#ue4 BackwardCompatibilityNodeConversion, ready for Editoronly Blueprints
Functions UKismetAIHelperLibrary::SpawnAI and  UGameplayStatics::CreateSaveGameObjectFromBlueprint are automatically converted.
Pins representing Blueprint and Subclass of UObject can match (switch in ini).

diff --git a/Engine/Config/BaseEditor.ini b/Engine/Config/BaseEditor.ini
--- a/Engine/Config/BaseEditor.ini
+++ b/Engine/Config/BaseEditor.ini
@@ -244,4 +244,7 @@ SupportWebsiteURL="http://unrealengine.com/support"
 [IntrinsicMetaData]
 
 [UserDefinedStructure]
-bUseUserDefinedStructure=false
\ No newline at end of file
+bUseUserDefinedStructure=false
+
+[EditoronlyBP]
+bAllowClassObjAndBluepirntPinMatching=true
\ No newline at end of file
diff --git a/Engine/Source/Editor/BlueprintGraph/Classes/EdGraphSchema_K2.h b/Engine/Source/Editor/BlueprintGraph/Classes/EdGraphSchema_K2.h
--- a/Engine/Source/Editor/BlueprintGraph/Classes/EdGraphSchema_K2.h
+++ b/Engine/Source/Editor/BlueprintGraph/Classes/EdGraphSchema_K2.h
@@ -415,7 +415,7 @@ public:
 	virtual int32 GetNodeSelectionCount(const UEdGraph* Graph) const OVERRIDE;
 	virtual TSharedPtr<FEdGraphSchemaAction> GetCreateCommentAction() const OVERRIDE;
 	virtual bool FadeNodeWhenDraggingOffPin(const UEdGraphNode* Node, const UEdGraphPin* Pin) const OVERRIDE;
-	virtual void BackwardCompatibilityNodeConversion(UEdGraph* Graph) const OVERRIDE;
+	virtual void BackwardCompatibilityNodeConversion(UEdGraph* Graph, bool bOnlySafeChanges) const OVERRIDE;
 	// End EdGraphSchema Interface
 
 	// Do validation, that doesn't require a knowledge about actual pin. 
diff --git a/Engine/Source/Editor/BlueprintGraph/Private/EdGraphSchema_K2.cpp b/Engine/Source/Editor/BlueprintGraph/Private/EdGraphSchema_K2.cpp
--- a/Engine/Source/Editor/BlueprintGraph/Private/EdGraphSchema_K2.cpp
+++ b/Engine/Source/Editor/BlueprintGraph/Private/EdGraphSchema_K2.cpp
@@ -2447,6 +2447,37 @@ bool UEdGraphSchema_K2::ArePinTypesCompatible(const FEdGraphPinType& Output, con
 		return true;
 	}
 
+	// Pins representing BLueprint objects and subclass of UObject can match when EditoronlyBP.bAllowClassObjAndBluepirntPinMatching=true (BaseEngine.ini)
+	// It's required for converting all UBlueprint references into UClass.
+	struct FObjClassAndBlueprintHelper
+	{
+	private:
+		bool bAllow;
+	public:
+		FObjClassAndBlueprintHelper() : bAllow(false)
+		{
+			GConfig->GetBool(TEXT("EditoronlyBP"), TEXT("bAllowClassObjAndBluepirntPinMatching"), bAllow, GEditorIni);
+		}
+
+		bool Match(const FEdGraphPinType& A, const FEdGraphPinType& B, const UEdGraphSchema_K2& Schema) const
+		{
+			if (bAllow && (B.PinCategory == Schema.PC_Object) && (A.PinCategory == Schema.PC_Class))
+			{
+				const bool bAIsObjectClass = (UObject::StaticClass() == A.PinSubCategoryObject.Get());
+				const UClass* BClass = Cast<UClass>(B.PinSubCategoryObject.Get());
+				const bool bBIsBlueprintObj = BClass && BClass->IsChildOf(UBlueprint::StaticClass());
+				return bAIsObjectClass && bBIsBlueprintObj;
+			}
+			return false;
+		}
+	};
+
+	static FObjClassAndBlueprintHelper MatchHelper;
+	if (MatchHelper.Match(Input, Output, *this) || MatchHelper.Match(Output, Input, *this))
+	{
+		return true;
+	}
+
 	return false;
 }
 
@@ -3103,89 +3134,269 @@ bool UEdGraphSchema_K2::FadeNodeWhenDraggingOffPin(const UEdGraphNode* Node, con
 	return false;
 }
 
-void UEdGraphSchema_K2::BackwardCompatibilityNodeConversion(UEdGraph* Graph) const 
+struct FBackwardCompatibilityConversionHelper
 {
-	if (Graph)
+	static bool ConvertNode(
+		UK2Node* OldNode,
+		const FString& BlueprintPinName,
+		UK2Node* NewNode,
+		const FString& ClassPinName,
+		const UEdGraphSchema_K2& Schema,
+		bool bOnlyWithDefaultBlueprint)
 	{
-		TArray<UK2Node_SpawnActor*> SpawnActorNodes;
-		Graph->GetNodesOfClass(SpawnActorNodes);
-		for(auto It = SpawnActorNodes.CreateIterator(); It; ++It)
-		{
-			UK2Node_SpawnActor* const OldSpawnActor = *It;
-			check(OldSpawnActor);
-			UEdGraphPin* OldBlueprintPin = OldSpawnActor->GetBlueprintPin();
-			check(OldBlueprintPin);
-
-			if ((0 == OldBlueprintPin->LinkedTo.Num()) && (NULL != OldBlueprintPin->DefaultObject))
-			{
-				const auto SpawnActorFromClass = NewObject<UK2Node_SpawnActorFromClass>(Graph);
-				check(SpawnActorFromClass);
-				SpawnActorFromClass->SetFlags(RF_Transactional);
-				Graph->AddNode(SpawnActorFromClass, false, false);
-				SpawnActorFromClass->CreateNewGuid();
-				SpawnActorFromClass->PostPlacedNewNode();
-				SpawnActorFromClass->AllocateDefaultPins();
-				SpawnActorFromClass->NodePosX = OldSpawnActor->NodePosX;
-				SpawnActorFromClass->NodePosY = OldSpawnActor->NodePosY;
-
-				const auto ClassPin = SpawnActorFromClass->GetClassPin();
-				check(ClassPin);
+		check(OldNode && NewNode);
+		const auto Blueprint = OldNode->GetBlueprint();
+
+		auto Graph = OldNode->GetGraph();
+		if (!Graph)
+		{
+			UE_LOG(LogBlueprint, Warning, TEXT("BackwardCompatibilityNodeConversion Error bp: '%s' node: '%s'. No graph containing the node."),
+				Blueprint ? *Blueprint->GetName() : TEXT("Unknown"),
+				*OldNode->GetName(),
+				*BlueprintPinName);
+			return false;
+		}
+
+		auto OldBlueprintPin = OldNode->FindPin(BlueprintPinName);
+		if (!OldBlueprintPin)
+		{
+			UE_LOG(LogBlueprint, Warning, TEXT("BackwardCompatibilityNodeConversion Error bp: '%s' node: '%s'. No bp pin found '%s'"),
+				Blueprint ? *Blueprint->GetName() : TEXT("Unknown"),
+				*OldNode->GetName(),
+				*BlueprintPinName);
+			return false;
+		}
+
+		const bool bNondefaultBPConnected = (OldBlueprintPin->LinkedTo.Num() > 0);
+		const bool bNodeUsesDefaultBlueprint = !bNondefaultBPConnected && OldBlueprintPin->DefaultObject;
+		const bool bTryConvert = bNodeUsesDefaultBlueprint || (!bOnlyWithDefaultBlueprint && bNondefaultBPConnected);
+		if (bTryConvert)
+		{
+			// CREATE NEW NODE
+			NewNode->SetFlags(RF_Transactional);
+			Graph->AddNode(NewNode, false, false);
+			NewNode->CreateNewGuid();
+			NewNode->PostPlacedNewNode();
+			NewNode->AllocateDefaultPins();
+			NewNode->NodePosX = OldNode->NodePosX;
+			NewNode->NodePosY = OldNode->NodePosY;
+
+			const auto ClassPin = NewNode->FindPin(ClassPinName);
+			if (!ClassPin)
+			{
+				UE_LOG(LogBlueprint, Warning, TEXT("BackwardCompatibilityNodeConversion Error bp: '%s' node: '%s'. No class pin found '%s'"),
+					Blueprint ? *Blueprint->GetName() : TEXT("Unknown"),
+					*NewNode->GetName(),
+					*ClassPinName);
+				return false;
+			}
+			auto TargetClass = Cast<UClass>(ClassPin->PinType.PinSubCategoryObject.Get());
+			if (!TargetClass)
+			{
+				UE_LOG(LogBlueprint, Warning, TEXT("BackwardCompatibilityNodeConversion Error bp: '%s' node: '%s'. No class found '%s'"),
+					Blueprint ? *Blueprint->GetName() : TEXT("Unknown"),
+					*NewNode->GetName(),
+					*ClassPinName);
+				return false;
+			}
+
+			// REPLACE BLUEPRINT WITH CLASS
+			if (!bNondefaultBPConnected)
+			{
+				// DEFAULT VALUE
 				const auto UsedBlueprint = Cast<UBlueprint>(OldBlueprintPin->DefaultObject);
-				check(NULL != UsedBlueprint);
+				if (!UsedBlueprint)
+				{
+					UE_LOG(LogBlueprint, Warning, TEXT("BackwardCompatibilityNodeConversion Error bp: '%s' node: '%s'. Wrong blueprint default value"),
+						Blueprint ? *Blueprint->GetName() : TEXT("Unknown"),
+						*OldNode->GetName());
+					return false;
+				}
 				ensure(NULL != *UsedBlueprint->GeneratedClass);
-				const auto Blueprint = OldSpawnActor->GetBlueprint();
-				TrySetDefaultObject(*ClassPin, *UsedBlueprint->GeneratedClass);
+				Schema.TrySetDefaultObject(*ClassPin, *UsedBlueprint->GeneratedClass);
 				if (ClassPin->DefaultObject != *UsedBlueprint->GeneratedClass)
 				{
-					FString ErrorStr = IsPinDefaultValid(ClassPin, FString(), *UsedBlueprint->GeneratedClass, FText());
-					UE_LOG(LogBlueprint, Warning, TEXT("BackwardCompatibilityNodeConversion Error 'cannot set class' in blueprint: %s actor bp: %s, reason: %s"),
+					auto ErrorStr = Schema.IsPinDefaultValid(ClassPin, FString(), *UsedBlueprint->GeneratedClass, FText());
+					UE_LOG(LogBlueprint, Warning, TEXT("BackwardCompatibilityNodeConversion Error 'cannot set class' in blueprint: %s node: '%s' actor bp: %s, reason: %s"),
 						Blueprint ? *Blueprint->GetName() : TEXT("Unknown"),
+						*OldNode->GetName(),
 						UsedBlueprint ? *UsedBlueprint->GetName() : TEXT("Unknown"),
 						ErrorStr.IsEmpty() ? TEXT("Unknown") : *ErrorStr);
+					return false;
+				}
+			}
+			else
+			{
+				// LINK
+				auto CastNode = NewObject<UK2Node_ClassDynamicCast>(Graph);
+				CastNode->TargetType = TargetClass;
+				Graph->AddNode(CastNode, false, false);
+				CastNode->CreateNewGuid();
+				CastNode->PostPlacedNewNode();
+				CastNode->AllocateDefaultPins();
+				const int32 OffsetOnGraph = 200;
+				CastNode->NodePosX = OldNode->NodePosX - OffsetOnGraph;
+				CastNode->NodePosY = OldNode->NodePosY;
+
+				auto ExecPin = OldNode->GetExecPin();
+				auto ExecCastPin = CastNode->GetExecPin();
+				check(ExecCastPin);
+				if (!ExecPin || !Schema.MovePinLinks(*ExecPin, *ExecCastPin).CanSafeConnect())
+				{
+					UE_LOG(LogBlueprint, Warning, TEXT("BackwardCompatibilityNodeConversion Error 'cannot connect' in blueprint: %s, pin: %s"),
+						Blueprint ? *Blueprint->GetName() : TEXT("Unknown"),
+						*ExecCastPin->PinName);
+					return false;
+				}
+
+				auto ValidCastPin = CastNode->GetValidCastPin();
+				check(ValidCastPin);
+				if (!Schema.TryCreateConnection(ValidCastPin, ExecPin))
+				{
+					UE_LOG(LogBlueprint, Warning, TEXT("BackwardCompatibilityNodeConversion Error 'cannot connect' in blueprint: %s, pin: %s"),
+						Blueprint ? *Blueprint->GetName() : TEXT("Unknown"),
+						*ValidCastPin->PinName);
+					return false;
+				}
+
+				auto CastSourcePin = CastNode->GetCastSourcePin();
+				check(CastSourcePin);
+				if (!Schema.MovePinLinks(*OldBlueprintPin, *CastSourcePin).CanSafeConnect())
+				{
+					UE_LOG(LogBlueprint, Warning, TEXT("BackwardCompatibilityNodeConversion Error 'cannot connect' in blueprint: %s, pin: %s"),
+						Blueprint ? *Blueprint->GetName() : TEXT("Unknown"),
+						*CastSourcePin->PinName);
+					return false;
+				}
+
+				auto CastResultPin = CastNode->GetCastResultPin();
+				check(CastResultPin);
+				if (!Schema.TryCreateConnection(CastResultPin, ClassPin))
+				{
+					UE_LOG(LogBlueprint, Warning, TEXT("BackwardCompatibilityNodeConversion Error 'cannot connect' in blueprint: %s, pin: %s"),
+						Blueprint ? *Blueprint->GetName() : TEXT("Unknown"),
+						*CastResultPin->PinName);
+					return false;
 				}
+			}
 
-				TArray<UEdGraphPin*> OldPins;
-				OldPins.Add(OldBlueprintPin);
-				for(auto PinIter = SpawnActorFromClass->Pins.CreateIterator(); PinIter; ++PinIter)
+			// MOVE OTHER PINS
+			TArray<UEdGraphPin*> OldPins;
+			OldPins.Add(OldBlueprintPin);
+			for (auto PinIter = NewNode->Pins.CreateIterator(); PinIter; ++PinIter)
+			{
+				UEdGraphPin* const Pin = *PinIter;
+				check(Pin);
+				if (ClassPin != Pin)
 				{
-					UEdGraphPin* const Pin = *PinIter;
-					check(Pin);
-					if (ClassPin != Pin)
+					const auto OldPin = OldNode->FindPin(Pin->PinName);
+					if (NULL != OldPin)
 					{
-						const auto OldPin = OldSpawnActor->FindPin(Pin->PinName);
-						if(NULL != OldPin)
+						OldPins.Add(OldPin);
+						if (!Schema.MovePinLinks(*OldPin, *Pin).CanSafeConnect())
 						{
-							OldPins.Add(OldPin);
-							if(!MovePinLinks(*OldPin, *Pin).CanSafeConnect())
-							{
-								UE_LOG(LogBlueprint, Warning, TEXT("BackwardCompatibilityNodeConversion Error 'cannot connect' in blueprint: %s actor bp: %s, pin: %s"), 
-									Blueprint ? *Blueprint->GetName() : TEXT("Unknown"), 
-									UsedBlueprint ? *UsedBlueprint->GetName() : TEXT("Unknown"),
-									*Pin->PinName);
-							}
-						}
-						else
-						{
-							UE_LOG(LogBlueprint, Warning, TEXT("BackwardCompatibilityNodeConversion Error 'missing old pin' in blueprint: %s actor bp: %s"), 
-								Blueprint ? *Blueprint->GetName() : TEXT("Unknown"), 
-								UsedBlueprint ? *UsedBlueprint->GetName() : TEXT("Unknown"),
-								Pin ? *Pin->PinName : TEXT("Unknown"));
+							UE_LOG(LogBlueprint, Warning, TEXT("BackwardCompatibilityNodeConversion Error 'cannot connect' in blueprint: %s, pin: %s"),
+								Blueprint ? *Blueprint->GetName() : TEXT("Unknown"),
+								*Pin->PinName);
 						}
 					}
-				}
-				OldSpawnActor->BreakAllNodeLinks();
-				for(auto PinIter = OldSpawnActor->Pins.CreateIterator(); PinIter; ++PinIter)
-				{
-					if(!OldPins.Contains(*PinIter))
+					else
 					{
-						UEdGraphPin* Pin = *PinIter;
-						UE_LOG(LogBlueprint, Warning, TEXT("BackwardCompatibilityNodeConversion Error 'missing new pin' in blueprint: %s actor bp: %s"), 
-							Blueprint ? *Blueprint->GetName() : TEXT("Unknown"), 
-							UsedBlueprint ? *UsedBlueprint->GetName() : TEXT("Unknown"),
+						UE_LOG(LogBlueprint, Warning, TEXT("BackwardCompatibilityNodeConversion Error 'missing old pin' in blueprint: %s"),
+							Blueprint ? *Blueprint->GetName() : TEXT("Unknown"),
 							Pin ? *Pin->PinName : TEXT("Unknown"));
 					}
 				}
-				Graph->RemoveNode(OldSpawnActor);
+			}
+			OldNode->BreakAllNodeLinks();
+			for (auto PinIter = OldNode->Pins.CreateIterator(); PinIter; ++PinIter)
+			{
+				if (!OldPins.Contains(*PinIter))
+				{
+					UEdGraphPin* Pin = *PinIter;
+					UE_LOG(LogBlueprint, Warning, TEXT("BackwardCompatibilityNodeConversion Error 'missing new pin' in blueprint: %s"),
+						Blueprint ? *Blueprint->GetName() : TEXT("Unknown"),
+						Pin ? *Pin->PinName : TEXT("Unknown"));
+				}
+			}
+			Graph->RemoveNode(OldNode);
+			return true;
+		}
+		return false;
+	}
+
+	struct FFunctionCallParams
+	{
+		const FName OldFuncName;
+		const FName NewFuncName;
+		const FString& BlueprintPinName;
+		const FString& ClassPinName;
+		const UClass* FuncScope;
+
+		FFunctionCallParams(FName InOldFunc, FName InNewFunc, const FString& InBlueprintPinName, const FString& InClassPinName, const UClass* InFuncScope)
+			: OldFuncName(InOldFunc), NewFuncName(InNewFunc), BlueprintPinName(InBlueprintPinName), ClassPinName(InClassPinName), FuncScope(InFuncScope)
+		{
+			check(FuncScope);
+		}
+
+	};
+
+	static void ConvertFunctionCallNodes(const FFunctionCallParams& ConversionParams, TArray<UK2Node_CallFunction*>& Nodes, UEdGraph* Graph, const UEdGraphSchema_K2& Schema, bool bOnlyWithDefaultBlueprint)
+	{
+		const UFunction* OldFunc = ConversionParams.FuncScope->FindFunctionByName(ConversionParams.OldFuncName);
+		check(OldFunc);
+		const UFunction* NewFunc = ConversionParams.FuncScope->FindFunctionByName(ConversionParams.NewFuncName);
+		check(NewFunc);
+
+		for (auto It = Nodes.CreateIterator(); It; ++It)
+		{
+			if (OldFunc == (*It)->GetTargetFunction())
+			{
+				auto NewNode = NewObject<UK2Node_CallFunction>(Graph);
+				NewNode->SetFromFunction(NewFunc);
+				ConvertNode(*It, ConversionParams.BlueprintPinName, NewNode, 
+					ConversionParams.ClassPinName, Schema, bOnlyWithDefaultBlueprint);
+			}
+		}
+	}
+};
+
+void UEdGraphSchema_K2::BackwardCompatibilityNodeConversion(UEdGraph* Graph, bool bOnlySafeChanges) const 
+{
+	if (Graph)
+	{
+		{
+			static const FString BlueprintPinName(TEXT("Blueprint"));
+			static const FString ClassPinName(TEXT("Class"));
+			TArray<UK2Node_SpawnActor*> SpawnActorNodes;
+			Graph->GetNodesOfClass(SpawnActorNodes);
+			for (auto It = SpawnActorNodes.CreateIterator(); It; ++It)
+			{
+				FBackwardCompatibilityConversionHelper::ConvertNode(
+					*It, BlueprintPinName, NewObject<UK2Node_SpawnActorFromClass>(Graph),
+					ClassPinName, *this, bOnlySafeChanges);
+			}
+		}
+
+		{
+			TArray<UK2Node_CallFunction*> Nodes;
+			Graph->GetNodesOfClass(Nodes);
+			{
+				static const FString BlueprintPinName(TEXT("Pawn"));
+				static const FString ClassPinName(TEXT("PawnClass"));
+				static const FName OldFuncName(GET_FUNCTION_NAME_CHECKED(UKismetAIHelperLibrary, SpawnAI));
+				static const FName NewFuncName(GET_FUNCTION_NAME_CHECKED(UKismetAIHelperLibrary, SpawnAIFromClass));
+				FBackwardCompatibilityConversionHelper::FFunctionCallParams Params(OldFuncName, NewFuncName, BlueprintPinName, ClassPinName, UKismetAIHelperLibrary::StaticClass());
+				FBackwardCompatibilityConversionHelper::ConvertFunctionCallNodes(Params, Nodes, Graph, *this, bOnlySafeChanges);
+			}
+
+			{
+				static const FString BlueprintPinName(TEXT("SaveGameBlueprint"));
+				static const FString ClassPinName(TEXT("SaveGameClass"));
+				static const FName OldFuncName(GET_FUNCTION_NAME_CHECKED(UGameplayStatics, CreateSaveGameObjectFromBlueprint));
+				static const FName NewFuncName(GET_FUNCTION_NAME_CHECKED(UGameplayStatics, CreateSaveGameObject));
+				FBackwardCompatibilityConversionHelper::FFunctionCallParams Params(OldFuncName, NewFuncName, BlueprintPinName, ClassPinName, UGameplayStatics::StaticClass());
+				FBackwardCompatibilityConversionHelper::ConvertFunctionCallNodes(Params, Nodes, Graph, *this, bOnlySafeChanges);
 			}
 		}
 
diff --git a/Engine/Source/Runtime/Engine/Classes/EdGraph/EdGraphSchema.h b/Engine/Source/Runtime/Engine/Classes/EdGraph/EdGraphSchema.h
--- a/Engine/Source/Runtime/Engine/Classes/EdGraph/EdGraphSchema.h
+++ b/Engine/Source/Runtime/Engine/Classes/EdGraph/EdGraphSchema.h
@@ -694,5 +694,5 @@ class ENGINE_API UEdGraphSchema : public UObject
 	/* When dragging off a pin, we want to duck the alpha of some nodes */
 	virtual bool FadeNodeWhenDraggingOffPin(const UEdGraphNode* Node, const UEdGraphPin* Pin) const { return false; }
 
-	virtual void BackwardCompatibilityNodeConversion(UEdGraph* Graph) const { }
+	virtual void BackwardCompatibilityNodeConversion(UEdGraph* Graph, bool bOnlySafeChanges) const { }
 };
diff --git a/Engine/Source/Runtime/Engine/Private/Blueprint.cpp b/Engine/Source/Runtime/Engine/Private/Blueprint.cpp
--- a/Engine/Source/Runtime/Engine/Private/Blueprint.cpp
+++ b/Engine/Source/Runtime/Engine/Private/Blueprint.cpp
@@ -395,7 +395,7 @@ void UBlueprint::PostLoad()
 		{
 			UEdGraph* const Graph = *It;
 			const UEdGraphSchema* Schema = Graph->GetSchema();
-			Schema->BackwardCompatibilityNodeConversion(Graph);
+			Schema->BackwardCompatibilityNodeConversion(Graph, true);
 		}
 	}
 
