TTP #326008: Adding a scoped transaction to node refreshes (so that removed pins properly get reconstructed on undo).

diff --git a/Engine/Source/Editor/BlueprintGraph/Private/K2Node.cpp b/Engine/Source/Editor/BlueprintGraph/Private/K2Node.cpp
--- a/Engine/Source/Editor/BlueprintGraph/Private/K2Node.cpp
+++ b/Engine/Source/Editor/BlueprintGraph/Private/K2Node.cpp
@@ -3,6 +3,7 @@
 
 #include "BlueprintGraphPrivatePCH.h"
 #include "KismetDebugUtilities.h" // for HasDebuggingData(), GetWatchText()
+#include "ScopedTransaction.h"
 
 #define LOCTEXT_NAMESPACE "K2Node"
 
@@ -315,7 +316,10 @@ void UK2Node::RemovePinsFromOldPins(TArray<UEdGraphPin*>& OldPins, int32 Removed
 }
 
 void UK2Node::ReconstructNode()
-{							  
+{		
+	FString TransactionDesc = FString::Printf(*LOCTEXT("ReconstructNodeTransaction", "'%s' Node Refreshed").ToString(), *GetNodeTitle(ENodeTitleType::FullTitle));
+	FScopedTransaction const Transaction(FText::FromString(TransactionDesc));
+
 	Modify();
 
 	UBlueprint* Blueprint = GetBlueprint();
diff --git a/Engine/Source/Editor/UnrealEd/Private/Kismet2/BlueprintEditorUtils.cpp b/Engine/Source/Editor/UnrealEd/Private/Kismet2/BlueprintEditorUtils.cpp
--- a/Engine/Source/Editor/UnrealEd/Private/Kismet2/BlueprintEditorUtils.cpp
+++ b/Engine/Source/Editor/UnrealEd/Private/Kismet2/BlueprintEditorUtils.cpp
@@ -244,6 +244,9 @@ void FBlueprintEditorUtils::RefreshAllNodes(UBlueprint* Blueprint)
 		AllNodes.Sort(FCompareNodePriority());
 	}
 
+	// scoped transaction (so users don't have to undo every single ReconstructNode)
+	FScopedTransaction const Transaction(LOCTEXT("ReconstructNodes", "Refresh All Nodes"));
+
 	bool bLastChangesStructure = (AllNodes.Num() > 0) ? AllNodes[0]->NodeCausesStructuralBlueprintChange() : true;
 	for( TArray<UK2Node*>::TIterator NodeIt(AllNodes); NodeIt; ++NodeIt )
 	{
diff --git a/Engine/Source/Runtime/Engine/Private/EdGraph/EdGraph.cpp b/Engine/Source/Runtime/Engine/Private/EdGraph/EdGraph.cpp
--- a/Engine/Source/Runtime/Engine/Private/EdGraph/EdGraph.cpp
+++ b/Engine/Source/Runtime/Engine/Private/EdGraph/EdGraph.cpp
@@ -846,6 +846,10 @@ void UEdGraphPin::CopyPersistentDataFromOldPin(const UEdGraphPin& SourcePin)
 	{
 		UEdGraphPin* OtherPin = SourcePin.LinkedTo[LinkIndex];
 		check(NULL != OtherPin);
+		
+		Modify();
+		OtherPin->Modify();
+
 		OtherPin->MakeLinkTo(this);
 	}
 
