#ue4 BP: EnumLiteral node.
ttp331514 BP: need simpler way to specify literal enum values

diff --git a/Engine/Source/Editor/BlueprintGraph/Classes/K2Node_EnumLiteral.h b/Engine/Source/Editor/BlueprintGraph/Classes/K2Node_EnumLiteral.h
new file mode 100644
--- /dev/null
+++ b/Engine/Source/Editor/BlueprintGraph/Classes/K2Node_EnumLiteral.h
@@ -0,0 +1,33 @@
+// Copyright 1998-2014 Epic Games, Inc. All Rights Reserved.
+
+#pragma once
+#include "K2Node_EnumLiteral.generated.h"
+
+UCLASS(MinimalAPI)
+class UK2Node_EnumLiteral : public UK2Node, public INodeDependingOnEnumInterface
+{
+	GENERATED_UCLASS_BODY()
+
+	UPROPERTY()
+	UEnum* Enum;
+
+	static const FString& GetEnumInputPinName();
+
+	// Begin UEdGraphNode interface
+	virtual void AllocateDefaultPins() OVERRIDE;
+	virtual FString GetTooltip() const OVERRIDE;
+	virtual FString GetNodeTitle(ENodeTitleType::Type TitleType) const OVERRIDE;
+	virtual FName GetPaletteIcon(FLinearColor& OutColor) const OVERRIDE;
+	// End UEdGraphNode interface
+
+	// Begin UK2Node interface
+	virtual bool IsNodePure() const OVERRIDE { return true; }
+	virtual FNodeHandlingFunctor* CreateNodeHandler(FKismetCompilerContext& CompilerContext) const OVERRIDE;
+	// End UK2Node interface
+
+	// INodeDependingOnEnumInterface
+	virtual class UEnum* GetEnum() const OVERRIDE { return Enum; }
+	virtual bool ShouldBeReconstructedAfterEnumChanged() const OVERRIDE { return true; }
+	// End of INodeDependingOnEnumInterface
+};
+
diff --git a/Engine/Source/Editor/BlueprintGraph/Private/K2ActionMenuBuilder.cpp b/Engine/Source/Editor/BlueprintGraph/Private/K2ActionMenuBuilder.cpp
--- a/Engine/Source/Editor/BlueprintGraph/Private/K2ActionMenuBuilder.cpp
+++ b/Engine/Source/Editor/BlueprintGraph/Private/K2ActionMenuBuilder.cpp
@@ -236,11 +236,11 @@ static void GetCommentAction(UBlueprint const* BlueprintIn, FGraphActionListBuil
 }
 
 /** */
-static void GetEnumUtilitiesNodes(FBlueprintGraphActionListBuilder& ContextMenuBuilder, bool bNumEnum, bool bForEach, bool bCastFromByte, UEnum* OnlyEnum = NULL)
+static void GetEnumUtilitiesNodes(FBlueprintGraphActionListBuilder& ContextMenuBuilder, bool bNumEnum, bool bForEach, bool bCastFromByte, bool bLiteralByte, UEnum* OnlyEnum = NULL)
 {
 	struct FGetEnumUtilitiesNodes
 	{
-		static void Get(FBlueprintGraphActionListBuilder& InContextMenuBuilder, bool bInNumEnum, bool bInForEach, bool bInCastFromByte, UEnum* Enum, const FString& Category)
+		static void Get(FBlueprintGraphActionListBuilder& InContextMenuBuilder, bool bInNumEnum, bool bInForEach, bool bInCastFromByte, bool bLiteralByte, UEnum* Enum, const FString& Category)
 		{
 			const bool bIsBlueprintType = UEdGraphSchema_K2::IsAllowableBlueprintVariableType(Enum);
 			if (bIsBlueprintType)
@@ -275,13 +275,23 @@ static void GetEnumUtilitiesNodes(FBlueprintGraphActionListBuilder& ContextMenuB
 						InContextMenuBuilder, Category, NodeTemplate->GetNodeTitle(ENodeTitleType::ListView),  NodeTemplate->GetTooltip(), 0, NodeTemplate->GetKeywords());
 					Action->NodeTemplate = NodeTemplate;
 				}
+
+				if (bLiteralByte)
+				{
+					UK2Node_EnumLiteral* EnumTemplate = InContextMenuBuilder.CreateTemplateNode<UK2Node_EnumLiteral>();
+					EnumTemplate->Enum = Enum;
+					UK2Node* NodeTemplate = EnumTemplate;
+					TSharedPtr<FEdGraphSchemaAction_K2NewNode> Action = FK2ActionMenuBuilder::AddNewNodeAction(
+						InContextMenuBuilder, Category, NodeTemplate->GetNodeTitle(ENodeTitleType::ListView), NodeTemplate->GetTooltip(), 0, NodeTemplate->GetKeywords());
+					Action->NodeTemplate = NodeTemplate;
+				}
 			}
 		}
 	};
 
 	if (OnlyEnum)
 	{
-		FGetEnumUtilitiesNodes::Get(ContextMenuBuilder, bNumEnum, bForEach, bCastFromByte, OnlyEnum, K2ActionCategories::EnumCategory);
+		FGetEnumUtilitiesNodes::Get(ContextMenuBuilder, bNumEnum, bForEach, bCastFromByte, bLiteralByte, OnlyEnum, K2ActionCategories::EnumCategory);
 	}
 	else
 	{
@@ -289,7 +299,7 @@ static void GetEnumUtilitiesNodes(FBlueprintGraphActionListBuilder& ContextMenuB
 		{
 			if (UEnum* CurrentEnum = *EnumIt)
 			{
-				FGetEnumUtilitiesNodes::Get(ContextMenuBuilder, bNumEnum, bForEach, bCastFromByte, CurrentEnum, K2ActionCategories::EnumCategory);
+				FGetEnumUtilitiesNodes::Get(ContextMenuBuilder, bNumEnum, bForEach, bCastFromByte, bLiteralByte, CurrentEnum, K2ActionCategories::EnumCategory);
 			}
 		}
 	}
@@ -984,7 +994,7 @@ void FK2ActionMenuBuilder::GetContextAllowedNodeTypes(FBlueprintGraphActionListB
 
 		FEditorDelegates::OnBlueprintContextMenuCreated.Broadcast(ContextMenuBuilder);
 
-		GetEnumUtilitiesNodes(ContextMenuBuilder, true, true, true);
+		GetEnumUtilitiesNodes(ContextMenuBuilder, true, true, true, true);
 	}
 }
 
@@ -1303,7 +1313,7 @@ void FK2ActionMenuBuilder::GetPinAllowedNodeTypes(FBlueprintGraphActionListBuild
 
 		if ((FromPin.Direction == EGPD_Input) && (FromPin.PinType.PinCategory == K2Schema->PC_Int))
 		{
-			GetEnumUtilitiesNodes(ContextMenuBuilder, true, false, false);
+			GetEnumUtilitiesNodes(ContextMenuBuilder, true, false, false, true);
 		}
 		else if(FromPin.PinType.PinCategory == K2Schema->PC_Text)
 		{
@@ -1321,9 +1331,13 @@ void FK2ActionMenuBuilder::GetPinAllowedNodeTypes(FBlueprintGraphActionListBuild
 			}
 			if (FromPin.Direction == EGPD_Input)
 			{
-				GetEnumUtilitiesNodes(ContextMenuBuilder, false, true, true, Enum);
+				GetEnumUtilitiesNodes(ContextMenuBuilder, false, true, true, false, Enum);
 			}
 		}
+		else if ((FromPin.PinType.PinCategory == K2Schema->PC_Byte) && !Enum && FromPin.Direction == EGPD_Input)
+		{
+			GetEnumUtilitiesNodes(ContextMenuBuilder, false, false, false, true);
+		}
 
 		// Assignment statement for local variables
 		if(FromPin.GetOwningNode()->IsA(UK2Node_TemporaryVariable::StaticClass()))
diff --git a/Engine/Source/Editor/BlueprintGraph/Private/K2Node_EnumLiteral.cpp b/Engine/Source/Editor/BlueprintGraph/Private/K2Node_EnumLiteral.cpp
new file mode 100644
--- /dev/null
+++ b/Engine/Source/Editor/BlueprintGraph/Private/K2Node_EnumLiteral.cpp
@@ -0,0 +1,84 @@
+// Copyright 1998-2014 Epic Games, Inc. All Rights Reserved.
+
+
+#include "BlueprintGraphPrivatePCH.h"
+
+const FString& UK2Node_EnumLiteral::GetEnumInputPinName()
+{
+	static const FString Name(TEXT("Enum"));
+	return Name;
+}
+
+UK2Node_EnumLiteral::UK2Node_EnumLiteral(const class FPostConstructInitializeProperties& PCIP)
+	: Super(PCIP)
+{
+}
+
+void UK2Node_EnumLiteral::AllocateDefaultPins()
+{
+	check(Enum);
+
+	const UEdGraphSchema_K2* Schema = GetDefault<UEdGraphSchema_K2>();
+
+	auto InputPin = CreatePin(EGPD_Input, Schema->PC_Byte, TEXT(""), Enum, false, false, GetEnumInputPinName());
+	Schema->SetPinDefaultValueBasedOnType(InputPin);
+
+	CreatePin(EGPD_Output, Schema->PC_Byte, TEXT(""), Enum, false, false, Schema->PN_ReturnValue);
+
+	Super::AllocateDefaultPins();
+}
+
+FName UK2Node_EnumLiteral::GetPaletteIcon(FLinearColor& OutColor) const
+{ 
+	static const FName PaletteIconName(TEXT("GraphEditor.Enum_16x"));
+	return PaletteIconName;
+}
+
+FString UK2Node_EnumLiteral::GetTooltip() const
+{
+	const FString EnumName = Enum ? Enum->GetName() : TEXT("(bad enum)");
+
+	return FString::Printf(*NSLOCTEXT("K2Node", "EnumLiteral_Tooltip", "Literal enum %s").ToString(), *EnumName);
+}
+
+FString UK2Node_EnumLiteral::GetNodeTitle(ENodeTitleType::Type TitleType) const
+{
+	return GetTooltip();
+}
+
+class FKCHandler_EnumLiteral : public FNodeHandlingFunctor
+{
+public:
+	FKCHandler_EnumLiteral(FKismetCompilerContext& InCompilerContext)
+		: FNodeHandlingFunctor(InCompilerContext)
+	{
+	}
+
+	virtual void RegisterNets(FKismetFunctionContext& Context, UEdGraphNode* Node) OVERRIDE
+	{
+		check(Context.Schema && Cast<UK2Node_EnumLiteral>(Node));
+		FNodeHandlingFunctor::RegisterNets(Context, Node);
+
+		UEdGraphPin* InPin = Node->FindPinChecked(UK2Node_EnumLiteral::GetEnumInputPinName());
+		UEdGraphPin* Net = FEdGraphUtilities::GetNetFromPin(InPin);
+		if (Context.NetMap.Find(Net) == NULL)
+		{
+			FBPTerminal* Term = new (Context.IsEventGraph() ? Context.EventGraphLocals : Context.Locals) FBPTerminal();
+			Term->CopyFromPin(Net, Context.NetNameMap->MakeValidName(Net));
+			Context.NetMap.Add(Net, Term);
+		}
+
+		FBPTerminal** ValueSource = Context.NetMap.Find(Net);
+		check(ValueSource && *ValueSource);
+		UEdGraphPin* OutPin = Node->FindPinChecked(Context.Schema->PN_ReturnValue);
+		if (ensure(Context.NetMap.Find(OutPin) == NULL))
+		{
+			Context.NetMap.Add(OutPin, *ValueSource);
+		}
+	}
+};
+
+FNodeHandlingFunctor* UK2Node_EnumLiteral::CreateNodeHandler(FKismetCompilerContext& CompilerContext) const
+{
+	return new FKCHandler_EnumLiteral(CompilerContext);
+}
