#ue4 ExposeOnSpawn is not mutually exclusive with BlueprintReadOnly

diff --git a/Engine/Source/Editor/BlueprintGraph/Private/K2Node_SpawnActor.cpp b/Engine/Source/Editor/BlueprintGraph/Private/K2Node_SpawnActor.cpp
--- a/Engine/Source/Editor/BlueprintGraph/Private/K2Node_SpawnActor.cpp
+++ b/Engine/Source/Editor/BlueprintGraph/Private/K2Node_SpawnActor.cpp
@@ -65,9 +65,7 @@ void UK2Node_SpawnActor::CreatePinsForClass(UClass* InClass)
 		UClass* PropertyClass = CastChecked<UClass>(Property->GetOuter());
 		const bool bIsDelegate = Property->IsA(UMulticastDelegateProperty::StaticClass());
 		const bool bIsExposedToSpawn = UEdGraphSchema_K2::IsPropertyExposedOnSpawn(Property);
-
-		const bool bIsSettableExternally = !Property->HasAnyPropertyFlags(CPF_DisableEditOnInstance)
-			|| !Property->HasAnyPropertyFlags(CPF_BlueprintReadOnly);//@TODO: Remove this after old content is fixed up
+		const bool bIsSettableExternally = !Property->HasAnyPropertyFlags(CPF_DisableEditOnInstance);
 
 		if(	bIsExposedToSpawn &&
 			!Property->HasAnyPropertyFlags(CPF_Parm) && 
diff --git a/Engine/Source/Editor/BlueprintGraph/Private/K2Node_SpawnActorFromClass.cpp b/Engine/Source/Editor/BlueprintGraph/Private/K2Node_SpawnActorFromClass.cpp
--- a/Engine/Source/Editor/BlueprintGraph/Private/K2Node_SpawnActorFromClass.cpp
+++ b/Engine/Source/Editor/BlueprintGraph/Private/K2Node_SpawnActorFromClass.cpp
@@ -86,9 +86,7 @@ void UK2Node_SpawnActorFromClass::CreatePinsForClass(UClass* InClass)
 		UClass* PropertyClass = CastChecked<UClass>(Property->GetOuter());
 		const bool bIsDelegate = Property->IsA(UMulticastDelegateProperty::StaticClass());
 		const bool bIsExposedToSpawn = UEdGraphSchema_K2::IsPropertyExposedOnSpawn(Property);
-
-		const bool bIsSettableExternally = !Property->HasAnyPropertyFlags(CPF_DisableEditOnInstance)
-			|| !Property->HasAnyPropertyFlags(CPF_BlueprintReadOnly);//@TODO: Remove this after old content is fixed up
+		const bool bIsSettableExternally = !Property->HasAnyPropertyFlags(CPF_DisableEditOnInstance);
 
 		if(	bIsExposedToSpawn &&
 			!Property->HasAnyPropertyFlags(CPF_Parm) && 
diff --git a/Engine/Source/Programs/UnrealHeaderTool/Private/HeaderParser.cpp b/Engine/Source/Programs/UnrealHeaderTool/Private/HeaderParser.cpp
--- a/Engine/Source/Programs/UnrealHeaderTool/Private/HeaderParser.cpp
+++ b/Engine/Source/Programs/UnrealHeaderTool/Private/HeaderParser.cpp
@@ -2791,7 +2791,7 @@ bool FHeaderParser::GetVarType
 		const bool bExposeOnSpawn = (NULL != ExposeOnSpawnStr);
 		if (bExposeOnSpawn)
 		{
-			if ((CPF_DisableEditOnInstance | CPF_BlueprintReadOnly) & Flags)
+			if (0 != (CPF_DisableEditOnInstance & Flags))
 			{
 				UE_LOG(LogCompile, Warning, TEXT("Property cannot have 'DisableEditOnInstance' or 'BlueprintReadOnly' and 'ExposeOnSpawn' flags"));
 			}
