TTP #331765: Forcing coders to use MovePinLinksToIntermediate() when expanding nodes for compilation.

diff --git a/Engine/Source/Editor/BlueprintGraph/Private/K2Node_CommutativeAssociativeBinaryOperator.cpp b/Engine/Source/Editor/BlueprintGraph/Private/K2Node_CommutativeAssociativeBinaryOperator.cpp
--- a/Engine/Source/Editor/BlueprintGraph/Private/K2Node_CommutativeAssociativeBinaryOperator.cpp
+++ b/Engine/Source/Editor/BlueprintGraph/Private/K2Node_CommutativeAssociativeBinaryOperator.cpp
@@ -288,12 +288,12 @@ void UK2Node_CommutativeAssociativeBinaryOperator::ExpandNode(FKismetCompilerCon
 			else
 			{
 				// handle first created node (SrcFirstInput is skipped, and has no own node).
-				Schema->MovePinLinks(*SrcFirstInput, *NewOperatorInputA);
+				CompilerContext.MovePinLinksToIntermediate(*SrcFirstInput, *NewOperatorInputA);
 			}
 
 			UEdGraphPin* NewOperatorInputB = NewOperator->GetInputPin(1);
 			check(NewOperatorInputB);
-			Schema->MovePinLinks(*CurrentPin, *NewOperatorInputB);
+			CompilerContext.MovePinLinksToIntermediate(*CurrentPin, *NewOperatorInputB);
 
 			LastOutPin = NewOperator->FindOutPin();
 			check(LastOutPin);
@@ -301,7 +301,7 @@ void UK2Node_CommutativeAssociativeBinaryOperator::ExpandNode(FKismetCompilerCon
 
 		UEdGraphPin* TrueOutPin = FindOutPin();
 		check(TrueOutPin);
-		Schema->MovePinLinks(*TrueOutPin, *LastOutPin);
+		CompilerContext.MovePinLinksToIntermediate(*TrueOutPin, *LastOutPin);
 
 		BreakAllNodeLinks();
 	}
diff --git a/Engine/Source/Editor/BlueprintGraph/Private/K2Node_InputTouch.cpp b/Engine/Source/Editor/BlueprintGraph/Private/K2Node_InputTouch.cpp
--- a/Engine/Source/Editor/BlueprintGraph/Private/K2Node_InputTouch.cpp
+++ b/Engine/Source/Editor/BlueprintGraph/Private/K2Node_InputTouch.cpp
@@ -162,12 +162,12 @@ void UK2Node_InputTouch::ExpandNode(FKismetCompilerContext& CompilerContext, UEd
 			Schema->TryCreateConnection(TouchReleasedLocationInitialize->GetThenPin(), TouchReleasedFingerInitialize->GetExecPin());
 
 			// Move the original event connections to the then pin of the finger index assign
-			Schema->MovePinLinks(*GetPressedPin(), *TouchPressedFingerInitialize->GetThenPin());
-			Schema->MovePinLinks(*GetReleasedPin(), *TouchReleasedFingerInitialize->GetThenPin());
+			CompilerContext.MovePinLinksToIntermediate(*GetPressedPin(), *TouchPressedFingerInitialize->GetThenPin());
+			CompilerContext.MovePinLinksToIntermediate(*GetReleasedPin(), *TouchReleasedFingerInitialize->GetThenPin());
 
 			// Move the original event variable connections to the intermediate nodes
-			Schema->MovePinLinks(*GetLocationPin(), *TouchLocationVar->GetVariablePin());
-			Schema->MovePinLinks(*GetFingerIndexPin(), *TouchFingerVar->GetVariablePin());
+			CompilerContext.MovePinLinksToIntermediate(*GetLocationPin(), *TouchLocationVar->GetVariablePin());
+			CompilerContext.MovePinLinksToIntermediate(*GetFingerIndexPin(), *TouchFingerVar->GetVariablePin());
 		}
 		else
 		{
@@ -197,9 +197,9 @@ void UK2Node_InputTouch::ExpandNode(FKismetCompilerContext& CompilerContext, UEd
 				InputTouchEvent->bInternalEvent = true;
 				InputTouchEvent->AllocateDefaultPins();
 
-				Schema->MovePinLinks(*InputTouchPin, *Schema->FindExecutionPin(*InputTouchEvent, EGPD_Output));
-				Schema->MovePinLinks(*GetLocationPin(), *InputTouchEvent->FindPin(TEXT("Location")));
-				Schema->MovePinLinks(*GetFingerIndexPin(), *InputTouchEvent->FindPin(TEXT("FingerIndex")));
+				CompilerContext.MovePinLinksToIntermediate(*InputTouchPin, *Schema->FindExecutionPin(*InputTouchEvent, EGPD_Output));
+				CompilerContext.MovePinLinksToIntermediate(*GetLocationPin(), *InputTouchEvent->FindPin(TEXT("Location")));
+				CompilerContext.MovePinLinksToIntermediate(*GetFingerIndexPin(), *InputTouchEvent->FindPin(TEXT("FingerIndex")));
 			}
 		}
 	}
diff --git a/Engine/Source/Editor/KismetCompiler/Private/KismetCompiler.cpp b/Engine/Source/Editor/KismetCompiler/Private/KismetCompiler.cpp
--- a/Engine/Source/Editor/KismetCompiler/Private/KismetCompiler.cpp
+++ b/Engine/Source/Editor/KismetCompiler/Private/KismetCompiler.cpp
@@ -1987,7 +1987,7 @@ void FKismetCompilerContext::ExpandPlayMovieSceneNodes( UEdGraph* SourceGraph )
 FPinConnectionResponse FKismetCompilerContext::MovePinLinksToIntermediate(UEdGraphPin& SourcePin, UEdGraphPin& IntermediatePin)
 {
 	 UEdGraphSchema_K2 const* K2Schema = GetSchema();
-	 FPinConnectionResponse ConnectionResult = K2Schema->MovePinLinks(SourcePin, IntermediatePin);
+	 FPinConnectionResponse ConnectionResult = K2Schema->MovePinLinks(SourcePin, IntermediatePin, true);
 
 	 CheckConnectionResponse(ConnectionResult, SourcePin.GetOwningNode());
 	 MessageLog.NotifyIntermediateObjectCreation(&IntermediatePin, &SourcePin);
@@ -1998,7 +1998,7 @@ FPinConnectionResponse FKismetCompilerContext::MovePinLinksToIntermediate(UEdGra
 FPinConnectionResponse FKismetCompilerContext::CopyPinLinksToIntermediate(UEdGraphPin& SourcePin, UEdGraphPin& IntermediatePin)
 {
 	UEdGraphSchema_K2 const* K2Schema = GetSchema();
-	FPinConnectionResponse ConnectionResult = K2Schema->CopyPinLinks(SourcePin, IntermediatePin);
+	FPinConnectionResponse ConnectionResult = K2Schema->CopyPinLinks(SourcePin, IntermediatePin, true);
 
 	CheckConnectionResponse(ConnectionResult, SourcePin.GetOwningNode());
 	MessageLog.NotifyIntermediateObjectCreation(&IntermediatePin, &SourcePin);
diff --git a/Engine/Source/Editor/ReferenceViewer/Classes/ReferenceViewerSchema.h b/Engine/Source/Editor/ReferenceViewer/Classes/ReferenceViewerSchema.h
--- a/Engine/Source/Editor/ReferenceViewer/Classes/ReferenceViewerSchema.h
+++ b/Engine/Source/Editor/ReferenceViewer/Classes/ReferenceViewerSchema.h
@@ -14,8 +14,8 @@ public:
 	virtual FLinearColor GetPinTypeColor(const FEdGraphPinType& PinType) const OVERRIDE;
 	virtual void BreakPinLinks(UEdGraphPin& TargetPin, bool bSendsNodeNotifcation) const OVERRIDE;
 	virtual void BreakSinglePinLink(UEdGraphPin* SourcePin, UEdGraphPin* TargetPin) OVERRIDE;
-	virtual FPinConnectionResponse MovePinLinks(UEdGraphPin& MoveFromPin, UEdGraphPin& MoveToPin) const OVERRIDE;
-	virtual FPinConnectionResponse CopyPinLinks(UEdGraphPin& CopyFromPin, UEdGraphPin& CopyToPin) const OVERRIDE;
+	virtual FPinConnectionResponse MovePinLinks(UEdGraphPin& MoveFromPin, UEdGraphPin& MoveToPin, bool bIsItermeadiateMove = false) const OVERRIDE;
+	virtual FPinConnectionResponse CopyPinLinks(UEdGraphPin& CopyFromPin, UEdGraphPin& CopyToPin, bool bIsItermeadiateCopy = false) const OVERRIDE;
 	// End of UEdGraphSchema interface
 
 private:
diff --git a/Engine/Source/Editor/ReferenceViewer/Private/ReferenceViewerSchema.cpp b/Engine/Source/Editor/ReferenceViewer/Private/ReferenceViewerSchema.cpp
--- a/Engine/Source/Editor/ReferenceViewer/Private/ReferenceViewerSchema.cpp
+++ b/Engine/Source/Editor/ReferenceViewer/Private/ReferenceViewerSchema.cpp
@@ -43,13 +43,13 @@ void UReferenceViewerSchema::BreakSinglePinLink(UEdGraphPin* SourcePin, UEdGraph
 	// Don't allow breaking any links
 }
 
-FPinConnectionResponse UReferenceViewerSchema::MovePinLinks(UEdGraphPin& MoveFromPin, UEdGraphPin& MoveToPin) const
+FPinConnectionResponse UReferenceViewerSchema::MovePinLinks(UEdGraphPin& MoveFromPin, UEdGraphPin& MoveToPin, bool bIsItermeadiateMove) const
 {
 	// Don't allow moving any links
 	return FPinConnectionResponse(CONNECT_RESPONSE_DISALLOW, TEXT(""));
 }
 
-FPinConnectionResponse UReferenceViewerSchema::CopyPinLinks(UEdGraphPin& CopyFromPin, UEdGraphPin& CopyToPin) const
+FPinConnectionResponse UReferenceViewerSchema::CopyPinLinks(UEdGraphPin& CopyFromPin, UEdGraphPin& CopyToPin, bool bIsItermeadiateCopy) const
 {
 	// Don't allow copying any links
 	return FPinConnectionResponse(CONNECT_RESPONSE_DISALLOW, TEXT(""));
diff --git a/Engine/Source/Runtime/Engine/Classes/EdGraph/EdGraphSchema.h b/Engine/Source/Runtime/Engine/Classes/EdGraph/EdGraphSchema.h
--- a/Engine/Source/Runtime/Engine/Classes/EdGraph/EdGraphSchema.h
+++ b/Engine/Source/Runtime/Engine/Classes/EdGraph/EdGraphSchema.h
@@ -586,10 +586,10 @@ class ENGINE_API UEdGraphSchema : public UObject
 	virtual void BreakSinglePinLink(UEdGraphPin* SourcePin, UEdGraphPin* TargetPin);
 
 	/** Break links on this pin and create links instead on MoveToPin */
-	virtual FPinConnectionResponse MovePinLinks(UEdGraphPin& MoveFromPin, UEdGraphPin& MoveToPin) const;
+	virtual FPinConnectionResponse MovePinLinks(UEdGraphPin& MoveFromPin, UEdGraphPin& MoveToPin, bool bIsIntermeadiateMove = false) const;
 
 	/** Copies pin links from one pin to another without breaking the original links */
-	virtual FPinConnectionResponse CopyPinLinks(UEdGraphPin& CopyFromPin, UEdGraphPin& CopyToPin) const;
+	virtual FPinConnectionResponse CopyPinLinks(UEdGraphPin& CopyFromPin, UEdGraphPin& CopyToPin, bool bIsIntermeadiateCopy = false) const;
 
 	/** Is self pin type? */
 	virtual bool IsSelfPin(const UEdGraphPin& Pin) const   {return false;}
diff --git a/Engine/Source/Runtime/Engine/Private/EdGraph/EdGraph.cpp b/Engine/Source/Runtime/Engine/Private/EdGraph/EdGraph.cpp
--- a/Engine/Source/Runtime/Engine/Private/EdGraph/EdGraph.cpp
+++ b/Engine/Source/Runtime/Engine/Private/EdGraph/EdGraph.cpp
@@ -1198,8 +1198,11 @@ void UEdGraphSchema::BreakSinglePinLink(UEdGraphPin* SourcePin, UEdGraphPin* Tar
 #endif	//#if WITH_EDITOR
 }
 
-FPinConnectionResponse UEdGraphSchema::MovePinLinks(UEdGraphPin& MoveFromPin, UEdGraphPin& MoveToPin) const
+FPinConnectionResponse UEdGraphSchema::MovePinLinks(UEdGraphPin& MoveFromPin, UEdGraphPin& MoveToPin, bool bIsIntermeadiateMove) const
 {
+	ensureMsg(bIsIntermeadiateMove || !MoveToPin.GetOwningNode()->GetGraph()->HasAnyFlags(RF_Transient),
+		TEXT("When moving to an intermeadiate pin, use FKismetCompilerContext::MovePinLinksToIntermediate() instead of UEdGraphSchema::MovePinLinks()"));
+
 	FPinConnectionResponse FinalResponse = FPinConnectionResponse(CONNECT_RESPONSE_MAKE, TEXT(""));
 	// First copy the current set of links
 	TArray<UEdGraphPin*> CurrentLinks = MoveFromPin.LinkedTo;
@@ -1226,8 +1229,11 @@ FPinConnectionResponse UEdGraphSchema::MovePinLinks(UEdGraphPin& MoveFromPin, UE
 	return FinalResponse;
 }
 
-FPinConnectionResponse UEdGraphSchema::CopyPinLinks(UEdGraphPin& CopyFromPin, UEdGraphPin& CopyToPin) const
+FPinConnectionResponse UEdGraphSchema::CopyPinLinks(UEdGraphPin& CopyFromPin, UEdGraphPin& CopyToPin, bool bIsIntermeadiateCopy) const
 {
+	ensureMsg(bIsIntermeadiateCopy || !CopyToPin.GetOwningNode()->GetGraph()->HasAnyFlags(RF_Transient),
+		TEXT("When copying to an intermeadiate pin, use FKismetCompilerContext::CopyPinLinksToIntermediate() instead of UEdGraphSchema::CopyPinLinks()"));
+
 	FPinConnectionResponse FinalResponse = FPinConnectionResponse(CONNECT_RESPONSE_MAKE, TEXT(""));
 	for (int32 i=0; i<CopyFromPin.LinkedTo.Num(); i++)
 	{
