Preventing you from invoking infinite recursion in a BP's construction script (through adding a child component of the same type).

diff --git a/Engine/Source/Editor/BlueprintGraph/Classes/EdGraphSchema_K2.h b/Engine/Source/Editor/BlueprintGraph/Classes/EdGraphSchema_K2.h
--- a/Engine/Source/Editor/BlueprintGraph/Classes/EdGraphSchema_K2.h
+++ b/Engine/Source/Editor/BlueprintGraph/Classes/EdGraphSchema_K2.h
@@ -433,11 +433,7 @@ public:
 	 * @param	TestEdGraph		Graph to test
 	 * @return	true if this is a construction script
 	 */
-	inline bool IsConstructionScript(const UEdGraph* TestEdGraph) const
-	{
-		const FName GraphName = TestEdGraph->GetFName();
-		return (GraphName == FN_UserConstructionScript);
-	}
+	bool IsConstructionScript(const UEdGraph* TestEdGraph) const;
 
 	/**
 	 * Checks to see if the specified graph is a composite graph
diff --git a/Engine/Source/Editor/BlueprintGraph/Private/EdGraphSchema_K2.cpp b/Engine/Source/Editor/BlueprintGraph/Private/EdGraphSchema_K2.cpp
--- a/Engine/Source/Editor/BlueprintGraph/Private/EdGraphSchema_K2.cpp
+++ b/Engine/Source/Editor/BlueprintGraph/Private/EdGraphSchema_K2.cpp
@@ -2996,6 +2996,20 @@ FVector2D UEdGraphSchema_K2::CalculateAveragePositionBetweenNodes(UEdGraphPin* I
 	return (InputCorner + OutputCorner) * 0.5f;
 }
 
+bool UEdGraphSchema_K2::IsConstructionScript(const UEdGraph* TestEdGraph) const
+{
+	TArray<class UK2Node_FunctionEntry*> EntryNodes;
+	TestEdGraph->GetNodesOfClass<UK2Node_FunctionEntry>(EntryNodes);
+
+	bool bIsConstructionScript = false;
+	if (EntryNodes.Num() > 0)
+	{
+		UK2Node_FunctionEntry const* const EntryNode = EntryNodes[0];
+		bIsConstructionScript = (EntryNode->SignatureName == FN_UserConstructionScript);
+	}
+	return bIsConstructionScript;
+}
+
 bool UEdGraphSchema_K2::IsCompositeGraph( const UEdGraph* TestEdGraph ) const
 {
 	check(TestEdGraph);
diff --git a/Engine/Source/Editor/BlueprintGraph/Private/K2Node_AddComponent.cpp b/Engine/Source/Editor/BlueprintGraph/Private/K2Node_AddComponent.cpp
--- a/Engine/Source/Editor/BlueprintGraph/Private/K2Node_AddComponent.cpp
+++ b/Engine/Source/Editor/BlueprintGraph/Private/K2Node_AddComponent.cpp
@@ -155,6 +155,23 @@ void UK2Node_AddComponent::ValidateNodeDuringCompilation(FCompilerResultsLog& Me
 		{
 			MessageLog.Error(*FString::Printf(*NSLOCTEXT("KismetCompiler", "InvalidComponentTemplate_Error", "Invalid class '%s' used as template by '%s' for @@").ToString(), *TemplateClass->GetName(), *GetNodeTitle(ENodeTitleType::FullTitle)), this);
 		}
+
+		if (UChildActorComponent const* ChildActorComponent = Cast<UChildActorComponent const>(Template))
+		{
+			UBlueprint const* Blueprint = GetBlueprint();
+
+			UClass const* ChildActorClass = ChildActorComponent->ChildActorClass;
+			if (ChildActorClass == Blueprint->GeneratedClass)
+			{
+				UEdGraph const* ParentGraph = GetGraph();
+				UEdGraphSchema_K2 const* K2Schema = GetDefault<UEdGraphSchema_K2>();
+
+				if (K2Schema->IsConstructionScript(ParentGraph))
+				{
+					MessageLog.Error(*FString::Printf(*NSLOCTEXT("KismetCompiler", "AddSelfComponent_Error", "@@ cannot add a '%s' component in the construction script (could cause infinite recursion).").ToString(), *ChildActorClass->GetName()), this);
+				}
+			}
+		}
 	}
 	else
 	{
