TTP #331574: Distinguishing interface pins/vars from other objects in blueprints with their own type/color (handling PC_Interface like we handle PC_Object).

diff --git a/Engine/Source/Editor/BlueprintGraph/Classes/EdGraphSchema_K2.h b/Engine/Source/Editor/BlueprintGraph/Classes/EdGraphSchema_K2.h
--- a/Engine/Source/Editor/BlueprintGraph/Classes/EdGraphSchema_K2.h
+++ b/Engine/Source/Editor/BlueprintGraph/Classes/EdGraphSchema_K2.h
@@ -173,13 +173,16 @@ class BLUEPRINTGRAPH_API UEdGraphSchema_K2 : public UEdGraphSchema
 	FString PC_Name;
 
 	UPROPERTY()
-	FString PC_Delegate;    // SubCategoryObject is the UFunction of the delegate signature
+	FString PC_Delegate;	// SubCategoryObject is the UFunction of the delegate signature
 
 	UPROPERTY()
-	FString PC_MCDelegate;  // SubCategoryObject is the UFunction of the delegate signature
+	FString PC_MCDelegate;	// SubCategoryObject is the UFunction of the delegate signature
 
 	UPROPERTY()
-	FString PC_Object;    // SubCategoryObject is the Class of the object passed thru this pin, or SubCategory can be 'self'. The DefaultValue string should always be empty, use DefaultObject.
+	FString PC_Object;		// SubCategoryObject is the Class of the object passed thru this pin, or SubCategory can be 'self'. The DefaultValue string should always be empty, use DefaultObject.
+
+	UPROPERTY()
+	FString PC_Interface;	// SubCategoryObject is the Class of the object passed thru this pin.
 
 	UPROPERTY()
 	FString PC_String;
@@ -191,9 +194,7 @@ class BLUEPRINTGRAPH_API UEdGraphSchema_K2 : public UEdGraphSchema
 	FString PC_Struct;    // SubCategoryObject is the ScriptStruct of the struct passed thru this pin, 'self' is not a valid SubCategory. DefaultObject should always be empty, the DefaultValue string may be used for supported structs.
 
 	UPROPERTY()
-	FString PC_Wildcard;    // Special matching rules are imposed by the node itself
-
-	// PC_Interface - not implemented yet
+	FString PC_Wildcard;  // Special matching rules are imposed by the node itself
 
 	// Common PinType.PinSubCategory values
 	UPROPERTY()
diff --git a/Engine/Source/Editor/BlueprintGraph/Private/CallFunctionHandler.cpp b/Engine/Source/Editor/BlueprintGraph/Private/CallFunctionHandler.cpp
--- a/Engine/Source/Editor/BlueprintGraph/Private/CallFunctionHandler.cpp
+++ b/Engine/Source/Editor/BlueprintGraph/Private/CallFunctionHandler.cpp
@@ -451,12 +451,14 @@ UClass* FKCHandler_CallFunction::GetTrueCallingClass(FKismetFunctionContext& Con
 {
 	if (SelfPin != NULL)
 	{
+		UEdGraphSchema_K2 const* K2Schema = CompilerContext.GetSchema();
+
 		// TODO: here FBlueprintCompiledStatement::GetScopeFromPinType should be called, but since FEdGraphPinType::PinSubCategory is not always initialized properly that function works wrong
 		// return Cast<UClass>(Context.GetScopeFromPinType(SelfPin->PinType, Context.NewClass));
 		FEdGraphPinType& Type = SelfPin->PinType;
-		if ((Type.PinCategory == CompilerContext.GetSchema()->PC_Object) || (Type.PinCategory == CompilerContext.GetSchema()->PC_Class))
+		if ((Type.PinCategory == K2Schema->PC_Object) || (Type.PinCategory == K2Schema->PC_Class) || (Type.PinCategory == K2Schema->PC_Interface))
 		{
-			if (!Type.PinSubCategory.IsEmpty() && (Type.PinSubCategory != CompilerContext.GetSchema()->PSC_Self))
+			if (!Type.PinSubCategory.IsEmpty() && (Type.PinSubCategory != K2Schema->PSC_Self))
 			{
 				return Cast<UClass>(Type.PinSubCategoryObject.Get());
 			}
@@ -519,13 +521,15 @@ void FKCHandler_CallFunction::RegisterNets(FKismetFunctionContext& Context, UEdG
 			UEdGraphPin* Pin = (*It);
 			const bool bIsConnected = (Pin->LinkedTo.Num() != 0);
 
+			UEdGraphSchema_K2 const* K2Schema = CompilerContext.GetSchema();
+
 			// if this pin could use a default (it doesn't have a connection or default of its own)
 			if (!bIsConnected && (Pin->DefaultObject == NULL))
 			{
 				if (DefaultToSelfParamNames.Contains(Pin->PinName))
 				{
 					ensure(Pin->PinType.PinSubCategoryObject != NULL);
-					ensure(Pin->PinType.PinCategory == CompilerContext.GetSchema()->PC_Object);
+					ensure((Pin->PinType.PinCategory == K2Schema->PC_Object) || (Pin->PinType.PinCategory == K2Schema->PC_Interface));
 
 					FBPTerminal* Term = Context.RegisterLiteral(Pin);
 					Term->Type.PinSubCategory = CompilerContext.GetSchema()->PN_Self;
diff --git a/Engine/Source/Editor/BlueprintGraph/Private/EdGraphSchema_K2.cpp b/Engine/Source/Editor/BlueprintGraph/Private/EdGraphSchema_K2.cpp
--- a/Engine/Source/Editor/BlueprintGraph/Private/EdGraphSchema_K2.cpp
+++ b/Engine/Source/Editor/BlueprintGraph/Private/EdGraphSchema_K2.cpp
@@ -144,6 +144,7 @@ UEdGraphSchema_K2::UEdGraphSchema_K2(const class FPostConstructInitializePropert
 	PC_Delegate = TEXT("delegate");
 	PC_MCDelegate = TEXT("mcdelegate");
 	PC_Object = TEXT("object");
+	PC_Interface = TEXT("interface");
 	PC_String = TEXT("string");
 	PC_Text = TEXT("text");
 	PC_Struct = TEXT("struct");
@@ -1344,15 +1345,12 @@ bool UEdGraphSchema_K2::SearchForAutocastFunction(const UEdGraphPin* OutputPin,
 			TargetFunction = TEXT("Conv_TextToString");
 		}
 	}
-	else if ((OutputPin->PinType.PinCategory == PC_Object) && (InputPin->PinType.PinCategory == PC_Object))
+	else if ((OutputPin->PinType.PinCategory == PC_Interface) && (InputPin->PinType.PinCategory == PC_Object))
 	{
-		UClass const* OutputClass = Cast<UClass const>(OutputPin->PinType.PinSubCategoryObject.Get());
-		bool const bOutputIsInterface = ((OutputClass != NULL) && OutputClass->IsChildOf(UInterface::StaticClass()));
-
 		UClass const* InputClass = Cast<UClass const>(InputPin->PinType.PinSubCategoryObject.Get());
-		bool const bInputIsUObject = ((InputClass != NULL) && (InputClass == UObject::StaticClass()));
 
-		if (bOutputIsInterface && bInputIsUObject)
+		bool const bInputIsUObject = ((InputClass != NULL) && (InputClass == UObject::StaticClass()));
+		if (bInputIsUObject)
 		{
 			TargetFunction = TEXT("Conv_InterfaceToObject");
 		}
@@ -1376,7 +1374,8 @@ bool UEdGraphSchema_K2::FindSpecializedConversionNode(const UEdGraphPin* OutputP
 		}
 	}
 	// If connecting an object to a 'call function' self pin, and not currently compatible, see if there is a property we can call a function on
-	else if( InputPin->GetOwningNode()->IsA(UK2Node_CallFunction::StaticClass()) && IsSelfPin(*InputPin) && (OutputPin->PinType.PinCategory == PC_Object) )
+	else if (InputPin->GetOwningNode()->IsA(UK2Node_CallFunction::StaticClass()) && IsSelfPin(*InputPin) && 
+		((OutputPin->PinType.PinCategory == PC_Object) || (OutputPin->PinType.PinCategory == PC_Interface)))
 	{
 		UK2Node_CallFunction* CallFunctionNode = (UK2Node_CallFunction*)(InputPin->GetOwningNode());
 		UClass* OutputPinClass = Cast<UClass>(OutputPin->PinType.PinSubCategoryObject.Get());
@@ -1680,7 +1679,7 @@ bool UEdGraphSchema_K2::DefaultValueSimpleValidation(const FEdGraphPinType& PinT
 			DVSV_RETURN_MSG( FString::Printf(TEXT("Invalid default name for pin %s"), *(PinName)) ); 
 		}
 	}
-	else if (PinCategory == PC_Object)
+	else if ((PinCategory == PC_Object) || (PinCategory == PC_Interface))
 	{
 		if(PinSubCategoryObject == NULL && PinSubCategory != PSC_Self)
 		{
@@ -1810,6 +1809,10 @@ FLinearColor UEdGraphSchema_K2::GetPinTypeColor(const FEdGraphPinType& PinType)
 	{
 		return Options.ObjectPinTypeColor;
 	}
+	else if (TypeString == PC_Interface)
+	{
+		return Options.InterfacePinTypeColor;
+	}
 	else if (TypeString == PC_Float)
 	{
 		return Options.FloatPinTypeColor;
@@ -2076,7 +2079,7 @@ bool UEdGraphSchema_K2::ConvertPropertyToPinType(const UProperty* Property, /*ou
 	}
 	else if (const UInterfaceProperty* InterfaceProperty = Cast<const UInterfaceProperty>(TestProperty))
 	{
-		TypeOut.PinCategory = PC_Object;
+		TypeOut.PinCategory = PC_Interface;
 		TypeOut.PinSubCategoryObject = InterfaceProperty->InterfaceClass;
 	}
 	else if (const UClassProperty* ClassProperty = Cast<const UClassProperty>(TestProperty))
@@ -2271,10 +2274,9 @@ void UEdGraphSchema_K2::GetVariableTypeTree( TArray< TSharedPtr<FPinTypeTreeInfo
 	// Add the types that have subtrees
 	TypeTree.Add( MakeShareable( new FPinTypeTreeInfo(PC_Struct, this, LOCTEXT("StructType", "Struct (value) types.").ToString(), true) ) );
 	TypeTree.Add( MakeShareable( new FPinTypeTreeInfo(PC_Object, this, LOCTEXT("ObjectType", "Object pointer.").ToString(), true) ) );
+	TypeTree.Add( MakeShareable( new FPinTypeTreeInfo(PC_Interface, this, LOCTEXT("InterfaceType", "Interface pointer.").ToString(), true) ) );
 	TypeTree.Add( MakeShareable( new FPinTypeTreeInfo(PC_Class, this, LOCTEXT("ClassType", "Class pointers.").ToString(), true) ) );
 	TypeTree.Add( MakeShareable( new FPinTypeTreeInfo(TEXT("Enum"), PC_Byte, this, LOCTEXT("EnumType", "Enumeration types.").ToString(), true) ) );
-
-
 }
 
 void UEdGraphSchema_K2::GetVariableIndexTypeTree( TArray< TSharedPtr<FPinTypeTreeInfo> >& TypeTree, bool bAllowExec, bool bAllowWildcard ) const
@@ -2303,7 +2305,7 @@ void UEdGraphSchema_K2::GetVariableIndexTypeTree( TArray< TSharedPtr<FPinTypeTre
 
 bool UEdGraphSchema_K2::DoesTypeHaveSubtypes(const FString& FriendlyTypeName) const
 {
-	return (FriendlyTypeName == PC_Struct) || (FriendlyTypeName == PC_Object) || (FriendlyTypeName == PC_Class) || (FriendlyTypeName == TEXT("Enum"));
+	return (FriendlyTypeName == PC_Struct) || (FriendlyTypeName == PC_Object) || (FriendlyTypeName == PC_Interface) || (FriendlyTypeName == PC_Class) || (FriendlyTypeName == TEXT("Enum"));
 }
 
 void UEdGraphSchema_K2::GetVariableSubtypes(const FString& Type, TArray<UObject*>& SubtypesList) const
@@ -2322,7 +2324,7 @@ void UEdGraphSchema_K2::GetVariableSubtypes(const FString& Type, TArray<UObject*
 			}
 		}
 	}
-	else if ((Type == PC_Object) || (Type == PC_Class))
+	else if (Type == PC_Class)
 	{
 		// Generate a list of all potential objects which have "BlueprintType=true" in their metadata
 		for (TObjectIterator<UClass> ClassIt; ClassIt; ++ClassIt)
@@ -2334,6 +2336,30 @@ void UEdGraphSchema_K2::GetVariableSubtypes(const FString& Type, TArray<UObject*
 			}
 		}
 	}
+	else if (Type == PC_Object)
+	{
+		// Generate a list of all potential objects which have "BlueprintType=true" in their metadata (that aren't interfaces)
+		for (TObjectIterator<UClass> ClassIt; ClassIt; ++ClassIt)
+		{
+			UClass* CurrentClass = *ClassIt;
+			if (!CurrentClass->IsChildOf(UInterface::StaticClass()) && UEdGraphSchema_K2::IsAllowableBlueprintVariableType(CurrentClass))
+			{
+				SubtypesList.Add(CurrentClass);
+			}
+		}
+	}
+	else if (Type == PC_Interface)
+	{
+		// Generate a list of all potential objects which have "BlueprintType=true" in their metadata (only ones that are interfaces)
+		for (TObjectIterator<UClass> ClassIt; ClassIt; ++ClassIt)
+		{
+			UClass* CurrentClass = *ClassIt;
+			if (CurrentClass->IsChildOf(UInterface::StaticClass()) && UEdGraphSchema_K2::IsAllowableBlueprintVariableType(CurrentClass))
+			{
+				SubtypesList.Add(CurrentClass);
+			}
+		}
+	}
 	else if (Type == TEXT("Enum"))
 	{
 		// Generate a list of all potential enums which have "BlueprintType=true" in their metadata
@@ -2378,12 +2404,22 @@ bool UEdGraphSchema_K2::ArePinTypesCompatible(const FEdGraphPinType& Output, con
 		{
 			return true;
 		}
+		else if (Output.PinCategory == PC_Interface)
+		{
+			UClass const* OutputClass = Cast<UClass const>(Output.PinSubCategoryObject.Get());
+			check(OutputClass && OutputClass->IsChildOf(UInterface::StaticClass()));
+
+			UClass const* InputClass = Cast<UClass const>(Input.PinSubCategoryObject.Get());
+			check(InputClass && InputClass->IsChildOf(UInterface::StaticClass()));
+			
+			return OutputClass->IsChildOf(InputClass);
+		}
 		else if ((Output.PinCategory == PC_Object) || (Output.PinCategory == PC_Struct) || (Output.PinCategory == PC_Class))
 		{
 			// Subcategory mismatch, but the two could be castable
 			// Only allow a match if the input is a superclass of the output
-			const UStruct* OutputObject = (Output.PinSubCategory == PSC_Self) ? CallingContext : Cast<UStruct>(Output.PinSubCategoryObject.Get());
-			const UStruct* InputObject  = (Input.PinSubCategory == PSC_Self) ? CallingContext : Cast<UStruct>(Input.PinSubCategoryObject.Get());
+			UStruct const* OutputObject = (Output.PinSubCategory == PSC_Self) ? CallingContext : Cast<UStruct>(Output.PinSubCategoryObject.Get());
+			UStruct const* InputObject  = (Input.PinSubCategory == PSC_Self)  ? CallingContext : Cast<UStruct>(Input.PinSubCategoryObject.Get());
 
 			if ((OutputObject != NULL) && (InputObject != NULL))
 			{
@@ -2446,6 +2482,18 @@ bool UEdGraphSchema_K2::ArePinTypesCompatible(const FEdGraphPinType& Output, con
 
 		return true;
 	}
+	else if ((Output.PinCategory == PC_Interface) && (Input.PinCategory == PC_Object))
+	{
+		UClass const* InterfaceClass = Cast<UClass const>(Output.PinSubCategoryObject.Get());
+		UClass const* InputClass     = Cast<UClass const>(Input.PinSubCategoryObject.Get());
+		return InputClass->ImplementsInterface(InterfaceClass) || InterfaceClass->IsChildOf(InputClass);
+	}
+	else if ((Output.PinCategory == PC_Object) && (Input.PinCategory == PC_Interface))
+	{
+		UClass const* OutputClass    = Cast<UClass const>(Output.PinSubCategoryObject.Get());
+		UClass const* InterfaceClass = Cast<UClass const>(Input.PinSubCategoryObject.Get());
+		return OutputClass->ImplementsInterface(InterfaceClass) || OutputClass->IsChildOf(InterfaceClass);
+	}
 
 	// Pins representing BLueprint objects and subclass of UObject can match when EditoronlyBP.bAllowClassObjAndBluepirntPinMatching=true (BaseEngine.ini)
 	// It's required for converting all UBlueprint references into UClass.
@@ -2577,7 +2625,7 @@ void UEdGraphSchema_K2::TrySetDefaultValue(UEdGraphPin& Pin, const FString& NewD
 	UObject* UseDefaultObject = NULL;
 	FText UseDefaultText;
 
-	if ((Pin.PinType.PinCategory == PC_Object) || (Pin.PinType.PinCategory == PC_Class))
+	if ((Pin.PinType.PinCategory == PC_Object) || (Pin.PinType.PinCategory == PC_Class) || (Pin.PinType.PinCategory == PC_Interface))
 	{
 		UseDefaultObject = FindObject<UObject>(ANY_PACKAGE, *NewDefaultValue);
 		UseDefaultValue = NULL;
@@ -2803,6 +2851,11 @@ UFunction* UEdGraphSchema_K2::FindSetVariableByNameFunction(const FEdGraphPinTyp
 		static FName SetObjectName(GET_FUNCTION_NAME_CHECKED(UKismetSystemLibrary, SetObjectPropertyByName));
 		SetFunctionName = SetObjectName;
 	}
+	else if(PinType.PinCategory == K2Schema->PC_Interface)
+	{
+		static FName SetObjectName(GET_FUNCTION_NAME_CHECKED(UKismetSystemLibrary, SetObjectPropertyByName));
+		SetFunctionName = SetObjectName;
+	}
 	else if(PinType.PinCategory == K2Schema->PC_String)
 	{
 		static FName SetStringName(GET_FUNCTION_NAME_CHECKED(UKismetSystemLibrary, SetStringPropertyByName));
@@ -2870,12 +2923,15 @@ bool UEdGraphSchema_K2::CanPromotePinToVariable( const UEdGraphPin& Pin ) const
 		{
 			bCanPromote = false;
 		}
-		else if ((PinType.PinCategory == PC_Object) && (PinType.PinSubCategoryObject != NULL))
+		else if ((PinType.PinCategory == PC_Object) || (PinType.PinCategory == PC_Interface))
 		{
-			if (UClass* Class = Cast<UClass>(PinType.PinSubCategoryObject.Get()))
+			if (PinType.PinSubCategoryObject != NULL)
 			{
-				bCanPromote = UEdGraphSchema_K2::IsAllowableBlueprintVariableType(Class);
-			}	
+				if (UClass* Class = Cast<UClass>(PinType.PinSubCategoryObject.Get()))
+				{
+					bCanPromote = UEdGraphSchema_K2::IsAllowableBlueprintVariableType(Class);
+				}	
+			}
 		}
 		else if ((PinType.PinCategory == PC_Struct) && (PinType.PinSubCategoryObject != NULL))
 		{
@@ -3068,7 +3124,7 @@ void UEdGraphSchema_K2::DroppedAssetsOnNode(const TArray<FAssetData>& Assets, co
 void UEdGraphSchema_K2::DroppedAssetsOnPin(const TArray<FAssetData>& Assets, const FVector2D& GraphPosition, UEdGraphPin* Pin) const
 {
 	// If dropping onto an 'object' pin, try and set the literal
-	if(Pin->PinType.PinCategory == PC_Object)
+	if ((Pin->PinType.PinCategory == PC_Object) || (Pin->PinType.PinCategory == PC_Interface))
 	{
 		UClass* PinClass = Cast<UClass>(Pin->PinType.PinSubCategoryObject.Get());
 		if(PinClass != NULL)
@@ -3096,7 +3152,7 @@ void UEdGraphSchema_K2::GetAssetsPinHoverMessage(const TArray<FAssetData>& Asset
 	OutOkIcon = false;
 
 	// If dropping onto an 'object' pin, try and set the literal
-	if(HoverPin->PinType.PinCategory == PC_Object)
+	if ((HoverPin->PinType.PinCategory == PC_Object) || (HoverPin->PinType.PinCategory == PC_Interface))
 	{
 		UClass* PinClass = Cast<UClass>(HoverPin->PinType.PinSubCategoryObject.Get());
 		if(PinClass != NULL)
diff --git a/Engine/Source/Editor/BlueprintGraph/Private/K2ActionMenuBuilder.cpp b/Engine/Source/Editor/BlueprintGraph/Private/K2ActionMenuBuilder.cpp
--- a/Engine/Source/Editor/BlueprintGraph/Private/K2ActionMenuBuilder.cpp
+++ b/Engine/Source/Editor/BlueprintGraph/Private/K2ActionMenuBuilder.cpp
@@ -1122,7 +1122,7 @@ void FK2ActionMenuBuilder::GetPinAllowedNodeTypes(FBlueprintGraphActionListBuild
 
 			GetAllInterfaceMessageActions(ContextMenuBuilder);		
 		}
-		else if (FromPin.PinType.PinCategory == K2Schema->PC_Object)
+		else if ((FromPin.PinType.PinCategory == K2Schema->PC_Object) || (FromPin.PinType.PinCategory == K2Schema->PC_Interface))
 		{
 			// For Object types, look for functions we can call on that object
 			UClass* PinClass = Cast<UClass>(FromPin.PinType.PinSubCategoryObject.Get());
diff --git a/Engine/Source/Editor/BlueprintGraph/Private/K2Node.cpp b/Engine/Source/Editor/BlueprintGraph/Private/K2Node.cpp
--- a/Engine/Source/Editor/BlueprintGraph/Private/K2Node.cpp
+++ b/Engine/Source/Editor/BlueprintGraph/Private/K2Node.cpp
@@ -252,6 +252,14 @@ void UK2Node::PostReconstructNode()
 						}
 					}
 				}
+				else if (PinCategory == Schema->PC_Object)
+				{
+					UClass const* PinClass = Cast<UClass const>(CurrentPin->PinType.PinSubCategoryObject.Get());
+					if ((PinClass != nullptr) && PinClass->IsChildOf(UInterface::StaticClass()))
+					{
+						CurrentPin->PinType.PinCategory = Schema->PC_Interface;
+					}
+				}
 			}
 		}
 	}
diff --git a/Engine/Source/Editor/BlueprintGraph/Private/K2Node_CallFunction.cpp b/Engine/Source/Editor/BlueprintGraph/Private/K2Node_CallFunction.cpp
--- a/Engine/Source/Editor/BlueprintGraph/Private/K2Node_CallFunction.cpp
+++ b/Engine/Source/Editor/BlueprintGraph/Private/K2Node_CallFunction.cpp
@@ -202,6 +202,10 @@ UEdGraphPin* UK2Node_CallFunction::CreateSelfPin(const UFunction* Function)
 		// This means the function is defined within the blueprint, so the pin should be a true "self" pin
 		SelfPin = CreatePin(EGPD_Input, K2Schema->PC_Object, K2Schema->PSC_Self, NULL, false, false, K2Schema->PN_Self);
 	}
+	else if (FunctionClass->IsChildOf(UInterface::StaticClass()))
+	{
+		SelfPin = CreatePin(EGPD_Input, K2Schema->PC_Interface, TEXT(""), FunctionClass, false, false, K2Schema->PN_Self);
+	}
 	else
 	{
 		// This means that the function is declared in an external class, and should reference that class
diff --git a/Engine/Source/Editor/BlueprintGraph/Private/K2Node_DelegateSet.cpp b/Engine/Source/Editor/BlueprintGraph/Private/K2Node_DelegateSet.cpp
--- a/Engine/Source/Editor/BlueprintGraph/Private/K2Node_DelegateSet.cpp
+++ b/Engine/Source/Editor/BlueprintGraph/Private/K2Node_DelegateSet.cpp
@@ -233,9 +233,11 @@ UK2Node::ERedirectType UK2Node_DelegateSet::DoPinsMatchForReconstruction(const U
 	const UEdGraphSchema_K2* K2Schema = Cast<const UEdGraphSchema_K2>(GetSchema());
 	if ((ERedirectType::ERedirectType_None == OrginalResult) && K2Schema && NewPin && OldPin)
 	{
+		bool const bOldPinIsObj = (OldPin->PinType.PinCategory == K2Schema->PC_Object) || (OldPin->PinType.PinCategory == K2Schema->PC_Interface);
+		bool const bNewPinIsObj = (OldPin->PinType.PinCategory == K2Schema->PC_Object) || (OldPin->PinType.PinCategory == K2Schema->PC_Interface);
+
 		if ((NewPin->Direction == EGPD_Input && OldPin->Direction == EGPD_Input) &&
-			(OldPin->PinType.PinCategory == K2Schema->PC_Object) &&
-			(NewPin->PinType.PinCategory == K2Schema->PC_Object))
+			bOldPinIsObj && bNewPinIsObj)
 		{
 			return ERedirectType_Name;
 		}
diff --git a/Engine/Source/Editor/BlueprintGraph/Private/K2Node_DynamicCast.cpp b/Engine/Source/Editor/BlueprintGraph/Private/K2Node_DynamicCast.cpp
--- a/Engine/Source/Editor/BlueprintGraph/Private/K2Node_DynamicCast.cpp
+++ b/Engine/Source/Editor/BlueprintGraph/Private/K2Node_DynamicCast.cpp
@@ -31,8 +31,14 @@ void UK2Node_DynamicCast::AllocateDefaultPins()
 	if (TargetType != NULL)
 	{
 		FString CastResultPinName = K2Schema->PN_CastedValuePrefix + TargetType->GetName();
-		UEdGraphPin* CastResultPin = CreatePin(EGPD_Output, K2Schema->PC_Object, TEXT(""), UObject::StaticClass(), false, false, CastResultPinName);
-		CastResultPin->PinType.PinSubCategoryObject = *TargetType;
+		if (TargetType->IsChildOf(UInterface::StaticClass()))
+		{
+			CreatePin(EGPD_Output, K2Schema->PC_Interface, TEXT(""), *TargetType, false, false, CastResultPinName);
+		}
+		else 
+		{
+			CreatePin(EGPD_Output, K2Schema->PC_Object, TEXT(""), *TargetType, false, false, CastResultPinName);
+		}
 	}
 
 	Super::AllocateDefaultPins();
diff --git a/Engine/Source/Editor/BlueprintGraph/Private/K2Node_Variable.cpp b/Engine/Source/Editor/BlueprintGraph/Private/K2Node_Variable.cpp
--- a/Engine/Source/Editor/BlueprintGraph/Private/K2Node_Variable.cpp
+++ b/Engine/Source/Editor/BlueprintGraph/Private/K2Node_Variable.cpp
@@ -165,7 +165,8 @@ UK2Node::ERedirectType UK2Node_Variable::DoPinsMatchForReconstruction( const UEd
 	{
 		return ERedirectType_Name;
 	}
-	else if( (OldPin->PinName == NewPin->PinName) && (NewPin->PinType.PinCategory == K2Schema->PC_Object) && (NewPin->PinType.PinSubCategoryObject == NULL))
+	else if ((OldPin->PinName == NewPin->PinName) && ((NewPin->PinType.PinCategory == K2Schema->PC_Object) || 
+		(NewPin->PinType.PinCategory == K2Schema->PC_Interface)) && (NewPin->PinType.PinSubCategoryObject == NULL))
 	{
 		// Special Case:  If we had a pin match, and the class isn't loaded yet because of a cyclic dependency, temporarily cast away the const, and fix up.
 		// @TODO:  Fix this up to be less hacky
diff --git a/Engine/Source/Editor/GraphEditor/Private/NodeFactory.cpp b/Engine/Source/Editor/GraphEditor/Private/NodeFactory.cpp
--- a/Engine/Source/Editor/GraphEditor/Private/NodeFactory.cpp
+++ b/Engine/Source/Editor/GraphEditor/Private/NodeFactory.cpp
@@ -253,6 +253,10 @@ TSharedPtr<SGraphPin> FNodeFactory::CreatePinWidget(UEdGraphPin* InPin)
 		{
 			return SNew(SGraphPinObject, InPin);
 		}
+		else if (InPin->PinType.PinCategory == K2Schema->PC_Interface)
+		{
+			return SNew(SGraphPinObject, InPin);
+		}
 		else if (InPin->PinType.PinCategory == K2Schema->PC_Class)
 		{
 			return SNew(SGraphPinClass, InPin);
diff --git a/Engine/Source/Editor/KismetCompiler/Private/KismetCompiler.cpp b/Engine/Source/Editor/KismetCompiler/Private/KismetCompiler.cpp
--- a/Engine/Source/Editor/KismetCompiler/Private/KismetCompiler.cpp
+++ b/Engine/Source/Editor/KismetCompiler/Private/KismetCompiler.cpp
@@ -261,7 +261,7 @@ void FKismetCompilerContext::ValidatePin(const UEdGraphPin* Pin) const
 	Super::ValidatePin(Pin);
 
 	// Fixing up references to the skel or the generated classes to be PSC_Self pins
-	if (Pin->PinType.PinCategory == Schema->PC_Object)
+	if ((Pin->PinType.PinCategory == Schema->PC_Object) || (Pin->PinType.PinCategory == Schema->PC_Interface))
 	{
 		//@todo:  This is modifying the model, but is acceptable to save another prepass on the pins
 		UEdGraphPin* MutablePin = const_cast<UEdGraphPin*>(Pin);
@@ -1448,7 +1448,8 @@ void FKismetCompilerContext::FinishCompilingFunction(FKismetFunctionContext& Con
 		// No defaults for object/class pins
 		if(	!Schema->IsMetaPin(*EntryPin) && 
 			(EntryPin->PinType.PinCategory != Schema->PC_Object) && 
-			(EntryPin->PinType.PinCategory != Schema->PC_Class) && 
+			(EntryPin->PinType.PinCategory != Schema->PC_Class)  && 
+			(EntryPin->PinType.PinCategory != Schema->PC_Interface) && 
 			!EntryPin->DefaultValue.IsEmpty() )
 		{
 			Function->SetMetaData(*EntryPin->PinName, *EntryPin->DefaultValue);
diff --git a/Engine/Source/Editor/KismetCompiler/Private/KismetCompilerMisc.cpp b/Engine/Source/Editor/KismetCompiler/Private/KismetCompilerMisc.cpp
--- a/Engine/Source/Editor/KismetCompiler/Private/KismetCompilerMisc.cpp
+++ b/Engine/Source/Editor/KismetCompiler/Private/KismetCompilerMisc.cpp
@@ -160,7 +160,7 @@ bool FKismetCompilerUtilities::IsTypeCompatibleWithProperty(UEdGraphPin* SourceP
 			&& PropertyDelegate->SignatureFunction 
 			&& PropertyDelegate->SignatureFunction->IsSignatureCompatibleWith(SignatureFunction));
 	}
-	else if (PinCategory == Schema->PC_Object)
+	else if ((PinCategory == Schema->PC_Object) || (PinCategory == Schema->PC_Interface))
 	{
 		const UClass* ObjectType = (PinSubCategory == Schema->PSC_Self) ? SelfClass : Cast<const UClass>(PinSubCategoryObject);
 
@@ -536,7 +536,7 @@ UProperty* FKismetCompilerUtilities::CreatePropertyOnScope(UStruct* Scope, const
 	}
 
 	//@TODO: Nasty string if-else tree
-	if (Type.PinCategory == Schema->PC_Object)
+	if ((Type.PinCategory == Schema->PC_Object) || (Type.PinCategory == Schema->PC_Interface))
 	{
 		UClass* SubType = (Type.PinSubCategory == Schema->PSC_Self) ? SelfClass : Cast<UClass>(Type.PinSubCategoryObject.Get());
 
diff --git a/Engine/Source/Editor/KismetCompiler/Public/KismetCompiledFunctionContext.h b/Engine/Source/Editor/KismetCompiler/Public/KismetCompiledFunctionContext.h
--- a/Engine/Source/Editor/KismetCompiler/Public/KismetCompiledFunctionContext.h
+++ b/Engine/Source/Editor/KismetCompiler/Public/KismetCompiledFunctionContext.h
@@ -170,7 +170,7 @@ public:
 	/** Returns a UStruct scope corresponding to the pin type passed in, if one exists */
 	UStruct* GetScopeFromPinType(FEdGraphPinType& Type, UClass* SelfClass)
 	{
-		if ((Type.PinCategory == Schema->PC_Object) || (Type.PinCategory == Schema->PC_Class))
+		if ((Type.PinCategory == Schema->PC_Object) || (Type.PinCategory == Schema->PC_Class) || (Type.PinCategory == Schema->PC_Interface))
 		{
 			UClass* SubType = (Type.PinSubCategory == Schema->PSC_Self) ? SelfClass : Cast<UClass>(Type.PinSubCategoryObject.Get());
 			return SubType;
diff --git a/Engine/Source/Editor/UnrealEd/Classes/Editor/EditorUserSettings.h b/Engine/Source/Editor/UnrealEd/Classes/Editor/EditorUserSettings.h
--- a/Engine/Source/Editor/UnrealEd/Classes/Editor/EditorUserSettings.h
+++ b/Engine/Source/Editor/UnrealEd/Classes/Editor/EditorUserSettings.h
@@ -106,6 +106,10 @@ class UEditorUserSettings : public UObject
 	UPROPERTY()
 	FLinearColor ObjectPinTypeColor;
 
+	/** Interface pin type color */
+	UPROPERTY()
+	FLinearColor InterfacePinTypeColor;
+
 	/** String pin type color */
 	UPROPERTY()
 	FLinearColor StringPinTypeColor;
diff --git a/Engine/Source/Editor/UnrealEd/Private/Editor/EditorUserSettings.cpp b/Engine/Source/Editor/UnrealEd/Private/Editor/EditorUserSettings.cpp
--- a/Engine/Source/Editor/UnrealEd/Private/Editor/EditorUserSettings.cpp
+++ b/Engine/Source/Editor/UnrealEd/Private/Editor/EditorUserSettings.cpp
@@ -13,24 +13,25 @@ UEditorUserSettings::UEditorUserSettings(const class FPostConstructInitializePro
 	DataPinStyle = BPST_VariantA;
 
 	// Blueprint editor graph node pin type colors
-	DefaultPinTypeColor = FLinearColor(0.750000f, 0.6f, 0.4f, 1.0f);
-	ExecutionPinTypeColor = FLinearColor(1.0f, 1.0f, 1.0f, 1.0f);
-	BooleanPinTypeColor = FLinearColor(0.300000f, 0.0f, 0.0f, 1.0f);
-	BytePinTypeColor = FLinearColor(0.0f, 0.160000f, 0.131270f, 1.0f);
-	ClassPinTypeColor = FLinearColor(0.1f, 0.0f, 0.5f, 1.0f);
-	IntPinTypeColor = FLinearColor(0.013575f, 0.770000f, 0.429609f, 1.0f);
-	FloatPinTypeColor = FLinearColor(0.357667f, 1.0f, 0.060000f, 1.0f);
-	NamePinTypeColor = FLinearColor(0.607717f, 0.224984f, 1.0f, 1.0f);
-	DelegatePinTypeColor = FLinearColor(1.0f, 0.04f, 0.04f, 1.0f);
-	ObjectPinTypeColor = FLinearColor(0.0f, 0.4f, 0.910000f, 1.0f);
-	StringPinTypeColor = FLinearColor(1.0f, 0.0f, 0.660537f, 1.0f);
-	TextPinTypeColor = FLinearColor(0.8f, 0.2f, 0.4f, 1.0f);
-	StructPinTypeColor = FLinearColor(0.0f, 0.1f, 0.6f, 1.0f);
-	WildcardPinTypeColor = FLinearColor(0.220000f, 0.195800f, 0.195800f, 1.0f);
-	VectorPinTypeColor = FLinearColor(1.0f, 0.591255f, 0.016512f, 1.0f);
-	RotatorPinTypeColor = FLinearColor(0.353393f, 0.454175f, 1.0f, 1.0f);
-	TransformPinTypeColor = FLinearColor(1.0f, 0.172585f, 0.0f, 1.0f);
-	IndexPinTypeColor = FLinearColor(0.013575f, 0.770000f, 0.429609f, 1.0f);
+	DefaultPinTypeColor   = FLinearColor(0.750000f, 0.6f, 0.4f, 1.0f);			// light brown
+	ExecutionPinTypeColor = FLinearColor(1.0f, 1.0f, 1.0f, 1.0f);				// white
+	BooleanPinTypeColor	  = FLinearColor(0.300000f, 0.0f, 0.0f, 1.0f);			// maroon
+	BytePinTypeColor      = FLinearColor(0.0f, 0.160000f, 0.131270f, 1.0f);		// dark green
+	ClassPinTypeColor     = FLinearColor(0.1f, 0.0f, 0.5f, 1.0f);				// deep purple (violet)
+	IntPinTypeColor       = FLinearColor(0.013575f, 0.770000f, 0.429609f, 1.0f);// green-blue
+	FloatPinTypeColor     = FLinearColor(0.357667f, 1.0f, 0.060000f, 1.0f);		// bright green
+	NamePinTypeColor      = FLinearColor(0.607717f, 0.224984f, 1.0f, 1.0f);		// lilac
+	DelegatePinTypeColor  = FLinearColor(1.0f, 0.04f, 0.04f, 1.0f);				// bright red
+	ObjectPinTypeColor    = FLinearColor(0.0f, 0.4f, 0.910000f, 1.0f);			// sharp blue
+	InterfacePinTypeColor = FLinearColor(0.8784f, 1.0f, 0.4f, 1.0f);			// pale green
+	StringPinTypeColor    = FLinearColor(1.0f, 0.0f, 0.660537f, 1.0f);			// bright pink
+	TextPinTypeColor      = FLinearColor(0.8f, 0.2f, 0.4f, 1.0f);				// salmon (light pink)
+	StructPinTypeColor    = FLinearColor(0.0f, 0.1f, 0.6f, 1.0f);				// deep blue
+	WildcardPinTypeColor  = FLinearColor(0.220000f, 0.195800f, 0.195800f, 1.0f);// dark gray
+	VectorPinTypeColor    = FLinearColor(1.0f, 0.591255f, 0.016512f, 1.0f);		// yellow
+	RotatorPinTypeColor   = FLinearColor(0.353393f, 0.454175f, 1.0f, 1.0f);		// periwinkle 
+	TransformPinTypeColor = FLinearColor(1.0f, 0.172585f, 0.0f, 1.0f);			// orange
+	IndexPinTypeColor     = FLinearColor(0.013575f, 0.770000f, 0.429609f, 1.0f);// green-blue
 
 	// Blueprint editor graph node title colors
 	EventNodeTitleColor = FLinearColor(1.f, 0.0f, 0.0f, 1.0f);
diff --git a/Engine/Source/Editor/UnrealEd/Private/Kismet2/BlueprintEditorUtils.cpp b/Engine/Source/Editor/UnrealEd/Private/Kismet2/BlueprintEditorUtils.cpp
--- a/Engine/Source/Editor/UnrealEd/Private/Kismet2/BlueprintEditorUtils.cpp
+++ b/Engine/Source/Editor/UnrealEd/Private/Kismet2/BlueprintEditorUtils.cpp
@@ -2734,7 +2734,7 @@ bool FBlueprintEditorUtils::AddMemberVariable(UBlueprint* Blueprint, const FName
 	{
 		NewVar.PropertyFlags |= CPF_BlueprintAssignable | CPF_BlueprintCallable;
 	}
-	else if (NewVarType.PinCategory == K2Schema->PC_Object)
+	else if ((NewVarType.PinCategory == K2Schema->PC_Object) || (NewVarType.PinCategory == K2Schema->PC_Interface))
 	{
 		// if it's a PC_Object, then it should have an associated UClass object
 		check(NewVarType.PinSubCategoryObject.IsValid());
@@ -2955,7 +2955,7 @@ void FBlueprintEditorUtils::ChangeMemberVariableType(UBlueprint* Blueprint, cons
 				// Destroy all our nodes, because the pin types could be incorrect now (as will the links)
 				RemoveVariableNodes(Blueprint, VariableName);
 
-				if (NewPinType.PinCategory == K2Schema->PC_Object)
+				if ((NewPinType.PinCategory == K2Schema->PC_Object) || (NewPinType.PinCategory == K2Schema->PC_Interface))
 				{
 					// if it's a PC_Object, then it should have an associated UClass object
 					check(NewPinType.PinSubCategoryObject.IsValid());
diff --git a/Engine/Source/Editor/UnrealEd/Private/Tests/BlueprintAutomationTests.cpp b/Engine/Source/Editor/UnrealEd/Private/Tests/BlueprintAutomationTests.cpp
--- a/Engine/Source/Editor/UnrealEd/Private/Tests/BlueprintAutomationTests.cpp
+++ b/Engine/Source/Editor/UnrealEd/Private/Tests/BlueprintAutomationTests.cpp
@@ -611,7 +611,8 @@ public:
 			
 			FEdGraphPinType const& PinType = MalformedPin->PinType;
 			// only object pins would reference the unloaded blueprint
-			if (!PinType.PinSubCategoryObject.IsValid() || ((PinType.PinCategory != K2Schema->PC_Object) && (PinType.PinCategory != K2Schema->PSC_Self)))
+			if (!PinType.PinSubCategoryObject.IsValid() || ((PinType.PinCategory != K2Schema->PC_Object) && 
+				(PinType.PinCategory != K2Schema->PSC_Self) && (PinType.PinCategory != K2Schema->PC_Interface)))
 			{
 				continue;
 			}
