#ue4 ttp331729 Blueprints: Can't access properties of variables of the same type as the Blueprint

SelfContextInfo added to Variable node.

diff --git a/Engine/Source/Editor/BlueprintGraph/Classes/K2Node_Variable.h b/Engine/Source/Editor/BlueprintGraph/Classes/K2Node_Variable.h
--- a/Engine/Source/Editor/BlueprintGraph/Classes/K2Node_Variable.h
+++ b/Engine/Source/Editor/BlueprintGraph/Classes/K2Node_Variable.h
@@ -4,6 +4,16 @@
 #pragma once
 #include "K2Node_Variable.generated.h"
 
+UENUM()
+namespace ESelfContextInfo
+{
+	enum Type
+	{
+		Unspecified,
+		NotSelfContext,
+	};
+}
+
 UCLASS(MinimalAPI, abstract)
 class UK2Node_Variable : public UK2Node
 {
@@ -13,6 +23,9 @@ class UK2Node_Variable : public UK2Node
 	UPROPERTY()
 	FMemberReference	VariableReference;
 
+	UPROPERTY()
+	TEnumAsByte<ESelfContextInfo::Type> SelfContextInfo;
+
 protected:
 	/** Class that this variable is defined in. Should be NULL if bSelfContext is true.  */
 	UPROPERTY()
diff --git a/Engine/Source/Editor/BlueprintGraph/Private/K2Node_Variable.cpp b/Engine/Source/Editor/BlueprintGraph/Private/K2Node_Variable.cpp
--- a/Engine/Source/Editor/BlueprintGraph/Private/K2Node_Variable.cpp
+++ b/Engine/Source/Editor/BlueprintGraph/Private/K2Node_Variable.cpp
@@ -51,6 +51,7 @@ void UK2Node_Variable::Serialize(FArchive& Ar)
 
 void UK2Node_Variable::SetFromProperty(const UProperty* Property, bool bSelfContext)
 {
+	SelfContextInfo = bSelfContext ? ESelfContextInfo::Unspecified : ESelfContextInfo::NotSelfContext;
 	VariableReference.SetFromField<UProperty>(Property, bSelfContext);
 }
 
@@ -82,14 +83,16 @@ void UK2Node_Variable::CreatePinForSelf()
 	// Create the self pin
 	if (!K2Schema->FindSelfPin(*this, EGPD_Input))
 	{
-		if (VariableReference.IsSelfContext())
+		if (VariableReference.IsSelfContext() && (ESelfContextInfo::NotSelfContext != SelfContextInfo))
 		{
 			UEdGraphPin* SelfPin = CreatePin(EGPD_Input, K2Schema->PC_Object, K2Schema->PSC_Self, NULL, false, false, K2Schema->PN_Self);
 			SelfPin->bHidden = true; // don't show in 'self' context
 		}
 		else
 		{
-			UEdGraphPin* TargetPin = CreatePin(EGPD_Input, K2Schema->PC_Object, TEXT(""), VariableReference.GetMemberParentClass(this), false, false, K2Schema->PN_Self);
+			UClass* MemberParentClass = VariableReference.GetMemberParentClass(this);
+			UClass* AuthoritativeClass = MemberParentClass ? MemberParentClass->GetAuthoritativeClass() : NULL;
+			UEdGraphPin* TargetPin = CreatePin(EGPD_Input, K2Schema->PC_Object, TEXT(""), AuthoritativeClass, false, false, K2Schema->PN_Self);
 			TargetPin->PinFriendlyName =  TEXT("Target");
 		}
 	}
