#ue4 When self and variable has the same type, the K2Node_Variable can become corrupted.

https://forums.unrealengine.com/showthread.php?2203-Setting-and-Getting-in-same-flow-causes-Blueprint-bug

diff --git a/Engine/Source/Editor/BlueprintGraph/Private/K2Node_Variable.cpp b/Engine/Source/Editor/BlueprintGraph/Private/K2Node_Variable.cpp
--- a/Engine/Source/Editor/BlueprintGraph/Private/K2Node_Variable.cpp
+++ b/Engine/Source/Editor/BlueprintGraph/Private/K2Node_Variable.cpp
@@ -164,36 +164,42 @@ UK2Node::ERedirectType UK2Node_Variable::DoPinsMatchForReconstruction( const UEd
 	{
 		return Super::DoPinsMatchForReconstruction(NewPin, NewPinIndex, OldPin, OldPinIndex);
 	}
-	else if (K2Schema->ArePinTypesCompatible( NewPin->PinType, OldPin->PinType))
-	{
-		return ERedirectType_Name;
-	}
-	else if ((OldPin->PinName == NewPin->PinName) && ((NewPin->PinType.PinCategory == K2Schema->PC_Object) || 
-		(NewPin->PinType.PinCategory == K2Schema->PC_Interface)) && (NewPin->PinType.PinSubCategoryObject == NULL))
+
+	const bool bCanMatchSelfs = ((OldPin->PinName == K2Schema->PN_Self) == (NewPin->PinName == K2Schema->PN_Self));
+	const bool bTheSameDirection = (NewPin->Direction == OldPin->Direction);
+	if (bCanMatchSelfs && bTheSameDirection)
 	{
-		// Special Case:  If we had a pin match, and the class isn't loaded yet because of a cyclic dependency, temporarily cast away the const, and fix up.
-		// @TODO:  Fix this up to be less hacky
-		UBlueprintGeneratedClass* TypeClass = Cast<UBlueprintGeneratedClass>(OldPin->PinType.PinSubCategoryObject.Get());
-		if(TypeClass && TypeClass->ClassGeneratedBy && TypeClass->ClassGeneratedBy->HasAnyFlags(RF_BeingRegenerated))
+		if (K2Schema->ArePinTypesCompatible(NewPin->PinType, OldPin->PinType))
 		{
-			UEdGraphPin* NonConstNewPin = (UEdGraphPin*)NewPin;
-			NonConstNewPin->PinType.PinSubCategoryObject = OldPin->PinType.PinSubCategoryObject.Get();
 			return ERedirectType_Name;
 		}
+		else if ((OldPin->PinName == NewPin->PinName) && ((NewPin->PinType.PinCategory == K2Schema->PC_Object) ||
+			(NewPin->PinType.PinCategory == K2Schema->PC_Interface)) && (NewPin->PinType.PinSubCategoryObject == NULL))
+		{
+			// Special Case:  If we had a pin match, and the class isn't loaded yet because of a cyclic dependency, temporarily cast away the const, and fix up.
+			// @TODO:  Fix this up to be less hacky
+			UBlueprintGeneratedClass* TypeClass = Cast<UBlueprintGeneratedClass>(OldPin->PinType.PinSubCategoryObject.Get());
+			if (TypeClass && TypeClass->ClassGeneratedBy && TypeClass->ClassGeneratedBy->HasAnyFlags(RF_BeingRegenerated))
+			{
+				UEdGraphPin* NonConstNewPin = (UEdGraphPin*)NewPin;
+				NonConstNewPin->PinType.PinSubCategoryObject = OldPin->PinType.PinSubCategoryObject.Get();
+				return ERedirectType_Name;
+			}
+		}
+		else
+		{
+			// Special Case:  If we're migrating from old blueprint references to class references, allow pins to be reconnected if coerced
+			const UClass* PSCOClass = Cast<UClass>(OldPin->PinType.PinSubCategoryObject.Get());
+			const bool bOldIsBlueprint = PSCOClass && PSCOClass->IsChildOf(UBlueprint::StaticClass());
+			const bool bNewIsClass = (NewPin->PinType.PinCategory == K2Schema->PC_Class);
+			if (bNewIsClass && bOldIsBlueprint)
+			{
+				UEdGraphPin* OldPinNonConst = (UEdGraphPin*)OldPin;
+				OldPinNonConst->PinName = NewPin->PinName;
+				return ERedirectType_Name;
+			}
+		}
 	}
- 	else
- 	{
-		// Special Case:  If we're migrating from old blueprint references to class references, allow pins to be reconnected if coerced
-		const UClass* PSCOClass = Cast<UClass>(OldPin->PinType.PinSubCategoryObject.Get());
- 		const bool bOldIsBlueprint = PSCOClass && PSCOClass->IsChildOf(UBlueprint::StaticClass());
- 		const bool bNewIsClass = (NewPin->PinType.PinCategory == K2Schema->PC_Class);
- 		if(bNewIsClass && bOldIsBlueprint)
- 		{
- 			UEdGraphPin* OldPinNonConst = (UEdGraphPin*)OldPin;
- 			OldPinNonConst->PinName = NewPin->PinName;
-			return ERedirectType_Name;
- 		}
- 	}
 
 	return ERedirectType_None;
 }
