Remark InputAxis(Key)Events as internal so that they cannot be called.
TTP #332148

diff --git a/Engine/Source/Editor/BlueprintGraph/Private/K2Node_InputAxisEvent.cpp b/Engine/Source/Editor/BlueprintGraph/Private/K2Node_InputAxisEvent.cpp
--- a/Engine/Source/Editor/BlueprintGraph/Private/K2Node_InputAxisEvent.cpp
+++ b/Engine/Source/Editor/BlueprintGraph/Private/K2Node_InputAxisEvent.cpp
@@ -9,6 +9,7 @@ UK2Node_InputAxisEvent::UK2Node_InputAxisEvent(const class FPostConstructInitial
 {
 	bConsumeInput = true;
 	bOverrideParentBinding = true;
+	bInternalEvent = true;
 
 	EventSignatureName = TEXT("InputAxisHandlerDynamicSignature__DelegateSignature");
 	EventSignatureClass = UInputComponent::StaticClass();
diff --git a/Engine/Source/Editor/BlueprintGraph/Private/K2Node_InputAxisKeyEvent.cpp b/Engine/Source/Editor/BlueprintGraph/Private/K2Node_InputAxisKeyEvent.cpp
--- a/Engine/Source/Editor/BlueprintGraph/Private/K2Node_InputAxisKeyEvent.cpp
+++ b/Engine/Source/Editor/BlueprintGraph/Private/K2Node_InputAxisKeyEvent.cpp
@@ -9,6 +9,7 @@ UK2Node_InputAxisKeyEvent::UK2Node_InputAxisKeyEvent(const class FPostConstructI
 {
 	bConsumeInput = true;
 	bOverrideParentBinding = true;
+	bInternalEvent = true;
 
 	EventSignatureName = TEXT("InputAxisHandlerDynamicSignature__DelegateSignature");
 	EventSignatureClass = UInputComponent::StaticClass();
