Expose Get Axis Key Value node to complement Axis Key Events and have parity with Mapped Axis Events/Get Value

diff --git a/Engine/Source/Editor/BlueprintGraph/Private/K2ActionMenuBuilder.cpp b/Engine/Source/Editor/BlueprintGraph/Private/K2ActionMenuBuilder.cpp
--- a/Engine/Source/Editor/BlueprintGraph/Private/K2ActionMenuBuilder.cpp
+++ b/Engine/Source/Editor/BlueprintGraph/Private/K2ActionMenuBuilder.cpp
@@ -31,15 +31,21 @@ namespace K2ActionCategories
 	static FString const AxisValuesSubCategory   = LOCTEXT("ActionValuesCategory", "Axis Values").ToString();
 	static FString const AxisEventsSubCategory   = LOCTEXT("AxisEventsCategory",   "Axis Events").ToString();
 	static FString const KeyEventsSubCategory    = LOCTEXT("KeyEventsCategory",    "Key Events").ToString();
+	static FString const KeyValuesSubCategory    = LOCTEXT("KeyValuesCategory", "Key Values").ToString();
 	static FString const GamepadEventsSubCategory = LOCTEXT("GamepadEventsCategory", "Gamepad Events").ToString();
+	static FString const GamepadValuesSubCategory = LOCTEXT("GamepadValuesCategory", "Gamepad Values").ToString();
 	static FString const MouseEventsSubCategory = LOCTEXT("MouseEventsCategory", "Mouse Events").ToString();
+	static FString const MouseValuesSubCategory = LOCTEXT("MouseValuesCategory", "Mouse Values").ToString();
 
 	static FString const ActionEventsCategory = InputCategory + SubCategoryDelim + ActionEventsSubCategory;
 	static FString const AxisCategory         = InputCategory + SubCategoryDelim + AxisValuesSubCategory;
 	static FString const AxisEventCategory    = InputCategory + SubCategoryDelim + AxisEventsSubCategory;
 	static FString const KeyEventsCategory    = InputCategory + SubCategoryDelim + KeyEventsSubCategory;
+	static FString const KeyValuesCategory    = InputCategory + SubCategoryDelim + KeyValuesSubCategory;
 	static FString const GamepadEventsCategory = InputCategory + SubCategoryDelim + GamepadEventsSubCategory;
+	static FString const GamepadValuesCategory = InputCategory + SubCategoryDelim + GamepadValuesSubCategory;
 	static FString const MouseEventsCategory = InputCategory + SubCategoryDelim + MouseEventsSubCategory;
+	static FString const MouseValuesCategory = InputCategory + SubCategoryDelim + MouseValuesSubCategory;
 	static FString const TouchEventsCategory = InputCategory;
 
 
@@ -149,22 +155,10 @@ static void GetInputNodes(FGraphActionListBuilderBase& ActionMenuBuilder, const
 
 				const bool bGamepad = Key.IsGamepadKey();
 				const bool bMouse = Key.IsMouseButton();
-				const FString& KeyCategory = (bGamepad ? K2ActionCategories::GamepadEventsCategory : (bMouse ? K2ActionCategories::MouseEventsCategory : K2ActionCategories::KeyEventsCategory));
+				const FString& KeyEventCategory = (bGamepad ? K2ActionCategories::GamepadEventsCategory : (bMouse ? K2ActionCategories::MouseEventsCategory : K2ActionCategories::KeyEventsCategory));
 
-				TSharedPtr<FEdGraphSchemaAction_K2NewNode> Action = FK2ActionMenuBuilder::AddNewNodeAction(ActionMenuBuilder, KeyCategory, KeyName.ToString(), KeyName.ToString());
-
-				if (Key.IsAxis())
-				{
-					UK2Node_InputAxisKeyEvent* InputKeyAxisNode = ActionMenuBuilder.CreateTemplateNode<UK2Node_InputAxisKeyEvent>();
-					InputKeyAxisNode->Initialize(Key);
-					Action->NodeTemplate = InputKeyAxisNode;
-				}
-				else
-				{
-					UK2Node_InputKey* InputKeyNode = ActionMenuBuilder.CreateTemplateNode<UK2Node_InputKey>();
-					InputKeyNode->InputKey = Key;
-					Action->NodeTemplate = InputKeyNode;
-				}
+				TSharedPtr<FEdGraphSchemaAction_K2NewNode> Action = FK2ActionMenuBuilder::AddNewNodeAction(ActionMenuBuilder, KeyEventCategory, KeyName.ToString(), KeyName.ToString());
+				TSharedPtr<FEdGraphSchemaAction_K2NewNode> GetAction;
 
 				if (bGamepad)
 				{
@@ -178,6 +172,26 @@ static void GetInputNodes(FGraphActionListBuilderBase& ActionMenuBuilder, const
 				{
 					Action->Keywords = TEXT("Input Key InputKey Keyboard");
 				}
+
+				if (Key.IsAxis())
+				{
+					UK2Node_InputAxisKeyEvent* InputKeyAxisNode = ActionMenuBuilder.CreateTemplateNode<UK2Node_InputAxisKeyEvent>();
+					InputKeyAxisNode->Initialize(Key);
+					Action->NodeTemplate = InputKeyAxisNode;
+
+					const FString& KeyValuesCategory = (bGamepad ? K2ActionCategories::GamepadValuesCategory : (bMouse ? K2ActionCategories::MouseValuesCategory : K2ActionCategories::KeyValuesCategory));
+					GetAction = FK2ActionMenuBuilder::AddNewNodeAction(ActionMenuBuilder, KeyValuesCategory, KeyName.ToString(), KeyName.ToString());
+					UK2Node_GetInputAxisKeyValue* GetInputAxisKeyValue = ActionMenuBuilder.CreateTemplateNode<UK2Node_GetInputAxisKeyValue>();
+					GetInputAxisKeyValue->Initialize(Key);
+					GetAction->NodeTemplate = GetInputAxisKeyValue;
+					GetAction->Keywords = Action->Keywords;
+				}
+				else
+				{
+					UK2Node_InputKey* InputKeyNode = ActionMenuBuilder.CreateTemplateNode<UK2Node_InputKey>();
+					InputKeyNode->InputKey = Key;
+					Action->NodeTemplate = InputKeyNode;
+				}
 			}
 		}
 
diff --git a/Engine/Source/Editor/BlueprintGraph/Private/K2Node_GetInputAxisValue.cpp b/Engine/Source/Editor/BlueprintGraph/Private/K2Node_GetInputAxisValue.cpp
--- a/Engine/Source/Editor/BlueprintGraph/Private/K2Node_GetInputAxisValue.cpp
+++ b/Engine/Source/Editor/BlueprintGraph/Private/K2Node_GetInputAxisValue.cpp
@@ -26,7 +26,7 @@ void UK2Node_GetInputAxisValue::Initialize(const FName AxisName)
 
 FString UK2Node_GetInputAxisValue::GetNodeTitle(ENodeTitleType::Type TitleType) const
 {
-	return FString::Printf(*NSLOCTEXT("K2Node", "GetInputAxis_Name", "Get InputAxis %s").ToString(), *InputAxisName.ToString());
+	return FString::Printf(*NSLOCTEXT("K2Node", "GetInputAxis_Name", "Get %s").ToString(), *InputAxisName.ToString());
 }
 
 FString UK2Node_GetInputAxisValue::GetTooltip() const
diff --git a/Engine/Source/Editor/BlueprintGraph/Private/K2Node_InputAxisKeyEvent.cpp b/Engine/Source/Editor/BlueprintGraph/Private/K2Node_InputAxisKeyEvent.cpp
--- a/Engine/Source/Editor/BlueprintGraph/Private/K2Node_InputAxisKeyEvent.cpp
+++ b/Engine/Source/Editor/BlueprintGraph/Private/K2Node_InputAxisKeyEvent.cpp
@@ -37,15 +37,15 @@ void UK2Node_InputAxisKeyEvent::ValidateNodeDuringCompilation(class FCompilerRes
 
 	if (!AxisKey.IsValid())
 	{
-		MessageLog.Warning(*FText::Format(NSLOCTEXT("KismetCompiler", "Invalid_InputKey_Warning", "InputKey Event specifies invalid FKey'{0}' for @@"), FText::FromString(AxisKey.ToString())).ToString(), this);
+		MessageLog.Warning(*FText::Format(NSLOCTEXT("KismetCompiler", "Invalid_InputAxisKey_Warning", "InputAxisKey Event specifies invalid FKey'{0}' for @@"), FText::FromString(AxisKey.ToString())).ToString(), this);
 	}
 	else if (!AxisKey.IsAxis())
 	{
-		MessageLog.Warning(*FText::Format(NSLOCTEXT("KismetCompiler", "NotAxis_InputKey_Warning", "InputKey Event specifies non-axis FKey'{0}' for @@"), FText::FromString(AxisKey.ToString())).ToString(), this);
+		MessageLog.Warning(*FText::Format(NSLOCTEXT("KismetCompiler", "NotAxis_InputAxisKey_Warning", "InputAxisKey Event specifies non-axis FKey'{0}' for @@"), FText::FromString(AxisKey.ToString())).ToString(), this);
 	}
 	else if (!AxisKey.IsBindableInBlueprints())
 	{
-		MessageLog.Warning(*FText::Format(NSLOCTEXT("KismetCompiler", "NotBindanble_InputKey_Warning", "InputKey Event specifies FKey'{0}' that is not blueprint bindable for @@"), FText::FromString(AxisKey.ToString())).ToString(), this);
+		MessageLog.Warning(*FText::Format(NSLOCTEXT("KismetCompiler", "NotBindable_InputAxisKey_Warning", "InputAxisKey Event specifies FKey'{0}' that is not blueprint bindable for @@"), FText::FromString(AxisKey.ToString())).ToString(), this);
 	}
 }
 
diff --git a/Engine/Source/Runtime/Engine/Classes/GameFramework/Actor.h b/Engine/Source/Runtime/Engine/Classes/GameFramework/Actor.h
--- a/Engine/Source/Runtime/Engine/Classes/GameFramework/Actor.h
+++ b/Engine/Source/Runtime/Engine/Classes/GameFramework/Actor.h
@@ -570,6 +570,9 @@ public:
 	UFUNCTION(BlueprintCallable, meta=(BlueprintInternalUseOnly="true", BlueprintProtected = "true", HidePin="InputAxisName", ToolTip="Gets the value of the input axis if input is enabled for this actor"))
 	float GetInputAxisValue(const FName InputAxisName) const;
 
+	UFUNCTION(BlueprintCallable, meta=(BlueprintInternalUseOnly="true", BlueprintProtected = "true", HidePin="InputAxisKey", ToolTip="Gets the value of the input axis key if input is enabled for this actor"))
+	float GetInputAxisKeyValue(const FKey InputAxisKey) const;
+
 	/** Returns the instigator for this actor, or NULL if there is none. */
 	UFUNCTION(BlueprintCallable, meta=(BlueprintProtected = "true"), Category="Game|Damage")
 	APawn* GetInstigator() const;
diff --git a/Engine/Source/Runtime/Engine/Private/Actor.cpp b/Engine/Source/Runtime/Engine/Private/Actor.cpp
--- a/Engine/Source/Runtime/Engine/Private/Actor.cpp
+++ b/Engine/Source/Runtime/Engine/Private/Actor.cpp
@@ -2253,6 +2253,17 @@ float AActor::GetInputAxisValue(const FName InputAxisName) const
 	return Value;
 }
 
+float AActor::GetInputAxisKeyValue(const FKey InputAxisKey) const
+{
+	float Value = 0.f;
+
+	if (InputComponent)
+	{
+		Value = InputComponent->GetAxisKeyValue(InputAxisKey);
+	}
+
+	return Value;
+}
 
 
 bool AActor::SetActorLocation(const FVector& NewLocation, bool bSweep)
