Add missed files for CL# 2044415

diff --git a/Engine/Source/Editor/BlueprintGraph/Classes/K2Node_GetInputAxisKeyValue.h b/Engine/Source/Editor/BlueprintGraph/Classes/K2Node_GetInputAxisKeyValue.h
new file mode 100644
--- /dev/null
+++ b/Engine/Source/Editor/BlueprintGraph/Classes/K2Node_GetInputAxisKeyValue.h
@@ -0,0 +1,37 @@
+// Copyright 1998-2014 Epic Games, Inc. All Rights Reserved.
+
+#pragma once
+#include "K2Node_GetInputAxisKeyValue.generated.h"
+
+UCLASS(MinimalAPI)
+class UK2Node_GetInputAxisKeyValue : public UK2Node_CallFunction
+{
+	GENERATED_UCLASS_BODY()
+
+	UPROPERTY()
+	FKey InputAxisKey;
+
+	// Prevents actors with lower priority from handling this input
+	UPROPERTY(EditAnywhere, Category="Input")
+	uint32 bConsumeInput:1;
+
+	// Should the binding gather input even when the game is paused
+	UPROPERTY(EditAnywhere, Category="Input")
+	uint32 bExecuteWhenPaused:1;
+
+	// Begin EdGraphNode interface
+	virtual void AllocateDefaultPins() OVERRIDE;
+	virtual FString GetNodeTitle(ENodeTitleType::Type TitleType) const OVERRIDE;
+	virtual FString GetTooltip() const OVERRIDE;
+	virtual FName GetPaletteIcon(FLinearColor& OutColor) const OVERRIDE;
+	// End EdGraphNode interface
+
+	// Begin UK2Node interface
+	virtual void ValidateNodeDuringCompilation(class FCompilerResultsLog& MessageLog) const OVERRIDE;
+	virtual bool ShouldShowNodeProperties() const OVERRIDE { return true; }
+	virtual UClass* GetDynamicBindingClass() const OVERRIDE;
+	virtual void RegisterDynamicBinding(UDynamicBlueprintBinding* BindingObject) const OVERRIDE;
+	// End UK2Node interface
+	
+	void Initialize(const FKey AxisKey);
+};
\ No newline at end of file
diff --git a/Engine/Source/Editor/BlueprintGraph/Private/K2Node_GetInputAxisKeyValue.cpp b/Engine/Source/Editor/BlueprintGraph/Private/K2Node_GetInputAxisKeyValue.cpp
new file mode 100644
--- /dev/null
+++ b/Engine/Source/Editor/BlueprintGraph/Private/K2Node_GetInputAxisKeyValue.cpp
@@ -0,0 +1,86 @@
+// Copyright 1998-2014 Epic Games, Inc. All Rights Reserved.
+
+#include "BlueprintGraphPrivatePCH.h"
+
+#include "CompilerResultsLog.h"
+
+UK2Node_GetInputAxisKeyValue::UK2Node_GetInputAxisKeyValue(const class FPostConstructInitializeProperties& PCIP)
+	: Super(PCIP)
+{
+	bConsumeInput = true;
+}
+
+void UK2Node_GetInputAxisKeyValue::AllocateDefaultPins()
+{
+	Super::AllocateDefaultPins();
+
+	UEdGraphPin* InputAxisKeyPin = FindPinChecked(TEXT("InputAxisKey"));
+	InputAxisKeyPin->DefaultValue = InputAxisKey.ToString();
+}
+
+void UK2Node_GetInputAxisKeyValue::Initialize(const FKey AxisKey)
+{
+	InputAxisKey = AxisKey;
+	SetFromFunction(AActor::StaticClass()->FindFunctionByName(TEXT("GetInputAxisKeyValue")));
+}
+
+FString UK2Node_GetInputAxisKeyValue::GetNodeTitle(ENodeTitleType::Type TitleType) const
+{
+	return FString::Printf(*NSLOCTEXT("K2Node", "GetInputAxisKey_Name", "Get %s").ToString(), *InputAxisKey.GetDisplayName().ToString());
+}
+
+FString UK2Node_GetInputAxisKeyValue::GetTooltip() const
+{
+	return FString::Printf(*NSLOCTEXT("K2Node", "GetInputAxisKey_Tooltip", "Returns the current value of input axis key %s.  If input is disabled for the actor the value will be 0.").ToString(), *InputAxisKey.GetDisplayName().ToString());
+}
+
+void UK2Node_GetInputAxisKeyValue::ValidateNodeDuringCompilation(class FCompilerResultsLog& MessageLog) const
+{
+	Super::ValidateNodeDuringCompilation(MessageLog);
+
+	if (!InputAxisKey.IsValid())
+	{
+		MessageLog.Warning(*FText::Format(NSLOCTEXT("KismetCompiler", "Invalid_GetInputAxisKey_Warning", "GetInputAxisKey Value specifies invalid FKey'{0}' for @@"), FText::FromString(InputAxisKey.ToString())).ToString(), this);
+	}
+	else if (!InputAxisKey.IsAxis())
+	{
+		MessageLog.Warning(*FText::Format(NSLOCTEXT("KismetCompiler", "NotAxis_GetInputAxisKey_Warning", "GetInputAxisKey Value specifies non-axis FKey'{0}' for @@"), FText::FromString(InputAxisKey.ToString())).ToString(), this);
+	}
+	else if (!InputAxisKey.IsBindableInBlueprints())
+	{
+		MessageLog.Warning(*FText::Format(NSLOCTEXT("KismetCompiler", "NotBindanble_GetInputAxisKey_Warning", "GetInputAxisKey Value specifies FKey'{0}' that is not blueprint bindable for @@"), FText::FromString(InputAxisKey.ToString())).ToString(), this);
+	}
+}
+
+UClass* UK2Node_GetInputAxisKeyValue::GetDynamicBindingClass() const
+{
+	return UInputAxisKeyDelegateBinding::StaticClass();
+}
+
+void UK2Node_GetInputAxisKeyValue::RegisterDynamicBinding(UDynamicBlueprintBinding* BindingObject) const
+{
+	UInputAxisKeyDelegateBinding* InputAxisKeyBindingObject = CastChecked<UInputAxisKeyDelegateBinding>(BindingObject);
+
+	FBlueprintInputAxisKeyDelegateBinding Binding;
+	Binding.AxisKey = InputAxisKey;
+	Binding.bConsumeInput = bConsumeInput;
+	Binding.bExecuteWhenPaused = bExecuteWhenPaused;
+
+	InputAxisKeyBindingObject->InputAxisKeyDelegateBindings.Add(Binding);
+}
+
+FName UK2Node_GetInputAxisKeyValue::GetPaletteIcon(FLinearColor& OutColor) const
+{
+	if (InputAxisKey.IsMouseButton())
+	{
+		return TEXT("GraphEditor.MouseEvent_16x");
+	}
+	else if (InputAxisKey.IsGamepadKey())
+	{
+		return TEXT("GraphEditor.PadEvent_16x");
+	}
+	else
+	{
+		return TEXT("GraphEditor.KeyEvent_16x");
+	}
+}
\ No newline at end of file
