#summary Can search for nodes in the palette menu in Blueprints with both the native name and the localized name.

diff --git a/Engine/Source/Editor/AnimGraph/Private/AnimationCustomTransitionSchema.cpp b/Engine/Source/Editor/AnimGraph/Private/AnimationCustomTransitionSchema.cpp
--- a/Engine/Source/Editor/AnimGraph/Private/AnimationCustomTransitionSchema.cpp
+++ b/Engine/Source/Editor/AnimGraph/Private/AnimationCustomTransitionSchema.cpp
@@ -45,10 +45,12 @@ void UAnimationCustomTransitionSchema::CreateDefaultNodesForGraph(UEdGraph& Grap
 
 void UAnimationCustomTransitionSchema::GetGraphDisplayInformation(const UEdGraph& Graph, /*out*/ FGraphDisplayInfo& DisplayInfo) const
 {
-	DisplayInfo.DisplayName = FText::FromString( Graph.GetName() );
+	DisplayInfo.PlainName = FText::FromString( Graph.GetName() );
 
 	if (const UAnimStateTransitionNode* TransNode = Cast<const UAnimStateTransitionNode>(Graph.GetOuter()))
 	{
-		DisplayInfo.DisplayName = FText::Format( NSLOCTEXT("Animation", "CustomBlendGraphTitle", "{0} (custom blend)"), FText::FromString( TransNode->GetNodeTitle(ENodeTitleType::FullTitle)) );
+		DisplayInfo.PlainName = FText::Format( NSLOCTEXT("Animation", "CustomBlendGraphTitle", "{0} (custom blend)"), TransNode->GetNodeTitle(ENodeTitleType::FullTitle) );
 	}
+
+	DisplayInfo.DisplayName = DisplayInfo.PlainName;
 }
diff --git a/Engine/Source/Editor/AnimGraph/Private/AnimationGraphSchema.cpp b/Engine/Source/Editor/AnimGraph/Private/AnimationGraphSchema.cpp
--- a/Engine/Source/Editor/AnimGraph/Private/AnimationGraphSchema.cpp
+++ b/Engine/Source/Editor/AnimGraph/Private/AnimationGraphSchema.cpp
@@ -50,10 +50,11 @@ void UAnimationGraphSchema::GetStateMachineMenuItems(FGraphContextMenuBuilder& C
 			NodeTemplate->GetterType = ETransitionGetter::ArbitraryState_GetBlendWeight;
 
 			const FString OwnerName = StateNode->GetOuter()->GetName();
-			const FString Title = FString::Printf(TEXT("Current %s for state '%s.%s'"), *UK2Node_TransitionRuleGetter::GetFriendlyName(NodeTemplate->GetterType), *OwnerName, *StateNode->GetStateName());
+			const FString Title = FString::Printf(TEXT("Current %s for state '%s.%s'"), *UK2Node_TransitionRuleGetter::GetFriendlyName(NodeTemplate->GetterType).ToString(), *OwnerName, *StateNode->GetStateName());
 
-			TSharedPtr<FEdGraphSchemaAction_K2NewNode> Action = FK2ActionMenuBuilder::AddNewNodeAction(ContextMenuBuilder, StateCategory, Title, NodeTemplate->GetTooltip(), 0, NodeTemplate->GetKeywords());
+			TSharedPtr<FEdGraphSchemaAction_K2NewNode> Action = FK2ActionMenuBuilder::AddNewNodeAction(ContextMenuBuilder, StateCategory, FText::FromString(Title), NodeTemplate->GetTooltip(), 0, NodeTemplate->GetKeywords());
 			Action->NodeTemplate = NodeTemplate;
+			Action->SearchTitle = NodeTemplate->GetNodeSearchTitle();
 		}
 	}
 }
diff --git a/Engine/Source/Editor/AnimGraph/Private/AnimationStateGraphSchema.cpp b/Engine/Source/Editor/AnimGraph/Private/AnimationStateGraphSchema.cpp
--- a/Engine/Source/Editor/AnimGraph/Private/AnimationStateGraphSchema.cpp
+++ b/Engine/Source/Editor/AnimGraph/Private/AnimationStateGraphSchema.cpp
@@ -30,12 +30,14 @@ void UAnimationStateGraphSchema::CreateDefaultNodesForGraph(UEdGraph& Graph) con
 
 void UAnimationStateGraphSchema::GetGraphDisplayInformation(const UEdGraph& Graph, /*out*/ FGraphDisplayInfo& DisplayInfo) const
 {
-	DisplayInfo.DisplayName = FText::FromString( Graph.GetName() );
+	DisplayInfo.PlainName = FText::FromString( Graph.GetName() );
 
 	if (const UAnimStateNode* StateNode = Cast<const UAnimStateNode>(Graph.GetOuter()))
 	{
-		DisplayInfo.DisplayName = FText::Format( LOCTEXT("StateNameGraphTitle", "{0} (state)"), FText::FromString( StateNode->GetStateName() ) );
+		DisplayInfo.PlainName = FText::Format( LOCTEXT("StateNameGraphTitle", "{0} (state)"), FText::FromString( StateNode->GetStateName() ) );
 	}
+
+	DisplayInfo.DisplayName = DisplayInfo.PlainName;
 }
 
 #undef LOCTEXT_NAMESPACE
diff --git a/Engine/Source/Editor/AnimGraph/Private/AnimationStateMachineSchema.cpp b/Engine/Source/Editor/AnimGraph/Private/AnimationStateMachineSchema.cpp
--- a/Engine/Source/Editor/AnimGraph/Private/AnimationStateMachineSchema.cpp
+++ b/Engine/Source/Editor/AnimGraph/Private/AnimationStateMachineSchema.cpp
@@ -17,7 +17,7 @@
 
 /////////////////////////////////////////////////////
 
-TSharedPtr<FEdGraphSchemaAction_NewStateNode> AddNewStateNodeAction(FGraphContextMenuBuilder& ContextMenuBuilder, const FString& Category, const FString& MenuDesc, const FString& Tooltip, const int32 Grouping = 0)
+TSharedPtr<FEdGraphSchemaAction_NewStateNode> AddNewStateNodeAction(FGraphContextMenuBuilder& ContextMenuBuilder, const FString& Category, const FText& MenuDesc, const FString& Tooltip, const int32 Grouping = 0)
 {
 	TSharedPtr<FEdGraphSchemaAction_NewStateNode> NewStateNode( new FEdGraphSchemaAction_NewStateNode(Category, MenuDesc, Tooltip, Grouping) );
 	ContextMenuBuilder.AddAction( NewStateNode );
@@ -255,14 +255,16 @@ void UAnimationStateMachineSchema::GetGraphContextActions(FGraphContextMenuBuild
 {
 	// Add state node
 	{
-		TSharedPtr<FEdGraphSchemaAction_NewStateNode> Action = AddNewStateNodeAction(ContextMenuBuilder, TEXT(""), TEXT("Add State..."), TEXT("A new state"));
+		TSharedPtr<FEdGraphSchemaAction_NewStateNode> Action = AddNewStateNodeAction(ContextMenuBuilder, TEXT(""), LOCTEXT("AddState", "Add State..."), TEXT("A new state"));
 		Action->NodeTemplate = NewObject<UAnimStateNode>(ContextMenuBuilder.OwnerOfTemporaries);
+		Action->SearchTitle = Action->NodeTemplate->GetNodeSearchTitle();
 	}
 
 	// Add conduit node
 	{
-		TSharedPtr<FEdGraphSchemaAction_NewStateNode> Action = AddNewStateNodeAction(ContextMenuBuilder, TEXT(""), TEXT("Add Conduit..."), TEXT("A new conduit state"));
+		TSharedPtr<FEdGraphSchemaAction_NewStateNode> Action = AddNewStateNodeAction(ContextMenuBuilder, TEXT(""), LOCTEXT("AddConduit", "Add Conduit..."), TEXT("A new conduit state"));
 		Action->NodeTemplate = NewObject<UAnimStateConduitNode>(ContextMenuBuilder.OwnerOfTemporaries);
+		Action->SearchTitle = Action->NodeTemplate->GetNodeSearchTitle();
 	}
 
 	// Entry point (only if doesn't already exist)
@@ -280,8 +282,9 @@ void UAnimationStateMachineSchema::GetGraphContextActions(FGraphContextMenuBuild
 
 		if (!bHasEntry)
 		{
-			TSharedPtr<FEdGraphSchemaAction_NewStateNode> Action = AddNewStateNodeAction(ContextMenuBuilder, TEXT(""), TEXT("Add Entry Point..."), TEXT("Define State Machine's Entry Point"));
+			TSharedPtr<FEdGraphSchemaAction_NewStateNode> Action = AddNewStateNodeAction(ContextMenuBuilder, TEXT(""), LOCTEXT("AddEntryPoint", "Add Entry Point..."), TEXT("Define State Machine's Entry Point"));
 			Action->NodeTemplate = NewObject<UAnimStateEntryNode>(ContextMenuBuilder.OwnerOfTemporaries);
+			Action->SearchTitle = Action->NodeTemplate->GetNodeSearchTitle();
 		}
 	}
 
@@ -290,7 +293,7 @@ void UAnimationStateMachineSchema::GetGraphContextActions(FGraphContextMenuBuild
 	{
 		UBlueprint* OwnerBlueprint = FBlueprintEditorUtils::FindBlueprintForGraphChecked(ContextMenuBuilder.CurrentGraph);
 		const bool bIsManyNodesSelected = (FKismetEditorUtilities::GetNumberOfSelectedNodes(OwnerBlueprint) > 0);
-		const FString MenuDescription = bIsManyNodesSelected ? TEXT("Create Comment from Selection") : TEXT("Add Comment...");
+		const FText MenuDescription = bIsManyNodesSelected ? LOCTEXT("CreateCommentSelection", "Create Comment from Selection") : LOCTEXT("AddComment", "Add Comment...");
 		const FString ToolTip = TEXT("Create a resizeable comment box around selected nodes.");
 
 		TSharedPtr<FEdGraphSchemaAction_NewStateComment> NewComment( new FEdGraphSchemaAction_NewStateComment(TEXT(""), MenuDescription, ToolTip, 0) );
@@ -340,7 +343,8 @@ FLinearColor UAnimationStateMachineSchema::GetPinTypeColor(const FEdGraphPinType
 
 void UAnimationStateMachineSchema::GetGraphDisplayInformation(const UEdGraph& Graph, /*out*/ FGraphDisplayInfo& DisplayInfo) const
 {
-	DisplayInfo.DisplayName = FText::FromString( Graph.GetName() );
+	DisplayInfo.PlainName = FText::FromString( Graph.GetName() );
+	DisplayInfo.DisplayName = DisplayInfo.PlainName;
 }
 
 void UAnimationStateMachineSchema::DroppedAssetsOnGraph(const TArray<FAssetData>& Assets, const FVector2D& GraphPosition, UEdGraph* Graph) const
diff --git a/Engine/Source/Editor/AnimGraph/Private/AnimationTransitionSchema.cpp b/Engine/Source/Editor/AnimGraph/Private/AnimationTransitionSchema.cpp
--- a/Engine/Source/Editor/AnimGraph/Private/AnimationTransitionSchema.cpp
+++ b/Engine/Source/Editor/AnimGraph/Private/AnimationTransitionSchema.cpp
@@ -16,6 +16,8 @@
 /////////////////////////////////////////////////////
 // UAnimationTransitionSchema
 
+#define LOCTEXT_NAMESPACE "AnimationTransitionSchema"
+
 UAnimationTransitionSchema::UAnimationTransitionSchema(const class FPostConstructInitializeProperties& PCIP)
 	: Super(PCIP)
 {
@@ -80,10 +82,14 @@ void UAnimationTransitionSchema::GetSourceStateActions(FGraphContextMenuBuilder&
 
 								NodeTemplate->GetterType = SequenceSpecificGetters[TypeIndex];
 
-								FString Title = FString::Printf(TEXT("%s for '%s'"), *UK2Node_TransitionRuleGetter::GetFriendlyName(NodeTemplate->GetterType), *AssetName);
+								FFormatNamedArguments Args;
+								Args.Add(TEXT("NodeName"), UK2Node_TransitionRuleGetter::GetFriendlyName(NodeTemplate->GetterType));
+								Args.Add(TEXT("AssetName"), FText::FromString(AssetName));
+								FText Title = FText::Format(LOCTEXT("TransitionFor", "{NodeName} for '{AssetName}'"), Args);
 
 								TSharedPtr<FEdGraphSchemaAction_K2NewNode> Action = FK2ActionMenuBuilder::AddNewNodeAction(ContextMenuBuilder, Category_AssetPlayer, Title, NodeTemplate->GetTooltip(), 0, NodeTemplate->GetKeywords());
 								Action->NodeTemplate = NodeTemplate;
+								Action->SearchTitle = Action->NodeTemplate->GetNodeSearchTitle();
 							}
 						}
 					}
@@ -103,10 +109,11 @@ void UAnimationTransitionSchema::GetSourceStateActions(FGraphContextMenuBuilder&
 						UK2Node_TransitionRuleGetter* NodeTemplate = ContextMenuBuilder.CreateTemplateNode<UK2Node_TransitionRuleGetter>();
 						NodeTemplate->GetterType = NonSpecificGetters[TypeIndex];
 
-						FString Title = FString::Printf(TEXT("%s"), *UK2Node_TransitionRuleGetter::GetFriendlyName(NodeTemplate->GetterType));
+						FText Title = UK2Node_TransitionRuleGetter::GetFriendlyName(NodeTemplate->GetterType);
 
 						TSharedPtr<FEdGraphSchemaAction_K2NewNode> Action = FK2ActionMenuBuilder::AddNewNodeAction(ContextMenuBuilder, Category_Transition, Title, NodeTemplate->GetTooltip(), 0, NodeTemplate->GetKeywords());
 						Action->NodeTemplate = NodeTemplate;
+						Action->SearchTitle = Action->NodeTemplate->GetNodeSearchTitle();
 					}
 				}
 			}
@@ -127,7 +134,7 @@ void UAnimationTransitionSchema::CreateDefaultNodesForGraph(UEdGraph& Graph) con
 
 void UAnimationTransitionSchema::GetGraphDisplayInformation(const UEdGraph& Graph, /*out*/ FGraphDisplayInfo& DisplayInfo) const
 {
-	DisplayInfo.DisplayName = FText::FromString( Graph.GetName() );
+	DisplayInfo.PlainName = FText::FromString( Graph.GetName() );
 
 	const UAnimStateTransitionNode* TransNode = Cast<const UAnimStateTransitionNode>(Graph.GetOuter());
 	if (TransNode == NULL)
@@ -142,6 +149,12 @@ void UAnimationTransitionSchema::GetGraphDisplayInformation(const UEdGraph& Grap
 
 	if (TransNode)
 	{
-		DisplayInfo.DisplayName = FText::Format( NSLOCTEXT("Animation", "TransitionRuleGraphTitle", "{0} (rule)"), FText::FromString( TransNode->GetNodeTitle(ENodeTitleType::FullTitle)) );
+		FFormatNamedArguments Args;
+		Args.Add(TEXT("NodeTitle"), TransNode->GetNodeTitle(ENodeTitleType::FullTitle));
+		DisplayInfo.PlainName = FText::Format( NSLOCTEXT("Animation", "TransitionRuleGraphTitle", "{NodeTitle} (rule)"), Args );
 	}
+
+	DisplayInfo.DisplayName = DisplayInfo.PlainName;
 }
+
+#undef LOCTEXT_NAMESPACE
\ No newline at end of file
diff --git a/Engine/Source/Editor/AnimGraph/Private/K2Node_TransitionRuleGetter.cpp b/Engine/Source/Editor/AnimGraph/Private/K2Node_TransitionRuleGetter.cpp
--- a/Engine/Source/Editor/AnimGraph/Private/K2Node_TransitionRuleGetter.cpp
+++ b/Engine/Source/Editor/AnimGraph/Private/K2Node_TransitionRuleGetter.cpp
@@ -19,10 +19,8 @@ void UK2Node_TransitionRuleGetter::AllocateDefaultPins()
 {
 	const UEdGraphSchema_K2* Schema = GetDefault<UEdGraphSchema_K2>();
 
-	FString FriendlyName = GetFriendlyName(GetterType);
-
 	UEdGraphPin* OutputPin = CreatePin(EGPD_Output, Schema->PC_Float, /*PSC=*/ TEXT(""), /*PSC object=*/ NULL, /*bIsArray=*/ false, /*bIsReference=*/ false, TEXT("Output"));
-	OutputPin->PinFriendlyName = FriendlyName;
+	OutputPin->PinFriendlyName = GetFriendlyName(GetterType);
 
 	PreloadObject(AssociatedAnimAssetPlayerNode);
 	PreloadObject(AssociatedStateNode);
@@ -30,7 +28,7 @@ void UK2Node_TransitionRuleGetter::AllocateDefaultPins()
 	Super::AllocateDefaultPins();
 }
 
-FString UK2Node_TransitionRuleGetter::GetFriendlyName(ETransitionGetter::Type TypeID)
+FText UK2Node_TransitionRuleGetter::GetFriendlyName(ETransitionGetter::Type TypeID)
 {
 	FText FriendlyName;
 	switch (TypeID)
@@ -67,17 +65,19 @@ FString UK2Node_TransitionRuleGetter::GetFriendlyName(ETransitionGetter::Type Ty
 		break;
 	}
 
-	return FriendlyName.ToString();
+	return FriendlyName;
 }
 
-FString UK2Node_TransitionRuleGetter::GetNodeTitle(ENodeTitleType::Type TitleType) const
+FText UK2Node_TransitionRuleGetter::GetNodeTitle(ENodeTitleType::Type TitleType) const
 {
 	if (AssociatedAnimAssetPlayerNode != NULL)
 	{
 		UAnimationAsset* BoundAsset = AssociatedAnimAssetPlayerNode->GetAnimationAsset();
 		if (BoundAsset)
 		{
-			return  FString::Printf(*LOCTEXT("AnimationAssetInfoGetterTitle", "%s Asset").ToString(), *BoundAsset->GetName());
+			FFormatNamedArguments Args;
+			Args.Add(TEXT("BoundAsset"), FText::FromString(BoundAsset->GetName()));
+			return  FText::Format(LOCTEXT("AnimationAssetInfoGetterTitle", "{BoundAsset} Asset"), Args);
 		}
 	}
 	else if (AssociatedStateNode != NULL)
@@ -86,25 +86,61 @@ FString UK2Node_TransitionRuleGetter::GetNodeTitle(ENodeTitleType::Type TitleTyp
 		{
 			const FString OwnerName = State->GetOuter()->GetName();
 
-			return FString::Printf(*LOCTEXT("StateInfoGetterTitle", "%s.%s State").ToString(), *OwnerName, *(State->GetStateName()));
+			FFormatNamedArguments Args;
+			Args.Add(TEXT("OwnerName"), FText::FromString(OwnerName));
+			Args.Add(TEXT("StateName"), FText::FromString(State->GetStateName()));
+			return FText::Format(LOCTEXT("StateInfoGetterTitle", "{OwnerName}.{StateName} State"), Args);
 		}
 	}
 	else if (GetterType == ETransitionGetter::CurrentTransitionDuration)
 	{
-		return LOCTEXT("TransitionDuration", "Transition").ToString();
+		return LOCTEXT("TransitionDuration", "Transition");
 	}
 	else if (GetterType == ETransitionGetter::CurrentState_ElapsedTime ||
 			 GetterType == ETransitionGetter::CurrentState_GetBlendWeight)
 	{
-		return LOCTEXT("CurrentState", "Current State").ToString();
+		return LOCTEXT("CurrentState", "Current State");
 	}
 
 	return Super::GetNodeTitle(TitleType);
 }
 
+FString UK2Node_TransitionRuleGetter::GetNodeNativeTitle(ENodeTitleType::Type TitleType) const
+{
+	// Do not setup this function for localization, intentionally left unlocalized!
+	if (AssociatedAnimAssetPlayerNode != NULL)
+	{
+		UAnimationAsset* BoundAsset = AssociatedAnimAssetPlayerNode->GetAnimationAsset();
+		if (BoundAsset)
+		{
+			return  FString::Printf(TEXT("%s Asset"), *BoundAsset->GetName());
+		}
+	}
+	else if (AssociatedStateNode != NULL)
+	{
+		if (UAnimStateNode* State = Cast<UAnimStateNode>(AssociatedStateNode))
+		{
+			const FString OwnerName = State->GetOuter()->GetName();
+
+			return FString::Printf(TEXT("%s.%s State"), *OwnerName, *(State->GetStateName()));
+		}
+	}
+	else if (GetterType == ETransitionGetter::CurrentTransitionDuration)
+	{
+		return TEXT("Transition");
+	}
+	else if (GetterType == ETransitionGetter::CurrentState_ElapsedTime ||
+		GetterType == ETransitionGetter::CurrentState_GetBlendWeight)
+	{
+		return TEXT("Current State");
+	}
+
+	return Super::GetNodeNativeTitle(TitleType);
+}
+
 FString UK2Node_TransitionRuleGetter::GetTooltip() const
 {
-	return GetNodeTitle(ENodeTitleType::FullTitle);
+	return GetNodeTitle(ENodeTitleType::FullTitle).ToString();
 }
 
 UEdGraphPin* UK2Node_TransitionRuleGetter::GetOutputPin() const
diff --git a/Engine/Source/Editor/BehaviorTreeEditor/Classes/EdGraphSchema_BehaviorTree.h b/Engine/Source/Editor/BehaviorTreeEditor/Classes/EdGraphSchema_BehaviorTree.h
--- a/Engine/Source/Editor/BehaviorTreeEditor/Classes/EdGraphSchema_BehaviorTree.h
+++ b/Engine/Source/Editor/BehaviorTreeEditor/Classes/EdGraphSchema_BehaviorTree.h
@@ -19,7 +19,7 @@ struct FBehaviorTreeSchemaAction_NewNode : public FEdGraphSchemaAction
 		, NodeTemplate(NULL)
 	{}
 
-	FBehaviorTreeSchemaAction_NewNode(const FString& InNodeCategory, const FString& InMenuDesc, const FString& InToolTip, const int32 InGrouping)
+	FBehaviorTreeSchemaAction_NewNode(const FString& InNodeCategory, const FText& InMenuDesc, const FString& InToolTip, const int32 InGrouping)
 		: FEdGraphSchemaAction(InNodeCategory, InMenuDesc, InToolTip, InGrouping) 
 		, NodeTemplate(NULL)
 	{}
@@ -59,7 +59,7 @@ struct FBehaviorTreeSchemaAction_NewSubNode : public FEdGraphSchemaAction
 		, NodeTemplate(NULL)
 	{}
 
-	FBehaviorTreeSchemaAction_NewSubNode(const FString& InNodeCategory, const FString& InMenuDesc, const FString& InToolTip, const int32 InGrouping)
+	FBehaviorTreeSchemaAction_NewSubNode(const FString& InNodeCategory, const FText& InMenuDesc, const FString& InToolTip, const int32 InGrouping)
 		: FEdGraphSchemaAction(InNodeCategory, InMenuDesc, InToolTip, InGrouping) 
 		, NodeTemplate(NULL)
 	{}
@@ -80,7 +80,7 @@ struct FBehaviorTreeSchemaAction_AutoArrange : public FEdGraphSchemaAction
 	FBehaviorTreeSchemaAction_AutoArrange() 
 		: FEdGraphSchemaAction() {}
 
-	FBehaviorTreeSchemaAction_AutoArrange(const FString& InNodeCategory, const FString& InMenuDesc, const FString& InToolTip, const int32 InGrouping)
+	FBehaviorTreeSchemaAction_AutoArrange(const FString& InNodeCategory, const FText& InMenuDesc, const FString& InToolTip, const int32 InGrouping)
 		: FEdGraphSchemaAction(InNodeCategory, InMenuDesc, InToolTip, InGrouping) 
 	{}
 
diff --git a/Engine/Source/Editor/BehaviorTreeEditor/Classes/EdGraphSchema_BehaviorTreeDecorator.h b/Engine/Source/Editor/BehaviorTreeEditor/Classes/EdGraphSchema_BehaviorTreeDecorator.h
--- a/Engine/Source/Editor/BehaviorTreeEditor/Classes/EdGraphSchema_BehaviorTreeDecorator.h
+++ b/Engine/Source/Editor/BehaviorTreeEditor/Classes/EdGraphSchema_BehaviorTreeDecorator.h
@@ -18,7 +18,7 @@ struct FDecoratorSchemaAction_NewNode : public FEdGraphSchemaAction
 		, NodeTemplate(NULL)
 	{}
 
-	FDecoratorSchemaAction_NewNode(const FString& InNodeCategory, const FString& InMenuDesc, const FString& InToolTip, const int32 InGrouping)
+	FDecoratorSchemaAction_NewNode(const FString& InNodeCategory, const FText& InMenuDesc, const FString& InToolTip, const int32 InGrouping)
 		: FEdGraphSchemaAction(InNodeCategory, InMenuDesc, InToolTip, InGrouping) 
 		, NodeTemplate(NULL)
 	{}
@@ -60,5 +60,5 @@ class UEdGraphSchema_BehaviorTreeDecorator : public UEdGraphSchema
 	virtual bool ShouldHidePinDefaultValue(UEdGraphPin* Pin) const OVERRIDE;
 	// End EdGraphSchema interface
 
-	static TSharedPtr<FDecoratorSchemaAction_NewNode> AddNewDecoratorAction(FGraphContextMenuBuilder& ContextMenuBuilder, const FString& Category, const FString& MenuDesc, const FString& Tooltip);
+	static TSharedPtr<FDecoratorSchemaAction_NewNode> AddNewDecoratorAction(FGraphContextMenuBuilder& ContextMenuBuilder, const FString& Category, const FText& MenuDesc, const FString& Tooltip);
 };
diff --git a/Engine/Source/Editor/BehaviorTreeEditor/Private/BehaviorTreeDecoratorGraphNode_Decorator.cpp b/Engine/Source/Editor/BehaviorTreeEditor/Private/BehaviorTreeDecoratorGraphNode_Decorator.cpp
--- a/Engine/Source/Editor/BehaviorTreeEditor/Private/BehaviorTreeDecoratorGraphNode_Decorator.cpp
+++ b/Engine/Source/Editor/BehaviorTreeEditor/Private/BehaviorTreeDecoratorGraphNode_Decorator.cpp
@@ -28,12 +28,12 @@ void UBehaviorTreeDecoratorGraphNode_Decorator::PostPlacedNewNode()
 	}
 }
 
-FString UBehaviorTreeDecoratorGraphNode_Decorator::GetNodeTitle(ENodeTitleType::Type TitleType) const
+FText UBehaviorTreeDecoratorGraphNode_Decorator::GetNodeTitle(ENodeTitleType::Type TitleType) const
 {
 	const UBTNode* MyNode = Cast<UBTNode>(NodeInstance);
 	if (MyNode != NULL)
 	{
-		return MyNode->GetStaticDescription();
+		return FText::FromString(MyNode->GetStaticDescription());
 	}
 
 	return Super::GetNodeTitle(TitleType);
diff --git a/Engine/Source/Editor/BehaviorTreeEditor/Private/BehaviorTreeDecoratorGraphNode_Logic.cpp b/Engine/Source/Editor/BehaviorTreeEditor/Private/BehaviorTreeDecoratorGraphNode_Logic.cpp
--- a/Engine/Source/Editor/BehaviorTreeEditor/Private/BehaviorTreeDecoratorGraphNode_Logic.cpp
+++ b/Engine/Source/Editor/BehaviorTreeEditor/Private/BehaviorTreeDecoratorGraphNode_Logic.cpp
@@ -31,9 +31,9 @@ static FString DescribeLogicModeHelper(const EDecoratorLogicMode::Type& Mode)
 	return Desc[Mode];
 }
 
-FString UBehaviorTreeDecoratorGraphNode_Logic::GetNodeTitle(ENodeTitleType::Type TitleType) const
+FText UBehaviorTreeDecoratorGraphNode_Logic::GetNodeTitle(ENodeTitleType::Type TitleType) const
 {
-	return DescribeLogicModeHelper(LogicMode);
+	return FText::FromString(DescribeLogicModeHelper(LogicMode));
 }
 
 void UBehaviorTreeDecoratorGraphNode_Logic::GetMenuEntries(struct FGraphContextMenuBuilder& ContextMenuBuilder) const
@@ -42,11 +42,12 @@ void UBehaviorTreeDecoratorGraphNode_Logic::GetMenuEntries(struct FGraphContextM
 	const int32 NumModes = sizeof(Modes) / sizeof(Modes[0]);
 	for (int32 i = 0; i < NumModes; i++)
 	{
-		TSharedPtr<FDecoratorSchemaAction_NewNode> AddOpAction = UEdGraphSchema_BehaviorTreeDecorator::AddNewDecoratorAction(ContextMenuBuilder, TEXT("Logic"), DescribeLogicModeHelper(Modes[i]), "");
+		TSharedPtr<FDecoratorSchemaAction_NewNode> AddOpAction = UEdGraphSchema_BehaviorTreeDecorator::AddNewDecoratorAction(ContextMenuBuilder, TEXT("Logic"), FText::FromString(DescribeLogicModeHelper(Modes[i])), "");
 
 		UBehaviorTreeDecoratorGraphNode_Logic* OpNode = NewObject<UBehaviorTreeDecoratorGraphNode_Logic>(ContextMenuBuilder.OwnerOfTemporaries);
 		OpNode->LogicMode = Modes[i];
 		AddOpAction->NodeTemplate = OpNode;
+		AddOpAction->SearchTitle = AddOpAction->NodeTemplate->GetNodeSearchTitle();
 	}
 }
 
diff --git a/Engine/Source/Editor/BehaviorTreeEditor/Private/BehaviorTreeGraphNode_Composite.cpp b/Engine/Source/Editor/BehaviorTreeEditor/Private/BehaviorTreeGraphNode_Composite.cpp
--- a/Engine/Source/Editor/BehaviorTreeEditor/Private/BehaviorTreeGraphNode_Composite.cpp
+++ b/Engine/Source/Editor/BehaviorTreeEditor/Private/BehaviorTreeGraphNode_Composite.cpp
@@ -9,12 +9,12 @@ UBehaviorTreeGraphNode_Composite::UBehaviorTreeGraphNode_Composite(const class F
 {
 }
 
-FString UBehaviorTreeGraphNode_Composite::GetNodeTitle(ENodeTitleType::Type TitleType) const
+FText UBehaviorTreeGraphNode_Composite::GetNodeTitle(ENodeTitleType::Type TitleType) const
 {
 	const UBTNode* MyNode = Cast<UBTNode>(NodeInstance);
 	if (MyNode != NULL)
 	{
-		return  MyNode->GetNodeName();
+		return FText::FromString(MyNode->GetNodeName());
 	}
 	return Super::GetNodeTitle(TitleType);
 }
diff --git a/Engine/Source/Editor/BehaviorTreeEditor/Private/BehaviorTreeGraphNode_CompositeDecorator.cpp b/Engine/Source/Editor/BehaviorTreeEditor/Private/BehaviorTreeGraphNode_CompositeDecorator.cpp
--- a/Engine/Source/Editor/BehaviorTreeEditor/Private/BehaviorTreeGraphNode_CompositeDecorator.cpp
+++ b/Engine/Source/Editor/BehaviorTreeEditor/Private/BehaviorTreeGraphNode_CompositeDecorator.cpp
@@ -29,9 +29,9 @@ FString UBehaviorTreeGraphNode_CompositeDecorator::GetNodeTypeDescription() cons
 	return LOCTEXT("Composite","Composite").ToString();
 }
 
-FString UBehaviorTreeGraphNode_CompositeDecorator::GetNodeTitle(ENodeTitleType::Type TitleType) const
+FText UBehaviorTreeGraphNode_CompositeDecorator::GetNodeTitle(ENodeTitleType::Type TitleType) const
 {
-	return CompositeName.Len() ? CompositeName : GetNodeTypeDescription();
+	return FText::FromString(CompositeName.Len() ? CompositeName : GetNodeTypeDescription());
 }
 
 FName UBehaviorTreeGraphNode_CompositeDecorator::GetNameIcon() const
diff --git a/Engine/Source/Editor/BehaviorTreeEditor/Private/BehaviorTreeGraphNode_Decorator.cpp b/Engine/Source/Editor/BehaviorTreeEditor/Private/BehaviorTreeGraphNode_Decorator.cpp
--- a/Engine/Source/Editor/BehaviorTreeEditor/Private/BehaviorTreeGraphNode_Decorator.cpp
+++ b/Engine/Source/Editor/BehaviorTreeEditor/Private/BehaviorTreeGraphNode_Decorator.cpp
@@ -12,12 +12,12 @@ void UBehaviorTreeGraphNode_Decorator::AllocateDefaultPins()
 	//No Pins for decorators
 }
 
-FString UBehaviorTreeGraphNode_Decorator::GetNodeTitle(ENodeTitleType::Type TitleType) const
+FText UBehaviorTreeGraphNode_Decorator::GetNodeTitle(ENodeTitleType::Type TitleType) const
 {
 	const UBTDecorator* Decorator = Cast<UBTDecorator>(NodeInstance);
 	if (Decorator != NULL)
 	{
-		return Decorator->GetNodeName();
+		return FText::FromString(Decorator->GetNodeName());
 	}
 
 	return Super::GetNodeTitle(TitleType);
diff --git a/Engine/Source/Editor/BehaviorTreeEditor/Private/BehaviorTreeGraphNode_Root.cpp b/Engine/Source/Editor/BehaviorTreeEditor/Private/BehaviorTreeGraphNode_Root.cpp
--- a/Engine/Source/Editor/BehaviorTreeEditor/Private/BehaviorTreeGraphNode_Root.cpp
+++ b/Engine/Source/Editor/BehaviorTreeEditor/Private/BehaviorTreeGraphNode_Root.cpp
@@ -28,9 +28,9 @@ void UBehaviorTreeGraphNode_Root::AllocateDefaultPins()
 	CreatePin(EGPD_Output, UBehaviorTreeEditorTypes::PinCategory_SingleComposite, TEXT(""), NULL, false, false, TEXT("In"));
 }
 
-FString UBehaviorTreeGraphNode_Root::GetNodeTitle(ENodeTitleType::Type TitleType) const
+FText UBehaviorTreeGraphNode_Root::GetNodeTitle(ENodeTitleType::Type TitleType) const
 {
-	return FString("ROOT");
+	return NSLOCTEXT("BehaviorTreeGraphNode", "Root", "ROOT");
 }
 
 FName UBehaviorTreeGraphNode_Root::GetNameIcon() const
diff --git a/Engine/Source/Editor/BehaviorTreeEditor/Private/BehaviorTreeGraphNode_Service.cpp b/Engine/Source/Editor/BehaviorTreeEditor/Private/BehaviorTreeGraphNode_Service.cpp
--- a/Engine/Source/Editor/BehaviorTreeEditor/Private/BehaviorTreeGraphNode_Service.cpp
+++ b/Engine/Source/Editor/BehaviorTreeEditor/Private/BehaviorTreeGraphNode_Service.cpp
@@ -13,12 +13,12 @@ void UBehaviorTreeGraphNode_Service::AllocateDefaultPins()
 }
 
 
-FString UBehaviorTreeGraphNode_Service::GetNodeTitle(ENodeTitleType::Type TitleType) const
+FText UBehaviorTreeGraphNode_Service::GetNodeTitle(ENodeTitleType::Type TitleType) const
 {
 	const UBTService* Service = Cast<UBTService>(NodeInstance);
 	if (Service != NULL)
 	{
-		return Service->GetNodeName();
+		return FText::FromString(Service->GetNodeName());
 	}
 
 	return Super::GetNodeTitle(TitleType);
diff --git a/Engine/Source/Editor/BehaviorTreeEditor/Private/BehaviorTreeGraphNode_Task.cpp b/Engine/Source/Editor/BehaviorTreeEditor/Private/BehaviorTreeGraphNode_Task.cpp
--- a/Engine/Source/Editor/BehaviorTreeEditor/Private/BehaviorTreeGraphNode_Task.cpp
+++ b/Engine/Source/Editor/BehaviorTreeEditor/Private/BehaviorTreeGraphNode_Task.cpp
@@ -13,12 +13,12 @@ void UBehaviorTreeGraphNode_Task::AllocateDefaultPins()
 	CreatePin(EGPD_Input, UBehaviorTreeEditorTypes::PinCategory_SingleComposite, TEXT(""), NULL, false, false, TEXT("In"));
 }
 
-FString UBehaviorTreeGraphNode_Task::GetNodeTitle(ENodeTitleType::Type TitleType) const
+FText UBehaviorTreeGraphNode_Task::GetNodeTitle(ENodeTitleType::Type TitleType) const
 {
 	const UBTNode* MyNode = Cast<UBTNode>(NodeInstance);
 	if (MyNode != NULL)
 	{
-		return MyNode->GetNodeName();
+		return FText::FromString(MyNode->GetNodeName());
 	}
 
 	return Super::GetNodeTitle(TitleType);
diff --git a/Engine/Source/Editor/BehaviorTreeEditor/Private/EdGraphSchema_BehaviorTree.cpp b/Engine/Source/Editor/BehaviorTreeEditor/Private/EdGraphSchema_BehaviorTree.cpp
--- a/Engine/Source/Editor/BehaviorTreeEditor/Private/EdGraphSchema_BehaviorTree.cpp
+++ b/Engine/Source/Editor/BehaviorTreeEditor/Private/EdGraphSchema_BehaviorTree.cpp
@@ -217,7 +217,7 @@ UEdGraphSchema_BehaviorTree::UEdGraphSchema_BehaviorTree(const class FPostConstr
 {
 }
 
-TSharedPtr<FBehaviorTreeSchemaAction_NewNode> AddNewNodeAction(FGraphContextMenuBuilder& ContextMenuBuilder, const FString& Category, const FString& MenuDesc, const FString& Tooltip)
+TSharedPtr<FBehaviorTreeSchemaAction_NewNode> AddNewNodeAction(FGraphContextMenuBuilder& ContextMenuBuilder, const FString& Category, const FText& MenuDesc, const FString& Tooltip)
 {
 	TSharedPtr<FBehaviorTreeSchemaAction_NewNode> NewAction = TSharedPtr<FBehaviorTreeSchemaAction_NewNode>(new FBehaviorTreeSchemaAction_NewNode(Category, MenuDesc, Tooltip, 0));
 
@@ -226,7 +226,7 @@ TSharedPtr<FBehaviorTreeSchemaAction_NewNode> AddNewNodeAction(FGraphContextMenu
 	return NewAction;
 }
 
-TSharedPtr<FBehaviorTreeSchemaAction_NewSubNode> AddNewSubNodeAction(FGraphContextMenuBuilder& ContextMenuBuilder, const FString& Category, const FString& MenuDesc, const FString& Tooltip)
+TSharedPtr<FBehaviorTreeSchemaAction_NewSubNode> AddNewSubNodeAction(FGraphContextMenuBuilder& ContextMenuBuilder, const FString& Category, const FText& MenuDesc, const FString& Tooltip)
 {
 	TSharedPtr<FBehaviorTreeSchemaAction_NewSubNode> NewAction = TSharedPtr<FBehaviorTreeSchemaAction_NewSubNode>(new FBehaviorTreeSchemaAction_NewSubNode(Category, MenuDesc, Tooltip, 0));
 	ContextMenuBuilder.AddAction( NewAction );
@@ -268,14 +268,15 @@ void UEdGraphSchema_BehaviorTree::GetGraphNodeContextActions(FGraphContextMenuBu
 
 		{
 			UBehaviorTreeGraphNode_CompositeDecorator* OpNode = NewObject<UBehaviorTreeGraphNode_CompositeDecorator>(Graph);
-			TSharedPtr<FBehaviorTreeSchemaAction_NewSubNode> AddOpAction = AddNewSubNodeAction(ContextMenuBuilder, TEXT(""), OpNode->GetNodeTypeDescription(), "");
+			TSharedPtr<FBehaviorTreeSchemaAction_NewSubNode> AddOpAction = AddNewSubNodeAction(ContextMenuBuilder, TEXT(""), FText::FromString(OpNode->GetNodeTypeDescription()), "");
 			AddOpAction->ParentNode = Cast<UBehaviorTreeGraphNode>(ContextMenuBuilder.SelectedObjects[0]);
 			AddOpAction->NodeTemplate = OpNode;
+			AddOpAction->SearchTitle = AddOpAction->NodeTemplate->GetNodeSearchTitle();
 		}
 
 		for (int32 i = 0; i < NodeClasses.Num(); i++)
 		{
-			const FString NodeTypeName = FName::NameToDisplayString(NodeClasses[i].ToString(), false);
+			const FText NodeTypeName = FText::FromString(FName::NameToDisplayString(NodeClasses[i].ToString(), false));
 
 			UBehaviorTreeGraphNode* OpNode = NewObject<UBehaviorTreeGraphNode_Decorator>(Graph);
 			OpNode->ClassData = NodeClasses[i];
@@ -283,6 +284,7 @@ void UEdGraphSchema_BehaviorTree::GetGraphNodeContextActions(FGraphContextMenuBu
 			TSharedPtr<FBehaviorTreeSchemaAction_NewSubNode> AddOpAction = AddNewSubNodeAction(ContextMenuBuilder, TEXT(""), NodeTypeName, "");
 			AddOpAction->ParentNode = Cast<UBehaviorTreeGraphNode>(ContextMenuBuilder.SelectedObjects[0]);
 			AddOpAction->NodeTemplate = OpNode;
+			AddOpAction->SearchTitle = AddOpAction->NodeTemplate->GetNodeSearchTitle();
 		}
 	}
 	else if (SubNodeType == ESubNode::Service)
@@ -290,7 +292,7 @@ void UEdGraphSchema_BehaviorTree::GetGraphNodeContextActions(FGraphContextMenuBu
 		FClassBrowseHelper::GatherClasses(UBTService::StaticClass(), NodeClasses);
 		for (int32 i = 0; i < NodeClasses.Num(); i++)
 		{
-			const FString NodeTypeName = FName::NameToDisplayString(NodeClasses[i].ToString(), false);
+			const FText NodeTypeName = FText::FromString(FName::NameToDisplayString(NodeClasses[i].ToString(), false));
 
 			UBehaviorTreeGraphNode* OpNode = NewObject<UBehaviorTreeGraphNode_Service>(Graph);
 			OpNode->ClassData = NodeClasses[i];
@@ -298,6 +300,7 @@ void UEdGraphSchema_BehaviorTree::GetGraphNodeContextActions(FGraphContextMenuBu
 			TSharedPtr<FBehaviorTreeSchemaAction_NewSubNode> AddOpAction = AddNewSubNodeAction(ContextMenuBuilder, TEXT(""), NodeTypeName, "");
 			AddOpAction->ParentNode = Cast<UBehaviorTreeGraphNode>(ContextMenuBuilder.SelectedObjects[0]);
 			AddOpAction->NodeTemplate = OpNode;
+			AddOpAction->SearchTitle = AddOpAction->NodeTemplate->GetNodeSearchTitle();
 		}
 	}
 }
@@ -323,7 +326,7 @@ void UEdGraphSchema_BehaviorTree::GetGraphContextActions(FGraphContextMenuBuilde
 
 		for (int32 i = 0; i < NodeClasses.Num(); i++)
 		{
-			const FString NodeTypeName = FName::NameToDisplayString(NodeClasses[i].ToString(), false);
+			const FText NodeTypeName = FText::FromString(FName::NameToDisplayString(NodeClasses[i].ToString(), false));
 
 			TSharedPtr<FBehaviorTreeSchemaAction_NewNode> AddOpAction = AddNewNodeAction(ContextMenuBuilder, TEXT("Composites"), NodeTypeName, "");
 
@@ -333,6 +336,7 @@ void UEdGraphSchema_BehaviorTree::GetGraphContextActions(FGraphContextMenuBuilde
 
 			OpNode->ClassData = NodeClasses[i];
 			AddOpAction->NodeTemplate = OpNode;
+			AddOpAction->SearchTitle = AddOpAction->NodeTemplate->GetNodeSearchTitle();
 		}
 	}
 
@@ -343,19 +347,20 @@ void UEdGraphSchema_BehaviorTree::GetGraphContextActions(FGraphContextMenuBuilde
 
 		for (int32 i = 0; i < NodeClasses.Num(); i++)
 		{
-			const FString NodeTypeName = FName::NameToDisplayString(NodeClasses[i].ToString(), false);
+			const FText NodeTypeName = FText::FromString(FName::NameToDisplayString(NodeClasses[i].ToString(), false));
 
 			TSharedPtr<FBehaviorTreeSchemaAction_NewNode> AddOpAction = AddNewNodeAction(ContextMenuBuilder, TEXT("Tasks"), NodeTypeName, "");
 
 			UBehaviorTreeGraphNode* OpNode = NewObject<UBehaviorTreeGraphNode_Task>(ContextMenuBuilder.OwnerOfTemporaries);
 			OpNode->ClassData = NodeClasses[i];
 			AddOpAction->NodeTemplate = OpNode;
+			AddOpAction->SearchTitle = AddOpAction->NodeTemplate->GetNodeSearchTitle();
 		}
 	}
 	
 	if (bNoParent)
 	{
-		TSharedPtr<FBehaviorTreeSchemaAction_AutoArrange> Action = TSharedPtr<FBehaviorTreeSchemaAction_AutoArrange>(new FBehaviorTreeSchemaAction_AutoArrange(FString(),FString(TEXT("Auto Arrange")),FString(),0));
+		TSharedPtr<FBehaviorTreeSchemaAction_AutoArrange> Action = TSharedPtr<FBehaviorTreeSchemaAction_AutoArrange>(new FBehaviorTreeSchemaAction_AutoArrange(FString(), LOCTEXT("AutoArrange", "Auto Arrange"), FString(), 0));
 		ContextMenuBuilder.AddAction(Action);
 	}
 }
@@ -427,7 +432,7 @@ void UEdGraphSchema_BehaviorTree::GetBreakLinkToSubMenuActions( class FMenuBuild
 	for(TArray<class UEdGraphPin*>::TConstIterator Links(InGraphPin->LinkedTo); Links; ++Links)
 	{
 		UEdGraphPin* Pin = *Links;
-		FString TitleString = Pin->GetOwningNode()->GetNodeTitle(ENodeTitleType::ListView);
+		FString TitleString = Pin->GetOwningNode()->GetNodeTitle(ENodeTitleType::ListView).ToString();
 		FText Title = FText::FromString( TitleString );
 		if ( Pin->PinName != TEXT("") )
 		{
diff --git a/Engine/Source/Editor/BehaviorTreeEditor/Private/EdGraphSchema_BehaviorTreeDecorator.cpp b/Engine/Source/Editor/BehaviorTreeEditor/Private/EdGraphSchema_BehaviorTreeDecorator.cpp
--- a/Engine/Source/Editor/BehaviorTreeEditor/Private/EdGraphSchema_BehaviorTreeDecorator.cpp
+++ b/Engine/Source/Editor/BehaviorTreeEditor/Private/EdGraphSchema_BehaviorTreeDecorator.cpp
@@ -103,7 +103,7 @@ UEdGraphSchema_BehaviorTreeDecorator::UEdGraphSchema_BehaviorTreeDecorator(const
 	PC_Boolean = TEXT("bool");
 }
 
-TSharedPtr<FDecoratorSchemaAction_NewNode> UEdGraphSchema_BehaviorTreeDecorator::AddNewDecoratorAction(FGraphContextMenuBuilder& ContextMenuBuilder, const FString& Category, const FString& MenuDesc, const FString& Tooltip)
+TSharedPtr<FDecoratorSchemaAction_NewNode> UEdGraphSchema_BehaviorTreeDecorator::AddNewDecoratorAction(FGraphContextMenuBuilder& ContextMenuBuilder, const FString& Category, const FText& MenuDesc, const FString& Tooltip)
 {
 	TSharedPtr<FDecoratorSchemaAction_NewNode> NewAction = TSharedPtr<FDecoratorSchemaAction_NewNode>(new FDecoratorSchemaAction_NewNode(Category, MenuDesc, Tooltip, 0));
 
@@ -131,12 +131,13 @@ void UEdGraphSchema_BehaviorTreeDecorator::GetGraphContextActions(FGraphContextM
 
 	for (int32 i = 0; i < NodeClasses.Num(); i++)
 	{
-		const FString NodeTypeName = NodeClasses[i].ToString();
+		const FText NodeTypeName = FText::FromString(NodeClasses[i].ToString());
 		TSharedPtr<FDecoratorSchemaAction_NewNode> AddOpAction = AddNewDecoratorAction(ContextMenuBuilder, TEXT("Decorators"), NodeTypeName, "");
 
 		UBehaviorTreeDecoratorGraphNode_Decorator* OpNode = NewObject<UBehaviorTreeDecoratorGraphNode_Decorator>(ContextMenuBuilder.OwnerOfTemporaries);
 		OpNode->ClassData = NodeClasses[i];
 		AddOpAction->NodeTemplate = OpNode;
+		AddOpAction->SearchTitle = AddOpAction->NodeTemplate->GetNodeSearchTitle();
 	}
 
 #if WITH_EDITOR
@@ -199,7 +200,7 @@ void UEdGraphSchema_BehaviorTreeDecorator::GetBreakLinkToSubMenuActions( class F
 	for(TArray<class UEdGraphPin*>::TConstIterator Links(InGraphPin->LinkedTo); Links; ++Links)
 	{
 		UEdGraphPin* Pin = *Links;
-		FString TitleString = Pin->GetOwningNode()->GetNodeTitle(ENodeTitleType::ListView);
+		FString TitleString = Pin->GetOwningNode()->GetNodeTitle(ENodeTitleType::ListView).ToString();
 		FText Title = FText::FromString( TitleString );
 		if ( Pin->PinName != TEXT("") )
 		{
diff --git a/Engine/Source/Editor/BehaviorTreeEditor/Private/FindInBT.cpp b/Engine/Source/Editor/BehaviorTreeEditor/Private/FindInBT.cpp
--- a/Engine/Source/Editor/BehaviorTreeEditor/Private/FindInBT.cpp
+++ b/Engine/Source/Editor/BehaviorTreeEditor/Private/FindInBT.cpp
@@ -212,7 +212,7 @@ void SFindInBT::MatchTokens(const TArray<FString>& Tokens)
 	{
 		UEdGraphNode* Node = *It;
 			
-		const FString NodeName = Node->GetNodeTitle(ENodeTitleType::ListView);
+		const FString NodeName = Node->GetNodeTitle(ENodeTitleType::ListView).ToString();
 		FSearchResult NodeResult(new FFindInBTResult(NodeName, RootSearchResult, Node));
 
 		FString NodeSearchString = NodeName + Node->GetClass()->GetName() + Node->NodeComment;
@@ -253,7 +253,7 @@ void SFindInBT::MatchTokensInChild(const TArray<FString>& Tokens, UBehaviorTreeG
 		return;
 	}
 
-	FString ChildName = Child->GetNodeTitle(ENodeTitleType::ListView);
+	FString ChildName = Child->GetNodeTitle(ENodeTitleType::ListView).ToString();
 	FString ChildSearchString = ChildName + Child->GetClass()->GetName() + Child->NodeComment;
 	ChildSearchString = ChildSearchString.Replace(TEXT(" "), TEXT(""));
 	if (StringMatchesSearchTokens(Tokens, ChildSearchString))
diff --git a/Engine/Source/Editor/BlueprintGraph/Classes/EdGraphSchema_K2.h b/Engine/Source/Editor/BlueprintGraph/Classes/EdGraphSchema_K2.h
--- a/Engine/Source/Editor/BlueprintGraph/Classes/EdGraphSchema_K2.h
+++ b/Engine/Source/Editor/BlueprintGraph/Classes/EdGraphSchema_K2.h
@@ -694,6 +694,15 @@ public:
 	 */
 	static FString TypeToString(UProperty* const Property);
 
+	/**
+	 * Converts a pin type into a fully qualified FText (e.g., object'ObjectName').
+	 *
+	 * @param	Type	The type to convert into a FText.
+	 *
+	 * @return	The converted type text.
+	 */
+	static FText TypeToText(const FEdGraphPinType& Type);
+
 	/**
 	 * Get the type tree for all of the property types valid for this schema
 	 *
diff --git a/Engine/Source/Editor/BlueprintGraph/Classes/EdGraphSchema_K2_Actions.h b/Engine/Source/Editor/BlueprintGraph/Classes/EdGraphSchema_K2_Actions.h
--- a/Engine/Source/Editor/BlueprintGraph/Classes/EdGraphSchema_K2_Actions.h
+++ b/Engine/Source/Editor/BlueprintGraph/Classes/EdGraphSchema_K2_Actions.h
@@ -45,7 +45,7 @@ struct BLUEPRINTGRAPH_API FEdGraphSchemaAction_K2NewNode : public FEdGraphSchema
 		, NodeTemplate(NULL)
 	{}
 
-	FEdGraphSchemaAction_K2NewNode(const FString& InNodeCategory, const FString& InMenuDesc, const FString& InToolTip, const int32 InGrouping)
+	FEdGraphSchemaAction_K2NewNode(const FString& InNodeCategory, const FText& InMenuDesc, const FString& InToolTip, const int32 InGrouping)
 		: FEdGraphSchemaAction(InNodeCategory, InMenuDesc, InToolTip, InGrouping) 
 		, NodeTemplate(NULL)
 	{}
@@ -91,7 +91,7 @@ struct BLUEPRINTGRAPH_API FEdGraphSchemaAction_K2ViewNode : public FEdGraphSchem
 		, NodePtr(NULL)
 	{}
 
-	FEdGraphSchemaAction_K2ViewNode(const FString& InNodeCategory, const FString& InMenuDesc, const FString& InToolTip, const int32 InGrouping)
+	FEdGraphSchemaAction_K2ViewNode(const FString& InNodeCategory, const FText& InMenuDesc, const FString& InToolTip, const int32 InGrouping)
 		: FEdGraphSchemaAction(InNodeCategory, InMenuDesc, InToolTip, InGrouping) 
 		, NodePtr(NULL)
 	{}
@@ -120,7 +120,7 @@ struct BLUEPRINTGRAPH_API FEdGraphSchemaAction_K2AssignDelegate : public FEdGrap
 		:FEdGraphSchemaAction_K2NewNode()
 	{}
 
-	FEdGraphSchemaAction_K2AssignDelegate(const FString& InNodeCategory, const FString& InMenuDesc, const FString& InToolTip, const int32 InGrouping)
+	FEdGraphSchemaAction_K2AssignDelegate(const FString& InNodeCategory, const FText& InMenuDesc, const FString& InToolTip, const int32 InGrouping)
 		: FEdGraphSchemaAction_K2NewNode(InNodeCategory, InMenuDesc, InToolTip, InGrouping) 
 	{}
 
@@ -152,7 +152,7 @@ struct BLUEPRINTGRAPH_API FEdGraphSchemaAction_EventFromFunction : public FEdGra
 		, SignatureFunction(NULL)
 	{}
 
-	FEdGraphSchemaAction_EventFromFunction(const FString& InNodeCategory, const FString& InMenuDesc, const FString& InToolTip, const int32 InGrouping)
+	FEdGraphSchemaAction_EventFromFunction(const FString& InNodeCategory, const FText& InMenuDesc, const FString& InToolTip, const int32 InGrouping)
 		: FEdGraphSchemaAction(InNodeCategory, InMenuDesc, InToolTip, InGrouping) 
 		, SignatureFunction(NULL)
 	{}
@@ -192,7 +192,7 @@ struct FEdGraphSchemaAction_K2AddComponent : public FEdGraphSchemaAction_K2NewNo
 		, ComponentAsset(NULL)
 	{}
 
-	FEdGraphSchemaAction_K2AddComponent(const FString& InNodeCategory, const FString& InMenuDesc, const FString& InToolTip, const int32 InGrouping)
+	FEdGraphSchemaAction_K2AddComponent(const FString& InNodeCategory, const FText& InMenuDesc, const FString& InToolTip, const int32 InGrouping)
 		: FEdGraphSchemaAction_K2NewNode(InNodeCategory, InMenuDesc, InToolTip, InGrouping)
 		, ComponentClass(NULL)
 		, ComponentAsset(NULL)
@@ -222,7 +222,7 @@ struct FEdGraphSchemaAction_K2AddTimeline : public FEdGraphSchemaAction_K2NewNod
 		: FEdGraphSchemaAction_K2NewNode()
 	{}
 
-	FEdGraphSchemaAction_K2AddTimeline(const FString& InNodeCategory, const FString& InMenuDesc, const FString& InToolTip, const int32 InGrouping)
+	FEdGraphSchemaAction_K2AddTimeline(const FString& InNodeCategory, const FText& InMenuDesc, const FString& InToolTip, const int32 InGrouping)
 		: FEdGraphSchemaAction_K2NewNode(InNodeCategory, InMenuDesc, InToolTip, InGrouping)
 	{}
 
@@ -249,7 +249,7 @@ struct FEdGraphSchemaAction_K2AddEvent : public FEdGraphSchemaAction_K2NewNode
 		: FEdGraphSchemaAction_K2NewNode()
 	{}
 
-	FEdGraphSchemaAction_K2AddEvent(const FString& InNodeCategory, const FString& InMenuDesc, const FString& InToolTip, const int32 InGrouping)
+	FEdGraphSchemaAction_K2AddEvent(const FString& InNodeCategory, const FText& InMenuDesc, const FString& InToolTip, const int32 InGrouping)
 		: FEdGraphSchemaAction_K2NewNode(InNodeCategory, InMenuDesc, InToolTip, InGrouping)
 	{}
 
@@ -279,7 +279,7 @@ struct FEdGraphSchemaAction_K2AddCustomEvent : public FEdGraphSchemaAction_K2New
 		: FEdGraphSchemaAction_K2NewNode()
 	{}
 
-	FEdGraphSchemaAction_K2AddCustomEvent(const FString& InNodeCategory, const FString& InMenuDesc, const FString& InToolTip, const int32 InGrouping)
+	FEdGraphSchemaAction_K2AddCustomEvent(const FString& InNodeCategory, const FText& InMenuDesc, const FString& InToolTip, const int32 InGrouping)
 		: FEdGraphSchemaAction_K2NewNode(InNodeCategory, InMenuDesc, InToolTip, InGrouping)
 	{}
 
@@ -310,7 +310,7 @@ struct FEdGraphSchemaAction_K2AddCallOnActor : public FEdGraphSchemaAction_K2New
 		: FEdGraphSchemaAction_K2NewNode()
 	{}
 
-	FEdGraphSchemaAction_K2AddCallOnActor(const FString& InNodeCategory, const FString& InMenuDesc, const FString& InToolTip, const int32 InGrouping)
+	FEdGraphSchemaAction_K2AddCallOnActor(const FString& InNodeCategory, const FText& InMenuDesc, const FString& InToolTip, const int32 InGrouping)
 		: FEdGraphSchemaAction_K2NewNode(InNodeCategory, InMenuDesc, InToolTip, InGrouping)
 	{}
 
@@ -338,7 +338,7 @@ struct FEdGraphSchemaAction_K2AddCallOnVariable : public FEdGraphSchemaAction_K2
 		: FEdGraphSchemaAction_K2NewNode()
 	{}
 
-	FEdGraphSchemaAction_K2AddCallOnVariable(const FString& InNodeCategory, const FString& InMenuDesc, const FString& InToolTip, const int32 InGrouping)
+	FEdGraphSchemaAction_K2AddCallOnVariable(const FString& InNodeCategory, const FText& InMenuDesc, const FString& InToolTip, const int32 InGrouping)
 		: FEdGraphSchemaAction_K2NewNode(InNodeCategory, InMenuDesc, InToolTip, InGrouping)
 	{}
 
@@ -365,11 +365,11 @@ struct FEdGraphSchemaAction_K2AddComment : public FEdGraphSchemaAction
 		: FEdGraphSchemaAction()
 	{
 		Category = TEXT("");
-		MenuDescription = TEXT("Add Comment...");
-		TooltipDescription = TEXT("Create a resizable comment box.");
+		MenuDescription = NSLOCTEXT("K2AddComment", "AddComment", "Add Comment...");
+		TooltipDescription = NSLOCTEXT("K2AddComment", "AddComment_Tooltip", "Create a resizable comment box.").ToString();
 	}
 
-	FEdGraphSchemaAction_K2AddComment(const FString& InDescription, const FString& InToolTip)
+	FEdGraphSchemaAction_K2AddComment(const FText& InDescription, const FString& InToolTip)
 		: FEdGraphSchemaAction()
 	{
 		Category = TEXT("");
@@ -401,7 +401,7 @@ struct BLUEPRINTGRAPH_API FEdGraphSchemaAction_K2TargetNode : public FEdGraphSch
 		: FEdGraphSchemaAction_K2NewNode()
 	{}
 
-	FEdGraphSchemaAction_K2TargetNode(const FString& InNodeCategory, const FString& InMenuDesc, const FString& InToolTip, const int32 InGrouping)
+	FEdGraphSchemaAction_K2TargetNode(const FString& InNodeCategory, const FText& InMenuDesc, const FString& InToolTip, const int32 InGrouping)
 		: FEdGraphSchemaAction_K2NewNode(InNodeCategory, InMenuDesc, InToolTip, InGrouping)
 	{}
 
@@ -428,7 +428,7 @@ struct FEdGraphSchemaAction_K2PasteHere : public FEdGraphSchemaAction_K2NewNode
 		: FEdGraphSchemaAction_K2NewNode()
 	{}
 
-	FEdGraphSchemaAction_K2PasteHere (const FString& InNodeCategory, const FString& InMenuDesc, const FString& InToolTip, const int32 InGrouping)
+	FEdGraphSchemaAction_K2PasteHere (const FString& InNodeCategory, const FText& InMenuDesc, const FString& InToolTip, const int32 InGrouping)
 		: FEdGraphSchemaAction_K2NewNode(InNodeCategory, InMenuDesc, InToolTip, InGrouping)
 	{}
 
@@ -470,7 +470,7 @@ struct BLUEPRINTGRAPH_API FEdGraphSchemaAction_K2Enum : public FEdGraphSchemaAct
 		: FEdGraphSchemaAction()
 	{}
 
-	FEdGraphSchemaAction_K2Enum (const FString& InNodeCategory, const FString& InMenuDesc, const FString& InToolTip, const int32 InGrouping)
+	FEdGraphSchemaAction_K2Enum (const FString& InNodeCategory, const FText& InMenuDesc, const FString& InToolTip, const int32 InGrouping)
 		: FEdGraphSchemaAction(InNodeCategory, InMenuDesc, InToolTip, InGrouping)
 	{}
 };
@@ -509,7 +509,7 @@ public:
 		: FEdGraphSchemaAction()
 	{}
 
-	FEdGraphSchemaAction_K2Var (const FString& InNodeCategory, const FString& InMenuDesc, const FString& InToolTip, const int32 InGrouping)
+	FEdGraphSchemaAction_K2Var (const FString& InNodeCategory, const FText& InMenuDesc, const FString& InToolTip, const int32 InGrouping)
 		: FEdGraphSchemaAction(InNodeCategory, InMenuDesc, InToolTip, InGrouping)
 	{}
 
@@ -575,7 +575,7 @@ struct BLUEPRINTGRAPH_API FEdGraphSchemaAction_K2Graph : public FEdGraphSchemaAc
 		: FEdGraphSchemaAction()
 	{}
 
-	FEdGraphSchemaAction_K2Graph (EEdGraphSchemaAction_K2Graph::Type InType, const FString& InNodeCategory, const FString& InMenuDesc, const FString& InToolTip, const int32 InGrouping)
+	FEdGraphSchemaAction_K2Graph (EEdGraphSchemaAction_K2Graph::Type InType, const FString& InNodeCategory, const FText& InMenuDesc, const FString& InToolTip, const int32 InGrouping)
 		: FEdGraphSchemaAction(InNodeCategory, InMenuDesc, InToolTip, InGrouping), EdGraph(NULL), GraphType(InType)
 	{}
 
@@ -613,7 +613,7 @@ struct BLUEPRINTGRAPH_API FEdGraphSchemaAction_K2Event : public FEdGraphSchemaAc
 	 * @param   Tooltip			A string to display when hovering over this action entry.
 	 * @param   Grouping		Used to override list ordering (actions with the same number get grouped together, higher numbers get sorted first).
 	 */
-	FEdGraphSchemaAction_K2Event(FString const& Category, FString const& MenuDescription, FString const& Tooltip, int32 const Grouping)
+	FEdGraphSchemaAction_K2Event(FString const& Category, FText const& MenuDescription, FString const& Tooltip, int32 const Grouping)
 		: FEdGraphSchemaAction_K2TargetNode(Category, MenuDescription, Tooltip, Grouping)
 	{}
 
@@ -671,7 +671,7 @@ public:
 		: FEdGraphSchemaAction(), EdGraph(NULL)
 	{}
 
-	FEdGraphSchemaAction_K2Delegate(const FString& InNodeCategory, const FString& InMenuDesc, const FString& InToolTip, const int32 InGrouping)
+	FEdGraphSchemaAction_K2Delegate(const FString& InNodeCategory, const FText& InMenuDesc, const FString& InToolTip, const int32 InGrouping)
 		: FEdGraphSchemaAction(InNodeCategory, InMenuDesc, InToolTip, InGrouping), EdGraph(NULL)
 	{}
 
diff --git a/Engine/Source/Editor/BlueprintGraph/Classes/K2Node.h b/Engine/Source/Editor/BlueprintGraph/Classes/K2Node.h
--- a/Engine/Source/Editor/BlueprintGraph/Classes/K2Node.h
+++ b/Engine/Source/Editor/BlueprintGraph/Classes/K2Node.h
@@ -151,7 +151,7 @@ class UK2Node : public UEdGraphNode
 	virtual bool ShouldDrawCompact() const { return false; }
 
 	/** Return title if drawing this node in 'compact' mode */
-	virtual FString GetCompactNodeTitle() const { return GetNodeTitle(ENodeTitleType::FullTitle); }
+	virtual FText GetCompactNodeTitle() const { return GetNodeTitle(ENodeTitleType::FullTitle); }
 
 	/** */
 	BLUEPRINTGRAPH_API virtual FText GetToolTipHeading() const;
diff --git a/Engine/Source/Editor/BlueprintGraph/Classes/K2Node_CallFunction.h b/Engine/Source/Editor/BlueprintGraph/Classes/K2Node_CallFunction.h
--- a/Engine/Source/Editor/BlueprintGraph/Classes/K2Node_CallFunction.h
+++ b/Engine/Source/Editor/BlueprintGraph/Classes/K2Node_CallFunction.h
@@ -58,7 +58,8 @@ public:
 	virtual void DestroyNode() OVERRIDE;
 	virtual FLinearColor GetNodeTitleColor() const OVERRIDE;
 	virtual FString GetTooltip() const OVERRIDE;
-	virtual FString GetNodeTitle(ENodeTitleType::Type TitleType) const OVERRIDE;
+	virtual FText GetNodeTitle(ENodeTitleType::Type TitleType) const OVERRIDE;
+	virtual FString GetNodeNativeTitle(ENodeTitleType::Type TitleType) const OVERRIDE;
 	virtual FString GetDescriptiveCompiledName() const OVERRIDE;
 	virtual bool IsDeprecated() const OVERRIDE;
 	virtual bool ShouldWarnOnDeprecation() const OVERRIDE;
@@ -76,7 +77,7 @@ public:
 	virtual void PostReconstructNode() OVERRIDE;
 	virtual bool ShouldDrawCompact() const OVERRIDE;
 	virtual bool ShouldDrawAsBead() const OVERRIDE;
-	virtual FString GetCompactNodeTitle() const OVERRIDE;
+	virtual FText GetCompactNodeTitle() const OVERRIDE;
 	virtual void PostPasteNode() OVERRIDE;
 	virtual void ValidateNodeDuringCompilation(class FCompilerResultsLog& MessageLog) const OVERRIDE;
 	virtual bool ShouldShowNodeProperties() const OVERRIDE;
diff --git a/Engine/Source/Editor/BlueprintGraph/Classes/K2Node_CastByteToEnum.h b/Engine/Source/Editor/BlueprintGraph/Classes/K2Node_CastByteToEnum.h
--- a/Engine/Source/Editor/BlueprintGraph/Classes/K2Node_CastByteToEnum.h
+++ b/Engine/Source/Editor/BlueprintGraph/Classes/K2Node_CastByteToEnum.h
@@ -20,12 +20,12 @@ class UK2Node_CastByteToEnum : public UK2Node
 	// Begin UEdGraphNode interface
 	virtual void AllocateDefaultPins() OVERRIDE;
 	virtual FString GetTooltip() const OVERRIDE;
-	virtual FString GetNodeTitle(ENodeTitleType::Type TitleType) const OVERRIDE;
+	virtual FText GetNodeTitle(ENodeTitleType::Type TitleType) const OVERRIDE;
 	virtual FName GetPaletteIcon(FLinearColor& OutColor) const OVERRIDE{ return TEXT("GraphEditor.Enum_16x"); }
 	// End UEdGraphNode interface
 
 	// Begin UK2Node interface
-	virtual FString GetCompactNodeTitle() const OVERRIDE;
+	virtual FText GetCompactNodeTitle() const OVERRIDE;
 	virtual bool ShouldDrawCompact() const OVERRIDE { return true; }
 	virtual bool IsNodePure() const OVERRIDE { return true; }
 	FNodeHandlingFunctor* CreateNodeHandler(FKismetCompilerContext& CompilerContext) const OVERRIDE;
diff --git a/Engine/Source/Editor/BlueprintGraph/Classes/K2Node_Composite.h b/Engine/Source/Editor/BlueprintGraph/Classes/K2Node_Composite.h
--- a/Engine/Source/Editor/BlueprintGraph/Classes/K2Node_Composite.h
+++ b/Engine/Source/Editor/BlueprintGraph/Classes/K2Node_Composite.h
@@ -18,7 +18,8 @@ class UK2Node_Composite : public UK2Node_Tunnel
 	BLUEPRINTGRAPH_API virtual void PostPasteNode() OVERRIDE;
 	BLUEPRINTGRAPH_API virtual FString GetTooltip() const OVERRIDE;
 	BLUEPRINTGRAPH_API virtual FLinearColor GetNodeTitleColor() const OVERRIDE;
-	BLUEPRINTGRAPH_API virtual FString GetNodeTitle(ENodeTitleType::Type TitleType) const OVERRIDE;
+	BLUEPRINTGRAPH_API virtual FText GetNodeTitle(ENodeTitleType::Type TitleType) const OVERRIDE;
+	BLUEPRINTGRAPH_API virtual FString GetNodeNativeTitle(ENodeTitleType::Type TitleType) const OVERRIDE;
 	BLUEPRINTGRAPH_API virtual bool CanUserDeleteNode() const OVERRIDE;
 	BLUEPRINTGRAPH_API virtual UObject* GetJumpTargetForDoubleClick() const OVERRIDE;
 	BLUEPRINTGRAPH_API virtual void PostPlacedNewNode() OVERRIDE;
diff --git a/Engine/Source/Editor/BlueprintGraph/Classes/K2Node_EnumEquality.h b/Engine/Source/Editor/BlueprintGraph/Classes/K2Node_EnumEquality.h
--- a/Engine/Source/Editor/BlueprintGraph/Classes/K2Node_EnumEquality.h
+++ b/Engine/Source/Editor/BlueprintGraph/Classes/K2Node_EnumEquality.h
@@ -12,7 +12,7 @@ class UK2Node_EnumEquality : public UK2Node
 	// Begin UEdGraphNode interface
 	virtual void AllocateDefaultPins() OVERRIDE;
 	virtual FString GetTooltip() const OVERRIDE;
-	virtual FString GetNodeTitle(ENodeTitleType::Type TitleType) const OVERRIDE;
+	virtual FText GetNodeTitle(ENodeTitleType::Type TitleType) const OVERRIDE;
 	virtual FString GetKeywords() const OVERRIDE;
 	// End UEdGraphNode interface
 
@@ -21,7 +21,7 @@ class UK2Node_EnumEquality : public UK2Node
 	virtual bool IsNodeSafeToIgnore() const OVERRIDE { return true; }
 	virtual bool IsNodePure() const OVERRIDE { return true; }
 	virtual bool ShouldDrawCompact() const OVERRIDE { return true; }
-	virtual FString GetCompactNodeTitle() const OVERRIDE { return TEXT("=="); }
+	virtual FText GetCompactNodeTitle() const OVERRIDE { return NSLOCTEXT("K2Node", "EqualEqual", "=="); }
 	virtual void GetMenuEntries(FGraphContextMenuBuilder& ContextMenuBuilder) const OVERRIDE;
 	virtual void NotifyPinConnectionListChanged(UEdGraphPin* Pin) OVERRIDE;
 	virtual class FNodeHandlingFunctor* CreateNodeHandler(class FKismetCompilerContext& CompilerContext) const OVERRIDE;
diff --git a/Engine/Source/Editor/BlueprintGraph/Classes/K2Node_Event.h b/Engine/Source/Editor/BlueprintGraph/Classes/K2Node_Event.h
--- a/Engine/Source/Editor/BlueprintGraph/Classes/K2Node_Event.h
+++ b/Engine/Source/Editor/BlueprintGraph/Classes/K2Node_Event.h
@@ -38,7 +38,8 @@ class UK2Node_Event : public UK2Node_EditablePinBase
 	BLUEPRINTGRAPH_API virtual void AllocateDefaultPins() OVERRIDE;
 	BLUEPRINTGRAPH_API virtual FString GetTooltip() const OVERRIDE;
 	BLUEPRINTGRAPH_API virtual FString GetKeywords() const OVERRIDE;	
-	BLUEPRINTGRAPH_API virtual FString GetNodeTitle(ENodeTitleType::Type TitleType) const OVERRIDE;
+	BLUEPRINTGRAPH_API virtual FText GetNodeTitle(ENodeTitleType::Type TitleType) const OVERRIDE;
+	BLUEPRINTGRAPH_API virtual FString GetNodeNativeTitle(ENodeTitleType::Type TitleType) const OVERRIDE;
 	BLUEPRINTGRAPH_API virtual FLinearColor GetNodeTitleColor() const OVERRIDE;
 	BLUEPRINTGRAPH_API virtual bool CanPasteHere(const UEdGraph* TargetGraph, const UEdGraphSchema* Schema) const OVERRIDE;
 	BLUEPRINTGRAPH_API virtual FName GetCornerIcon() const OVERRIDE;
diff --git a/Engine/Source/Editor/BlueprintGraph/Classes/K2Node_GetEnumeratorName.h b/Engine/Source/Editor/BlueprintGraph/Classes/K2Node_GetEnumeratorName.h
--- a/Engine/Source/Editor/BlueprintGraph/Classes/K2Node_GetEnumeratorName.h
+++ b/Engine/Source/Editor/BlueprintGraph/Classes/K2Node_GetEnumeratorName.h
@@ -13,13 +13,13 @@ class UK2Node_GetEnumeratorName : public UK2Node
 	// Begin UEdGraphNode interface
 	virtual void AllocateDefaultPins() OVERRIDE;
 	virtual FString GetTooltip() const OVERRIDE;
-	virtual FString GetNodeTitle(ENodeTitleType::Type TitleType) const OVERRIDE;
+	virtual FText GetNodeTitle(ENodeTitleType::Type TitleType) const OVERRIDE;
 	virtual void ValidateNodeDuringCompilation(class FCompilerResultsLog& MessageLog) const OVERRIDE;
 	virtual FName GetPaletteIcon(FLinearColor& OutColor) const OVERRIDE{ return TEXT("GraphEditor.Enum_16x"); }
 	// End UEdGraphNode interface
 
 	// Begin UK2Node interface
-	virtual FString GetCompactNodeTitle() const OVERRIDE;
+	virtual FText GetCompactNodeTitle() const OVERRIDE;
 	virtual bool ShouldDrawCompact() const OVERRIDE { return true; }
 	virtual bool IsNodePure() const OVERRIDE { return true; }
 	virtual void ExpandNode(class FKismetCompilerContext& CompilerContext, UEdGraph* SourceGraph) OVERRIDE;
diff --git a/Engine/Source/Editor/BlueprintGraph/Private/EdGraphSchema_K2.cpp b/Engine/Source/Editor/BlueprintGraph/Private/EdGraphSchema_K2.cpp
--- a/Engine/Source/Editor/BlueprintGraph/Private/EdGraphSchema_K2.cpp
+++ b/Engine/Source/Editor/BlueprintGraph/Private/EdGraphSchema_K2.cpp
@@ -913,8 +913,8 @@ void UEdGraphSchema_K2::GetBreakLinkToSubMenuActions( class FMenuBuilder& MenuBu
 	for(TArray<class UEdGraphPin*>::TConstIterator Links(InGraphPin->LinkedTo); Links; ++Links)
 	{
 		UEdGraphPin* Pin = *Links;
-		FString TitleString = Pin->GetOwningNode()->GetNodeTitle(ENodeTitleType::ListView);
-		FText Title = FText::FromString( TitleString );
+		FText Title = Pin->GetOwningNode()->GetNodeTitle(ENodeTitleType::ListView);
+		FString TitleString = Title.ToString();
 		if ( Pin->PinName != TEXT("") )
 		{
 			TitleString = FString::Printf(TEXT("%s (%s)"), *TitleString, *Pin->PinName);
@@ -956,8 +956,8 @@ void UEdGraphSchema_K2::GetJumpToConnectionSubMenuActions( class FMenuBuilder& M
 	// Add all the links we could break from
 	for(auto PinLink : InGraphPin->LinkedTo )
 	{
-		FString TitleString = PinLink->GetOwningNode()->GetNodeTitle(ENodeTitleType::ListView);
-		FText Title = FText::FromString( TitleString );
+		FText Title = PinLink->GetOwningNode()->GetNodeTitle(ENodeTitleType::ListView);
+		FString TitleString = Title.ToString();
 		if ( PinLink->PinName != TEXT("") )
 		{
 			TitleString = FString::Printf(TEXT("%s (%s)"), *TitleString, *PinLink->PinName);
@@ -2160,50 +2160,7 @@ bool UEdGraphSchema_K2::ConvertPropertyToPinType(const UProperty* Property, /*ou
 
 FString UEdGraphSchema_K2::TypeToString(const FEdGraphPinType& Type)
 {
-	FString PropertyString;
-
-	if (Type.PinSubCategoryObject != NULL)
-	{
-		const UEdGraphSchema_K2* Schema = GetDefault<UEdGraphSchema_K2>();
-		if (Type.PinCategory == Schema->PC_Byte)
-		{
-			PropertyString = FString::Printf(TEXT("enum'%s'"), *Type.PinSubCategoryObject->GetName());
-		}
-		else
-		{
-			if( !Type.bIsWeakPointer )
-			{
-				UClass* PSCOAsClass = Cast<UClass>(Type.PinSubCategoryObject.Get());
-				const bool bIsInterface = PSCOAsClass && PSCOAsClass->HasAnyClassFlags(CLASS_Interface);
-
-				FString CategoryDesc = !bIsInterface ? (*Type.PinCategory) : TEXT("interface");
-				PropertyString = FString::Printf(TEXT("%s'%s'"), *CategoryDesc, *Type.PinSubCategoryObject.Get()->GetName());
-			}
-			else
-			{
-				PropertyString = FString::Printf(TEXT("weak_ptr_%s'%s'"), *Type.PinCategory, *Type.PinSubCategoryObject->GetName());
-			}
-		}
-	}
-	else if (Type.PinSubCategory != TEXT(""))
-	{
-		PropertyString = FString::Printf(TEXT("%s'%s'"), *Type.PinCategory, *Type.PinSubCategory);
-	}
-	else
-	{
-		PropertyString = Type.PinCategory;
-	}
-
-	if (Type.bIsArray)
-	{
-		PropertyString = FString::Printf(TEXT("array[%s]"), *PropertyString);
-	}
-	else if (Type.bIsReference)
-	{
-		PropertyString += TEXT(" (by ref)"); 
-	}
-
-	return PropertyString;
+	return TypeToText(Type).ToString();
 }
 
 FString UEdGraphSchema_K2::TypeToString(UProperty* const Property)
@@ -2246,6 +2203,70 @@ FString UEdGraphSchema_K2::TypeToString(UProperty* const Property)
 	}
 }
 
+FText UEdGraphSchema_K2::TypeToText(const FEdGraphPinType& Type)
+{
+	FText PropertyText;
+
+	if (Type.PinSubCategoryObject != NULL)
+	{
+		const UEdGraphSchema_K2* Schema = GetDefault<UEdGraphSchema_K2>();
+		if (Type.PinCategory == Schema->PC_Byte)
+		{
+			FFormatNamedArguments Args;
+			Args.Add(TEXT("EnumName"), FText::FromString(Type.PinSubCategoryObject->GetName()));
+			PropertyText = FText::Format(LOCTEXT("EnumAsText", "enum'{EnumName}'"), Args);
+		}
+		else
+		{
+			if( !Type.bIsWeakPointer )
+			{
+				UClass* PSCOAsClass = Cast<UClass>(Type.PinSubCategoryObject.Get());
+				const bool bIsInterface = PSCOAsClass && PSCOAsClass->HasAnyClassFlags(CLASS_Interface);
+
+				FText CategoryDesc = !bIsInterface ? FText::FromString(Type.PinCategory) : LOCTEXT("Interface", "interface");
+
+				FFormatNamedArguments Args;
+				Args.Add(TEXT("Category"), CategoryDesc);
+				Args.Add(TEXT("ObjectName"), FText::FromString(Type.PinSubCategoryObject.Get()->GetName()));
+				PropertyText = FText::Format(LOCTEXT("ObjectAsText", "{Category}'{ObjectName}'"), Args);
+			}
+			else
+			{
+				FFormatNamedArguments Args;
+				Args.Add(TEXT("Category"), FText::FromString(Type.PinCategory));
+				Args.Add(TEXT("ObjectName"), FText::FromString(Type.PinSubCategoryObject.Get()->GetName()));
+				PropertyText = FText::Format(LOCTEXT("WeakPtrAsText", "weak_ptr_{Category}'{ObjectName}'"), Args);
+			}
+		}
+	}
+	else if (Type.PinSubCategory != TEXT(""))
+	{
+		FFormatNamedArguments Args;
+		Args.Add(TEXT("Category"), FText::FromString(Type.PinCategory));
+		Args.Add(TEXT("ObjectName"), FText::FromString(Type.PinSubCategory));
+		PropertyText = FText::Format(LOCTEXT("ObjectAsText", "{Category}'{ObjectName}'"), Args);
+	}
+	else
+	{
+		PropertyText = FText::FromString(Type.PinCategory);
+	}
+
+	if (Type.bIsArray)
+	{
+		FFormatNamedArguments Args;
+		Args.Add(TEXT("PropertyTitle"), PropertyText);
+		PropertyText = FText::Format(LOCTEXT("ArrayAsText", "array[{PropertyTitle}]"), Args);
+	}
+	else if (Type.bIsReference)
+	{
+		FFormatNamedArguments Args;
+		Args.Add(TEXT("PropertyTitle"), PropertyText);
+		PropertyText = FText::Format(LOCTEXT("PropertyByRef", "{PropertyTitle} (by ref)"), Args);
+	}
+
+	return PropertyText;
+}
+
 void UEdGraphSchema_K2::GetVariableTypeTree( TArray< TSharedPtr<FPinTypeTreeInfo> >& TypeTree, bool bAllowExec, bool bAllowWildCard ) const
 {
 	// Clear the list
@@ -2955,7 +2976,7 @@ bool UEdGraphSchema_K2::CanPromotePinToVariable( const UEdGraphPin& Pin ) const
 void UEdGraphSchema_K2::GetGraphDisplayInformation(const UEdGraph& Graph, /*out*/ FGraphDisplayInfo& DisplayInfo) const
 {
 	DisplayInfo.DocLink = TEXT("Shared/Editors/BlueprintEditor/GraphTypes");
-	DisplayInfo.DisplayName = FText::FromString( Graph.GetName() ); // Fallback is graph name
+	DisplayInfo.PlainName = FText::FromString( Graph.GetName() ); // Fallback is graph name
 
 	UFunction* Function = NULL;
 	UBlueprint* Blueprint = FBlueprintEditorUtils::FindBlueprintForGraph(&Graph);
@@ -2972,8 +2993,8 @@ void UEdGraphSchema_K2::GetGraphDisplayInformation(const UEdGraph& Graph, /*out*
 		if (Graph.GetFName() == GN_EventGraph)
 		{
 			// localized name for the first event graph
-			DisplayInfo.DisplayName = LOCTEXT("GraphDisplayName_EventGraph", "EventGraph");
-			DisplayInfo.Tooltip = DisplayInfo.DisplayName.ToString();
+			DisplayInfo.PlainName = LOCTEXT("GraphDisplayName_EventGraph", "EventGraph");
+			DisplayInfo.Tooltip = DisplayInfo.PlainName.ToString();
 		}
 		else
 		{
@@ -2984,7 +3005,7 @@ void UEdGraphSchema_K2::GetGraphDisplayInformation(const UEdGraph& Graph, /*out*
 	{
 		if ( Graph.GetFName() == FN_UserConstructionScript )
 		{
-			DisplayInfo.DisplayName = LOCTEXT("GraphDisplayName_ConstructionScript", "ConstructionScript");
+			DisplayInfo.PlainName = LOCTEXT("GraphDisplayName_ConstructionScript", "ConstructionScript");
 
 			DisplayInfo.Tooltip = LOCTEXT("GraphTooltip_ConstructionScript", "Function executed when Blueprint is placed or modified.").ToString();
 			DisplayInfo.DocExcerptName = TEXT("ConstructionScript");
@@ -3014,7 +3035,7 @@ void UEdGraphSchema_K2::GetGraphDisplayInformation(const UEdGraph& Graph, /*out*
 	}
 	else if (GraphType == GT_Animation)
 	{
-		DisplayInfo.DisplayName = LOCTEXT("GraphDisplayName_AnimGraph", "AnimGraph");
+		DisplayInfo.PlainName = LOCTEXT("GraphDisplayName_AnimGraph", "AnimGraph");
 
 		DisplayInfo.Tooltip = LOCTEXT("GraphTooltip_AnimGraph", "Graph used to blend together different animations.").ToString();
 		DisplayInfo.DocExcerptName = TEXT("AnimGraph");
@@ -3034,9 +3055,18 @@ void UEdGraphSchema_K2::GetGraphDisplayInformation(const UEdGraph& Graph, /*out*
 	// Mark transient graphs as obviously so
 	if (Graph.HasAllFlags(RF_Transient))
 	{
-		DisplayInfo.DisplayName = FText::FromString( FString::Printf(TEXT("$$ %s $$"), *DisplayInfo.DisplayName.ToString()) );
+		DisplayInfo.PlainName = FText::FromString( FString::Printf(TEXT("$$ %s $$"), *DisplayInfo.PlainName.ToString()) );
 		DisplayInfo.Notes.Add(TEXT("intermediate build product"));
 	}
+
+	if( GEditor && GetDefault<UEditorStyleSettings>()->bShowFriendlyNames )
+	{
+		DisplayInfo.DisplayName = FText::FromString(FName::NameToDisplayString( DisplayInfo.PlainName.ToString(), false ));
+	}
+	else
+	{
+		DisplayInfo.DisplayName = DisplayInfo.PlainName;
+	}
 }
 
 bool UEdGraphSchema_K2::IsSelfPin(const UEdGraphPin& Pin) const 
