#summary Can search for nodes in the palette menu in Blueprints with both the native name and the localized name.

diff --git a/Engine/Source/Editor/BlueprintGraph/Private/K2ActionMenuBuilder.cpp b/Engine/Source/Editor/BlueprintGraph/Private/K2ActionMenuBuilder.cpp
--- a/Engine/Source/Editor/BlueprintGraph/Private/K2ActionMenuBuilder.cpp
+++ b/Engine/Source/Editor/BlueprintGraph/Private/K2ActionMenuBuilder.cpp
@@ -96,6 +96,7 @@ static void GetRequestedAction(FGraphActionListBuilderBase& ActionMenuBuilder, c
 	UK2Node* NodeTemplate = NewObject<ActionType>(ActionMenuBuilder.OwnerOfTemporaries);
 	TSharedPtr<FEdGraphSchemaAction_K2NewNode> Action = FK2ActionMenuBuilder::AddNewNodeAction(ActionMenuBuilder, Category, NodeTemplate->GetNodeTitle(ENodeTitleType::ListView), NodeTemplate->GetTooltip(), 0, NodeTemplate->GetKeywords());
 	Action->NodeTemplate = NodeTemplate;
+	Action->SearchTitle = Action->NodeTemplate->GetNodeSearchTitle();
 }
 
 /** Utility for creating spawn from archetype nodes */
@@ -106,7 +107,7 @@ static void AddSpawnActorNodeAction(FGraphActionListBuilderBase& ContextMenuBuil
 	// Add the new SpawnActorFromClass node
 	{
 		UK2Node* NodeTemplate = ContextMenuBuilder.CreateTemplateNode<UK2Node_SpawnActorFromClass>();
-		TSharedPtr<FEdGraphSchemaAction_K2NewNode> Action = FK2ActionMenuBuilder::AddNewNodeAction(ContextMenuBuilder, SpawnActorCategory, TEXT("Spawn Actor from Class"), NodeTemplate->GetTooltip(), 0, NodeTemplate->GetKeywords());
+		TSharedPtr<FEdGraphSchemaAction_K2NewNode> Action = FK2ActionMenuBuilder::AddNewNodeAction(ContextMenuBuilder, SpawnActorCategory, LOCTEXT("SpawnActorFromClass", "Spawn Actor from Class"), NodeTemplate->GetTooltip(), 0, NodeTemplate->GetKeywords());
 		Action->NodeTemplate = NodeTemplate;
 	}
 }
@@ -120,10 +121,11 @@ static void GetInputNodes(FGraphActionListBuilderBase& ActionMenuBuilder, const
 		GetDefault<UInputSettings>()->GetActionNames(ActionNames);
 		for (const FName InputActionName : ActionNames)
 		{
-			TSharedPtr<FEdGraphSchemaAction_K2NewNode> Action = FK2ActionMenuBuilder::AddNewNodeAction(ActionMenuBuilder, K2ActionCategories::ActionEventsCategory, InputActionName.ToString(), InputActionName.ToString());
+			TSharedPtr<FEdGraphSchemaAction_K2NewNode> Action = FK2ActionMenuBuilder::AddNewNodeAction(ActionMenuBuilder, K2ActionCategories::ActionEventsCategory, FText::FromName(InputActionName), InputActionName.ToString());
 			UK2Node_InputAction* InputActionNode = ActionMenuBuilder.CreateTemplateNode<UK2Node_InputAction>();
 			InputActionNode->InputActionName = InputActionName;
 			Action->NodeTemplate = InputActionNode;
+			Action->SearchTitle = Action->NodeTemplate->GetNodeSearchTitle();
 			Action->Keywords = TEXT("InputAction");
 		}
 
@@ -131,16 +133,18 @@ static void GetInputNodes(FGraphActionListBuilderBase& ActionMenuBuilder, const
 		GetDefault<UInputSettings>()->GetAxisNames(AxisNames);
 		for (const FName InputAxisName : AxisNames)
 		{
-			TSharedPtr<FEdGraphSchemaAction_K2NewNode> GetAction = FK2ActionMenuBuilder::AddNewNodeAction(ActionMenuBuilder, K2ActionCategories::AxisCategory, InputAxisName.ToString(), InputAxisName.ToString());
+			TSharedPtr<FEdGraphSchemaAction_K2NewNode> GetAction = FK2ActionMenuBuilder::AddNewNodeAction(ActionMenuBuilder, K2ActionCategories::AxisCategory, FText::FromName(InputAxisName), InputAxisName.ToString());
 			UK2Node_GetInputAxisValue* GetInputAxisValue = ActionMenuBuilder.CreateTemplateNode<UK2Node_GetInputAxisValue>();
 			GetInputAxisValue->Initialize(InputAxisName);
 			GetAction->NodeTemplate = GetInputAxisValue;
+			GetAction->SearchTitle = GetAction->NodeTemplate->GetNodeSearchTitle();
 			GetAction->Keywords = TEXT("InputAxis");
 
-			TSharedPtr<FEdGraphSchemaAction_K2NewNode> Action = FK2ActionMenuBuilder::AddNewNodeAction(ActionMenuBuilder, K2ActionCategories::AxisEventCategory, InputAxisName.ToString(), InputAxisName.ToString());
+			TSharedPtr<FEdGraphSchemaAction_K2NewNode> Action = FK2ActionMenuBuilder::AddNewNodeAction(ActionMenuBuilder, K2ActionCategories::AxisEventCategory, FText::FromName(InputAxisName), InputAxisName.ToString());
 			UK2Node_InputAxisEvent* InputAxisEvent = ActionMenuBuilder.CreateTemplateNode<UK2Node_InputAxisEvent>();
 			InputAxisEvent->Initialize(InputAxisName);
 			Action->NodeTemplate = InputAxisEvent;
+			Action->SearchTitle = Action->NodeTemplate->GetNodeSearchTitle();
 			Action->Keywords = TEXT("InputAxis");
 		}
 
@@ -157,7 +161,7 @@ static void GetInputNodes(FGraphActionListBuilderBase& ActionMenuBuilder, const
 				const bool bMouse = Key.IsMouseButton();
 				const FString& KeyEventCategory = (bGamepad ? K2ActionCategories::GamepadEventsCategory : (bMouse ? K2ActionCategories::MouseEventsCategory : K2ActionCategories::KeyEventsCategory));
 
-				TSharedPtr<FEdGraphSchemaAction_K2NewNode> Action = FK2ActionMenuBuilder::AddNewNodeAction(ActionMenuBuilder, KeyEventCategory, KeyName.ToString(), KeyName.ToString());
+				TSharedPtr<FEdGraphSchemaAction_K2NewNode> Action = FK2ActionMenuBuilder::AddNewNodeAction(ActionMenuBuilder, KeyEventCategory, KeyName, KeyName.ToString());
 				TSharedPtr<FEdGraphSchemaAction_K2NewNode> GetAction;
 
 				if (bGamepad)
@@ -178,9 +182,10 @@ static void GetInputNodes(FGraphActionListBuilderBase& ActionMenuBuilder, const
 					UK2Node_InputAxisKeyEvent* InputKeyAxisNode = ActionMenuBuilder.CreateTemplateNode<UK2Node_InputAxisKeyEvent>();
 					InputKeyAxisNode->Initialize(Key);
 					Action->NodeTemplate = InputKeyAxisNode;
+					Action->SearchTitle = Action->NodeTemplate->GetNodeSearchTitle();
 
 					const FString& KeyValuesCategory = (bGamepad ? K2ActionCategories::GamepadValuesCategory : (bMouse ? K2ActionCategories::MouseValuesCategory : K2ActionCategories::KeyValuesCategory));
-					GetAction = FK2ActionMenuBuilder::AddNewNodeAction(ActionMenuBuilder, KeyValuesCategory, KeyName.ToString(), KeyName.ToString());
+					GetAction = FK2ActionMenuBuilder::AddNewNodeAction(ActionMenuBuilder, KeyValuesCategory, KeyName, KeyName.ToString());
 					UK2Node_GetInputAxisKeyValue* GetInputAxisKeyValue = ActionMenuBuilder.CreateTemplateNode<UK2Node_GetInputAxisKeyValue>();
 					GetInputAxisKeyValue->Initialize(Key);
 					GetAction->NodeTemplate = GetInputAxisKeyValue;
@@ -191,13 +196,15 @@ static void GetInputNodes(FGraphActionListBuilderBase& ActionMenuBuilder, const
 					UK2Node_InputKey* InputKeyNode = ActionMenuBuilder.CreateTemplateNode<UK2Node_InputKey>();
 					InputKeyNode->InputKey = Key;
 					Action->NodeTemplate = InputKeyNode;
+					Action->SearchTitle = Action->NodeTemplate->GetNodeSearchTitle();
 				}
 			}
 		}
 
-		TSharedPtr<FEdGraphSchemaAction_K2NewNode> Action = FK2ActionMenuBuilder::AddNewNodeAction(ActionMenuBuilder, K2ActionCategories::TouchEventsCategory, TEXT("Touch"), TEXT("Touch"));
+		TSharedPtr<FEdGraphSchemaAction_K2NewNode> Action = FK2ActionMenuBuilder::AddNewNodeAction(ActionMenuBuilder, K2ActionCategories::TouchEventsCategory, LOCTEXT("Touch", "Touch"), LOCTEXT("Touch", "Touch").ToString());
 		UK2Node_InputTouch* InputTouchNode = ActionMenuBuilder.CreateTemplateNode<UK2Node_InputTouch>();
 		Action->NodeTemplate = InputTouchNode;
+		Action->SearchTitle = Action->NodeTemplate->GetNodeSearchTitle();
 		Action->Keywords = TEXT("Input Touch");
 	}
 }
@@ -242,7 +249,7 @@ static void GetAddComponentClasses(UBlueprint const* BlueprintIn, FGraphActionLi
 static void GetCommentAction(UBlueprint const* BlueprintIn, FGraphActionListBuilderBase& ActionMenuBuilder)
 {
 	const bool bIsManyNodesSelected = (FKismetEditorUtilities::GetNumberOfSelectedNodes(BlueprintIn) > 0);
-	const FString MenuDescription = bIsManyNodesSelected ? LOCTEXT("AddCommentFromSelection", "Create Comment from Selection").ToString() : LOCTEXT("AddComment", "Add Comment...").ToString();
+	const FText MenuDescription = bIsManyNodesSelected ? LOCTEXT("AddCommentFromSelection", "Create Comment from Selection") : LOCTEXT("AddComment", "Add Comment...");
 	const FString ToolTip = LOCTEXT("CreateCommentToolTip", "Create a resizable comment box.").ToString();
 	TSharedPtr<FEdGraphSchemaAction_K2AddComment> NewAction = TSharedPtr<FEdGraphSchemaAction_K2AddComment>(new FEdGraphSchemaAction_K2AddComment( MenuDescription, ToolTip ));
 
@@ -267,6 +274,7 @@ static void GetEnumUtilitiesNodes(FBlueprintGraphActionListBuilder& ContextMenuB
 					TSharedPtr<FEdGraphSchemaAction_K2NewNode> Action = FK2ActionMenuBuilder::AddNewNodeAction(
 						InContextMenuBuilder, Category, NodeTemplate->GetNodeTitle(ENodeTitleType::ListView),  NodeTemplate->GetTooltip(), 0, NodeTemplate->GetKeywords());
 					Action->NodeTemplate = NodeTemplate;
+					Action->SearchTitle = NodeTemplate->GetNodeSearchTitle();
 				}
 
 				if (bInForEach)
@@ -277,6 +285,7 @@ static void GetEnumUtilitiesNodes(FBlueprintGraphActionListBuilder& ContextMenuB
 					TSharedPtr<FEdGraphSchemaAction_K2NewNode> Action = FK2ActionMenuBuilder::AddNewNodeAction(
 						InContextMenuBuilder, Category, NodeTemplate->GetNodeTitle(ENodeTitleType::ListView),  NodeTemplate->GetTooltip(), 0, NodeTemplate->GetKeywords());
 					Action->NodeTemplate = NodeTemplate;
+					Action->SearchTitle = NodeTemplate->GetNodeSearchTitle();
 				}
 
 				if (bInCastFromByte)
@@ -288,6 +297,7 @@ static void GetEnumUtilitiesNodes(FBlueprintGraphActionListBuilder& ContextMenuB
 					TSharedPtr<FEdGraphSchemaAction_K2NewNode> Action = FK2ActionMenuBuilder::AddNewNodeAction(
 						InContextMenuBuilder, Category, NodeTemplate->GetNodeTitle(ENodeTitleType::ListView),  NodeTemplate->GetTooltip(), 0, NodeTemplate->GetKeywords());
 					Action->NodeTemplate = NodeTemplate;
+					Action->SearchTitle = NodeTemplate->GetNodeSearchTitle();
 				}
 
 				if (bLiteralByte)
@@ -332,8 +342,9 @@ static void GetInterfaceMessages(UClass* CurrentInterface, const FString& Messag
 			{
 				UK2Node_Message* MessageTemplate = ContextMenuBuilder.CreateTemplateNode<UK2Node_Message>();
 				MessageTemplate->FunctionReference.SetExternalMember(Function->GetFName(), CurrentInterface);
-				TSharedPtr<FEdGraphSchemaAction_K2NewNode> Action = FK2ActionMenuBuilder::AddNewNodeAction(ContextMenuBuilder, MessageSubCategory, Function->GetName(), MessageTemplate->GetTooltip(), 0, MessageTemplate->GetKeywords());
+				TSharedPtr<FEdGraphSchemaAction_K2NewNode> Action = FK2ActionMenuBuilder::AddNewNodeAction(ContextMenuBuilder, MessageSubCategory, FText::FromString(Function->GetName()), MessageTemplate->GetTooltip(), 0, MessageTemplate->GetKeywords());
 				Action->NodeTemplate = MessageTemplate;
+				Action->SearchTitle = MessageTemplate->GetNodeSearchTitle();
 			}
 		}
 	}
@@ -343,12 +354,12 @@ static void GetInterfaceMessages(UClass* CurrentInterface, const FString& Messag
 static void GetReferenceSetterItems(FBlueprintGraphActionListBuilder& ContextMenuBuilder)
 {
 	const UEdGraphPin* AssociatedPin = ContextMenuBuilder.FromPin;
-	FString SetRefTitle = FString::Printf(TEXT("Set the value of %s"), *AssociatedPin->PinFriendlyName);
 
 	UK2Node_VariableSetRef* SetRefTemplate = ContextMenuBuilder.CreateTemplateNode<UK2Node_VariableSetRef>();
 	TSharedPtr<FEdGraphSchemaAction_K2NewNode> Action = FK2ActionMenuBuilder::AddNewNodeAction(ContextMenuBuilder, K2ActionCategories::VariablesCategory, SetRefTemplate->GetNodeTitle(ENodeTitleType::ListView), SetRefTemplate->GetTooltip(), 0, SetRefTemplate->GetKeywords());
 
 	Action->NodeTemplate = SetRefTemplate;
+	Action->SearchTitle = SetRefTemplate->GetNodeSearchTitle();
 }
 
 /** 
@@ -476,9 +487,10 @@ TSharedPtr<FEdGraphSchemaAction_K2AddComponent> FK2ActionMenuBuilder::CreateAddC
 
 	// Create menu category
 	FString MenuCategory = K2ActionCategories::AddComponentCategory + FString(TEXT("|")) + ClassGroup;
-	TSharedPtr<FEdGraphSchemaAction_K2AddComponent> Action = TSharedPtr<FEdGraphSchemaAction_K2AddComponent>(new FEdGraphSchemaAction_K2AddComponent(MenuCategory, Description, ToolTip, Asset != NULL ? 2 : 0));
+	TSharedPtr<FEdGraphSchemaAction_K2AddComponent> Action = TSharedPtr<FEdGraphSchemaAction_K2AddComponent>(new FEdGraphSchemaAction_K2AddComponent(MenuCategory, FText::FromString(Description), ToolTip, Asset != NULL ? 2 : 0));
 
 	Action->NodeTemplate = CallFuncNode;
+	Action->SearchTitle = Action->NodeTemplate->GetNodeSearchTitle();
 	Action->ComponentClass = DestinationComponentType;
 	Action->ComponentAsset = Asset;
 
@@ -486,7 +498,7 @@ TSharedPtr<FEdGraphSchemaAction_K2AddComponent> FK2ActionMenuBuilder::CreateAddC
 }
 
 //------------------------------------------------------------------------------
-TSharedPtr<FEdGraphSchemaAction_K2NewNode> FK2ActionMenuBuilder::AddNewNodeAction(FGraphActionListBuilderBase& ContextMenuBuilder, const FString& Category, const FString& MenuDesc, const FString& Tooltip, const int32 Grouping, const FString& Keywords)
+TSharedPtr<FEdGraphSchemaAction_K2NewNode> FK2ActionMenuBuilder::AddNewNodeAction(FGraphActionListBuilderBase& ContextMenuBuilder, const FString& Category, const FText& MenuDesc, const FString& Tooltip, const int32 Grouping, const FString& Keywords)
 {
 	TSharedPtr<FEdGraphSchemaAction_K2NewNode> NewActionNode = TSharedPtr<FEdGraphSchemaAction_K2NewNode>(new FEdGraphSchemaAction_K2NewNode(Category, MenuDesc, Tooltip, Grouping));
 	NewActionNode->Keywords = Keywords;
@@ -503,7 +515,7 @@ void FK2ActionMenuBuilder::AddEventForFunction(FGraphActionListBuilderBase& Acti
 	{
 		check(!Function->HasAnyFunctionFlags(FUNC_Delegate));
 		const FString SigName   = UEdGraphSchema_K2::GetFriendlySignitureName(Function);
-		const FString EventName = FString(TEXT("Event ")) + SigName;
+		const FString EventName = FText::Format( LOCTEXT("EventWithSignatureName", "Event {0}"), FText::FromString(SigName)).ToString();
 
 		UK2Node_Event* EventNodeTemplate = ActionMenuBuilder.CreateTemplateNode<UK2Node_Event>();
 		EventNodeTemplate->EventSignatureName = Function->GetFName();
@@ -512,9 +524,10 @@ void FK2ActionMenuBuilder::AddEventForFunction(FGraphActionListBuilderBase& Acti
 
 		const FString Category = UK2Node_CallFunction::GetDefaultCategoryForFunction(Function, K2ActionCategories::AddEventCategory);
 
-		TSharedPtr<FEdGraphSchemaAction_K2AddEvent> Action(new FEdGraphSchemaAction_K2AddEvent(Category, EventName, EventNodeTemplate->GetTooltip(), 0));
+		TSharedPtr<FEdGraphSchemaAction_K2AddEvent> Action(new FEdGraphSchemaAction_K2AddEvent(Category, FText::FromString(EventName), EventNodeTemplate->GetTooltip(), 0));
 		Action->Keywords     = EventNodeTemplate->GetKeywords();
 		Action->NodeTemplate = EventNodeTemplate;
+		Action->SearchTitle = EventNodeTemplate->GetNodeSearchTitle();
 
 		ActionMenuBuilder.AddAction( Action );
 	}
@@ -532,6 +545,7 @@ void FK2ActionMenuBuilder::AttachMacroGraphAction(FGraphActionListBuilderBase& A
 		MacroTemplate->SetMacroGraph(MacroGraph);
 		TSharedPtr<FEdGraphSchemaAction_K2NewNode> Action = FK2ActionMenuBuilder::AddNewNodeAction(ActionMenuBuilder, CategoryName, MacroTemplate->GetNodeTitle(ENodeTitleType::ListView), MacroTemplate->GetTooltip());
 		Action->NodeTemplate = MacroTemplate;
+		Action->SearchTitle = MacroTemplate->GetNodeSearchTitle();
 	}
 }
 
@@ -559,7 +573,7 @@ void FK2ActionMenuBuilder::AddSpawnInfoForFunction(const UFunction* Function, bo
 	// See if we want to create an actor ref node for the target
 	if (TargetInfo.FunctionTarget == FFunctionTargetInfo::EFT_Actor)
 	{
-		TSharedPtr<FEdGraphSchemaAction_K2AddCallOnActor> ActorAction = TSharedPtr<FEdGraphSchemaAction_K2AddCallOnActor>(new FEdGraphSchemaAction_K2AddCallOnActor(NodeCategory, Description, Tooltip, ActionGrouping));
+		TSharedPtr<FEdGraphSchemaAction_K2AddCallOnActor> ActorAction = TSharedPtr<FEdGraphSchemaAction_K2AddCallOnActor>(new FEdGraphSchemaAction_K2AddCallOnActor(NodeCategory, FText::FromString(Description), Tooltip, ActionGrouping));
 
 		for ( int32 ActorIndex = 0; ActorIndex < TargetInfo.Actors.Num(); ActorIndex++ )
 		{
@@ -576,7 +590,7 @@ void FK2ActionMenuBuilder::AddSpawnInfoForFunction(const UFunction* Function, bo
 	// Or create a component property ref node for the target
 	else if(TargetInfo.FunctionTarget == FFunctionTargetInfo::EFT_Component)
 	{
-		TSharedPtr<FEdGraphSchemaAction_K2AddCallOnVariable> VarAction = TSharedPtr<FEdGraphSchemaAction_K2AddCallOnVariable>(new FEdGraphSchemaAction_K2AddCallOnVariable(NodeCategory, Description, Tooltip, ActionGrouping));
+		TSharedPtr<FEdGraphSchemaAction_K2AddCallOnVariable> VarAction = TSharedPtr<FEdGraphSchemaAction_K2AddCallOnVariable>(new FEdGraphSchemaAction_K2AddCallOnVariable(NodeCategory, FText::FromString(Description), Tooltip, ActionGrouping));
 
 		VarAction->VariableName = TargetInfo.ComponentPropertyName;
 		Action = VarAction;
@@ -586,7 +600,7 @@ void FK2ActionMenuBuilder::AddSpawnInfoForFunction(const UFunction* Function, bo
 	// Default case (no additional nodes created for target)
 	else 
 	{
-		Action = FK2ActionMenuBuilder::AddNewNodeAction(ListBuilder, NodeCategory, Description, Tooltip, 0, Keywords);
+		Action = FK2ActionMenuBuilder::AddNewNodeAction(ListBuilder, NodeCategory, FText::FromString(Description), Tooltip, 0, Keywords);
 	}
 
 	Action->Keywords = Keywords;
@@ -634,6 +648,7 @@ void FK2ActionMenuBuilder::AddSpawnInfoForFunction(const UFunction* Function, bo
 	//@TODO: FUNCTIONREFACTOR:	CallFuncNode->FunctionReference.SetExternalMember(Function->GetFName(), Function->GetOuterUClass());
 
 	Action->NodeTemplate = CallTemplate;
+	Action->SearchTitle = Action->NodeTemplate->GetNodeSearchTitle();
 }
 
 /*******************************************************************************
@@ -842,7 +857,7 @@ void FK2ActionMenuBuilder::GetContextAllowedNodeTypes(FBlueprintGraphActionListB
 
 		if (FKismetEditorUtilities::CanPasteNodes(ContextMenuBuilder.CurrentGraph))
 		{
-			TSharedPtr<FEdGraphSchemaAction_K2PasteHere> NewAction(new FEdGraphSchemaAction_K2PasteHere(TEXT(""), TEXT("Paste here"), TEXT(""), 0));
+			TSharedPtr<FEdGraphSchemaAction_K2PasteHere> NewAction(new FEdGraphSchemaAction_K2PasteHere(TEXT(""), LOCTEXT("PasteHere", "Paste here"), TEXT(""), 0));
 			ContextMenuBuilder.AddAction(NewAction);
 		}
 
@@ -875,6 +890,7 @@ void FK2ActionMenuBuilder::GetContextAllowedNodeTypes(FBlueprintGraphActionListB
 				UK2Node* NodeTemplate = ContextMenuBuilder.CreateTemplateNode<UK2Node_AssignmentStatement>();
 				TSharedPtr<FEdGraphSchemaAction_K2NewNode> Action = FK2ActionMenuBuilder::AddNewNodeAction(ContextMenuBuilder, K2ActionCategories::MacroToolsCategory, NodeTemplate->GetNodeTitle(ENodeTitleType::ListView), NodeTemplate->GetTooltip(), 0, NodeTemplate->GetKeywords());
 				Action->NodeTemplate = NodeTemplate;
+				Action->SearchTitle = NodeTemplate->GetNodeSearchTitle();
 			}
 
 			// Temporary variables
@@ -944,6 +960,7 @@ void FK2ActionMenuBuilder::GetContextAllowedNodeTypes(FBlueprintGraphActionListB
 
 				TSharedPtr<FEdGraphSchemaAction_K2NewNode> Action = FK2ActionMenuBuilder::AddNewNodeAction(ContextMenuBuilder, K2ActionCategories::MacroToolsCategory, NodeTemplate->GetNodeTitle(ENodeTitleType::ListView), NodeTemplate->GetTooltip(), 0, NodeTemplate->GetKeywords());
 				Action->NodeTemplate = NodeTemplate;
+				Action->SearchTitle = NodeTemplate->GetNodeSearchTitle();
 			}
 
 			// Add in bool and int persistent pin types
@@ -955,6 +972,7 @@ void FK2ActionMenuBuilder::GetContextAllowedNodeTypes(FBlueprintGraphActionListB
 				PersistentIntNodeTemplate->bIsPersistent = true;
 				TSharedPtr<FEdGraphSchemaAction_K2NewNode> IntAction = FK2ActionMenuBuilder::AddNewNodeAction(ContextMenuBuilder, K2ActionCategories::MacroToolsCategory, PersistentIntNodeTemplate->GetNodeTitle(ENodeTitleType::ListView), PersistentIntNodeTemplate->GetTooltip(), 0, PersistentIntNodeTemplate->GetKeywords());
 				IntAction->NodeTemplate = PersistentIntNodeTemplate;
+				IntAction->SearchTitle = PersistentIntNodeTemplate->GetNodeSearchTitle();
 
 				// Persistent Bool
 				UK2Node_TemporaryVariable* PersistentBoolNodeTemplate = ContextMenuBuilder.CreateTemplateNode<UK2Node_TemporaryVariable>();
@@ -962,6 +980,7 @@ void FK2ActionMenuBuilder::GetContextAllowedNodeTypes(FBlueprintGraphActionListB
 				PersistentBoolNodeTemplate->bIsPersistent = true;
 				TSharedPtr<FEdGraphSchemaAction_K2NewNode> BoolAction = FK2ActionMenuBuilder::AddNewNodeAction(ContextMenuBuilder, K2ActionCategories::MacroToolsCategory, PersistentBoolNodeTemplate->GetNodeTitle(ENodeTitleType::ListView), PersistentBoolNodeTemplate->GetTooltip(), 0, PersistentBoolNodeTemplate->GetKeywords());
 				BoolAction->NodeTemplate = PersistentBoolNodeTemplate;
+				BoolAction->SearchTitle = PersistentBoolNodeTemplate->GetNodeSearchTitle();
 			}
 		}
 
@@ -984,6 +1003,7 @@ void FK2ActionMenuBuilder::GetContextAllowedNodeTypes(FBlueprintGraphActionListB
 			UK2Node* NodeTemplate = ContextMenuBuilder.CreateTemplateNode<UK2Node_Select>();
 			TSharedPtr<FEdGraphSchemaAction_K2NewNode> SelectionAction = FK2ActionMenuBuilder::AddNewNodeAction(ContextMenuBuilder, K2ActionCategories::UtilitiesCategory, NodeTemplate->GetNodeTitle(ENodeTitleType::ListView), NodeTemplate->GetTooltip(), 0, NodeTemplate->GetKeywords());
 			SelectionAction->NodeTemplate = NodeTemplate;
+			SelectionAction->SearchTitle = NodeTemplate->GetNodeSearchTitle();
 		}
 
 		if (GraphType == GT_Ubergraph && bAllowEvents && !bIsConstructionScript && bAllowImpureFuncs)
@@ -1002,6 +1022,7 @@ void FK2ActionMenuBuilder::GetContextAllowedNodeTypes(FBlueprintGraphActionListB
 					FString CategoryName = FString::Printf( TEXT("%s|%s"), *K2ActionCategories::GenericFunctionCategory, *CurrentInterface->GetDefaultObject<UK2Node_BaseAsyncTask>()->GetCategoryName());
 					TSharedPtr<FEdGraphSchemaAction_K2NewNode> SelectionAction = FK2ActionMenuBuilder::AddNewNodeAction(ContextMenuBuilder, CategoryName, NodeTemplate->GetNodeTitle(ENodeTitleType::ListView), NodeTemplate->GetTooltip(), 0, NodeTemplate->GetKeywords());
 					SelectionAction->NodeTemplate = NodeTemplate;
+					SelectionAction->SearchTitle = NodeTemplate->GetNodeSearchTitle();
 				}
 			}
 		}
@@ -1033,6 +1054,7 @@ struct FClassDynamicCastHelper
 
 		TSharedPtr<FEdGraphSchemaAction_K2NewNode> Action = FK2ActionMenuBuilder::AddNewNodeAction(ContextMenuBuilder, CastCategory, CastNode->GetNodeTitle(ENodeTitleType::ListView), TEXT("Dynamic cast to a specific type"), 0, CastNode->GetKeywords());
 		Action->NodeTemplate = CastNode;
+		Action->SearchTitle = CastNode->GetNodeSearchTitle();
 	}
 
 	static void GetClassDynamicCastNodes(FBlueprintGraphActionListBuilder& ContextMenuBuilder, const UEdGraphSchema_K2* K2Schema)
@@ -1236,6 +1258,7 @@ void FK2ActionMenuBuilder::GetPinAllowedNodeTypes(FBlueprintGraphActionListBuild
 
 								TSharedPtr<FEdGraphSchemaAction_K2NewNode> Action = FK2ActionMenuBuilder::AddNewNodeAction(ContextMenuBuilder,  CastCategory, CastNode->GetNodeTitle(ENodeTitleType::ListView), TEXT("Dynamic cast to a specific type"), 0, CastNode->GetKeywords());
 								Action->NodeTemplate = CastNode;
+								Action->SearchTitle = CastNode->GetNodeSearchTitle();
 							}
 						}
 					}
@@ -1295,6 +1318,7 @@ void FK2ActionMenuBuilder::GetPinAllowedNodeTypes(FBlueprintGraphActionListBuild
 					NodeTemplate->StructType = PinStruct;
 					TSharedPtr<FEdGraphSchemaAction_K2NewNode> Action = FK2ActionMenuBuilder::AddNewNodeAction(ContextMenuBuilder, FString(), NodeTemplate->GetNodeTitle(ENodeTitleType::ListView), NodeTemplate->GetTooltip(), 0, NodeTemplate->GetKeywords());
 					Action->NodeTemplate = NodeTemplate;
+					Action->SearchTitle = NodeTemplate->GetNodeSearchTitle();
 				}
 				else if ((FromPin.Direction == EGPD_Input) && UK2Node_MakeStruct::CanBeMade(PinStruct))
 				{
@@ -1302,6 +1326,7 @@ void FK2ActionMenuBuilder::GetPinAllowedNodeTypes(FBlueprintGraphActionListBuild
 					NodeTemplate->StructType = PinStruct;
 					TSharedPtr<FEdGraphSchemaAction_K2NewNode> Action = FK2ActionMenuBuilder::AddNewNodeAction(ContextMenuBuilder, FString(), NodeTemplate->GetNodeTitle(ENodeTitleType::ListView), NodeTemplate->GetTooltip(), 0, NodeTemplate->GetKeywords());
 					Action->NodeTemplate = NodeTemplate;
+					Action->SearchTitle = NodeTemplate->GetNodeSearchTitle();
 				}
 			}
 
@@ -1310,13 +1335,14 @@ void FK2ActionMenuBuilder::GetPinAllowedNodeTypes(FBlueprintGraphActionListBuild
 				UK2Node_CreateDelegate* NodeTemplate = ContextMenuBuilder.CreateTemplateNode<UK2Node_CreateDelegate>();
 				TSharedPtr<FEdGraphSchemaAction_K2NewNode> Action = FK2ActionMenuBuilder::AddNewNodeAction(ContextMenuBuilder, FString(), NodeTemplate->GetNodeTitle(ENodeTitleType::ListView), NodeTemplate->GetTooltip(), 0, NodeTemplate->GetKeywords());
 				Action->NodeTemplate = NodeTemplate;
+				Action->SearchTitle = NodeTemplate->GetNodeSearchTitle();
 
 				UFunction* Signature = Cast<UFunction>(FromPin.PinType.PinSubCategoryObject.Get());
 				const bool bSupportsEventGraphs = (Blueprint && FBlueprintEditorUtils::DoesSupportEventGraphs(Blueprint));
 				const bool bAllowEvents = (GraphType == GT_Ubergraph);
 				if(Signature && bSupportsEventGraphs && bAllowEvents)
 				{
-					TSharedPtr<FEdGraphSchemaAction_EventFromFunction> DelegateEvent(new FEdGraphSchemaAction_EventFromFunction( FString(), TEXT("Add Custom Event for Dispatcher"), TEXT("Add a new event for given Dispatcher"), 0));
+					TSharedPtr<FEdGraphSchemaAction_EventFromFunction> DelegateEvent(new FEdGraphSchemaAction_EventFromFunction( FString(), LOCTEXT("AddCustomEventForDispatcher", "Add Custom Event for Dispatcher"), LOCTEXT("AddCustomEventForDispatcher_Tooltip", "Add a new event for given Dispatcher").ToString(), 0));
 					ContextMenuBuilder.AddAction( DelegateEvent );
 					DelegateEvent->SignatureFunction = Signature;
 				}
@@ -1359,6 +1385,7 @@ void FK2ActionMenuBuilder::GetPinAllowedNodeTypes(FBlueprintGraphActionListBuild
 			UK2Node* NodeTemplate = ContextMenuBuilder.CreateTemplateNode<UK2Node_AssignmentStatement>();
 			TSharedPtr<FEdGraphSchemaAction_K2NewNode> Action = FK2ActionMenuBuilder::AddNewNodeAction(ContextMenuBuilder, TEXT(""), NodeTemplate->GetNodeTitle(ENodeTitleType::ListView), NodeTemplate->GetTooltip(), 0, NodeTemplate->GetKeywords());
 			Action->NodeTemplate = NodeTemplate;
+			Action->SearchTitle = NodeTemplate->GetNodeSearchTitle();
 		}
 	}
 }
@@ -1465,7 +1492,7 @@ void FK2ActionMenuBuilder::GetAllActionsForClass(FBlueprintPaletteListBuilder& P
 					const FString PropertyDesc = GetDefault<UEditorStyleSettings>()->bShowFriendlyNames ? UEditorEngine::GetFriendlyName(Property) : Property->GetName();
 					const FString PropertyCategory = FObjectEditorUtils::GetCategory(Property);
 
-					TSharedPtr<FEdGraphSchemaAction_K2Var> NewVarAction = TSharedPtr<FEdGraphSchemaAction_K2Var>(new FEdGraphSchemaAction_K2Var(PropertyCategory, PropertyDesc, PropertyTooltip, 0));
+					TSharedPtr<FEdGraphSchemaAction_K2Var> NewVarAction = TSharedPtr<FEdGraphSchemaAction_K2Var>(new FEdGraphSchemaAction_K2Var(PropertyCategory, FText::FromString(PropertyDesc), PropertyTooltip, 0));
 					NewVarAction->SetVariableInfo(Property->GetFName(), InClass);
 
 					PaletteBuilder.AddAction(NewVarAction);
@@ -1493,7 +1520,7 @@ void FK2ActionMenuBuilder::GetAllActionsForClass(FBlueprintPaletteListBuilder& P
 					}
 					PropertyCategory = FName::NameToDisplayString(PropertyCategory, /*bIsBool =*/false);
 
-					TSharedPtr<FEdGraphSchemaAction_K2Delegate> NewFuncAction = MakeShareable(new FEdGraphSchemaAction_K2Delegate(PropertyCategory, PropertyDesc, PropertyTooltip, 0));
+					TSharedPtr<FEdGraphSchemaAction_K2Delegate> NewFuncAction = MakeShareable(new FEdGraphSchemaAction_K2Delegate(PropertyCategory, FText::FromString(PropertyDesc), PropertyTooltip, 0));
 					NewFuncAction->SetDelegateInfo(Property->GetFName(), InClass);
 
 					UClass* CurrentPropertyOwner = Cast<UClass>(Property->GetOuterUField());
@@ -1560,6 +1587,7 @@ void FK2ActionMenuBuilder::GetSwitchMenuItems(FBlueprintPaletteListBuilder& Acti
 				UK2Node* NodeTemplate = EnumTemplate;
 				TSharedPtr<FEdGraphSchemaAction_K2NewNode> Action = FK2ActionMenuBuilder::AddNewNodeAction(ActionMenuBuilder, SwitchControlCategory, NodeTemplate->GetNodeTitle(ENodeTitleType::ListView), NodeTemplate->GetTooltip(), 0, NodeTemplate->GetKeywords());
 				Action->NodeTemplate = NodeTemplate;
+				Action->SearchTitle = NodeTemplate->GetNodeSearchTitle();
 			}
 		}
 	}
@@ -1579,7 +1607,7 @@ void FK2ActionMenuBuilder::GetAnimNotifyMenuItems(FBlueprintGraphActionListBuild
 				FName NotifyName = AnimBlueprint->TargetSkeleton->AnimationNotifies[I];
 				FString Label = NotifyName.ToString();
 
-				TSharedPtr<FEdGraphSchemaAction_K2NewNode> Action = FK2ActionMenuBuilder::AddNewNodeAction(ContextMenuBuilder, K2ActionCategories::AnimNotifyCategory, Label, TEXT(""));
+				TSharedPtr<FEdGraphSchemaAction_K2NewNode> Action = FK2ActionMenuBuilder::AddNewNodeAction(ContextMenuBuilder, K2ActionCategories::AnimNotifyCategory, FText::FromString(Label), TEXT(""));
 
 				UK2Node_Event* EventNode = ContextMenuBuilder.CreateTemplateNode<UK2Node_Event>();
 				FString SignatureName = FString::Printf(TEXT("AnimNotify_%s"), *Label);
@@ -1587,6 +1615,7 @@ void FK2ActionMenuBuilder::GetAnimNotifyMenuItems(FBlueprintGraphActionListBuild
 				EventNode->EventSignatureClass = UAnimInstance::StaticClass();
 				EventNode->CustomFunctionName = EventNode->EventSignatureName;
 				Action->NodeTemplate = EventNode;
+				Action->SearchTitle = Action->NodeTemplate->GetNodeSearchTitle();
 			}
 
 			// @todo anim: fix this to be same as notifies, save same list in the 
@@ -1612,7 +1641,7 @@ void FK2ActionMenuBuilder::GetAnimNotifyMenuItems(FBlueprintGraphActionListBuild
 				FName NotifyName = BrancingPointHandlers[I];
 				FString Label = NotifyName.ToString();
 
-				TSharedPtr<FEdGraphSchemaAction_K2NewNode> Action = FK2ActionMenuBuilder::AddNewNodeAction(ContextMenuBuilder, K2ActionCategories::BranchPointCategory, Label, TEXT(""));
+				TSharedPtr<FEdGraphSchemaAction_K2NewNode> Action = FK2ActionMenuBuilder::AddNewNodeAction(ContextMenuBuilder, K2ActionCategories::BranchPointCategory, FText::FromString(Label), TEXT(""));
 
 				UK2Node_Event* EventNode = ContextMenuBuilder.CreateTemplateNode<UK2Node_Event>();
 				FString SignatureName = FString::Printf(TEXT("MontageBranchingPoint_%s"), *Label);
@@ -1620,6 +1649,7 @@ void FK2ActionMenuBuilder::GetAnimNotifyMenuItems(FBlueprintGraphActionListBuild
 				EventNode->EventSignatureClass = UAnimInstance::StaticClass();
 				EventNode->CustomFunctionName = EventNode->EventSignatureName;
 				Action->NodeTemplate = EventNode;
+				Action->SearchTitle = Action->NodeTemplate->GetNodeSearchTitle();
 			}
 		}
 
@@ -1633,7 +1663,7 @@ void FK2ActionMenuBuilder::GetAnimNotifyMenuItems(FBlueprintGraphActionListBuild
 				{
 					FString Label = NotifyName.ToString();
 
-					TSharedPtr<FEdGraphSchemaAction_K2NewNode> Action = FK2ActionMenuBuilder::AddNewNodeAction(ContextMenuBuilder, K2ActionCategories::StateMachineCategory, Label, TEXT(""));
+					TSharedPtr<FEdGraphSchemaAction_K2NewNode> Action = FK2ActionMenuBuilder::AddNewNodeAction(ContextMenuBuilder, K2ActionCategories::StateMachineCategory, FText::FromString(Label), TEXT(""));
 
 					UK2Node_Event* EventNode = ContextMenuBuilder.CreateTemplateNode<UK2Node_Event>();
 					FString SignatureName = FString::Printf(TEXT("AnimNotify_%s"), *Label);
@@ -1641,6 +1671,7 @@ void FK2ActionMenuBuilder::GetAnimNotifyMenuItems(FBlueprintGraphActionListBuild
 					EventNode->EventSignatureClass = UAnimInstance::StaticClass();
 					EventNode->CustomFunctionName = EventNode->EventSignatureName;
 					Action->NodeTemplate = EventNode;
+					Action->SearchTitle = Action->NodeTemplate->GetNodeSearchTitle();
 				}
 			}
 		}
@@ -1723,20 +1754,22 @@ void FK2ActionMenuBuilder::GetVariableGettersSettersForClass(FBlueprintGraphActi
 				// Variable getters and setters
 				if (bCanRead)
 				{
-					TSharedPtr<FEdGraphSchemaAction_K2NewNode> ReadVar = FK2ActionMenuBuilder::AddNewNodeAction(ContextMenuBuilder, FullCategoryName, GetString + PropertyDescription, PropertyTooltip);
+					TSharedPtr<FEdGraphSchemaAction_K2NewNode> ReadVar = FK2ActionMenuBuilder::AddNewNodeAction(ContextMenuBuilder, FullCategoryName, FText::FromString(GetString + PropertyDescription), PropertyTooltip);
 
 					UK2Node_VariableGet* VarGetNode = ContextMenuBuilder.CreateTemplateNode<UK2Node_VariableGet>();
 					VarGetNode->SetFromProperty(Property, bSelfContext);
 					ReadVar->NodeTemplate = VarGetNode;
+					ReadVar->SearchTitle = ReadVar->NodeTemplate->GetNodeSearchTitle();
 				}
 
 				if (bCanWrite)
 				{
-					TSharedPtr<FEdGraphSchemaAction_K2NewNode> WriteVar = FK2ActionMenuBuilder::AddNewNodeAction(ContextMenuBuilder, FullCategoryName, SetString + PropertyDescription, PropertyTooltip);
+					TSharedPtr<FEdGraphSchemaAction_K2NewNode> WriteVar = FK2ActionMenuBuilder::AddNewNodeAction(ContextMenuBuilder, FullCategoryName, FText::FromString(SetString + PropertyDescription), PropertyTooltip);
 
 					UK2Node_VariableSet* VarSetNode = ContextMenuBuilder.CreateTemplateNode<UK2Node_VariableSet>();
 					VarSetNode->SetFromProperty(Property, bSelfContext);
 					WriteVar->NodeTemplate = VarSetNode;
+					WriteVar->SearchTitle = WriteVar->NodeTemplate->GetNodeSearchTitle();
 				}
 			}
 		}
@@ -1750,11 +1783,12 @@ void FK2ActionMenuBuilder::GetVariableGettersSettersForClass(FBlueprintGraphActi
 	// Add in self-context variable getter, if needed
 	if( bSelfContext && bWantGetters )
 	{
-		TSharedPtr<FEdGraphSchemaAction_K2NewNode> SelfVar = FK2ActionMenuBuilder::AddNewNodeAction(ContextMenuBuilder, K2ActionCategories::VariablesCategory, SelfString, SelfTooltip);
+		TSharedPtr<FEdGraphSchemaAction_K2NewNode> SelfVar = FK2ActionMenuBuilder::AddNewNodeAction(ContextMenuBuilder, K2ActionCategories::VariablesCategory, FText::FromString(SelfString), SelfTooltip);
 
 		UK2Node_Self* SelfNode = ContextMenuBuilder.CreateTemplateNode<UK2Node_Self>();
 		SelfNode->SelfClass = (UClass*)Class;
 		SelfVar->NodeTemplate = SelfNode;
+		SelfVar->SearchTitle = SelfVar->NodeTemplate->GetNodeSearchTitle();
 		SelfVar->Keywords = SelfNode->GetKeywords();
 	}
 }
@@ -1795,12 +1829,15 @@ void FK2ActionMenuBuilder::GetEventDispatcherNodesForClass(FBlueprintGraphAction
 					if(bAllowEvents)
 					{
 						const FString PropertyDescription = GetDefault<UEditorStyleSettings>()->bShowFriendlyNames ? UEditorEngine::GetFriendlyName(Property) : Property->GetName();
-						TSharedPtr<FEdGraphSchemaAction_K2AssignDelegate> AssignDelegateEvent(new FEdGraphSchemaAction_K2AssignDelegate(K2ActionCategories::DelegateBindingCategory, BindString + PropertyDescription, PropertyTooltip, 0));
+						TSharedPtr<FEdGraphSchemaAction_K2AssignDelegate> AssignDelegateEvent(new FEdGraphSchemaAction_K2AssignDelegate(K2ActionCategories::DelegateBindingCategory, FText::FromString(BindString + PropertyDescription), PropertyTooltip, 0));
 						AssignDelegateEvent->NodeTemplate = DelegateEventNode;
+						AssignDelegateEvent->SearchTitle = AssignDelegateEvent->NodeTemplate->GetNodeSearchTitle();
 						ContextMenuBuilder.AddAction( AssignDelegateEvent );
 					}
 					TSharedPtr<FEdGraphSchemaAction_K2NewNode> DelegateEvent = FK2ActionMenuBuilder::AddNewNodeAction(ContextMenuBuilder, K2ActionCategories::DelegateCategory, DelegateEventNode->GetNodeTitle(ENodeTitleType::FullTitle), PropertyTooltip);
 					DelegateEvent->NodeTemplate = DelegateEventNode;
+					DelegateEvent->SearchTitle = DelegateEvent->NodeTemplate->GetNodeSearchTitle();
+					DelegateEvent->SearchTitle = DelegateEventNode->GetNodeSearchTitle();
 				}
 
 				{
@@ -1808,6 +1845,8 @@ void FK2ActionMenuBuilder::GetEventDispatcherNodesForClass(FBlueprintGraphAction
 					DelegateEventNode->SetFromProperty(Property, bSelfContext);
 					TSharedPtr<FEdGraphSchemaAction_K2NewNode> DelegateEvent = FK2ActionMenuBuilder::AddNewNodeAction(ContextMenuBuilder, K2ActionCategories::DelegateCategory, DelegateEventNode->GetNodeTitle(ENodeTitleType::FullTitle), PropertyTooltip);
 					DelegateEvent->NodeTemplate = DelegateEventNode;
+					DelegateEvent->SearchTitle = DelegateEvent->NodeTemplate->GetNodeSearchTitle();
+					DelegateEvent->SearchTitle = DelegateEventNode->GetNodeSearchTitle();
 				}
 
 				{
@@ -1815,6 +1854,8 @@ void FK2ActionMenuBuilder::GetEventDispatcherNodesForClass(FBlueprintGraphAction
 					DelegateEventNode->SetFromProperty(Property, bSelfContext);
 					TSharedPtr<FEdGraphSchemaAction_K2NewNode> DelegateEvent = FK2ActionMenuBuilder::AddNewNodeAction(ContextMenuBuilder, K2ActionCategories::DelegateCategory, DelegateEventNode->GetNodeTitle(ENodeTitleType::FullTitle), PropertyTooltip);
 					DelegateEvent->NodeTemplate = DelegateEventNode;
+					DelegateEvent->SearchTitle = DelegateEvent->NodeTemplate->GetNodeSearchTitle();
+					DelegateEvent->SearchTitle = DelegateEventNode->GetNodeSearchTitle();
 				}
 			}
 
@@ -1824,6 +1865,8 @@ void FK2ActionMenuBuilder::GetEventDispatcherNodesForClass(FBlueprintGraphAction
 				DelegateEventNode->SetFromProperty(Property, bSelfContext);
 				TSharedPtr<FEdGraphSchemaAction_K2NewNode> DelegateEvent = FK2ActionMenuBuilder::AddNewNodeAction(ContextMenuBuilder, K2ActionCategories::DelegateCategory, DelegateEventNode->GetNodeTitle(ENodeTitleType::FullTitle), PropertyTooltip);
 				DelegateEvent->NodeTemplate = DelegateEventNode;
+				DelegateEvent->SearchTitle = DelegateEvent->NodeTemplate->GetNodeSearchTitle();
+				DelegateEvent->SearchTitle = DelegateEventNode->GetNodeSearchTitle();
 			}
 		}
 	}
@@ -1847,6 +1890,7 @@ void FK2ActionMenuBuilder::GetStructActions(FBlueprintGraphActionListBuilder& Co
 					BreakNodeTemplate->StructType = Struct;
 					TSharedPtr<FEdGraphSchemaAction_K2NewNode> Action = FK2ActionMenuBuilder::AddNewNodeAction(ContextMenuBuilder, K2ActionCategories::BreakStructCategory, BreakNodeTemplate->GetNodeTitle(ENodeTitleType::ListView), BreakNodeTemplate->GetTooltip(), 0, BreakNodeTemplate->GetKeywords());
 					Action->NodeTemplate = BreakNodeTemplate;
+					Action->SearchTitle = Action->NodeTemplate->GetNodeSearchTitle();
 				}
 			}
 
@@ -1858,6 +1902,7 @@ void FK2ActionMenuBuilder::GetStructActions(FBlueprintGraphActionListBuilder& Co
 					MakeNodeTemplate->StructType = Struct;
 					TSharedPtr<FEdGraphSchemaAction_K2NewNode> Action = FK2ActionMenuBuilder::AddNewNodeAction(ContextMenuBuilder, K2ActionCategories::MakeStructCategory, MakeNodeTemplate->GetNodeTitle(ENodeTitleType::ListView), MakeNodeTemplate->GetTooltip(), 0, MakeNodeTemplate->GetKeywords());
 					Action->NodeTemplate = MakeNodeTemplate;
+					Action->SearchTitle = Action->NodeTemplate->GetNodeSearchTitle();
 				}
 			}
 		}
@@ -1874,12 +1919,13 @@ void FK2ActionMenuBuilder::GetAllEventActions(FBlueprintPaletteListBuilder& Acti
 	if (FBlueprintEditorUtils::DoesSupportTimelines(ActionMenuBuilder.Blueprint))
 	{
 		const FString ToolTip = LOCTEXT("Timeline_Tooltip", "Add a Timeline node to modify values over time.").ToString();
-		TSharedPtr<FEdGraphSchemaAction_K2AddTimeline> Action = TSharedPtr<FEdGraphSchemaAction_K2AddTimeline>(new FEdGraphSchemaAction_K2AddTimeline(TEXT(""), TEXT("Add Timeline..."), ToolTip, 0));
+		TSharedPtr<FEdGraphSchemaAction_K2AddTimeline> Action = TSharedPtr<FEdGraphSchemaAction_K2AddTimeline>(new FEdGraphSchemaAction_K2AddTimeline(TEXT(""), LOCTEXT("AddTimeline", "Add Timeline..."), ToolTip, 0));
 		ActionMenuBuilder.AddAction( Action );
 
 		UK2Node_Timeline* TimelineNodeTemplate = ActionMenuBuilder.CreateTemplateNode<UK2Node_Timeline>();
 		TimelineNodeTemplate->TimelineName = FBlueprintEditorUtils::FindUniqueTimelineName(ActionMenuBuilder.Blueprint);
 		Action->NodeTemplate = TimelineNodeTemplate;
+		Action->SearchTitle = Action->NodeTemplate->GetNodeSearchTitle();
 	}
 }
 
@@ -1974,13 +2020,14 @@ void FK2ActionMenuBuilder::GetEventsForBlueprint(FBlueprintPaletteListBuilder& A
 	}
 
 	// Add support for custom events, if applicable
-	TSharedPtr<FEdGraphSchemaAction_K2AddCustomEvent> Action = TSharedPtr<FEdGraphSchemaAction_K2AddCustomEvent>(new FEdGraphSchemaAction_K2AddCustomEvent(TEXT(""), LOCTEXT("AddCustomEvent", "Custom Event...").ToString(), TEXT(""), 0));
+	TSharedPtr<FEdGraphSchemaAction_K2AddCustomEvent> Action = TSharedPtr<FEdGraphSchemaAction_K2AddCustomEvent>(new FEdGraphSchemaAction_K2AddCustomEvent(TEXT(""), LOCTEXT("AddCustomEvent", "Custom Event..."), TEXT(""), 0));
 	ActionMenuBuilder.AddAction(Action, K2ActionCategories::AddEventCategory);
 
 	UK2Node_CustomEvent* CustomEventNode = ActionMenuBuilder.CreateTemplateNode<UK2Node_CustomEvent>();
 	CustomEventNode->CustomFunctionName = FBlueprintEditorUtils::FindUniqueCustomEventName(Blueprint);
 	CustomEventNode->bIsEditable = true;
 	Action->NodeTemplate = CustomEventNode;
+	Action->SearchTitle = Action->NodeTemplate->GetNodeSearchTitle();
 }
 
 //------------------------------------------------------------------------------
@@ -2017,9 +2064,10 @@ void FK2ActionMenuBuilder::GetLiteralsFromActorSelection(FBlueprintGraphActionLi
 				FString::Printf(*LOCTEXT("AddReferenceToActorToolTip", "Add a reference to %s").ToString(), *Actor->GetActorLabel()) :
 				LOCTEXT("AddReferencesToSelectedActorsToolTip", "Add references to currently selected actors in the scene").ToString();
 
-			TSharedPtr<FEdGraphSchemaAction_K2NewNode> Action( new FEdGraphSchemaAction_K2NewNode(TEXT(""), Description, ToolTip, K2Schema->AG_LevelReference) );
+			TSharedPtr<FEdGraphSchemaAction_K2NewNode> Action( new FEdGraphSchemaAction_K2NewNode(TEXT(""), FText::FromString(Description), ToolTip, K2Schema->AG_LevelReference) );
 			Action->Keywords = LiteralNodeTemplate->GetKeywords();
 			Action->NodeTemplate = LiteralNodeTemplate;
+			Action->SearchTitle = Action->NodeTemplate->GetNodeSearchTitle();
 
 			NewActions.Add( Action );
 		}
@@ -2092,18 +2140,19 @@ void FK2ActionMenuBuilder::AddBoundEventActionsForClass(FBlueprintGraphActionLis
 
 					if (EventNode)
 					{
-						TSharedPtr<FEdGraphSchemaAction_K2ViewNode> NewDelegateNode = TSharedPtr<FEdGraphSchemaAction_K2ViewNode>( new FEdGraphSchemaAction_K2ViewNode(EventCategory, ViewString + DelegateProperty->GetName(), PropertyTooltip, K2Schema->AG_LevelReference));
+						TSharedPtr<FEdGraphSchemaAction_K2ViewNode> NewDelegateNode = TSharedPtr<FEdGraphSchemaAction_K2ViewNode>( new FEdGraphSchemaAction_K2ViewNode(EventCategory, FText::FromString(ViewString + DelegateProperty->GetName()), PropertyTooltip, K2Schema->AG_LevelReference));
 						NewDelegateNode->NodePtr = EventNode;
 
 						ContextMenuBuilder.AddAction( NewDelegateNode );
 					}
 					else
 					{
-						TSharedPtr<FEdGraphSchemaAction_K2NewNode> NewDelegateNode = TSharedPtr<FEdGraphSchemaAction_K2NewNode>( new FEdGraphSchemaAction_K2NewNode(EventCategory, AddString + DelegateProperty->GetName(), PropertyTooltip, K2Schema->AG_LevelReference));
+						TSharedPtr<FEdGraphSchemaAction_K2NewNode> NewDelegateNode = TSharedPtr<FEdGraphSchemaAction_K2NewNode>( new FEdGraphSchemaAction_K2NewNode(EventCategory, FText::FromString(AddString + DelegateProperty->GetName()), PropertyTooltip, K2Schema->AG_LevelReference));
 
 						UK2Node_ComponentBoundEvent* NewComponentEvent = ContextMenuBuilder.CreateTemplateNode<UK2Node_ComponentBoundEvent>();
 						NewComponentEvent->InitializeComponentBoundEventParams( ComponentProperty, DelegateProperty );
 						NewDelegateNode->NodeTemplate = NewComponentEvent;
+						NewDelegateNode->SearchTitle = NewDelegateNode->NodeTemplate->GetNodeSearchTitle();
 
 						ContextMenuBuilder.AddAction( NewDelegateNode );
 					}
@@ -2117,18 +2166,19 @@ void FK2ActionMenuBuilder::AddBoundEventActionsForClass(FBlueprintGraphActionLis
 
 						if (ActorEventNode)
 						{
-							TSharedPtr<FEdGraphSchemaAction_K2ViewNode> NewDelegateNode = TSharedPtr<FEdGraphSchemaAction_K2ViewNode>( new FEdGraphSchemaAction_K2ViewNode(EventCategory, ViewString + DelegateProperty->GetName(), PropertyTooltip, K2Schema->AG_LevelReference));
+							TSharedPtr<FEdGraphSchemaAction_K2ViewNode> NewDelegateNode = TSharedPtr<FEdGraphSchemaAction_K2ViewNode>( new FEdGraphSchemaAction_K2ViewNode(EventCategory, FText::FromString(ViewString + DelegateProperty->GetName()), PropertyTooltip, K2Schema->AG_LevelReference));
 							NewDelegateNode->NodePtr = ActorEventNode;
 
 							NewActions.Add( NewDelegateNode );
 						}
 						else
 						{
-							TSharedPtr<FEdGraphSchemaAction_K2NewNode> NewDelegateNode = TSharedPtr<FEdGraphSchemaAction_K2NewNode>(new FEdGraphSchemaAction_K2NewNode(EventCategory, AddString + DelegateProperty->GetName(), PropertyTooltip, K2Schema->AG_LevelReference));
+							TSharedPtr<FEdGraphSchemaAction_K2NewNode> NewDelegateNode = TSharedPtr<FEdGraphSchemaAction_K2NewNode>(new FEdGraphSchemaAction_K2NewNode(EventCategory, FText::FromString(AddString + DelegateProperty->GetName()), PropertyTooltip, K2Schema->AG_LevelReference));
 
 							UK2Node_ActorBoundEvent* NewActorEvent = ContextMenuBuilder.CreateTemplateNode<UK2Node_ActorBoundEvent>();
 							NewActorEvent->InitializeActorBoundEventParams( Actors[ActorIndex], DelegateProperty );
 							NewDelegateNode->NodeTemplate = NewActorEvent;
+							NewDelegateNode->SearchTitle = NewDelegateNode->NodeTemplate->GetNodeSearchTitle();
 
 							NewActions.Add( NewDelegateNode );
 						}
@@ -2194,14 +2244,15 @@ void FK2ActionMenuBuilder::GetMatineeControllers(FBlueprintGraphActionListBuilde
 	{
 		AMatineeActor* MatActor = MatineeActors[i];
 
-		FString Description = FString::Printf(TEXT("Add MatineeController for %s"), *MatActor->GetName());
-		FString ToolTip = FString::Printf(TEXT("Add Matinee controller node for %s"), *MatActor->GetName());
+		FText Description = FText::Format(LOCTEXT("AddMatineeControllerForActor", "Add MatineeController for {actorName}"), FText::FromString(MatActor->GetName()));
+		FString ToolTip = FText::Format(LOCTEXT("AddMatineeControllerForActor_Tooltip", "Add Matinee controller node for {actorName}"), FText::FromString(MatActor->GetName())).ToString();
 
 		UK2Node_MatineeController* MatControlNode = ContextMenuBuilder.CreateTemplateNode<UK2Node_MatineeController>();
 		MatControlNode->MatineeActor = MatActor;
 
 		TSharedPtr<FEdGraphSchemaAction_K2NewNode> Action = FK2ActionMenuBuilder::AddNewNodeAction(ContextMenuBuilder, TEXT(""), Description, ToolTip, K2Schema->AG_LevelReference, MatControlNode->GetKeywords());
 		Action->NodeTemplate = MatControlNode;
+		Action->SearchTitle = Action->NodeTemplate->GetNodeSearchTitle();
 	}
 }
 
@@ -2237,7 +2288,7 @@ void FK2ActionMenuBuilder::GetFunctionCallsOnSelectedActors(FBlueprintGraphActio
 	{
 		FText SelectActorsMsg = LOCTEXT("SelectActorForEvents", "Select Actor(s) to see available Events & Functions");
 		FText SelectActorsToolTip = LOCTEXT("SelectActorForEventsTooltip", "Select Actor(s) in the level to see available Events and Functions in this menu.");
-		TSharedPtr<FEdGraphSchemaAction_Dummy> MsgAction = TSharedPtr<FEdGraphSchemaAction_Dummy>(new FEdGraphSchemaAction_Dummy(TEXT(""), SelectActorsMsg.ToString(), SelectActorsToolTip.ToString(), K2Schema->AG_LevelReference));
+		TSharedPtr<FEdGraphSchemaAction_Dummy> MsgAction = TSharedPtr<FEdGraphSchemaAction_Dummy>(new FEdGraphSchemaAction_Dummy(TEXT(""), SelectActorsMsg, SelectActorsToolTip.ToString(), K2Schema->AG_LevelReference));
 		ContextMenuBuilder.AddAction(MsgAction);
 	}
 }
@@ -2288,7 +2339,7 @@ void FK2ActionMenuBuilder::GetFunctionCallsOnSelectedComponents(FBlueprintGraphA
 	{
 		FText SelectComponentMsg = LOCTEXT("SelectComponentForEvents", "Select a Component to see available Events & Functions");
 		FText SelectComponentToolTip = LOCTEXT("SelectComponentForEventsTooltip", "Select a Component in the MyBlueprint tab to see available Events and Functions in this menu.");
-		TSharedPtr<FEdGraphSchemaAction_Dummy> MsgAction = TSharedPtr<FEdGraphSchemaAction_Dummy>(new FEdGraphSchemaAction_Dummy(TEXT(""), SelectComponentMsg.ToString(), SelectComponentToolTip.ToString(), K2Schema->AG_LevelReference));
+		TSharedPtr<FEdGraphSchemaAction_Dummy> MsgAction = TSharedPtr<FEdGraphSchemaAction_Dummy>(new FEdGraphSchemaAction_Dummy(TEXT(""), SelectComponentMsg, SelectComponentToolTip.ToString(), K2Schema->AG_LevelReference));
 		ContextMenuBuilder.AddAction(MsgAction);
 	}
 }
@@ -2461,8 +2512,9 @@ void FK2ActionMenuBuilder::GetAllInterfaceMessageActions(FGraphActionListBuilder
 
 					UK2Node_Message* MessageTemplate = ActionMenuBuilder.CreateTemplateNode<UK2Node_Message>();
 					MessageTemplate->FunctionReference.SetExternalMember(Function->GetFName(), CurrentInterface);
-					TSharedPtr<FEdGraphSchemaAction_K2NewNode> Action = FK2ActionMenuBuilder::AddNewNodeAction(ActionMenuBuilder, MessageCategory, Function->GetName(), MessageTemplate->GetTooltip(), 0, MessageTemplate->GetKeywords());
+					TSharedPtr<FEdGraphSchemaAction_K2NewNode> Action = FK2ActionMenuBuilder::AddNewNodeAction(ActionMenuBuilder, MessageCategory, FText::FromString(Function->GetName()), MessageTemplate->GetTooltip(), 0, MessageTemplate->GetKeywords());
 					Action->NodeTemplate = MessageTemplate;
+					Action->SearchTitle = Action->NodeTemplate->GetNodeSearchTitle();
 				}
 			}
 		}
diff --git a/Engine/Source/Editor/BlueprintGraph/Private/K2Node.cpp b/Engine/Source/Editor/BlueprintGraph/Private/K2Node.cpp
--- a/Engine/Source/Editor/BlueprintGraph/Private/K2Node.cpp
+++ b/Engine/Source/Editor/BlueprintGraph/Private/K2Node.cpp
@@ -324,9 +324,12 @@ void UK2Node::RemovePinsFromOldPins(TArray<UEdGraphPin*>& OldPins, int32 Removed
 }
 
 void UK2Node::ReconstructNode()
-{		
-	FString TransactionDesc = FString::Printf(*LOCTEXT("ReconstructNodeTransaction", "'%s' Node Refreshed").ToString(), *GetNodeTitle(ENodeTitleType::FullTitle));
-	FScopedTransaction const Transaction(FText::FromString(TransactionDesc));
+{
+	FFormatNamedArguments Args;
+	Args.Add(TEXT("NodeTitle"), GetNodeTitle(ENodeTitleType::FullTitle));
+
+	FText TransactionDesc = FText::Format(LOCTEXT("ReconstructNodeTransaction", "'{NodeTitle}' Node Refreshed"), Args);
+	FScopedTransaction const Transaction(TransactionDesc);
 
 	Modify();
 
@@ -454,7 +457,7 @@ UK2Node::ERedirectType UK2Node::ShouldRedirectParam(const TArray<FString>& OldPi
 
 				// if it has it, return true
 				if (OldPinFNames.Contains(ParamRemap.OldParam)
-					&& (ParamRemap.NodeTitle == NAME_None || ParamRemap.NodeTitle.ToString() == NewPinNode->GetNodeTitle(ENodeTitleType::FullTitle)))
+					&& (ParamRemap.NodeTitle == NAME_None || ParamRemap.NodeTitle.ToString() == NewPinNode->GetNodeTitle(ENodeTitleType::FullTitle).ToString()))
 				{
 					NewPinName = ParamRemap.NewParam;
 					return (ParamRemap.bCustomValueMapping ? ERedirectType_Custom : (ParamRemap.ParamValueMap.Num() ? ERedirectType_Value : ERedirectType_Name));
@@ -892,9 +895,11 @@ void FOptionalPinManager::CreateVisiblePins(TArray<FOptionalPinFromProperty>& Pr
 						UEdGraphPin* NewPin = NULL;
 						if (PropertyEntry.bShowPin)
 						{
-							const FString PinName = FString::Printf(TEXT("%s_%d"), *(PropertyEntry.PropertyName.ToString()), Index);
-							const FString PinFriendlyName = PropertyEntry.PropertyFriendlyName.IsEmpty() ? PinName : 
-								FString::Printf(TEXT("%s_%d"), *PropertyEntry.PropertyFriendlyName, Index);
+							FFormatNamedArguments Args;
+							Args.Add(TEXT("PinName"), FText::FromName(PropertyEntry.PropertyName));
+							Args.Add(TEXT("Index"), Index);
+							const FText PinFriendlyName = FText::Format(LOCTEXT("PinFriendlyNameWithIndex", "{PinFriendlyName}_{Index}"), Args);
+							const FString PinName = PinFriendlyName.ToString();
 							NewPin = TargetNode->CreatePin(Direction, PinType.PinCategory, PinType.PinSubCategory, PinType.PinSubCategoryObject.Get(), PinType.bIsArray, PinType.bIsReference, PinName);
 							NewPin->PinFriendlyName = PinFriendlyName;
 
@@ -928,7 +933,7 @@ void FOptionalPinManager::CreateVisiblePins(TArray<FOptionalPinFromProperty>& Pr
 					{
 						const FString PinName = PropertyEntry.PropertyName.ToString();
 						NewPin = TargetNode->CreatePin(Direction, PinType.PinCategory, PinType.PinSubCategory, PinType.PinSubCategoryObject.Get(), PinType.bIsArray, PinType.bIsReference, PinName);
-						NewPin->PinFriendlyName = PropertyEntry.PropertyFriendlyName.IsEmpty() ? PinName : PropertyEntry.PropertyFriendlyName;
+						NewPin->PinFriendlyName = PropertyEntry.PropertyFriendlyName.IsEmpty() ? FText::FromString(PinName) : FText::FromString(PropertyEntry.PropertyFriendlyName);
 
 						// Allow the derived class to customize the created pin
 						CustomizePinData(NewPin, PropertyEntry.PropertyName, INDEX_NONE, OuterProperty);
diff --git a/Engine/Source/Editor/BlueprintGraph/Private/K2Node_AIMoveTo.cpp b/Engine/Source/Editor/BlueprintGraph/Private/K2Node_AIMoveTo.cpp
--- a/Engine/Source/Editor/BlueprintGraph/Private/K2Node_AIMoveTo.cpp
+++ b/Engine/Source/Editor/BlueprintGraph/Private/K2Node_AIMoveTo.cpp
@@ -5,6 +5,8 @@
 #include "EngineKismetLibraryClasses.h"
 #include "CompilerResultsLog.h"
 
+#define LOCTEXT_NAMESPACE "K2Node_AIMoveTo"
+
 UK2Node_AIMoveTo::UK2Node_AIMoveTo(const class FPostConstructInitializeProperties& PCIP)
 	: Super(PCIP)
 {
@@ -23,8 +25,11 @@ FString UK2Node_AIMoveTo::GetTooltip() const
 	return TEXT("Simple order for Pawn with AIController to move to a specific location");
 }
 
-FString UK2Node_AIMoveTo::GetNodeTitle(ENodeTitleType::Type TitleType) const
+FText UK2Node_AIMoveTo::GetNodeTitle(ENodeTitleType::Type TitleType) const
 {
-	return "AI MoveTo";
+	return LOCTEXT("AIMoveTo", "AI MoveTo");
 }
 
+#undef LOCTEXT_NAMESPACE
+
+
diff --git a/Engine/Source/Editor/BlueprintGraph/Private/K2Node_ActorBoundEvent.cpp b/Engine/Source/Editor/BlueprintGraph/Private/K2Node_ActorBoundEvent.cpp
--- a/Engine/Source/Editor/BlueprintGraph/Private/K2Node_ActorBoundEvent.cpp
+++ b/Engine/Source/Editor/BlueprintGraph/Private/K2Node_ActorBoundEvent.cpp
@@ -79,8 +79,24 @@ void UK2Node_ActorBoundEvent::DestroyNode()
 	Super::DestroyNode();
 }
 
-FString UK2Node_ActorBoundEvent::GetNodeTitle(ENodeTitleType::Type TitleType) const
+FText UK2Node_ActorBoundEvent::GetNodeTitle(ENodeTitleType::Type TitleType) const
 {
+	FText TargetName = LOCTEXT("None", "None");
+	if( EventOwner )
+	{
+		TargetName = FText::FromString(EventOwner->GetActorLabel());		
+	}
+
+	FFormatNamedArguments Args;
+	Args.Add(TEXT("DelegatePropertyName"), FText::FromName(DelegatePropertyName));
+	Args.Add(TEXT("TargetName"), TargetName);
+
+	return FText::Format(LOCTEXT("ActorBoundEventTitle", "{DelegatePropertyName} ({TargetName})"), Args);
+}
+
+FString UK2Node_ActorBoundEvent::GetNodeNativeTitle(ENodeTitleType::Type TitleType) const
+{
+	// Do not setup this function for localization, intentionally left unlocalized!
 	FString TargetName = TEXT("None");
 	if( EventOwner )
 	{
diff --git a/Engine/Source/Editor/BlueprintGraph/Private/K2Node_AddComponent.cpp b/Engine/Source/Editor/BlueprintGraph/Private/K2Node_AddComponent.cpp
--- a/Engine/Source/Editor/BlueprintGraph/Private/K2Node_AddComponent.cpp
+++ b/Engine/Source/Editor/BlueprintGraph/Private/K2Node_AddComponent.cpp
@@ -153,7 +153,10 @@ void UK2Node_AddComponent::ValidateNodeDuringCompilation(FCompilerResultsLog& Me
 		UClass* TemplateClass = Template->GetClass();
 		if (!TemplateClass->IsChildOf(UActorComponent::StaticClass()) || TemplateClass->HasAnyClassFlags(CLASS_Abstract) || !TemplateClass->HasMetaData(FBlueprintMetadata::MD_BlueprintSpawnableComponent) )
 		{
-			MessageLog.Error(*FString::Printf(*NSLOCTEXT("KismetCompiler", "InvalidComponentTemplate_Error", "Invalid class '%s' used as template by '%s' for @@").ToString(), *TemplateClass->GetName(), *GetNodeTitle(ENodeTitleType::FullTitle)), this);
+			FFormatNamedArguments Args;
+			Args.Add(TEXT("TemplateClass"), FText::FromString(TemplateClass->GetName()));
+			Args.Add(TEXT("NodeTitle"), GetNodeTitle(ENodeTitleType::FullTitle));
+			MessageLog.Error(*FText::Format(NSLOCTEXT("KismetCompiler", "InvalidComponentTemplate_Error", "Invalid class '{TemplateClass}' used as template by '{NodeTitle}' for @@"), Args).ToString(), this);
 		}
 
 		if (UChildActorComponent const* ChildActorComponent = Cast<UChildActorComponent const>(Template))
@@ -168,14 +171,18 @@ void UK2Node_AddComponent::ValidateNodeDuringCompilation(FCompilerResultsLog& Me
 
 				if (K2Schema->IsConstructionScript(ParentGraph))
 				{
-					MessageLog.Error(*FString::Printf(*NSLOCTEXT("KismetCompiler", "AddSelfComponent_Error", "@@ cannot add a '%s' component in the construction script (could cause infinite recursion).").ToString(), *ChildActorClass->GetName()), this);
+					FFormatNamedArguments Args;
+					Args.Add(TEXT("ChildActorClass"), FText::FromString(ChildActorClass->GetName()));
+					MessageLog.Error(*FText::Format(NSLOCTEXT("KismetCompiler", "AddSelfComponent_Error", "@@ cannot add a '{ChildActorClass}' component in the construction script (could cause infinite recursion)."), Args).ToString(), this);
 				}
 			}
 		}
 	}
 	else
 	{
-		MessageLog.Error(*FString::Printf(*NSLOCTEXT("KismetCompiler", "MissingComponentTemplate_Error", "Unknown template referenced by '%s' for @@").ToString(), *GetNodeTitle(ENodeTitleType::FullTitle)), this);
+		FFormatNamedArguments Args;
+		Args.Add(TEXT("NodeTitle"), GetNodeTitle(ENodeTitleType::FullTitle));
+		MessageLog.Error(*FText::Format(NSLOCTEXT("KismetCompiler", "MissingComponentTemplate_Error", "Unknown template referenced by '{NodeTitle}' for @@"), Args).ToString(), this);
 	}
 }
 
@@ -281,7 +288,7 @@ void UK2Node_AddComponent::PostPasteNode()
 	}
 }
 
-FString UK2Node_AddComponent::GetNodeTitle(ENodeTitleType::Type TitleType) const
+FText UK2Node_AddComponent::GetNodeTitle(ENodeTitleType::Type TitleType) const
 {
 	UEdGraphPin* TemplateNamePin = GetTemplateNamePin();
 	if (TemplateNamePin != NULL)
@@ -298,23 +305,33 @@ FString UK2Node_AddComponent::GetNodeTitle(ENodeTitleType::Type TitleType) const
 
 			if(StaticMeshComp != NULL && StaticMeshComp->StaticMesh != NULL)
 			{
-				return FString(TEXT("Add StaticMesh ")) + StaticMeshComp->StaticMesh->GetName();
+				FFormatNamedArguments Args;
+				Args.Add(TEXT("StaticMeshName"), FText::FromString(StaticMeshComp->StaticMesh->GetName()));
+				return FText::Format(LOCTEXT("AddStaticMesh", "Add StaticMesh {StaticMeshName}"), Args);
 			}
 			else if(SkelMeshComp != NULL && SkelMeshComp->SkeletalMesh != NULL)
 			{
-				return FString(TEXT("Add SkeletalMesh ")) + SkelMeshComp->SkeletalMesh->GetName();		
+				FFormatNamedArguments Args;
+				Args.Add(TEXT("SkeletalMeshName"), FText::FromString(SkelMeshComp->SkeletalMesh->GetName()));
+				return FText::Format(LOCTEXT("AddSkeletalMesh", "Add SkeletalMesh {SkeletalMeshName}"), Args);
 			}
 			else if(PSysComp != NULL && PSysComp->Template != NULL)
 			{
-				return FString(TEXT("Add ParticleSystem ")) + PSysComp->Template->GetName();		
+				FFormatNamedArguments Args;
+				Args.Add(TEXT("ParticleSystemName"), FText::FromString(PSysComp->Template->GetName()));
+				return FText::Format(LOCTEXT("AddParticleSystem", "Add ParticleSystem {ParticleSystemName}"), Args);
 			}
 			else if (SubActorComp && SubActorComp->ChildActorClass)
 			{
-				return FString(TEXT("Add ChildActorComponent ")) + SubActorComp->ChildActorClass->GetName();
+				FFormatNamedArguments Args;
+				Args.Add(TEXT("ComponentClassName"), FText::FromString(SubActorComp->ChildActorClass->GetName()));
+				return FText::Format(LOCTEXT("AddChildActorComponent", "Add ChildActorComponent {ComponentClassName}"), Args);
 			}
 			else
 			{
-				return FString(TEXT("Add ")) + SourceTemplate->GetClass()->GetName();		
+				FFormatNamedArguments Args;
+				Args.Add(TEXT("ClassName"), FText::FromString(SourceTemplate->GetClass()->GetName()));
+				return FText::Format(LOCTEXT("AddClass", "Add {ClassName}"), Args);		
 			}
 		}
 	}
diff --git a/Engine/Source/Editor/BlueprintGraph/Private/K2Node_AssignmentStatement.cpp b/Engine/Source/Editor/BlueprintGraph/Private/K2Node_AssignmentStatement.cpp
--- a/Engine/Source/Editor/BlueprintGraph/Private/K2Node_AssignmentStatement.cpp
+++ b/Engine/Source/Editor/BlueprintGraph/Private/K2Node_AssignmentStatement.cpp
@@ -77,9 +77,9 @@ FString UK2Node_AssignmentStatement::GetTooltip() const
 	return TEXT("Assigns Value to Variable");
 }
 
-FString UK2Node_AssignmentStatement::GetNodeTitle(ENodeTitleType::Type TitleType) const
+FText UK2Node_AssignmentStatement::GetNodeTitle(ENodeTitleType::Type TitleType) const
 {
-	return TEXT("Assign");
+	return LOCTEXT("Assign", "Assign");
 }
 
 void UK2Node_AssignmentStatement::NotifyPinConnectionListChanged(UEdGraphPin* Pin)
diff --git a/Engine/Source/Editor/BlueprintGraph/Private/K2Node_BreakStruct.cpp b/Engine/Source/Editor/BlueprintGraph/Private/K2Node_BreakStruct.cpp
--- a/Engine/Source/Editor/BlueprintGraph/Private/K2Node_BreakStruct.cpp
+++ b/Engine/Source/Editor/BlueprintGraph/Private/K2Node_BreakStruct.cpp
@@ -174,8 +174,16 @@ void UK2Node_BreakStruct::AllocateDefaultPins()
 	}
 }
 
-FString UK2Node_BreakStruct::GetNodeTitle(ENodeTitleType::Type TitleType) const
+FText UK2Node_BreakStruct::GetNodeTitle(ENodeTitleType::Type TitleType) const
 {
+	FFormatNamedArguments Args;
+	Args.Add(TEXT("StructName"), FText::FromString(StructType ? StructType->GetName() : FString()));
+	return FText::Format(LOCTEXT("BreakNodeTitle", "Break {StructName}"), Args);
+}
+
+FString UK2Node_BreakStruct::GetNodeNativeTitle(ENodeTitleType::Type TitleType) const
+{
+	// Do not setup this function for localization, intentionally left unlocalized!
 	return FString(TEXT("Break ")) + (StructType ? StructType->GetName() : FString());
 }
 
diff --git a/Engine/Source/Editor/BlueprintGraph/Private/K2Node_CallFunction.cpp b/Engine/Source/Editor/BlueprintGraph/Private/K2Node_CallFunction.cpp
--- a/Engine/Source/Editor/BlueprintGraph/Private/K2Node_CallFunction.cpp
+++ b/Engine/Source/Editor/BlueprintGraph/Private/K2Node_CallFunction.cpp
@@ -66,7 +66,7 @@ FString UK2Node_CallFunction::GetFunctionContextString() const
 }
 
 
-FString UK2Node_CallFunction::GetNodeTitle(ENodeTitleType::Type TitleType) const
+FText UK2Node_CallFunction::GetNodeTitle(ENodeTitleType::Type TitleType) const
 {
 	FString FunctionName;
 	FString ContextString;
@@ -87,11 +87,46 @@ FString UK2Node_CallFunction::GetNodeTitle(ENodeTitleType::Type TitleType) const
 		}
 	}
 
+	if(TitleType == ENodeTitleType::FullTitle)
+	{
+		FFormatNamedArguments Args;
+		Args.Add(TEXT("FunctionName"), FText::FromString(FunctionName));
+		Args.Add(TEXT("ContextString"), FText::FromString(ContextString));
+		Args.Add(TEXT("RPCString"), FText::FromString(RPCString));
+		return FText::Format(LOCTEXT("CallFunction_FullTitle", "{FunctionName}{ContextString}{RPCString}"), Args);
+	}
+	else
+	{
+		return FText::FromString(FunctionName);
+	}
+}
+
+FString UK2Node_CallFunction::GetNodeNativeTitle(ENodeTitleType::Type TitleType) const
+{
+	// Do not setup this function for localization, intentionally left unlocalized!
+	FString FunctionName;
+	FString ContextString;
+	FString RPCString;
+
+	if (UFunction* Function = GetTargetFunction())
+	{
+		RPCString = UK2Node_Event::GetLocalizedNetString(Function->FunctionFlags, true);
+		FunctionName = UK2Node_CallFunction::GetUserFacingFunctionName(Function);
+		ContextString = GetFunctionContextString();
+	}
+	else
+	{
+		FunctionName = FunctionReference.GetMemberName().ToString();
+		if ((GEditor != NULL) && (GetDefault<UEditorStyleSettings>()->bShowFriendlyNames))
+		{
+			FunctionName = FName::NameToDisplayString(FunctionName, false);
+		}
+	}
+
 	const FString Result = (TitleType == ENodeTitleType::ListView) ? FunctionName : FunctionName + ContextString + RPCString;
 	return Result;
 }
 
-
 void UK2Node_CallFunction::AllocateDefaultPins()
 {
 	UBlueprint* MyBlueprint = GetBlueprint();
@@ -249,7 +284,7 @@ void UK2Node_CallFunction::CreateExecPinsForFunctionCall(const UFunction* Functi
 		UEdGraphPin* OutputExecPin = CreatePin(EGPD_Output, K2Schema->PC_Exec, TEXT(""), NULL, false, false, K2Schema->PN_Then);
 		if(Function->HasMetaData(FBlueprintMetadata::MD_Latent))
 		{
-			OutputExecPin->PinFriendlyName = K2Schema->PN_Completed;
+			OutputExecPin->PinFriendlyName = FText::FromString(K2Schema->PN_Completed);
 		}
 	}
 }
@@ -422,7 +457,7 @@ bool UK2Node_CallFunction::CreatePinsForFunctionCall(const UFunction* Function)
 	UEdGraphPin* SelfPin = CreateSelfPin(Function);
 
 	//Renamed self pin to target
-	SelfPin->PinFriendlyName =  TEXT("Target");
+	SelfPin->PinFriendlyName =  LOCTEXT("Target", "Target");
 
 	// fill out the self-pin's default tool-tip
 	GeneratePinHoverText(*SelfPin, SelfPin->PinToolTip);
@@ -876,12 +911,12 @@ FString UK2Node_CallFunction::GetCompactNodeTitle(const UFunction* Function)
 	return Function->GetName();
 }
 
-FString UK2Node_CallFunction::GetCompactNodeTitle() const
+FText UK2Node_CallFunction::GetCompactNodeTitle() const
 {
 	UFunction* Function = GetTargetFunction();
 	if (Function != NULL)
 	{
-		return GetCompactNodeTitle(Function);
+		return FText::FromString(GetCompactNodeTitle(Function));
 	}
 	else
 	{
diff --git a/Engine/Source/Editor/BlueprintGraph/Private/K2Node_CallParentFunction.cpp b/Engine/Source/Editor/BlueprintGraph/Private/K2Node_CallParentFunction.cpp
--- a/Engine/Source/Editor/BlueprintGraph/Private/K2Node_CallParentFunction.cpp
+++ b/Engine/Source/Editor/BlueprintGraph/Private/K2Node_CallParentFunction.cpp
@@ -10,12 +10,28 @@ UK2Node_CallParentFunction::UK2Node_CallParentFunction(const class FPostConstruc
 	bIsFinalFunction = true;
 }
 
-FString UK2Node_CallParentFunction::GetNodeTitle(ENodeTitleType::Type TitleType) const
+FText UK2Node_CallParentFunction::GetNodeTitle(ENodeTitleType::Type TitleType) const
 {
 	UFunction* Function = FunctionReference.ResolveMember<UFunction>(this);
 	FString FunctionName = Function ? GetUserFacingFunctionName( Function ) : FunctionReference.GetMemberName().ToString();
-	FString Result = FText::Format( LOCTEXT( "CallSuperFunction", "Parent: {0}" ), FText::FromString( FunctionName )).ToString();
-	
+
+	if( GEditor && GetDefault<UEditorStyleSettings>()->bShowFriendlyNames )
+	{
+		FunctionName = FName::NameToDisplayString(FunctionName, false);
+	}
+
+	FFormatNamedArguments Args;
+	Args.Add(TEXT("FunctionName"),  FText::FromString( FunctionName ));
+	return FText::Format( LOCTEXT( "CallSuperFunction", "Parent: {FunctionName}" ), Args);
+}
+
+FString UK2Node_CallParentFunction::GetNodeNativeTitle(ENodeTitleType::Type TitleType) const
+{
+	// Do not setup this function for localization, intentionally left unlocalized!
+	UFunction* Function = FunctionReference.ResolveMember<UFunction>(this);
+	FString FunctionName = Function ? GetUserFacingFunctionName( Function ) : FunctionReference.GetMemberName().ToString();
+	FString Result = FString::Printf( TEXT("Parent: %s" ), *FunctionName);
+
 	if( GEditor && GetDefault<UEditorStyleSettings>()->bShowFriendlyNames )
 	{
 		Result = FName::NameToDisplayString(Result, false);
diff --git a/Engine/Source/Editor/BlueprintGraph/Private/K2Node_CastByteToEnum.cpp b/Engine/Source/Editor/BlueprintGraph/Private/K2Node_CastByteToEnum.cpp
--- a/Engine/Source/Editor/BlueprintGraph/Private/K2Node_CastByteToEnum.cpp
+++ b/Engine/Source/Editor/BlueprintGraph/Private/K2Node_CastByteToEnum.cpp
@@ -29,14 +29,14 @@ FString UK2Node_CastByteToEnum::GetTooltip() const
 		*Enum->GetName());
 }
 
-FString UK2Node_CastByteToEnum::GetNodeTitle(ENodeTitleType::Type TitleType) const
+FText UK2Node_CastByteToEnum::GetNodeTitle(ENodeTitleType::Type TitleType) const
 {
-	return GetTooltip();
+	return FText::FromString(GetTooltip());
 }
 
-FString UK2Node_CastByteToEnum::GetCompactNodeTitle() const
+FText UK2Node_CastByteToEnum::GetCompactNodeTitle() const
 {
-	return TEXT("\x2022");
+	return NSLOCTEXT("K2Node", "CastSymbol", "\x2022");
 }
 
 FName UK2Node_CastByteToEnum::GetFunctionName() const
diff --git a/Engine/Source/Editor/BlueprintGraph/Private/K2Node_ComponentBoundEvent.cpp b/Engine/Source/Editor/BlueprintGraph/Private/K2Node_ComponentBoundEvent.cpp
--- a/Engine/Source/Editor/BlueprintGraph/Private/K2Node_ComponentBoundEvent.cpp
+++ b/Engine/Source/Editor/BlueprintGraph/Private/K2Node_ComponentBoundEvent.cpp
@@ -4,13 +4,24 @@
 
 #include "EngineLevelScriptClasses.h"
 
+#define LOCTEXT_NAMESPACE "K2Node"
+
 UK2Node_ComponentBoundEvent::UK2Node_ComponentBoundEvent(const class FPostConstructInitializeProperties& PCIP)
 	: Super(PCIP)
 {
 }
 
-FString UK2Node_ComponentBoundEvent::GetNodeTitle(ENodeTitleType::Type TitleType) const
+FText UK2Node_ComponentBoundEvent::GetNodeTitle(ENodeTitleType::Type TitleType) const
+{
+	FFormatNamedArguments Args;
+	Args.Add(TEXT("DelegatePropertyName"), FText::FromName(DelegatePropertyName));
+	Args.Add(TEXT("ComponentPropertyName"), FText::FromName(ComponentPropertyName));
+	return FText::Format(LOCTEXT("ComponentBoundEvent_Title", "{DelegatePropertyName} ({ComponentPropertyName})"), Args);
+}
+
+FString UK2Node_ComponentBoundEvent::GetNodeNativeTitle(ENodeTitleType::Type TitleType) const
 {
+	// Do not setup this function for localization, intentionally left unlocalized!
 	return FString::Printf(TEXT("%s (%s)"), *DelegatePropertyName.ToString(), *ComponentPropertyName.ToString());
 }
 
@@ -83,4 +94,6 @@ FString UK2Node_ComponentBoundEvent::GetDocumentationLink() const
 FString UK2Node_ComponentBoundEvent::GetDocumentationExcerptName() const
 {
 	return DelegatePropertyName.ToString();
-}
\ No newline at end of file
+}
+
+#undef LOCTEXT_NAMESPACE
\ No newline at end of file
diff --git a/Engine/Source/Editor/BlueprintGraph/Private/K2Node_Composite.cpp b/Engine/Source/Editor/BlueprintGraph/Private/K2Node_Composite.cpp
--- a/Engine/Source/Editor/BlueprintGraph/Private/K2Node_Composite.cpp
+++ b/Engine/Source/Editor/BlueprintGraph/Private/K2Node_Composite.cpp
@@ -162,11 +162,26 @@ FLinearColor UK2Node_Composite::GetNodeTitleColor() const
 	return FLinearColor::White;
 }
 
-FString UK2Node_Composite::GetNodeTitle(ENodeTitleType::Type TitleType) const
+FText UK2Node_Composite::GetNodeTitle(ENodeTitleType::Type TitleType) const
 {
 	if(TitleType == ENodeTitleType::FullTitle)
 	{
-		return FString::Printf(*LOCTEXT("Collapsed_Name", "%s\nCollapsed Graph").ToString(), (BoundGraph)?*BoundGraph->GetName() : TEXT("Invalid Graph"));
+		FFormatNamedArguments Args;
+		Args.Add(TEXT("BoundGraphName"), (BoundGraph)? FText::FromString(BoundGraph->GetName()) : LOCTEXT("InvalidGraph", "Invalid Graph"));
+		return FText::Format(LOCTEXT("Collapsed_Name", "{BoundGraphName}\nCollapsed Graph"), Args);
+	}
+	else
+	{
+		return (BoundGraph)? FText::FromString(BoundGraph->GetName()) : LOCTEXT("InvalidGraph", "Invalid Graph");
+	}
+}
+
+FString UK2Node_Composite::GetNodeNativeTitle(ENodeTitleType::Type TitleType) const
+{
+	// Do not setup this function for localization, intentionally left unlocalized!
+	if(TitleType == ENodeTitleType::FullTitle)
+	{
+		return FString::Printf(TEXT("%s\nCollapsed Graph"), (BoundGraph)?*BoundGraph->GetName() : TEXT("Invalid Graph"));
 	}
 	else
 	{
diff --git a/Engine/Source/Editor/BlueprintGraph/Private/K2Node_CreateDelegate.cpp b/Engine/Source/Editor/BlueprintGraph/Private/K2Node_CreateDelegate.cpp
--- a/Engine/Source/Editor/BlueprintGraph/Private/K2Node_CreateDelegate.cpp
+++ b/Engine/Source/Editor/BlueprintGraph/Private/K2Node_CreateDelegate.cpp
@@ -24,12 +24,12 @@ void UK2Node_CreateDelegate::AllocateDefaultPins()
 
 	if(UEdGraphPin* ObjPin = CreatePin(EGPD_Input, K2Schema->PC_Object, TEXT(""), UObject::StaticClass(), false, false, FK2Node_CreateDelegate_Helper::ObjectInputName))
 	{
-		ObjPin->PinFriendlyName = *NSLOCTEXT("K2Node", "CreateDelegate_ObjectInputName", "Object").ToString();
+		ObjPin->PinFriendlyName = NSLOCTEXT("K2Node", "CreateDelegate_ObjectInputName", "Object");
 	}
 
 	if(UEdGraphPin* DelegatePin = CreatePin(EGPD_Output, K2Schema->PC_Delegate, TEXT(""), NULL, false, false, FK2Node_CreateDelegate_Helper::DelegateOutputName))
 	{
-		DelegatePin->PinFriendlyName = *NSLOCTEXT("K2Node", "CreateDelegate_DelegateOutName", "Event").ToString();
+		DelegatePin->PinFriendlyName = NSLOCTEXT("K2Node", "CreateDelegate_DelegateOutName", "Event");
 	}
 
 	Super::AllocateDefaultPins();
@@ -275,9 +275,9 @@ UEdGraphPin* UK2Node_CreateDelegate::GetObjectInPin() const
 	return FindPinChecked(FK2Node_CreateDelegate_Helper::ObjectInputName);
 }
 
-FString UK2Node_CreateDelegate::GetNodeTitle(ENodeTitleType::Type TitleType) const
+FText UK2Node_CreateDelegate::GetNodeTitle(ENodeTitleType::Type TitleType) const
 {
-	return FString::Printf(*NSLOCTEXT("K2Node", "CreateDelegate", "Create Event").ToString());
+	return NSLOCTEXT("K2Node", "CreateDelegate", "Create Event");
 }
 
 UObject* UK2Node_CreateDelegate::GetJumpTargetForDoubleClick() const
diff --git a/Engine/Source/Editor/BlueprintGraph/Private/K2Node_CustomEvent.cpp b/Engine/Source/Editor/BlueprintGraph/Private/K2Node_CustomEvent.cpp
--- a/Engine/Source/Editor/BlueprintGraph/Private/K2Node_CustomEvent.cpp
+++ b/Engine/Source/Editor/BlueprintGraph/Private/K2Node_CustomEvent.cpp
@@ -103,20 +103,42 @@ UK2Node_CustomEvent::UK2Node_CustomEvent(const class FPostConstructInitializePro
 	bCanRenameNode = true;
 }
 
-FString UK2Node_CustomEvent::GetNodeTitle(ENodeTitleType::Type TitleType) const
+FText UK2Node_CustomEvent::GetNodeTitle(ENodeTitleType::Type TitleType) const
 {
+	if (TitleType == ENodeTitleType::EditableTitle)
+	{
+		return FText::FromName(CustomFunctionName);
+	}
+	else if(TitleType == ENodeTitleType::ListView)
+	{
+		return NSLOCTEXT("K2Node", "CustomEvent_Title", "Custom Event");
+	}
+	else
+	{
+		FString RPCString = UK2Node_Event::GetLocalizedNetString(FunctionFlags, false);
+
+		FFormatNamedArguments Args;
+		Args.Add(TEXT("FunctionName"), FText::FromName(CustomFunctionName));
+		Args.Add(TEXT("RPCString"), FText::FromString(RPCString));
+		return FText::Format(NSLOCTEXT("K2Node", "CustomEvent_Name", "{FunctionName}{RPCString}\nCustom Event"), Args);
+	}
+}
+
+FString UK2Node_CustomEvent::GetNodeNativeTitle(ENodeTitleType::Type TitleType) const
+{
+	// Do not setup this function for localization, intentionally left unlocalized!
 	if (TitleType == ENodeTitleType::EditableTitle)
 	{
 		return CustomFunctionName.ToString();
 	}
 	else if(TitleType == ENodeTitleType::ListView)
 	{
-		return NSLOCTEXT("K2Node", "CustomEvent_Title", "Custom Event").ToString();
+		return TEXT("Custom Event");
 	}
 	else
 	{
 		FString RPCString = UK2Node_Event::GetLocalizedNetString(FunctionFlags, false);
-		return FString::Printf(*NSLOCTEXT("K2Node", "CustomEvent_Name", "%s\nCustom Event").ToString(), *CustomFunctionName.ToString()) + RPCString;
+		return FString::Printf(TEXT("%s\nCustom Event"), *CustomFunctionName.ToString()) + RPCString;
 	}
 }
 
diff --git a/Engine/Source/Editor/BlueprintGraph/Private/K2Node_DelegateSet.cpp b/Engine/Source/Editor/BlueprintGraph/Private/K2Node_DelegateSet.cpp
--- a/Engine/Source/Editor/BlueprintGraph/Private/K2Node_DelegateSet.cpp
+++ b/Engine/Source/Editor/BlueprintGraph/Private/K2Node_DelegateSet.cpp
@@ -162,9 +162,17 @@ FString UK2Node_DelegateSet::GetTooltip() const
 	return Tooltip;
 }
 
-FString UK2Node_DelegateSet::GetNodeTitle(ENodeTitleType::Type TitleType) const
+FText UK2Node_DelegateSet::GetNodeTitle(ENodeTitleType::Type TitleType) const
 {
-	return FString::Printf(*NSLOCTEXT("K2Node", "Assign_Name", "Assign %s").ToString(), *DelegatePropertyName.ToString());
+	FFormatNamedArguments Args;
+	Args.Add(TEXT("DelegatePropertyName"), FText::FromName(DelegatePropertyName));
+	return FText::Format(NSLOCTEXT("K2Node", "Assign_Name", "Assign {DelegatePropertyName}"), Args);
+}
+
+FString UK2Node_DelegateSet::GetNodeNativeTitle(ENodeTitleType::Type TitleType) const
+{
+	// Do not setup this function for localization, intentionally left unlocalized!
+	return FString::Printf(TEXT("Assign %s"), *DelegatePropertyName.ToString());
 }
 
 UEdGraphPin* UK2Node_DelegateSet::GetDelegateOwner() const
diff --git a/Engine/Source/Editor/BlueprintGraph/Private/K2Node_DynamicCast.cpp b/Engine/Source/Editor/BlueprintGraph/Private/K2Node_DynamicCast.cpp
--- a/Engine/Source/Editor/BlueprintGraph/Private/K2Node_DynamicCast.cpp
+++ b/Engine/Source/Editor/BlueprintGraph/Private/K2Node_DynamicCast.cpp
@@ -49,8 +49,36 @@ FLinearColor UK2Node_DynamicCast::GetNodeTitleColor() const
 	return FLinearColor(0.0f, 0.55f, 0.62f);
 }
 
-FString UK2Node_DynamicCast::GetNodeTitle(ENodeTitleType::Type TitleType) const
+FText UK2Node_DynamicCast::GetNodeTitle(ENodeTitleType::Type TitleType) const
 {
+	if(TargetType != NULL)
+	{
+		// If casting to BP class, use BP name not class name (ie. remove the _C)
+		FString TargetName;
+		UBlueprint* CastToBP = UBlueprint::GetBlueprintFromClass(TargetType);
+		if(CastToBP != NULL)
+		{
+			TargetName = CastToBP->GetName();
+		}
+		else
+		{
+			TargetName = TargetType->GetName();
+		}
+
+		FFormatNamedArguments Args;
+		Args.Add(TEXT("TargetName"), FText::FromString(TargetName));
+		return FText::Format(NSLOCTEXT("K2Node_DynamicCast", "CastTo", "Cast To {TargetName}"), Args);
+	}
+	// No target type, bad node
+	else
+	{
+		return NSLOCTEXT("K2Node_DynamicCast", "BadCastNode", "Bad cast node");
+	}
+}
+
+FString UK2Node_DynamicCast::GetNodeNativeTitle(ENodeTitleType::Type TitleType) const
+{
+	// Do not setup this function for localization, intentionally left unlocalized!
 	if(TargetType != NULL)
 	{
 		// If casting to BP class, use BP name not class name (ie. remove the _C)
@@ -70,7 +98,7 @@ FString UK2Node_DynamicCast::GetNodeTitle(ENodeTitleType::Type TitleType) const
 	// No target type, bad node
 	else
 	{
-		return FString(TEXT("Bad cast node"));
+		return TEXT("Bad cast node");
 	}
 }
 
diff --git a/Engine/Source/Editor/BlueprintGraph/Private/K2Node_EnumEquality.cpp b/Engine/Source/Editor/BlueprintGraph/Private/K2Node_EnumEquality.cpp
--- a/Engine/Source/Editor/BlueprintGraph/Private/K2Node_EnumEquality.cpp
+++ b/Engine/Source/Editor/BlueprintGraph/Private/K2Node_EnumEquality.cpp
@@ -30,9 +30,9 @@ FString UK2Node_EnumEquality::GetTooltip() const
 	return TEXT("Returns true if A is equal to B (A == B)");
 }
 
-FString UK2Node_EnumEquality::GetNodeTitle(ENodeTitleType::Type TitleType) const
+FText UK2Node_EnumEquality::GetNodeTitle(ENodeTitleType::Type TitleType) const
 {
-	return TEXT("Equal (Enum)");
+	return NSLOCTEXT("K2Node_EnumEquality", "EqualEnum", "Equal (Enum)");
 }
 
 FString UK2Node_EnumEquality::GetKeywords() const
@@ -186,12 +186,13 @@ void UK2Node_EnumEquality::GetMenuEntries(FGraphContextMenuBuilder& Context) con
 		UK2Node* EnumNodeTemplate = Context.CreateTemplateNode<UK2Node_EnumEquality>();
 
 		const FString Category = FunctionCategory + TEXT("|Utilities| Enum");
-		const FString MenuDesc = EnumNodeTemplate->GetNodeTitle(ENodeTitleType::ListView);
+		const FText MenuDesc = EnumNodeTemplate->GetNodeTitle(ENodeTitleType::ListView);
 		const FString Tooltip = EnumNodeTemplate->GetTooltip();
 		const FString Keywords = EnumNodeTemplate->GetKeywords();
 
 		TSharedPtr<FEdGraphSchemaAction_K2NewNode> NodeAction = FK2ActionMenuBuilder::AddNewNodeAction(Context, Category, MenuDesc, Tooltip, 0, Keywords);
 		NodeAction->NodeTemplate = EnumNodeTemplate;
+		NodeAction->SearchTitle = EnumNodeTemplate->GetNodeSearchTitle();
 	}
 }
 
diff --git a/Engine/Source/Editor/BlueprintGraph/Private/K2Node_EnumLiteral.cpp b/Engine/Source/Editor/BlueprintGraph/Private/K2Node_EnumLiteral.cpp
--- a/Engine/Source/Editor/BlueprintGraph/Private/K2Node_EnumLiteral.cpp
+++ b/Engine/Source/Editor/BlueprintGraph/Private/K2Node_EnumLiteral.cpp
@@ -37,14 +37,16 @@ FName UK2Node_EnumLiteral::GetPaletteIcon(FLinearColor& OutColor) const
 
 FString UK2Node_EnumLiteral::GetTooltip() const
 {
-	const FString EnumName = Enum ? Enum->GetName() : TEXT("(bad enum)");
+	const FText EnumName = Enum ? FText::FromString(Enum->GetName()) : NSLOCTEXT("K2Node", "BadEnum", "(bad enum)");
 
-	return FString::Printf(*NSLOCTEXT("K2Node", "EnumLiteral_Tooltip", "Literal enum %s").ToString(), *EnumName);
+	FFormatNamedArguments Args;
+	Args.Add(TEXT("EnumName"), EnumName);
+	return FText::Format(NSLOCTEXT("K2Node", "EnumLiteral_Tooltip", "Literal enum {EnumName}"), Args).ToString();
 }
 
-FString UK2Node_EnumLiteral::GetNodeTitle(ENodeTitleType::Type TitleType) const
+FText UK2Node_EnumLiteral::GetNodeTitle(ENodeTitleType::Type TitleType) const
 {
-	return GetTooltip();
+	return FText::FromString(GetTooltip());
 }
 
 class FKCHandler_EnumLiteral : public FNodeHandlingFunctor
diff --git a/Engine/Source/Editor/BlueprintGraph/Private/K2Node_Event.cpp b/Engine/Source/Editor/BlueprintGraph/Private/K2Node_Event.cpp
--- a/Engine/Source/Editor/BlueprintGraph/Private/K2Node_Event.cpp
+++ b/Engine/Source/Editor/BlueprintGraph/Private/K2Node_Event.cpp
@@ -50,7 +50,7 @@ FLinearColor UK2Node_Event::GetNodeTitleColor() const
 	return GEditor->AccessEditorUserSettings().EventNodeTitleColor;
 }
 
-FString UK2Node_Event::GetNodeTitle(ENodeTitleType::Type TitleType) const
+FText UK2Node_Event::GetNodeTitle(ENodeTitleType::Type TitleType) const
 {
 	if (bOverrideFunction || (CustomFunctionName == NAME_None))
 	{
@@ -61,7 +61,9 @@ FString UK2Node_Event::GetNodeTitle(ENodeTitleType::Type TitleType) const
 			FunctionName = UEdGraphSchema_K2::GetFriendlySignitureName(Function);
 		}
 
-		FString TitleString = NSLOCTEXT("K2Node", "Event_Name", "Event ").ToString() + FunctionName;
+		FFormatNamedArguments Args;
+		Args.Add(TEXT("FunctionName"), FText::FromString(FunctionName));
+		FText Title = FText::Format(NSLOCTEXT("K2Node", "Event_Name", "Event {FunctionName}"), Args);
 
 		if(TitleType == ENodeTitleType::FullTitle && EventSignatureClass != NULL && EventSignatureClass->IsChildOf(UInterface::StaticClass()))
 		{
@@ -73,7 +75,46 @@ FString UK2Node_Event::GetNodeTitle(ENodeTitleType::Type TitleType) const
 				SourceString = SourceString.LeftChop(2);
 			}
 
-			TitleString += FString(TEXT("\n")) + FString::Printf(*LOCTEXT("EventFromInterface", "From %s").ToString(), *SourceString);
+			FFormatNamedArguments FullTitleArgs;
+			FullTitleArgs.Add(TEXT("Title"), Title);
+			FullTitleArgs.Add(TEXT("InterfaceClass"), FText::FromString(SourceString));
+
+			Title = FText::Format(LOCTEXT("EventFromInterface", "{Title}\nFrom {InterfaceClass}"), FullTitleArgs);
+		}
+
+		return Title;
+	}
+	else
+	{
+		return FText::FromName(CustomFunctionName);
+	}
+}
+
+FString UK2Node_Event::GetNodeNativeTitle(ENodeTitleType::Type TitleType) const
+{
+	// Do not setup this function for localization, intentionally left unlocalized!
+	if (bOverrideFunction || (CustomFunctionName == NAME_None))
+	{
+		FString FunctionName = EventSignatureName.ToString(); // If we fail to find the function, still want to write something on the node.
+
+		if (UFunction* Function = FindField<UFunction>(EventSignatureClass, EventSignatureName))
+		{
+			FunctionName = UEdGraphSchema_K2::GetFriendlySignitureName(Function);
+		}
+
+		FString TitleString = FString(TEXT("Event ")) + FunctionName;
+
+		if(TitleType == ENodeTitleType::FullTitle && EventSignatureClass != NULL && EventSignatureClass->IsChildOf(UInterface::StaticClass()))
+		{
+			FString SourceString = EventSignatureClass->GetName();
+
+			// @todo: This action won't be necessary once the new name convention is used.
+			if(SourceString.EndsWith(TEXT("_C")))
+			{
+				SourceString = SourceString.LeftChop(2);
+			}
+
+			TitleString += FString(TEXT("\n")) + FString::Printf(TEXT("From %s"), *SourceString);
 		}
 
 		return TitleString;
diff --git a/Engine/Source/Editor/BlueprintGraph/Private/K2Node_ExecutionSequence.cpp b/Engine/Source/Editor/BlueprintGraph/Private/K2Node_ExecutionSequence.cpp
--- a/Engine/Source/Editor/BlueprintGraph/Private/K2Node_ExecutionSequence.cpp
+++ b/Engine/Source/Editor/BlueprintGraph/Private/K2Node_ExecutionSequence.cpp
@@ -136,9 +136,9 @@ void UK2Node_ExecutionSequence::AllocateDefaultPins()
 	Super::AllocateDefaultPins();
 }
 
-FString UK2Node_ExecutionSequence::GetNodeTitle(ENodeTitleType::Type TitleType) const
+FText UK2Node_ExecutionSequence::GetNodeTitle(ENodeTitleType::Type TitleType) const
 {
-	return NSLOCTEXT("K2Node", "Sequence", "Sequence").ToString();
+	return NSLOCTEXT("K2Node", "Sequence", "Sequence");
 }
 
 FLinearColor UK2Node_ExecutionSequence::GetNodeTitleColor() const
diff --git a/Engine/Source/Editor/BlueprintGraph/Private/K2Node_ForEachElementInEnum.cpp b/Engine/Source/Editor/BlueprintGraph/Private/K2Node_ForEachElementInEnum.cpp
--- a/Engine/Source/Editor/BlueprintGraph/Private/K2Node_ForEachElementInEnum.cpp
+++ b/Engine/Source/Editor/BlueprintGraph/Private/K2Node_ForEachElementInEnum.cpp
@@ -6,6 +6,8 @@
 #include "../../../Runtime/Engine/Classes/Kismet/KismetNodeHelperLibrary.h"
 #include "../../../Runtime/Engine/Classes/Kismet/KismetMathLibrary.h"
 
+#define LOCTEXT_NAMESPACE "K2Node"
+
 struct FForExpandNodeHelper
 {
 	UEdGraphPin* StartLoopExecInPin;
@@ -121,14 +123,22 @@ void UK2Node_ForEachElementInEnum::ValidateNodeDuringCompilation(FCompilerResult
 
 FString UK2Node_ForEachElementInEnum::GetTooltip() const
 {
-	return FString::Printf( 
-		*NSLOCTEXT("K2Node", "ForEachElementInEnum_Tooltip", "ForEach %s").ToString(),
-		Enum ? *Enum->GetName() : TEXT("UNKNOWN"));
+	return GetNodeTitle(ENodeTitleType::FullTitle).ToString();
 }
 
-FString UK2Node_ForEachElementInEnum::GetNodeTitle(ENodeTitleType::Type TitleType) const
+FText UK2Node_ForEachElementInEnum::GetNodeTitle(ENodeTitleType::Type TitleType) const
 {
-	return GetTooltip();
+	FFormatNamedArguments Args;
+	Args.Add(TEXT("EnumName"), Enum ? FText::FromString(Enum->GetName()) : LOCTEXT("Unknown", "UNKNOWN"));
+	return FText::Format(LOCTEXT("ForEachElementInEnum_Title", "ForEach {EnumName}"), Args);
+}
+
+FString UK2Node_ForEachElementInEnum::GetNodeNativeTitle(ENodeTitleType::Type TitleType) const
+{
+	// Do not setup this function for localization, intentionally left unlocalized!
+	return FString::Printf( 
+		TEXT("ForEach %s"),
+		Enum ? *Enum->GetName() : TEXT("UNKNOWN"));
 }
 
 void UK2Node_ForEachElementInEnum::ExpandNode(class FKismetCompilerContext& CompilerContext, UEdGraph* SourceGraph)
@@ -179,3 +189,5 @@ void UK2Node_ForEachElementInEnum::ExpandNode(class FKismetCompilerContext& Comp
 		BreakAllNodeLinks();
 	}
 }
+
+#undef LOCTEXT_NAMESPACE
\ No newline at end of file
diff --git a/Engine/Source/Editor/BlueprintGraph/Private/K2Node_FormatText.cpp b/Engine/Source/Editor/BlueprintGraph/Private/K2Node_FormatText.cpp
--- a/Engine/Source/Editor/BlueprintGraph/Private/K2Node_FormatText.cpp
+++ b/Engine/Source/Editor/BlueprintGraph/Private/K2Node_FormatText.cpp
@@ -33,9 +33,9 @@ void UK2Node_FormatText::AllocateDefaultPins()
 	}
 }
 
-FString UK2Node_FormatText::GetNodeTitle(ENodeTitleType::Type TitleType) const
+FText UK2Node_FormatText::GetNodeTitle(ENodeTitleType::Type TitleType) const
 {
-	return LOCTEXT("FormatText_Title", "Format Text").ToString();
+	return LOCTEXT("FormatText_Title", "Format Text");
 }
 
 FString UK2Node_FormatText::GetPinDisplayName(const UEdGraphPin* Pin) const
diff --git a/Engine/Source/Editor/BlueprintGraph/Private/K2Node_FunctionEntry.cpp b/Engine/Source/Editor/BlueprintGraph/Private/K2Node_FunctionEntry.cpp
--- a/Engine/Source/Editor/BlueprintGraph/Private/K2Node_FunctionEntry.cpp
+++ b/Engine/Source/Editor/BlueprintGraph/Private/K2Node_FunctionEntry.cpp
@@ -99,16 +99,22 @@ UK2Node_FunctionEntry::UK2Node_FunctionEntry(const class FPostConstructInitializ
 {
 }
 
-FString UK2Node_FunctionEntry::GetNodeTitle(ENodeTitleType::Type TitleType) const
+FText UK2Node_FunctionEntry::GetNodeTitle(ENodeTitleType::Type TitleType) const
 {
 	UEdGraph* Graph = GetGraph();
 	FGraphDisplayInfo DisplayInfo;
 	Graph->GetSchema()->GetGraphDisplayInformation(*Graph, DisplayInfo);
 
-	if( GEditor && GetDefault<UEditorStyleSettings>()->bShowFriendlyNames )
-	{
-		DisplayInfo.DisplayName = FText::FromString( FName::NameToDisplayString(DisplayInfo.DisplayName.ToString(), false) );
-	}
+	return DisplayInfo.DisplayName;
+}
+
+FString UK2Node_FunctionEntry::GetNodeNativeTitle(ENodeTitleType::Type TitleType) const
+{
+	// Do not setup this function for localization, intentionally left unlocalized!
+	UEdGraph* Graph = GetGraph();
+	FGraphDisplayInfo DisplayInfo;
+	Graph->GetSchema()->GetGraphDisplayInformation(*Graph, DisplayInfo);
+
 	return DisplayInfo.DisplayName.ToString();
 }
 
diff --git a/Engine/Source/Editor/BlueprintGraph/Private/K2Node_FunctionResult.cpp b/Engine/Source/Editor/BlueprintGraph/Private/K2Node_FunctionResult.cpp
--- a/Engine/Source/Editor/BlueprintGraph/Private/K2Node_FunctionResult.cpp
+++ b/Engine/Source/Editor/BlueprintGraph/Private/K2Node_FunctionResult.cpp
@@ -72,9 +72,9 @@ UK2Node_FunctionResult::UK2Node_FunctionResult(const class FPostConstructInitial
 {
 }
 
-FString UK2Node_FunctionResult::GetNodeTitle(ENodeTitleType::Type TitleType) const
+FText UK2Node_FunctionResult::GetNodeTitle(ENodeTitleType::Type TitleType) const
 {
-	return NSLOCTEXT("K2Node", "ReturnNode", "ReturnNode").ToString();
+	return NSLOCTEXT("K2Node", "ReturnNode", "ReturnNode");
 }
 
 void UK2Node_FunctionResult::AllocateDefaultPins()
diff --git a/Engine/Source/Editor/BlueprintGraph/Private/K2Node_GetEnumeratorName.cpp b/Engine/Source/Editor/BlueprintGraph/Private/K2Node_GetEnumeratorName.cpp
--- a/Engine/Source/Editor/BlueprintGraph/Private/K2Node_GetEnumeratorName.cpp
+++ b/Engine/Source/Editor/BlueprintGraph/Private/K2Node_GetEnumeratorName.cpp
@@ -26,14 +26,14 @@ FString UK2Node_GetEnumeratorName::GetTooltip() const
 	return NSLOCTEXT("K2Node", "GetEnumeratorName_Tooltip", "Returns name of enumerator").ToString();
 }
 
-FString UK2Node_GetEnumeratorName::GetNodeTitle(ENodeTitleType::Type TitleType) const
+FText UK2Node_GetEnumeratorName::GetNodeTitle(ENodeTitleType::Type TitleType) const
 {
-	return NSLOCTEXT("K2Node", "GetNode_Title", "Enum to Name").ToString();
+	return NSLOCTEXT("K2Node", "GetNode_Title", "Enum to Name");
 }
 
-FString UK2Node_GetEnumeratorName::GetCompactNodeTitle() const
+FText UK2Node_GetEnumeratorName::GetCompactNodeTitle() const
 {
-	return TEXT("\x2022");
+	return NSLOCTEXT("K2Node", "CastSymbol", "\x2022");
 }
 
 UEnum* UK2Node_GetEnumeratorName::GetEnum() const
diff --git a/Engine/Source/Editor/BlueprintGraph/Private/K2Node_GetEnumeratorNameAsString.cpp b/Engine/Source/Editor/BlueprintGraph/Private/K2Node_GetEnumeratorNameAsString.cpp
--- a/Engine/Source/Editor/BlueprintGraph/Private/K2Node_GetEnumeratorNameAsString.cpp
+++ b/Engine/Source/Editor/BlueprintGraph/Private/K2Node_GetEnumeratorNameAsString.cpp
@@ -23,9 +23,9 @@ FString UK2Node_GetEnumeratorNameAsString::GetTooltip() const
 	return NSLOCTEXT("K2Node", "GetEnumeratorNameAsString_Tooltip", "Returns user friendly name of enumerator").ToString();
 }
 
-FString UK2Node_GetEnumeratorNameAsString::GetNodeTitle(ENodeTitleType::Type TitleType) const
+FText UK2Node_GetEnumeratorNameAsString::GetNodeTitle(ENodeTitleType::Type TitleType) const
 {
-	return NSLOCTEXT("K2Node", "GetEnumeratorNameAsString_Title", "Enum to String").ToString();
+	return NSLOCTEXT("K2Node", "GetEnumeratorNameAsString_Title", "Enum to String");
 }
 
 FName UK2Node_GetEnumeratorNameAsString::GetFunctionName() const
diff --git a/Engine/Source/Editor/BlueprintGraph/Private/K2Node_GetInputAxisValue.cpp b/Engine/Source/Editor/BlueprintGraph/Private/K2Node_GetInputAxisValue.cpp
--- a/Engine/Source/Editor/BlueprintGraph/Private/K2Node_GetInputAxisValue.cpp
+++ b/Engine/Source/Editor/BlueprintGraph/Private/K2Node_GetInputAxisValue.cpp
@@ -24,9 +24,17 @@ void UK2Node_GetInputAxisValue::Initialize(const FName AxisName)
 	SetFromFunction(AActor::StaticClass()->FindFunctionByName(TEXT("GetInputAxisValue")));
 }
 
-FString UK2Node_GetInputAxisValue::GetNodeTitle(ENodeTitleType::Type TitleType) const
+FText UK2Node_GetInputAxisValue::GetNodeTitle(ENodeTitleType::Type TitleType) const
 {
-	return FString::Printf(*NSLOCTEXT("K2Node", "GetInputAxis_Name", "Get %s").ToString(), *InputAxisName.ToString());
+	FFormatNamedArguments Args;
+	Args.Add(TEXT("InputAxisName"), FText::FromName(InputAxisName));
+	return FText::Format(NSLOCTEXT("K2Node", "GetInputAxis_Name", "Get {InputAxisName}"), Args);
+}
+
+FString UK2Node_GetInputAxisValue::GetNodeNativeTitle(ENodeTitleType::Type TitleType) const
+{
+	// Do not setup this function for localization, intentionally left unlocalized!
+	return FString::Printf(TEXT("Get %s"), *InputAxisName.ToString());
 }
 
 FString UK2Node_GetInputAxisValue::GetTooltip() const
diff --git a/Engine/Source/Editor/BlueprintGraph/Private/K2Node_GetNumEnumEntries.cpp b/Engine/Source/Editor/BlueprintGraph/Private/K2Node_GetNumEnumEntries.cpp
--- a/Engine/Source/Editor/BlueprintGraph/Private/K2Node_GetNumEnumEntries.cpp
+++ b/Engine/Source/Editor/BlueprintGraph/Private/K2Node_GetNumEnumEntries.cpp
@@ -27,11 +27,21 @@ FString UK2Node_GetNumEnumEntries::GetTooltip() const
 	return FString::Printf(*NSLOCTEXT("K2Node", "GetNumEnumEntries_Tooltip", "Returns %s_MAX value").ToString(), *EnumName);
 }
 
-FString UK2Node_GetNumEnumEntries::GetNodeTitle(ENodeTitleType::Type TitleType) const
+FText UK2Node_GetNumEnumEntries::GetNodeTitle(ENodeTitleType::Type TitleType) const
 {
+	const FText EnumName = (Enum != NULL) ? FText::FromString(Enum->GetName()) : NSLOCTEXT("K2Node", "BadEnum", "(bad enum)");
+
+	FFormatNamedArguments Args;
+	Args.Add(TEXT("EnumName"), EnumName);
+	return FText::Format(NSLOCTEXT("K2Node", "GetNumEnumEntries_Title", "Get number of entries in {EnumName}"), Args);
+}
+
+FString UK2Node_GetNumEnumEntries::GetNodeNativeTitle(ENodeTitleType::Type TitleType) const
+{
+	// Do not setup this function for localization, intentionally left unlocalized!
 	const FString EnumName = (Enum != NULL) ? Enum->GetName() : TEXT("(bad enum)");
 
-	return FString::Printf(*NSLOCTEXT("K2Node", "GetNumEnumEntries_Title", "Get number of entries in %s").ToString(), *EnumName);
+	return FString::Printf(TEXT("Get number of entries in %s"), *EnumName);
 }
 
 void UK2Node_GetNumEnumEntries::ExpandNode(class FKismetCompilerContext& CompilerContext, UEdGraph* SourceGraph)
diff --git a/Engine/Source/Editor/BlueprintGraph/Private/K2Node_IfThenElse.cpp b/Engine/Source/Editor/BlueprintGraph/Private/K2Node_IfThenElse.cpp
--- a/Engine/Source/Editor/BlueprintGraph/Private/K2Node_IfThenElse.cpp
+++ b/Engine/Source/Editor/BlueprintGraph/Private/K2Node_IfThenElse.cpp
@@ -83,17 +83,17 @@ void UK2Node_IfThenElse::AllocateDefaultPins()
 	ConditionPin->DefaultValue = ConditionPin->AutogeneratedDefaultValue = TEXT("true");
 
 	UEdGraphPin* TruePin = CreatePin(EGPD_Output, K2Schema->PC_Exec, TEXT(""), NULL, false, false, K2Schema->PN_Then);
-	TruePin->PinFriendlyName = FString(*LOCTEXT("true", "true").ToString());
+	TruePin->PinFriendlyName =LOCTEXT("true", "true");
 
 	UEdGraphPin* FalsePin = CreatePin(EGPD_Output, K2Schema->PC_Exec, TEXT(""), NULL, false, false, K2Schema->PN_Else);
-	FalsePin->PinFriendlyName = FString(*LOCTEXT("false", "false").ToString());
+	FalsePin->PinFriendlyName = LOCTEXT("false", "false");
 
 	Super::AllocateDefaultPins();
 }
 
-FString UK2Node_IfThenElse::GetNodeTitle(ENodeTitleType::Type TitleType) const
+FText UK2Node_IfThenElse::GetNodeTitle(ENodeTitleType::Type TitleType) const
 {
-	return *LOCTEXT("Branch", "Branch").ToString();
+	return LOCTEXT("Branch", "Branch");
 }
 
 FLinearColor UK2Node_IfThenElse::GetNodeTitleColor() const
diff --git a/Engine/Source/Editor/BlueprintGraph/Private/K2Node_InputAction.cpp b/Engine/Source/Editor/BlueprintGraph/Private/K2Node_InputAction.cpp
--- a/Engine/Source/Editor/BlueprintGraph/Private/K2Node_InputAction.cpp
+++ b/Engine/Source/Editor/BlueprintGraph/Private/K2Node_InputAction.cpp
@@ -38,9 +38,17 @@ FLinearColor UK2Node_InputAction::GetNodeTitleColor() const
 	return GEditor->AccessEditorUserSettings().EventNodeTitleColor;
 }
 
-FString UK2Node_InputAction::GetNodeTitle(ENodeTitleType::Type TitleType) const
+FText UK2Node_InputAction::GetNodeTitle(ENodeTitleType::Type TitleType) const
 {
-	return FString::Printf(*NSLOCTEXT("K2Node", "InputAction_Name", "InputAction %s").ToString(), *InputActionName.ToString());
+	FFormatNamedArguments Args;
+	Args.Add(TEXT("InputActionName"), FText::FromName(InputActionName));
+	return FText::Format(NSLOCTEXT("K2Node", "InputAction_Name", "InputAction {InputActionName}"), Args);
+}
+
+FString UK2Node_InputAction::GetNodeNativeTitle(ENodeTitleType::Type TitleType) const
+{
+	// Do not setup this function for localization, intentionally left unlocalized!
+	return FString::Printf(TEXT("InputAction %s"), *InputActionName.ToString());
 }
 
 FString UK2Node_InputAction::GetTooltip() const
diff --git a/Engine/Source/Editor/BlueprintGraph/Private/K2Node_InputAxisEvent.cpp b/Engine/Source/Editor/BlueprintGraph/Private/K2Node_InputAxisEvent.cpp
--- a/Engine/Source/Editor/BlueprintGraph/Private/K2Node_InputAxisEvent.cpp
+++ b/Engine/Source/Editor/BlueprintGraph/Private/K2Node_InputAxisEvent.cpp
@@ -32,9 +32,17 @@ void UK2Node_InputAxisEvent::Initialize(const FName AxisName)
 	CustomFunctionName = FName( *FString::Printf(TEXT("InpAxisEvt_%s_%s"), *InputAxisName.ToString(), *GetName()));
 }
 
-FString UK2Node_InputAxisEvent::GetNodeTitle(ENodeTitleType::Type TitleType) const
+FText UK2Node_InputAxisEvent::GetNodeTitle(ENodeTitleType::Type TitleType) const
 {
-	return FString::Printf(*NSLOCTEXT("K2Node", "InputAxis_Name", "InputAxis %s").ToString(), *InputAxisName.ToString());
+	FFormatNamedArguments Args;
+	Args.Add(TEXT("InputAxisName"), FText::FromName(InputAxisName));
+	return FText::Format(NSLOCTEXT("K2Node", "InputAxis_Name", "InputAxis {InputAxisName}"), Args);
+}
+
+FString UK2Node_InputAxisEvent::GetNodeNativeTitle(ENodeTitleType::Type TitleType) const
+{
+	// Do not setup this function for localization, intentionally left unlocalized!
+	return FString::Printf(TEXT("InputAxis %s"), *InputAxisName.ToString());
 }
 
 FString UK2Node_InputAxisEvent::GetTooltip() const
diff --git a/Engine/Source/Editor/BlueprintGraph/Private/K2Node_InputAxisKeyEvent.cpp b/Engine/Source/Editor/BlueprintGraph/Private/K2Node_InputAxisKeyEvent.cpp
--- a/Engine/Source/Editor/BlueprintGraph/Private/K2Node_InputAxisKeyEvent.cpp
+++ b/Engine/Source/Editor/BlueprintGraph/Private/K2Node_InputAxisKeyEvent.cpp
@@ -21,9 +21,15 @@ void UK2Node_InputAxisKeyEvent::Initialize(const FKey InAxisKey)
 	CustomFunctionName = FName(*FString::Printf(TEXT("InpAxisKeyEvt_%s_%s"), *AxisKey.ToString(), *GetName()));
 }
 
-FString UK2Node_InputAxisKeyEvent::GetNodeTitle(ENodeTitleType::Type TitleType) const
+FText UK2Node_InputAxisKeyEvent::GetNodeTitle(ENodeTitleType::Type TitleType) const
 {
-	return FString::Printf(*NSLOCTEXT("K2Node", "InputAxisKey_Name", "%s").ToString(), *AxisKey.GetDisplayName().ToString());
+	return AxisKey.GetDisplayName();
+}
+
+FString UK2Node_InputAxisKeyEvent::GetNodeNativeTitle(ENodeTitleType::Type TitleType) const
+{
+	// Do not setup this function for localization, intentionally left unlocalized!
+	return AxisKey.GetDisplayName().ToString();
 }
 
 FString UK2Node_InputAxisKeyEvent::GetTooltip() const
