#summary Commiting some missing files from CL# 2044506

diff --git a/Engine/Source/Editor/BlueprintGraph/Classes/K2Node_GetInputAxisKeyValue.h b/Engine/Source/Editor/BlueprintGraph/Classes/K2Node_GetInputAxisKeyValue.h
--- a/Engine/Source/Editor/BlueprintGraph/Classes/K2Node_GetInputAxisKeyValue.h
+++ b/Engine/Source/Editor/BlueprintGraph/Classes/K2Node_GetInputAxisKeyValue.h
@@ -21,7 +21,7 @@ class UK2Node_GetInputAxisKeyValue : public UK2Node_CallFunction
 
 	// Begin EdGraphNode interface
 	virtual void AllocateDefaultPins() OVERRIDE;
-	virtual FString GetNodeTitle(ENodeTitleType::Type TitleType) const OVERRIDE;
+	virtual FText GetNodeTitle(ENodeTitleType::Type TitleType) const OVERRIDE;
 	virtual FString GetTooltip() const OVERRIDE;
 	virtual FName GetPaletteIcon(FLinearColor& OutColor) const OVERRIDE;
 	// End EdGraphNode interface
diff --git a/Engine/Source/Editor/BlueprintGraph/Private/K2Node_GetInputAxisKeyValue.cpp b/Engine/Source/Editor/BlueprintGraph/Private/K2Node_GetInputAxisKeyValue.cpp
--- a/Engine/Source/Editor/BlueprintGraph/Private/K2Node_GetInputAxisKeyValue.cpp
+++ b/Engine/Source/Editor/BlueprintGraph/Private/K2Node_GetInputAxisKeyValue.cpp
@@ -24,14 +24,18 @@ void UK2Node_GetInputAxisKeyValue::Initialize(const FKey AxisKey)
 	SetFromFunction(AActor::StaticClass()->FindFunctionByName(TEXT("GetInputAxisKeyValue")));
 }
 
-FString UK2Node_GetInputAxisKeyValue::GetNodeTitle(ENodeTitleType::Type TitleType) const
+FText UK2Node_GetInputAxisKeyValue::GetNodeTitle(ENodeTitleType::Type TitleType) const
 {
-	return FString::Printf(*NSLOCTEXT("K2Node", "GetInputAxisKey_Name", "Get %s").ToString(), *InputAxisKey.GetDisplayName().ToString());
+	FFormatNamedArguments Args;
+	Args.Add(TEXT("AxisKey"), InputAxisKey.GetDisplayName());
+	return FText::Format(NSLOCTEXT("K2Node", "GetInputAxisKey_Name", "Get {AxisKey}"), Args);
 }
 
 FString UK2Node_GetInputAxisKeyValue::GetTooltip() const
 {
-	return FString::Printf(*NSLOCTEXT("K2Node", "GetInputAxisKey_Tooltip", "Returns the current value of input axis key %s.  If input is disabled for the actor the value will be 0.").ToString(), *InputAxisKey.GetDisplayName().ToString());
+	FFormatNamedArguments Args;
+	Args.Add(TEXT("AxisKey"), InputAxisKey.GetDisplayName());
+	return FText::Format(NSLOCTEXT("K2Node", "GetInputAxisKey_Tooltip", "Returns the current value of input axis key {AxisKey}.  If input is disabled for the actor the value will be 0."), Args).ToString();
 }
 
 void UK2Node_GetInputAxisKeyValue::ValidateNodeDuringCompilation(class FCompilerResultsLog& MessageLog) const
