#ue4 K2Node_BaseAsyncTask works with arrays

https://udn.unrealengine.com/questions/204560/baseasynctask-with-array-parameter.html

diff --git a/Engine/Source/Editor/BlueprintGraph/Private/K2Node_BaseAsyncTask.cpp b/Engine/Source/Editor/BlueprintGraph/Private/K2Node_BaseAsyncTask.cpp
--- a/Engine/Source/Editor/BlueprintGraph/Private/K2Node_BaseAsyncTask.cpp
+++ b/Engine/Source/Editor/BlueprintGraph/Private/K2Node_BaseAsyncTask.cpp
@@ -283,7 +283,7 @@ void UK2Node_BaseAsyncTask::ExpandNode(class FKismetCompilerContext& CompilerCon
 					OutPin->PinType.PinCategory,
 					OutPin->PinType.PinSubCategory,
 					OutPin->PinType.PinSubCategoryObject.Get(),
-					false);
+					OutPin->PinType.bIsArray);
 			}
 			UEdGraphPin* PinWithData = CurrentCENode->FindPin(OutPin->PinName);
 
