#summary Fixed an issue with some pins being disconnected due to an incorrect and invalid name change from CL# 2044506, this was causing, as an example, BlendListByBools's inputs to become disconnected.

diff --git a/Engine/Source/Editor/BlueprintGraph/Private/K2Node.cpp b/Engine/Source/Editor/BlueprintGraph/Private/K2Node.cpp
--- a/Engine/Source/Editor/BlueprintGraph/Private/K2Node.cpp
+++ b/Engine/Source/Editor/BlueprintGraph/Private/K2Node.cpp
@@ -898,7 +898,7 @@ void FOptionalPinManager::CreateVisiblePins(TArray<FOptionalPinFromProperty>& Pr
 							FFormatNamedArguments Args;
 							Args.Add(TEXT("PinName"), FText::FromName(PropertyEntry.PropertyName));
 							Args.Add(TEXT("Index"), Index);
-							const FText PinFriendlyName = FText::Format(LOCTEXT("PinFriendlyNameWithIndex", "{PinFriendlyName}_{Index}"), Args);
+							const FText PinFriendlyName = FText::Format(LOCTEXT("PinFriendlyNameWithIndex", "{PinName}_{Index}"), Args);
 							const FString PinName = PinFriendlyName.ToString();
 							NewPin = TargetNode->CreatePin(Direction, PinType.PinCategory, PinType.PinSubCategory, PinType.PinSubCategoryObject.Get(), PinType.bIsArray, PinType.bIsReference, PinName);
 							NewPin->PinFriendlyName = PinFriendlyName;
