#ue4 K2Node_BaseAsyncTask - Expansion simplified. (No Select node.)

diff --git a/Engine/Source/Editor/BlueprintGraph/Private/K2Node_BaseAsyncTask.cpp b/Engine/Source/Editor/BlueprintGraph/Private/K2Node_BaseAsyncTask.cpp
--- a/Engine/Source/Editor/BlueprintGraph/Private/K2Node_BaseAsyncTask.cpp
+++ b/Engine/Source/Editor/BlueprintGraph/Private/K2Node_BaseAsyncTask.cpp
@@ -139,210 +139,183 @@ void UK2Node_BaseAsyncTask::AllocateDefaultPins()
 	Super::AllocateDefaultPins();
 }
 
-void UK2Node_BaseAsyncTask::ExpandNode(class FKismetCompilerContext& CompilerContext, UEdGraph* SourceGraph)
+struct FBaseAsyncTaskHelper
 {
-	if (!CompilerContext.bIsFullCompile)
+	static bool ValidDataPin(const UEdGraphPin* Pin, EEdGraphPinDirection Direction, const UEdGraphSchema_K2* Schema)
 	{
-		return;
-	}
-
-	check(SourceGraph);
-	const UEdGraphSchema_K2* Schema = CompilerContext.GetSchema();
+		check(Schema);
+		const bool bValidDataPin = Pin
+			&& (Pin->PinName != Schema->PN_Execute)
+			&& (Pin->PinName != Schema->PN_Then)
+			&& (Pin->PinType.PinCategory != Schema->PC_Exec);
 
-	UK2Node_BaseAsyncTask* CurrentNode = this;
+		const bool bProperDirection = Pin && (Pin->Direction == Direction);
 
-	TArray<UMulticastDelegateProperty*> DelegateProperties;
-	for (TFieldIterator<UProperty> PropertyIt(CurrentNode->GetProxyClass(), EFieldIteratorFlags::ExcludeSuper); PropertyIt; ++PropertyIt)
-	{
-		if (UMulticastDelegateProperty* Property = Cast<UMulticastDelegateProperty>(*PropertyIt))
-		{
-			DelegateProperties.AddUnique(Property);
-		}
+		return bValidDataPin && bProperDirection;
 	}
 
-	if (DelegateProperties.Num() == 0)
+	static bool CreateDelegateForNewFunction(UEdGraphPin* DelegateInputPin, FName FunctionName, UK2Node* CurrentNode, UEdGraph* SourceGraph, FKismetCompilerContext& CompilerContext)
 	{
-		// we don't have any delegates to handle, skip this node
-		return;
+		const UEdGraphSchema_K2* Schema = CompilerContext.GetSchema();
+		check(DelegateInputPin && Schema && CurrentNode && SourceGraph && (FunctionName != NAME_None));
+		bool bResult = true;
+
+		// WORKAROUND, so we can create delegate from nonexistent function by avoiding check at expanding step
+		// instead simply: Schema->TryCreateConnection(AddDelegateNode->GetDelegatePin(), CurrentCENode->FindPinChecked(UK2Node_CustomEvent::DelegateOutputName));
+		UK2Node_Self* SelfNode = CompilerContext.SpawnIntermediateNode<UK2Node_Self>(CurrentNode, SourceGraph);
+		SelfNode->AllocateDefaultPins();
+
+		UK2Node_CreateDelegate* CreateDelegateNode = CompilerContext.SpawnIntermediateNode<UK2Node_CreateDelegate>(CurrentNode, SourceGraph);
+		CreateDelegateNode->AllocateDefaultPins();
+		bResult &= Schema->TryCreateConnection(DelegateInputPin, CreateDelegateNode->GetDelegateOutPin());
+		bResult &= Schema->TryCreateConnection(SelfNode->FindPinChecked(Schema->PN_Self), CreateDelegateNode->GetObjectInPin());
+		CreateDelegateNode->SelectedFunctionName = FunctionName;
+
+		return bResult;
 	}
 
-	const FName CreateMoveToProxyObjectBlueprintFuncName = CurrentNode->GetProxyFactoryFunctionName();
-	UK2Node_CallFunction* CallCreateMoveToProxyObjectNode = CompilerContext.SpawnIntermediateNode<UK2Node_CallFunction>(CurrentNode, SourceGraph);
-	CallCreateMoveToProxyObjectNode->FunctionReference.SetExternalMember(CreateMoveToProxyObjectBlueprintFuncName, CurrentNode->GetProxyFactoryClass());
-	CallCreateMoveToProxyObjectNode->AllocateDefaultPins();
-	CompilerContext.MovePinLinksToIntermediate(*CurrentNode->FindPinChecked(Schema->PN_Execute), *CallCreateMoveToProxyObjectNode->FindPinChecked(Schema->PN_Execute));
-
-	// match function inputs, to pass data to function from CallFunction node
-	for (int32 Index=0; Index < CurrentNode->Pins.Num(); ++Index)
+	struct FOutputPinAndLocalVariable
 	{
-		UEdGraphPin* CurrentPin = CurrentNode->Pins[Index];
-		if (!CurrentPin || CurrentPin->Direction != EGPD_Input || CurrentPin->PinName == Schema->PN_Execute || CurrentPin->PinName == Schema->PN_Then)
-		{
-			continue;
-		}
+		UEdGraphPin* OutputPin;
+		UK2Node_TemporaryVariable* TempVar;
 
-		if (UEdGraphPin* DestPin = CallCreateMoveToProxyObjectNode->FindPin(CurrentPin->PinName))
-		{
-			CompilerContext.MovePinLinksToIntermediate(*CurrentPin, *DestPin);
-		}
-	}
+		FOutputPinAndLocalVariable(UEdGraphPin* Pin, UK2Node_TemporaryVariable* Var) : OutputPin(Pin), TempVar(Var) {}
+	};
 
-	TArray<UEdGraphPin*> VariableOutputPins;
-	for (int32 Index=0; Index < CurrentNode->Pins.Num(); ++Index)
+	static bool CopyEventSignature(UK2Node_CustomEvent* CENode, UFunction* Function, const UEdGraphSchema_K2* Schema)
 	{
-		UEdGraphPin* CurrentPin = CurrentNode->Pins[Index];
-		if (!CurrentPin || CurrentPin->Direction != EGPD_Output || CurrentPin->PinType.PinCategory == Schema->PC_Exec || CurrentPin->PinName == Schema->PN_Execute || CurrentPin->PinName == Schema->PN_Then)
+		check(CENode && Function && Schema);
+
+		bool bResult = true;
+		for (TFieldIterator<UProperty> PropIt(Function); PropIt && (PropIt->PropertyFlags & CPF_Parm); ++PropIt)
 		{
-			continue;
+			const UProperty* Param = *PropIt;
+			if (!Param->HasAnyPropertyFlags(CPF_OutParm) || Param->HasAnyPropertyFlags(CPF_ReferenceParm))
+			{
+				FEdGraphPinType PinType;
+				bResult &= Schema->ConvertPropertyToPinType(Param, /*out*/ PinType);
+				bResult &= (NULL != CENode->CreateUserDefinedPin(Param->GetName(), PinType));
+			}
 		}
 
-		VariableOutputPins.AddUnique(CurrentPin);
+		return bResult;
 	}
 
-	// Create int selection index
-	UK2Node_TemporaryVariable* TempIndexVar = CompilerContext.SpawnIntermediateNode<UK2Node_TemporaryVariable>(CurrentNode, SourceGraph);
-	TempIndexVar->VariableType.PinCategory = Schema->PC_Int;
-	TempIndexVar->AllocateDefaultPins();
-
-	TArray<UK2Node_Select*> OutputSelectNodes;
-	for (int32 Index=0; Index < VariableOutputPins.Num(); ++Index)
+	static bool HandleDelegateImplementation(
+		UMulticastDelegateProperty* CurrentProperty, const TArray<FBaseAsyncTaskHelper::FOutputPinAndLocalVariable>& VariableOutputs,
+		UEdGraphPin* ProxyObjectPin, UEdGraphPin*& InOutLastThenPin,
+		UK2Node* CurrentNode, UEdGraph* SourceGraph, FKismetCompilerContext& CompilerContext)
 	{
-		UK2Node_Select* SelectNode = CompilerContext.SpawnIntermediateNode<UK2Node_Select>(CurrentNode, SourceGraph);
-		SelectNode->NumOptionPins = DelegateProperties.Num();
-		SelectNode->IndexPinType = TempIndexVar->VariableType;
-		SelectNode->AllocateDefaultPins();
-		OutputSelectNodes.AddUnique(SelectNode);
-	}
+		bool bIsErrorFree = true;
+		const UEdGraphSchema_K2* Schema = CompilerContext.GetSchema();
+		check(CurrentProperty && ProxyObjectPin && InOutLastThenPin && CurrentNode && SourceGraph && Schema);
 
-	bool bMoveDataPins = true;
-	UEdGraphPin* LastThenPin = NULL;
-	for (int32 Index=0; Index < DelegateProperties.Num(); ++Index)
-	{
-		UProperty* CurrentProperty = DelegateProperties[Index];
 		UEdGraphPin* PinForCurrentDelegateProperty = CurrentNode->FindPin(CurrentProperty->GetName());
-		if (CurrentProperty == NULL || PinForCurrentDelegateProperty == NULL)
+		if (!PinForCurrentDelegateProperty || (Schema->PC_Exec != PinForCurrentDelegateProperty->PinType.PinCategory))
 		{
-			// we can't match current delegate property with output pin, skip it
-			continue;
+			FText ErrorMessage = FText::Format(LOCTEXT("WrongDelegateProperty", "BaseAsyncTask: Cannot find execution pin for delegate "), FText::FromString(CurrentProperty->GetName()));
+			CompilerContext.MessageLog.Error(*ErrorMessage.ToString(), CurrentNode);
+			return false;
 		}
 
-		UK2Node_AddDelegate* AddDelegateNode = CompilerContext.SpawnIntermediateNode<UK2Node_AddDelegate>(CurrentNode, SourceGraph);
-		check(AddDelegateNode);
-		AddDelegateNode->SetFromProperty(CurrentProperty, false);
-		AddDelegateNode->AllocateDefaultPins();
-		Schema->TryCreateConnection(AddDelegateNode->FindPinChecked(Schema->PN_Self), CallCreateMoveToProxyObjectNode->GetReturnValuePin());
-		if (LastThenPin == NULL)
-		{
-			Schema->TryCreateConnection(CallCreateMoveToProxyObjectNode->FindPinChecked(Schema->PN_Then), AddDelegateNode->FindPinChecked(Schema->PN_Execute));
-		}
-		else
-		{
-			Schema->TryCreateConnection(LastThenPin, AddDelegateNode->FindPinChecked(Schema->PN_Execute));
-		}
-		LastThenPin = AddDelegateNode->FindPinChecked(Schema->PN_Then);
-
 		UK2Node_CustomEvent* CurrentCENode = CompilerContext.SpawnIntermediateNode<UK2Node_CustomEvent>(CurrentNode, SourceGraph);
-		check(CurrentCENode);
-		CurrentCENode->CustomFunctionName = *FString::Printf(TEXT("%s_%u"), *CurrentProperty->GetName(), FPlatformTime::Cycles());
-		CurrentCENode->AllocateDefaultPins();
-
 		{
-			// WORKAROUND, so we can create delegate from nonexistent function by avoiding check at expanding step
-			// instead simply: Schema->TryCreateConnection(AddDelegateNode->GetDelegatePin(), CurrentCENode->FindPinChecked(UK2Node_CustomEvent::DelegateOutputName));
-
-			UK2Node_Self* SelfNode = CompilerContext.SpawnIntermediateNode<UK2Node_Self>(this, SourceGraph);
-			SelfNode->AllocateDefaultPins();
-
-			UK2Node_CreateDelegate* CreateDelegateNode = CompilerContext.SpawnIntermediateNode<UK2Node_CreateDelegate>(this, SourceGraph);
-			CreateDelegateNode->AllocateDefaultPins();
-			Schema->TryCreateConnection(AddDelegateNode->GetDelegatePin(), CreateDelegateNode->GetDelegateOutPin());
-			Schema->TryCreateConnection(SelfNode->FindPinChecked(Schema->PN_Self), CreateDelegateNode->GetObjectInPin());
-			CreateDelegateNode->SelectedFunctionName = CurrentCENode->GetFunctionName();
-		}
-
-		for (TFieldIterator<UProperty> PropIt(AddDelegateNode->GetDelegateSignature()); PropIt && (PropIt->PropertyFlags & CPF_Parm); ++PropIt)
-		{
-			const UProperty* Param = *PropIt;
-			if (!Param->HasAnyPropertyFlags(CPF_OutParm) || Param->HasAnyPropertyFlags(CPF_ReferenceParm))
-			{
-				FEdGraphPinType PinType;
-				Schema->ConvertPropertyToPinType(Param, /*out*/ PinType);
-				CurrentCENode->CreateUserDefinedPin(Param->GetName(), PinType);
-			}
+			UK2Node_AddDelegate* AddDelegateNode = CompilerContext.SpawnIntermediateNode<UK2Node_AddDelegate>(CurrentNode, SourceGraph);
+			AddDelegateNode->SetFromProperty(CurrentProperty, false);
+			AddDelegateNode->AllocateDefaultPins();
+			bIsErrorFree &= Schema->TryCreateConnection(AddDelegateNode->FindPinChecked(Schema->PN_Self), ProxyObjectPin);
+			bIsErrorFree &= Schema->TryCreateConnection(InOutLastThenPin, AddDelegateNode->FindPinChecked(Schema->PN_Execute));
+			InOutLastThenPin = AddDelegateNode->FindPinChecked(Schema->PN_Then);
+			CurrentCENode->CustomFunctionName = *FString::Printf(TEXT("%s_%s"), *CurrentProperty->GetName(), *CurrentCENode->NodeGuid.ToString());
+			CurrentCENode->AllocateDefaultPins();
+
+			bIsErrorFree &= FBaseAsyncTaskHelper::CreateDelegateForNewFunction(AddDelegateNode->GetDelegatePin(), CurrentCENode->GetFunctionName(), CurrentNode, SourceGraph, CompilerContext);
+			bIsErrorFree &= FBaseAsyncTaskHelper::CopyEventSignature(CurrentCENode, AddDelegateNode->GetDelegateSignature(), Schema);
 		}
 
-
 		UEdGraphPin* LastActivatedNodeThen = CurrentCENode->FindPinChecked(Schema->PN_Then);
-		for (int32 OutIndex = 0; OutIndex < VariableOutputPins.Num(); ++OutIndex)
+		for (auto OutputPair : VariableOutputs) // CREATE CHAIN OF ASSIGMENTS
 		{
-			UEdGraphPin* OutPin = VariableOutputPins[OutIndex];
-			UK2Node_TemporaryVariable* TempVarOutput;
-			{
-				TempVarOutput = CompilerContext.SpawnInternalVariable(
-					CurrentNode,
-					OutPin->PinType.PinCategory,
-					OutPin->PinType.PinSubCategory,
-					OutPin->PinType.PinSubCategoryObject.Get(),
-					OutPin->PinType.bIsArray);
-			}
-			UEdGraphPin* PinWithData = CurrentCENode->FindPin(OutPin->PinName);
-
+			UEdGraphPin* PinWithData = CurrentCENode->FindPin(OutputPair.OutputPin->PinName);
 			if (PinWithData == NULL)
 			{
-				FText ErrorMessage = FText::Format(LOCTEXT("MissingDataPin", "ICE: Pin @@ was expecting a data output pin named {0} on @@ (each delegate must have the same signature)"), FText::FromString(OutPin->PinName));
-				CompilerContext.MessageLog.Error(*ErrorMessage.ToString(), OutPin, CurrentCENode);
-				continue;
+				FText ErrorMessage = FText::Format(LOCTEXT("MissingDataPin", "ICE: Pin @@ was expecting a data output pin named {0} on @@ (each delegate must have the same signature)"), FText::FromString(OutputPair.OutputPin->PinName));
+				CompilerContext.MessageLog.Error(*ErrorMessage.ToString(), OutputPair.OutputPin, CurrentCENode);
+				return false;
 			}
 
 			UK2Node_AssignmentStatement* AssignNode = CompilerContext.SpawnIntermediateNode<UK2Node_AssignmentStatement>(CurrentNode, SourceGraph);
 			AssignNode->AllocateDefaultPins();
-			AssignNode->GetVariablePin()->DefaultValue = OutPin->DefaultValue;
-			Schema->TryCreateConnection(LastActivatedNodeThen, AssignNode->GetExecPin());
-
-			// Move connections from fake UK2Node_AddDelegate pint to this assignment node
-			Schema->TryCreateConnection(TempVarOutput->GetVariablePin(), AssignNode->GetVariablePin());
+			bIsErrorFree &= Schema->TryCreateConnection(LastActivatedNodeThen, AssignNode->GetExecPin());
+			bIsErrorFree &= Schema->TryCreateConnection(OutputPair.TempVar->GetVariablePin(), AssignNode->GetVariablePin());
 			AssignNode->NotifyPinConnectionListChanged(AssignNode->GetVariablePin());
-
-			Schema->TryCreateConnection(AssignNode->GetValuePin(), PinWithData);
+			bIsErrorFree &= Schema->TryCreateConnection(AssignNode->GetValuePin(), PinWithData);
 			AssignNode->NotifyPinConnectionListChanged(AssignNode->GetValuePin());
 
-			// Set correct index for K2Node_Select node to return correct value from delegate
-			UK2Node_AssignmentStatement* TempIndexVarInitialize = CompilerContext.SpawnIntermediateNode<UK2Node_AssignmentStatement>(CurrentNode, SourceGraph);
-			TempIndexVarInitialize->AllocateDefaultPins();
-			TempIndexVarInitialize->GetVariablePin()->PinType = TempIndexVar->GetVariablePin()->PinType;
-			TempIndexVarInitialize->GetVariablePin()->MakeLinkTo(TempIndexVar->GetVariablePin());
-			TempIndexVarInitialize->PinConnectionListChanged(TempIndexVarInitialize->GetVariablePin());
-			TempIndexVarInitialize->GetValuePin()->PinType = TempIndexVar->GetVariablePin()->PinType;
-			TempIndexVarInitialize->GetValuePin()->DefaultValue = FString::Printf(TEXT("%d"), Index);
-			Schema->TryCreateConnection(AssignNode->GetThenPin(), TempIndexVarInitialize->GetExecPin());
-			LastActivatedNodeThen = TempIndexVarInitialize->GetThenPin();
-
-			TempIndexVar->GetVariablePin()->MakeLinkTo(OutputSelectNodes[OutIndex]->GetIndexPin());
-
-			TArray<UEdGraphPin*> OptionPins;
-			OutputSelectNodes[OutIndex]->GetOptionPins(OptionPins);
-			Schema->TryCreateConnection(TempVarOutput->GetVariablePin(), OptionPins[Index]);
+			LastActivatedNodeThen = AssignNode->GetThenPin();
 		}
 
-		// connect CE output with OnSuccess output of AIMoveTo node
-		CompilerContext.MovePinLinksToIntermediate(*PinForCurrentDelegateProperty, *LastActivatedNodeThen);
+		bIsErrorFree &= CompilerContext.MovePinLinksToIntermediate(*PinForCurrentDelegateProperty, *LastActivatedNodeThen).CanSafeConnect();
+		return bIsErrorFree;
+	}
+};
+
+void UK2Node_BaseAsyncTask::ExpandNode(class FKismetCompilerContext& CompilerContext, UEdGraph* SourceGraph)
+{
+	if (!CompilerContext.bIsFullCompile)
+	{
+		return;
+	}
+	const UEdGraphSchema_K2* Schema = CompilerContext.GetSchema();
+	UK2Node_BaseAsyncTask* const CurrentNode = this;
+	check(SourceGraph && Schema && CurrentNode);
+	bool bIsErrorFree = true;
+
+	UK2Node_CallFunction* const CallCreateMoveToProxyObjectNode = CompilerContext.SpawnIntermediateNode<UK2Node_CallFunction>(CurrentNode, SourceGraph);
+	CallCreateMoveToProxyObjectNode->FunctionReference.SetExternalMember(CurrentNode->GetProxyFactoryFunctionName(), CurrentNode->GetProxyFactoryClass());
+	CallCreateMoveToProxyObjectNode->AllocateDefaultPins();
+	bIsErrorFree &= CompilerContext.MovePinLinksToIntermediate(*CurrentNode->FindPinChecked(Schema->PN_Execute), *CallCreateMoveToProxyObjectNode->FindPinChecked(Schema->PN_Execute)).CanSafeConnect();
+	for (auto CurrentPin : CurrentNode->Pins)
+	{
+		if (FBaseAsyncTaskHelper::ValidDataPin(CurrentPin, EGPD_Input, Schema))
+		{
+			UEdGraphPin* DestPin = CallCreateMoveToProxyObjectNode->FindPin(CurrentPin->PinName); // match function inputs, to pass data to function from CallFunction node
+			bIsErrorFree &= DestPin && CompilerContext.MovePinLinksToIntermediate(*CurrentPin, *DestPin).CanSafeConnect();
+		}
+	}
+	
+	// GATHER OUTPUT PARAMETERS AND PAIR THEM WITH LOCAL VARIABLES
+	TArray<FBaseAsyncTaskHelper::FOutputPinAndLocalVariable> VariableOutputs;
+	for (auto CurrentPin : CurrentNode->Pins)
+	{
+		if (FBaseAsyncTaskHelper::ValidDataPin(CurrentPin, EGPD_Output, Schema))
+		{
+			const FEdGraphPinType& PinType = CurrentPin->PinType;
+			UK2Node_TemporaryVariable* TempVarOutput = CompilerContext.SpawnInternalVariable(
+				CurrentNode, PinType.PinCategory, PinType.PinSubCategory, PinType.PinSubCategoryObject.Get(), PinType.bIsArray);
+			bIsErrorFree &= TempVarOutput->GetVariablePin() && CompilerContext.MovePinLinksToIntermediate(*CurrentPin, *TempVarOutput->GetVariablePin()).CanSafeConnect();
+			VariableOutputs.Add(FBaseAsyncTaskHelper::FOutputPinAndLocalVariable(CurrentPin, TempVarOutput));
+		}
 	}
 
-	for (int32 OutIndex = 0; OutIndex < VariableOutputPins.Num(); ++OutIndex)
+	// FOR EACH DELEGATE DEFINE EVENT, CONNECT IT TO DELEGATE AND IMPLEMENT A CHAIN OF ASSIGMENTS
+	UEdGraphPin* LastThenPin = CallCreateMoveToProxyObjectNode->FindPinChecked(Schema->PN_Then);
+	UEdGraphPin* const ProxyObjectPin = CallCreateMoveToProxyObjectNode->GetReturnValuePin();
+	for (TFieldIterator<UMulticastDelegateProperty> PropertyIt(CurrentNode->GetProxyClass(), EFieldIteratorFlags::ExcludeSuper); PropertyIt && bIsErrorFree; ++PropertyIt)
 	{
-		UEdGraphPin* OutPin = VariableOutputPins[OutIndex];
-		TArray<UEdGraphPin*> OptionPins;
-		OutputSelectNodes[OutIndex]->GetOptionPins(OptionPins);
-		CompilerContext.MovePinLinksToIntermediate(*OutPin, *OutputSelectNodes[OutIndex]->GetReturnValuePin());
+		bIsErrorFree &= FBaseAsyncTaskHelper::HandleDelegateImplementation(*PropertyIt, VariableOutputs, ProxyObjectPin, LastThenPin, CurrentNode, SourceGraph, CompilerContext);
 	}
 
-	if (LastThenPin != NULL)
+	if (CallCreateMoveToProxyObjectNode->FindPinChecked(Schema->PN_Then) == LastThenPin)
 	{
-		CompilerContext.MovePinLinksToIntermediate(*CurrentNode->FindPinChecked(Schema->PN_Then), *LastThenPin);
+		CompilerContext.MessageLog.Error(*LOCTEXT("MissingDelegateProperties", "BaseAsyncTask: Proxy has no delegates defined. @@").ToString(), CurrentNode);
+		return;
 	}
-	else
+
+	bIsErrorFree &= CompilerContext.MovePinLinksToIntermediate(*CurrentNode->FindPinChecked(Schema->PN_Then), *LastThenPin).CanSafeConnect();
+	if (!bIsErrorFree)
 	{
-		CompilerContext.MovePinLinksToIntermediate(*CurrentNode->FindPinChecked(Schema->PN_Then), *CallCreateMoveToProxyObjectNode->FindPinChecked(Schema->PN_Then));
+		CompilerContext.MessageLog.Error(*LOCTEXT("InternalConnectionError", "BaseAsyncTask: Internal connection error. @@").ToString(), CurrentNode);
 	}
 
 	CurrentNode->BreakAllNodeLinks();
