TTP #332341: Making it so drag-n-dropped assets get a temporary graph outer (so we don't hit any checks to ensure a valid outer graph).

diff --git a/Engine/Source/Editor/BlueprintGraph/Private/EdGraphSchema_K2.cpp b/Engine/Source/Editor/BlueprintGraph/Private/EdGraphSchema_K2.cpp
--- a/Engine/Source/Editor/BlueprintGraph/Private/EdGraphSchema_K2.cpp
+++ b/Engine/Source/Editor/BlueprintGraph/Private/EdGraphSchema_K2.cpp
@@ -3170,7 +3170,10 @@ void UEdGraphSchema_K2::DroppedAssetsOnGraph(const TArray<FAssetData>& Assets, c
 			// Make sure we have an asset type that's registered with the component list
 			if (DestinationComponentType != NULL)
 			{
-				TSharedPtr<FEdGraphSchemaAction_K2AddComponent> Action = FK2ActionMenuBuilder::CreateAddComponentAction(GetTransientPackage(), Blueprint, DestinationComponentType, Asset);
+				UEdGraph* TempOuter = NewObject<UEdGraph>((UObject*)Blueprint);
+				TempOuter->SetFlags(RF_Transient);
+
+				TSharedPtr<FEdGraphSchemaAction_K2AddComponent> Action = FK2ActionMenuBuilder::CreateAddComponentAction(TempOuter, Blueprint, DestinationComponentType, Asset);
 				Action->PerformAction(Graph, NULL, GraphPosition);
 			}
 		}
