#ue4 BP Editoronly. Blueprint is not longes marked as dirty after the conversion.

diff --git a/Engine/Source/Editor/BlueprintGraph/Private/EdGraphSchema_K2.cpp b/Engine/Source/Editor/BlueprintGraph/Private/EdGraphSchema_K2.cpp
--- a/Engine/Source/Editor/BlueprintGraph/Private/EdGraphSchema_K2.cpp
+++ b/Engine/Source/Editor/BlueprintGraph/Private/EdGraphSchema_K2.cpp
@@ -3363,6 +3363,7 @@ struct FBackwardCompatibilityConversionHelper
 			{
 				// LINK
 				auto CastNode = NewObject<UK2Node_ClassDynamicCast>(Graph);
+				CastNode->SetFlags(RF_Transactional);
 				CastNode->TargetType = TargetClass;
 				Graph->AddNode(CastNode, false, false);
 				CastNode->CreateNewGuid();
