TTP #332676: When calling GetFunctionContextString() making sure that the function's class has a valid ClassGeneratedBy before using it.

diff --git a/Engine/Source/Editor/BlueprintGraph/Private/K2Node_CallFunction.cpp b/Engine/Source/Editor/BlueprintGraph/Private/K2Node_CallFunction.cpp
--- a/Engine/Source/Editor/BlueprintGraph/Private/K2Node_CallFunction.cpp
+++ b/Engine/Source/Editor/BlueprintGraph/Private/K2Node_CallFunction.cpp
@@ -50,7 +50,11 @@ FString UK2Node_CallFunction::GetFunctionContextString() const
 	{
 		const UFunction* Function = GetTargetFunction();
 		UClass* CurrentSelfClass = (Function != NULL) ? Function->GetOwnerClass() : NULL;
-		const UClass* TrueSelfClass = CurrentSelfClass ? CurrentSelfClass->GetAuthoritativeClass() : NULL;
+		UClass const* TrueSelfClass = CurrentSelfClass;
+		if (CurrentSelfClass && CurrentSelfClass->ClassGeneratedBy)
+		{
+			TrueSelfClass = CurrentSelfClass->GetAuthoritativeClass();
+		}
 		FString TargetString = (TrueSelfClass != NULL) ? TrueSelfClass->GetName() : TEXT("None");
 
 		// This action won't be necessary once the new name convention is used.
