#ue4
Various improvements to UK2Node_BaseAsyncTask:
- Moved node enumeration from the menu builder into the node
- Added the latent icon to the corner of the node
- Added optional support for a separate Activate call on the proxy object, allowing immediate success/failure instead of having to delay it a frame
- Added default behavior for category and node title, pulling it from the factory function

diff --git a/Engine/Source/Editor/BlueprintGraph/Classes/K2Node_BaseAsyncTask.h b/Engine/Source/Editor/BlueprintGraph/Classes/K2Node_BaseAsyncTask.h
--- a/Engine/Source/Editor/BlueprintGraph/Classes/K2Node_BaseAsyncTask.h
+++ b/Engine/Source/Editor/BlueprintGraph/Classes/K2Node_BaseAsyncTask.h
@@ -2,6 +2,8 @@
 
 
 #pragma once
+
+#include "EdGraphSchema_K2_Actions.h"
 #include "K2Node_BaseAsyncTask.generated.h"
 
 UCLASS(Abstract)
@@ -9,22 +11,42 @@ class BLUEPRINTGRAPH_API UK2Node_BaseAsyncTask : public UK2Node
 {
 	GENERATED_UCLASS_BODY()
 
-	// Begin UEdGraphNode interface
+	// UEdGraphNode interface
 	virtual void AllocateDefaultPins() OVERRIDE;
-	// End UEdGraphNode interface
+	virtual void GetMenuEntries(FGraphContextMenuBuilder& ContextMenuBuilder) const OVERRIDE;
+	virtual FString GetTooltip() const OVERRIDE;
+	virtual FText GetNodeTitle(ENodeTitleType::Type TitleType) const OVERRIDE;
+	// End of UEdGraphNode interface
 
-	// Begin UK2Node interface
+	// UK2Node interface
 	virtual void ExpandNode(class FKismetCompilerContext& CompilerContext, UEdGraph* SourceGraph) OVERRIDE;
 	virtual bool HasExternalBlueprintDependencies(TArray<class UStruct*>* OptionalOutput) const OVERRIDE;
-	// End UK2Node interface
+	virtual FName GetCornerIcon() const OVERRIDE;
+	// End of UK2Node interface
 
-	FName	GetProxyFactoryFunctionName();
-	UClass*	GetProxyFactoryClass();
-	UClass* GetProxyClass();
+protected:
 	virtual FString GetCategoryName();
 
+	// Creates a default menu entry
+	TSharedPtr<FEdGraphSchemaAction_K2NewNode> CreateDefaultMenuEntry(UK2Node_BaseAsyncTask* NodeTemplate, FGraphContextMenuBuilder& ContextMenuBuilder) const;
+
+	// Returns the factory function (checked)
+	UFunction* GetFactoryFunction() const;
+
 protected:
+	// The name of the function to call to create a proxy object
+	UPROPERTY()
 	FName ProxyFactoryFunctionName;
+
+	// The class containing the proxy object functions
+	UPROPERTY()
 	UClass* ProxyFactoryClass;
+
+	// The type of proxy object that will be created
+	UPROPERTY()
 	UClass* ProxyClass;
+
+	// The name of the 'go' function on the proxy object that will be called after delegates are in place, can be NAME_None
+	UPROPERTY()
+	FName ProxyActivateFunctionName;
 };
diff --git a/Engine/Source/Editor/BlueprintGraph/Private/K2ActionMenuBuilder.cpp b/Engine/Source/Editor/BlueprintGraph/Private/K2ActionMenuBuilder.cpp
--- a/Engine/Source/Editor/BlueprintGraph/Private/K2ActionMenuBuilder.cpp
+++ b/Engine/Source/Editor/BlueprintGraph/Private/K2ActionMenuBuilder.cpp
@@ -1006,27 +1006,6 @@ void FK2ActionMenuBuilder::GetContextAllowedNodeTypes(FBlueprintGraphActionListB
 			SelectionAction->SearchTitle = NodeTemplate->GetNodeSearchTitle();
 		}
 
-		if (GraphType == GT_Ubergraph && bAllowEvents && !bIsConstructionScript && bAllowImpureFuncs)
-		{
-			for ( TObjectIterator<UClass> it; it; ++it )
-			{
-				UClass* CurrentInterface = *it;
-				if (CurrentInterface->HasAnyClassFlags(CLASS_Abstract))
-				{
-					continue;
-				}
-
-				if (CurrentInterface->IsChildOf( UK2Node_BaseAsyncTask::StaticClass() ))
-				{
-					UK2Node* NodeTemplate = NewObject<UK2Node_BaseAsyncTask>(ContextMenuBuilder.OwnerOfTemporaries, CurrentInterface);
-					FString CategoryName = FString::Printf( TEXT("%s|%s"), *K2ActionCategories::GenericFunctionCategory, *CurrentInterface->GetDefaultObject<UK2Node_BaseAsyncTask>()->GetCategoryName());
-					TSharedPtr<FEdGraphSchemaAction_K2NewNode> SelectionAction = FK2ActionMenuBuilder::AddNewNodeAction(ContextMenuBuilder, CategoryName, NodeTemplate->GetNodeTitle(ENodeTitleType::ListView), NodeTemplate->GetTooltip(), 0, NodeTemplate->GetKeywords());
-					SelectionAction->NodeTemplate = NodeTemplate;
-					SelectionAction->SearchTitle = NodeTemplate->GetNodeSearchTitle();
-				}
-			}
-		}
-
 		FEditorDelegates::OnBlueprintContextMenuCreated.Broadcast(ContextMenuBuilder);
 
 		GetEnumUtilitiesNodes(ContextMenuBuilder, true, true, true, true);
diff --git a/Engine/Source/Editor/BlueprintGraph/Private/K2Node_BaseAsyncTask.cpp b/Engine/Source/Editor/BlueprintGraph/Private/K2Node_BaseAsyncTask.cpp
--- a/Engine/Source/Editor/BlueprintGraph/Private/K2Node_BaseAsyncTask.cpp
+++ b/Engine/Source/Editor/BlueprintGraph/Private/K2Node_BaseAsyncTask.cpp
@@ -5,6 +5,7 @@
 
 #include "CompilerResultsLog.h"
 #include "KismetCompiler.h"
+#include "K2ActionMenuBuilder.h"
 
 #define LOCTEXT_NAMESPACE "UK2Node_BaseAsyncTask"
 
@@ -13,27 +14,48 @@ UK2Node_BaseAsyncTask::UK2Node_BaseAsyncTask(const class FPostConstructInitializ
 	, ProxyFactoryFunctionName(NAME_None)
 	, ProxyFactoryClass(NULL)
 	, ProxyClass(NULL)
+	, ProxyActivateFunctionName(NAME_None)
 {
 }
 
-FName UK2Node_BaseAsyncTask::GetProxyFactoryFunctionName()
+FString UK2Node_BaseAsyncTask::GetTooltip() const
 {
-	return ProxyFactoryFunctionName;
+	const FString FunctionToolTipText = UK2Node_CallFunction::GetDefaultTooltipForFunction(GetFactoryFunction());
+	return FunctionToolTipText;
 }
 
-UClass*	UK2Node_BaseAsyncTask::GetProxyFactoryClass()
+FText UK2Node_BaseAsyncTask::GetNodeTitle(ENodeTitleType::Type TitleType) const
 {
-	return ProxyFactoryClass;
+	const FString FunctionToolTipText = UK2Node_CallFunction::GetUserFacingFunctionName(GetFactoryFunction());
+	return FText::FromString(FunctionToolTipText);
 }
 
-UClass* UK2Node_BaseAsyncTask::GetProxyClass()
+FString UK2Node_BaseAsyncTask::GetCategoryName()
 {
-	return ProxyClass;
+	const FString GenericFunctionCategory = LOCTEXT("FunctionCategory", "Call Function").ToString();
+	return UK2Node_CallFunction::GetDefaultCategoryForFunction(GetFactoryFunction(), GenericFunctionCategory);
 }
 
-FString UK2Node_BaseAsyncTask::GetCategoryName()
+void UK2Node_BaseAsyncTask::GetMenuEntries(FGraphContextMenuBuilder& ContextMenuBuilder) const
+{
+	const UEdGraphSchema_K2* K2Schema = GetDefault<UEdGraphSchema_K2>();
+	const bool bAllowLatentFuncs = (K2Schema->GetGraphType(ContextMenuBuilder.CurrentGraph) == GT_Ubergraph);
+
+	if (bAllowLatentFuncs)
+	{
+		UK2Node_BaseAsyncTask* NodeTemplate = NewObject<UK2Node_BaseAsyncTask>(ContextMenuBuilder.OwnerOfTemporaries, GetClass());
+		CreateDefaultMenuEntry(NodeTemplate, ContextMenuBuilder);
+	}
+}
+
+TSharedPtr<FEdGraphSchemaAction_K2NewNode> UK2Node_BaseAsyncTask::CreateDefaultMenuEntry(UK2Node_BaseAsyncTask* NodeTemplate, FGraphContextMenuBuilder& ContextMenuBuilder) const
 {
-	return TEXT("Latent Execution");
+	TSharedPtr<FEdGraphSchemaAction_K2NewNode> NodeAction = FK2ActionMenuBuilder::AddNewNodeAction(ContextMenuBuilder, NodeTemplate->GetCategoryName(), NodeTemplate->GetNodeTitle(ENodeTitleType::ListView), NodeTemplate->GetTooltip(), 0, NodeTemplate->GetKeywords());
+	
+	NodeAction->NodeTemplate = NodeTemplate;
+	NodeAction->SearchTitle = NodeTemplate->GetNodeSearchTitle();
+
+	return NodeAction;
 }
 
 void UK2Node_BaseAsyncTask::AllocateDefaultPins()
@@ -46,14 +68,13 @@ void UK2Node_BaseAsyncTask::AllocateDefaultPins()
 	UFunction* DelegateSignatureFunction = NULL;
 	for (TFieldIterator<UProperty> PropertyIt(ProxyClass, EFieldIteratorFlags::ExcludeSuper); PropertyIt; ++PropertyIt)
 	{
-		UMulticastDelegateProperty* Property = Cast<UMulticastDelegateProperty>(*PropertyIt);
-		if(!Property)
-			continue;
-
-		CreatePin(EGPD_Output, K2Schema->PC_Exec, TEXT(""), NULL, false, false, *Property->GetName());
-		if (!DelegateSignatureFunction)
+		if (UMulticastDelegateProperty* Property = Cast<UMulticastDelegateProperty>(*PropertyIt))
 		{
-			DelegateSignatureFunction = Property->SignatureFunction;
+			CreatePin(EGPD_Output, K2Schema->PC_Exec, TEXT(""), NULL, false, false, *Property->GetName());
+			if (!DelegateSignatureFunction)
+			{
+				DelegateSignatureFunction = Property->SignatureFunction;
+			}
 		}
 	}
 
@@ -267,58 +288,81 @@ void UK2Node_BaseAsyncTask::ExpandNode(class FKismetCompilerContext& CompilerCon
 		return;
 	}
 	const UEdGraphSchema_K2* Schema = CompilerContext.GetSchema();
-	UK2Node_BaseAsyncTask* const CurrentNode = this;
-	check(SourceGraph && Schema && CurrentNode);
+	check(SourceGraph && Schema);
 	bool bIsErrorFree = true;
 
-	UK2Node_CallFunction* const CallCreateMoveToProxyObjectNode = CompilerContext.SpawnIntermediateNode<UK2Node_CallFunction>(CurrentNode, SourceGraph);
-	CallCreateMoveToProxyObjectNode->FunctionReference.SetExternalMember(CurrentNode->GetProxyFactoryFunctionName(), CurrentNode->GetProxyFactoryClass());
-	CallCreateMoveToProxyObjectNode->AllocateDefaultPins();
-	bIsErrorFree &= CompilerContext.MovePinLinksToIntermediate(*CurrentNode->FindPinChecked(Schema->PN_Execute), *CallCreateMoveToProxyObjectNode->FindPinChecked(Schema->PN_Execute)).CanSafeConnect();
-	for (auto CurrentPin : CurrentNode->Pins)
+	// Create a call to factory the proxy object
+	UK2Node_CallFunction* const CallCreateProxyObjectNode = CompilerContext.SpawnIntermediateNode<UK2Node_CallFunction>(this, SourceGraph);
+	CallCreateProxyObjectNode->FunctionReference.SetExternalMember(ProxyFactoryFunctionName, ProxyFactoryClass);
+	CallCreateProxyObjectNode->AllocateDefaultPins();
+	bIsErrorFree &= CompilerContext.MovePinLinksToIntermediate(*FindPinChecked(Schema->PN_Execute), *CallCreateProxyObjectNode->FindPinChecked(Schema->PN_Execute)).CanSafeConnect();
+	for (auto CurrentPin : Pins)
 	{
 		if (FBaseAsyncTaskHelper::ValidDataPin(CurrentPin, EGPD_Input, Schema))
 		{
-			UEdGraphPin* DestPin = CallCreateMoveToProxyObjectNode->FindPin(CurrentPin->PinName); // match function inputs, to pass data to function from CallFunction node
+			UEdGraphPin* DestPin = CallCreateProxyObjectNode->FindPin(CurrentPin->PinName); // match function inputs, to pass data to function from CallFunction node
 			bIsErrorFree &= DestPin && CompilerContext.MovePinLinksToIntermediate(*CurrentPin, *DestPin).CanSafeConnect();
 		}
 	}
 	
 	// GATHER OUTPUT PARAMETERS AND PAIR THEM WITH LOCAL VARIABLES
 	TArray<FBaseAsyncTaskHelper::FOutputPinAndLocalVariable> VariableOutputs;
-	for (auto CurrentPin : CurrentNode->Pins)
+	for (auto CurrentPin : Pins)
 	{
 		if (FBaseAsyncTaskHelper::ValidDataPin(CurrentPin, EGPD_Output, Schema))
 		{
 			const FEdGraphPinType& PinType = CurrentPin->PinType;
 			UK2Node_TemporaryVariable* TempVarOutput = CompilerContext.SpawnInternalVariable(
-				CurrentNode, PinType.PinCategory, PinType.PinSubCategory, PinType.PinSubCategoryObject.Get(), PinType.bIsArray);
+				this, PinType.PinCategory, PinType.PinSubCategory, PinType.PinSubCategoryObject.Get(), PinType.bIsArray);
 			bIsErrorFree &= TempVarOutput->GetVariablePin() && CompilerContext.MovePinLinksToIntermediate(*CurrentPin, *TempVarOutput->GetVariablePin()).CanSafeConnect();
 			VariableOutputs.Add(FBaseAsyncTaskHelper::FOutputPinAndLocalVariable(CurrentPin, TempVarOutput));
 		}
 	}
 
 	// FOR EACH DELEGATE DEFINE EVENT, CONNECT IT TO DELEGATE AND IMPLEMENT A CHAIN OF ASSIGMENTS
-	UEdGraphPin* LastThenPin = CallCreateMoveToProxyObjectNode->FindPinChecked(Schema->PN_Then);
-	UEdGraphPin* const ProxyObjectPin = CallCreateMoveToProxyObjectNode->GetReturnValuePin();
-	for (TFieldIterator<UMulticastDelegateProperty> PropertyIt(CurrentNode->GetProxyClass(), EFieldIteratorFlags::ExcludeSuper); PropertyIt && bIsErrorFree; ++PropertyIt)
+	UEdGraphPin* LastThenPin = CallCreateProxyObjectNode->FindPinChecked(Schema->PN_Then);
+	UEdGraphPin* const ProxyObjectPin = CallCreateProxyObjectNode->GetReturnValuePin();
+	for (TFieldIterator<UMulticastDelegateProperty> PropertyIt(ProxyClass, EFieldIteratorFlags::ExcludeSuper); PropertyIt && bIsErrorFree; ++PropertyIt)
 	{
-		bIsErrorFree &= FBaseAsyncTaskHelper::HandleDelegateImplementation(*PropertyIt, VariableOutputs, ProxyObjectPin, LastThenPin, CurrentNode, SourceGraph, CompilerContext);
+		bIsErrorFree &= FBaseAsyncTaskHelper::HandleDelegateImplementation(*PropertyIt, VariableOutputs, ProxyObjectPin, LastThenPin, this, SourceGraph, CompilerContext);
 	}
 
-	if (CallCreateMoveToProxyObjectNode->FindPinChecked(Schema->PN_Then) == LastThenPin)
+	if (CallCreateProxyObjectNode->FindPinChecked(Schema->PN_Then) == LastThenPin)
 	{
-		CompilerContext.MessageLog.Error(*LOCTEXT("MissingDelegateProperties", "BaseAsyncTask: Proxy has no delegates defined. @@").ToString(), CurrentNode);
+		CompilerContext.MessageLog.Error(*LOCTEXT("MissingDelegateProperties", "BaseAsyncTask: Proxy has no delegates defined. @@").ToString(), this);
 		return;
 	}
 
-	bIsErrorFree &= CompilerContext.MovePinLinksToIntermediate(*CurrentNode->FindPinChecked(Schema->PN_Then), *LastThenPin).CanSafeConnect();
+	// Create a call to activate the proxy object if necessary
+	if (ProxyActivateFunctionName != NAME_None)
+	{
+		UK2Node_CallFunction* const CallActivateProxyObjectNode = CompilerContext.SpawnIntermediateNode<UK2Node_CallFunction>(this, SourceGraph);
+		CallActivateProxyObjectNode->FunctionReference.SetExternalMember(ProxyActivateFunctionName, ProxyClass);
+		CallActivateProxyObjectNode->AllocateDefaultPins();
+
+		// Hook up the self connection
+		UEdGraphPin* ActivateCallSelfPin = Schema->FindSelfPin(*CallActivateProxyObjectNode, EGPD_Input);
+		check(ActivateCallSelfPin);
+
+		bIsErrorFree &= Schema->TryCreateConnection(ProxyObjectPin, ActivateCallSelfPin);
+		
+		// Hook the activate node up in the exec chain
+		UEdGraphPin* ActivateExecPin = CallActivateProxyObjectNode->FindPinChecked(Schema->PN_Execute);
+		UEdGraphPin* ActivateThenPin = CallActivateProxyObjectNode->FindPinChecked(Schema->PN_Then);
+
+		bIsErrorFree &= CompilerContext.MovePinLinksToIntermediate(*LastThenPin, *ActivateExecPin).CanSafeConnect();
+		LastThenPin = ActivateThenPin;
+	}
+
+	// Move the connections from the original node then pin to the last internal then pin
+	bIsErrorFree &= CompilerContext.MovePinLinksToIntermediate(*FindPinChecked(Schema->PN_Then), *LastThenPin).CanSafeConnect();
 	if (!bIsErrorFree)
 	{
-		CompilerContext.MessageLog.Error(*LOCTEXT("InternalConnectionError", "BaseAsyncTask: Internal connection error. @@").ToString(), CurrentNode);
+		CompilerContext.MessageLog.Error(*LOCTEXT("InternalConnectionError", "BaseAsyncTask: Internal connection error. @@").ToString(), this);
 	}
 
-	CurrentNode->BreakAllNodeLinks();
+	// Make sure we caught everything
+	BreakAllNodeLinks();
 }
 
 bool UK2Node_BaseAsyncTask::HasExternalBlueprintDependencies(TArray<class UStruct*>* OptionalOutput) const
@@ -340,4 +384,17 @@ bool UK2Node_BaseAsyncTask::HasExternalBlueprintDependencies(TArray<class UStruc
 	return bProxyFactoryResult || bProxyResult;
 }
 
+FName UK2Node_BaseAsyncTask::GetCornerIcon() const
+{
+	return TEXT("Graph.Latent.LatentIcon");
+}
+
+UFunction* UK2Node_BaseAsyncTask::GetFactoryFunction() const
+{
+	check(ProxyFactoryClass);
+	UFunction* FactoryFunction = ProxyFactoryClass->FindFunctionByName(ProxyFactoryFunctionName);
+	check(FactoryFunction);
+	return FactoryFunction;
+}
+
 #undef LOCTEXT_NAMESPACE
