#ue4
- Fixed a pin connection issue / potential crash when changing the spawn class type on a SpawnActor node with connected inputs

diff --git a/Engine/Source/Editor/BlueprintGraph/Private/K2Node_SpawnActor.cpp b/Engine/Source/Editor/BlueprintGraph/Private/K2Node_SpawnActor.cpp
--- a/Engine/Source/Editor/BlueprintGraph/Private/K2Node_SpawnActor.cpp
+++ b/Engine/Source/Editor/BlueprintGraph/Private/K2Node_SpawnActor.cpp
@@ -123,9 +123,9 @@ bool UK2Node_SpawnActor::IsSpawnVarPin(UEdGraphPin* Pin)
 }
 
 
-void UK2Node_SpawnActor::PinDefaultValueChanged(UEdGraphPin* Pin) 
+void UK2Node_SpawnActor::PinDefaultValueChanged(UEdGraphPin* ChangedPin) 
 {
-	if(Pin->PinName == BlueprintPinName)
+	if (ChangedPin->PinName == BlueprintPinName)
 	{
 		const UEdGraphSchema_K2* K2Schema = GetDefault<UEdGraphSchema_K2>();
 
@@ -135,12 +135,12 @@ void UK2Node_SpawnActor::PinDefaultValueChanged(UEdGraphPin* Pin)
 
 		// Remove all pins related to archetype variables
 		TArray<UEdGraphPin*> OldPins = Pins;
-		for(int32 i=0; i<OldPins.Num(); i++)
+		for (int32 i = 0; i < OldPins.Num(); i++)
 		{
 			UEdGraphPin* OldPin = OldPins[i];
-			if(	IsSpawnVarPin(OldPin) )
+			if (IsSpawnVarPin(OldPin))
 			{
-				Pin->BreakAllPinLinks();
+				OldPin->BreakAllPinLinks();
 				Pins.Remove(OldPin);
 			}
 		}
diff --git a/Engine/Source/Editor/BlueprintGraph/Private/K2Node_SpawnActorFromClass.cpp b/Engine/Source/Editor/BlueprintGraph/Private/K2Node_SpawnActorFromClass.cpp
--- a/Engine/Source/Editor/BlueprintGraph/Private/K2Node_SpawnActorFromClass.cpp
+++ b/Engine/Source/Editor/BlueprintGraph/Private/K2Node_SpawnActorFromClass.cpp
@@ -144,9 +144,9 @@ bool UK2Node_SpawnActorFromClass::IsSpawnVarPin(UEdGraphPin* Pin)
 }
 
 
-void UK2Node_SpawnActorFromClass::PinDefaultValueChanged(UEdGraphPin* Pin) 
+void UK2Node_SpawnActorFromClass::PinDefaultValueChanged(UEdGraphPin* ChangedPin) 
 {
-	if(Pin->PinName == FK2Node_SpawnActorFromClassHelper::ClassPinName)
+	if (ChangedPin->PinName == FK2Node_SpawnActorFromClassHelper::ClassPinName)
 	{
 		const UEdGraphSchema_K2* K2Schema = GetDefault<UEdGraphSchema_K2>();
 
@@ -156,18 +156,18 @@ void UK2Node_SpawnActorFromClass::PinDefaultValueChanged(UEdGraphPin* Pin)
 
 		// Remove all pins related to archetype variables
 		TArray<UEdGraphPin*> OldPins = Pins;
-		for(int32 i=0; i<OldPins.Num(); i++)
+		for (int32 i = 0; i < OldPins.Num(); i++)
 		{
 			UEdGraphPin* OldPin = OldPins[i];
-			if(	IsSpawnVarPin(OldPin) )
+			if (IsSpawnVarPin(OldPin))
 			{
-				Pin->BreakAllPinLinks();
+				OldPin->BreakAllPinLinks();
 				Pins.Remove(OldPin);
 			}
 		}
 
 		UClass* UseSpawnClass = GetClassToSpawn();
-		if(UseSpawnClass != NULL)
+		if (UseSpawnClass != NULL)
 		{
 			CreatePinsForClass(UseSpawnClass);
 		}
