Fix for AnswerHub: https://answers.unrealengine.com/questions/15846/blueprint-corrupted-crash-when-opening-project.html

Making sure a class property has a valid MetaClass before using it (the user could have renamed/deleted the class).

diff --git a/Engine/Source/Editor/BlueprintGraph/Private/EdGraphSchema_K2.cpp b/Engine/Source/Editor/BlueprintGraph/Private/EdGraphSchema_K2.cpp
--- a/Engine/Source/Editor/BlueprintGraph/Private/EdGraphSchema_K2.cpp
+++ b/Engine/Source/Editor/BlueprintGraph/Private/EdGraphSchema_K2.cpp
@@ -2181,15 +2181,24 @@ FString UEdGraphSchema_K2::TypeToString(UProperty* const Property)
 {
 	if (UStructProperty* Struct = Cast<UStructProperty>(Property))
 	{
-		return FString::Printf(TEXT("struct'%s'"), *Struct->Struct->GetName());
+		if (Struct->Struct)
+		{
+			return FString::Printf(TEXT("struct'%s'"), *Struct->Struct->GetName());
+		}
 	}
 	else if (UClassProperty* Class = Cast<UClassProperty>(Property))
 	{
-		return FString::Printf(TEXT("class'%s'"), *Class->MetaClass->GetName());
+		if (Class->MetaClass != nullptr)
+		{
+			return FString::Printf(TEXT("class'%s'"), *Class->MetaClass->GetName());
+		}
 	}
 	else if (UInterfaceProperty* Interface = Cast<UInterfaceProperty>(Property))
 	{
-		return FString::Printf(TEXT("interface'%s'"), *Interface->InterfaceClass->GetName());
+		if (Interface->InterfaceClass != nullptr)
+		{
+			return FString::Printf(TEXT("interface'%s'"), *Interface->InterfaceClass->GetName());
+		}
 	}
 	else if (UObjectPropertyBase* Obj = Cast<UObjectPropertyBase>(Property))
 	{
@@ -2209,12 +2218,13 @@ FString UEdGraphSchema_K2::TypeToString(UProperty* const Property)
 	}
 	else if (UArrayProperty* Array = Cast<UArrayProperty>(Property))
 	{
-		return FString::Printf(TEXT("array[%s]"), *TypeToString(Array->Inner)); 
-	}
-	else
-	{
-		return Property->GetClass()->GetName();
+		if (Array->Inner)
+		{
+			return FString::Printf(TEXT("array[%s]"), *TypeToString(Array->Inner));
+		}
 	}
+	
+	return Property->GetClass()->GetName();
 }
 
 FText UEdGraphSchema_K2::TypeToText(const FEdGraphPinType& Type)
diff --git a/Engine/Source/Editor/KismetCompiler/Private/KismetCompiler.cpp b/Engine/Source/Editor/KismetCompiler/Private/KismetCompiler.cpp
--- a/Engine/Source/Editor/KismetCompiler/Private/KismetCompiler.cpp
+++ b/Engine/Source/Editor/KismetCompiler/Private/KismetCompiler.cpp
@@ -961,7 +961,10 @@ void FKismetCompilerContext::ValidateSelfPinsInGraph(const UEdGraph* SourceGraph
 									(Pin->PinType.PinCategory == Schema->PC_Interface) ||
 									(Pin->PinType.PinCategory == Schema->PC_Class))
 								{
-									PinType = Pin->PinType.PinSubCategoryObject->GetName();
+									if (Pin->PinType.PinSubCategoryObject.IsValid())
+									{
+										PinType = Pin->PinType.PinSubCategoryObject->GetName();
+									}
 								}
 
 								FString ErrorMsg = FString::Printf(*LOCTEXT("PinMustHaveConnection_Error", "This blueprint (self) is not a %s, therefore '@@' must have a connection").ToString(), *PinType);
