#ue4
- Fixed some issues with async tasks that use ProxyActivateFunctionName

diff --git a/Engine/Source/Editor/BlueprintGraph/Private/K2Node_BaseAsyncTask.cpp b/Engine/Source/Editor/BlueprintGraph/Private/K2Node_BaseAsyncTask.cpp
--- a/Engine/Source/Editor/BlueprintGraph/Private/K2Node_BaseAsyncTask.cpp
+++ b/Engine/Source/Editor/BlueprintGraph/Private/K2Node_BaseAsyncTask.cpp
@@ -350,7 +350,8 @@ void UK2Node_BaseAsyncTask::ExpandNode(class FKismetCompilerContext& CompilerCon
 		UEdGraphPin* ActivateExecPin = CallActivateProxyObjectNode->FindPinChecked(Schema->PN_Execute);
 		UEdGraphPin* ActivateThenPin = CallActivateProxyObjectNode->FindPinChecked(Schema->PN_Then);
 
-		bIsErrorFree &= CompilerContext.MovePinLinksToIntermediate(*LastThenPin, *ActivateExecPin).CanSafeConnect();
+		bIsErrorFree &= Schema->TryCreateConnection(LastThenPin, ActivateExecPin);
+
 		LastThenPin = ActivateThenPin;
 	}
 
diff --git a/Engine/Source/Runtime/Online/OnlineSubsystemUtils/Classes/OnlineBlueprintCallProxyBase.h b/Engine/Source/Runtime/Online/OnlineSubsystemUtils/Classes/OnlineBlueprintCallProxyBase.h
--- a/Engine/Source/Runtime/Online/OnlineSubsystemUtils/Classes/OnlineBlueprintCallProxyBase.h
+++ b/Engine/Source/Runtime/Online/OnlineSubsystemUtils/Classes/OnlineBlueprintCallProxyBase.h
@@ -11,6 +11,6 @@ class UOnlineBlueprintCallProxyBase : public UObject
 	GENERATED_UCLASS_BODY()
 
 	// Called to trigger the actual online action once the delegates have been bound
-	UFUNCTION()
+	UFUNCTION(BlueprintCallable, meta=(BlueprintInternalUseOnly="true"), Category="Online")
 	virtual void Activate();
 };
