ttp#332789 LIVE: CRASH: BP: User's PlayerController became corrupted during development and crashes when trying to load it now

UEdGraphPin::GetOwningNodeUnchecked added. Some 'defensive' code added, to ensure the kismet compiler won't crash, when handling a corrupted asset (especially when a pins outer is not a node).

diff --git a/Engine/Source/Editor/BlueprintGraph/Private/EdGraphSchema_K2.cpp b/Engine/Source/Editor/BlueprintGraph/Private/EdGraphSchema_K2.cpp
--- a/Engine/Source/Editor/BlueprintGraph/Private/EdGraphSchema_K2.cpp
+++ b/Engine/Source/Editor/BlueprintGraph/Private/EdGraphSchema_K2.cpp
@@ -1081,17 +1081,14 @@ const FPinConnectionResponse UEdGraphSchema_K2::DetermineConnectionResponseOfCom
 
 const FPinConnectionResponse UEdGraphSchema_K2::CanCreateConnection(const UEdGraphPin* PinA, const UEdGraphPin* PinB) const
 {
-	// Find the calling context in case one of the pins is of type object and has a value of Self
-	UClass* CallingContext = NULL;
-	const UBlueprint* Blueprint = FBlueprintEditorUtils::FindBlueprintForNode(PinA->GetOwningNode());
-	if (Blueprint)
+	const UK2Node* OwningNodeA = Cast<UK2Node>(PinA->GetOwningNodeUnchecked());
+	const UK2Node* OwningNodeB = Cast<UK2Node>(PinB->GetOwningNodeUnchecked());
+
+	if (!OwningNodeA || !OwningNodeB)
 	{
-		CallingContext = (Blueprint->GeneratedClass != NULL) ? Blueprint->GeneratedClass : Blueprint->ParentClass;
+		return FPinConnectionResponse(CONNECT_RESPONSE_DISALLOW, TEXT("Invalid nodes"));
 	}
 
-	const UK2Node* OwningNodeA = Cast<UK2Node>(PinA->GetOwningNode());
-	const UK2Node* OwningNodeB = Cast<UK2Node>(PinB->GetOwningNode());
-
 	// Make sure the pins are not on the same node
 	if (OwningNodeA == OwningNodeB)
 	{
@@ -1129,6 +1126,14 @@ const FPinConnectionResponse UEdGraphSchema_K2::CanCreateConnection(const UEdGra
 		bIgnoreArray = bAllowMultipleSelfs && bNotAnArrayFunction && bSelfPin;
 	}
 
+	// Find the calling context in case one of the pins is of type object and has a value of Self
+	UClass* CallingContext = NULL;
+	const UBlueprint* Blueprint = FBlueprintEditorUtils::FindBlueprintForNode(PinA->GetOwningNodeUnchecked());
+	if (Blueprint)
+	{
+		CallingContext = (Blueprint->GeneratedClass != NULL) ? Blueprint->GeneratedClass : Blueprint->ParentClass;
+	}
+
 	// Compare the types
 	const bool bTypesMatch = ArePinsCompatible(OutputPin, InputPin, CallingContext, bIgnoreArray);
 
@@ -2183,22 +2188,22 @@ FString UEdGraphSchema_K2::TypeToString(UProperty* const Property)
 	{
 		if (Struct->Struct)
 		{
-			return FString::Printf(TEXT("struct'%s'"), *Struct->Struct->GetName());
-		}
+		return FString::Printf(TEXT("struct'%s'"), *Struct->Struct->GetName());
+	}
 	}
 	else if (UClassProperty* Class = Cast<UClassProperty>(Property))
 	{
 		if (Class->MetaClass != nullptr)
 		{
-			return FString::Printf(TEXT("class'%s'"), *Class->MetaClass->GetName());
-		}
+		return FString::Printf(TEXT("class'%s'"), *Class->MetaClass->GetName());
+	}
 	}
 	else if (UInterfaceProperty* Interface = Cast<UInterfaceProperty>(Property))
 	{
 		if (Interface->InterfaceClass != nullptr)
 		{
-			return FString::Printf(TEXT("interface'%s'"), *Interface->InterfaceClass->GetName());
-		}
+		return FString::Printf(TEXT("interface'%s'"), *Interface->InterfaceClass->GetName());
+	}
 	}
 	else if (UObjectPropertyBase* Obj = Cast<UObjectPropertyBase>(Property))
 	{
@@ -2220,12 +2225,12 @@ FString UEdGraphSchema_K2::TypeToString(UProperty* const Property)
 	{
 		if (Array->Inner)
 		{
-			return FString::Printf(TEXT("array[%s]"), *TypeToString(Array->Inner));
-		}
+		return FString::Printf(TEXT("array[%s]"), *TypeToString(Array->Inner)); 
+	}
 	}
 	
-	return Property->GetClass()->GetName();
-}
+		return Property->GetClass()->GetName();
+	}
 
 FText UEdGraphSchema_K2::TypeToText(const FEdGraphPinType& Type)
 {
diff --git a/Engine/Source/Editor/BlueprintGraph/Private/K2Node.cpp b/Engine/Source/Editor/BlueprintGraph/Private/K2Node.cpp
--- a/Engine/Source/Editor/BlueprintGraph/Private/K2Node.cpp
+++ b/Engine/Source/Editor/BlueprintGraph/Private/K2Node.cpp
@@ -344,7 +344,7 @@ void UK2Node::ReconstructNode()
 		{
 			UEdGraphPin* OtherPin = LinkedToCopy[LinkIdx];
 			// If we are linked to a pin that its owner doesn't know about, break that link
-			if ((OtherPin == NULL) || !OtherPin->GetOwningNode()->Pins.Contains(OtherPin))
+			if ((OtherPin == NULL) || !OtherPin->GetOwningNodeUnchecked() || !OtherPin->GetOwningNode()->Pins.Contains(OtherPin))
 			{
 				Pin->LinkedTo.Remove(OtherPin);
 			}
diff --git a/Engine/Source/Editor/KismetCompiler/Private/EdGraphCompiler.cpp b/Engine/Source/Editor/KismetCompiler/Private/EdGraphCompiler.cpp
--- a/Engine/Source/Editor/KismetCompiler/Private/EdGraphCompiler.cpp
+++ b/Engine/Source/Editor/KismetCompiler/Private/EdGraphCompiler.cpp
@@ -18,7 +18,7 @@ void FGraphCompilerContext::ValidateLink(const UEdGraphPin* PinA, const UEdGraph
 		MessageLog.Error(TEXT("Direction mismatch between pins @@ and @@"), PinA, PinB);
 	}
 
-	if (PinA->GetOwningNode() == PinB->GetOwningNode())
+	if (PinA->GetOwningNodeUnchecked() == PinB->GetOwningNodeUnchecked())
 	{
 		MessageLog.Error(TEXT("Pins @@ and @@ on the same node @@ are connected to each other, creating a loop."), PinA, PinB, PinA->GetOwningNode());
 	}
@@ -27,8 +27,21 @@ void FGraphCompilerContext::ValidateLink(const UEdGraphPin* PinA, const UEdGraph
 /** Validate that the wiring for a single pin is schema compatible */
 void FGraphCompilerContext::ValidatePin(const UEdGraphPin* Pin) const
 {
+	if (NULL == Pin->GetOwningNodeUnchecked())
+	{
+		MessageLog.Error(
+			*FString::Printf(
+				*NSLOCTEXT("EdGraphCompiler", "PinWrongOuterError", "Blueprint is corrupted! Pin '%s' has wrong outer '%s'.").ToString(),
+				*Pin->GetName(),
+				Pin->GetOuter() ? *Pin->GetOuter()->GetName() : TEXT("NULL")),
+			Pin);
+		return;
+	}
+
+	const int32 ErrorNum = MessageLog.NumErrors;
+
 	// Validate the links to each connected pin
-	for (int32 LinkIndex = 0; LinkIndex < Pin->LinkedTo.Num(); ++LinkIndex)
+	for (int32 LinkIndex = 0; (LinkIndex < Pin->LinkedTo.Num()) && (ErrorNum == MessageLog.NumErrors); ++LinkIndex)
 	{
 		if (const UEdGraphPin* OtherPin = Pin->LinkedTo[LinkIndex])
 		{
@@ -207,10 +220,11 @@ void FGraphCompilerContext::CreateExecutionSchedule(const TArray<UEdGraphNode*>&
 						continue;
 					}
 
-					UEdGraphNode* WasDependentNode = OutPin->LinkedTo[LinkIndex]->GetOwningNode();
+					UEdGraphNode* WasDependentNode = OutPin->LinkedTo[LinkIndex]->GetOwningNodeUnchecked();
+					int32* pNumEdgesLeft = WasDependentNode ? NumIncomingEdges.Find(WasDependentNode) : NULL;
 
 					// Remove the edge between these two nodes, since node is scheduled
-					if (int32* pNumEdgesLeft = NumIncomingEdges.Find(WasDependentNode))
+					if (pNumEdgesLeft)
 					{
 						int32& NumEdgesLeft = *pNumEdgesLeft;
 
diff --git a/Engine/Source/Editor/UnrealEd/Private/Kismet2/CompilerResultsLog.cpp b/Engine/Source/Editor/UnrealEd/Private/Kismet2/CompilerResultsLog.cpp
--- a/Engine/Source/Editor/UnrealEd/Private/Kismet2/CompilerResultsLog.cpp
+++ b/Engine/Source/Editor/UnrealEd/Private/Kismet2/CompilerResultsLog.cpp
@@ -128,7 +128,7 @@ void FCompilerResultsLog::InternalLogMessage(const EMessageSeverity::Type& Sever
 					OwnerNode = Cast<UEdGraphNode>(ObjectArgument);
 					if (Pin)
 					{
-						OwnerNode = Cast<UEdGraphNode>(Pin->GetOwningNode());
+						OwnerNode = Pin->GetOwningNodeUnchecked();
 					}
 
 					if (ObjectArgument->GetOutermost() == GetTransientPackage())
diff --git a/Engine/Source/Editor/UnrealEd/Public/EdGraphUtilities.h b/Engine/Source/Editor/UnrealEd/Public/EdGraphUtilities.h
--- a/Engine/Source/Editor/UnrealEd/Public/EdGraphUtilities.h
+++ b/Engine/Source/Editor/UnrealEd/Public/EdGraphUtilities.h
@@ -116,8 +116,8 @@ public:
 					UEdGraphPin* OtherPin = MyPin->LinkedTo[j];
 					if( OtherPin )
 					{
-						UEdGraphNode* OtherNode = OtherPin->GetOwningNode();
-						if (!VisitedNodes.Contains(OtherNode))
+						UEdGraphNode* OtherNode = OtherPin->GetOwningNodeUnchecked();
+						if (OtherNode && !VisitedNodes.Contains(OtherNode))
 						{
 							TraverseNodes(OtherNode);
 						}
diff --git a/Engine/Source/Runtime/Engine/Classes/EdGraph/EdGraphPin.h b/Engine/Source/Runtime/Engine/Classes/EdGraph/EdGraphPin.h
--- a/Engine/Source/Runtime/Engine/Classes/EdGraph/EdGraphPin.h
+++ b/Engine/Source/Runtime/Engine/Classes/EdGraph/EdGraphPin.h
@@ -215,6 +215,11 @@ public:
 		return CastChecked<UEdGraphNode>(GetOuter());
 	}
 
+	class UEdGraphNode* GetOwningNodeUnchecked() const
+	{
+		return Cast<UEdGraphNode>(GetOuter());
+	}
+
 	/** Shorthand way to access the schema of the graph that owns the node that owns this pin */
 	ENGINE_API const class UEdGraphSchema* GetSchema() const;
 
