Make/Break struct cleanup
- Use display name instead of friendly name metadata
- Cleaned up small number of properties settting friendly name
- Changed blueprint created variables and user defined struct variables to write to DisplayName instead of FriendlyName (fallback exists when displaying pins from blueprints/structs that have not been recompiled/re-saved)
- Property tooltips now display when hovering over pin in addition to type information
- Hide the edit box next to pin for break struct

diff --git a/Engine/Source/Editor/BlueprintGraph/Classes/K2Node.h b/Engine/Source/Editor/BlueprintGraph/Classes/K2Node.h
--- a/Engine/Source/Editor/BlueprintGraph/Classes/K2Node.h
+++ b/Engine/Source/Editor/BlueprintGraph/Classes/K2Node.h
@@ -59,6 +59,9 @@ struct FOptionalPinFromProperty
 	UPROPERTY(EditAnywhere, Category=Hi, BlueprintReadOnly)
 	FString PropertyFriendlyName;
 
+	UPROPERTY(EditAnywhere, Category=Hi, BlueprintReadOnly)
+	FText PropertyTooltip;
+
 	UPROPERTY(EditAnywhere, BlueprintReadWrite, Category=Hi)
 	bool bShowPin;
 
@@ -69,11 +72,12 @@ struct FOptionalPinFromProperty
 	{
 	}
 	
-	FOptionalPinFromProperty(FName InPropertyName, bool bInShowPin, bool bInCanToggleVisibility, const FString& InFriendlyName)
+	FOptionalPinFromProperty(FName InPropertyName, bool bInShowPin, bool bInCanToggleVisibility, const FString& InFriendlyName, const FText& InTooltip)
 		: PropertyName(InPropertyName)
+		, PropertyFriendlyName(InFriendlyName)
+		, PropertyTooltip(InTooltip)
 		, bShowPin(bInShowPin)
 		, bCanToggleVisibility(bInCanToggleVisibility)
-		, PropertyFriendlyName(InFriendlyName)
 	{
 	}
 };
diff --git a/Engine/Source/Editor/BlueprintGraph/Private/EdGraphSchema_K2.cpp b/Engine/Source/Editor/BlueprintGraph/Private/EdGraphSchema_K2.cpp
--- a/Engine/Source/Editor/BlueprintGraph/Private/EdGraphSchema_K2.cpp
+++ b/Engine/Source/Editor/BlueprintGraph/Private/EdGraphSchema_K2.cpp
@@ -2052,7 +2052,10 @@ void UEdGraphSchema_K2::ConstructBasicPinTooltip(const UEdGraphPin& Pin, const F
 		}
 	}	
 
-	ConstructedTooltip += FString(TEXT("\n")) + PinDescription;
+	if (!PinDescription.IsEmpty())
+	{
+		ConstructedTooltip += FString(TEXT("\n")) + PinDescription;
+	}
 	TooltipOut = ConstructedTooltip; // using a local FString, so PinDescription and TooltipOut can be the same variable
 }
 
diff --git a/Engine/Source/Editor/BlueprintGraph/Private/K2Node.cpp b/Engine/Source/Editor/BlueprintGraph/Private/K2Node.cpp
--- a/Engine/Source/Editor/BlueprintGraph/Private/K2Node.cpp
+++ b/Engine/Source/Editor/BlueprintGraph/Private/K2Node.cpp
@@ -795,8 +795,8 @@ void FOptionalPinManager::RebuildPropertyList(TArray<FOptionalPinFromProperty>&
 		{
 			FOptionalPinFromProperty* Record = new (Properties) FOptionalPinFromProperty;
 			Record->PropertyName = TestProperty->GetFName();
-			const FString FriendlyName = TestProperty->GetMetaData(TEXT("FriendlyName"));
-			Record->PropertyFriendlyName = FriendlyName.IsEmpty() ? Record->PropertyName.ToString() : FriendlyName;
+			Record->PropertyFriendlyName = UEditorEngine::GetFriendlyName(TestProperty, SourceStruct);
+			Record->PropertyTooltip = TestProperty->GetToolTipText();
 
 			// Get the defaults
 			GetRecordDefaults(TestProperty, *Record);
@@ -848,6 +848,7 @@ void FOptionalPinManager::CreateVisiblePins(TArray<FOptionalPinFromProperty>& Pr
 							const FString PinName = PinFriendlyName.ToString();
 							NewPin = TargetNode->CreatePin(Direction, PinType.PinCategory, PinType.PinSubCategory, PinType.PinSubCategoryObject.Get(), PinType.bIsArray, PinType.bIsReference, PinName);
 							NewPin->PinFriendlyName = PinFriendlyName;
+							Schema->ConstructBasicPinTooltip(*NewPin, TEXT("\n") + PropertyEntry.PropertyTooltip.ToString(), NewPin->PinToolTip);
 
 							// Allow the derived class to customize the created pin
 							CustomizePinData(NewPin, PropertyEntry.PropertyName, Index, InnerProperty);
@@ -880,6 +881,7 @@ void FOptionalPinManager::CreateVisiblePins(TArray<FOptionalPinFromProperty>& Pr
 						const FString PinName = PropertyEntry.PropertyName.ToString();
 						NewPin = TargetNode->CreatePin(Direction, PinType.PinCategory, PinType.PinSubCategory, PinType.PinSubCategoryObject.Get(), PinType.bIsArray, PinType.bIsReference, PinName);
 						NewPin->PinFriendlyName = PropertyEntry.PropertyFriendlyName.IsEmpty() ? FText::FromString(PinName) : FText::FromString(PropertyEntry.PropertyFriendlyName);
+						Schema->ConstructBasicPinTooltip(*NewPin, TEXT("\n") + PropertyEntry.PropertyTooltip.ToString(), NewPin->PinToolTip);
 
 						// Allow the derived class to customize the created pin
 						CustomizePinData(NewPin, PropertyEntry.PropertyName, INDEX_NONE, OuterProperty);
diff --git a/Engine/Source/Editor/BlueprintGraph/Private/K2Node_BreakStruct.cpp b/Engine/Source/Editor/BlueprintGraph/Private/K2Node_BreakStruct.cpp
--- a/Engine/Source/Editor/BlueprintGraph/Private/K2Node_BreakStruct.cpp
+++ b/Engine/Source/Editor/BlueprintGraph/Private/K2Node_BreakStruct.cpp
@@ -151,7 +151,7 @@ void UK2Node_BreakStruct::AllocateDefaultPins()
 	const UEdGraphSchema_K2* Schema = GetDefault<UEdGraphSchema_K2>();
 	if(Schema && StructType)
 	{
-		CreatePin(EGPD_Input, Schema->PC_Struct, TEXT(""), StructType, false, false, StructType->GetName());
+		CreatePin(EGPD_Input, Schema->PC_Struct, TEXT(""), StructType, false, true, StructType->GetName(), true);
 		
 		UK2Node_StructMemberGet::AllocateDefaultPins();
 
diff --git a/Engine/Source/Editor/KismetCompiler/Private/KismetCompiler.cpp b/Engine/Source/Editor/KismetCompiler/Private/KismetCompiler.cpp
--- a/Engine/Source/Editor/KismetCompiler/Private/KismetCompiler.cpp
+++ b/Engine/Source/Editor/KismetCompiler/Private/KismetCompiler.cpp
@@ -455,7 +455,7 @@ void FKismetCompilerContext::CreateClassVariablesFromBlueprint()
 		if (NewProperty != NULL)
 		{
 			NewProperty->SetPropertyFlags(Variable.PropertyFlags);
-			NewProperty->SetMetaData(TEXT("FriendlyName"), *Variable.FriendlyName);
+			NewProperty->SetMetaData(TEXT("DisplayName"), *Variable.FriendlyName);
 			NewProperty->SetMetaData(TEXT("Category"), *Variable.Category.ToString());
 			NewProperty->RepNotifyFunc = Variable.RepNotifyFunc;
 
diff --git a/Engine/Source/Editor/KismetCompiler/Private/UserDefinedStructureCompilerUtils.cpp b/Engine/Source/Editor/KismetCompiler/Private/UserDefinedStructureCompilerUtils.cpp
--- a/Engine/Source/Editor/KismetCompiler/Private/UserDefinedStructureCompilerUtils.cpp
+++ b/Engine/Source/Editor/KismetCompiler/Private/UserDefinedStructureCompilerUtils.cpp
@@ -146,7 +146,6 @@ struct FUserDefinedStructureCompilerInner
 			}
 
 			NewProperty->SetPropertyFlags(CPF_Edit);
-			NewProperty->SetMetaData(TEXT("FriendlyName"), *VarDesc.FriendlyName);
 			NewProperty->SetMetaData(TEXT("DisplayName"), *VarDesc.FriendlyName);
 			NewProperty->SetMetaData(TEXT("Category"), *VarDesc.Category);
 			NewProperty->RepNotifyFunc = NAME_None;
diff --git a/Engine/Source/Editor/UnrealEd/Classes/Editor/EditorEngine.h b/Engine/Source/Editor/UnrealEd/Classes/Editor/EditorEngine.h
--- a/Engine/Source/Editor/UnrealEd/Classes/Editor/EditorEngine.h
+++ b/Engine/Source/Editor/UnrealEd/Classes/Editor/EditorEngine.h
@@ -2178,7 +2178,7 @@ public:
 	 *
 	 * @return	the friendly name for the property.  localized first, then metadata, then the property's name.
 	 */
-	static FString GetFriendlyName( const UProperty* Property, UClass* OwnerClass = NULL );
+	static FString GetFriendlyName( const UProperty* Property, UStruct* OwnerClass = NULL );
 
 	/**
 	 * Register a client tool to receive undo events 
diff --git a/Engine/Source/Editor/UnrealEd/Private/Editor.cpp b/Engine/Source/Editor/UnrealEd/Private/Editor.cpp
--- a/Engine/Source/Editor/UnrealEd/Private/Editor.cpp
+++ b/Engine/Source/Editor/UnrealEd/Private/Editor.cpp
@@ -4499,7 +4499,7 @@ void UEditorEngine::SetActorLabelUnique( AActor* Actor, const FString& NewActorL
 
 
 
-FString UEditorEngine::GetFriendlyName( const UProperty* Property, UClass* OwnerClass/* = NULL*/ )
+FString UEditorEngine::GetFriendlyName( const UProperty* Property, UStruct* OwnerClass/* = NULL*/ )
 {
 	// first, try to pull the friendly name from the loc file
 	check( Property );
@@ -4512,19 +4512,24 @@ FString UEditorEngine::GetFriendlyName( const UProperty* Property, UClass* Owner
 
 	FText FoundText;
 	bool DidFindText = false;
-	UClass* CurrentClass = OwnerClass;
+	UStruct* CurrentClass = OwnerClass;
 	do 
 	{
 		FString PropertyPathName = Property->GetPathName(CurrentClass);
 
 		DidFindText = FText::FindText(*CurrentClass->GetName(), *(PropertyPathName + TEXT(".FriendlyName")), /*OUT*/FoundText );
-		CurrentClass = CurrentClass->GetSuperClass();
+		CurrentClass = CurrentClass->GetSuperStruct();
 	} while( CurrentClass != NULL && CurrentClass->IsChildOf(RealOwnerClass) && !DidFindText );
 
 	if ( !DidFindText )
 	{
-		FString DefaultFriendlyName = Property->GetMetaData(TEXT("FriendlyName"));
-		if ( DefaultFriendlyName.Len() == 0 )
+		FString DefaultFriendlyName = Property->GetMetaData(TEXT("DisplayName"));
+		if (DefaultFriendlyName.IsEmpty())
+		{
+			// Fallback for old blueprint variables that were saved as friendly name
+			DefaultFriendlyName = Property->GetMetaData(TEXT("FriendlyName"));
+		}
+		if ( DefaultFriendlyName.IsEmpty() )
 		{
 			const bool bIsBool = Cast<const UBoolProperty>(Property) != NULL;
 			return FName::NameToDisplayString( Property->GetName(), bIsBool );
diff --git a/Engine/Source/Runtime/Engine/Classes/Components/SkeletalMeshComponent.h b/Engine/Source/Runtime/Engine/Classes/Components/SkeletalMeshComponent.h
--- a/Engine/Source/Runtime/Engine/Classes/Components/SkeletalMeshComponent.h
+++ b/Engine/Source/Runtime/Engine/Classes/Components/SkeletalMeshComponent.h
@@ -156,19 +156,19 @@ struct FSingleAnimationPlayData
 	class UVertexAnimation* VertexAnimToPlay;
 
 	// Default setting for looping for SequenceToPlay. This is not current state of looping.
-	UPROPERTY(EditAnywhere, BlueprintReadWrite, Category=Animation, meta=(FriendlyName="Looping"))
+	UPROPERTY(EditAnywhere, BlueprintReadWrite, Category=Animation, meta=(DisplayName="Looping"))
 	uint32 bSavedLooping:1;
 
 	// Default setting for playing for SequenceToPlay. This is not current state of playing.
-	UPROPERTY(EditAnywhere, BlueprintReadWrite, Category=Animation, meta=(FriendlyName="Playing"))
+	UPROPERTY(EditAnywhere, BlueprintReadWrite, Category=Animation, meta=(DisplayName="Playing"))
 	uint32 bSavedPlaying:1;
 
 	// Default setting for position of SequenceToPlay to play. 
-	UPROPERTY(EditAnywhere, BlueprintReadWrite, Category=Animation, meta=(FriendlyName="InitialPosition"))
+	UPROPERTY(EditAnywhere, BlueprintReadWrite, Category=Animation, meta=(DisplayName="InitialPosition"))
 	float SavedPosition;
 
 	// Default setting for playrate of SequenceToPlay to play. 
-	UPROPERTY(EditAnywhere, AdvancedDisplay, BlueprintReadWrite, Category=Animation, meta=(FriendlyName="PlayRate"))
+	UPROPERTY(EditAnywhere, AdvancedDisplay, BlueprintReadWrite, Category=Animation, meta=(DisplayName="PlayRate"))
 	float SavedPlayRate;
 
 	FSingleAnimationPlayData()
diff --git a/Engine/Source/Runtime/Engine/Classes/EdGraph/EdGraphNode_Comment.h b/Engine/Source/Runtime/Engine/Classes/EdGraph/EdGraphNode_Comment.h
--- a/Engine/Source/Runtime/Engine/Classes/EdGraph/EdGraphNode_Comment.h
+++ b/Engine/Source/Runtime/Engine/Classes/EdGraph/EdGraphNode_Comment.h
@@ -28,7 +28,7 @@ class UEdGraphNode_Comment : public UEdGraphNode
 	FLinearColor CommentColor;
 
 	/** Whether to use Comment Color to color the background of the comment bubble shown when zoomed out. */
-	UPROPERTY(EditAnywhere, Category=Comment, meta=(FriendlyName="Color Bubble"))
+	UPROPERTY(EditAnywhere, Category=Comment, meta=(DisplayName="Color Bubble"))
 	uint32 bColorCommentBubble:1;
 
 	/** Whether the comment should move any fully enclosed nodes around when it is moved */
diff --git a/Engine/Source/Runtime/Engine/Classes/GameFramework/AIController.h b/Engine/Source/Runtime/Engine/Classes/GameFramework/AIController.h
--- a/Engine/Source/Runtime/Engine/Classes/GameFramework/AIController.h
+++ b/Engine/Source/Runtime/Engine/Classes/GameFramework/AIController.h
@@ -165,7 +165,7 @@ public:
 	FORCEINLINE FAIRequestID GetCurrentMoveRequestID() const { return PathFollowingComponent.IsValid() ? PathFollowingComponent->GetCurrentRequestId() : FAIRequestID::InvalidRequest; }
 
 	/** Blueprint notification that we've completed the current movement request */
-	UPROPERTY(BlueprintAssignable, meta=(FriendlyName="MoveCompleted"))
+	UPROPERTY(BlueprintAssignable, meta=(DisplayName="MoveCompleted"))
 	FAIMoveCompletedSignature ReceiveMoveCompleted;
 
 	/** @returns path to actor Goal (can be incomplete if async pathfinding is used) */
diff --git a/Engine/Source/Runtime/Engine/Classes/GameFramework/WorldSettings.h b/Engine/Source/Runtime/Engine/Classes/GameFramework/WorldSettings.h
--- a/Engine/Source/Runtime/Engine/Classes/GameFramework/WorldSettings.h
+++ b/Engine/Source/Runtime/Engine/Classes/GameFramework/WorldSettings.h
@@ -193,7 +193,7 @@ class ENGINE_API AWorldSettings : public AInfo
 	uint32 bWorldGravitySet:1;
 
 	/** If set to true we will use GlobalGravityZ instead of project setting DefaultGravityZ */
-	UPROPERTY(EditAnywhere, BlueprintReadOnly, meta=(FriendlyName = "Override World Gravity"), Category = Physics)
+	UPROPERTY(EditAnywhere, BlueprintReadOnly, meta=(DisplayName = "Override World Gravity"), Category = Physics)
 	uint32 bGlobalGravitySet:1;
 
 	// any actor falling below this level gets destroyed
