Mac compilation fix - K2Node_Variable

diff --git a/Engine/Source/Editor/BlueprintGraph/Classes/K2Node_Variable.h b/Engine/Source/Editor/BlueprintGraph/Classes/K2Node_Variable.h
--- a/Engine/Source/Editor/BlueprintGraph/Classes/K2Node_Variable.h
+++ b/Engine/Source/Editor/BlueprintGraph/Classes/K2Node_Variable.h
@@ -115,7 +115,7 @@ public:
 	BLUEPRINTGRAPH_API UEdGraphPin* GetValuePin() const;
 
 	/** Validates there are no errors in the node */
-	void CheckForErrors(const UEdGraphSchema_K2* Schema, FCompilerResultsLog& MessageLog);
+	void CheckForErrors(const class UEdGraphSchema_K2* Schema, class FCompilerResultsLog& MessageLog);
 
 	/**
 	 * Utility method intended to serve as a choke point for various slate 
