Blueprint fucntion library: Self node in BFL causes a warning. ActionMenuBluilder minior fixes.

diff --git a/Engine/Source/Editor/BlueprintGraph/Classes/K2Node_Self.h b/Engine/Source/Editor/BlueprintGraph/Classes/K2Node_Self.h
--- a/Engine/Source/Editor/BlueprintGraph/Classes/K2Node_Self.h
+++ b/Engine/Source/Editor/BlueprintGraph/Classes/K2Node_Self.h
@@ -18,6 +18,7 @@ class UK2Node_Self : public UK2Node
 	virtual FString GetTooltip() const OVERRIDE;
 	virtual FString GetKeywords() const OVERRIDE;
 	virtual void AllocateDefaultPins() OVERRIDE;
+	virtual void ValidateNodeDuringCompilation(class FCompilerResultsLog& MessageLog) const OVERRIDE;
 	// End UEdGraphNode interface
 
 	// Begin K2Node interface
diff --git a/Engine/Source/Editor/BlueprintGraph/Private/K2ActionMenuBuilder.cpp b/Engine/Source/Editor/BlueprintGraph/Private/K2ActionMenuBuilder.cpp
--- a/Engine/Source/Editor/BlueprintGraph/Private/K2ActionMenuBuilder.cpp
+++ b/Engine/Source/Editor/BlueprintGraph/Private/K2ActionMenuBuilder.cpp
@@ -1444,12 +1444,14 @@ void FK2ActionMenuBuilder::GetFuncNodesForClass(FGraphActionListBuilderBase& Lis
 	else
 	{
 		// Find all functions in each function library
-		for (TObjectIterator<UClass> ClassIt; ClassIt; ++ClassIt)
+		for (TObjectIterator<UClass> ClassIt(RF_Transient | RF_PendingKill); ClassIt; ++ClassIt)
 		{
 			UClass* TestClass = *ClassIt;
-			const bool bIsSkeletonClass = TestClass->HasAnyFlags(RF_Transient) && TestClass->HasAnyClassFlags(CLASS_CompiledFromBlueprint);
-			if (TestClass->IsChildOf(UBlueprintFunctionLibrary::StaticClass()) &&
-				CanUseLibraryFunctionsForScope(TestClass, OwnerClass) && !bIsSkeletonClass)
+			const bool bTransient = (GetTransientPackage() == TestClass->GetOutermost());
+			const UClass* AuthoritativeOwnerClass = OwnerClass ? const_cast<UClass*>(OwnerClass)->GetAuthoritativeClass() : NULL;
+			const bool bCanShowInOwnerClass = !AuthoritativeOwnerClass || !TestClass->IsChildOf(AuthoritativeOwnerClass);
+			const bool bIsProperFucntionLibrary = TestClass->IsChildOf(UBlueprintFunctionLibrary::StaticClass()) && CanUseLibraryFunctionsForScope(TestClass, OwnerClass);
+			if (bCanShowInOwnerClass && bIsProperFucntionLibrary && !bTransient)
 			{
 				GetFuncNodesForClass(ListBuilder, TestClass, OwnerClass, DestGraph, FunctionTypes, BaseCategory, TargetInfo);
 			}
@@ -1775,7 +1777,8 @@ void FK2ActionMenuBuilder::GetVariableGettersSettersForClass(FBlueprintGraphActi
 	}
 
 	// Add in self-context variable getter, if needed
-	if( bSelfContext && bWantGetters )
+	const bool bSelfIsAllowed = !Class->IsChildOf(UBlueprintFunctionLibrary::StaticClass());
+	if (bSelfContext && bWantGetters && bSelfIsAllowed)
 	{
 		TSharedPtr<FEdGraphSchemaAction_K2NewNode> SelfVar = FK2ActionMenuBuilder::AddNewNodeAction(ContextMenuBuilder, K2ActionCategories::VariablesCategory, FText::FromString(SelfString), SelfTooltip);
 
diff --git a/Engine/Source/Editor/BlueprintGraph/Private/K2Node_Self.cpp b/Engine/Source/Editor/BlueprintGraph/Private/K2Node_Self.cpp
--- a/Engine/Source/Editor/BlueprintGraph/Private/K2Node_Self.cpp
+++ b/Engine/Source/Editor/BlueprintGraph/Private/K2Node_Self.cpp
@@ -61,3 +61,15 @@ FNodeHandlingFunctor* UK2Node_Self::CreateNodeHandler(FKismetCompilerContext& Co
 {
 	return new FKCHandler_Self(CompilerContext);
 }
+
+void UK2Node_Self::ValidateNodeDuringCompilation(class FCompilerResultsLog& MessageLog) const
+{
+	Super::ValidateNodeDuringCompilation(MessageLog);
+	const UBlueprint* Blueprint = GetBlueprint();
+	const UClass* MyClass = SelfClass ? SelfClass : (Blueprint ? Blueprint->GeneratedClass : NULL);
+	const bool bValidClass = !MyClass || !MyClass->IsChildOf(UBlueprintFunctionLibrary::StaticClass());
+	if (!bValidClass)
+	{
+		MessageLog.Warning(*NSLOCTEXT("K2Node", "InvalidSelfNode", "Self node @@ cannot be used in static library.").ToString(), this);
+	}
+}
\ No newline at end of file
