Blueprint Function Library handles 'DefaultToSelf' and 'WorldContext; parameters.

ttp333765 Function libraries that use blueprint nodes requiring a worldcontextobject cause ensures in cooked builds

diff --git a/Engine/Source/Editor/BlueprintGraph/Classes/K2Node_FunctionEntry.h b/Engine/Source/Editor/BlueprintGraph/Classes/K2Node_FunctionEntry.h
--- a/Engine/Source/Editor/BlueprintGraph/Classes/K2Node_FunctionEntry.h
+++ b/Engine/Source/Editor/BlueprintGraph/Classes/K2Node_FunctionEntry.h
@@ -45,5 +45,7 @@ class UK2Node_FunctionEntry : public UK2Node_FunctionTerminator
 
 	// Removes an output pin from the node
 	BLUEPRINTGRAPH_API void RemoveOutputPin(UEdGraphPin* PinToRemove);
+
+	BLUEPRINTGRAPH_API UEdGraphPin* GetWorldContextPin() const;
 };
 
diff --git a/Engine/Source/Editor/BlueprintGraph/Private/K2Node_CallFunction.cpp b/Engine/Source/Editor/BlueprintGraph/Private/K2Node_CallFunction.cpp
--- a/Engine/Source/Editor/BlueprintGraph/Private/K2Node_CallFunction.cpp
+++ b/Engine/Source/Editor/BlueprintGraph/Private/K2Node_CallFunction.cpp
@@ -618,7 +618,7 @@ bool UK2Node_CallFunction::CreatePinsForFunctionCall(const UFunction* Function)
 				if (!bShowHiddenSelfPins || !bIsSelfPin)
 				{
 					Pin->bHidden = true;
-					Pin->DefaultValue = TEXT("0");
+					K2Schema->SetPinDefaultValueBasedOnType(Pin);
 				}
 			}
 
@@ -1082,6 +1082,7 @@ void UK2Node_CallFunction::PostPasteNode()
 		FString const DefaultToSelfMetaValue = Function->GetMetaData(FBlueprintMetadata::MD_DefaultToSelf);
 		FString const WorldContextMetaValue  = Function->GetMetaData(FBlueprintMetadata::MD_WorldContext);
 
+		const UEdGraphSchema_K2* K2Schema = GetDefault<UEdGraphSchema_K2>();
 		for (int32 PinIndex = 0; PinIndex < Pins.Num(); ++PinIndex)
 		{
 			UEdGraphPin* Pin = Pins[PinIndex];
@@ -1092,7 +1093,7 @@ void UK2Node_CallFunction::PostPasteNode()
 			if (bPinShouldBeHidden && !Pin->bHidden)
 			{
 				Pin->BreakAllPinLinks();
-				Pin->DefaultValue = TEXT("0");
+				K2Schema->SetPinDefaultValueBasedOnType(Pin);
 			}
 			Pin->bHidden = bPinShouldBeHidden;
 		}
@@ -1197,113 +1198,51 @@ FNodeHandlingFunctor* UK2Node_CallFunction::CreateNodeHandler(FKismetCompilerCon
 	return new FKCHandler_CallFunction(CompilerContext);
 }
 
-/*
-struct FBackwardCompatibilityHelper
+void UK2Node_CallFunction::ExpandNode(class FKismetCompilerContext& CompilerContext, UEdGraph* SourceGraph)
 {
-private:
-	static void Replace_CreateSaveGameObjectFromBlueprint(UK2Node_CallFunction* OldNode, FKismetCompilerContext& CompilerContext, UEdGraph* SourceGraph)
-	{
-		const UEdGraphSchema_K2* Schema = CompilerContext.GetSchema();
-		check(Schema);
-
-		UK2Node_CallFunction* NewNode = CompilerContext.SpawnIntermediateNode<UK2Node_CallFunction>(OldNode, SourceGraph); 
-		check(NewNode);
-		const UFunction* NewFunction = UGameplayStatics::StaticClass()->FindFunctionByName(TEXT("CreateSaveGameObject"));
-		check(NewFunction);
-		NewNode->SetFromFunction(NewFunction);
-		NewNode->AllocateDefaultPins();
-
-		check(OldNode->GetThenPin() && NewNode->GetThenPin());
-		CompilerContext.CheckConnectionResponse(Schema->MovePinLinks(*OldNode->GetThenPin(), *NewNode->GetThenPin()), OldNode);
-
-		check(OldNode->GetReturnValuePin() && NewNode->GetReturnValuePin());
-		CompilerContext.CheckConnectionResponse(Schema->MovePinLinks(*OldNode->GetReturnValuePin(), *NewNode->GetReturnValuePin()), OldNode);
-
-		UEdGraphPin* OldExecPin = OldNode->FindPinChecked(Schema->PN_Execute);
-		UEdGraphPin* NewExecPin = NewNode->FindPinChecked(Schema->PN_Execute);
-		CompilerContext.CheckConnectionResponse(Schema->MovePinLinks(*OldExecPin, *NewExecPin), OldNode);
-
-		UEdGraphPin* OldBPPin = OldNode->FindPinChecked(TEXT("SaveGameBlueprint"));
-		ensure(0 == OldBPPin->LinkedTo.Num());
-		if (UBlueprint* SpawnBlueprint = Cast<UBlueprint>(OldBPPin->DefaultObject))
-		{
-			UEdGraphPin* NewClassPin = NewNode->FindPinChecked(TEXT("SaveGameClass"));
-			NewClassPin->DefaultObject = SpawnBlueprint->GeneratedClass;
-		}
-
-		OldNode->BreakAllNodeLinks();
-	}
-
-	static void Replace_SetAnimBlueprint(UK2Node_CallFunction* OldNode, FKismetCompilerContext& CompilerContext, UEdGraph* SourceGraph)
+	if (CompilerContext.bIsFullCompile)
 	{
 		const UEdGraphSchema_K2* Schema = CompilerContext.GetSchema();
-		check(Schema);
-
-		UK2Node_CallFunction* NewNode = CompilerContext.SpawnIntermediateNode<UK2Node_CallFunction>(OldNode, SourceGraph); 
-		check(NewNode);
-		const UFunction* NewFunction = USkeletalMeshComponent::StaticClass()->FindFunctionByName(TEXT("SetAnimClass"));
-		check(NewFunction);
-		NewNode->SetFromFunction(NewFunction);
-		NewNode->AllocateDefaultPins();
-
-		check(OldNode->GetThenPin() && NewNode->GetThenPin());
-		CompilerContext.CheckConnectionResponse(Schema->MovePinLinks(*OldNode->GetThenPin(), *NewNode->GetThenPin()), OldNode);
+		UFunction* Function = GetTargetFunction();
 
-		UEdGraphPin* OldExecPin = OldNode->FindPinChecked(Schema->PN_Execute);
-		UEdGraphPin* NewExecPin = NewNode->FindPinChecked(Schema->PN_Execute);
-		CompilerContext.CheckConnectionResponse(Schema->MovePinLinks(*OldExecPin, *NewExecPin), OldNode);
-
-		UEdGraphPin* OldSelfPin = OldNode->FindPinChecked(Schema->PN_Self);
-		UEdGraphPin* NewSelfPin = NewNode->FindPinChecked(Schema->PN_Self);
-		CompilerContext.CheckConnectionResponse(Schema->MovePinLinks(*OldSelfPin, *NewSelfPin), OldNode);
-
-		UEdGraphPin* OldBPPin = OldNode->FindPinChecked(TEXT("NewBlueprint"));
-		ensure(0 == OldBPPin->LinkedTo.Num());
-		if (UBlueprint* SpawnBlueprint = Cast<UBlueprint>(OldBPPin->DefaultObject))
+		// connect DefaultToSelf and WorldContext inside static functions to proper 'self'  
+		const bool bInsideBpFuncLibrary = CompilerContext.Blueprint && (BPTYPE_FunctionLibrary == CompilerContext.Blueprint->BlueprintType);
+		if (bInsideBpFuncLibrary && SourceGraph && Function)
 		{
-			UEdGraphPin* NewClassPin = NewNode->FindPinChecked(TEXT("NewClass"));
-			NewClassPin->DefaultObject = SpawnBlueprint->GeneratedClass;
-		}
-
-		OldNode->BreakAllNodeLinks();
-	}
-
-public:
-	static void ReplaceObsoleteFunctionsCall(UK2Node_CallFunction* CallFunctionNode, FKismetCompilerContext& CompilerContext, UEdGraph* SourceGraph)
-	{
-		check(CallFunctionNode && SourceGraph);
-		const FMemberReference& MemberReference = CallFunctionNode->FunctionReference;
-		static const FName CreateSaveGameObjectFromBlueprint("CreateSaveGameObjectFromBlueprint");
-		static const FName SetAnimBlueprint("SetAnimBlueprint");
+			TArray<UK2Node_FunctionEntry*> EntryPoints;
+			SourceGraph->GetNodesOfClass(EntryPoints);
+			if (1 != EntryPoints.Num())
+			{
+				CompilerContext.MessageLog.Warning(*FString::Printf(*LOCTEXT("WrongEntryPointsNum", "%i entry points found while expanding node @@").ToString(), EntryPoints.Num()), this);
+			}
+			else if (auto BetterSelfPin = EntryPoints[0]->GetWorldContextPin())
+			{
+				FString const DefaultToSelfMetaValue = Function->GetMetaData(FBlueprintMetadata::MD_DefaultToSelf);
+				FString const WorldContextMetaValue = Function->GetMetaData(FBlueprintMetadata::MD_WorldContext);
 
-		if ((MemberReference.GetMemberName() == CreateSaveGameObjectFromBlueprint) &&
-			(MemberReference.GetMemberParentClass((UClass *)NULL) == UGameplayStatics::StaticClass()))
-		{
-			Replace_CreateSaveGameObjectFromBlueprint(CallFunctionNode, CompilerContext, SourceGraph);
-		}
-		else if ((MemberReference.GetMemberName() == SetAnimBlueprint) &&
-			(MemberReference.GetMemberParentClass((UClass *)NULL) == USkeletalMeshComponent::StaticClass()))
-		{
-			Replace_SetAnimBlueprint(CallFunctionNode, CompilerContext, SourceGraph);
+				struct FStructConnectHelper
+				{
+					static void Connect(const FString& PinName, UK2Node* Node, UEdGraphPin* BetterSelf, const UEdGraphSchema_K2* Schema, FCompilerResultsLog& MessageLog)
+					{
+						auto Pin = Node->FindPin(PinName);
+						if (!PinName.IsEmpty() && Pin && !Pin->LinkedTo.Num())
+						{
+							const bool bConnected = Schema->TryCreateConnection(Pin, BetterSelf);
+							if (!bConnected)
+							{
+								MessageLog.Warning(*LOCTEXT("DefaultToSelfNotConnected", "DefaultToSelf pin @@ from node @@ cannot be connected to @@").ToString(), Pin, Node, BetterSelf);
+							}
+						}
+					}
+				};
+				FStructConnectHelper::Connect(DefaultToSelfMetaValue, this, BetterSelfPin, Schema, CompilerContext.MessageLog);
+				FStructConnectHelper::Connect(WorldContextMetaValue, this, BetterSelfPin, Schema, CompilerContext.MessageLog);
+			}
 		}
-	}
-};
-*/
-void UK2Node_CallFunction::ExpandNode(class FKismetCompilerContext& CompilerContext, UEdGraph* SourceGraph)
-{
-	if (CompilerContext.bIsFullCompile)
-	{
-		const UEdGraphSchema_K2* Schema = CompilerContext.GetSchema();
-
-		// >>> Backwards Compatibility:  VER_UE4_EDITORONLY_BLUEPRINTS
-		//FBackwardCompatibilityHelper::ReplaceObsoleteFunctionsCall(this, CompilerContext, SourceGraph);
-		// <<< End Backwards Compatibility
 
 		// If we have an enum param that is expanded, we handle that first
 		if(bWantsEnumToExecExpansion)
 		{
-			// Grab the function
-			UFunction* Function = GetTargetFunction();
 			if(Function)
 			{
 				// Get the metadata that identifies which param is the enum, and try and find it
diff --git a/Engine/Source/Editor/BlueprintGraph/Private/K2Node_FunctionEntry.cpp b/Engine/Source/Editor/BlueprintGraph/Private/K2Node_FunctionEntry.cpp
--- a/Engine/Source/Editor/BlueprintGraph/Private/K2Node_FunctionEntry.cpp
+++ b/Engine/Source/Editor/BlueprintGraph/Private/K2Node_FunctionEntry.cpp
@@ -93,6 +93,21 @@ public:
 	}
 };
 
+struct FFunctionEntryHelper
+{
+	static const FString& GetWorldContextPinName()
+	{
+		static const FString WorldContextPinName(TEXT("__WorldContext"));
+		return WorldContextPinName;
+	}
+
+	static bool RequireWorldContextParameter(const UK2Node_FunctionEntry* Node)
+	{
+		auto Blueprint = Node->GetBlueprint();
+		ensure(Blueprint);
+		return Blueprint && (BPTYPE_FunctionLibrary == Blueprint->BlueprintType);
+	}
+};
 
 UK2Node_FunctionEntry::UK2Node_FunctionEntry(const class FPostConstructInitializeProperties& PCIP)
 	: Super(PCIP)
@@ -130,6 +145,25 @@ void UK2Node_FunctionEntry::AllocateDefaultPins()
 	}
 
 	Super::AllocateDefaultPins();
+
+	if (FFunctionEntryHelper::RequireWorldContextParameter(this) 
+		&& ensure(!FindPin(FFunctionEntryHelper::GetWorldContextPinName())))
+	{
+		UEdGraphPin* WorldContextPin = CreatePin(
+			EGPD_Output,
+			K2Schema->PC_Object,
+			FString(),
+			UObject::StaticClass(),
+			false,
+			false,
+			FFunctionEntryHelper::GetWorldContextPinName());
+		WorldContextPin->bHidden = true;
+	}
+}
+
+UEdGraphPin* UK2Node_FunctionEntry::GetWorldContextPin() const
+{
+	return FFunctionEntryHelper::RequireWorldContextParameter(this) ? FindPin(FFunctionEntryHelper::GetWorldContextPinName()) : NULL;
 }
 
 void UK2Node_FunctionEntry::RemoveOutputPin(UEdGraphPin* PinToRemove)
diff --git a/Engine/Source/Editor/KismetCompiler/Private/KismetCompiler.cpp b/Engine/Source/Editor/KismetCompiler/Private/KismetCompiler.cpp
--- a/Engine/Source/Editor/KismetCompiler/Private/KismetCompiler.cpp
+++ b/Engine/Source/Editor/KismetCompiler/Private/KismetCompiler.cpp
@@ -1451,6 +1451,12 @@ void FKismetCompilerContext::FinishCompilingFunction(FKismetFunctionContext& Con
 		Function->SetMetaData(FBlueprintMetadata::MD_Tooltip, *EntryNode->MetaData.ToolTip);
 	}
 
+	if (auto WorldContextPin = EntryNode->GetWorldContextPin())
+	{
+		Function->SetMetaData(FBlueprintMetadata::MD_DefaultToSelf, *WorldContextPin->PinName); 
+		Function->SetMetaData(TEXT("HidePin"), *WorldContextPin->PinName);
+	}
+
 	for (int32 EntryPinIndex = 0; EntryPinIndex < EntryNode->Pins.Num(); ++EntryPinIndex)
 	{
 		UEdGraphPin* EntryPin = EntryNode->Pins[EntryPinIndex];
