Blueprint: Menu builder fix (Blueprint libraries).

diff --git a/Engine/Source/Editor/BlueprintGraph/Private/K2ActionMenuBuilder.cpp b/Engine/Source/Editor/BlueprintGraph/Private/K2ActionMenuBuilder.cpp
--- a/Engine/Source/Editor/BlueprintGraph/Private/K2ActionMenuBuilder.cpp
+++ b/Engine/Source/Editor/BlueprintGraph/Private/K2ActionMenuBuilder.cpp
@@ -1444,14 +1444,15 @@ void FK2ActionMenuBuilder::GetFuncNodesForClass(FGraphActionListBuilderBase& Lis
 	else
 	{
 		// Find all functions in each function library
-		for (TObjectIterator<UClass> ClassIt(RF_Transient | RF_PendingKill); ClassIt; ++ClassIt)
+		for (TObjectIterator<UClass> ClassIt; ClassIt; ++ClassIt)
 		{
 			UClass* TestClass = *ClassIt;
 			const bool bTransient = (GetTransientPackage() == TestClass->GetOutermost());
 			const UClass* AuthoritativeOwnerClass = OwnerClass ? const_cast<UClass*>(OwnerClass)->GetAuthoritativeClass() : NULL;
 			const bool bCanShowInOwnerClass = !AuthoritativeOwnerClass || !TestClass->IsChildOf(AuthoritativeOwnerClass);
 			const bool bIsProperFucntionLibrary = TestClass->IsChildOf(UBlueprintFunctionLibrary::StaticClass()) && CanUseLibraryFunctionsForScope(TestClass, OwnerClass);
-			if (bCanShowInOwnerClass && bIsProperFucntionLibrary && !bTransient)
+			const bool bIsSkeletonClass = TestClass->HasAnyFlags(RF_Transient) && TestClass->HasAnyClassFlags(CLASS_CompiledFromBlueprint);
+			if (bCanShowInOwnerClass && bIsProperFucntionLibrary && !bTransient && !bIsSkeletonClass)
 			{
 				GetFuncNodesForClass(ListBuilder, TestClass, OwnerClass, DestGraph, FunctionTypes, BaseCategory, TargetInfo);
 			}
