A comment added for GetWorldContextPin

diff --git a/Engine/Source/Editor/BlueprintGraph/Classes/K2Node_FunctionEntry.h b/Engine/Source/Editor/BlueprintGraph/Classes/K2Node_FunctionEntry.h
--- a/Engine/Source/Editor/BlueprintGraph/Classes/K2Node_FunctionEntry.h
+++ b/Engine/Source/Editor/BlueprintGraph/Classes/K2Node_FunctionEntry.h
@@ -46,6 +46,7 @@ class UK2Node_FunctionEntry : public UK2Node_FunctionTerminator
 	// Removes an output pin from the node
 	BLUEPRINTGRAPH_API void RemoveOutputPin(UEdGraphPin* PinToRemove);
 
+	// Returns pin for the automatically added WorldContext parameter (used only by BlueprintFunctionLibrary).
 	BLUEPRINTGRAPH_API UEdGraphPin* GetWorldContextPin() const;
 };
 
