Hide edit box on Bind/UnBind Delegate nodes

diff --git a/Engine/Source/Editor/BlueprintGraph/Private/K2Node_MCDelegate.cpp b/Engine/Source/Editor/BlueprintGraph/Private/K2Node_MCDelegate.cpp
--- a/Engine/Source/Editor/BlueprintGraph/Private/K2Node_MCDelegate.cpp
+++ b/Engine/Source/Editor/BlueprintGraph/Private/K2Node_MCDelegate.cpp
@@ -199,7 +199,7 @@ void UK2Node_AddDelegate::AllocateDefaultPins()
 
 	const UEdGraphSchema_K2* K2Schema = GetDefault<UEdGraphSchema_K2>();
 
-	if(UEdGraphPin* DelegatePin = CreatePin(EGPD_Input, K2Schema->PC_Delegate, TEXT(""), NULL, false, false, FK2Node_BaseMCDelegateHelper::DelegatePinName))
+	if(UEdGraphPin* DelegatePin = CreatePin(EGPD_Input, K2Schema->PC_Delegate, TEXT(""), NULL, false, true, FK2Node_BaseMCDelegateHelper::DelegatePinName, true))
 	{
 		FMemberReference::FillSimpleMemberReference<UFunction>(GetDelegateSignature(), DelegatePin->PinType.PinSubCategoryMemberReference);
 		DelegatePin->PinFriendlyName = NSLOCTEXT("K2Node", "PinFriendlyDelegatetName", "Event");
@@ -264,7 +264,7 @@ void UK2Node_RemoveDelegate::AllocateDefaultPins()
 
 	const UEdGraphSchema_K2* K2Schema = GetDefault<UEdGraphSchema_K2>();
 
-	if(UEdGraphPin* DelegatePin = CreatePin(EGPD_Input, K2Schema->PC_Delegate, TEXT(""), NULL, false, false, FK2Node_BaseMCDelegateHelper::DelegatePinName))
+	if(UEdGraphPin* DelegatePin = CreatePin(EGPD_Input, K2Schema->PC_Delegate, TEXT(""), NULL, false, true, FK2Node_BaseMCDelegateHelper::DelegatePinName, true))
 	{
 		FMemberReference::FillSimpleMemberReference<UFunction>(GetDelegateSignature(), DelegatePin->PinType.PinSubCategoryMemberReference);
 		DelegatePin->PinFriendlyName = NSLOCTEXT("K2Node", "PinFriendlyDelegatetName", "Event");
