Remove internal "SKEL_*" classes from the class picker.

- Modified FKismetEditorUtilities::CanCreateBlueprintFromClass() to return false if the input class is a BP skeleton class.
- Modified FKismetEditorUtilities::IsClassABlueprintSkeleton() to incorporate the alternative class flags check in case the SkeletonGeneratedClass has not yet been loaded during serialization.
- Replaced individual occurrences of the alternative class flags skeleton class check with FKismetEditorUtilities::IsClassABlueprintSkeleton() instead (for consistency and per discussion).

diff --git a/Engine/Source/Editor/BlueprintGraph/Private/K2ActionMenuBuilder.cpp b/Engine/Source/Editor/BlueprintGraph/Private/K2ActionMenuBuilder.cpp
--- a/Engine/Source/Editor/BlueprintGraph/Private/K2ActionMenuBuilder.cpp
+++ b/Engine/Source/Editor/BlueprintGraph/Private/K2ActionMenuBuilder.cpp
@@ -1033,7 +1033,7 @@ struct FClassDynamicCastHelper
 	static bool CanCastToClass(const UClass* TestClass, const UEdGraphSchema_K2* K2Schema)
 	{
 		check(TestClass && K2Schema);
-		const bool bIsSkeletonClass = TestClass->HasAnyFlags(RF_Transient) && TestClass->HasAnyClassFlags(CLASS_CompiledFromBlueprint);
+		const bool bIsSkeletonClass = FKismetEditorUtilities::IsClassABlueprintSkeleton(TestClass);
 		const bool bProperClass = TestClass->HasAnyFlags(RF_Public) && !TestClass->HasAnyClassFlags(CLASS_Deprecated | CLASS_NewerVersionExists);
 		const bool bAllowedByBPComms = K2Schema->bAllowBlueprintComms || !TestClass->ClassGeneratedBy;
 		return !bIsSkeletonClass && bProperClass && bAllowedByBPComms;
@@ -1225,7 +1225,7 @@ void FK2ActionMenuBuilder::GetPinAllowedNodeTypes(FBlueprintGraphActionListBuild
 					for (TObjectIterator<UClass> ClassIt; ClassIt; ++ClassIt)
 					{
 						UClass* TestClass = *ClassIt;
-						bool bIsSkeletonClass = TestClass->HasAnyFlags(RF_Transient) && TestClass->HasAnyClassFlags(CLASS_CompiledFromBlueprint);
+						bool bIsSkeletonClass = FKismetEditorUtilities::IsClassABlueprintSkeleton(TestClass);
 
 						if (TestClass->HasAnyFlags(RF_Public) && !bIsSkeletonClass && !TestClass->HasAnyClassFlags(CLASS_Deprecated|CLASS_NewerVersionExists) && (K2Schema->bAllowBlueprintComms || !TestClass->ClassGeneratedBy))
 						{
@@ -1451,7 +1451,7 @@ void FK2ActionMenuBuilder::GetFuncNodesForClass(FGraphActionListBuilderBase& Lis
 			const UClass* AuthoritativeOwnerClass = OwnerClass ? const_cast<UClass*>(OwnerClass)->GetAuthoritativeClass() : NULL;
 			const bool bCanShowInOwnerClass = !AuthoritativeOwnerClass || !TestClass->IsChildOf(AuthoritativeOwnerClass);
 			const bool bIsProperFucntionLibrary = TestClass->IsChildOf(UBlueprintFunctionLibrary::StaticClass()) && CanUseLibraryFunctionsForScope(TestClass, OwnerClass);
-			const bool bIsSkeletonClass = TestClass->HasAnyFlags(RF_Transient) && TestClass->HasAnyClassFlags(CLASS_CompiledFromBlueprint);
+			const bool bIsSkeletonClass = FKismetEditorUtilities::IsClassABlueprintSkeleton(TestClass);
 			if (bCanShowInOwnerClass && bIsProperFucntionLibrary && !bTransient && !bIsSkeletonClass)
 			{
 				GetFuncNodesForClass(ListBuilder, TestClass, OwnerClass, DestGraph, FunctionTypes, BaseCategory, TargetInfo);
@@ -2479,7 +2479,7 @@ void FK2ActionMenuBuilder::GetFuncNodesWithPinType(FBlueprintGraphActionListBuil
 		for (TObjectIterator<UClass> ClassIt; ClassIt; ++ClassIt)
 		{
 			UClass* TestClass = *ClassIt;
-			const bool bIsSkeletonClass = TestClass->HasAnyFlags(RF_Transient) && TestClass->HasAnyClassFlags(CLASS_CompiledFromBlueprint);
+			const bool bIsSkeletonClass = FKismetEditorUtilities::IsClassABlueprintSkeleton(TestClass);
 			if (TestClass->IsChildOf(UBlueprintFunctionLibrary::StaticClass()) &&
 				CanUseLibraryFunctionsForScope(OwnerClass, TestClass) && !bIsSkeletonClass)
 			{
diff --git a/Engine/Source/Editor/ContentBrowser/Private/SAssetView.cpp b/Engine/Source/Editor/ContentBrowser/Private/SAssetView.cpp
--- a/Engine/Source/Editor/ContentBrowser/Private/SAssetView.cpp
+++ b/Engine/Source/Editor/ContentBrowser/Private/SAssetView.cpp
@@ -11,6 +11,7 @@
 #include "FileHelpers.h"
 #include "ContentBrowserModule.h"
 #include "ObjectTools.h"
+#include "KismetEditorUtilities.h"
 
 #define LOCTEXT_NAMESPACE "ContentBrowser"
 
@@ -1333,7 +1334,7 @@ void SAssetView::RefreshSourceItems()
 		for (TObjectIterator<UClass> ClassIt; ClassIt; ++ClassIt)
 		{
 			// Don't offer skeleton classes
-			bool bIsSkeletonClass = ClassIt->HasAnyFlags(RF_Transient) && ClassIt->HasAnyClassFlags(CLASS_CompiledFromBlueprint);
+			bool bIsSkeletonClass = FKismetEditorUtilities::IsClassABlueprintSkeleton(*ClassIt);
 
 			if ( !ClassIt->HasAllClassFlags(CLASS_NotPlaceable) &&
 				!ClassIt->HasAnyClassFlags(CLASS_Abstract | CLASS_Deprecated | CLASS_NewerVersionExists) &&
diff --git a/Engine/Source/Editor/PlacementMode/Private/SPlacementModeTools.cpp b/Engine/Source/Editor/PlacementMode/Private/SPlacementModeTools.cpp
--- a/Engine/Source/Editor/PlacementMode/Private/SPlacementModeTools.cpp
+++ b/Engine/Source/Editor/PlacementMode/Private/SPlacementModeTools.cpp
@@ -7,6 +7,7 @@
 #include "SPlacementModeTools.h"
 #include "DragAndDrop/AssetDragDropOp.h"
 #include "IBspModeModule.h"
+#include "KismetEditorUtilities.h"
 
 #include "STutorialWrapper.h"
 
@@ -721,7 +722,7 @@ void SPlacementModeTools::RefreshVolumes()
 	for ( TObjectIterator<UClass> ClassIt; ClassIt; ++ClassIt )
 	{
 		// Don't offer skeleton classes
-		bool bIsSkeletonClass = ClassIt->HasAnyFlags( RF_Transient ) && ClassIt->HasAnyClassFlags( CLASS_CompiledFromBlueprint );
+		bool bIsSkeletonClass = FKismetEditorUtilities::IsClassABlueprintSkeleton(*ClassIt);
 
 		if ( !ClassIt->HasAllClassFlags( CLASS_NotPlaceable ) &&
 			 !ClassIt->HasAnyClassFlags( CLASS_Abstract | CLASS_Deprecated | CLASS_NewerVersionExists ) &&
@@ -774,7 +775,7 @@ void SPlacementModeTools::RebuildPlaceableClassWidgetCache()
 	for ( TObjectIterator<UClass> ClassIt; ClassIt; ++ClassIt )
 	{
 		// Don't offer skeleton classes
-		bool bIsSkeletonClass = ClassIt->HasAnyFlags(RF_Transient) && ClassIt->HasAnyClassFlags(CLASS_CompiledFromBlueprint);
+		bool bIsSkeletonClass = FKismetEditorUtilities::IsClassABlueprintSkeleton(*ClassIt);
 
 		if ( !ClassIt->HasAllClassFlags(CLASS_NotPlaceable) &&
 			!ClassIt->HasAnyClassFlags(CLASS_Abstract | CLASS_Deprecated | CLASS_NewerVersionExists) &&
diff --git a/Engine/Source/Editor/UnrealEd/Private/Commandlets/DumpHiddenCategoriesCommandlet.cpp b/Engine/Source/Editor/UnrealEd/Private/Commandlets/DumpHiddenCategoriesCommandlet.cpp
--- a/Engine/Source/Editor/UnrealEd/Private/Commandlets/DumpHiddenCategoriesCommandlet.cpp
+++ b/Engine/Source/Editor/UnrealEd/Private/Commandlets/DumpHiddenCategoriesCommandlet.cpp
@@ -3,6 +3,7 @@
 #include "UnrealEd.h"
 //#include "DumpHiddenCategoriesCommandlet.h"
 #include "ObjectEditorUtils.h"
+#include "KismetEditorUtilities.h"
 
 /*******************************************************************************
  * Static Helpers
@@ -108,7 +109,7 @@ static int32 GetHiddenFunctions(uint32 Indent, UClass* Class, bool bShowFunLibs,
 		{
 			UClass* TestClass = *ClassIt;
 			// if this is a skeleton class, don't bother
-			if (TestClass->HasAnyFlags(RF_Transient) && TestClass->HasAnyClassFlags(CLASS_CompiledFromBlueprint))
+			if (FKismetEditorUtilities::IsClassABlueprintSkeleton(TestClass))
 			{
 				continue;
 			}
diff --git a/Engine/Source/Editor/UnrealEd/Private/Kismet2/Kismet2.cpp b/Engine/Source/Editor/UnrealEd/Private/Kismet2/Kismet2.cpp
--- a/Engine/Source/Editor/UnrealEd/Private/Kismet2/Kismet2.cpp
+++ b/Engine/Source/Editor/UnrealEd/Private/Kismet2/Kismet2.cpp
@@ -482,7 +482,7 @@ bool FKismetEditorUtilities::CanCreateBlueprintOfClass(const UClass* Class)
 	const bool bCanCreateBlueprint =
 		!Class->HasAnyClassFlags(CLASS_Deprecated)
 		&& !Class->HasAnyClassFlags(CLASS_NewerVersionExists)
-		&& (!Class->ClassGeneratedBy || bAllowDerivedBlueprints);
+		&& (!Class->ClassGeneratedBy || (bAllowDerivedBlueprints && !IsClassABlueprintSkeleton(Class)));
 
 	const UEdGraphSchema_K2* K2Schema = GetDefault<UEdGraphSchema_K2>();
 	const bool bIsValidClass = Class->GetBoolMetaDataHierarchical(FBlueprintMetadata::MD_IsBlueprintBase) 
@@ -1298,11 +1298,11 @@ bool FKismetEditorUtilities::IsClassABlueprintSkeleton(const UClass* Class)
 {
 	// Find generating blueprint for a class
 	UBlueprint* GeneratingBP = Cast<UBlueprint>(Class->ClassGeneratedBy);
-	if( GeneratingBP )
+	if( GeneratingBP && GeneratingBP->SkeletonGeneratedClass )
 	{
 		return (Class == GeneratingBP->SkeletonGeneratedClass) && (GeneratingBP->SkeletonGeneratedClass != GeneratingBP->GeneratedClass);
 	}
-	return false;
+	return Class->HasAnyFlags(RF_Transient) && Class->HasAnyClassFlags(CLASS_CompiledFromBlueprint);
 }
 
 /** Run over the components references, and then NULL any that fall outside this blueprint's scope (e.g. components brought over after reparenting from another class, which are now in the transient package) */
