FKCHandler_CallDelegate uses usually the skel class.
Partial fix for ttp#334298CRASH: LIVE: BLUEPRINTS: Opening specific blueprints causes crash

diff --git a/Engine/Source/Editor/BlueprintGraph/Private/DelegateNodeHandlers.cpp b/Engine/Source/Editor/BlueprintGraph/Private/DelegateNodeHandlers.cpp
--- a/Engine/Source/Editor/BlueprintGraph/Private/DelegateNodeHandlers.cpp
+++ b/Engine/Source/Editor/BlueprintGraph/Private/DelegateNodeHandlers.cpp
@@ -58,7 +58,7 @@ struct FKCHandlerDelegateHelper
 
 		{
 			// MulticastDelegateProperty from NewClass may have empty signature, but property from skeletal class should have it.
-			const UFunction* OrgSignature = DelegateNode->GetDelegateSignature(true);
+			const UFunction* OrgSignature = DelegateNode->GetDelegateSignature();
 			if(const UEdGraphPin* DelegatePin = DelegateNode->GetDelegatePin())
 			{
 				const UFunction* PinSignature = FMemberReference::ResolveSimpleMemberReference<UFunction>(DelegatePin->PinType.PinSubCategoryMemberReference);
@@ -398,8 +398,11 @@ void FKCHandler_CallDelegate::Compile(FKismetFunctionContext& Context, UEdGraphN
 
 UFunction* FKCHandler_CallDelegate::FindFunction(FKismetFunctionContext& Context, UEdGraphNode* Node)
 {
-	UK2Node_CallDelegate* DelegateNode = CastChecked<UK2Node_CallDelegate>(Node);
-	return DelegateNode->GetDelegateSignature(true);
+	const UK2Node_CallDelegate* DelegateNode = CastChecked<UK2Node_CallDelegate>(Node);
+	const UClass* TestClass = Context.NewClass;
+	check(TestClass);
+	const bool bIsSkeletonClass = TestClass->HasAnyFlags(RF_Transient) && TestClass->HasAnyClassFlags(CLASS_CompiledFromBlueprint);
+	return DelegateNode->GetDelegateSignature(!bIsSkeletonClass);
 }
 
 void FKCHandler_CallDelegate::AdditionalCompiledStatementHandling(FKismetFunctionContext& Context, UEdGraphNode* Node, FBlueprintCompiledStatement& Statement)
