UK2Node::FindRemappedField - Remapped field cannot be owned by a structure, that is not parrent of current scope.

Problem description:
There is BPGC based on Pawn. One can add OnLand event. But FMemberReference::ResolveMember will return Character::OnLand because of existing redirector.
"+K2FieldRedirects=(OldFieldName="Pawn.OnLanded",NewFieldName="Character.OnLanded")"

diff --git a/Engine/Source/Editor/BlueprintGraph/Classes/K2Node.h b/Engine/Source/Editor/BlueprintGraph/Classes/K2Node.h
--- a/Engine/Source/Editor/BlueprintGraph/Classes/K2Node.h
+++ b/Engine/Source/Editor/BlueprintGraph/Classes/K2Node.h
@@ -240,9 +240,10 @@ class UK2Node : public UEdGraphNode
 	 *
 	 * @param	InitialScope	The scope the field was initially defined in.  The function will search up into parent scopes to attempt to find remappings
 	 * @param	InitialName		The name of the field to attempt to find a redirector for
+	 * @param	bInitialScopeMustBeOwnerOfField		if true the InitialScope must be Child of the field's owner
 	 * @return	The remapped field, if one exists
 	 */
-	BLUEPRINTGRAPH_API static UField* FindRemappedField(UClass* InitialScope, FName InitialName);
+	BLUEPRINTGRAPH_API static UField* FindRemappedField(UClass* InitialScope, FName InitialName, bool bInitialScopeMustBeOwnerOfField = false);
 
 	// Renames an existing pin on the node.
 	BLUEPRINTGRAPH_API virtual ERenamePinResult RenameUserDefinedPin(const FString& OldName, const FString& NewName, bool bTest = false);
@@ -494,7 +495,7 @@ public:
 		UClass* TargetScope = bSelfContext ? SelfScope : (UClass*)MemberParentClass;
 		if( TargetScope != NULL &&  !GIsSavingPackage )
 		{
-			ReturnField = Cast<TFieldType>(UK2Node::FindRemappedField(TargetScope, MemberName));
+			ReturnField = Cast<TFieldType>(UK2Node::FindRemappedField(TargetScope, MemberName, true));
 		}
 
 		if(ReturnField != NULL)
diff --git a/Engine/Source/Editor/BlueprintGraph/Private/K2Node.cpp b/Engine/Source/Editor/BlueprintGraph/Private/K2Node.cpp
--- a/Engine/Source/Editor/BlueprintGraph/Private/K2Node.cpp
+++ b/Engine/Source/Editor/BlueprintGraph/Private/K2Node.cpp
@@ -682,7 +682,7 @@ void UK2Node::InitFieldRedirectMap()
 	}
 }
 
-UField* UK2Node::FindRemappedField(UClass* InitialScope, FName InitialName)
+UField* UK2Node::FindRemappedField(UClass* InitialScope, FName InitialName, bool bInitialScopeMustBeOwnerOfField)
 {
 	FFieldRemapInfo NewFieldInfo;
 
@@ -711,8 +711,15 @@ UField* UK2Node::FindRemappedField(UClass* InitialScope, FName InitialName)
 		UField* NewField = FindField<UField>(SearchClass, NewFieldInfo.FieldName);
 		if( NewField != NULL )
 		{
-			UE_LOG(LogBlueprint, Log, TEXT("UK2Node:  Fixed up old field '%s' to new name '%s' on class '%s'."), *InitialName.ToString(), *NewFieldInfo.FieldName.ToString(), *SearchClass->GetName());
-			return NewField;
+			if (bInitialScopeMustBeOwnerOfField && !InitialScope->IsChildOf(SearchClass))
+			{
+				UE_LOG(LogBlueprint, Log, TEXT("UK2Node:  Unable to update field. Remapped field '%s' in not owned by given scope. Scope: '%s', Owner: '%s'."), *InitialName.ToString(), *InitialScope->GetName(), *NewFieldInfo.FieldClass.ToString());
+			}
+			else
+			{
+				UE_LOG(LogBlueprint, Log, TEXT("UK2Node:  Fixed up old field '%s' to new name '%s' on class '%s'."), *InitialName.ToString(), *NewFieldInfo.FieldName.ToString(), *SearchClass->GetName());
+				return NewField;
+			}
 		}
 		else if (SearchClass != NULL)
 		{
