UMG - Now filtering bindable function options.  Allowing the creation of a binding based on delegate signature.  New function for checking function signature compatiabilty that accepts a custom set of ignore flags.

diff --git a/Engine/Source/Editor/BlueprintGraph/Classes/EdGraphSchema_K2.h b/Engine/Source/Editor/BlueprintGraph/Classes/EdGraphSchema_K2.h
--- a/Engine/Source/Editor/BlueprintGraph/Classes/EdGraphSchema_K2.h
+++ b/Engine/Source/Editor/BlueprintGraph/Classes/EdGraphSchema_K2.h
@@ -662,7 +662,7 @@ public:
 	/** Can this function be overridden by kismet (either placed as event or new function graph created) */
 	static bool CanKismetOverrideFunction(const UFunction* Function);
 
-	/** returns friendly signiture name if possible or Removes any mangling to get the unmangled signiture name of the function */
+	/** returns friendly signiture name if possible or Removes any mangling to get the unmangled signature name of the function */
 	static FString GetFriendlySignitureName(const UFunction* Function);
 
 	static bool IsAllowableBlueprintVariableType(const class UEnum* InEnum);
@@ -723,6 +723,14 @@ public:
 	 */
 	virtual void CreateFunctionGraphTerminators(UEdGraph& Graph, UClass* Class) const;
 
+	/** 
+	 * Populate new function graph with entry and possibly return node
+	 * 
+	 * @param	Graph			Graph to add the function terminators to
+	 * @param	FunctionSignature	The function signature to mimic when creating the inputs and outputs for the function.
+	 */
+	virtual void CreateFunctionGraphTerminators(UEdGraph& Graph, UFunction* FunctionSignature) const;
+
 	/**
 	 * Converts a pin type into a fully qualified string (e.g., object'ObjectName').
 	 *
diff --git a/Engine/Source/Editor/BlueprintGraph/Private/EdGraphSchema_K2.cpp b/Engine/Source/Editor/BlueprintGraph/Private/EdGraphSchema_K2.cpp
--- a/Engine/Source/Editor/BlueprintGraph/Private/EdGraphSchema_K2.cpp
+++ b/Engine/Source/Editor/BlueprintGraph/Private/EdGraphSchema_K2.cpp
@@ -728,56 +728,56 @@ void UEdGraphSchema_K2::GetContextMenuActions(const UEdGraph* CurrentGraph, cons
 		MenuBuilder->BeginSection("EdGraphSchemaPinActions", LOCTEXT("PinActionsMenuHeader", "Pin Actions"));
 		{
 			if (!bIsDebugging)
-		    {
-			    // Break pin links
-			    if (InGraphPin->LinkedTo.Num() > 1)
-			    {
-				    MenuBuilder->AddMenuEntry( FGraphEditorCommands::Get().BreakPinLinks );
-			    }
-    
-			    // Add the change pin type action, if this is a select node
-			    if (InGraphNode->IsA(UK2Node_Select::StaticClass()))
-			    {
-				    MenuBuilder->AddMenuEntry(FGraphEditorCommands::Get().ChangePinType);
-			    }
-    
-			    // add sub menu for break link to
-			    if (InGraphPin->LinkedTo.Num() > 0)
-			    {
-				    if(InGraphPin->LinkedTo.Num() > 1)
-				    {
-					    MenuBuilder->AddSubMenu(
-						    LOCTEXT("BreakLinkTo", "Break Link To..."),
-						    LOCTEXT("BreakSpecificLinks", "Break a specific link..."),
-						    FNewMenuDelegate::CreateUObject( (UEdGraphSchema_K2*const)this, &UEdGraphSchema_K2::GetBreakLinkToSubMenuActions, const_cast<UEdGraphPin*>(InGraphPin)));
+			{
+				// Break pin links
+				if (InGraphPin->LinkedTo.Num() > 1)
+				{
+					MenuBuilder->AddMenuEntry( FGraphEditorCommands::Get().BreakPinLinks );
+				}
+	
+				// Add the change pin type action, if this is a select node
+				if (InGraphNode->IsA(UK2Node_Select::StaticClass()))
+				{
+					MenuBuilder->AddMenuEntry(FGraphEditorCommands::Get().ChangePinType);
+				}
+	
+				// add sub menu for break link to
+				if (InGraphPin->LinkedTo.Num() > 0)
+				{
+					if(InGraphPin->LinkedTo.Num() > 1)
+					{
+						MenuBuilder->AddSubMenu(
+							LOCTEXT("BreakLinkTo", "Break Link To..."),
+							LOCTEXT("BreakSpecificLinks", "Break a specific link..."),
+							FNewMenuDelegate::CreateUObject( (UEdGraphSchema_K2*const)this, &UEdGraphSchema_K2::GetBreakLinkToSubMenuActions, const_cast<UEdGraphPin*>(InGraphPin)));
 
 						MenuBuilder->AddSubMenu(
 							LOCTEXT("JumpToConnection", "Jump to Connection..."),
 							LOCTEXT("JumpToSpecificConnection", "Jump to specific connection..."),
 							FNewMenuDelegate::CreateUObject( (UEdGraphSchema_K2*const)this, &UEdGraphSchema_K2::GetJumpToConnectionSubMenuActions, const_cast<UEdGraphPin*>(InGraphPin)));
-				    }
-				    else
-				    {
-					    ((UEdGraphSchema_K2*const)this)->GetBreakLinkToSubMenuActions(*MenuBuilder, const_cast<UEdGraphPin*>(InGraphPin));
+					}
+					else
+					{
+						((UEdGraphSchema_K2*const)this)->GetBreakLinkToSubMenuActions(*MenuBuilder, const_cast<UEdGraphPin*>(InGraphPin));
 						((UEdGraphSchema_K2*const)this)->GetJumpToConnectionSubMenuActions(*MenuBuilder, const_cast<UEdGraphPin*>(InGraphPin));
 					}
-			    }
-    
-			    // Conditionally add the var promotion pin if this is an output pin and it's not an exec pin
-			    if (InGraphPin->PinType.PinCategory != PC_Exec)
-			    {
-				    MenuBuilder->AddMenuEntry( FGraphEditorCommands::Get().PromoteToVariable );
-			    }
-    
-			    // Conditionally add the execution path pin removal if this is an execution branching node
-			    if( InGraphPin->Direction == EGPD_Output && InGraphPin->GetOwningNode())
-			    {
-				    if ( InGraphPin->GetOwningNode()->IsA(UK2Node_ExecutionSequence::StaticClass()) ||  InGraphPin->GetOwningNode()->IsA(UK2Node_Switch::StaticClass()) )
-				    {
-					    MenuBuilder->AddMenuEntry( FGraphEditorCommands::Get().RemoveExecutionPin );
-				    }
-			    }
-		    }
+				}
+	
+				// Conditionally add the var promotion pin if this is an output pin and it's not an exec pin
+				if (InGraphPin->PinType.PinCategory != PC_Exec)
+				{
+					MenuBuilder->AddMenuEntry( FGraphEditorCommands::Get().PromoteToVariable );
+				}
+	
+				// Conditionally add the execution path pin removal if this is an execution branching node
+				if( InGraphPin->Direction == EGPD_Output && InGraphPin->GetOwningNode())
+				{
+					if ( InGraphPin->GetOwningNode()->IsA(UK2Node_ExecutionSequence::StaticClass()) ||  InGraphPin->GetOwningNode()->IsA(UK2Node_Switch::StaticClass()) )
+					{
+						MenuBuilder->AddMenuEntry( FGraphEditorCommands::Get().RemoveExecutionPin );
+					}
+				}
+			}
 		}
 		MenuBuilder->EndSection();
 
@@ -2172,9 +2172,54 @@ void UEdGraphSchema_K2::CreateFunctionGraphTerminators(UEdGraph& Graph, UClass*
 			UEdGraphPin* EntryNodeExec = FindExecutionPin(*EntryNode, EGPD_Output);
 			UEdGraphPin* ResultNodeExec = FindExecutionPin(*ReturnNode, EGPD_Input);
 			EntryNodeExec->MakeLinkTo(ResultNodeExec);
+		}
+	}
+}
+
+void UEdGraphSchema_K2::CreateFunctionGraphTerminators(UEdGraph& Graph, UFunction* FunctionSignature) const
+{
+	const FName GraphName = Graph.GetFName();
 
+	UBlueprint* Blueprint = FBlueprintEditorUtils::FindBlueprintForGraphChecked(&Graph);
+	check(Blueprint->BlueprintType != BPTYPE_MacroLibrary);
+
+	// Create a function entry node
+	FGraphNodeCreator<UK2Node_FunctionEntry> FunctionEntryCreator(Graph);
+	UK2Node_FunctionEntry* EntryNode = FunctionEntryCreator.CreateNode();
+	EntryNode->SignatureClass = NULL;
+	EntryNode->SignatureName = GraphName;
+	FunctionEntryCreator.Finalize();
+
+	EntryNode->CreatePinsForFunctionEntryExit(FunctionSignature, /*bIsFunctionEntry=*/ true);
+
+	// See if any function params are marked as out
+	bool bHasOutParam =  false;
+	for ( TFieldIterator<UProperty> It(FunctionSignature); It && ( It->PropertyFlags & CPF_Parm ); ++It )
+	{
+		if ( It->PropertyFlags & CPF_OutParm )
+		{
+			bHasOutParam = true;
+			break;
 		}
 	}
+
+	if ( bHasOutParam )
+	{
+		FGraphNodeCreator<UK2Node_FunctionResult> NodeCreator(Graph);
+		UK2Node_FunctionResult* ReturnNode = NodeCreator.CreateNode();
+		ReturnNode->SignatureClass = NULL;
+		ReturnNode->SignatureName = GraphName;
+		ReturnNode->NodePosX = EntryNode->NodePosX + EntryNode->NodeWidth + 256;
+		ReturnNode->NodePosY = EntryNode->NodePosY;
+		NodeCreator.Finalize();
+
+		ReturnNode->CreatePinsForFunctionEntryExit(FunctionSignature, /*bIsFunctionEntry=*/ false);
+
+		// Auto-connect the pins for entry and exit, so that by default the signature is properly generated
+		UEdGraphPin* EntryNodeExec = FindExecutionPin(*EntryNode, EGPD_Output);
+		UEdGraphPin* ResultNodeExec = FindExecutionPin(*ReturnNode, EGPD_Input);
+		EntryNodeExec->MakeLinkTo(ResultNodeExec);
+	}
 }
 
 bool UEdGraphSchema_K2::ConvertPropertyToPinType(const UProperty* Property, /*out*/ FEdGraphPinType& TypeOut) const
@@ -2417,9 +2462,9 @@ void UEdGraphSchema_K2::GetVariableTypeTree( TArray< TSharedPtr<FPinTypeTreeInfo
 	UScriptStruct* VectorStruct = FindObjectChecked<UScriptStruct>(UObject::StaticClass(), TEXT("Vector"));
 	UScriptStruct* RotatorStruct = FindObjectChecked<UScriptStruct>(UObject::StaticClass(), TEXT("Rotator"));
 	UScriptStruct* TransformStruct = FindObjectChecked<UScriptStruct>(UObject::StaticClass(), TEXT("Transform"));
- 	TypeTree.Add( MakeShareable( new FPinTypeTreeInfo(PC_Struct, VectorStruct, LOCTEXT("VectorType", "A 3D vector").ToString()) ) );
- 	TypeTree.Add( MakeShareable( new FPinTypeTreeInfo(PC_Struct, RotatorStruct, LOCTEXT("RotatorType", "A 3D rotation").ToString()) ) );
- 	TypeTree.Add( MakeShareable( new FPinTypeTreeInfo(PC_Struct, TransformStruct, LOCTEXT("TransformType", "A 3D transformation, including translation, rotation and 3D scale.").ToString()) ) );
+	TypeTree.Add( MakeShareable( new FPinTypeTreeInfo(PC_Struct, VectorStruct, LOCTEXT("VectorType", "A 3D vector").ToString()) ) );
+	TypeTree.Add( MakeShareable( new FPinTypeTreeInfo(PC_Struct, RotatorStruct, LOCTEXT("RotatorType", "A 3D rotation").ToString()) ) );
+	TypeTree.Add( MakeShareable( new FPinTypeTreeInfo(PC_Struct, TransformStruct, LOCTEXT("TransformType", "A 3D transformation, including translation, rotation and 3D scale.").ToString()) ) );
 	
 	// Add wildcard type
 	if (bAllowWildCard)
@@ -2737,7 +2782,7 @@ bool UEdGraphSchema_K2::CanEncapuslateNode(UEdGraphNode const& TestNode) const
 {
 	// Can't encapsulate entry points (may relax this restriction in the future, but it makes sense for now)
 	return !TestNode.IsA(UK2Node_FunctionTerminator::StaticClass()) && 
-		    TestNode.GetClass() != UK2Node_Tunnel::StaticClass(); //Tunnel nodes getting sucked into collapsed graphs fails badly, want to allow derived types though(composite node/Macroinstances)
+			TestNode.GetClass() != UK2Node_Tunnel::StaticClass(); //Tunnel nodes getting sucked into collapsed graphs fails badly, want to allow derived types though(composite node/Macroinstances)
 }
 
 void UEdGraphSchema_K2::HandleGraphBeingDeleted(UEdGraph& GraphBeingRemoved) const
diff --git a/Engine/Source/Editor/Kismet/Private/BlueprintDetailsCustomization.cpp b/Engine/Source/Editor/Kismet/Private/BlueprintDetailsCustomization.cpp
--- a/Engine/Source/Editor/Kismet/Private/BlueprintDetailsCustomization.cpp
+++ b/Engine/Source/Editor/Kismet/Private/BlueprintDetailsCustomization.cpp
@@ -1335,7 +1335,7 @@ void FBlueprintVarActionDetails::OnChangeReplication(TSharedPtr<FString> ItemSel
 				if (!FuncGraph)
 				{
 					FuncGraph = FBlueprintEditorUtils::CreateNewGraph(GetBlueprintObj(), FName(*NewFuncName), UEdGraph::StaticClass(), UEdGraphSchema_K2::StaticClass());
-					FBlueprintEditorUtils::AddFunctionGraph(GetBlueprintObj(), FuncGraph, false, NULL);
+					FBlueprintEditorUtils::AddFunctionGraph<UClass>(GetBlueprintObj(), FuncGraph, false, NULL);
 				}
 
 				if (FuncGraph)
diff --git a/Engine/Source/Editor/Kismet/Private/BlueprintEditor.cpp b/Engine/Source/Editor/Kismet/Private/BlueprintEditor.cpp
--- a/Engine/Source/Editor/Kismet/Private/BlueprintEditor.cpp
+++ b/Engine/Source/Editor/Kismet/Private/BlueprintEditor.cpp
@@ -3333,7 +3333,7 @@ bool FBlueprintEditor::CanCollapseSelectionToFunction(TSet<class UEdGraphNode*>&
 
 	// Create a function graph
 	UEdGraph* FunctionGraph = FBlueprintEditorUtils::CreateNewGraph(GetBlueprintObj(), FBlueprintEditorUtils::FindUniqueKismetName(GetBlueprintObj(), TEXT("TempGraph")), UEdGraph::StaticClass(), UEdGraphSchema_K2::StaticClass());
-	FBlueprintEditorUtils::AddFunctionGraph(GetBlueprintObj(), FunctionGraph, /*bIsUserCreated=*/ true, NULL);
+	FBlueprintEditorUtils::AddFunctionGraph<UClass>(GetBlueprintObj(), FunctionGraph, /*bIsUserCreated=*/ true, NULL);
 
 	const UEdGraphSchema_K2* K2Schema = GetDefault<UEdGraphSchema_K2>();
 
@@ -4984,7 +4984,7 @@ UEdGraph* FBlueprintEditor::CollapseSelectionToFunction(TSet<class UEdGraphNode*
 	FName DocumentName = FBlueprintEditorUtils::FindUniqueKismetName(GetBlueprintObj(), TEXT("NewFunction"));
 
 	NewGraph = FBlueprintEditorUtils::CreateNewGraph(GetBlueprintObj(), DocumentName, UEdGraph::StaticClass(), UEdGraphSchema_K2::StaticClass());
-	FBlueprintEditorUtils::AddFunctionGraph(GetBlueprintObj(), NewGraph, /*bIsUserCreated=*/ true, NULL);
+	FBlueprintEditorUtils::AddFunctionGraph<UClass>(GetBlueprintObj(), NewGraph, /*bIsUserCreated=*/ true, NULL);
 
 	TArray<UK2Node_FunctionEntry*> EntryNodes;
 	NewGraph->GetNodesOfClass(EntryNodes);
@@ -5385,7 +5385,7 @@ void FBlueprintEditor::OnAddNewDelegate()
 	NewGraph->bEditable = false;
 
 	K2Schema->CreateDefaultNodesForGraph(*NewGraph);
-	K2Schema->CreateFunctionGraphTerminators(*NewGraph, NULL);
+	K2Schema->CreateFunctionGraphTerminators(*NewGraph, (UClass*)NULL);
 	K2Schema->AddExtraFunctionFlags(NewGraph, (FUNC_BlueprintCallable|FUNC_BlueprintEvent|FUNC_Public));
 	K2Schema->MarkFunctionEntryAsEditable(NewGraph, true);
 
@@ -5449,7 +5449,7 @@ void FBlueprintEditor::NewDocument_OnClicked(ECreatedDocumentType GraphType)
 	if (GraphType == CGT_NewFunctionGraph)
 	{
 		NewGraph = FBlueprintEditorUtils::CreateNewGraph(GetBlueprintObj(), DocumentName, UEdGraph::StaticClass(), UEdGraphSchema_K2::StaticClass());
-		FBlueprintEditorUtils::AddFunctionGraph(GetBlueprintObj(), NewGraph, /*bIsUserCreated=*/ true, NULL);
+		FBlueprintEditorUtils::AddFunctionGraph<UClass>(GetBlueprintObj(), NewGraph, /*bIsUserCreated=*/ true, NULL);
 	}
 	else if (GraphType == CGT_NewMacroGraph)
 	{
diff --git a/Engine/Source/Editor/UMGEditor/Private/DetailCustomizations.cpp b/Engine/Source/Editor/UMGEditor/Private/DetailCustomizations.cpp
--- a/Engine/Source/Editor/UMGEditor/Private/DetailCustomizations.cpp
+++ b/Engine/Source/Editor/UMGEditor/Private/DetailCustomizations.cpp
@@ -63,7 +63,7 @@ void FBlueprintWidgetCustomization::CustomizeDetails( IDetailLayoutBuilder& Deta
 				PropertyRow.GetDefaultWidgets(NameWidget, ValueWidget, Row);
 
 				TSharedRef<SComboButton> DelegateButton = SNew(SComboButton)
-				.OnGetMenuContent(this, &FBlueprintWidgetCustomization::OnGenerateDelegateMenu, ConstPropertyHandle)
+				.OnGetMenuContent(this, &FBlueprintWidgetCustomization::OnGenerateDelegateMenu, ConstPropertyHandle, Property->SignatureFunction)
 				.ContentPadding(0)
 				.ButtonContent()
 				[
@@ -104,14 +104,15 @@ void FBlueprintWidgetCustomization::CustomizeDetails( IDetailLayoutBuilder& Deta
 	}
 }
 
-void FBlueprintWidgetCustomization::RefreshBlueprintFunctionCache()
+void FBlueprintWidgetCustomization::RefreshBlueprintFunctionCache(const UFunction* DelegateSignature)
 {
 	const UWidgetGraphSchema* Schema = GetDefault<UWidgetGraphSchema>();
 	const FSlateFontInfo DetailFontInfo = IDetailLayoutBuilder::GetDetailFont();
 
 	BlueprintFunctionCache.Reset();
 
-	//TODO UMG Filter the list of functions to only be functions matching the signature.
+	// Get the current skeleton class, think header for the blueprint.
+	UBlueprintGeneratedClass* SkeletonClass = Cast<UBlueprintGeneratedClass>(Blueprint->SkeletonGeneratedClass);
 
 	// Grab functions implemented by the blueprint
 	for ( int32 i = 0; i < Blueprint->FunctionGraphs.Num(); i++ )
@@ -125,37 +126,62 @@ void FBlueprintWidgetCustomization::RefreshBlueprintFunctionCache()
 			continue;
 		}
 
-		FGraphDisplayInfo DisplayInfo;
-		Graph->GetSchema()->GetGraphDisplayInformation(*Graph, DisplayInfo);
+		FName FunctionFName = Graph->GetFName();
 
-		TSharedPtr<FunctionInfo> NewFuncAction = MakeShareable(new FunctionInfo());
-		NewFuncAction->DisplayName = DisplayInfo.PlainName;
-		NewFuncAction->Tooltip = DisplayInfo.Tooltip;
-		NewFuncAction->FuncName = Graph->GetFName();
-		NewFuncAction->EdGraph = Graph;
+		UFunction* Function = SkeletonClass->FindFunctionByName(FunctionFName, EIncludeSuperFlag::IncludeSuper);
+		if ( Function == NULL )
+		{
+			continue;
+		}
+
+		// We ignore CPF_ReturnParm because all that matters for binding to script functions is that the number of out parameters match.
+		if ( Function->IsSignatureCompatibleWith(DelegateSignature, UFunction::GetDefaultIgnoredSignatureCompatibilityFlags() | CPF_ReturnParm) )
+		{
+			FGraphDisplayInfo DisplayInfo;
+			Graph->GetSchema()->GetGraphDisplayInformation(*Graph, DisplayInfo);
+
+			TSharedPtr<FunctionInfo> NewFuncAction = MakeShareable(new FunctionInfo());
+			NewFuncAction->DisplayName = DisplayInfo.PlainName;
+			NewFuncAction->Tooltip = DisplayInfo.Tooltip;
+			NewFuncAction->FuncName = FunctionFName;
+			NewFuncAction->EdGraph = Graph;
 
-		BlueprintFunctionCache.Add(NewFuncAction);
+			BlueprintFunctionCache.Add(NewFuncAction);
+		}
 	}
 }
 
-TSharedRef<SWidget> FBlueprintWidgetCustomization::OnGenerateDelegateMenu(TSharedRef<IPropertyHandle> PropertyHandle)
+TSharedRef<SWidget> FBlueprintWidgetCustomization::OnGenerateDelegateMenu(TSharedRef<IPropertyHandle> PropertyHandle, UFunction* DelegateSignature)
 {
-	RefreshBlueprintFunctionCache();
+	RefreshBlueprintFunctionCache(DelegateSignature);
 
 	const bool bInShouldCloseWindowAfterMenuSelection = true;
 	FMenuBuilder MenuBuilder(bInShouldCloseWindowAfterMenuSelection, NULL);
 
+	//TODO UMG Enable or disable Remove Binding if needed.
+
 	MenuBuilder.BeginSection("ClearBinding");
 	{
 		MenuBuilder.AddMenuEntry(
-			LOCTEXT("FilterListResetFilters", "Remove Binding"),
-			LOCTEXT("FilterListResetToolTip", "Removes the current binding"),
+			LOCTEXT("RemoveBinding", "Remove Binding"),
+			LOCTEXT("RemoveBindingToolTip", "Removes the current binding"),
 			FSlateIcon(),
 			FUIAction(FExecuteAction::CreateSP(this, &FBlueprintWidgetCustomization::HandleRemoveBinding, PropertyHandle))
 			);
 	}
 	MenuBuilder.EndSection(); //ClearBinding
 
+	MenuBuilder.BeginSection("CreateBinding");
+	{
+		MenuBuilder.AddMenuEntry(
+			LOCTEXT("CreateBinding", "Create Binding"),
+			LOCTEXT("CreateBindingToolTip", "Creates a new function for this property to be bound to"),
+			FSlateIcon(),
+			FUIAction(FExecuteAction::CreateSP(this, &FBlueprintWidgetCustomization::HandleCreateAndAddBinding, PropertyHandle, DelegateSignature))
+			);
+	}
+	MenuBuilder.EndSection(); //CreateBinding
+
 	MenuBuilder.BeginSection("Functions");
 	{
 		for ( TSharedPtr<FunctionInfo>& Function : BlueprintFunctionCache )
@@ -257,4 +283,31 @@ void FBlueprintWidgetCustomization::HandleAddBinding(TSharedRef<IPropertyHandle>
 	FBlueprintEditorUtils::MarkBlueprintAsStructurallyModified(Blueprint);
 }
 
+void FBlueprintWidgetCustomization::HandleCreateAndAddBinding(TSharedRef<IPropertyHandle> PropertyHandle, UFunction* DelegateSignature)
+{
+	const FScopedTransaction Transaction(LOCTEXT("CreateDelegate", "Create Binding"));
+
+	Blueprint->Modify();
+
+	// Create the function graph.
+	FString FunctionName = FString(TEXT("Get")) + PropertyHandle->GetProperty()->GetName();
+	UEdGraph* FunctionGraph = FBlueprintEditorUtils::CreateNewGraph(Blueprint, FBlueprintEditorUtils::FindUniqueKismetName(Blueprint, FunctionName), UEdGraph::StaticClass(), UEdGraphSchema_K2::StaticClass());
+
+	bool bUserCreated = true;
+	FBlueprintEditorUtils::AddFunctionGraph(Blueprint, FunctionGraph, bUserCreated, DelegateSignature);
+
+	//TODO UMG Once we start binding event functions, pure needs to be variable, so maybe pass it in.
+	
+	// All bound delegates should be pure.
+	const UEdGraphSchema_K2* Schema_K2 = Cast<UEdGraphSchema_K2>(FunctionGraph->GetSchema());
+	Schema_K2->AddExtraFunctionFlags(FunctionGraph, FUNC_BlueprintPure);
+
+	// Add the binding to the blueprint
+	TSharedPtr<FunctionInfo> SelectedFunction = MakeShareable(new FunctionInfo());
+	SelectedFunction->FuncName = FunctionGraph->GetFName();
+	SelectedFunction->EdGraph = FunctionGraph;
+
+	HandleAddBinding(PropertyHandle, SelectedFunction);
+}
+
 #undef LOCTEXT_NAMESPACE
diff --git a/Engine/Source/Editor/UMGEditor/Private/DetailCustomizations.h b/Engine/Source/Editor/UMGEditor/Private/DetailCustomizations.h
--- a/Engine/Source/Editor/UMGEditor/Private/DetailCustomizations.h
+++ b/Engine/Source/Editor/UMGEditor/Private/DetailCustomizations.h
@@ -29,13 +29,15 @@ private:
 		UEdGraph* EdGraph;
 	};
 
-	void RefreshBlueprintFunctionCache();
+	void RefreshBlueprintFunctionCache(const UFunction* DelegateSignature);
 
-	TSharedRef<SWidget> OnGenerateDelegateMenu(TSharedRef<IPropertyHandle> PropertyHandle);
+	TSharedRef<SWidget> OnGenerateDelegateMenu(TSharedRef<IPropertyHandle> PropertyHandle, UFunction* DelegateSignature);
 
 	FText GetCurrentBindingText(TSharedRef<IPropertyHandle> PropertyHandle) const;
+
 	void HandleRemoveBinding(TSharedRef<IPropertyHandle> PropertyHandle);
 	void HandleAddBinding(TSharedRef<IPropertyHandle> PropertyHandle, TSharedPtr<FunctionInfo> SelectedFunction);
+	void HandleCreateAndAddBinding(TSharedRef<IPropertyHandle> PropertyHandle, UFunction* DelegateSignature);
 
 private:
 
diff --git a/Engine/Source/Editor/UnrealEd/Private/Kismet2/BlueprintEditorUtils.cpp b/Engine/Source/Editor/UnrealEd/Private/Kismet2/BlueprintEditorUtils.cpp
--- a/Engine/Source/Editor/UnrealEd/Private/Kismet2/BlueprintEditorUtils.cpp
+++ b/Engine/Source/Editor/UnrealEd/Private/Kismet2/BlueprintEditorUtils.cpp
@@ -1541,40 +1541,6 @@ UEdGraph* FBlueprintEditorUtils::CreateNewGraph(UObject* ParentScope, const FNam
 	return NewGraph;
 }
 
-
-void FBlueprintEditorUtils::AddFunctionGraph(UBlueprint* Blueprint, class UEdGraph* Graph, bool bIsUserCreated, UClass* SignatureFromClass)
-{
-	// Give the schema a chance to fill out any required nodes (like the entry node or results node)
-	const UEdGraphSchema* Schema = Graph->GetSchema();
-	const UEdGraphSchema_K2* K2Schema = Cast<const UEdGraphSchema_K2>(Graph->GetSchema());
-
-	Schema->CreateDefaultNodesForGraph(*Graph);
-
-	if (K2Schema != NULL)
-	{
-		K2Schema->CreateFunctionGraphTerminators(*Graph, SignatureFromClass);
-
-		if (bIsUserCreated)
-		{
-			// We need to flag the entry node to make sure that the compiled function is callable from Kismet2
-			int32 ExtraFunctionFlags = (FUNC_BlueprintCallable | FUNC_BlueprintEvent | FUNC_Public);
-			if (BPTYPE_FunctionLibrary == Blueprint->BlueprintType)
-			{
-				ExtraFunctionFlags |= FUNC_Static;
-			}
-			K2Schema->AddExtraFunctionFlags(Graph, ExtraFunctionFlags);
-			K2Schema->MarkFunctionEntryAsEditable(Graph, true);
-		}
-	}
-
-	Blueprint->FunctionGraphs.Add(Graph);
-
-	// Potentially adjust variable names for any child blueprints
-	ValidateBlueprintChildVariables(Blueprint, Graph->GetFName());
-
-	FBlueprintEditorUtils::MarkBlueprintAsStructurallyModified(Blueprint);
-}
-
 UFunction* FBlueprintEditorUtils::FindFunctionInImplementedInterfaces(const UBlueprint* Blueprint, const FName& FunctionName, bool * bOutInvalidInterface )
 {
 	if(Blueprint)
@@ -2047,9 +2013,9 @@ void FBlueprintEditorUtils::GetDependentBlueprints(UBlueprint* Blueprint, TArray
 
 	//@TODO PRETEST MERGE:  Swap this when main is merged to pretest!
 #if 0
- 	TArray<UObject*> AllBlueprints;
- 	bool const bIncludeDerivedClasses = true;
- 	GetObjectsOfClass(UBlueprint::StaticClass(), AllBlueprints, bIncludeDerivedClasses );
+	TArray<UObject*> AllBlueprints;
+	bool const bIncludeDerivedClasses = true;
+	GetObjectsOfClass(UBlueprint::StaticClass(), AllBlueprints, bIncludeDerivedClasses );
 
 	for ( auto ObjIt = AllBlueprints.CreateConstIterator(); ObjIt; ++ObjIt )
 #else
diff --git a/Engine/Source/Editor/UnrealEd/Public/Kismet2/BlueprintEditorUtils.h b/Engine/Source/Editor/UnrealEd/Public/Kismet2/BlueprintEditorUtils.h
--- a/Engine/Source/Editor/UnrealEd/Public/Kismet2/BlueprintEditorUtils.h
+++ b/Engine/Source/Editor/UnrealEd/Public/Kismet2/BlueprintEditorUtils.h
@@ -149,8 +149,43 @@ public:
 	 */
 	static class UEdGraph* CreateNewGraph(UObject* ParentScope, const FName& GraphName, TSubclassOf<class UEdGraph> GraphClass, TSubclassOf<class UEdGraphSchema> SchemaClass);
 
-	/** Adds a function graph to this blueprint.  If bIsUserCreated is true, the entry/exit nodes will be editable. SignatureFromClass is used to find signature for entry/exit nodes if using an existing signature. */
-	static void AddFunctionGraph(UBlueprint* Blueprint, class UEdGraph* Graph,  bool bIsUserCreated, UClass* SignatureFromClass);
+	/** 
+	 * Adds a function graph to this blueprint.  If bIsUserCreated is true, the entry/exit nodes will be editable. SignatureFromObject is used to find signature for entry/exit nodes if using an existing signature.
+	 * The template argument SignatureType should be UClass or UFunction.
+	 */
+	template <typename SignatureType>
+	static void AddFunctionGraph(UBlueprint* Blueprint, class UEdGraph* Graph, bool bIsUserCreated, SignatureType* SignatureFromObject)
+	{
+		// Give the schema a chance to fill out any required nodes (like the entry node or results node)
+		const UEdGraphSchema* Schema = Graph->GetSchema();
+		const UEdGraphSchema_K2* K2Schema = Cast<const UEdGraphSchema_K2>(Graph->GetSchema());
+
+		Schema->CreateDefaultNodesForGraph(*Graph);
+
+		if ( K2Schema != NULL )
+		{
+			K2Schema->CreateFunctionGraphTerminators(*Graph, SignatureFromObject);
+
+			if ( bIsUserCreated )
+			{
+				// We need to flag the entry node to make sure that the compiled function is callable from Kismet2
+				int32 ExtraFunctionFlags = ( FUNC_BlueprintCallable | FUNC_BlueprintEvent | FUNC_Public );
+				if ( BPTYPE_FunctionLibrary == Blueprint->BlueprintType )
+				{
+					ExtraFunctionFlags |= FUNC_Static;
+				}
+				K2Schema->AddExtraFunctionFlags(Graph, ExtraFunctionFlags);
+				K2Schema->MarkFunctionEntryAsEditable(Graph, true);
+			}
+		}
+
+		Blueprint->FunctionGraphs.Add(Graph);
+
+		// Potentially adjust variable names for any child blueprints
+		ValidateBlueprintChildVariables(Blueprint, Graph->GetFName());
+
+		FBlueprintEditorUtils::MarkBlueprintAsStructurallyModified(Blueprint);
+	}
 
 	/** Adds a macro graph to this blueprint.  If bIsUserCreated is true, the entry/exit nodes will be editable. SignatureFromClass is used to find signature for entry/exit nodes if using an existing signature. */
 	static void AddMacroGraph(UBlueprint* Blueprint, class UEdGraph* Graph,  bool bIsUserCreated, UClass* SignatureFromClass);
diff --git a/Engine/Source/Runtime/CoreUObject/Private/UObject/Class.cpp b/Engine/Source/Runtime/CoreUObject/Private/UObject/Class.cpp
--- a/Engine/Source/Runtime/CoreUObject/Private/UObject/Class.cpp
+++ b/Engine/Source/Runtime/CoreUObject/Private/UObject/Class.cpp
@@ -3666,6 +3666,14 @@ void UFunction::Link(FArchive& Ar, bool bRelinkExistingProperties)
 }
 
 bool UFunction::IsSignatureCompatibleWith(const UFunction* OtherFunction) const
+{
+	// The script compiler pointlessly sets CPF_EditInline on components passed as arguments
+	const uint64 IgnoreFlags = UFunction::GetDefaultIgnoredSignatureCompatibilityFlags();
+
+	return IsSignatureCompatibleWith(OtherFunction, IgnoreFlags);
+}
+
+bool UFunction::IsSignatureCompatibleWith(const UFunction* OtherFunction, uint64 IgnoreFlags) const
 {
 	// Early out if they're exactly the same function
 	if (this == OtherFunction)
@@ -3673,9 +3681,6 @@ bool UFunction::IsSignatureCompatibleWith(const UFunction* OtherFunction) const
 		return true;
 	}
 
-	// The script compiler pointlessly sets CPF_EditInline on components passed as arguments
-	const uint64 IgnoreFlags = CPF_EditInline | CPF_ExportObject | CPF_InstancedReference | CPF_ContainsInstancedReference | CPF_ComputedFlags | CPF_ConstParm;//@TODO: UCREMOVAL: CPF_ConstParm added as a hack to get blueprints compiling with a const DamageType parameter.
-
 	// Run thru the parameter property chains to compare each property
 	TFieldIterator<UProperty> IteratorA(this);
 	TFieldIterator<UProperty> IteratorB(OtherFunction);
diff --git a/Engine/Source/Runtime/CoreUObject/Public/UObject/Class.h b/Engine/Source/Runtime/CoreUObject/Public/UObject/Class.h
--- a/Engine/Source/Runtime/CoreUObject/Public/UObject/Class.h
+++ b/Engine/Source/Runtime/CoreUObject/Public/UObject/Class.h
@@ -139,7 +139,7 @@ class COREUOBJECT_API UField : public UObject
 	
 	/**
 	 * Find the metadata value associated with the key
- 	 * and return bool  
+	 * and return bool  
 	 * @param Key The key to lookup in the metadata
 	 * @return return true if the value was true (case insensitive)
 	 */
@@ -1140,6 +1140,7 @@ public:
 	{
 		return (FunctionFlags&FlagsToCheck) != 0 || FlagsToCheck == FUNC_AllFlags;
 	}
+
 	/**
 	 * Used to safely check whether all of the passed in flags are set.
 	 *
@@ -1151,6 +1152,16 @@ public:
 		return ((FunctionFlags & FlagsToCheck) == FlagsToCheck);
 	}
 
+	/**
+	 * Returns the flags that are ignored by default when comparing function signatures.
+	 */
+	FORCEINLINE static uint64 GetDefaultIgnoredSignatureCompatibilityFlags()
+	{
+		//@TODO: UCREMOVAL: CPF_ConstParm added as a hack to get blueprints compiling with a const DamageType parameter.
+		const uint64 IgnoreFlags = CPF_EditInline | CPF_ExportObject | CPF_InstancedReference | CPF_ContainsInstancedReference | CPF_ComputedFlags | CPF_ConstParm;
+		return IgnoreFlags;
+	}
+
 	/**
 	 * Determines if two functions have an identical signature (note: currently doesn't allow
 	 * matches with class parameters that differ only in how derived they are; there is no
@@ -1161,6 +1172,18 @@ public:
 	 * @return	true if function signatures are compatible.
 	 */
 	COREUOBJECT_API bool IsSignatureCompatibleWith(const UFunction* OtherFunction) const;
+
+	/**
+	 * Determines if two functions have an identical signature (note: currently doesn't allow
+	 * matches with class parameters that differ only in how derived they are; there is no
+	 * directionality to the call)
+	 *
+	 * @param	OtherFunction	Function to compare this function against.
+	 * @param   IgnoreFlags     Custom flags to ignore when comparing parameters between the functions.
+	 *
+	 * @return	true if function signatures are compatible.
+	 */
+	COREUOBJECT_API bool IsSignatureCompatibleWith(const UFunction* OtherFunction, uint64 IgnoreFlags) const;
 };
 
 /*-----------------------------------------------------------------------------
@@ -1766,7 +1789,7 @@ public:
 	 * default if it doesn't exist yet.
 	 *
 	 * @return	name of the class specific ini file
- 	 */
+	 */
 	const FString GetConfigName() const;
 
 	UClass* GetSuperClass() const
