UMG - Allocating the default pin value for generated blueprint bindable functions.

diff --git a/Engine/Source/Editor/BlueprintGraph/Private/K2Node.cpp b/Engine/Source/Editor/BlueprintGraph/Private/K2Node.cpp
--- a/Engine/Source/Editor/BlueprintGraph/Private/K2Node.cpp
+++ b/Engine/Source/Editor/BlueprintGraph/Private/K2Node.cpp
@@ -63,6 +63,7 @@ bool UK2Node::CreatePinsForFunctionEntryExit(const UFunction* Function, bool bFo
 
 			UEdGraphPin* Pin = CreatePin(Direction, TEXT(""), TEXT(""), NULL, false, false, Param->GetName());
 			const bool bPinGood = K2Schema->ConvertPropertyToPinType(Param, /*out*/ Pin->PinType);
+			K2Schema->SetPinDefaultValueBasedOnType(Pin);
 
 			bAllPinsGood = bAllPinsGood && bPinGood;
 		}
