#ue4 Fix for bifurcation of variables in redirectors (TTP 334798)

diff --git a/Engine/Source/Editor/BlueprintGraph/Private/K2Node.cpp b/Engine/Source/Editor/BlueprintGraph/Private/K2Node.cpp
--- a/Engine/Source/Editor/BlueprintGraph/Private/K2Node.cpp
+++ b/Engine/Source/Editor/BlueprintGraph/Private/K2Node.cpp
@@ -703,6 +703,12 @@ UField* UK2Node::FindRemappedField(UClass* InitialScope, FName InitialName, bool
 		TestRemapClass = TestRemapClass->GetSuperClass();
 	}
 
+	// In the case of a bifurcation of a variable (e.g. moved from a parent into certain children), verify that we don't also define the variable in the current scope first
+	if( bFoundReplacement && (FindField<UField>(InitialScope, InitialName) != nullptr))
+	{
+		bFoundReplacement = false;		
+	}
+
 	if( bFoundReplacement )
 	{
 		const FName NewFieldName = NewFieldInfo.FieldName;
