ForEachLoop on Enum has the outputs in the wrong place. ttp335491

diff --git a/Engine/Source/Editor/BlueprintGraph/Private/K2Node_ForEachElementInEnum.cpp b/Engine/Source/Editor/BlueprintGraph/Private/K2Node_ForEachElementInEnum.cpp
--- a/Engine/Source/Editor/BlueprintGraph/Private/K2Node_ForEachElementInEnum.cpp
+++ b/Engine/Source/Editor/BlueprintGraph/Private/K2Node_ForEachElementInEnum.cpp
@@ -104,12 +104,16 @@ void UK2Node_ForEachElementInEnum::AllocateDefaultPins()
 	check(K2Schema);
 
 	CreatePin(EGPD_Input, K2Schema->PC_Exec, TEXT(""), NULL, false, false, K2Schema->PN_Execute);
-	CreatePin(EGPD_Output, K2Schema->PC_Exec, TEXT(""), NULL, false, false, K2Schema->PN_Then);
 
 	if (Enum)
 	{
-		CreatePin(EGPD_Output, K2Schema->PC_Byte, TEXT(""), Enum, false, false, EnumOuputPinName);
 		CreatePin(EGPD_Output, K2Schema->PC_Exec, TEXT(""), NULL, false, false, InsideLoopPinName);
+		CreatePin(EGPD_Output, K2Schema->PC_Byte, TEXT(""), Enum, false, false, EnumOuputPinName);
+	}
+
+	if (auto CompletedPin = CreatePin(EGPD_Output, K2Schema->PC_Exec, TEXT(""), NULL, false, false, K2Schema->PN_Then))
+	{
+		CompletedPin->PinFriendlyName = LOCTEXT("Completed", "Completed");
 	}
 }
 
