EDITOR: Crash on auto-converting SpawnAI to SpawnAIFromClass ttp334534

diff --git a/Engine/Source/Editor/BlueprintGraph/Private/EdGraphSchema_K2.cpp b/Engine/Source/Editor/BlueprintGraph/Private/EdGraphSchema_K2.cpp
--- a/Engine/Source/Editor/BlueprintGraph/Private/EdGraphSchema_K2.cpp
+++ b/Engine/Source/Editor/BlueprintGraph/Private/EdGraphSchema_K2.cpp
@@ -3702,24 +3702,33 @@ void UEdGraphSchema_K2::BackwardCompatibilityNodeConversion(UEdGraph* Graph, boo
 		}
 
 		{
-			TArray<UK2Node_CallFunction*> Nodes;
-			Graph->GetNodesOfClass(Nodes);
+			auto Blueprint = FBlueprintEditorUtils::FindBlueprintForGraph(Graph);
+			if (Blueprint && *Blueprint->SkeletonGeneratedClass)
 			{
-				static const FString BlueprintPinName(TEXT("Pawn"));
-				static const FString ClassPinName(TEXT("PawnClass"));
-				static const FName OldFuncName(GET_FUNCTION_NAME_CHECKED(UKismetAIHelperLibrary, SpawnAI));
-				static const FName NewFuncName(GET_FUNCTION_NAME_CHECKED(UKismetAIHelperLibrary, SpawnAIFromClass));
-				FBackwardCompatibilityConversionHelper::FFunctionCallParams Params(OldFuncName, NewFuncName, BlueprintPinName, ClassPinName, UKismetAIHelperLibrary::StaticClass());
-				FBackwardCompatibilityConversionHelper::ConvertFunctionCallNodes(Params, Nodes, Graph, *this, bOnlySafeChanges);
-			}
+				TArray<UK2Node_CallFunction*> Nodes;
+				Graph->GetNodesOfClass(Nodes);
+				{
+					static const FString BlueprintPinName(TEXT("Pawn"));
+					static const FString ClassPinName(TEXT("PawnClass"));
+					static const FName OldFuncName(GET_FUNCTION_NAME_CHECKED(UKismetAIHelperLibrary, SpawnAI));
+					static const FName NewFuncName(GET_FUNCTION_NAME_CHECKED(UKismetAIHelperLibrary, SpawnAIFromClass));
+					FBackwardCompatibilityConversionHelper::FFunctionCallParams Params(OldFuncName, NewFuncName, BlueprintPinName, ClassPinName, UKismetAIHelperLibrary::StaticClass());
+					FBackwardCompatibilityConversionHelper::ConvertFunctionCallNodes(Params, Nodes, Graph, *this, bOnlySafeChanges);
+				}
 
+				{
+					static const FString BlueprintPinName(TEXT("SaveGameBlueprint"));
+					static const FString ClassPinName(TEXT("SaveGameClass"));
+					static const FName OldFuncName(GET_FUNCTION_NAME_CHECKED(UGameplayStatics, CreateSaveGameObjectFromBlueprint));
+					static const FName NewFuncName(GET_FUNCTION_NAME_CHECKED(UGameplayStatics, CreateSaveGameObject));
+					FBackwardCompatibilityConversionHelper::FFunctionCallParams Params(OldFuncName, NewFuncName, BlueprintPinName, ClassPinName, UGameplayStatics::StaticClass());
+					FBackwardCompatibilityConversionHelper::ConvertFunctionCallNodes(Params, Nodes, Graph, *this, bOnlySafeChanges);
+				}
+			}
+			else
 			{
-				static const FString BlueprintPinName(TEXT("SaveGameBlueprint"));
-				static const FString ClassPinName(TEXT("SaveGameClass"));
-				static const FName OldFuncName(GET_FUNCTION_NAME_CHECKED(UGameplayStatics, CreateSaveGameObjectFromBlueprint));
-				static const FName NewFuncName(GET_FUNCTION_NAME_CHECKED(UGameplayStatics, CreateSaveGameObject));
-				FBackwardCompatibilityConversionHelper::FFunctionCallParams Params(OldFuncName, NewFuncName, BlueprintPinName, ClassPinName, UGameplayStatics::StaticClass());
-				FBackwardCompatibilityConversionHelper::ConvertFunctionCallNodes(Params, Nodes, Graph, *this, bOnlySafeChanges);
+				UE_LOG(LogBlueprint, Log, TEXT("BackwardCompatibilityNodeConversion: Blueprint '%s' cannot be fully converted. It has no skeleton class!"),
+					Blueprint ? *Blueprint->GetName() : TEXT("Unknown"));
 			}
 		}
 
