MC Delegate Nodes use 'self' pin smarter.

diff --git a/Engine/Source/Editor/BlueprintGraph/Private/K2Node_MCDelegate.cpp b/Engine/Source/Editor/BlueprintGraph/Private/K2Node_MCDelegate.cpp
--- a/Engine/Source/Editor/BlueprintGraph/Private/K2Node_MCDelegate.cpp
+++ b/Engine/Source/Editor/BlueprintGraph/Private/K2Node_MCDelegate.cpp
@@ -59,16 +59,20 @@ void UK2Node_BaseMCDelegate::AllocateDefaultPins()
 	CreatePin(EGPD_Input, K2Schema->PC_Exec, TEXT(""), NULL, false, false, K2Schema->PN_Execute);
 	CreatePin(EGPD_Output, K2Schema->PC_Exec, TEXT(""), NULL, false, false, K2Schema->PN_Then);
 
+	const auto Property = DelegateReference.ResolveMember<UMulticastDelegateProperty>(this);
+	const auto PropertyOwnerClass = Property->GetOwnerClass();
+	const auto Blueprint = GetBlueprint();
+	ensure(PropertyOwnerClass && Blueprint);
+	const bool bUseSelf = Blueprint && (PropertyOwnerClass == Blueprint->GeneratedClass || PropertyOwnerClass == Blueprint->SkeletonGeneratedClass);
+
 	UEdGraphPin* SelfPin = NULL;
-	if (DelegateReference.IsSelfContext())
+	if (bUseSelf)
 	{
 		SelfPin = CreatePin(EGPD_Input, K2Schema->PC_Object, K2Schema->PSC_Self, NULL, false, false, K2Schema->PN_Self);
 	}
 	else
 	{
-		// Allow redirects on the target node if necessary.
-		DelegateReference.ResolveMember<UMulticastDelegateProperty>((UClass*)NULL);
-		SelfPin = CreatePin(EGPD_Input, K2Schema->PC_Object, TEXT(""), DelegateReference.GetMemberParentClass(this), false, false, K2Schema->PN_Self);
+		SelfPin = CreatePin(EGPD_Input, K2Schema->PC_Object, TEXT(""), PropertyOwnerClass, false, false, K2Schema->PN_Self);
 	}
 
 	if(SelfPin)
