Editoronly BP: native blueprintcallable function redirect

diff --git a/Engine/Config/BaseEditor.ini b/Engine/Config/BaseEditor.ini
--- a/Engine/Config/BaseEditor.ini
+++ b/Engine/Config/BaseEditor.ini
@@ -83,6 +83,9 @@ MenuPriority=10
 
 [/Script/BlueprintGraph.EdGraphSchema_K2]
 bAllowBlueprintComms=true
++EditoronlyBPFunctionRedirects=(ClassName="KismetAIHelperLibrary", OldFunctionName="SpawnAI", NewFunctionName="SpawnAIFromClass", BlueprintParamName="Pawn", ClassParamName="PawnClass")
++EditoronlyBPFunctionRedirects=(ClassName="GameplayStatics", OldFunctionName="CreateSaveGameObjectFromBlueprint", NewFunctionName="CreateSaveGameObject", BlueprintParamName="SaveGameBlueprint", ClassParamName="SaveGameClass")
+
 
 [BlueprintEditor.Palette]
 bUseLegacyLayout=false
diff --git a/Engine/Source/Editor/BlueprintGraph/Classes/EdGraphSchema_K2.h b/Engine/Source/Editor/BlueprintGraph/Classes/EdGraphSchema_K2.h
--- a/Engine/Source/Editor/BlueprintGraph/Classes/EdGraphSchema_K2.h
+++ b/Engine/Source/Editor/BlueprintGraph/Classes/EdGraphSchema_K2.h
@@ -178,6 +178,29 @@ struct FFunctionTargetInfo
 	{}
 };
 
+USTRUCT()
+// Structure used to automatically convert blueprintcallable functions (that have blueprint parameter) calls (in bp graph) 
+// into their never versions (with class param instead of blueprint).
+struct FBlueprintCallableFunctionRedirect
+{
+	GENERATED_USTRUCT_BODY()
+
+	UPROPERTY()
+	FString ClassName;
+
+	UPROPERTY()
+	FString OldFunctionName;
+
+	UPROPERTY()
+	FString NewFunctionName;
+
+	UPROPERTY()
+	FString BlueprintParamName;
+
+	UPROPERTY()
+	FString ClassParamName;
+};
+
 UCLASS(config=Editor)
 class BLUEPRINTGRAPH_API UEdGraphSchema_K2 : public UEdGraphSchema
 {
@@ -329,6 +352,9 @@ class BLUEPRINTGRAPH_API UEdGraphSchema_K2 : public UEdGraphSchema
 	UPROPERTY(globalconfig)
 	bool bAllowBlueprintComms;
 	
+	UPROPERTY(globalconfig)
+	TArray<FBlueprintCallableFunctionRedirect> EditoronlyBPFunctionRedirects;
+
 public:
 	//////////////////////////////////////////////////////////////////////////
 	// FPinTypeInfo
diff --git a/Engine/Source/Editor/BlueprintGraph/Private/EdGraphSchema_K2.cpp b/Engine/Source/Editor/BlueprintGraph/Private/EdGraphSchema_K2.cpp
--- a/Engine/Source/Editor/BlueprintGraph/Private/EdGraphSchema_K2.cpp
+++ b/Engine/Source/Editor/BlueprintGraph/Private/EdGraphSchema_K2.cpp
@@ -3481,8 +3481,7 @@ struct FBackwardCompatibilityConversionHelper
 		}
 
 		const bool bNondefaultBPConnected = (OldBlueprintPin->LinkedTo.Num() > 0);
-		const bool bNodeUsesDefaultBlueprint = !bNondefaultBPConnected && OldBlueprintPin->DefaultObject;
-		const bool bTryConvert = bNodeUsesDefaultBlueprint || (!bOnlyWithDefaultBlueprint && bNondefaultBPConnected);
+		const bool bTryConvert = !bNondefaultBPConnected || !bOnlyWithDefaultBlueprint;
 		if (bTryConvert)
 		{
 			// CREATE NEW NODE
@@ -3518,18 +3517,13 @@ struct FBackwardCompatibilityConversionHelper
 			{
 				// DEFAULT VALUE
 				const auto UsedBlueprint = Cast<UBlueprint>(OldBlueprintPin->DefaultObject);
-				if (!UsedBlueprint)
+				ensure(!OldBlueprintPin->DefaultObject || UsedBlueprint);
+				ensure(!UsedBlueprint || *UsedBlueprint->GeneratedClass);
+				UClass* UsedClass = UsedBlueprint ? *UsedBlueprint->GeneratedClass : NULL;
+				Schema.TrySetDefaultObject(*ClassPin, UsedClass);
+				if (ClassPin->DefaultObject != UsedClass)
 				{
-					UE_LOG(LogBlueprint, Warning, TEXT("BackwardCompatibilityNodeConversion Error bp: '%s' node: '%s'. Wrong blueprint default value"),
-						Blueprint ? *Blueprint->GetName() : TEXT("Unknown"),
-						*OldNode->GetName());
-					return false;
-				}
-				ensure(NULL != *UsedBlueprint->GeneratedClass);
-				Schema.TrySetDefaultObject(*ClassPin, *UsedBlueprint->GeneratedClass);
-				if (ClassPin->DefaultObject != *UsedBlueprint->GeneratedClass)
-				{
-					auto ErrorStr = Schema.IsPinDefaultValid(ClassPin, FString(), *UsedBlueprint->GeneratedClass, FText());
+					auto ErrorStr = Schema.IsPinDefaultValid(ClassPin, FString(), UsedClass, FText());
 					UE_LOG(LogBlueprint, Warning, TEXT("BackwardCompatibilityNodeConversion Error 'cannot set class' in blueprint: %s node: '%s' actor bp: %s, reason: %s"),
 						Blueprint ? *Blueprint->GetName() : TEXT("Unknown"),
 						*OldNode->GetName(),
@@ -3663,10 +3657,22 @@ struct FBackwardCompatibilityConversionHelper
 			check(FuncScope);
 		}
 
+		FFunctionCallParams(const FBlueprintCallableFunctionRedirect& FunctionRedirect)
+			: OldFuncName(*FunctionRedirect.OldFunctionName)
+			, NewFuncName(*FunctionRedirect.NewFunctionName)
+			, BlueprintPinName(FunctionRedirect.BlueprintParamName)
+			, ClassPinName(FunctionRedirect.ClassParamName)
+			, FuncScope(NULL)
+		{
+			FuncScope = FindObject<UClass>(ANY_PACKAGE, *FunctionRedirect.ClassName);
+		}
+
 	};
 
 	static void ConvertFunctionCallNodes(const FFunctionCallParams& ConversionParams, TArray<UK2Node_CallFunction*>& Nodes, UEdGraph* Graph, const UEdGraphSchema_K2& Schema, bool bOnlyWithDefaultBlueprint)
 	{
+		if (ConversionParams.FuncScope)
+		{
 		const UFunction* OldFunc = ConversionParams.FuncScope->FindFunctionByName(ConversionParams.OldFuncName);
 		check(OldFunc);
 		const UFunction* NewFunc = ConversionParams.FuncScope->FindFunctionByName(ConversionParams.NewFuncName);
@@ -3683,6 +3689,7 @@ struct FBackwardCompatibilityConversionHelper
 			}
 		}
 	}
+	}
 };
 
 void UEdGraphSchema_K2::BackwardCompatibilityNodeConversion(UEdGraph* Graph, bool bOnlySafeChanges) const 
@@ -3708,22 +3715,11 @@ void UEdGraphSchema_K2::BackwardCompatibilityNodeConversion(UEdGraph* Graph, boo
 			{
 				TArray<UK2Node_CallFunction*> Nodes;
 				Graph->GetNodesOfClass(Nodes);
+				for (const auto& FunctionRedirect : EditoronlyBPFunctionRedirects)
 				{
-					static const FString BlueprintPinName(TEXT("Pawn"));
-					static const FString ClassPinName(TEXT("PawnClass"));
-					static const FName OldFuncName(GET_FUNCTION_NAME_CHECKED(UKismetAIHelperLibrary, SpawnAI));
-					static const FName NewFuncName(GET_FUNCTION_NAME_CHECKED(UKismetAIHelperLibrary, SpawnAIFromClass));
-					FBackwardCompatibilityConversionHelper::FFunctionCallParams Params(OldFuncName, NewFuncName, BlueprintPinName, ClassPinName, UKismetAIHelperLibrary::StaticClass());
-					FBackwardCompatibilityConversionHelper::ConvertFunctionCallNodes(Params, Nodes, Graph, *this, bOnlySafeChanges);
-				}
-
-				{
-					static const FString BlueprintPinName(TEXT("SaveGameBlueprint"));
-					static const FString ClassPinName(TEXT("SaveGameClass"));
-					static const FName OldFuncName(GET_FUNCTION_NAME_CHECKED(UGameplayStatics, CreateSaveGameObjectFromBlueprint));
-					static const FName NewFuncName(GET_FUNCTION_NAME_CHECKED(UGameplayStatics, CreateSaveGameObject));
-					FBackwardCompatibilityConversionHelper::FFunctionCallParams Params(OldFuncName, NewFuncName, BlueprintPinName, ClassPinName, UGameplayStatics::StaticClass());
-					FBackwardCompatibilityConversionHelper::ConvertFunctionCallNodes(Params, Nodes, Graph, *this, bOnlySafeChanges);
+					FBackwardCompatibilityConversionHelper::ConvertFunctionCallNodes(
+						FBackwardCompatibilityConversionHelper::FFunctionCallParams(FunctionRedirect),
+						Nodes, Graph, *this, bOnlySafeChanges);
 				}
 			}
 			else
