Remove redundant headers from Blueprint module. Removed #include 'BlueprintGraphClasses.h" from BlueprintGraphDefinitions.h. Added only headers included by modules other than BlueprintGraph. Fixed compilation issues for both unity and non-unity builds.

diff --git a/Engine/Source/Editor/BlueprintGraph/Classes/EdGraphSchema_K2_Actions.h b/Engine/Source/Editor/BlueprintGraph/Classes/EdGraphSchema_K2_Actions.h
--- a/Engine/Source/Editor/BlueprintGraph/Classes/EdGraphSchema_K2_Actions.h
+++ b/Engine/Source/Editor/BlueprintGraph/Classes/EdGraphSchema_K2_Actions.h
@@ -1,7 +1,7 @@
 // Copyright 1998-2014 Epic Games, Inc. All Rights Reserved.
 
 #pragma once
-
+#include "K2Node.h"
 #include "EdGraphSchema_K2_Actions.generated.h"
 
 /*******************************************************************************
diff --git a/Engine/Source/Editor/BlueprintGraph/Classes/K2Node_AIMoveTo.h b/Engine/Source/Editor/BlueprintGraph/Classes/K2Node_AIMoveTo.h
--- a/Engine/Source/Editor/BlueprintGraph/Classes/K2Node_AIMoveTo.h
+++ b/Engine/Source/Editor/BlueprintGraph/Classes/K2Node_AIMoveTo.h
@@ -2,6 +2,7 @@
 
 
 #pragma once
+#include "K2Node_BaseAsyncTask.h"
 #include "K2Node_AIMoveTo.generated.h"
 
 UCLASS()
diff --git a/Engine/Source/Editor/BlueprintGraph/Classes/K2Node_ActorBoundEvent.h b/Engine/Source/Editor/BlueprintGraph/Classes/K2Node_ActorBoundEvent.h
--- a/Engine/Source/Editor/BlueprintGraph/Classes/K2Node_ActorBoundEvent.h
+++ b/Engine/Source/Editor/BlueprintGraph/Classes/K2Node_ActorBoundEvent.h
@@ -2,6 +2,7 @@
 
 
 #pragma once
+#include "K2Node_Event.h"
 #include "K2Node_ActorBoundEvent.generated.h"
 
 UCLASS(MinimalAPI)
diff --git a/Engine/Source/Editor/BlueprintGraph/Classes/K2Node_AddComponent.h b/Engine/Source/Editor/BlueprintGraph/Classes/K2Node_AddComponent.h
--- a/Engine/Source/Editor/BlueprintGraph/Classes/K2Node_AddComponent.h
+++ b/Engine/Source/Editor/BlueprintGraph/Classes/K2Node_AddComponent.h
@@ -2,6 +2,7 @@
 
 
 #pragma once
+#include "K2Node_CallFunction.h"
 #include "K2Node_AddComponent.generated.h"
 
 UCLASS(MinimalAPI)
diff --git a/Engine/Source/Editor/BlueprintGraph/Classes/K2Node_AddDelegate.h b/Engine/Source/Editor/BlueprintGraph/Classes/K2Node_AddDelegate.h
--- a/Engine/Source/Editor/BlueprintGraph/Classes/K2Node_AddDelegate.h
+++ b/Engine/Source/Editor/BlueprintGraph/Classes/K2Node_AddDelegate.h
@@ -1,6 +1,6 @@
 // Copyright 1998-2014 Epic Games, Inc. All Rights Reserved.
 #pragma once
-
+#include "K2Node_BaseMCDelegate.h"
 #include "K2Node_AddDelegate.generated.h"
 
 UCLASS(MinimalAPI)
diff --git a/Engine/Source/Editor/BlueprintGraph/Classes/K2Node_AssignmentStatement.h b/Engine/Source/Editor/BlueprintGraph/Classes/K2Node_AssignmentStatement.h
--- a/Engine/Source/Editor/BlueprintGraph/Classes/K2Node_AssignmentStatement.h
+++ b/Engine/Source/Editor/BlueprintGraph/Classes/K2Node_AssignmentStatement.h
@@ -1,7 +1,7 @@
 // Copyright 1998-2014 Epic Games, Inc. All Rights Reserved.
 
 #pragma once
-
+#include "K2Node.h"
 #include "K2Node_AssignmentStatement.generated.h"
 
 UCLASS(MinimalAPI)
diff --git a/Engine/Source/Editor/BlueprintGraph/Classes/K2Node_BaseMCDelegate.h b/Engine/Source/Editor/BlueprintGraph/Classes/K2Node_BaseMCDelegate.h
--- a/Engine/Source/Editor/BlueprintGraph/Classes/K2Node_BaseMCDelegate.h
+++ b/Engine/Source/Editor/BlueprintGraph/Classes/K2Node_BaseMCDelegate.h
@@ -1,6 +1,6 @@
 // Copyright 1998-2014 Epic Games, Inc. All Rights Reserved.
 #pragma once
-
+#include "K2Node.h"
 #include "K2Node_BaseMCDelegate.generated.h"
 
 UCLASS(MinimalAPI, abstract)
diff --git a/Engine/Source/Editor/BlueprintGraph/Classes/K2Node_BreakStruct.h b/Engine/Source/Editor/BlueprintGraph/Classes/K2Node_BreakStruct.h
--- a/Engine/Source/Editor/BlueprintGraph/Classes/K2Node_BreakStruct.h
+++ b/Engine/Source/Editor/BlueprintGraph/Classes/K2Node_BreakStruct.h
@@ -1,7 +1,7 @@
 // Copyright 1998-2014 Epic Games, Inc. All Rights Reserved.
 
 #pragma once
-
+#include "K2Node_StructMemberGet.h"
 #include "K2Node_BreakStruct.generated.h"
 
 UCLASS(MinimalAPI)
diff --git a/Engine/Source/Editor/BlueprintGraph/Classes/K2Node_CallArrayFunction.h b/Engine/Source/Editor/BlueprintGraph/Classes/K2Node_CallArrayFunction.h
--- a/Engine/Source/Editor/BlueprintGraph/Classes/K2Node_CallArrayFunction.h
+++ b/Engine/Source/Editor/BlueprintGraph/Classes/K2Node_CallArrayFunction.h
@@ -2,6 +2,7 @@
 
 
 #pragma once
+#include "K2Node_CallFunction.h"
 #include "K2Node_CallArrayFunction.generated.h"
 
 struct FArrayPropertyPinCombo
diff --git a/Engine/Source/Editor/BlueprintGraph/Classes/K2Node_CallFunction.h b/Engine/Source/Editor/BlueprintGraph/Classes/K2Node_CallFunction.h
--- a/Engine/Source/Editor/BlueprintGraph/Classes/K2Node_CallFunction.h
+++ b/Engine/Source/Editor/BlueprintGraph/Classes/K2Node_CallFunction.h
@@ -2,6 +2,7 @@
 
 #pragma once
 
+#include "K2Node.h"
 #include "K2Node_CallFunction.generated.h"
 
 UCLASS()
diff --git a/Engine/Source/Editor/BlueprintGraph/Classes/K2Node_CallFunctionOnMember.h b/Engine/Source/Editor/BlueprintGraph/Classes/K2Node_CallFunctionOnMember.h
--- a/Engine/Source/Editor/BlueprintGraph/Classes/K2Node_CallFunctionOnMember.h
+++ b/Engine/Source/Editor/BlueprintGraph/Classes/K2Node_CallFunctionOnMember.h
@@ -2,6 +2,7 @@
 
 
 #pragma once
+#include "K2Node_CallFunction.h"
 #include "K2Node_CallFunctionOnMember.generated.h"
 
 
diff --git a/Engine/Source/Editor/BlueprintGraph/Classes/K2Node_CallMaterialParameterCollectionFunction.h b/Engine/Source/Editor/BlueprintGraph/Classes/K2Node_CallMaterialParameterCollectionFunction.h
--- a/Engine/Source/Editor/BlueprintGraph/Classes/K2Node_CallMaterialParameterCollectionFunction.h
+++ b/Engine/Source/Editor/BlueprintGraph/Classes/K2Node_CallMaterialParameterCollectionFunction.h
@@ -1,6 +1,7 @@
 // Copyright 1998-2014 Epic Games, Inc. All Rights Reserved.
 
 #pragma once
+#include "K2Node_CallFunction.h"
 #include "K2Node_CallMaterialParameterCollectionFunction.generated.h"
 
 UCLASS(MinimalAPI)
diff --git a/Engine/Source/Editor/BlueprintGraph/Classes/K2Node_CallParentFunction.h b/Engine/Source/Editor/BlueprintGraph/Classes/K2Node_CallParentFunction.h
--- a/Engine/Source/Editor/BlueprintGraph/Classes/K2Node_CallParentFunction.h
+++ b/Engine/Source/Editor/BlueprintGraph/Classes/K2Node_CallParentFunction.h
@@ -2,6 +2,7 @@
 
 
 #pragma once
+#include "K2Node_CallFunction.h"
 #include "K2Node_CallParentFunction.generated.h"
 
 UCLASS(MinimalAPI)
diff --git a/Engine/Source/Editor/BlueprintGraph/Classes/K2Node_CastByteToEnum.h b/Engine/Source/Editor/BlueprintGraph/Classes/K2Node_CastByteToEnum.h
--- a/Engine/Source/Editor/BlueprintGraph/Classes/K2Node_CastByteToEnum.h
+++ b/Engine/Source/Editor/BlueprintGraph/Classes/K2Node_CastByteToEnum.h
@@ -1,6 +1,7 @@
 // Copyright 1998-2014 Epic Games, Inc. All Rights Reserved.
 
 #pragma once
+#include "K2Node.h"
 #include "K2Node_CastByteToEnum.generated.h"
 
 UCLASS(MinimalAPI)
diff --git a/Engine/Source/Editor/BlueprintGraph/Classes/K2Node_ClassDynamicCast.h b/Engine/Source/Editor/BlueprintGraph/Classes/K2Node_ClassDynamicCast.h
--- a/Engine/Source/Editor/BlueprintGraph/Classes/K2Node_ClassDynamicCast.h
+++ b/Engine/Source/Editor/BlueprintGraph/Classes/K2Node_ClassDynamicCast.h
@@ -2,6 +2,7 @@
 
 
 #pragma once
+#include "K2Node_DynamicCast.h"
 #include "K2Node_ClassDynamicCast.generated.h"
 
 UCLASS(MinimalAPI)
diff --git a/Engine/Source/Editor/BlueprintGraph/Classes/K2Node_CommutativeAssociativeBinaryOperator.h b/Engine/Source/Editor/BlueprintGraph/Classes/K2Node_CommutativeAssociativeBinaryOperator.h
--- a/Engine/Source/Editor/BlueprintGraph/Classes/K2Node_CommutativeAssociativeBinaryOperator.h
+++ b/Engine/Source/Editor/BlueprintGraph/Classes/K2Node_CommutativeAssociativeBinaryOperator.h
@@ -1,7 +1,7 @@
 // Copyright 1998-2014 Epic Games, Inc. All Rights Reserved.
 
 #pragma once
-
+#include "K2Node_CallFunction.h"
 #include "K2Node_CommutativeAssociativeBinaryOperator.generated.h"
 
 UCLASS(MinimalAPI)
diff --git a/Engine/Source/Editor/BlueprintGraph/Classes/K2Node_ComponentBoundEvent.h b/Engine/Source/Editor/BlueprintGraph/Classes/K2Node_ComponentBoundEvent.h
--- a/Engine/Source/Editor/BlueprintGraph/Classes/K2Node_ComponentBoundEvent.h
+++ b/Engine/Source/Editor/BlueprintGraph/Classes/K2Node_ComponentBoundEvent.h
@@ -2,6 +2,7 @@
 
 
 #pragma once
+#include "K2Node_Event.h"
 #include "K2Node_ComponentBoundEvent.generated.h"
 
 UCLASS(MinimalAPI)
diff --git a/Engine/Source/Editor/BlueprintGraph/Classes/K2Node_Composite.h b/Engine/Source/Editor/BlueprintGraph/Classes/K2Node_Composite.h
--- a/Engine/Source/Editor/BlueprintGraph/Classes/K2Node_Composite.h
+++ b/Engine/Source/Editor/BlueprintGraph/Classes/K2Node_Composite.h
@@ -1,6 +1,6 @@
 // Copyright 1998-2014 Epic Games, Inc. All Rights Reserved.
 #pragma once
-
+#include "K2Node_Tunnel.h"
 #include "K2Node_Composite.generated.h"
 
 UCLASS(MinimalAPI)
diff --git a/Engine/Source/Editor/BlueprintGraph/Classes/K2Node_CreateDelegate.h b/Engine/Source/Editor/BlueprintGraph/Classes/K2Node_CreateDelegate.h
--- a/Engine/Source/Editor/BlueprintGraph/Classes/K2Node_CreateDelegate.h
+++ b/Engine/Source/Editor/BlueprintGraph/Classes/K2Node_CreateDelegate.h
@@ -1,6 +1,6 @@
 // Copyright 1998-2014 Epic Games, Inc. All Rights Reserved.
 #pragma once
-
+#include "K2Node.h"
 #include "K2Node_CreateDelegate.generated.h"
 
 UCLASS(MinimalAPI)
diff --git a/Engine/Source/Editor/BlueprintGraph/Classes/K2Node_CustomEvent.h b/Engine/Source/Editor/BlueprintGraph/Classes/K2Node_CustomEvent.h
--- a/Engine/Source/Editor/BlueprintGraph/Classes/K2Node_CustomEvent.h
+++ b/Engine/Source/Editor/BlueprintGraph/Classes/K2Node_CustomEvent.h
@@ -2,6 +2,7 @@
 
 
 #pragma once
+#include "K2Node_Event.h"
 #include "K2Node_CustomEvent.generated.h"
 
 UCLASS(MinimalAPI)
diff --git a/Engine/Source/Editor/BlueprintGraph/Classes/K2Node_DeadClass.h b/Engine/Source/Editor/BlueprintGraph/Classes/K2Node_DeadClass.h
--- a/Engine/Source/Editor/BlueprintGraph/Classes/K2Node_DeadClass.h
+++ b/Engine/Source/Editor/BlueprintGraph/Classes/K2Node_DeadClass.h
@@ -1,6 +1,7 @@
 // Copyright 1998-2014 Epic Games, Inc. All Rights Reserved.
 
 #pragma once
+#include "K2Node.h"
 #include "K2Node_DeadClass.generated.h"
 
 UCLASS(MinimalAPI)
diff --git a/Engine/Source/Editor/BlueprintGraph/Classes/K2Node_DelegateSet.h b/Engine/Source/Editor/BlueprintGraph/Classes/K2Node_DelegateSet.h
--- a/Engine/Source/Editor/BlueprintGraph/Classes/K2Node_DelegateSet.h
+++ b/Engine/Source/Editor/BlueprintGraph/Classes/K2Node_DelegateSet.h
@@ -2,6 +2,7 @@
 
 
 #pragma once
+#include "K2Node.h"
 #include "K2Node_DelegateSet.generated.h"
 
 UCLASS(MinimalAPI)
diff --git a/Engine/Source/Editor/BlueprintGraph/Classes/K2Node_DynamicCast.h b/Engine/Source/Editor/BlueprintGraph/Classes/K2Node_DynamicCast.h
--- a/Engine/Source/Editor/BlueprintGraph/Classes/K2Node_DynamicCast.h
+++ b/Engine/Source/Editor/BlueprintGraph/Classes/K2Node_DynamicCast.h
@@ -2,6 +2,7 @@
 
 
 #pragma once
+#include "K2Node.h"
 #include "K2Node_DynamicCast.generated.h"
 
 UCLASS(MinimalAPI)
diff --git a/Engine/Source/Editor/BlueprintGraph/Classes/K2Node_EditablePinBase.h b/Engine/Source/Editor/BlueprintGraph/Classes/K2Node_EditablePinBase.h
--- a/Engine/Source/Editor/BlueprintGraph/Classes/K2Node_EditablePinBase.h
+++ b/Engine/Source/Editor/BlueprintGraph/Classes/K2Node_EditablePinBase.h
@@ -1,7 +1,7 @@
 // Copyright 1998-2014 Epic Games, Inc. All Rights Reserved.
 
 #pragma once
-
+#include "K2Node.h"
 #include "K2Node_EditablePinBase.generated.h"
 
 USTRUCT()
diff --git a/Engine/Source/Editor/BlueprintGraph/Classes/K2Node_EnumEquality.h b/Engine/Source/Editor/BlueprintGraph/Classes/K2Node_EnumEquality.h
--- a/Engine/Source/Editor/BlueprintGraph/Classes/K2Node_EnumEquality.h
+++ b/Engine/Source/Editor/BlueprintGraph/Classes/K2Node_EnumEquality.h
@@ -2,6 +2,7 @@
 
 
 #pragma once
+#include "K2Node.h"
 #include "K2Node_EnumEquality.generated.h"
 
 UCLASS(MinimalAPI)
diff --git a/Engine/Source/Editor/BlueprintGraph/Classes/K2Node_EnumLiteral.h b/Engine/Source/Editor/BlueprintGraph/Classes/K2Node_EnumLiteral.h
--- a/Engine/Source/Editor/BlueprintGraph/Classes/K2Node_EnumLiteral.h
+++ b/Engine/Source/Editor/BlueprintGraph/Classes/K2Node_EnumLiteral.h
@@ -1,6 +1,8 @@
 // Copyright 1998-2014 Epic Games, Inc. All Rights Reserved.
 
 #pragma once
+#include "K2Node.h"
+#include "NodeDependingOnEnumInterface.h"
 #include "K2Node_EnumLiteral.generated.h"
 
 UCLASS(MinimalAPI)
diff --git a/Engine/Source/Editor/BlueprintGraph/Classes/K2Node_Event.h b/Engine/Source/Editor/BlueprintGraph/Classes/K2Node_Event.h
--- a/Engine/Source/Editor/BlueprintGraph/Classes/K2Node_Event.h
+++ b/Engine/Source/Editor/BlueprintGraph/Classes/K2Node_Event.h
@@ -2,6 +2,7 @@
 
 
 #pragma once
+#include "K2Node_EditablePinBase.h"
 #include "K2Node_Event.generated.h"
 
 UCLASS(MinimalAPI)
diff --git a/Engine/Source/Editor/BlueprintGraph/Classes/K2Node_ExecutionSequence.h b/Engine/Source/Editor/BlueprintGraph/Classes/K2Node_ExecutionSequence.h
--- a/Engine/Source/Editor/BlueprintGraph/Classes/K2Node_ExecutionSequence.h
+++ b/Engine/Source/Editor/BlueprintGraph/Classes/K2Node_ExecutionSequence.h
@@ -2,6 +2,7 @@
 
 
 #pragma once
+#include "K2Node.h"
 #include "K2Node_ExecutionSequence.generated.h"
 
 UCLASS(MinimalAPI)
diff --git a/Engine/Source/Editor/BlueprintGraph/Classes/K2Node_ForEachElementInEnum.h b/Engine/Source/Editor/BlueprintGraph/Classes/K2Node_ForEachElementInEnum.h
--- a/Engine/Source/Editor/BlueprintGraph/Classes/K2Node_ForEachElementInEnum.h
+++ b/Engine/Source/Editor/BlueprintGraph/Classes/K2Node_ForEachElementInEnum.h
@@ -1,6 +1,7 @@
 // Copyright 1998-2014 Epic Games, Inc. All Rights Reserved.
 
 #pragma once
+#include "K2Node.h"
 #include "K2Node_ForEachElementInEnum.generated.h"
 
 UCLASS(MinimalAPI)
diff --git a/Engine/Source/Editor/BlueprintGraph/Classes/K2Node_FunctionEntry.h b/Engine/Source/Editor/BlueprintGraph/Classes/K2Node_FunctionEntry.h
--- a/Engine/Source/Editor/BlueprintGraph/Classes/K2Node_FunctionEntry.h
+++ b/Engine/Source/Editor/BlueprintGraph/Classes/K2Node_FunctionEntry.h
@@ -1,7 +1,7 @@
 // Copyright 1998-2014 Epic Games, Inc. All Rights Reserved.
 
-
 #pragma once
+#include "K2Node_FunctionTerminator.h"
 #include "K2Node_FunctionEntry.generated.h"
 
 UCLASS(MinimalAPI)
diff --git a/Engine/Source/Editor/BlueprintGraph/Classes/K2Node_FunctionResult.h b/Engine/Source/Editor/BlueprintGraph/Classes/K2Node_FunctionResult.h
--- a/Engine/Source/Editor/BlueprintGraph/Classes/K2Node_FunctionResult.h
+++ b/Engine/Source/Editor/BlueprintGraph/Classes/K2Node_FunctionResult.h
@@ -2,6 +2,7 @@
 
 
 #pragma once
+#include "K2Node_FunctionTerminator.h"
 #include "K2Node_FunctionResult.generated.h"
 
 UCLASS(MinimalAPI)
diff --git a/Engine/Source/Editor/BlueprintGraph/Classes/K2Node_FunctionTerminator.h b/Engine/Source/Editor/BlueprintGraph/Classes/K2Node_FunctionTerminator.h
--- a/Engine/Source/Editor/BlueprintGraph/Classes/K2Node_FunctionTerminator.h
+++ b/Engine/Source/Editor/BlueprintGraph/Classes/K2Node_FunctionTerminator.h
@@ -2,6 +2,7 @@
 
 
 #pragma once
+#include "K2Node_EditablePinBase.h"
 #include "K2Node_FunctionTerminator.generated.h"
 
 UCLASS(abstract, MinimalAPI)
diff --git a/Engine/Source/Editor/BlueprintGraph/Classes/K2Node_GetEnumeratorName.h b/Engine/Source/Editor/BlueprintGraph/Classes/K2Node_GetEnumeratorName.h
--- a/Engine/Source/Editor/BlueprintGraph/Classes/K2Node_GetEnumeratorName.h
+++ b/Engine/Source/Editor/BlueprintGraph/Classes/K2Node_GetEnumeratorName.h
@@ -1,6 +1,7 @@
 // Copyright 1998-2014 Epic Games, Inc. All Rights Reserved.
 
 #pragma once
+#include "K2Node.h"
 #include "K2Node_GetEnumeratorName.generated.h"
 
 UCLASS(MinimalAPI)
diff --git a/Engine/Source/Editor/BlueprintGraph/Classes/K2Node_GetInputAxisKeyValue.h b/Engine/Source/Editor/BlueprintGraph/Classes/K2Node_GetInputAxisKeyValue.h
--- a/Engine/Source/Editor/BlueprintGraph/Classes/K2Node_GetInputAxisKeyValue.h
+++ b/Engine/Source/Editor/BlueprintGraph/Classes/K2Node_GetInputAxisKeyValue.h
@@ -1,6 +1,7 @@
 // Copyright 1998-2014 Epic Games, Inc. All Rights Reserved.
 
 #pragma once
+#include "K2Node_CallFunction.h"
 #include "K2Node_GetInputAxisKeyValue.generated.h"
 
 UCLASS(MinimalAPI)
diff --git a/Engine/Source/Editor/BlueprintGraph/Classes/K2Node_GetInputAxisValue.h b/Engine/Source/Editor/BlueprintGraph/Classes/K2Node_GetInputAxisValue.h
--- a/Engine/Source/Editor/BlueprintGraph/Classes/K2Node_GetInputAxisValue.h
+++ b/Engine/Source/Editor/BlueprintGraph/Classes/K2Node_GetInputAxisValue.h
@@ -1,6 +1,7 @@
 // Copyright 1998-2014 Epic Games, Inc. All Rights Reserved.
 
 #pragma once
+#include "K2Node_CallFunction.h"
 #include "K2Node_GetInputAxisValue.generated.h"
 
 UCLASS(MinimalAPI)
diff --git a/Engine/Source/Editor/BlueprintGraph/Classes/K2Node_GetInputVectorAxisValue.h b/Engine/Source/Editor/BlueprintGraph/Classes/K2Node_GetInputVectorAxisValue.h
--- a/Engine/Source/Editor/BlueprintGraph/Classes/K2Node_GetInputVectorAxisValue.h
+++ b/Engine/Source/Editor/BlueprintGraph/Classes/K2Node_GetInputVectorAxisValue.h
@@ -1,6 +1,7 @@
 // Copyright 1998-2014 Epic Games, Inc. All Rights Reserved.
 
 #pragma once
+#include "K2Node_GetInputAxisKeyValue.h"
 #include "K2Node_GetInputVectorAxisValue.generated.h"
 
 UCLASS(MinimalAPI)
diff --git a/Engine/Source/Editor/BlueprintGraph/Classes/K2Node_GetNumEnumEntries.h b/Engine/Source/Editor/BlueprintGraph/Classes/K2Node_GetNumEnumEntries.h
--- a/Engine/Source/Editor/BlueprintGraph/Classes/K2Node_GetNumEnumEntries.h
+++ b/Engine/Source/Editor/BlueprintGraph/Classes/K2Node_GetNumEnumEntries.h
@@ -1,6 +1,7 @@
 // Copyright 1998-2014 Epic Games, Inc. All Rights Reserved.
 
 #pragma once
+#include "NodeDependingOnEnumInterface.h"
 #include "K2Node_GetNumEnumEntries.generated.h"
 
 UCLASS(MinimalAPI)
diff --git a/Engine/Source/Editor/BlueprintGraph/Classes/K2Node_IfThenElse.h b/Engine/Source/Editor/BlueprintGraph/Classes/K2Node_IfThenElse.h
--- a/Engine/Source/Editor/BlueprintGraph/Classes/K2Node_IfThenElse.h
+++ b/Engine/Source/Editor/BlueprintGraph/Classes/K2Node_IfThenElse.h
@@ -2,6 +2,7 @@
 
 
 #pragma once
+#include "K2Node.h"
 #include "K2Node_IfThenElse.generated.h"
 
 UCLASS(MinimalAPI)
diff --git a/Engine/Source/Editor/BlueprintGraph/Classes/K2Node_InputAction.h b/Engine/Source/Editor/BlueprintGraph/Classes/K2Node_InputAction.h
--- a/Engine/Source/Editor/BlueprintGraph/Classes/K2Node_InputAction.h
+++ b/Engine/Source/Editor/BlueprintGraph/Classes/K2Node_InputAction.h
@@ -2,6 +2,7 @@
 
 
 #pragma once
+#include "K2Node.h"
 #include "K2Node_InputAction.generated.h"
 
 UCLASS(MinimalAPI)
diff --git a/Engine/Source/Editor/BlueprintGraph/Classes/K2Node_InputActionEvent.h b/Engine/Source/Editor/BlueprintGraph/Classes/K2Node_InputActionEvent.h
--- a/Engine/Source/Editor/BlueprintGraph/Classes/K2Node_InputActionEvent.h
+++ b/Engine/Source/Editor/BlueprintGraph/Classes/K2Node_InputActionEvent.h
@@ -2,6 +2,7 @@
 
 
 #pragma once
+#include "K2Node_Event.h"
 #include "K2Node_InputActionEvent.generated.h"
 
 UCLASS(MinimalAPI)
diff --git a/Engine/Source/Editor/BlueprintGraph/Classes/K2Node_InputAxisEvent.h b/Engine/Source/Editor/BlueprintGraph/Classes/K2Node_InputAxisEvent.h
--- a/Engine/Source/Editor/BlueprintGraph/Classes/K2Node_InputAxisEvent.h
+++ b/Engine/Source/Editor/BlueprintGraph/Classes/K2Node_InputAxisEvent.h
@@ -1,6 +1,7 @@
 // Copyright 1998-2014 Epic Games, Inc. All Rights Reserved.
 
 #pragma once
+#include "K2Node_Event.h"
 #include "K2Node_InputAxisEvent.generated.h"
 
 UCLASS(MinimalAPI)
diff --git a/Engine/Source/Editor/BlueprintGraph/Classes/K2Node_InputAxisKeyEvent.h b/Engine/Source/Editor/BlueprintGraph/Classes/K2Node_InputAxisKeyEvent.h
--- a/Engine/Source/Editor/BlueprintGraph/Classes/K2Node_InputAxisKeyEvent.h
+++ b/Engine/Source/Editor/BlueprintGraph/Classes/K2Node_InputAxisKeyEvent.h
@@ -1,6 +1,7 @@
 // Copyright 1998-2014 Epic Games, Inc. All Rights Reserved.
 
 #pragma once
+#include "K2Node_Event.h"
 #include "K2Node_InputAxisKeyEvent.generated.h"
 
 UCLASS(MinimalAPI)
diff --git a/Engine/Source/Editor/BlueprintGraph/Classes/K2Node_InputKey.h b/Engine/Source/Editor/BlueprintGraph/Classes/K2Node_InputKey.h
--- a/Engine/Source/Editor/BlueprintGraph/Classes/K2Node_InputKey.h
+++ b/Engine/Source/Editor/BlueprintGraph/Classes/K2Node_InputKey.h
@@ -2,6 +2,7 @@
 
 
 #pragma once
+#include "K2Node.h"
 #include "K2Node_InputKey.generated.h"
 
 UCLASS(MinimalAPI)
diff --git a/Engine/Source/Editor/BlueprintGraph/Classes/K2Node_InputKeyEvent.h b/Engine/Source/Editor/BlueprintGraph/Classes/K2Node_InputKeyEvent.h
--- a/Engine/Source/Editor/BlueprintGraph/Classes/K2Node_InputKeyEvent.h
+++ b/Engine/Source/Editor/BlueprintGraph/Classes/K2Node_InputKeyEvent.h
@@ -2,6 +2,7 @@
 
 
 #pragma once
+#include "K2Node_Event.h"
 #include "K2Node_InputKeyEvent.generated.h"
 
 UCLASS(MinimalAPI)
diff --git a/Engine/Source/Editor/BlueprintGraph/Classes/K2Node_InputTouch.h b/Engine/Source/Editor/BlueprintGraph/Classes/K2Node_InputTouch.h
--- a/Engine/Source/Editor/BlueprintGraph/Classes/K2Node_InputTouch.h
+++ b/Engine/Source/Editor/BlueprintGraph/Classes/K2Node_InputTouch.h
@@ -2,6 +2,7 @@
 
 
 #pragma once
+#include "K2Node.h"
 #include "K2Node_InputTouch.generated.h"
 
 UCLASS(MinimalAPI)
diff --git a/Engine/Source/Editor/BlueprintGraph/Classes/K2Node_InputTouchEvent.h b/Engine/Source/Editor/BlueprintGraph/Classes/K2Node_InputTouchEvent.h
--- a/Engine/Source/Editor/BlueprintGraph/Classes/K2Node_InputTouchEvent.h
+++ b/Engine/Source/Editor/BlueprintGraph/Classes/K2Node_InputTouchEvent.h
@@ -2,6 +2,7 @@
 
 
 #pragma once
+#include "K2Node_Event.h"
 #include "K2Node_InputTouchEvent.generated.h"
 
 UCLASS(MinimalAPI)
diff --git a/Engine/Source/Editor/BlueprintGraph/Classes/K2Node_InputVectorAxisEvent.h b/Engine/Source/Editor/BlueprintGraph/Classes/K2Node_InputVectorAxisEvent.h
--- a/Engine/Source/Editor/BlueprintGraph/Classes/K2Node_InputVectorAxisEvent.h
+++ b/Engine/Source/Editor/BlueprintGraph/Classes/K2Node_InputVectorAxisEvent.h
@@ -1,6 +1,7 @@
 // Copyright 1998-2014 Epic Games, Inc. All Rights Reserved.
 
 #pragma once
+#include "K2Node_InputAxisKeyEvent.h"
 #include "K2Node_InputVectorAxisEvent.generated.h"
 
 UCLASS(MinimalAPI)
diff --git a/Engine/Source/Editor/BlueprintGraph/Classes/K2Node_Knot.h b/Engine/Source/Editor/BlueprintGraph/Classes/K2Node_Knot.h
--- a/Engine/Source/Editor/BlueprintGraph/Classes/K2Node_Knot.h
+++ b/Engine/Source/Editor/BlueprintGraph/Classes/K2Node_Knot.h
@@ -1,7 +1,7 @@
 // Copyright 1998-2014 Epic Games, Inc. All Rights Reserved.
 
 #pragma once
-
+#include "K2Node.h"
 #include "K2Node_Knot.generated.h"
 
 UCLASS(MinimalAPI)
diff --git a/Engine/Source/Editor/BlueprintGraph/Classes/K2Node_Literal.h b/Engine/Source/Editor/BlueprintGraph/Classes/K2Node_Literal.h
--- a/Engine/Source/Editor/BlueprintGraph/Classes/K2Node_Literal.h
+++ b/Engine/Source/Editor/BlueprintGraph/Classes/K2Node_Literal.h
@@ -2,6 +2,7 @@
 
 
 #pragma once
+#include "K2Node.h"
 #include "K2Node_Literal.generated.h"
 
 UCLASS(MinimalAPI)
diff --git a/Engine/Source/Editor/BlueprintGraph/Classes/K2Node_LocalVariable.h b/Engine/Source/Editor/BlueprintGraph/Classes/K2Node_LocalVariable.h
--- a/Engine/Source/Editor/BlueprintGraph/Classes/K2Node_LocalVariable.h
+++ b/Engine/Source/Editor/BlueprintGraph/Classes/K2Node_LocalVariable.h
@@ -1,7 +1,7 @@
 // Copyright 1998-2014 Epic Games, Inc. All Rights Reserved.
 
 #pragma once
-
+#include "K2Node_TemporaryVariable.h"
 #include "K2Node_LocalVariable.generated.h"
 
 UCLASS(MinimalAPI, deprecated)
diff --git a/Engine/Source/Editor/BlueprintGraph/Classes/K2Node_MacroInstance.h b/Engine/Source/Editor/BlueprintGraph/Classes/K2Node_MacroInstance.h
--- a/Engine/Source/Editor/BlueprintGraph/Classes/K2Node_MacroInstance.h
+++ b/Engine/Source/Editor/BlueprintGraph/Classes/K2Node_MacroInstance.h
@@ -2,6 +2,7 @@
 
 
 #pragma once
+#include "K2Node_Tunnel.h"
 #include "K2Node_MacroInstance.generated.h"
 
 UCLASS(MinimalAPI)
diff --git a/Engine/Source/Editor/BlueprintGraph/Classes/K2Node_MakeArray.h b/Engine/Source/Editor/BlueprintGraph/Classes/K2Node_MakeArray.h
--- a/Engine/Source/Editor/BlueprintGraph/Classes/K2Node_MakeArray.h
+++ b/Engine/Source/Editor/BlueprintGraph/Classes/K2Node_MakeArray.h
@@ -2,6 +2,7 @@
 
 
 #pragma once
+#include "K2Node.h"
 #include "K2Node_MakeArray.generated.h"
 
 UCLASS(MinimalAPI)
diff --git a/Engine/Source/Editor/BlueprintGraph/Classes/K2Node_MakeStruct.h b/Engine/Source/Editor/BlueprintGraph/Classes/K2Node_MakeStruct.h
--- a/Engine/Source/Editor/BlueprintGraph/Classes/K2Node_MakeStruct.h
+++ b/Engine/Source/Editor/BlueprintGraph/Classes/K2Node_MakeStruct.h
@@ -1,6 +1,7 @@
 // Copyright 1998-2014 Epic Games, Inc. All Rights Reserved.
 
 #pragma once
+#include "K2Node_StructMemberSet.h"
 #include "K2Node_MakeStruct.generated.h"
 
 // Pure kismet node that creates a struct with specified values for each member
diff --git a/Engine/Source/Editor/BlueprintGraph/Classes/K2Node_MatineeController.h b/Engine/Source/Editor/BlueprintGraph/Classes/K2Node_MatineeController.h
--- a/Engine/Source/Editor/BlueprintGraph/Classes/K2Node_MatineeController.h
+++ b/Engine/Source/Editor/BlueprintGraph/Classes/K2Node_MatineeController.h
@@ -2,6 +2,7 @@
 
 
 #pragma once
+#include "K2Node.h"
 #include "K2Node_MatineeController.generated.h"
 
 UCLASS(MinimalAPI)
diff --git a/Engine/Source/Editor/BlueprintGraph/Classes/K2Node_Message.h b/Engine/Source/Editor/BlueprintGraph/Classes/K2Node_Message.h
--- a/Engine/Source/Editor/BlueprintGraph/Classes/K2Node_Message.h
+++ b/Engine/Source/Editor/BlueprintGraph/Classes/K2Node_Message.h
@@ -2,6 +2,7 @@
 
 
 #pragma once
+#include "K2Node_CallFunction.h"
 #include "K2Node_Message.generated.h"
 
 UCLASS(MinimalAPI)
diff --git a/Engine/Source/Editor/BlueprintGraph/Classes/K2Node_MultiGate.h b/Engine/Source/Editor/BlueprintGraph/Classes/K2Node_MultiGate.h
--- a/Engine/Source/Editor/BlueprintGraph/Classes/K2Node_MultiGate.h
+++ b/Engine/Source/Editor/BlueprintGraph/Classes/K2Node_MultiGate.h
@@ -1,6 +1,7 @@
 // Copyright 1998-2014 Epic Games, Inc. All Rights Reserved.
 
 #pragma once
+#include "K2Node_ExecutionSequence.h"
 #include "K2Node_MultiGate.generated.h"
 
 UCLASS(MinimalAPI)
diff --git a/Engine/Source/Editor/BlueprintGraph/Classes/K2Node_RemoveDelegate.h b/Engine/Source/Editor/BlueprintGraph/Classes/K2Node_RemoveDelegate.h
--- a/Engine/Source/Editor/BlueprintGraph/Classes/K2Node_RemoveDelegate.h
+++ b/Engine/Source/Editor/BlueprintGraph/Classes/K2Node_RemoveDelegate.h
@@ -1,6 +1,6 @@
 // Copyright 1998-2014 Epic Games, Inc. All Rights Reserved.
 #pragma once
-
+#include "K2Node_BaseMCDelegate.h"
 #include "K2Node_RemoveDelegate.generated.h"
 
 UCLASS(MinimalAPI)
diff --git a/Engine/Source/Editor/BlueprintGraph/Classes/K2Node_Select.h b/Engine/Source/Editor/BlueprintGraph/Classes/K2Node_Select.h
--- a/Engine/Source/Editor/BlueprintGraph/Classes/K2Node_Select.h
+++ b/Engine/Source/Editor/BlueprintGraph/Classes/K2Node_Select.h
@@ -2,6 +2,7 @@
 
 
 #pragma once
+#include "K2Node.h"
 #include "K2Node_Select.generated.h"
 
 UCLASS(MinimalAPI)
diff --git a/Engine/Source/Editor/BlueprintGraph/Classes/K2Node_SelectEnum.h b/Engine/Source/Editor/BlueprintGraph/Classes/K2Node_SelectEnum.h
--- a/Engine/Source/Editor/BlueprintGraph/Classes/K2Node_SelectEnum.h
+++ b/Engine/Source/Editor/BlueprintGraph/Classes/K2Node_SelectEnum.h
@@ -2,6 +2,7 @@
 
 
 #pragma once
+#include "K2Node_Select.h"
 #include "K2Node_SelectEnum.generated.h"
 
 UCLASS(MinimalAPI, deprecated)
diff --git a/Engine/Source/Editor/BlueprintGraph/Classes/K2Node_Self.h b/Engine/Source/Editor/BlueprintGraph/Classes/K2Node_Self.h
--- a/Engine/Source/Editor/BlueprintGraph/Classes/K2Node_Self.h
+++ b/Engine/Source/Editor/BlueprintGraph/Classes/K2Node_Self.h
@@ -2,6 +2,7 @@
 
 
 #pragma once
+#include "K2Node.h"
 #include "K2Node_Self.generated.h"
 
 UCLASS(MinimalAPI)
diff --git a/Engine/Source/Editor/BlueprintGraph/Classes/K2Node_SpawnActor.h b/Engine/Source/Editor/BlueprintGraph/Classes/K2Node_SpawnActor.h
--- a/Engine/Source/Editor/BlueprintGraph/Classes/K2Node_SpawnActor.h
+++ b/Engine/Source/Editor/BlueprintGraph/Classes/K2Node_SpawnActor.h
@@ -2,6 +2,7 @@
 
 
 #pragma once
+#include "K2Node.h"
 #include "K2Node_SpawnActor.generated.h"
 
 UCLASS(MinimalAPI)
diff --git a/Engine/Source/Editor/BlueprintGraph/Classes/K2Node_SpawnActorFromClass.h b/Engine/Source/Editor/BlueprintGraph/Classes/K2Node_SpawnActorFromClass.h
--- a/Engine/Source/Editor/BlueprintGraph/Classes/K2Node_SpawnActorFromClass.h
+++ b/Engine/Source/Editor/BlueprintGraph/Classes/K2Node_SpawnActorFromClass.h
@@ -2,6 +2,7 @@
 
 
 #pragma once
+#include "K2Node.h"
 #include "K2Node_SpawnActorFromClass.generated.h"
 
 UCLASS(MinimalAPI)
diff --git a/Engine/Source/Editor/BlueprintGraph/Classes/K2Node_StructMemberGet.h b/Engine/Source/Editor/BlueprintGraph/Classes/K2Node_StructMemberGet.h
--- a/Engine/Source/Editor/BlueprintGraph/Classes/K2Node_StructMemberGet.h
+++ b/Engine/Source/Editor/BlueprintGraph/Classes/K2Node_StructMemberGet.h
@@ -1,6 +1,7 @@
 // Copyright 1998-2014 Epic Games, Inc. All Rights Reserved.
 
 #pragma once
+#include "K2Node_StructOperation.h"
 #include "K2Node_StructMemberGet.generated.h"
 
 // Pure kismet node that gets one or more member variables of a struct
diff --git a/Engine/Source/Editor/BlueprintGraph/Classes/K2Node_StructMemberSet.h b/Engine/Source/Editor/BlueprintGraph/Classes/K2Node_StructMemberSet.h
--- a/Engine/Source/Editor/BlueprintGraph/Classes/K2Node_StructMemberSet.h
+++ b/Engine/Source/Editor/BlueprintGraph/Classes/K2Node_StructMemberSet.h
@@ -1,6 +1,7 @@
 // Copyright 1998-2014 Epic Games, Inc. All Rights Reserved.
 
 #pragma once
+#include "K2Node_StructOperation.h"
 #include "K2Node_StructMemberSet.generated.h"
 
 // Imperative kismet node that sets one or more member variables of a struct
diff --git a/Engine/Source/Editor/BlueprintGraph/Classes/K2Node_StructOperation.h b/Engine/Source/Editor/BlueprintGraph/Classes/K2Node_StructOperation.h
--- a/Engine/Source/Editor/BlueprintGraph/Classes/K2Node_StructOperation.h
+++ b/Engine/Source/Editor/BlueprintGraph/Classes/K2Node_StructOperation.h
@@ -2,6 +2,7 @@
 
 
 #pragma once
+#include "K2Node_Variable.h"
 #include "K2Node_StructOperation.generated.h"
 
 UCLASS(MinimalAPI, abstract)
diff --git a/Engine/Source/Editor/BlueprintGraph/Classes/K2Node_Switch.h b/Engine/Source/Editor/BlueprintGraph/Classes/K2Node_Switch.h
--- a/Engine/Source/Editor/BlueprintGraph/Classes/K2Node_Switch.h
+++ b/Engine/Source/Editor/BlueprintGraph/Classes/K2Node_Switch.h
@@ -2,6 +2,7 @@
 
 
 #pragma once
+#include "K2Node.h"
 #include "K2Node_Switch.generated.h"
 
 UCLASS(MinimalAPI, abstract)
diff --git a/Engine/Source/Editor/BlueprintGraph/Classes/K2Node_SwitchEnum.h b/Engine/Source/Editor/BlueprintGraph/Classes/K2Node_SwitchEnum.h
--- a/Engine/Source/Editor/BlueprintGraph/Classes/K2Node_SwitchEnum.h
+++ b/Engine/Source/Editor/BlueprintGraph/Classes/K2Node_SwitchEnum.h
@@ -2,6 +2,8 @@
 
 
 #pragma once
+#include "K2Node_Switch.h"
+#include "NodeDependingOnEnumInterface.h"
 #include "K2Node_SwitchEnum.generated.h"
 
 UCLASS(MinimalAPI)
diff --git a/Engine/Source/Editor/BlueprintGraph/Classes/K2Node_SwitchInteger.h b/Engine/Source/Editor/BlueprintGraph/Classes/K2Node_SwitchInteger.h
--- a/Engine/Source/Editor/BlueprintGraph/Classes/K2Node_SwitchInteger.h
+++ b/Engine/Source/Editor/BlueprintGraph/Classes/K2Node_SwitchInteger.h
@@ -3,6 +3,7 @@
 
 
 #pragma once
+#include "K2Node_Switch.h"
 #include "K2Node_SwitchInteger.generated.h"
 
 UCLASS(MinimalAPI)
diff --git a/Engine/Source/Editor/BlueprintGraph/Classes/K2Node_SwitchName.h b/Engine/Source/Editor/BlueprintGraph/Classes/K2Node_SwitchName.h
--- a/Engine/Source/Editor/BlueprintGraph/Classes/K2Node_SwitchName.h
+++ b/Engine/Source/Editor/BlueprintGraph/Classes/K2Node_SwitchName.h
@@ -2,6 +2,7 @@
 
 
 #pragma once
+#include "K2Node_Switch.h"
 #include "K2Node_SwitchName.generated.h"
 
 UCLASS(MinimalAPI)
diff --git a/Engine/Source/Editor/BlueprintGraph/Classes/K2Node_SwitchString.h b/Engine/Source/Editor/BlueprintGraph/Classes/K2Node_SwitchString.h
--- a/Engine/Source/Editor/BlueprintGraph/Classes/K2Node_SwitchString.h
+++ b/Engine/Source/Editor/BlueprintGraph/Classes/K2Node_SwitchString.h
@@ -2,6 +2,7 @@
 
 
 #pragma once
+#include "K2Node_Switch.h"
 #include "K2Node_SwitchString.generated.h"
 
 UCLASS(MinimalAPI)
diff --git a/Engine/Source/Editor/BlueprintGraph/Classes/K2Node_TemporaryVariable.h b/Engine/Source/Editor/BlueprintGraph/Classes/K2Node_TemporaryVariable.h
--- a/Engine/Source/Editor/BlueprintGraph/Classes/K2Node_TemporaryVariable.h
+++ b/Engine/Source/Editor/BlueprintGraph/Classes/K2Node_TemporaryVariable.h
@@ -1,7 +1,7 @@
 // Copyright 1998-2014 Epic Games, Inc. All Rights Reserved.
 
 #pragma once
-
+#include "K2Node.h"
 #include "K2Node_TemporaryVariable.generated.h"
 
 UCLASS(MinimalAPI)
diff --git a/Engine/Source/Editor/BlueprintGraph/Classes/K2Node_Timeline.h b/Engine/Source/Editor/BlueprintGraph/Classes/K2Node_Timeline.h
--- a/Engine/Source/Editor/BlueprintGraph/Classes/K2Node_Timeline.h
+++ b/Engine/Source/Editor/BlueprintGraph/Classes/K2Node_Timeline.h
@@ -2,6 +2,7 @@
 
 
 #pragma once
+#include "K2Node.h"
 #include "K2Node_Timeline.generated.h"
 
 UCLASS(MinimalAPI)
diff --git a/Engine/Source/Editor/BlueprintGraph/Classes/K2Node_Tunnel.h b/Engine/Source/Editor/BlueprintGraph/Classes/K2Node_Tunnel.h
--- a/Engine/Source/Editor/BlueprintGraph/Classes/K2Node_Tunnel.h
+++ b/Engine/Source/Editor/BlueprintGraph/Classes/K2Node_Tunnel.h
@@ -1,7 +1,7 @@
 // Copyright 1998-2014 Epic Games, Inc. All Rights Reserved.
 
-
 #pragma once
+#include "K2Node_EditablePinBase.h"
 #include "K2Node_Tunnel.generated.h"
 
 UCLASS(MinimalAPI)
diff --git a/Engine/Source/Editor/BlueprintGraph/Classes/K2Node_VariableGet.h b/Engine/Source/Editor/BlueprintGraph/Classes/K2Node_VariableGet.h
--- a/Engine/Source/Editor/BlueprintGraph/Classes/K2Node_VariableGet.h
+++ b/Engine/Source/Editor/BlueprintGraph/Classes/K2Node_VariableGet.h
@@ -2,6 +2,7 @@
 
 
 #pragma once
+#include "K2node_Variable.h"
 #include "K2Node_VariableGet.generated.h"
 
 UCLASS(MinimalAPI)
diff --git a/Engine/Source/Editor/BlueprintGraph/Classes/K2Node_VariableSet.h b/Engine/Source/Editor/BlueprintGraph/Classes/K2Node_VariableSet.h
--- a/Engine/Source/Editor/BlueprintGraph/Classes/K2Node_VariableSet.h
+++ b/Engine/Source/Editor/BlueprintGraph/Classes/K2Node_VariableSet.h
@@ -2,6 +2,7 @@
 
 
 #pragma once
+#include "K2Node_Variable.h"
 #include "K2Node_VariableSet.generated.h"
 
 UCLASS(MinimalAPI)
diff --git a/Engine/Source/Editor/BlueprintGraph/Private/EdGraphSchema_K2.cpp b/Engine/Source/Editor/BlueprintGraph/Private/EdGraphSchema_K2.cpp
--- a/Engine/Source/Editor/BlueprintGraph/Private/EdGraphSchema_K2.cpp
+++ b/Engine/Source/Editor/BlueprintGraph/Private/EdGraphSchema_K2.cpp
@@ -24,6 +24,12 @@
 #include "K2ActionMenuBuilder.h"
 #include "AssetRegistryModule.h"
 
+#include "K2Node_CastByteToEnum.h"
+#include "K2Node_ClassDynamicCast.h"
+#include "K2Node_GetEnumeratorName.h"
+#include "K2Node_GetEnumeratorNameAsString.h"
+#include "K2Node_Tunnel.h"
+
 //////////////////////////////////////////////////////////////////////////
 // FBlueprintMetadata
 
diff --git a/Engine/Source/Editor/BlueprintGraph/Private/K2ActionMenuBuilder.cpp b/Engine/Source/Editor/BlueprintGraph/Private/K2ActionMenuBuilder.cpp
--- a/Engine/Source/Editor/BlueprintGraph/Private/K2ActionMenuBuilder.cpp
+++ b/Engine/Source/Editor/BlueprintGraph/Private/K2ActionMenuBuilder.cpp
@@ -7,6 +7,25 @@
 #include "ObjectEditorUtils.h"
 #include "Matinee/MatineeActor.h"
 
+#include "K2Node_CastByteToEnum.h"
+#include "K2Node_ClassDynamicCast.h"
+#include "K2Node_EnumLiteral.h"
+#include "K2Node_ForEachElementInEnum.h"
+#include "K2Node_GetEnumeratorName.h"
+#include "K2Node_GetEnumeratorNameAsString.h"
+#include "K2Node_GetInputAxisValue.h"
+#include "K2Node_GetInputVectorAxisValue.h"
+#include "K2Node_GetNumEnumEntries.h"
+#include "K2Node_InputAxisKeyEvent.h"
+#include "K2Node_InputVectorAxisEvent.h"
+#include "K2Node_Message.h"
+#include "K2Node_MultiGate.h"
+#include "K2Node_SwitchEnum.h"
+#include "K2Node_SwitchString.h"
+#include "K2Node_VariableSetRef.h"
+
+#include "EdGraphSchema_K2_Actions.h"
+
 #define LOCTEXT_NAMESPACE "KismetSchema"
 
 /*******************************************************************************
diff --git a/Engine/Source/Editor/BlueprintGraph/Private/K2Node.cpp b/Engine/Source/Editor/BlueprintGraph/Private/K2Node.cpp
--- a/Engine/Source/Editor/BlueprintGraph/Private/K2Node.cpp
+++ b/Engine/Source/Editor/BlueprintGraph/Private/K2Node.cpp
@@ -2,6 +2,7 @@
 
 
 #include "BlueprintGraphPrivatePCH.h"
+#include "K2Node.h"
 #include "KismetDebugUtilities.h" // for HasDebuggingData(), GetWatchText()
 
 #define LOCTEXT_NAMESPACE "K2Node"
diff --git a/Engine/Source/Editor/BlueprintGraph/Private/K2Node_AIMoveTo.cpp b/Engine/Source/Editor/BlueprintGraph/Private/K2Node_AIMoveTo.cpp
--- a/Engine/Source/Editor/BlueprintGraph/Private/K2Node_AIMoveTo.cpp
+++ b/Engine/Source/Editor/BlueprintGraph/Private/K2Node_AIMoveTo.cpp
@@ -3,6 +3,7 @@
 
 #include "BlueprintGraphPrivatePCH.h"
 #include "CompilerResultsLog.h"
+#include "K2Node_AIMoveTo.h"
 
 #define LOCTEXT_NAMESPACE "K2Node_AIMoveTo"
 
diff --git a/Engine/Source/Editor/BlueprintGraph/Private/K2Node_BreakStruct.cpp b/Engine/Source/Editor/BlueprintGraph/Private/K2Node_BreakStruct.cpp
--- a/Engine/Source/Editor/BlueprintGraph/Private/K2Node_BreakStruct.cpp
+++ b/Engine/Source/Editor/BlueprintGraph/Private/K2Node_BreakStruct.cpp
@@ -2,7 +2,6 @@
 
 #include "BlueprintGraphPrivatePCH.h"
 
-#include "K2Node_BreakStruct.h"
 #include "KismetCompiler.h"
 
 #define LOCTEXT_NAMESPACE "K2Node_BreakStruct"
diff --git a/Engine/Source/Editor/BlueprintGraph/Private/K2Node_CastByteToEnum.cpp b/Engine/Source/Editor/BlueprintGraph/Private/K2Node_CastByteToEnum.cpp
--- a/Engine/Source/Editor/BlueprintGraph/Private/K2Node_CastByteToEnum.cpp
+++ b/Engine/Source/Editor/BlueprintGraph/Private/K2Node_CastByteToEnum.cpp
@@ -3,7 +3,8 @@
 
 #include "BlueprintGraphPrivatePCH.h"
 #include "KismetCompiler.h"
-#include "../../../Runtime/Engine/Classes/Kismet/KismetNodeHelperLibrary.h"
+#include "Kismet/KismetNodeHelperLibrary.h"
+#include "K2Node_CastByteToEnum.h"
 
 const FString UK2Node_CastByteToEnum::ByteInputPinName = TEXT("Byte");
 
diff --git a/Engine/Source/Editor/BlueprintGraph/Private/K2Node_ClassDynamicCast.cpp b/Engine/Source/Editor/BlueprintGraph/Private/K2Node_ClassDynamicCast.cpp
--- a/Engine/Source/Editor/BlueprintGraph/Private/K2Node_ClassDynamicCast.cpp
+++ b/Engine/Source/Editor/BlueprintGraph/Private/K2Node_ClassDynamicCast.cpp
@@ -2,6 +2,7 @@
 
 
 #include "BlueprintGraphPrivatePCH.h"
+#include "K2Node_ClassDynamicCast.h"
 #include "DynamicCastHandler.h"
 
 struct FClassDynamicCastHelper
diff --git a/Engine/Source/Editor/BlueprintGraph/Private/K2Node_CommutativeAssociativeBinaryOperator.cpp b/Engine/Source/Editor/BlueprintGraph/Private/K2Node_CommutativeAssociativeBinaryOperator.cpp
--- a/Engine/Source/Editor/BlueprintGraph/Private/K2Node_CommutativeAssociativeBinaryOperator.cpp
+++ b/Engine/Source/Editor/BlueprintGraph/Private/K2Node_CommutativeAssociativeBinaryOperator.cpp
@@ -2,7 +2,6 @@
 
 #include "BlueprintGraphPrivatePCH.h"
 
-#include "K2Node_CommutativeAssociativeBinaryOperator.h"
 #include "ScopedTransaction.h"
 
 #include "KismetCompiler.h"
diff --git a/Engine/Source/Editor/BlueprintGraph/Private/K2Node_CreateDelegate.cpp b/Engine/Source/Editor/BlueprintGraph/Private/K2Node_CreateDelegate.cpp
--- a/Engine/Source/Editor/BlueprintGraph/Private/K2Node_CreateDelegate.cpp
+++ b/Engine/Source/Editor/BlueprintGraph/Private/K2Node_CreateDelegate.cpp
@@ -1,9 +1,7 @@
 // Copyright 1998-2014 Epic Games, Inc. All Rights Reserved.
 #include "BlueprintGraphPrivatePCH.h"
 
-#include "K2Node_CreateDelegate.h"
 #include "DelegateNodeHandlers.h"
-#include "K2Node_BaseMCDelegate.h"
 
 struct FK2Node_CreateDelegate_Helper
 {
diff --git a/Engine/Source/Editor/BlueprintGraph/Private/K2Node_CustomEvent.cpp b/Engine/Source/Editor/BlueprintGraph/Private/K2Node_CustomEvent.cpp
--- a/Engine/Source/Editor/BlueprintGraph/Private/K2Node_CustomEvent.cpp
+++ b/Engine/Source/Editor/BlueprintGraph/Private/K2Node_CustomEvent.cpp
@@ -4,7 +4,6 @@
 #include "BlueprintGraphPrivatePCH.h"
 
 #include "Kismet2NameValidators.h"
-#include "K2Node_BaseMCDelegate.h"
 #include "CompilerResultsLog.h"
 
 #define LOCTEXT_NAMESPACE "K2Node_CustomEvent"
diff --git a/Engine/Source/Editor/BlueprintGraph/Private/K2Node_DeadClass.cpp b/Engine/Source/Editor/BlueprintGraph/Private/K2Node_DeadClass.cpp
--- a/Engine/Source/Editor/BlueprintGraph/Private/K2Node_DeadClass.cpp
+++ b/Engine/Source/Editor/BlueprintGraph/Private/K2Node_DeadClass.cpp
@@ -1,6 +1,7 @@
 // Copyright 1998-2014 Epic Games, Inc. All Rights Reserved.
 
 #include "BlueprintGraphPrivatePCH.h"
+#include "K2Node_DeadClass.h"
 
 UK2Node_DeadClass::UK2Node_DeadClass(const class FPostConstructInitializeProperties& PCIP)
 	: Super(PCIP)
diff --git a/Engine/Source/Editor/BlueprintGraph/Private/K2Node_EnumEquality.cpp b/Engine/Source/Editor/BlueprintGraph/Private/K2Node_EnumEquality.cpp
--- a/Engine/Source/Editor/BlueprintGraph/Private/K2Node_EnumEquality.cpp
+++ b/Engine/Source/Editor/BlueprintGraph/Private/K2Node_EnumEquality.cpp
@@ -3,8 +3,9 @@
 
 #include "BlueprintGraphPrivatePCH.h"
 #include "KismetCompiler.h"
-#include "../../../Runtime/Engine/Classes/Kismet/KismetMathLibrary.h"
+#include "Kismet/KismetMathLibrary.h"
 #include "K2ActionMenuBuilder.h" // for FK2ActionMenuBuilder::AddNewNodeAction()
+#include "K2Node_EnumEquality.h"
 
 UK2Node_EnumEquality::UK2Node_EnumEquality(const class FPostConstructInitializeProperties& PCIP)
 	: Super(PCIP)
diff --git a/Engine/Source/Editor/BlueprintGraph/Private/K2Node_EnumLiteral.cpp b/Engine/Source/Editor/BlueprintGraph/Private/K2Node_EnumLiteral.cpp
--- a/Engine/Source/Editor/BlueprintGraph/Private/K2Node_EnumLiteral.cpp
+++ b/Engine/Source/Editor/BlueprintGraph/Private/K2Node_EnumLiteral.cpp
@@ -3,6 +3,7 @@
 
 #include "BlueprintGraphPrivatePCH.h"
 #include "KismetCompilerMisc.h"
+#include "K2Node_EnumLiteral.h"
 
 const FString& UK2Node_EnumLiteral::GetEnumInputPinName()
 {
diff --git a/Engine/Source/Editor/BlueprintGraph/Private/K2Node_ExecutionSequence.cpp b/Engine/Source/Editor/BlueprintGraph/Private/K2Node_ExecutionSequence.cpp
--- a/Engine/Source/Editor/BlueprintGraph/Private/K2Node_ExecutionSequence.cpp
+++ b/Engine/Source/Editor/BlueprintGraph/Private/K2Node_ExecutionSequence.cpp
@@ -3,6 +3,7 @@
 
 #include "BlueprintGraphPrivatePCH.h"
 #include "KismetCompiler.h"
+#include "KismetCompilerMisc.h"
 
 #define LOCTEXT_NAMESPACE "K2Node_MultiGate"
 
diff --git a/Engine/Source/Editor/BlueprintGraph/Private/K2Node_ForEachElementInEnum.cpp b/Engine/Source/Editor/BlueprintGraph/Private/K2Node_ForEachElementInEnum.cpp
--- a/Engine/Source/Editor/BlueprintGraph/Private/K2Node_ForEachElementInEnum.cpp
+++ b/Engine/Source/Editor/BlueprintGraph/Private/K2Node_ForEachElementInEnum.cpp
@@ -3,8 +3,11 @@
 
 #include "BlueprintGraphPrivatePCH.h"
 #include "KismetCompiler.h"
-#include "../../../Runtime/Engine/Classes/Kismet/KismetNodeHelperLibrary.h"
-#include "../../../Runtime/Engine/Classes/Kismet/KismetMathLibrary.h"
+#include "Kismet/KismetNodeHelperLibrary.h"
+#include "Kismet/KismetMathLibrary.h"
+#include "K2Node_CastByteToEnum.h"
+#include "K2Node_ForEachElementInEnum.h"
+#include "K2Node_GetNumEnumEntries.h"
 
 #define LOCTEXT_NAMESPACE "K2Node"
 
diff --git a/Engine/Source/Editor/BlueprintGraph/Private/K2Node_GetEnumeratorName.cpp b/Engine/Source/Editor/BlueprintGraph/Private/K2Node_GetEnumeratorName.cpp
--- a/Engine/Source/Editor/BlueprintGraph/Private/K2Node_GetEnumeratorName.cpp
+++ b/Engine/Source/Editor/BlueprintGraph/Private/K2Node_GetEnumeratorName.cpp
@@ -3,8 +3,9 @@
 
 #include "BlueprintGraphPrivatePCH.h"
 #include "KismetCompiler.h"
-#include "../../../Runtime/Engine/Classes/Kismet/KismetNodeHelperLibrary.h"
-#include "../../../Runtime/Engine/Classes/Kismet/KismetSystemLibrary.h"
+#include "Kismet/KismetNodeHelperLibrary.h"
+#include "Kismet/KismetSystemLibrary.h"
+#include "K2Node_GetEnumeratorName.h"
 
 FString UK2Node_GetEnumeratorName::EnumeratorPinName = TEXT("Enumerator");
 
diff --git a/Engine/Source/Editor/BlueprintGraph/Private/K2Node_GetEnumeratorNameAsString.cpp b/Engine/Source/Editor/BlueprintGraph/Private/K2Node_GetEnumeratorNameAsString.cpp
--- a/Engine/Source/Editor/BlueprintGraph/Private/K2Node_GetEnumeratorNameAsString.cpp
+++ b/Engine/Source/Editor/BlueprintGraph/Private/K2Node_GetEnumeratorNameAsString.cpp
@@ -3,7 +3,8 @@
 
 #include "BlueprintGraphPrivatePCH.h"
 #include "KismetCompiler.h"
-#include "../../../Runtime/Engine/Classes/Kismet/KismetNodeHelperLibrary.h"
+#include "Kismet/KismetNodeHelperLibrary.h"
+#include "K2Node_GetEnumeratorNameAsString.h"
 
 UK2Node_GetEnumeratorNameAsString::UK2Node_GetEnumeratorNameAsString(const class FPostConstructInitializeProperties& PCIP)
 	: Super(PCIP)
diff --git a/Engine/Source/Editor/BlueprintGraph/Private/K2Node_GetInputAxisKeyValue.cpp b/Engine/Source/Editor/BlueprintGraph/Private/K2Node_GetInputAxisKeyValue.cpp
--- a/Engine/Source/Editor/BlueprintGraph/Private/K2Node_GetInputAxisKeyValue.cpp
+++ b/Engine/Source/Editor/BlueprintGraph/Private/K2Node_GetInputAxisKeyValue.cpp
@@ -1,7 +1,7 @@
 // Copyright 1998-2014 Epic Games, Inc. All Rights Reserved.
 
 #include "BlueprintGraphPrivatePCH.h"
-
+#include "K2Node_GetInputAxisKeyValue.h"
 #include "CompilerResultsLog.h"
 
 UK2Node_GetInputAxisKeyValue::UK2Node_GetInputAxisKeyValue(const class FPostConstructInitializeProperties& PCIP)
diff --git a/Engine/Source/Editor/BlueprintGraph/Private/K2Node_GetInputAxisValue.cpp b/Engine/Source/Editor/BlueprintGraph/Private/K2Node_GetInputAxisValue.cpp
--- a/Engine/Source/Editor/BlueprintGraph/Private/K2Node_GetInputAxisValue.cpp
+++ b/Engine/Source/Editor/BlueprintGraph/Private/K2Node_GetInputAxisValue.cpp
@@ -1,7 +1,7 @@
 // Copyright 1998-2014 Epic Games, Inc. All Rights Reserved.
 
 #include "BlueprintGraphPrivatePCH.h"
-
+#include "K2Node_GetInputAxisValue.h"
 #include "CompilerResultsLog.h"
 
 UK2Node_GetInputAxisValue::UK2Node_GetInputAxisValue(const class FPostConstructInitializeProperties& PCIP)
diff --git a/Engine/Source/Editor/BlueprintGraph/Private/K2Node_GetInputVectorAxisValue.cpp b/Engine/Source/Editor/BlueprintGraph/Private/K2Node_GetInputVectorAxisValue.cpp
--- a/Engine/Source/Editor/BlueprintGraph/Private/K2Node_GetInputVectorAxisValue.cpp
+++ b/Engine/Source/Editor/BlueprintGraph/Private/K2Node_GetInputVectorAxisValue.cpp
@@ -1,6 +1,7 @@
 // Copyright 1998-2014 Epic Games, Inc. All Rights Reserved.
 
 #include "BlueprintGraphPrivatePCH.h"
+#include "K2Node_GetInputVectorAxisValue.h"
 
 #include "CompilerResultsLog.h"
 
diff --git a/Engine/Source/Editor/BlueprintGraph/Private/K2Node_GetNumEnumEntries.cpp b/Engine/Source/Editor/BlueprintGraph/Private/K2Node_GetNumEnumEntries.cpp
--- a/Engine/Source/Editor/BlueprintGraph/Private/K2Node_GetNumEnumEntries.cpp
+++ b/Engine/Source/Editor/BlueprintGraph/Private/K2Node_GetNumEnumEntries.cpp
@@ -4,6 +4,7 @@
 #include "BlueprintGraphPrivatePCH.h"
 #include "KismetCompiler.h"
 #include "../../../Runtime/Engine/Classes/Kismet/KismetSystemLibrary.h"
+#include "K2Node_GetNumEnumEntries.h"
 
 UK2Node_GetNumEnumEntries::UK2Node_GetNumEnumEntries(const class FPostConstructInitializeProperties& PCIP)
 	: Super(PCIP)
diff --git a/Engine/Source/Editor/BlueprintGraph/Private/K2Node_InputAction.cpp b/Engine/Source/Editor/BlueprintGraph/Private/K2Node_InputAction.cpp
--- a/Engine/Source/Editor/BlueprintGraph/Private/K2Node_InputAction.cpp
+++ b/Engine/Source/Editor/BlueprintGraph/Private/K2Node_InputAction.cpp
@@ -1,7 +1,7 @@
 // Copyright 1998-2014 Epic Games, Inc. All Rights Reserved.
 
 #include "BlueprintGraphPrivatePCH.h"
-
+#include "K2Node_InputActionEvent.h"
 #include "CompilerResultsLog.h"
 #include "KismetCompiler.h"
 
diff --git a/Engine/Source/Editor/BlueprintGraph/Private/K2Node_InputActionEvent.cpp b/Engine/Source/Editor/BlueprintGraph/Private/K2Node_InputActionEvent.cpp
--- a/Engine/Source/Editor/BlueprintGraph/Private/K2Node_InputActionEvent.cpp
+++ b/Engine/Source/Editor/BlueprintGraph/Private/K2Node_InputActionEvent.cpp
@@ -1,6 +1,7 @@
 // Copyright 1998-2014 Epic Games, Inc. All Rights Reserved.
 
 #include "BlueprintGraphPrivatePCH.h"
+#include "K2Node_InputActionEvent.h"
 
 UK2Node_InputActionEvent::UK2Node_InputActionEvent(const class FPostConstructInitializeProperties& PCIP)
 	: Super(PCIP)
diff --git a/Engine/Source/Editor/BlueprintGraph/Private/K2Node_InputAxisKeyEvent.cpp b/Engine/Source/Editor/BlueprintGraph/Private/K2Node_InputAxisKeyEvent.cpp
--- a/Engine/Source/Editor/BlueprintGraph/Private/K2Node_InputAxisKeyEvent.cpp
+++ b/Engine/Source/Editor/BlueprintGraph/Private/K2Node_InputAxisKeyEvent.cpp
@@ -1,7 +1,7 @@
 // Copyright 1998-2014 Epic Games, Inc. All Rights Reserved.
 
 #include "BlueprintGraphPrivatePCH.h"
-
+#include "K2Node_InputAxisKeyEvent.h"
 #include "CompilerResultsLog.h"
 
 UK2Node_InputAxisKeyEvent::UK2Node_InputAxisKeyEvent(const class FPostConstructInitializeProperties& PCIP)
diff --git a/Engine/Source/Editor/BlueprintGraph/Private/K2Node_InputKey.cpp b/Engine/Source/Editor/BlueprintGraph/Private/K2Node_InputKey.cpp
--- a/Engine/Source/Editor/BlueprintGraph/Private/K2Node_InputKey.cpp
+++ b/Engine/Source/Editor/BlueprintGraph/Private/K2Node_InputKey.cpp
@@ -1,7 +1,7 @@
 // Copyright 1998-2014 Epic Games, Inc. All Rights Reserved.
 
 #include "BlueprintGraphPrivatePCH.h"
-
+#include "K2Node_InputKeyEvent.h"
 #include "CompilerResultsLog.h"
 #include "KismetCompiler.h"
 
diff --git a/Engine/Source/Editor/BlueprintGraph/Private/K2Node_InputKeyEvent.cpp b/Engine/Source/Editor/BlueprintGraph/Private/K2Node_InputKeyEvent.cpp
--- a/Engine/Source/Editor/BlueprintGraph/Private/K2Node_InputKeyEvent.cpp
+++ b/Engine/Source/Editor/BlueprintGraph/Private/K2Node_InputKeyEvent.cpp
@@ -1,6 +1,7 @@
 // Copyright 1998-2014 Epic Games, Inc. All Rights Reserved.
 
 #include "BlueprintGraphPrivatePCH.h"
+#include "K2Node_InputKeyEvent.h"
 
 UK2Node_InputKeyEvent::UK2Node_InputKeyEvent(const class FPostConstructInitializeProperties& PCIP)
 	: Super(PCIP)
diff --git a/Engine/Source/Editor/BlueprintGraph/Private/K2Node_InputTouch.cpp b/Engine/Source/Editor/BlueprintGraph/Private/K2Node_InputTouch.cpp
--- a/Engine/Source/Editor/BlueprintGraph/Private/K2Node_InputTouch.cpp
+++ b/Engine/Source/Editor/BlueprintGraph/Private/K2Node_InputTouch.cpp
@@ -1,7 +1,7 @@
 // Copyright 1998-2014 Epic Games, Inc. All Rights Reserved.
 
 #include "BlueprintGraphPrivatePCH.h"
-
+#include "K2Node_InputTouchEvent.h"
 #include "CompilerResultsLog.h"
 #include "KismetCompiler.h"
 
diff --git a/Engine/Source/Editor/BlueprintGraph/Private/K2Node_InputTouchEvent.cpp b/Engine/Source/Editor/BlueprintGraph/Private/K2Node_InputTouchEvent.cpp
--- a/Engine/Source/Editor/BlueprintGraph/Private/K2Node_InputTouchEvent.cpp
+++ b/Engine/Source/Editor/BlueprintGraph/Private/K2Node_InputTouchEvent.cpp
@@ -1,7 +1,7 @@
 // Copyright 1998-2014 Epic Games, Inc. All Rights Reserved.
 
 #include "BlueprintGraphPrivatePCH.h"
-
+#include "K2Node_InputTouchEvent.h"
 #include "CompilerResultsLog.h"
 
 UK2Node_InputTouchEvent::UK2Node_InputTouchEvent(const class FPostConstructInitializeProperties& PCIP)
diff --git a/Engine/Source/Editor/BlueprintGraph/Private/K2Node_InputVectorAxisEvent.cpp b/Engine/Source/Editor/BlueprintGraph/Private/K2Node_InputVectorAxisEvent.cpp
--- a/Engine/Source/Editor/BlueprintGraph/Private/K2Node_InputVectorAxisEvent.cpp
+++ b/Engine/Source/Editor/BlueprintGraph/Private/K2Node_InputVectorAxisEvent.cpp
@@ -1,7 +1,7 @@
 // Copyright 1998-2014 Epic Games, Inc. All Rights Reserved.
 
 #include "BlueprintGraphPrivatePCH.h"
-
+#include "K2Node_InputVectorAxisEvent.h"
 #include "CompilerResultsLog.h"
 
 UK2Node_InputVectorAxisEvent::UK2Node_InputVectorAxisEvent(const class FPostConstructInitializeProperties& PCIP)
diff --git a/Engine/Source/Editor/BlueprintGraph/Private/K2Node_LocalVariable.cpp b/Engine/Source/Editor/BlueprintGraph/Private/K2Node_LocalVariable.cpp
--- a/Engine/Source/Editor/BlueprintGraph/Private/K2Node_LocalVariable.cpp
+++ b/Engine/Source/Editor/BlueprintGraph/Private/K2Node_LocalVariable.cpp
@@ -5,6 +5,9 @@
 #include "Kismet2NameValidators.h"
 #include "ScopedTransaction.h"
 
+#include "K2Node_FunctionEntry.h"
+#include "K2Node_LocalVariable.h"
+
 UDEPRECATED_K2Node_LocalVariable::UDEPRECATED_K2Node_LocalVariable(const class FPostConstructInitializeProperties& PCIP)
 	: Super(PCIP)
 {
diff --git a/Engine/Source/Editor/BlueprintGraph/Private/K2Node_MCDelegate.cpp b/Engine/Source/Editor/BlueprintGraph/Private/K2Node_MCDelegate.cpp
--- a/Engine/Source/Editor/BlueprintGraph/Private/K2Node_MCDelegate.cpp
+++ b/Engine/Source/Editor/BlueprintGraph/Private/K2Node_MCDelegate.cpp
@@ -1,11 +1,6 @@
 // Copyright 1998-2014 Epic Games, Inc. All Rights Reserved.
 #include "BlueprintGraphPrivatePCH.h"
 
-#include "K2Node_BaseMCDelegate.h"
-#include "K2Node_AddDelegate.h"
-#include "K2Node_CallDelegate.h"
-#include "K2Node_ClearDelegate.h"
-#include "K2Node_RemoveDelegate.h"
 #include "CompilerResultsLog.h"
 #include "DelegateNodeHandlers.h"
 
diff --git a/Engine/Source/Editor/BlueprintGraph/Private/K2Node_Message.cpp b/Engine/Source/Editor/BlueprintGraph/Private/K2Node_Message.cpp
--- a/Engine/Source/Editor/BlueprintGraph/Private/K2Node_Message.cpp
+++ b/Engine/Source/Editor/BlueprintGraph/Private/K2Node_Message.cpp
@@ -2,7 +2,8 @@
 
 #include "BlueprintGraphPrivatePCH.h"
 #include "KismetCompiler.h"
-
+#include "K2Node_CallArrayFunction.h"
+#include "K2Node_Message.h"
 #define LOCTEXT_NAMESPACE "K2Node_Message"
 
 UK2Node_Message::UK2Node_Message(const class FPostConstructInitializeProperties& PCIP)
diff --git a/Engine/Source/Editor/BlueprintGraph/Private/K2Node_MultiGate.cpp b/Engine/Source/Editor/BlueprintGraph/Private/K2Node_MultiGate.cpp
--- a/Engine/Source/Editor/BlueprintGraph/Private/K2Node_MultiGate.cpp
+++ b/Engine/Source/Editor/BlueprintGraph/Private/K2Node_MultiGate.cpp
@@ -2,10 +2,10 @@
 
 #include "BlueprintGraphPrivatePCH.h"
 #include "KismetCompiler.h"
-#include "../../../Runtime/Engine/Classes/Kismet/KismetMathLibrary.h"
-#include "../../../Runtime/Engine/Classes/Kismet/KismetNodeHelperLibrary.h"
-#include "../../../Runtime/Engine/Classes/Kismet/KismetSystemLibrary.h"
-
+#include "Kismet/KismetMathLibrary.h"
+#include "Kismet/KismetNodeHelperLibrary.h"
+#include "Kismet/KismetSystemLibrary.h"
+#include "K2Node_MultiGate.h"
 #define LOCTEXT_NAMESPACE "K2Node_MultiGate"
 
 //////////////////////////////////////////////////////////////////////////
diff --git a/Engine/Source/Editor/BlueprintGraph/Private/K2Node_SelectEnum.cpp b/Engine/Source/Editor/BlueprintGraph/Private/K2Node_SelectEnum.cpp
--- a/Engine/Source/Editor/BlueprintGraph/Private/K2Node_SelectEnum.cpp
+++ b/Engine/Source/Editor/BlueprintGraph/Private/K2Node_SelectEnum.cpp
@@ -2,7 +2,7 @@
 
 
 #include "BlueprintGraphPrivatePCH.h"
-
+#include "K2Node_SelectEnum.h"
 
 UDEPRECATED_K2Node_SelectEnum::UDEPRECATED_K2Node_SelectEnum(const class FPostConstructInitializeProperties& PCIP)
 	: Super(PCIP)
diff --git a/Engine/Source/Editor/BlueprintGraph/Private/K2Node_StructMemberSet.cpp b/Engine/Source/Editor/BlueprintGraph/Private/K2Node_StructMemberSet.cpp
--- a/Engine/Source/Editor/BlueprintGraph/Private/K2Node_StructMemberSet.cpp
+++ b/Engine/Source/Editor/BlueprintGraph/Private/K2Node_StructMemberSet.cpp
@@ -1,7 +1,7 @@
 // Copyright 1998-2014 Epic Games, Inc. All Rights Reserved.
 
 #include "BlueprintGraphPrivatePCH.h"
-
+#include "K2Node_StructMemberSet.h"
 #include "StructMemberNodeHandlers.h"
 
 //////////////////////////////////////////////////////////////////////////
diff --git a/Engine/Source/Editor/BlueprintGraph/Private/K2Node_Switch.cpp b/Engine/Source/Editor/BlueprintGraph/Private/K2Node_Switch.cpp
--- a/Engine/Source/Editor/BlueprintGraph/Private/K2Node_Switch.cpp
+++ b/Engine/Source/Editor/BlueprintGraph/Private/K2Node_Switch.cpp
@@ -4,6 +4,7 @@
 #include "BlueprintGraphPrivatePCH.h"
 #include "KismetCompiler.h"
 #include "../../../Runtime/Engine/Classes/Kismet/KismetMathLibrary.h"
+#include "K2Node_SwitchEnum.h"
 
 #define LOCTEXT_NAMESPACE "K2Node_Switch"
 
diff --git a/Engine/Source/Editor/BlueprintGraph/Private/K2Node_SwitchEnum.cpp b/Engine/Source/Editor/BlueprintGraph/Private/K2Node_SwitchEnum.cpp
--- a/Engine/Source/Editor/BlueprintGraph/Private/K2Node_SwitchEnum.cpp
+++ b/Engine/Source/Editor/BlueprintGraph/Private/K2Node_SwitchEnum.cpp
@@ -1,7 +1,8 @@
 // Copyright 1998-2014 Epic Games, Inc. All Rights Reserved.
 
 #include "BlueprintGraphPrivatePCH.h"
-#include "../../../Runtime/Engine/Classes/Kismet/KismetMathLibrary.h"
+#include "Kismet/KismetMathLibrary.h"
+#include "K2Node_SwitchEnum.h"
 
 #define LOCTEXT_NAMESPACE "K2Node"
 
diff --git a/Engine/Source/Editor/BlueprintGraph/Private/K2Node_SwitchString.cpp b/Engine/Source/Editor/BlueprintGraph/Private/K2Node_SwitchString.cpp
--- a/Engine/Source/Editor/BlueprintGraph/Private/K2Node_SwitchString.cpp
+++ b/Engine/Source/Editor/BlueprintGraph/Private/K2Node_SwitchString.cpp
@@ -1,7 +1,8 @@
 // Copyright 1998-2014 Epic Games, Inc. All Rights Reserved.
 
 #include "BlueprintGraphPrivatePCH.h"
-#include "../../../Runtime/Engine/Classes/Kismet/KismetMathLibrary.h"
+#include "Kismet/KismetMathLibrary.h"
+#include "K2Node_SwitchString.h"
 
 UK2Node_SwitchString::UK2Node_SwitchString(const class FPostConstructInitializeProperties& PCIP)
 	: Super(PCIP)
diff --git a/Engine/Source/Editor/BlueprintGraph/Private/K2Node_Tunnel.cpp b/Engine/Source/Editor/BlueprintGraph/Private/K2Node_Tunnel.cpp
--- a/Engine/Source/Editor/BlueprintGraph/Private/K2Node_Tunnel.cpp
+++ b/Engine/Source/Editor/BlueprintGraph/Private/K2Node_Tunnel.cpp
@@ -3,6 +3,9 @@
 #include "BlueprintGraphPrivatePCH.h"
 #include "EdGraphUtilities.h"
 
+#include "K2Node_MacroInstance.h"
+#include "K2Node_Tunnel.h"
+
 UK2Node_Tunnel::UK2Node_Tunnel(const class FPostConstructInitializeProperties& PCIP)
 	: Super(PCIP)
 {
diff --git a/Engine/Source/Editor/BlueprintGraph/Private/K2Node_VariableSetRef.cpp b/Engine/Source/Editor/BlueprintGraph/Private/K2Node_VariableSetRef.cpp
--- a/Engine/Source/Editor/BlueprintGraph/Private/K2Node_VariableSetRef.cpp
+++ b/Engine/Source/Editor/BlueprintGraph/Private/K2Node_VariableSetRef.cpp
@@ -4,6 +4,7 @@
 #include "BlueprintGraphPrivatePCH.h"
 #include "KismetCompiler.h"
 #include "VariableSetHandler.h"
+#include "K2Node_VariableSetRef.h"
 
 static FString TargetVarPinName(TEXT("Target"));
 static FString VarValuePinName(TEXT("Value"));
diff --git a/Engine/Source/Editor/BlueprintGraph/Public/BlueprintGraphDefinitions.h b/Engine/Source/Editor/BlueprintGraph/Public/BlueprintGraphDefinitions.h
--- a/Engine/Source/Editor/BlueprintGraph/Public/BlueprintGraphDefinitions.h
+++ b/Engine/Source/Editor/BlueprintGraph/Public/BlueprintGraphDefinitions.h
@@ -1,3 +1,61 @@
 // Copyright 1998-2014 Epic Games, Inc. All Rights Reserved.
 
-#include "BlueprintGraphClasses.h"
+#include "EdGraphSchema_K2.h"
+#include "EdGraphSchema_K2_Actions.h"
+
+#include "K2Node_ActorBoundEvent.h"
+#include "K2Node_AddComponent.h"
+#include "K2Node_AddDelegate.h"
+#include "K2Node_AssignmentStatement.h"
+#include "K2Node_BaseAsyncTask.h"
+#include "K2Node_BaseMCDelegate.h"
+#include "K2Node_BreakStruct.h"
+#include "K2Node_CallArrayFunction.h"
+#include "K2Node_CallDelegate.h"
+#include "K2Node_CallFunction.h"
+#include "K2Node_CallFunctionOnMember.h"
+#include "K2Node_CallMaterialParameterCollectionFunction.h"
+#include "K2Node_CallParentFunction.h"
+#include "K2Node_ClearDelegate.h"
+#include "K2Node_CommutativeAssociativeBinaryOperator.h"
+#include "K2Node_ComponentBoundEvent.h"
+#include "K2Node_Composite.h"
+#include "K2Node_CreateDelegate.h"
+#include "K2Node_CustomEvent.h"
+#include "K2Node_DelegateSet.h"
+#include "K2Node_DynamicCast.h"
+#include "K2Node_Event.h"
+#include "K2Node_ExecutionSequence.h"
+#include "K2Node_FormatText.h"
+#include "K2Node_FunctionEntry.h"
+#include "K2Node_FunctionResult.h"
+#include "K2Node_IfThenElse.h"
+#include "K2Node_InputAction.h"
+#include "K2Node_InputAxisEvent.h"
+#include "K2Node_InputKey.h"
+#include "K2Node_InputTouch.h"
+#include "K2Node_Knot.h"
+#include "K2Node_Literal.h"
+#include "K2Node_MacroInstance.h"
+#include "K2Node_MakeArray.h"
+#include "K2Node_MakeStruct.h"
+#include "K2Node_MatineeController.h"
+#include "K2Node_RemoveDelegate.h"
+#include "K2Node_Select.h"
+#include "K2Node_Self.h"
+#include "K2Node_SpawnActor.h"
+#include "K2Node_SpawnActorFromClass.h"
+#include "K2Node_StructMemberGet.h"
+#include "K2Node_StructMemberSet.h"
+#include "K2Node_StructOperation.h"
+#include "K2Node_Switch.h"
+#include "K2Node_SwitchInteger.h"
+#include "K2Node_SwitchName.h"
+#include "K2Node_TemporaryVariable.h"
+#include "K2Node_Timeline.h"
+#include "K2Node_Variable.h"
+#include "K2Node_VariableGet.h"
+#include "K2Node_VariableSet.h"
+
+
+#include "NodeDependingOnEnumInterface.h"
