Array parameter in interface does not generate error.

diff --git a/Engine/Source/Editor/BlueprintGraph/Private/EdGraphSchema_K2.cpp b/Engine/Source/Editor/BlueprintGraph/Private/EdGraphSchema_K2.cpp
--- a/Engine/Source/Editor/BlueprintGraph/Private/EdGraphSchema_K2.cpp
+++ b/Engine/Source/Editor/BlueprintGraph/Private/EdGraphSchema_K2.cpp
@@ -1661,14 +1661,19 @@ FString UEdGraphSchema_K2::IsPinDefaultValid(const UEdGraphPin* Pin, const FStri
 	const bool bIsReference = Pin->PinType.bIsReference;
 	const bool bIsAutoCreateRefTerm = IsAutoCreateRefTerm(Pin);
 
-	if (bIsArray)
+	if (OwningBP->BlueprintType != BPTYPE_Interface)
 	{
-		return TEXT("Literal values are not allowed for array parameters.  Use a Make Array node instead");
-	}
-
-	if ((OwningBP->BlueprintType != BPTYPE_Interface) && !bIsAutoCreateRefTerm && bIsReference)
-	{
-		return TEXT("Literal values are not allowed for pass-by-reference parameters.");
+		if (!bIsAutoCreateRefTerm)
+		{
+			if (bIsArray)
+			{
+				return TEXT("Literal values are not allowed for array parameters.  Use a Make Array node instead");
+			}
+			else if (bIsReference)
+			{
+				return TEXT("Literal values are not allowed for pass-by-reference parameters.");
+			}
+		}
 	}
 
 	FString ReturnMsg;
