Fix a few more initialization order issues.

diff --git a/Engine/Source/Developer/CollisionAnalyzer/Private/CollisionAnalyzer.h b/Engine/Source/Developer/CollisionAnalyzer/Private/CollisionAnalyzer.h
--- a/Engine/Source/Developer/CollisionAnalyzer/Private/CollisionAnalyzer.h
+++ b/Engine/Source/Developer/CollisionAnalyzer/Private/CollisionAnalyzer.h
@@ -54,9 +54,9 @@ public:
 	int32 GetNumFramesOfRecording();
 
 	FCollisionAnalyzer()
-		: CurrentFrameNum(0)
+		: DrawBox(0)		
+		, CurrentFrameNum(0)
 		, bIsRecording(false)
-		, DrawBox(0)
 	{
 	}
 
diff --git a/Engine/Source/Editor/BlueprintGraph/Classes/EdGraphSchema_K2_Actions.h b/Engine/Source/Editor/BlueprintGraph/Classes/EdGraphSchema_K2_Actions.h
--- a/Engine/Source/Editor/BlueprintGraph/Classes/EdGraphSchema_K2_Actions.h
+++ b/Engine/Source/Editor/BlueprintGraph/Classes/EdGraphSchema_K2_Actions.h
@@ -634,7 +634,7 @@ struct BLUEPRINTGRAPH_API FEdGraphSchemaAction_K2Graph : public FEdGraphSchemaAc
 	{}
 
 	FEdGraphSchemaAction_K2Graph (EEdGraphSchemaAction_K2Graph::Type InType, const FString& InNodeCategory, const FText& InMenuDesc, const FString& InToolTip, const int32 InGrouping)
-		: FEdGraphSchemaAction(InNodeCategory, InMenuDesc, InToolTip, InGrouping), EdGraph(NULL), GraphType(InType)
+		: FEdGraphSchemaAction(InNodeCategory, InMenuDesc, InToolTip, InGrouping), GraphType(InType), EdGraph(NULL)
 	{}
 
 	// FEdGraphSchemaAction interface
