Making it so self nodes do not warn when connected with interface pins (for interfaces added in the BP editor).

diff --git a/Engine/Source/Editor/BlueprintGraph/Private/EdGraphSchema_K2.cpp b/Engine/Source/Editor/BlueprintGraph/Private/EdGraphSchema_K2.cpp
--- a/Engine/Source/Editor/BlueprintGraph/Private/EdGraphSchema_K2.cpp
+++ b/Engine/Source/Editor/BlueprintGraph/Private/EdGraphSchema_K2.cpp
@@ -2700,12 +2700,26 @@ bool UEdGraphSchema_K2::ArePinTypesCompatible(const FEdGraphPinType& Output, con
 	{
 		UClass const* InterfaceClass = Cast<UClass const>(Output.PinSubCategoryObject.Get());
 		UClass const* InputClass     = Cast<UClass const>(Input.PinSubCategoryObject.Get());
+
+		if (Input.PinSubCategory == PSC_Self)
+		{
+			InputClass = CallingContext;
+		}
+		check(InputClass != nullptr);
+
 		return InputClass->ImplementsInterface(InterfaceClass) || InterfaceClass->IsChildOf(InputClass);
 	}
 	else if ((Output.PinCategory == PC_Object) && (Input.PinCategory == PC_Interface))
 	{
 		UClass const* OutputClass    = Cast<UClass const>(Output.PinSubCategoryObject.Get());
 		UClass const* InterfaceClass = Cast<UClass const>(Input.PinSubCategoryObject.Get());
+
+		if (Output.PinSubCategory == PSC_Self)
+		{
+			OutputClass = CallingContext;
+		}
+		check(OutputClass != nullptr);
+
 		return OutputClass->ImplementsInterface(InterfaceClass) || OutputClass->IsChildOf(InterfaceClass);
 	}
 
diff --git a/Engine/Source/Editor/KismetCompiler/Private/KismetCompiler.cpp b/Engine/Source/Editor/KismetCompiler/Private/KismetCompiler.cpp
--- a/Engine/Source/Editor/KismetCompiler/Private/KismetCompiler.cpp
+++ b/Engine/Source/Editor/KismetCompiler/Private/KismetCompiler.cpp
@@ -3060,6 +3060,10 @@ void FKismetCompilerContext::Compile()
 	// Run thru the class defined variables first, get them registered
 	CreateClassVariablesFromBlueprint();
 
+	// Add any interfaces that the blueprint implements to the class
+	// (has to happen before we validate pin links in CreateFunctionList(), so that we can verify self/interface pins)
+	AddInterfacesFromBlueprint(NewClass);
+
 	// Construct a context for each function, doing validation and building the function interface
 	{	
 		SCOPE_CYCLE_COUNTER(EKismetCompilerStats_CreateFunctionList);
@@ -3178,9 +3182,6 @@ void FKismetCompilerContext::Compile()
 	{
 		SCOPE_CYCLE_COUNTER(EKismetCompilerStats_FinalizationWork);
 
-		// Add any interfaces that the blueprint implements to the class
-		AddInterfacesFromBlueprint(NewClass);
-
 		// Set any final flags and seal the class, build a CDO, etc...
 		FinishCompilingClass(NewClass);
 
