Giving duplicated graphs their own unique GUID upon duplication.
(TTP #334746)

diff --git a/Engine/Source/Editor/BlueprintGraph/Private/EdGraphSchema_K2.cpp b/Engine/Source/Editor/BlueprintGraph/Private/EdGraphSchema_K2.cpp
--- a/Engine/Source/Editor/BlueprintGraph/Private/EdGraphSchema_K2.cpp
+++ b/Engine/Source/Editor/BlueprintGraph/Private/EdGraphSchema_K2.cpp
@@ -4121,6 +4121,8 @@ UEdGraph* UEdGraphSchema_K2::DuplicateGraph(UEdGraph* GraphToDuplicate) const
 		if (NewGraph)
 		{
 			FEdGraphUtilities::RenameGraphCloseToName(NewGraph,GraphToDuplicate->GetFName().GetPlainNameString());
+			// can't have two graphs with the same guid... that'd be silly!
+			NewGraph->GraphGuid = FGuid::NewGuid();
 
 			//Rename the entry node or any further renames will not update the entry node, also fixes a duplicate node issue on compile
 			for (int32 NodeIndex = 0; NodeIndex < NewGraph->Nodes.Num(); ++NodeIndex)
