Copying (but not pasting) a Timeline will no longer cause the returns from the Timeline to be 0 in PIE.

On copy, a new Guid is no longer generated for the Timeline_Template.

#ttp 336588 - LIVE: BLUEPRINTS: Having a Timeline copied in your clipboard makes the track return 0

diff --git a/Engine/Source/Editor/BlueprintGraph/Private/K2Node_Timeline.cpp b/Engine/Source/Editor/BlueprintGraph/Private/K2Node_Timeline.cpp
--- a/Engine/Source/Editor/BlueprintGraph/Private/K2Node_Timeline.cpp
+++ b/Engine/Source/Editor/BlueprintGraph/Private/K2Node_Timeline.cpp
@@ -353,7 +353,7 @@ void UK2Node_Timeline::PrepareForCopying()
 	//Set the GUID so we can identify which timeline template the copied node should use
 	UTimelineTemplate* Template  = Blueprint->FindTimelineTemplateByVariableName(TimelineName);
 	check(Template);
-	TimelineGuid = Template->TimelineGuid =  FGuid::NewGuid(); //refresh the guid 
+	TimelineGuid = Template->TimelineGuid; // hold onto the template's Guid so on paste we can match it up on paste
 }
 
 //Determine if all the tracks contained with both arrays are identical
