Injecting a cast operation when plugging an Blueprint object pin into an interface paramter.

diff --git a/Engine/Source/Editor/BlueprintGraph/Private/CallFunctionHandler.cpp b/Engine/Source/Editor/BlueprintGraph/Private/CallFunctionHandler.cpp
--- a/Engine/Source/Editor/BlueprintGraph/Private/CallFunctionHandler.cpp
+++ b/Engine/Source/Editor/BlueprintGraph/Private/CallFunctionHandler.cpp
@@ -218,7 +218,31 @@ void FKCHandler_CallFunction::CreateFunctionCallStatement(FKismetFunctionContext
 							}
 							else
 							{
-								int32 ParameterIndex = RHSTerms.Add(*Term);
+								FBPTerminal* RHSTerm = *Term;
+
+								// if this term is an object that needs to be cast to an interface
+								if (FBPTerminal** InterfaceTerm = InterfaceTermMap.Find(PinMatch))
+								{
+									UClass* InterfaceClass = CastChecked<UClass>(PinMatch->PinType.PinSubCategoryObject.Get());
+
+									FBPTerminal* ClassTerm = new (Context.IsEventGraph() ? Context.EventGraphLocals : Context.Locals) FBPTerminal();
+									ClassTerm->Name       = InterfaceClass->GetName();
+									ClassTerm->bIsLiteral = true;
+									ClassTerm->Source     = Node;
+									ClassTerm->ObjectLiteral = InterfaceClass;
+
+									// insert a cast op before a call to the function (and replace
+									// the param with the result from the cast)
+									FBlueprintCompiledStatement& CastStatement = Context.AppendStatementForNode(Node);
+									CastStatement.Type = KCST_CastObjToInterface;
+									CastStatement.LHS = *InterfaceTerm;
+									CastStatement.RHS.Add(ClassTerm);
+									CastStatement.RHS.Add(*Term);
+
+									RHSTerm = *InterfaceTerm;
+								}
+
+								int32 ParameterIndex = RHSTerms.Add(RHSTerm);
 
 								if (PinMatch == LatentInfoPin)
 								{
@@ -467,11 +491,11 @@ UClass* FKCHandler_CallFunction::GetTrueCallingClass(FKismetFunctionContext& Con
 
 void FKCHandler_CallFunction::RegisterNets(FKismetFunctionContext& Context, UEdGraphNode* Node)
 {
-	UFunction* Function = FindFunction(Context, Node);
+	UEdGraphSchema_K2 const* K2Schema = CompilerContext.GetSchema();
 
+	UFunction* Function = FindFunction(Context, Node);
 	if (Function != NULL)
 	{
-
 		// Auto-created term defaults
 		TArray<FString> AutoCreateRefTermPinNames;
 		const bool bHasAutoCreateRefTerms = Function->HasMetaData(FBlueprintMetadata::MD_AutoCreateRefTerm);
@@ -519,8 +543,6 @@ void FKCHandler_CallFunction::RegisterNets(FKismetFunctionContext& Context, UEdG
 			UEdGraphPin* Pin = (*It);
 			const bool bIsConnected = (Pin->LinkedTo.Num() != 0);
 
-			UEdGraphSchema_K2 const* K2Schema = CompilerContext.GetSchema();
-
 			// if this pin could use a default (it doesn't have a connection or default of its own)
 			if (!bIsConnected && (Pin->DefaultObject == NULL))
 			{
@@ -573,6 +595,26 @@ void FKCHandler_CallFunction::RegisterNets(FKismetFunctionContext& Context, UEdG
 		}
 	}
 
+	for (UEdGraphPin* Pin : Node->Pins)
+	{
+		if ((Pin->Direction != EGPD_Input) || (Pin->LinkedTo.Num() == 0))
+		{
+			continue;
+		}
+
+		// if we have an object plugged into an interface pin, let's create a 
+		// term that'll be used as an intermediate, holding the result of a cast 
+		// from object to interface
+		if ((Pin->PinType.PinCategory == K2Schema->PC_Interface) && (Pin->LinkedTo[0]->PinType.PinCategory == K2Schema->PC_Object))
+		{
+			FBPTerminal* InterfaceTerm = new (Context.IsEventGraph() ? Context.EventGraphLocals : Context.Locals) FBPTerminal();
+			InterfaceTerm->CopyFromPin(Pin, Context.NetNameMap->MakeValidName(Pin->LinkedTo[0]) + TEXT("_CastResult"));
+			InterfaceTerm->Source = Node;
+
+			InterfaceTermMap.Add(Pin, InterfaceTerm);
+		}
+	}
+
 	FNodeHandlingFunctor::RegisterNets(Context, Node);
 }
 
diff --git a/Engine/Source/Editor/BlueprintGraph/Public/CallFunctionHandler.h b/Engine/Source/Editor/BlueprintGraph/Public/CallFunctionHandler.h
--- a/Engine/Source/Editor/BlueprintGraph/Public/CallFunctionHandler.h
+++ b/Engine/Source/Editor/BlueprintGraph/Public/CallFunctionHandler.h
@@ -63,4 +63,7 @@ public:
 	virtual void Compile(FKismetFunctionContext& Context, UEdGraphNode* Node) OVERRIDE;
 	virtual void CheckIfFunctionIsCallable(UFunction* Function, FKismetFunctionContext& Context, UEdGraphNode* Node);
 	virtual void AdditionalCompiledStatementHandling(FKismetFunctionContext& Context, UEdGraphNode* Node, FBlueprintCompiledStatement& Statement) {}
+
+protected:
+	TMap<UEdGraphPin*, FBPTerminal*> InterfaceTermMap;
 };
