DoOnceMultiInput Node: provide's DoOnce behavior but with multiple input/output pairs. Basically, a 'first through' node.

diff --git a/Engine/Source/Editor/BlueprintGraph/Classes/K2Node_DoOnceMultiInput.h b/Engine/Source/Editor/BlueprintGraph/Classes/K2Node_DoOnceMultiInput.h
new file mode 100644
--- /dev/null
+++ b/Engine/Source/Editor/BlueprintGraph/Classes/K2Node_DoOnceMultiInput.h
@@ -0,0 +1,55 @@
+// Copyright 1998-2014 Epic Games, Inc. All Rights Reserved.
+
+#pragma once
+
+#include "K2Node_DoOnceMultiInput.generated.h"
+
+UCLASS(MinimalAPI)
+class UK2Node_DoOnceMultiInput : public UK2Node
+{
+	GENERATED_UCLASS_BODY()
+
+	/** The number of additional input pins to generate for this node (2 base pins are not included) */
+	UPROPERTY()
+	int32 NumAdditionalInputs;
+
+	virtual void GetMenuEntries(FGraphContextMenuBuilder& ContextMenuBuilder) const OVERRIDE;
+	virtual FText GetNodeTitle(ENodeTitleType::Type TitleType) const OVERRIDE;
+
+	/** Reference to the integer that contains */
+	UPROPERTY(transient)
+	class UK2Node_TemporaryVariable* DataNode;
+
+private:
+
+	const static int32 NumBaseInputs = 1;
+
+	static int32 GetMaxInputPinsNum();
+	static FString GetNameForPin(int32 PinIndex, bool In);
+
+	FEdGraphPinType GetInType() const;
+	FEdGraphPinType GetOutType() const;
+
+	void AddPinsInner(int32 AdditionalPinIndex);
+	bool CanAddPin() const;
+	bool CanRemovePin(const UEdGraphPin* Pin) const;
+public:
+	BLUEPRINTGRAPH_API UEdGraphPin* FindOutPin() const;
+	BLUEPRINTGRAPH_API UEdGraphPin* FindSelfPin() const;
+
+	/** Get TRUE input type (self, etc.. are skipped) */
+	BLUEPRINTGRAPH_API UEdGraphPin* GetInputPin(int32 InputPinIndex);
+	BLUEPRINTGRAPH_API UEdGraphPin* GetOutputPin(int32 InputPinIndex);
+
+	BLUEPRINTGRAPH_API void AddInputPin();
+	BLUEPRINTGRAPH_API void RemoveInputPin(UEdGraphPin* Pin);
+
+	// UEdGraphNode interface
+	virtual void AllocateDefaultPins() OVERRIDE;
+	// End of UEdGraphNode interface
+
+	// UK2Node interface
+	virtual void GetContextMenuActions(const FGraphNodeContextMenuBuilder& Context) const OVERRIDE;
+	virtual void ExpandNode(class FKismetCompilerContext& CompilerContext, UEdGraph* SourceGraph) OVERRIDE;
+	// End of UK2Node interface
+};
\ No newline at end of file
diff --git a/Engine/Source/Editor/BlueprintGraph/Private/K2Node_DoOnceMultiInput.cpp b/Engine/Source/Editor/BlueprintGraph/Private/K2Node_DoOnceMultiInput.cpp
new file mode 100644
--- /dev/null
+++ b/Engine/Source/Editor/BlueprintGraph/Private/K2Node_DoOnceMultiInput.cpp
@@ -0,0 +1,378 @@
+// Copyright 1998-2014 Epic Games, Inc. All Rights Reserved.
+
+#include "BlueprintGraphPrivatePCH.h"
+
+#include "K2Node_CommutativeAssociativeBinaryOperator.h"
+#include "ScopedTransaction.h"
+#include "K2ActionMenuBuilder.h"
+#include "KismetCompiler.h"
+
+
+#define LOCTEXT_NAMESPACE "K2Node"
+
+void UK2Node_DoOnceMultiInput::GetMenuEntries(FGraphContextMenuBuilder& ContextMenuBuilder) const
+{
+	const UEdGraphSchema_K2* K2Schema = GetDefault<UEdGraphSchema_K2>();
+	EGraphType GraphType = K2Schema->GetGraphType(ContextMenuBuilder.CurrentGraph);
+
+	// Create a node template for this factory method
+	UK2Node_DoOnceMultiInput* NodeTemplate = NewObject<UK2Node_DoOnceMultiInput>(ContextMenuBuilder.OwnerOfTemporaries);
+	
+	TSharedPtr<FEdGraphSchemaAction_K2NewNode> NodeAction = FK2ActionMenuBuilder::AddNewNodeAction(ContextMenuBuilder, LOCTEXT("FlowCategory", "Flow Control").ToString(), NodeTemplate->GetNodeTitle(ENodeTitleType::ListView), NodeTemplate->GetTooltip(), 0, NodeTemplate->GetKeywords());
+
+	NodeAction->NodeTemplate = NodeTemplate;
+	NodeAction->SearchTitle = NodeTemplate->GetNodeSearchTitle();
+}
+
+FText UK2Node_DoOnceMultiInput::GetNodeTitle(ENodeTitleType::Type TitleType) const
+{
+	return LOCTEXT("DoOnceMultiInput", "DoOnce MultiInput");
+}
+
+int32 UK2Node_DoOnceMultiInput::GetMaxInputPinsNum()
+{
+	return (TCHAR('Z') - TCHAR('A'));
+}
+
+FString UK2Node_DoOnceMultiInput::GetNameForPin(int32 PinIndex, bool In)
+{
+	check(PinIndex < GetMaxInputPinsNum());
+	FString Name;
+	Name.AppendChar(TCHAR('A') + PinIndex);
+	Name = In ? Name + LOCTEXT("DoOnceMultiIn", " In").ToString() : Name + LOCTEXT("DoOnceMultiOut", " Out").ToString();
+	return Name;
+}
+
+UK2Node_DoOnceMultiInput::UK2Node_DoOnceMultiInput(const class FPostConstructInitializeProperties& PCIP)
+	: Super(PCIP)
+{
+	NumAdditionalInputs = 0;
+}
+
+UEdGraphPin* UK2Node_DoOnceMultiInput::FindOutPin() const
+{
+	for(int32 PinIdx=0; PinIdx<Pins.Num(); PinIdx++)
+	{
+		if(EEdGraphPinDirection::EGPD_Output == Pins[PinIdx]->Direction)
+		{
+			return Pins[PinIdx];
+		}
+	}
+	return NULL;
+}
+
+UEdGraphPin* UK2Node_DoOnceMultiInput::FindSelfPin() const
+{
+	const UEdGraphSchema_K2* K2Schema = GetDefault<UEdGraphSchema_K2>();
+	for(int32 PinIdx=0; PinIdx<Pins.Num(); PinIdx++)
+	{
+		if(Pins[PinIdx]->PinName == K2Schema->PN_Self)
+		{
+			return Pins[PinIdx];
+		}
+	}
+	return NULL;
+}
+
+bool UK2Node_DoOnceMultiInput::CanAddPin() const
+{
+	return (NumAdditionalInputs) < GetMaxInputPinsNum();
+}
+
+bool UK2Node_DoOnceMultiInput::CanRemovePin(const UEdGraphPin* Pin) const
+{
+	return (
+		Pin &&
+		NumAdditionalInputs &&
+		(INDEX_NONE != Pins.IndexOfByKey(Pin)) &&
+		(EEdGraphPinDirection::EGPD_Input == Pin->Direction)
+	);
+}
+
+UEdGraphPin* UK2Node_DoOnceMultiInput::GetInputPin(int32 InputPinIndex)
+{
+	const UEdGraphPin* SelfPin = FindSelfPin();
+
+	int32 CurrentInputIndex = 0;
+	for(int32 PinIdx=0; PinIdx<Pins.Num(); PinIdx++)
+	{
+		UEdGraphPin* CurrentPin = Pins[PinIdx];
+		if ((CurrentPin != SelfPin) && CurrentPin->Direction == EGPD_Input)
+		{
+			if(CurrentInputIndex == InputPinIndex)
+			{
+				return CurrentPin;
+			}
+			CurrentInputIndex++;
+		}
+	}
+	return NULL;
+}
+
+UEdGraphPin* UK2Node_DoOnceMultiInput::GetOutputPin(int32 InputPinIndex)
+{
+	const UEdGraphPin* SelfPin = FindSelfPin();
+
+	int32 CurrentInputIndex = 0;
+	for (int32 PinIdx = 0; PinIdx<Pins.Num(); PinIdx++)
+	{
+		UEdGraphPin* CurrentPin = Pins[PinIdx];
+		if ((CurrentPin != SelfPin) && CurrentPin->Direction == EGPD_Output)
+		{
+			if (CurrentInputIndex == InputPinIndex)
+			{
+				return CurrentPin;
+			}
+			CurrentInputIndex++;
+		}
+	}
+	return NULL;
+}
+
+FEdGraphPinType UK2Node_DoOnceMultiInput::GetInType() const
+{
+	for (int32 PinIt = 0; PinIt < Pins.Num(); PinIt++)
+	{
+		if (Pins[PinIt] != FindSelfPin())
+		{
+			return Pins[PinIt]->PinType;
+		}
+	}
+	return FEdGraphPinType();
+}
+
+FEdGraphPinType UK2Node_DoOnceMultiInput::GetOutType() const
+{
+	for (int32 PinIt = 0; PinIt < Pins.Num(); PinIt++)
+	{
+		if (Pins[PinIt] != FindOutPin())
+		{
+			return Pins[PinIt]->PinType;
+		}
+	}
+	return FEdGraphPinType();
+}
+
+void UK2Node_DoOnceMultiInput::AllocateDefaultPins()
+{
+	Super::AllocateDefaultPins();
+
+	const UEdGraphSchema_K2* K2Schema = GetDefault<UEdGraphSchema_K2>();
+
+	
+	CreatePin(EGPD_Input, K2Schema->PC_Exec, TEXT(""), NULL, false, false, GetNameForPin(0, true));
+	CreatePin(EGPD_Output, K2Schema->PC_Exec, TEXT(""), NULL, false, false, GetNameForPin(0, false));
+
+	CreatePin(EGPD_Input, K2Schema->PC_Exec, TEXT(""), NULL, false, false, LOCTEXT("DoOnceResetIn", "Reset In").ToString());
+	CreatePin(EGPD_Output, K2Schema->PC_Exec, TEXT(""), NULL, false, false, LOCTEXT("DoOnceResetOut", "Reset Out").ToString());
+
+	for (int32 i = 0; i < NumAdditionalInputs; ++i)
+	{
+		AddPinsInner(i+1);
+	}
+}
+
+void UK2Node_DoOnceMultiInput::AddPinsInner(int32 AdditionalPinIndex)
+{
+	{
+		const FEdGraphPinType InputType = GetInType();
+		CreatePin(EGPD_Input, 
+			InputType.PinCategory, 
+			InputType.PinSubCategory, 
+			InputType.PinSubCategoryObject.Get(), 
+			InputType.bIsArray, 
+			InputType.bIsReference, 
+			*GetNameForPin(AdditionalPinIndex, true)
+		);
+	}
+
+	{
+		const FEdGraphPinType OutputType = GetOutType();
+		CreatePin(EGPD_Output,
+			OutputType.PinCategory,
+			OutputType.PinSubCategory,
+			OutputType.PinSubCategoryObject.Get(),
+			OutputType.bIsArray,
+			OutputType.bIsReference,
+			*GetNameForPin(AdditionalPinIndex, false)
+		);
+	}
+}
+
+void UK2Node_DoOnceMultiInput::AddInputPin()
+{
+	if(CanAddPin())
+	{
+		FScopedTransaction Transaction( LOCTEXT("AddPinTx", "AddPin") );
+		Modify();
+
+		AddPinsInner(NumAdditionalInputs + NumBaseInputs);
+		++NumAdditionalInputs;
+	
+		FBlueprintEditorUtils::MarkBlueprintAsStructurallyModified(GetBlueprint());
+	}
+}
+
+void UK2Node_DoOnceMultiInput::RemoveInputPin(UEdGraphPin* Pin)
+{
+	if(CanRemovePin(Pin))
+	{
+		FScopedTransaction Transaction( LOCTEXT("RemovePinTx", "RemovePin") );
+		Modify();
+
+		int32 PinRemovalIndex = INDEX_NONE;
+		if (Pins.Find(Pin, /*out*/ PinRemovalIndex))
+		{
+			Pins.RemoveAt(PinRemovalIndex);
+			Pin->Modify();
+			Pin->BreakAllPinLinks();
+			--NumAdditionalInputs;
+
+			int32 NameIndex = 0;
+			const UEdGraphPin* OutPin = FindOutPin();
+			const UEdGraphPin* SelfPin = FindSelfPin();
+			for (int32 PinIndex = 0; PinIndex < Pins.Num(); ++PinIndex)
+			{
+				UEdGraphPin* LocalPin = Pins[PinIndex];
+				if(LocalPin && (LocalPin != OutPin) && (LocalPin != SelfPin))
+				{
+					const FString PinName = GetNameForPin(NameIndex + NumBaseInputs, true);  // FIXME
+					if(PinName != LocalPin->PinName)
+					{
+						LocalPin->Modify();
+						LocalPin->PinName = PinName;
+					}
+					NameIndex++;
+				}
+			}
+			FBlueprintEditorUtils::MarkBlueprintAsStructurallyModified(GetBlueprint());
+		}
+	}
+}
+
+void UK2Node_DoOnceMultiInput::GetContextMenuActions(const FGraphNodeContextMenuBuilder& Context) const
+{
+	Super::GetContextMenuActions(Context);
+
+	if (!Context.bIsDebugging)
+	{
+		static FName CommutativeAssociativeBinaryOperatorNodeName = FName("CommutativeAssociativeBinaryOperatorNode");
+		FText CommutativeAssociativeBinaryOperatorStr = LOCTEXT("CommutativeAssociativeBinaryOperatorNode", "Operator Node");
+		if (Context.Pin != NULL)
+		{
+			if(CanRemovePin(Context.Pin))
+			{
+				Context.MenuBuilder->BeginSection(CommutativeAssociativeBinaryOperatorNodeName, CommutativeAssociativeBinaryOperatorStr);
+				Context.MenuBuilder->AddMenuEntry(
+					LOCTEXT("RemovePin", "Remove pin"),
+					LOCTEXT("RemovePinTooltip", "Remove this input pin"),
+					FSlateIcon(),
+					FUIAction(
+						FExecuteAction::CreateUObject(this, &UK2Node_DoOnceMultiInput::RemoveInputPin, const_cast<UEdGraphPin*>(Context.Pin))
+					)
+				);
+				Context.MenuBuilder->EndSection();
+			}
+		}
+		else if(CanAddPin())
+		{
+			Context.MenuBuilder->BeginSection(CommutativeAssociativeBinaryOperatorNodeName, CommutativeAssociativeBinaryOperatorStr);
+			Context.MenuBuilder->AddMenuEntry(
+				LOCTEXT("AddPin", "Add pin"),
+				LOCTEXT("AddPinTooltip", "Add another input pin"),
+				FSlateIcon(),
+				FUIAction(
+					FExecuteAction::CreateUObject(this, &UK2Node_DoOnceMultiInput::AddInputPin)
+				)
+			);
+			Context.MenuBuilder->EndSection();
+		}
+	}
+}
+
+void UK2Node_DoOnceMultiInput::ExpandNode(FKismetCompilerContext& CompilerContext, UEdGraph* SourceGraph)
+{
+	Super::ExpandNode(CompilerContext, SourceGraph);
+
+	if (CompilerContext.bIsFullCompile)
+	{
+		const UEdGraphSchema_K2* Schema = CompilerContext.GetSchema();
+
+		/////////////////////////////
+		// Temporary Variable node
+		/////////////////////////////
+
+		// Create the node
+		UK2Node_TemporaryVariable* TempVarNode = SourceGraph->CreateBlankNode<UK2Node_TemporaryVariable>();
+		TempVarNode->VariableType.PinCategory = Schema->PC_Boolean;
+		TempVarNode->AllocateDefaultPins();
+		CompilerContext.MessageLog.NotifyIntermediateObjectCreation(TempVarNode, this);
+		// Give a reference of the variable node to the multi gate node
+		DataNode = TempVarNode;
+
+		// Create the conditional node we're replacing the enum node for
+
+		for (int32 idx = 0; idx < NumBaseInputs + NumAdditionalInputs + 1 /*ResetPin*/; ++idx)
+		{
+			UEdGraphPin* ExecPin = GetInputPin(idx);
+			UEdGraphPin* ThenPin = GetOutputPin(idx);
+
+			check(ExecPin);
+			check(ThenPin);
+
+			// AssignmentNode
+			UK2Node_AssignmentStatement* AssignmentNode = SourceGraph->CreateBlankNode<UK2Node_AssignmentStatement>();
+			AssignmentNode->AllocateDefaultPins();
+			CompilerContext.MessageLog.NotifyIntermediateObjectCreation(AssignmentNode, this);
+			AssignmentNode->GetVariablePin()->PinType = TempVarNode->GetVariablePin()->PinType;
+			AssignmentNode->GetVariablePin()->MakeLinkTo(TempVarNode->GetVariablePin());
+			AssignmentNode->GetValuePin()->PinType = TempVarNode->GetVariablePin()->PinType;
+			
+			if (!ExecPin->PinName.Contains(TEXT("Reset"))) // Fixme this wont work for localization
+			{
+				// BranchNode
+				UK2Node_IfThenElse* BranchNode = SourceGraph->CreateBlankNode<UK2Node_IfThenElse>();
+				BranchNode->AllocateDefaultPins();
+				CompilerContext.MessageLog.NotifyIntermediateObjectCreation(BranchNode, this);
+
+				// -------------------------------------------------
+
+				// Coerce the wildcards pin types (set the default of the value to 0)
+				AssignmentNode->GetValuePin()->DefaultValue = TEXT("1");
+
+				// -------------------------------------------------
+
+				// Link Tempvariable with the branch condtional
+				Schema->TryCreateConnection(TempVarNode->GetVariablePin(), BranchNode->GetConditionPin());
+			
+				// Link our input exec pin into the branch node
+				CompilerContext.MovePinLinksToIntermediate(*ExecPin, *BranchNode->GetExecPin());
+
+				// link branch else (false) to assigment node (set temp variable to true)
+				Schema->TryCreateConnection(BranchNode->GetElsePin(), AssignmentNode->GetExecPin());
+
+				// link set temp variable node to our ouput then pin
+				CompilerContext.MovePinLinksToIntermediate(*ThenPin, *AssignmentNode->GetThenPin());
+
+			}
+			else
+			{
+				// Coerce the wildcards pin types (set the default of the value to 1)				
+				AssignmentNode->GetValuePin()->DefaultValue = TEXT("0");
+
+				// -------------------------------------------------
+
+				// Link our input exec pin into the branch node
+				CompilerContext.MovePinLinksToIntermediate(*ExecPin, *AssignmentNode->GetExecPin());
+
+				// link set temp variable node to our ouput then pin
+				CompilerContext.MovePinLinksToIntermediate(*ThenPin, *AssignmentNode->GetThenPin());
+			}
+		}
+
+		// Break all links to the Select node so it goes away for at scheduling time
+		BreakAllNodeLinks();
+	}
+}
+
+#undef LOCTEXT_NAMESPACE
diff --git a/Engine/Source/Editor/BlueprintGraph/Public/BlueprintGraphDefinitions.h b/Engine/Source/Editor/BlueprintGraph/Public/BlueprintGraphDefinitions.h
--- a/Engine/Source/Editor/BlueprintGraph/Public/BlueprintGraphDefinitions.h
+++ b/Engine/Source/Editor/BlueprintGraph/Public/BlueprintGraphDefinitions.h
@@ -18,6 +18,7 @@
 #include "K2Node_CallParentFunction.h"
 #include "K2Node_ClearDelegate.h"
 #include "K2Node_CommutativeAssociativeBinaryOperator.h"
+#include "K2Node_DoOnceMultiInput.h"
 #include "K2Node_ComponentBoundEvent.h"
 #include "K2Node_Composite.h"
 #include "K2Node_CreateDelegate.h"
diff --git a/Engine/Source/Editor/GraphEditor/Private/KismetNodes/SGraphNodeK2Sequence.cpp b/Engine/Source/Editor/GraphEditor/Private/KismetNodes/SGraphNodeK2Sequence.cpp
--- a/Engine/Source/Editor/GraphEditor/Private/KismetNodes/SGraphNodeK2Sequence.cpp
+++ b/Engine/Source/Editor/GraphEditor/Private/KismetNodes/SGraphNodeK2Sequence.cpp
@@ -49,6 +49,12 @@ FReply SGraphNodeK2Sequence::OnAddPin()
 		UpdateGraphNode();
 		GraphNode->GetGraph()->NotifyGraphChanged();
 	}
+	else if (UK2Node_DoOnceMultiInput* DoOnceMultiNode = Cast<UK2Node_DoOnceMultiInput>(GraphNode))
+	{
+		DoOnceMultiNode->AddInputPin();
+		UpdateGraphNode();
+		GraphNode->GetGraph()->NotifyGraphChanged();
+	}
 
 	return FReply::Handled();
 }
\ No newline at end of file
diff --git a/Engine/Source/Editor/GraphEditor/Private/NodeFactory.cpp b/Engine/Source/Editor/GraphEditor/Private/NodeFactory.cpp
--- a/Engine/Source/Editor/GraphEditor/Private/NodeFactory.cpp
+++ b/Engine/Source/Editor/GraphEditor/Private/NodeFactory.cpp
@@ -157,6 +157,10 @@ TSharedPtr<SGraphNode> FNodeFactory::CreateNodeWidget(UEdGraphNode* InNode)
 		{
 			return SNew(SGraphNodeK2Sequence, OperatorNode);
 		}
+		else if (UK2Node_DoOnceMultiInput* DoOnceMultiInputNode = Cast<UK2Node_DoOnceMultiInput>(InNode))
+		{
+			return SNew(SGraphNodeK2Sequence, DoOnceMultiInputNode);
+		}
 		else if (UK2Node_Timeline* TimelineNode = Cast<UK2Node_Timeline>(InNode))
 		{
 			return SNew(SGraphNodeK2Timeline, TimelineNode);
