BP Cyclic Dependency fix. ttp#337588

diff --git a/Engine/Source/Editor/BlueprintGraph/Private/EdGraphSchema_K2.cpp b/Engine/Source/Editor/BlueprintGraph/Private/EdGraphSchema_K2.cpp
--- a/Engine/Source/Editor/BlueprintGraph/Private/EdGraphSchema_K2.cpp
+++ b/Engine/Source/Editor/BlueprintGraph/Private/EdGraphSchema_K2.cpp
@@ -3940,12 +3940,14 @@ UEdGraphNode* UEdGraphSchema_K2::CreateSubstituteNode(UEdGraphNode* Node, const
 			FBlueprintEditorUtils::GetFunctionNameList(Blueprint, ExistingNamesInUse);
 			FBlueprintEditorUtils::GetClassVariableList(Blueprint, ExistingNamesInUse);
 
+			const ERenameFlags RenameFlags = (Blueprint->bIsRegeneratingOnLoad ? REN_ForceNoResetLoaders : 0);
+
 			// Allow the old object name to be used in the graph
 			FName ObjName = EventNode->GetFName();
 			UObject* Found = FindObject<UObject>(EventNode->GetOuter(), *ObjName.ToString());
 			if(Found)
 			{
-				Found->Rename(NULL, NULL, REN_None);
+				Found->Rename(NULL, NULL, REN_DontCreateRedirectors | RenameFlags);
 			}
 
 			// Create a custom event node to replace the original event node imported from text
@@ -4030,7 +4032,7 @@ UEdGraphNode* UEdGraphSchema_K2::CreateSubstituteNode(UEdGraphNode* Node, const
 				check(Pin);
 
 				// Reparent the pin to the new custom event node
-				Pin->Rename(*Pin->GetName(), CustomEventNode);
+				Pin->Rename(*Pin->GetName(), CustomEventNode, RenameFlags);
 
 				// Don't include execution or delegate output pins as user-defined pins
 				if(!IsExecPin(*Pin) && !IsDelegateCategory(Pin->PinType.PinCategory))
diff --git a/Engine/Source/Editor/UnrealEd/Private/Kismet2/BlueprintEditorUtils.cpp b/Engine/Source/Editor/UnrealEd/Private/Kismet2/BlueprintEditorUtils.cpp
--- a/Engine/Source/Editor/UnrealEd/Private/Kismet2/BlueprintEditorUtils.cpp
+++ b/Engine/Source/Editor/UnrealEd/Private/Kismet2/BlueprintEditorUtils.cpp
@@ -340,7 +340,11 @@ void FBlueprintEditorUtils::PreloadMembers(UObject* InObject)
 		UObject* CurrentObject = *it;
 		if( CurrentObject->HasAnyFlags(RF_NeedLoad) )
 		{
-			CurrentObject->GetLinker()->Preload(CurrentObject);
+			auto Linker = CurrentObject->GetLinker();
+			if (Linker)
+			{
+				Linker->Preload(CurrentObject);
+			}
 			PreloadMembers(CurrentObject);
 		}
 	}
@@ -938,6 +942,8 @@ UClass* FBlueprintEditorUtils::RegenerateBlueprintClass(UBlueprint* Blueprint, U
 
 	bool bNeedsSkelRefRemoval = false;
 
+	FBlueprintEditorUtils::PreloadMembers(Blueprint);
+
 	if( ShouldRegenerateBlueprint(Blueprint) && !Blueprint->bHasBeenRegenerated )
 	{
 		Blueprint->bIsRegeneratingOnLoad = true;
