When checking pin compatibility, only using the context object when the self pin's type doesn't already have the class filled in.
(TTP #337679)

diff --git a/Engine/Source/Editor/BlueprintGraph/Private/EdGraphSchema_K2.cpp b/Engine/Source/Editor/BlueprintGraph/Private/EdGraphSchema_K2.cpp
--- a/Engine/Source/Editor/BlueprintGraph/Private/EdGraphSchema_K2.cpp
+++ b/Engine/Source/Editor/BlueprintGraph/Private/EdGraphSchema_K2.cpp
@@ -2701,7 +2701,7 @@ bool UEdGraphSchema_K2::ArePinTypesCompatible(const FEdGraphPinType& Output, con
 		UClass const* InterfaceClass = Cast<UClass const>(Output.PinSubCategoryObject.Get());
 		UClass const* InputClass     = Cast<UClass const>(Input.PinSubCategoryObject.Get());
 
-		if (Input.PinSubCategory == PSC_Self)
+		if ((InputClass == nullptr) && (Input.PinSubCategory == PSC_Self))
 		{
 			InputClass = CallingContext;
 		}
@@ -2714,7 +2714,7 @@ bool UEdGraphSchema_K2::ArePinTypesCompatible(const FEdGraphPinType& Output, con
 		UClass const* OutputClass    = Cast<UClass const>(Output.PinSubCategoryObject.Get());
 		UClass const* InterfaceClass = Cast<UClass const>(Input.PinSubCategoryObject.Get());
 
-		if (Output.PinSubCategory == PSC_Self)
+		if ((OutputClass == nullptr) && (Output.PinSubCategory == PSC_Self))
 		{
 			OutputClass = CallingContext;
 		}
