When using the Format node in Blueprints, the Blueprint will no longer fail to compile when running the editor with -game.

Blueprint was reporting errors due to the argument pin name being invariant, invariant texts should not be reported on these nodes when not in the editor.

diff --git a/Engine/Source/Editor/BlueprintGraph/Private/EdGraphSchema_K2.cpp b/Engine/Source/Editor/BlueprintGraph/Private/EdGraphSchema_K2.cpp
--- a/Engine/Source/Editor/BlueprintGraph/Private/EdGraphSchema_K2.cpp
+++ b/Engine/Source/Editor/BlueprintGraph/Private/EdGraphSchema_K2.cpp
@@ -1849,9 +1849,13 @@ bool UEdGraphSchema_K2::DefaultValueSimpleValidation(const FEdGraphPinType& PinT
 			DVSV_RETURN_MSG( TEXT("Invalid text literal, text is transient!") );
 		}
 
-		if( InNewDefaultText.IsCultureInvariant() )
+		// It should not be possible to get a culture invariant text in editor
+		if(GIsEditor)
 		{
-			DVSV_RETURN_MSG( TEXT("Invalid text literal, text is culture invariant!") );
+			if( InNewDefaultText.IsCultureInvariant() )
+			{
+				DVSV_RETURN_MSG( TEXT("Invalid text literal, text is culture invariant!") );
+			}
 		}
 	}
 	else if (PinCategory == PC_Struct)
