Adding Blueprint support for getting row information from CurveTables and DataTables

diff --git a/Engine/Source/Editor/BlueprintGraph/Classes/K2Node_GetDataTableRow.h b/Engine/Source/Editor/BlueprintGraph/Classes/K2Node_GetDataTableRow.h
new file mode 100644
--- /dev/null
+++ b/Engine/Source/Editor/BlueprintGraph/Classes/K2Node_GetDataTableRow.h
@@ -0,0 +1,67 @@
+// Copyright 1998-2014 Epic Games, Inc. All Rights Reserved.
+
+
+#pragma once
+#include "K2Node.h"
+#include "K2Node_GetDataTableRow.generated.h"
+
+UCLASS()
+class UK2Node_GetDataTableRow : public UK2Node
+{
+	GENERATED_UCLASS_BODY()
+
+	UPROPERTY()
+	FDataTableRowHandle DataTableRowHandle;
+
+	// Begin UEdGraphNode interface.
+	virtual void AllocateDefaultPins() OVERRIDE;
+	virtual FText GetNodeTitle(ENodeTitleType::Type TitleType) const OVERRIDE;
+	virtual FString GetNodeNativeTitle(ENodeTitleType::Type TitleType) const OVERRIDE;
+	virtual void PinDefaultValueChanged(UEdGraphPin* Pin) OVERRIDE;
+	virtual FString GetTooltip() const OVERRIDE;
+	virtual void ExpandNode(class FKismetCompilerContext& CompilerContext, UEdGraph* SourceGraph) OVERRIDE;
+	virtual FName GetPaletteIcon(FLinearColor& OutColor) const OVERRIDE
+    {
+        OutColor = GetNodeTitleColor();
+        return TEXT("Kismet.AllClasses.FunctionIcon");
+    }
+	// End UEdGraphNode interface.
+
+	// Begin UK2Node interface
+	virtual bool IsNodeSafeToIgnore() const OVERRIDE { return true; }
+	virtual void ReallocatePinsDuringReconstruction(TArray<UEdGraphPin*>& OldPins) OVERRIDE;
+	// End UK2Node interface
+
+
+	/** Create new pins to show properties on archetype */
+	void SetReturnTypeForStruct(UScriptStruct* InClass);
+	/** See if this is a spawn variable pin, or a 'default' pin */
+	bool IsDataTablePin(UEdGraphPin* Pin);
+
+	/** Get the then output pin */
+	UEdGraphPin* GetThenPin() const;
+	/** Get the Data Table input pin */
+	UEdGraphPin* GetDataTablePin(const TArray<UEdGraphPin*>* InPinsToSearch=NULL) const;
+	/** Get the spawn transform input pin */	
+	UEdGraphPin* GetRowNamePin() const;
+    /** Get the exec output pin for when the row was not found */
+	UEdGraphPin* GetRowNotFoundPin() const;
+	/** Get the result output pin */
+	UEdGraphPin* GetResultPin() const;
+
+	/** Get the type of the TableRow to return */
+	UScriptStruct* GetDataTableRowStructType(const TArray<UEdGraphPin*>* InPinsToSearch = NULL) const;
+
+private:
+	/**
+	 * Takes the specified "MutatablePin" and sets its 'PinToolTip' field (according
+	 * to the specified description)
+	 * 
+	 * @param   MutatablePin	The pin you want to set tool-tip text on
+	 * @param   PinDescription	A string describing the pin's purpose
+	 */
+	void SetPinToolTip(UEdGraphPin& MutatablePin, const FText& PinDescription) const;
+
+	/** Tooltip text for this node. */
+	FString NodeTooltip;
+};
diff --git a/Engine/Source/Editor/BlueprintGraph/Private/K2ActionMenuBuilder.cpp b/Engine/Source/Editor/BlueprintGraph/Private/K2ActionMenuBuilder.cpp
--- a/Engine/Source/Editor/BlueprintGraph/Private/K2ActionMenuBuilder.cpp
+++ b/Engine/Source/Editor/BlueprintGraph/Private/K2ActionMenuBuilder.cpp
@@ -133,6 +133,19 @@ static void AddSpawnActorNodeAction(FGraphActionListBuilderBase& ContextMenuBuil
 	}
 }
 
+/** Utility for getting a TableRow from a DataTable */
+static void AddGetDataTableRowNodeAction(FGraphActionListBuilderBase& ContextMenuBuilder, const FString& FunctionCategory)
+{
+	FString const GetDataTableRowCategory = FunctionCategory + K2ActionCategories::SubCategoryDelim + K2ActionCategories::UtilitiesCategory;
+    
+	// Add the new GetDataTableRow node
+	{
+		UK2Node* NodeTemplate = ContextMenuBuilder.CreateTemplateNode<UK2Node_GetDataTableRow>();
+		TSharedPtr<FEdGraphSchemaAction_K2NewNode> Action = FK2ActionMenuBuilder::AddNewNodeAction(ContextMenuBuilder, GetDataTableRowCategory, LOCTEXT("GetDataTableRow", "Get Data Table Row"), NodeTemplate->GetTooltip(), 0, NodeTemplate->GetKeywords());
+		Action->NodeTemplate = NodeTemplate;
+	}
+}
+
 /** Gets all the input action and axis events and nodes */
 static void GetInputNodes(FGraphActionListBuilderBase& ActionMenuBuilder, const bool bIncludeEvents)
 {
@@ -838,6 +851,8 @@ void FK2ActionMenuBuilder::GetContextAllowedNodeTypes(FBlueprintGraphActionListB
 				GetInputNodes(ContextMenuBuilder, true);
 			}
 		}
+        
+        AddGetDataTableRowNodeAction(ContextMenuBuilder, K2ActionCategories::GenericFunctionCategory);
 
 		// Add struct make/break nodes
 		GetStructActions( ContextMenuBuilder );
@@ -1145,6 +1160,8 @@ void FK2ActionMenuBuilder::GetPinAllowedNodeTypes(FBlueprintGraphActionListBuild
 					GetInputNodes(ContextMenuBuilder, bIncludeEvents);
 				}
 			}
+            
+            AddGetDataTableRowNodeAction(ContextMenuBuilder, K2ActionCategories::GenericFunctionCategory);
 
 			// for output pins, add macro flow control as a connection option
 			if ( FromPin.Direction == EGPD_Output )
diff --git a/Engine/Source/Editor/BlueprintGraph/Private/K2Node_GetDataTableRow.cpp b/Engine/Source/Editor/BlueprintGraph/Private/K2Node_GetDataTableRow.cpp
new file mode 100644
--- /dev/null
+++ b/Engine/Source/Editor/BlueprintGraph/Private/K2Node_GetDataTableRow.cpp
@@ -0,0 +1,323 @@
+// Copyright 1998-2014 Epic Games, Inc. All Rights Reserved.
+
+#include "BlueprintGraphPrivatePCH.h"
+#include "KismetCompiler.h"
+//#include "Runtime/Engine/Classes/Kismet/GameplayStatics.h"
+#include "Kismet/DataTableFunctionLibrary.h"
+
+#define LOCTEXT_NAMESPACE "K2Node_GetDataTableRow"
+
+struct UK2Node_GetDataTableRowHelper
+{
+	static FString DataTablePinName;
+	static FString RowNamePinName;
+	static FString RowNotFoundPinName;
+};
+
+FString UK2Node_GetDataTableRowHelper::DataTablePinName(LOCTEXT("DataTablePinName","DataTable").ToString());
+FString UK2Node_GetDataTableRowHelper::RowNotFoundPinName(LOCTEXT("RowNotFoundPinName","RowNotFound").ToString());
+FString UK2Node_GetDataTableRowHelper::RowNamePinName(LOCTEXT("RowNamePinName","RowName").ToString());
+
+UK2Node_GetDataTableRow::UK2Node_GetDataTableRow(const class FPostConstructInitializeProperties& PCIP)
+	: Super(PCIP)
+{
+	NodeTooltip = LOCTEXT("NodeTooltip", "Attempts to retrieve a TableRow from a DataTable via it's RowName").ToString();
+}
+
+void UK2Node_GetDataTableRow::AllocateDefaultPins()
+{
+	const UEdGraphSchema_K2* K2Schema = GetDefault<UEdGraphSchema_K2>();
+
+	// Add execution pins
+	CreatePin(EGPD_Input, K2Schema->PC_Exec, TEXT(""), NULL, false, false, K2Schema->PN_Execute);
+	CreatePin(EGPD_Output, K2Schema->PC_Exec, TEXT(""), NULL, false, false, K2Schema->PN_Then);
+	CreatePin(EGPD_Output, K2Schema->PC_Exec, TEXT(""), NULL, false, false, UK2Node_GetDataTableRowHelper::RowNotFoundPinName);
+
+	// Add DataTable pin
+	UEdGraphPin* DataTablePin = CreatePin(EGPD_Input, K2Schema->PC_Object, TEXT(""), UDataTable::StaticClass(), false, false, UK2Node_GetDataTableRowHelper::DataTablePinName);
+	SetPinToolTip(*DataTablePin, LOCTEXT("DataTablePinDescription", "The DataTable you want to retreive a row from"));
+
+	// Row Name pin
+	UEdGraphPin* TransformPin = CreatePin(EGPD_Input, K2Schema->PC_Name, TEXT(""), NULL, false, false, UK2Node_GetDataTableRowHelper::RowNamePinName);
+	SetPinToolTip(*TransformPin, LOCTEXT("RowNamePinDescription", "The name of the row to retrieve from the DataTable"));
+
+	// Result pin
+	UEdGraphPin* ResultPin = CreatePin(EGPD_Output, K2Schema->PC_Struct, TEXT(""), FTableRowBase::StaticStruct(), false, false, K2Schema->PN_ReturnValue);
+	SetPinToolTip(*ResultPin, LOCTEXT("ResultPinDescription", "The returned TableRow, if found"));
+
+	Super::AllocateDefaultPins();
+}
+
+void UK2Node_GetDataTableRow::SetPinToolTip(UEdGraphPin& MutatablePin, const FText& PinDescription) const
+{
+	MutatablePin.PinToolTip = UEdGraphSchema_K2::TypeToString(MutatablePin.PinType);
+
+	UEdGraphSchema_K2 const* const K2Schema = Cast<const UEdGraphSchema_K2>(GetSchema());
+	if (K2Schema != nullptr)
+	{
+		MutatablePin.PinToolTip += TEXT(" ");
+		MutatablePin.PinToolTip += K2Schema->GetPinDisplayName(&MutatablePin);
+	}
+
+	MutatablePin.PinToolTip += FString(TEXT("\n")) + PinDescription.ToString();
+}
+
+
+void UK2Node_GetDataTableRow::SetReturnTypeForStruct(UScriptStruct* RowStruct)
+{
+	check(RowStruct != NULL);
+
+    // Change class of output pin
+	UEdGraphPin* ResultPin = GetResultPin();
+	ResultPin->PinType.PinSubCategoryObject = RowStruct;
+}
+
+UScriptStruct* UK2Node_GetDataTableRow::GetDataTableRowStructType(const TArray<UEdGraphPin*>* InPinsToSearch /*=NULL*/) const
+{
+	UScriptStruct* RowStructType = NULL;
+	const TArray<UEdGraphPin*>* PinsToSearch = InPinsToSearch ? InPinsToSearch : &Pins;
+
+	UEdGraphPin* DataTablePin = GetDataTablePin(PinsToSearch);
+	if(DataTablePin && DataTablePin->DefaultObject != NULL && DataTablePin->LinkedTo.Num() == 0)
+	{
+		if (DataTablePin->DefaultObject->IsA(UDataTable::StaticClass()))
+		{
+			UDataTable* DataTable = (UDataTable*)DataTablePin->DefaultObject;
+			RowStructType = DataTable->RowStruct;
+		}
+	}
+
+	return RowStructType;
+}
+
+
+void UK2Node_GetDataTableRow::ReallocatePinsDuringReconstruction(TArray<UEdGraphPin*>& OldPins) 
+{
+	Super::ReallocatePinsDuringReconstruction(OldPins);
+		
+	UScriptStruct* RowStruct = GetDataTableRowStructType(&OldPins);
+	if( RowStruct != NULL )
+	{
+		SetReturnTypeForStruct(RowStruct);
+	}
+}
+
+bool UK2Node_GetDataTableRow::IsDataTablePin(UEdGraphPin* Pin)
+{
+	const UEdGraphSchema_K2* K2Schema = GetDefault<UEdGraphSchema_K2>();
+
+	return(	Pin->PinName != K2Schema->PN_Execute &&
+			Pin->PinName != K2Schema->PN_Then &&
+			Pin->PinName != K2Schema->PN_ReturnValue &&
+			Pin->PinName != UK2Node_GetDataTableRowHelper::DataTablePinName &&
+			Pin->PinName != UK2Node_GetDataTableRowHelper::RowNotFoundPinName &&
+			Pin->PinName != UK2Node_GetDataTableRowHelper::RowNamePinName );
+}
+
+
+void UK2Node_GetDataTableRow::PinDefaultValueChanged(UEdGraphPin* ChangedPin) 
+{
+	if (ChangedPin->PinName == UK2Node_GetDataTableRowHelper::DataTablePinName)
+	{
+		const UEdGraphSchema_K2* K2Schema = GetDefault<UEdGraphSchema_K2>();
+
+		// Because the archetype has changed, we break the output link as the output pin type will change
+		UEdGraphPin* ResultPin = GetResultPin();
+		ResultPin->BreakAllPinLinks();
+
+		// Remove all pins related to archetype variables
+		TArray<UEdGraphPin*> OldPins = Pins;
+		for (int32 i = 0; i < OldPins.Num(); i++)
+		{
+			UEdGraphPin* OldPin = OldPins[i];
+			if (IsDataTablePin(OldPin))
+			{
+				OldPin->BreakAllPinLinks();
+				Pins.Remove(OldPin);
+			}
+		}
+
+		UScriptStruct* RowStruct = GetDataTableRowStructType();
+		if (RowStruct != NULL)
+		{
+			SetReturnTypeForStruct(RowStruct);
+		}
+
+		// Refresh the UI for the graph so the pin changes show up
+		UEdGraph* Graph = GetGraph();
+		Graph->NotifyGraphChanged();
+
+		// Mark dirty
+		FBlueprintEditorUtils::MarkBlueprintAsModified(GetBlueprint());
+	}
+}
+
+FString UK2Node_GetDataTableRow::GetTooltip() const
+{
+	return NodeTooltip;
+}
+
+UEdGraphPin* UK2Node_GetDataTableRow::GetThenPin()const
+{
+	const UEdGraphSchema_K2* K2Schema = GetDefault<UEdGraphSchema_K2>();
+
+	UEdGraphPin* Pin = FindPinChecked(K2Schema->PN_Then);
+	check(Pin->Direction == EGPD_Output);
+	return Pin;
+}
+
+UEdGraphPin* UK2Node_GetDataTableRow::GetDataTablePin(const TArray<UEdGraphPin*>* InPinsToSearch /*= NULL*/) const
+{
+	const TArray<UEdGraphPin*>* PinsToSearch = InPinsToSearch ? InPinsToSearch : &Pins;
+    
+	UEdGraphPin* Pin = NULL;
+	for( auto PinIt = PinsToSearch->CreateConstIterator(); PinIt; ++PinIt )
+	{
+		UEdGraphPin* TestPin = *PinIt;
+		if( TestPin && TestPin->PinName == UK2Node_GetDataTableRowHelper::DataTablePinName)
+		{
+			Pin = TestPin;
+			break;
+		}
+	}
+	check(Pin == NULL || Pin->Direction == EGPD_Input);
+	return Pin;
+}
+
+UEdGraphPin* UK2Node_GetDataTableRow::GetRowNamePin() const
+{
+	UEdGraphPin* Pin = FindPinChecked(UK2Node_GetDataTableRowHelper::RowNamePinName);
+	check(Pin->Direction == EGPD_Input);
+	return Pin;
+}
+
+UEdGraphPin* UK2Node_GetDataTableRow::GetRowNotFoundPin() const
+{
+	UEdGraphPin* Pin = FindPinChecked(UK2Node_GetDataTableRowHelper::RowNotFoundPinName);
+	check(Pin->Direction == EGPD_Output);
+	return Pin;
+}
+
+UEdGraphPin* UK2Node_GetDataTableRow::GetResultPin() const
+{
+	const UEdGraphSchema_K2* K2Schema = GetDefault<UEdGraphSchema_K2>();
+
+	UEdGraphPin* Pin = FindPinChecked(K2Schema->PN_ReturnValue);
+	check(Pin->Direction == EGPD_Output);
+	return Pin;
+}
+
+
+FText UK2Node_GetDataTableRow::GetNodeTitle(ENodeTitleType::Type TitleType) const
+{
+	UEdGraphPin* DataTablePin = GetDataTablePin();
+
+	FText DataTableString = NSLOCTEXT("K2Node", "None", "NONE");
+	if(DataTablePin != NULL)
+	{
+		if(DataTablePin->LinkedTo.Num() > 0)
+		{
+			// Blueprint will be determined dynamically, so we don't have the name in this case
+			DataTableString = FText::GetEmpty();
+		}
+		else if(DataTablePin->DefaultObject != NULL)
+		{
+			DataTableString = FText::FromString(DataTablePin->DefaultObject->GetName());
+		}
+	}
+
+	FFormatNamedArguments Args;
+	Args.Add(TEXT("DataTableName"), DataTableString);
+	return FText::Format(NSLOCTEXT("K2Node", "DataTable", "GetDataTableRow {DataTableName}"), Args);
+}
+
+FString UK2Node_GetDataTableRow::GetNodeNativeTitle(ENodeTitleType::Type TitleType) const
+{
+	// Do not setup this function for localization, intentionally left unlocalized!
+	
+	UEdGraphPin* DataTablePin = GetDataTablePin();
+
+	FString DataTableString = TEXT("NONE");
+	if(DataTablePin != NULL)
+	{
+		if(DataTablePin->LinkedTo.Num() > 0)
+		{
+			// Blueprint will be determined dynamically, so we don't have the name in this case
+			DataTableString = TEXT("");
+		}
+		else if(DataTablePin->DefaultObject != NULL)
+		{
+			DataTableString = DataTablePin->DefaultObject->GetName();
+		}
+	}
+
+	return FString::Printf(TEXT("GetDataTableRow %s"), *DataTableString);
+}
+
+void UK2Node_GetDataTableRow::ExpandNode(class FKismetCompilerContext& CompilerContext, UEdGraph* SourceGraph)
+{
+    Super::ExpandNode(CompilerContext, SourceGraph);
+    
+	if (CompilerContext.bIsFullCompile)
+	{
+		const UEdGraphSchema_K2* Schema = CompilerContext.GetSchema();
+        
+        UEdGraphPin* OriginalDataTableInPin = GetDataTablePin();
+        UDataTable* Table = (OriginalDataTableInPin != NULL) ? Cast<UDataTable>(OriginalDataTableInPin->DefaultObject) : NULL;
+        if((0 == OriginalDataTableInPin->LinkedTo.Num()) && (NULL == Table))
+        {
+        	CompilerContext.MessageLog.Error(*LOCTEXT("GetDataTableRowNoDataTable_Error", "GetDataTableRow must have a DataTable specified.").ToString(), this);
+        	// we break exec links so this is the only error we get
+            BreakAllNodeLinks();
+            return;
+        }
+
+		// FUNCTION NODE
+		const FName FunctionName = GET_FUNCTION_NAME_CHECKED(UDataTableFunctionLibrary, GetDataTableRowFromName);
+		UK2Node_CallFunction* GetDataTableRowFunction = CompilerContext.SpawnIntermediateNode<UK2Node_CallFunction>(this, SourceGraph);
+		GetDataTableRowFunction->FunctionReference.SetExternalMember(FunctionName, UDataTableFunctionLibrary::StaticClass());
+		GetDataTableRowFunction->AllocateDefaultPins();
+        CompilerContext.MovePinLinksToIntermediate(*GetExecPin(), *(GetDataTableRowFunction->GetExecPin()));
+
+		// Connect the input of our GetDataTableRow to the Input of our Function pin
+        UEdGraphPin* DataTableInPin = GetDataTableRowFunction->FindPinChecked(TEXT("Table"));
+		if(OriginalDataTableInPin->LinkedTo.Num() > 0)
+		{
+			// Copy the connection
+			CompilerContext.MovePinLinksToIntermediate(*OriginalDataTableInPin, *DataTableInPin);
+		}
+		else
+		{
+			// Copy literal
+			DataTableInPin->DefaultObject = OriginalDataTableInPin->DefaultObject;
+		}
+		UEdGraphPin* RowNameInPin = GetDataTableRowFunction->FindPinChecked(TEXT("RowName"));
+		CompilerContext.MovePinLinksToIntermediate(*GetRowNamePin(), *RowNameInPin);
+
+		// Get some pins to work with
+		UEdGraphPin* OriginalOutRowPin = FindPinChecked(Schema->PN_ReturnValue);
+		UEdGraphPin* FunctionOutRowPin = GetDataTableRowFunction->FindPinChecked(TEXT("OutRow"));
+        UEdGraphPin* FunctionReturnPin = GetDataTableRowFunction->FindPinChecked(Schema->PN_ReturnValue);
+        UEdGraphPin* FunctionThenPin = GetDataTableRowFunction->GetThenPin();
+        
+        // Set the type of the OutRow pin on this expanded mode to match original
+        FunctionOutRowPin->PinType = OriginalOutRowPin->PinType;
+		FunctionOutRowPin->PinType.PinSubCategoryObject = OriginalOutRowPin->PinType.PinSubCategoryObject;
+        
+        //BRANCH NODE
+        UK2Node_IfThenElse* BranchNode = CompilerContext.SpawnIntermediateNode<UK2Node_IfThenElse>(this, SourceGraph);
+        BranchNode->AllocateDefaultPins();
+        // Hook up inputs to branch
+        FunctionThenPin->MakeLinkTo(BranchNode->GetExecPin());
+        FunctionReturnPin->MakeLinkTo(BranchNode->GetConditionPin());
+        
+        // Hook up outputs
+        CompilerContext.MovePinLinksToIntermediate(*GetThenPin(), *(BranchNode->GetThenPin()));
+        CompilerContext.MovePinLinksToIntermediate(*GetRowNotFoundPin(), *(BranchNode->GetElsePin()));
+        CompilerContext.MovePinLinksToIntermediate(*OriginalOutRowPin, *FunctionOutRowPin);
+
+		BreakAllNodeLinks();
+	}
+}
+
+#undef LOCTEXT_NAMESPACE
\ No newline at end of file
diff --git a/Engine/Source/Editor/BlueprintGraph/Public/BlueprintGraphDefinitions.h b/Engine/Source/Editor/BlueprintGraph/Public/BlueprintGraphDefinitions.h
--- a/Engine/Source/Editor/BlueprintGraph/Public/BlueprintGraphDefinitions.h
+++ b/Engine/Source/Editor/BlueprintGraph/Public/BlueprintGraphDefinitions.h
@@ -30,6 +30,7 @@
 #include "K2Node_FormatText.h"
 #include "K2Node_FunctionEntry.h"
 #include "K2Node_FunctionResult.h"
+#include "K2Node_GetDataTableRow.h"
 #include "K2Node_IfThenElse.h"
 #include "K2Node_InputAction.h"
 #include "K2Node_InputAxisEvent.h"
diff --git a/Engine/Source/Runtime/Engine/Classes/Kismet/DataTableFunctionLibrary.h b/Engine/Source/Runtime/Engine/Classes/Kismet/DataTableFunctionLibrary.h
new file mode 100644
--- /dev/null
+++ b/Engine/Source/Runtime/Engine/Classes/Kismet/DataTableFunctionLibrary.h
@@ -0,0 +1,52 @@
+// Copyright 1998-2014 Epic Games, Inc. All Rights Reserved.
+
+#pragma once
+#include "../Engine/DataTable.h"
+#include "../Engine/CurveTable.h"
+#include "KismetSystemLibrary.h"
+#include "DataTableFunctionLibrary.generated.h"
+
+UCLASS(MinimalAPI)
+class UDataTableFunctionLibrary : public UBlueprintFunctionLibrary
+{
+	GENERATED_UCLASS_BODY()
+
+	UFUNCTION(BlueprintCallable, Category = "DataTable")
+	static bool EvaluateCurveTableRow(UCurveTable* Table, FName RowName, float InXY, float& OutXY);
+    
+    /** Get a Row from a DataTable given a RowName */
+    UFUNCTION(BlueprintCallable, CustomThunk, Category = "DataTable", meta=(CustomStructureParam = "OutRow", BlueprintInternalUseOnly="true"))
+    static bool GetDataTableRowFromName(UDataTable* Table, FName RowName, FTableRowBase& OutRow);
+    
+    /** Based on UDataTableFunctionLibrary::GetDataTableRow */
+    DECLARE_FUNCTION(execGetDataTableRowFromName)
+    {
+        P_GET_OBJECT(UDataTable, Table);
+        P_GET_PROPERTY(UNameProperty, RowName);
+        
+        Stack.StepCompiledIn<UStructProperty>(NULL);
+        void* OutRowPtr = Stack.MostRecentPropertyAddress;
+        
+        P_FINISH;
+        
+        bool foundRow = false;
+        
+        if (OutRowPtr != NULL)
+        {
+            void* RowPtr = Table->FindRowUnchecked(RowName);
+            
+            if (RowPtr != NULL)
+            {
+                UScriptStruct* StructType = Table->RowStruct;
+                
+                if (StructType != NULL)
+                {
+                    StructType->CopyScriptStruct(OutRowPtr, RowPtr);
+                    foundRow = true;
+                }
+            }
+        }
+        
+        *(bool*)Result = foundRow;
+    }
+};
diff --git a/Engine/Source/Runtime/Engine/Private/DataTableFunctionLibrary.cpp b/Engine/Source/Runtime/Engine/Private/DataTableFunctionLibrary.cpp
new file mode 100644
--- /dev/null
+++ b/Engine/Source/Runtime/Engine/Private/DataTableFunctionLibrary.cpp
@@ -0,0 +1,25 @@
+// Copyright 1998-2014 Epic Games, Inc. All Rights Reserved.
+
+#include "EnginePrivate.h"
+#include "../Classes/Kismet/DataTableFunctionLibrary.h"
+
+UDataTableFunctionLibrary::UDataTableFunctionLibrary(const FPostConstructInitializeProperties& PCIP)
+	: Super(PCIP)
+{
+}
+
+bool UDataTableFunctionLibrary::EvaluateCurveTableRow(UCurveTable* Table, FName RowName, float InXY, float& OutXY)
+{
+    FCurveTableRowHandle Handle;
+    Handle.CurveTable = Table;
+    Handle.RowName = RowName;
+    
+	return Handle.Eval(InXY, &OutXY);
+}
+
+bool UDataTableFunctionLibrary::GetDataTableRowFromName(UDataTable* Table, FName RowName, FTableRowBase& OutRow)
+{
+	// We should never hit this!  stubs to avoid NoExport on the class.
+	check(0);
+    return false;
+}
\ No newline at end of file
