Blueprint code generation no longer duplicates cast result terms for objects being plugged into interface params.
(TTP #337973)

diff --git a/Engine/Source/Editor/BlueprintGraph/Private/CallFunctionHandler.cpp b/Engine/Source/Editor/BlueprintGraph/Private/CallFunctionHandler.cpp
--- a/Engine/Source/Editor/BlueprintGraph/Private/CallFunctionHandler.cpp
+++ b/Engine/Source/Editor/BlueprintGraph/Private/CallFunctionHandler.cpp
@@ -608,7 +608,7 @@ void FKCHandler_CallFunction::RegisterNets(FKismetFunctionContext& Context, UEdG
 		if ((Pin->PinType.PinCategory == K2Schema->PC_Interface) && (Pin->LinkedTo[0]->PinType.PinCategory == K2Schema->PC_Object))
 		{
 			FBPTerminal* InterfaceTerm = new (Context.IsEventGraph() ? Context.EventGraphLocals : Context.Locals) FBPTerminal();
-			InterfaceTerm->CopyFromPin(Pin, Context.NetNameMap->MakeValidName(Pin->LinkedTo[0]) + TEXT("_CastResult"));
+			InterfaceTerm->CopyFromPin(Pin, Context.NetNameMap->MakeValidName(Pin) + TEXT("_CastInput"));
 			InterfaceTerm->Source = Node;
 
 			InterfaceTermMap.Add(Pin, InterfaceTerm);
