#ue4 Fix for missing BPGRAPH_API unresolved external error

diff --git a/Engine/Source/Editor/BlueprintGraph/Classes/K2Node_Event.h b/Engine/Source/Editor/BlueprintGraph/Classes/K2Node_Event.h
--- a/Engine/Source/Editor/BlueprintGraph/Classes/K2Node_Event.h
+++ b/Engine/Source/Editor/BlueprintGraph/Classes/K2Node_Event.h
@@ -62,7 +62,7 @@ class UK2Node_Event : public UK2Node_EditablePinBase
 	BLUEPRINTGRAPH_API virtual FNodeHandlingFunctor* CreateNodeHandler(FKismetCompilerContext& CompilerContext) const OVERRIDE;
 	BLUEPRINTGRAPH_API virtual void ExpandNode(class FKismetCompilerContext& CompilerContext, UEdGraph* SourceGraph) OVERRIDE;
 	BLUEPRINTGRAPH_API virtual FText GetToolTipHeading() const;
-	virtual void GetNodeAttributes( TArray<TKeyValuePair<FString, FString>>& OutNodeAttributes ) const OVERRIDE;
+	BLUEPRINTGRAPH_API virtual void GetNodeAttributes( TArray<TKeyValuePair<FString, FString>>& OutNodeAttributes ) const OVERRIDE;
 	// End UK2Node interface
 
 	/** Checks whether the parameters for this event node are compatible with the specified function entry node */
