Support ExpandEnumAsExecs as output pins

Evaluate Curve Table Row node now uses ExpandEnumAsExecs to have separate execution pins for if the row was found or not

Add some friendly names for Get Data Table Row pins

diff --git a/Engine/Source/Editor/BlueprintGraph/Private/K2Node_CallFunction.cpp b/Engine/Source/Editor/BlueprintGraph/Private/K2Node_CallFunction.cpp
--- a/Engine/Source/Editor/BlueprintGraph/Private/K2Node_CallFunction.cpp
+++ b/Engine/Source/Editor/BlueprintGraph/Private/K2Node_CallFunction.cpp
@@ -5,6 +5,7 @@
 
 #include "CompilerResultsLog.h"
 #include "CallFunctionHandler.h"
+#include "K2Node_SwitchEnum.h"
 
 #define LOCTEXT_NAMESPACE "K2Node"
 
@@ -327,6 +328,9 @@ UEdGraphPin* UK2Node_CallFunction::CreateSelfPin(const UFunction* Function)
 
 void UK2Node_CallFunction::CreateExecPinsForFunctionCall(const UFunction* Function)
 {
+	bool bCreateSingleExecInputPin = true;
+	bool bCreateThenPin = true;
+	
 	// If not pure, create exec pins
 	if (!bIsPureFunc)
 	{
@@ -339,26 +343,46 @@ void UK2Node_CallFunction::CreateExecPinsForFunctionCall(const UFunction* Functi
 			UByteProperty* EnumProp = FindField<UByteProperty>(Function, FName(*EnumParamName));
 			if(EnumProp != NULL && EnumProp->Enum != NULL)
 			{
+				const bool bIsFunctionInput = !EnumProp->HasAnyPropertyFlags(CPF_ReturnParm) &&
+					(!EnumProp->HasAnyPropertyFlags(CPF_OutParm) ||
+					 EnumProp->HasAnyPropertyFlags(CPF_ReferenceParm));
+				const EEdGraphPinDirection Direction = bIsFunctionInput ? EGPD_Input : EGPD_Output;
+				
 				// yay, found it! Now create exec pin for each
-				int32 NumInputs = (EnumProp->Enum->NumEnums() - 1);
-				for(int32 InputIdx=0; InputIdx<NumInputs; InputIdx++)
+				int32 NumExecs = (EnumProp->Enum->NumEnums() - 1);
+				for(int32 ExecIdx=0; ExecIdx<NumExecs; ExecIdx++)
+				{
+					FString ExecName = EnumProp->Enum->GetEnumName(ExecIdx);
+					CreatePin(Direction, K2Schema->PC_Exec, TEXT(""), NULL, false, false, ExecName);
+				}
+				
+				if (bIsFunctionInput)
 				{
-					FString InputName = EnumProp->Enum->GetEnumName(InputIdx);
-					CreatePin(EGPD_Input, K2Schema->PC_Exec, TEXT(""), NULL, false, false, InputName);
+					// If using ExpandEnumAsExec for input, don't want to add a input exec pin
+					bCreateSingleExecInputPin = false;
+				}
+				else
+				{
+					// If using ExpandEnumAsExec for output, don't want to add a "then" pin
+					bCreateThenPin = false;
 				}
 			}
 		}
-		else
+		
+		if (bCreateSingleExecInputPin)
 		{
 			// Single input exec pin
 			CreatePin(EGPD_Input, K2Schema->PC_Exec, TEXT(""), NULL, false, false, K2Schema->PN_Execute);
 		}
 
-		// Use 'completed' name for output pins on latent functions
-		UEdGraphPin* OutputExecPin = CreatePin(EGPD_Output, K2Schema->PC_Exec, TEXT(""), NULL, false, false, K2Schema->PN_Then);
-		if(Function->HasMetaData(FBlueprintMetadata::MD_Latent))
+		if (bCreateThenPin)
 		{
-			OutputExecPin->PinFriendlyName = FText::FromString(K2Schema->PN_Completed);
+			UEdGraphPin* OutputExecPin = CreatePin(EGPD_Output, K2Schema->PC_Exec, TEXT(""), NULL, false, false, K2Schema->PN_Then);
+			// Use 'completed' name for output pins on latent functions
+			if(Function->HasMetaData(FBlueprintMetadata::MD_Latent))
+			{
+				OutputExecPin->PinFriendlyName = FText::FromString(K2Schema->PN_Completed);
+			}
 		}
 	}
 }
@@ -1246,52 +1270,88 @@ void UK2Node_CallFunction::ExpandNode(class FKismetCompilerContext& CompilerCont
 				// Get the metadata that identifies which param is the enum, and try and find it
 				const FString& EnumParamName = Function->GetMetaData(FBlueprintMetadata::MD_ExpandEnumAsExecs);
 				UByteProperty* EnumProp = FindField<UByteProperty>(Function, FName(*EnumParamName));
+				UEdGraphPin* EnumParamPin = FindPinChecked(EnumParamName);
 				if(EnumProp != NULL && EnumProp->Enum != NULL)
 				{
-					// Create normal exec input
-					UEdGraphPin* ExecutePin = CreatePin(EGPD_Input, Schema->PC_Exec, TEXT(""), NULL, false, false, Schema->PN_Execute);
-
-					// Create temp enum variable
-					UK2Node_TemporaryVariable* TempEnumVarNode = CompilerContext.SpawnIntermediateNode<UK2Node_TemporaryVariable>(this, SourceGraph);
-					TempEnumVarNode->VariableType.PinCategory = Schema->PC_Byte;
-					TempEnumVarNode->VariableType.PinSubCategoryObject = EnumProp->Enum;
-					TempEnumVarNode->AllocateDefaultPins();
-					// Get the output pin
-					UEdGraphPin* TempEnumVarOutput = TempEnumVarNode->GetVariablePin();
-
-					// Connect temp enum variable to (hidden) enum pin
-					UEdGraphPin* EnumParamPin = FindPin(EnumParamName);
-					check(EnumParamPin);
-					check(EnumParamPin->Direction == EGPD_Input);
-					Schema->TryCreateConnection(TempEnumVarOutput, EnumParamPin);
-
-					// Now we want to iterate over other exec inputs...
-					for(int32 PinIdx=Pins.Num()-1; PinIdx>=0; PinIdx--)
+					// Expanded as input execs pins
+					if (EnumParamPin->Direction == EGPD_Input)
 					{
-						UEdGraphPin* Pin = Pins[PinIdx];
-						if( Pin != NULL && 
-							Pin != ExecutePin &&
-							Pin->Direction == EGPD_Input && 
-							Pin->PinType.PinCategory == Schema->PC_Exec )
+						// Create normal exec input
+						UEdGraphPin* ExecutePin = CreatePin(EGPD_Input, Schema->PC_Exec, TEXT(""), NULL, false, false, Schema->PN_Execute);
+
+						// Create temp enum variable
+						UK2Node_TemporaryVariable* TempEnumVarNode = CompilerContext.SpawnIntermediateNode<UK2Node_TemporaryVariable>(this, SourceGraph);
+						TempEnumVarNode->VariableType.PinCategory = Schema->PC_Byte;
+						TempEnumVarNode->VariableType.PinSubCategoryObject = EnumProp->Enum;
+						TempEnumVarNode->AllocateDefaultPins();
+						// Get the output pin
+						UEdGraphPin* TempEnumVarOutput = TempEnumVarNode->GetVariablePin();
+
+						// Connect temp enum variable to (hidden) enum pin
+						Schema->TryCreateConnection(TempEnumVarOutput, EnumParamPin);
+
+						// Now we want to iterate over other exec inputs...
+						for(int32 PinIdx=Pins.Num()-1; PinIdx>=0; PinIdx--)
 						{
-							// Create node to set the temp enum var
-							UK2Node_AssignmentStatement* AssignNode = CompilerContext.SpawnIntermediateNode<UK2Node_AssignmentStatement>(this, SourceGraph);
-							AssignNode->AllocateDefaultPins();
+							UEdGraphPin* Pin = Pins[PinIdx];
+							if( Pin != NULL && 
+								Pin != ExecutePin &&
+								Pin->Direction == EGPD_Input && 
+								Pin->PinType.PinCategory == Schema->PC_Exec )
+							{
+								// Create node to set the temp enum var
+								UK2Node_AssignmentStatement* AssignNode = CompilerContext.SpawnIntermediateNode<UK2Node_AssignmentStatement>(this, SourceGraph);
+								AssignNode->AllocateDefaultPins();
 
-							// Move connections from fake 'enum exec' pint to this assignment node
-							CompilerContext.MovePinLinksToIntermediate(*Pin, *AssignNode->GetExecPin());
+								// Move connections from fake 'enum exec' pint to this assignment node
+								CompilerContext.MovePinLinksToIntermediate(*Pin, *AssignNode->GetExecPin());
 
-							// Connect this to out temp enum var
-							Schema->TryCreateConnection(AssignNode->GetVariablePin(), TempEnumVarOutput);
+								// Connect this to out temp enum var
+								Schema->TryCreateConnection(AssignNode->GetVariablePin(), TempEnumVarOutput);
 
-							// Connect exec output to 'real' exec pin
-							Schema->TryCreateConnection(AssignNode->GetThenPin(), ExecutePin);
+								// Connect exec output to 'real' exec pin
+								Schema->TryCreateConnection(AssignNode->GetThenPin(), ExecutePin);
 
-							// set the literal enum value to set to
-							AssignNode->GetValuePin()->DefaultValue = Pin->PinName;
+								// set the literal enum value to set to
+								AssignNode->GetValuePin()->DefaultValue = Pin->PinName;
 
-							// Finally remove this 'cosmetic' exec pin
-							Pins.RemoveAt(PinIdx);
+								// Finally remove this 'cosmetic' exec pin
+								Pins.RemoveAt(PinIdx);
+							}
+						}
+					}
+					// Expanded as output execs pins
+					else if (EnumParamPin->Direction == EGPD_Output)
+					{
+						// Create normal exec output
+						UEdGraphPin* ExecutePin = CreatePin(EGPD_Output, Schema->PC_Exec, TEXT(""), NULL, false, false, Schema->PN_Execute);
+						
+						// Create a SwitchEnum node to switch on the output enum
+						UK2Node_SwitchEnum* SwitchEnumNode = CompilerContext.SpawnIntermediateNode<UK2Node_SwitchEnum>(this, SourceGraph);
+						UEnum* EnumObject = Cast<UEnum>(EnumParamPin->PinType.PinSubCategoryObject.Get());
+						SwitchEnumNode->SetEnum(EnumObject);
+						SwitchEnumNode->AllocateDefaultPins();
+						
+						// Hook up execution to the switch node
+						Schema->TryCreateConnection(ExecutePin, SwitchEnumNode->GetExecPin());
+						// Connect (hidden) enum pin to switch node's selection pin
+						Schema->TryCreateConnection(EnumParamPin, SwitchEnumNode->GetSelectionPin());
+						
+						// Now we want to iterate over other exec outputs
+						for(int32 PinIdx=Pins.Num()-1; PinIdx>=0; PinIdx--)
+						{
+							UEdGraphPin* Pin = Pins[PinIdx];
+							if( Pin != NULL &&
+							   Pin != ExecutePin &&
+							   Pin->Direction == EGPD_Output &&
+							   Pin->PinType.PinCategory == Schema->PC_Exec )
+							{
+								// Move connections from fake 'enum exec' pint to this switch node
+								CompilerContext.MovePinLinksToIntermediate(*Pin, *SwitchEnumNode->FindPinChecked(Pin->PinName));
+								
+								// Finally remove this 'cosmetic' exec pin
+								Pins.RemoveAt(PinIdx);
+							}
 						}
 					}
 				}
diff --git a/Engine/Source/Editor/BlueprintGraph/Private/K2Node_GetDataTableRow.cpp b/Engine/Source/Editor/BlueprintGraph/Private/K2Node_GetDataTableRow.cpp
--- a/Engine/Source/Editor/BlueprintGraph/Private/K2Node_GetDataTableRow.cpp
+++ b/Engine/Source/Editor/BlueprintGraph/Private/K2Node_GetDataTableRow.cpp
@@ -30,7 +30,8 @@ void UK2Node_GetDataTableRow::AllocateDefaultPins()
 
 	// Add execution pins
 	CreatePin(EGPD_Input, K2Schema->PC_Exec, TEXT(""), NULL, false, false, K2Schema->PN_Execute);
-	CreatePin(EGPD_Output, K2Schema->PC_Exec, TEXT(""), NULL, false, false, K2Schema->PN_Then);
+	UEdGraphPin* RowFoundPin = CreatePin(EGPD_Output, K2Schema->PC_Exec, TEXT(""), NULL, false, false, K2Schema->PN_Then);
+	RowFoundPin->PinFriendlyName = LOCTEXT("GetDataTableRow Row Found Exec pin", "Row Found");
 	CreatePin(EGPD_Output, K2Schema->PC_Exec, TEXT(""), NULL, false, false, UK2Node_GetDataTableRowHelper::RowNotFoundPinName);
 
 	// Add DataTable pin
@@ -43,6 +44,7 @@ void UK2Node_GetDataTableRow::AllocateDefaultPins()
 
 	// Result pin
 	UEdGraphPin* ResultPin = CreatePin(EGPD_Output, K2Schema->PC_Struct, TEXT(""), FTableRowBase::StaticStruct(), false, false, K2Schema->PN_ReturnValue);
+	ResultPin->PinFriendlyName = LOCTEXT("GetDataTableRow Output Row", "Out Row");
 	SetPinToolTip(*ResultPin, LOCTEXT("ResultPinDescription", "The returned TableRow, if found"));
 
 	Super::AllocateDefaultPins();
diff --git a/Engine/Source/Runtime/Engine/Classes/Kismet/DataTableFunctionLibrary.h b/Engine/Source/Runtime/Engine/Classes/Kismet/DataTableFunctionLibrary.h
--- a/Engine/Source/Runtime/Engine/Classes/Kismet/DataTableFunctionLibrary.h
+++ b/Engine/Source/Runtime/Engine/Classes/Kismet/DataTableFunctionLibrary.h
@@ -6,13 +6,26 @@
 #include "KismetSystemLibrary.h"
 #include "DataTableFunctionLibrary.generated.h"
 
+/** Enum used to indicate success or failure of EvaluateCurveTableRow */
+UENUM()
+namespace EEvaluateCurveTableResult
+{
+    enum Type
+    {
+        /** Found the row successfully */
+        RowFound,
+        /** Failed to find row */
+        RowNotFound,
+    };
+}
+
 UCLASS(MinimalAPI)
 class UDataTableFunctionLibrary : public UBlueprintFunctionLibrary
 {
 	GENERATED_UCLASS_BODY()
 
-	UFUNCTION(BlueprintCallable, Category = "DataTable")
-	static bool EvaluateCurveTableRow(UCurveTable* Table, FName RowName, float InXY, float& OutXY);
+	UFUNCTION(BlueprintCallable, Category = "DataTable", meta=(ExpandEnumAsExecs="OutResult"))
+	static void EvaluateCurveTableRow(UCurveTable* CurveTable, FName RowName, float InXY, TEnumAsByte<EEvaluateCurveTableResult::Type>& OutResult, float& OutXY);
     
     /** Get a Row from a DataTable given a RowName */
     UFUNCTION(BlueprintCallable, CustomThunk, Category = "DataTable", meta=(CustomStructureParam = "OutRow", BlueprintInternalUseOnly="true"))
diff --git a/Engine/Source/Runtime/Engine/Private/DataTableFunctionLibrary.cpp b/Engine/Source/Runtime/Engine/Private/DataTableFunctionLibrary.cpp
--- a/Engine/Source/Runtime/Engine/Private/DataTableFunctionLibrary.cpp
+++ b/Engine/Source/Runtime/Engine/Private/DataTableFunctionLibrary.cpp
@@ -8,13 +8,22 @@ UDataTableFunctionLibrary::UDataTableFunctionLibrary(const FPostConstructInitial
 {
 }
 
-bool UDataTableFunctionLibrary::EvaluateCurveTableRow(UCurveTable* Table, FName RowName, float InXY, float& OutXY)
+void UDataTableFunctionLibrary::EvaluateCurveTableRow(UCurveTable* CurveTable, FName RowName, float InXY, TEnumAsByte<EEvaluateCurveTableResult::Type>& OutResult, float& OutXY)
 {
     FCurveTableRowHandle Handle;
-    Handle.CurveTable = Table;
+    Handle.CurveTable = CurveTable;
     Handle.RowName = RowName;
     
-	return Handle.Eval(InXY, &OutXY);
+	bool found = Handle.Eval(InXY, &OutXY);
+    
+    if (found)
+    {
+        OutResult = EEvaluateCurveTableResult::RowFound;
+    }
+    else
+    {
+        OutResult = EEvaluateCurveTableResult::RowNotFound;
+    }
 }
 
 bool UDataTableFunctionLibrary::GetDataTableRowFromName(UDataTable* Table, FName RowName, FTableRowBase& OutRow)
