Fixing it so users can change the input types on math expression nodes, and those types will remain consistent after node reconstruction.
(TTP #337678)

diff --git a/Engine/Source/Editor/BlueprintGraph/Private/K2Node_MathExpression.cpp b/Engine/Source/Editor/BlueprintGraph/Private/K2Node_MathExpression.cpp
--- a/Engine/Source/Editor/BlueprintGraph/Private/K2Node_MathExpression.cpp
+++ b/Engine/Source/Editor/BlueprintGraph/Private/K2Node_MathExpression.cpp
@@ -1162,6 +1162,8 @@ public:
 		// want to track if we generated any errors from this pass, so we need to know how many we started with
 		int32 StartingErrorCount = MessageLog.NumErrors;
 		
+		InputPinNames.Empty();
+
 		LayoutMapper.Clear();
 		// map the depth/height of expression tree (so we can position nodes prettily)
 		ExpressionRoot->Accept(LayoutMapper);
@@ -1208,6 +1210,17 @@ public:
 		
 		return !bHasErrors;
 	}
+
+	/**
+	 * When the node gen is over, we need to clear any old pins that weren't 
+	 * reused. This query method helps in identifying those that were utilized.
+	 *
+	 * @return True if the pin's name was used in the most recent expression, false if not.
+	 */
+	bool IsPinInUse(TSharedPtr<FUserPinInfo> PinInfo)
+	{
+		return (InputPinNames.Find(PinInfo->PinName) != INDEX_NONE);
+	}
     
     /**
 	 * Overloaded, part of the FExpressionVisitor interface; attempts to 
@@ -1263,7 +1276,7 @@ public:
      * attempts to generate a UK2Node_CallFunction node for the specified
      * FBinaryOperator.
 	 *
-	 * @return True to continue travesing the expression tree, false to stop.
+	 * @return True to continue traversing the expression tree, false to stop.
 	 */
 	virtual bool Visit(FBinaryOperator& ExpressionNode, EVisitPhase Phase) OVERRIDE
 	{
@@ -1293,7 +1306,11 @@ public:
 	}
 	
 	/**
+	 * Does nothing (but had to prevent this expression node from being flagged 
+	 * as "unhandled"). Expression lists are handled by whatever expression 
+	 * they're contained within.
 	 *
+	 * @return Always true, it is expected that cascading errors are handled (and all should be logged).
 	 */
 	virtual bool Visit(FExpressionList& ExpressionNode, EVisitPhase Phase) OVERRIDE
 	{
@@ -1307,7 +1324,11 @@ public:
 	}
 	
 	/**
+	 * Overloaded, part of the FExpressionVisitor interface... On VISIT_Post, 
+     * attempts to generate a UK2Node_CallFunction node for the specified
+     * FFunctionExpression.
 	 *
+	 * @return True to continue traversing the expression tree, false to stop.
 	 */
 	virtual bool Visit(FFunctionExpression& ExpressionNode, EVisitPhase Phase) OVERRIDE
 	{
@@ -1439,6 +1460,12 @@ private:
 			InputPinFragment = MakeShareable(new FCodeGenFragment_InputPin(NewInputPin));
 		}
 		
+		// when regenerating a node, we need to clear any old pins that weren't 
+		// reused (can't do this before the node gen because the user may have 
+		// altered a pin to how they want it), so here we track the ones that 
+		// were used by the latest expression
+		InputPinNames.Add(VariableIdentifier);
+
 		return InputPinFragment;
 	}
 	
@@ -1695,8 +1722,9 @@ private:
 	/** The node that we're generating sub-nodes and pins for */
 	UK2Node_MathExpression* CompilingNode;
 
-	/** The blueprint that CompilingNode belongs to (the blueprint this will 
-	 * generate a graph for)
+	/** 
+	 * The blueprint that CompilingNode belongs to (the blueprint this will 
+	 * generate a graph for).
 	 */
 	UBlueprint* TargetBlueprint;
 
@@ -1728,6 +1756,12 @@ private:
 	 * not in the middle of GenerateCode().
 	 */
 	FCompilerResultsLog* ActiveMessageLog;
+
+	/** 
+	 * After the code generation, we want to clear any old pins that 
+	 * weren't reused, so here we track the ones in use.
+	 */
+	TArray<FString> InputPinNames;
 };
 
 /*******************************************************************************
@@ -2290,6 +2324,23 @@ void UK2Node_MathExpression::RebuildExpression(FString InExpression)
 				{
 					CachedMessageLog->Error(*LOCTEXT("MathExprFailedGen", "Failed to generate full expression graph for: '@@'").ToString(), this);
 				}
+
+				if (UK2Node_Tunnel* EntryNode = GetEntryNode())
+				{
+					// iterate backwards so we can remove as we go... we want to 
+					// clear any pins that weren't used by the expression (if we
+					// clear any, then they were probably remnants from the last
+					// expression... we can't delete them before, because the 
+					// user may have mutated one for the new expression)
+					for (int32 PinIndex = EntryNode->UserDefinedPins.Num() - 1; PinIndex >= 0; --PinIndex)
+					{
+						TSharedPtr<FUserPinInfo> PinInfo = EntryNode->UserDefinedPins[PinIndex];
+						if (!GraphGenerator.IsPinInUse(PinInfo))
+						{
+							EntryNode->RemoveUserDefinedPin(PinInfo);
+						}
+					}					
+				}
 			}
 			else
 			{
@@ -2317,22 +2368,15 @@ void UK2Node_MathExpression::ClearExpression()
 	// clear out old nodes
 	DeleteGeneratedNodesInGraph(BoundGraph);
 
-	// delete old pins
-	if (UK2Node_Tunnel* EntryNode = GetEntryNode())
-	{
-		// iterate back wards so we can remove as we go
-		for (int32 pinIndex = EntryNode->UserDefinedPins.Num()-1; pinIndex >= 0; --pinIndex)
-		{
-			TSharedPtr<FUserPinInfo> PinInfo = EntryNode->UserDefinedPins[pinIndex];
-			EntryNode->RemoveUserDefinedPin(PinInfo);
-		}
-	}
+	// delete the old return pins (they will always be regenerated)... save the 
+	// input pins though (because someone may have changed the input type to 
+	// something other than a float)
 	if (UK2Node_Tunnel* ExitNode = GetExitNode())
 	{
-		// iterate back wards so we can remove as we go
-		for (int32 pinIndex = ExitNode->UserDefinedPins.Num()-1; pinIndex >= 0; --pinIndex)
+		// iterate backwards so we can remove as we go
+		for (int32 PinIndex = ExitNode->UserDefinedPins.Num() - 1; PinIndex >= 0; --PinIndex)
 		{
-			TSharedPtr<FUserPinInfo> PinInfo = ExitNode->UserDefinedPins[pinIndex];
+			TSharedPtr<FUserPinInfo> PinInfo = ExitNode->UserDefinedPins[PinIndex];
 			ExitNode->RemoveUserDefinedPin(PinInfo);
 		}
 	}
