UE4 Refactoring. Changed OVERRIDE and FINAL macros to keywords override and final.

diff --git a/Engine/Source/Editor/AnimGraph/Classes/AnimationStateMachineSchema.h b/Engine/Source/Editor/AnimGraph/Classes/AnimationStateMachineSchema.h
--- a/Engine/Source/Editor/AnimGraph/Classes/AnimationStateMachineSchema.h
+++ b/Engine/Source/Editor/AnimGraph/Classes/AnimationStateMachineSchema.h
@@ -22,8 +22,8 @@ struct ANIMGRAPH_API FEdGraphSchemaAction_NewStateNode : public FEdGraphSchemaAc
 		, NodeTemplate(NULL)
 	{}
 
-	virtual UEdGraphNode* PerformAction(class UEdGraph* ParentGraph, UEdGraphPin* FromPin, const FVector2D Location, bool bSelectNewNode = true) OVERRIDE;
-	virtual void AddReferencedObjects( FReferenceCollector& Collector ) OVERRIDE;
+	virtual UEdGraphNode* PerformAction(class UEdGraph* ParentGraph, UEdGraphPin* FromPin, const FVector2D Location, bool bSelectNewNode = true) override;
+	virtual void AddReferencedObjects( FReferenceCollector& Collector ) override;
 
 	template <typename NodeType>
 	static NodeType* SpawnNodeFromTemplate(class UEdGraph* ParentGraph, NodeType* InTemplateNode, const FVector2D Location = FVector2D(0.0f, 0.0f))
@@ -50,7 +50,7 @@ struct ANIMGRAPH_API FEdGraphSchemaAction_NewStateComment : public FEdGraphSchem
 	{}
 
 	// FEdGraphSchemaAction interface
-	virtual UEdGraphNode* PerformAction(class UEdGraph* ParentGraph, UEdGraphPin* FromPin, const FVector2D Location, bool bSelectNewNode = true) OVERRIDE;
+	virtual UEdGraphNode* PerformAction(class UEdGraph* ParentGraph, UEdGraphPin* FromPin, const FVector2D Location, bool bSelectNewNode = true) override;
 	// End of FEdGraphSchemaAction interface
 };
 
@@ -65,18 +65,18 @@ class UAnimationStateMachineSchema : public UEdGraphSchema
 
 	// Begin UEdGraphSchema interface
 	virtual void CreateDefaultNodesForGraph(UEdGraph& Graph) const;
-	virtual const FPinConnectionResponse CanCreateConnection(const UEdGraphPin* A, const UEdGraphPin* B) const OVERRIDE;
-	virtual bool TryCreateConnection(UEdGraphPin* A, UEdGraphPin* B) const OVERRIDE;
+	virtual const FPinConnectionResponse CanCreateConnection(const UEdGraphPin* A, const UEdGraphPin* B) const override;
+	virtual bool TryCreateConnection(UEdGraphPin* A, UEdGraphPin* B) const override;
 	virtual bool CreateAutomaticConversionNodeAndConnections(UEdGraphPin* PinA, UEdGraphPin* PinB) const;
-	virtual void GetGraphContextActions(FGraphContextMenuBuilder& ContextMenuBuilder) const OVERRIDE;
-	virtual EGraphType GetGraphType(const UEdGraph* TestEdGraph) const OVERRIDE;
-	virtual void GetContextMenuActions(const UEdGraph* CurrentGraph, const UEdGraphNode* InGraphNode, const UEdGraphPin* InGraphPin, FMenuBuilder* MenuBuilder, bool bIsDebugging) const OVERRIDE;
-	virtual FLinearColor GetPinTypeColor(const FEdGraphPinType& PinType) const OVERRIDE;
-	virtual void GetGraphDisplayInformation(const UEdGraph& Graph, /*out*/ FGraphDisplayInfo& DisplayInfo) const OVERRIDE;
-	virtual void DroppedAssetsOnGraph(const TArray<FAssetData>& Assets, const FVector2D& GraphPosition, UEdGraph* Graph) const OVERRIDE;
-	virtual void DroppedAssetsOnNode(const TArray<FAssetData>& Assets, const FVector2D& GraphPosition, UEdGraphNode* Node) const OVERRIDE;
-	virtual void DroppedAssetsOnPin(const TArray<FAssetData>& Assets, const FVector2D& GraphPosition, UEdGraphPin* Pin) const OVERRIDE;
-	virtual void GetAssetsNodeHoverMessage(const TArray<FAssetData>& Assets, const UEdGraphNode* HoverNode, FString& OutTooltipText, bool& OutOkIcon) const OVERRIDE;
-	virtual void GetAssetsPinHoverMessage(const TArray<FAssetData>& Assets, const UEdGraphPin* HoverPin, FString& OutTooltipText, bool& OutOkIcon) const OVERRIDE;
+	virtual void GetGraphContextActions(FGraphContextMenuBuilder& ContextMenuBuilder) const override;
+	virtual EGraphType GetGraphType(const UEdGraph* TestEdGraph) const override;
+	virtual void GetContextMenuActions(const UEdGraph* CurrentGraph, const UEdGraphNode* InGraphNode, const UEdGraphPin* InGraphPin, FMenuBuilder* MenuBuilder, bool bIsDebugging) const override;
+	virtual FLinearColor GetPinTypeColor(const FEdGraphPinType& PinType) const override;
+	virtual void GetGraphDisplayInformation(const UEdGraph& Graph, /*out*/ FGraphDisplayInfo& DisplayInfo) const override;
+	virtual void DroppedAssetsOnGraph(const TArray<FAssetData>& Assets, const FVector2D& GraphPosition, UEdGraph* Graph) const override;
+	virtual void DroppedAssetsOnNode(const TArray<FAssetData>& Assets, const FVector2D& GraphPosition, UEdGraphNode* Node) const override;
+	virtual void DroppedAssetsOnPin(const TArray<FAssetData>& Assets, const FVector2D& GraphPosition, UEdGraphPin* Pin) const override;
+	virtual void GetAssetsNodeHoverMessage(const TArray<FAssetData>& Assets, const UEdGraphNode* HoverNode, FString& OutTooltipText, bool& OutOkIcon) const override;
+	virtual void GetAssetsPinHoverMessage(const TArray<FAssetData>& Assets, const UEdGraphPin* HoverPin, FString& OutTooltipText, bool& OutOkIcon) const override;
 	// End UEdGraphSchema interface
 };
diff --git a/Engine/Source/Editor/AnimGraph/Classes/AnimationTransitionSchema.h b/Engine/Source/Editor/AnimGraph/Classes/AnimationTransitionSchema.h
--- a/Engine/Source/Editor/AnimGraph/Classes/AnimationTransitionSchema.h
+++ b/Engine/Source/Editor/AnimGraph/Classes/AnimationTransitionSchema.h
@@ -10,9 +10,9 @@ class UAnimationTransitionSchema : public UEdGraphSchema_K2
 	GENERATED_UCLASS_BODY()
 
 	// Begin UEdGraphSchema interface.
-	virtual void GetGraphContextActions(FGraphContextMenuBuilder& ContextMenuBuilder) const OVERRIDE;
-	virtual void CreateDefaultNodesForGraph(UEdGraph& Graph) const OVERRIDE;
-	virtual bool CanDuplicateGraph(UEdGraph* InSourceGraph) const OVERRIDE { return false; }
+	virtual void GetGraphContextActions(FGraphContextMenuBuilder& ContextMenuBuilder) const override;
+	virtual void CreateDefaultNodesForGraph(UEdGraph& Graph) const override;
+	virtual bool CanDuplicateGraph(UEdGraph* InSourceGraph) const override { return false; }
 	virtual void GetGraphDisplayInformation(const UEdGraph& Graph, /*out*/ FGraphDisplayInfo& DisplayInfo) const;
 	// End UEdGraphSchema interface.
 
diff --git a/Engine/Source/Editor/AnimGraph/Classes/K2Node_TransitionRuleGetter.h b/Engine/Source/Editor/AnimGraph/Classes/K2Node_TransitionRuleGetter.h
--- a/Engine/Source/Editor/AnimGraph/Classes/K2Node_TransitionRuleGetter.h
+++ b/Engine/Source/Editor/AnimGraph/Classes/K2Node_TransitionRuleGetter.h
@@ -37,11 +37,11 @@ class UK2Node_TransitionRuleGetter : public UK2Node
 	UObject* AssociatedStateNode;
 
 	// UEdGraphNode interface
-	virtual void AllocateDefaultPins() OVERRIDE;
-	virtual FText GetNodeTitle(ENodeTitleType::Type TitleType) const OVERRIDE;
-	virtual FString GetNodeNativeTitle(ENodeTitleType::Type TitleType) const OVERRIDE;
-	virtual FString GetTooltip() const OVERRIDE;
-	virtual bool ShowPaletteIconOnNode() const OVERRIDE{ return false; }
+	virtual void AllocateDefaultPins() override;
+	virtual FText GetNodeTitle(ENodeTitleType::Type TitleType) const override;
+	virtual FString GetNodeNativeTitle(ENodeTitleType::Type TitleType) const override;
+	virtual FString GetTooltip() const override;
+	virtual bool ShowPaletteIconOnNode() const override{ return false; }
 	// End of UEdGraphNode interface
 
 	// @todo document
diff --git a/Engine/Source/Editor/AnimGraph/Private/AnimGraphNode_Base.cpp b/Engine/Source/Editor/AnimGraph/Private/AnimGraphNode_Base.cpp
--- a/Engine/Source/Editor/AnimGraph/Private/AnimGraphNode_Base.cpp
+++ b/Engine/Source/Editor/AnimGraph/Private/AnimGraphNode_Base.cpp
@@ -32,7 +32,7 @@ public:
 		}
 	}
 
-	virtual void GetRecordDefaults(UProperty* TestProperty, FOptionalPinFromProperty& Record) const OVERRIDE
+	virtual void GetRecordDefaults(UProperty* TestProperty, FOptionalPinFromProperty& Record) const override
 	{
 		const UAnimationGraphSchema* Schema = GetDefault<UAnimationGraphSchema>();
 
@@ -53,7 +53,7 @@ public:
 		Record.bPropertyIsCustomized = bPropertyIsCustomized;
  	}
 
-	virtual void CustomizePinData(UEdGraphPin* Pin, FName SourcePropertyName, int32 ArrayIndex, UProperty* Property) const OVERRIDE
+	virtual void CustomizePinData(UEdGraphPin* Pin, FName SourcePropertyName, int32 ArrayIndex, UProperty* Property) const override
 	{
 		if (BaseNode != NULL)
 		{
@@ -61,7 +61,7 @@ public:
 		}
 	}
 
-	void PostInitNewPin(UEdGraphPin* Pin, FOptionalPinFromProperty& Record, int32 ArrayIndex, UProperty* Property, uint8* PropertyAddress) const OVERRIDE
+	void PostInitNewPin(UEdGraphPin* Pin, FOptionalPinFromProperty& Record, int32 ArrayIndex, UProperty* Property, uint8* PropertyAddress) const override
 	{
 		check(PropertyAddress != NULL);
 		check(Record.bShowPin);
@@ -86,7 +86,7 @@ public:
 		}
 	}
 
-	void PostRemovedOldPin(FOptionalPinFromProperty& Record, int32 ArrayIndex, UProperty* Property, uint8* PropertyAddress) const OVERRIDE
+	void PostRemovedOldPin(FOptionalPinFromProperty& Record, int32 ArrayIndex, UProperty* Property, uint8* PropertyAddress) const override
 	{
 		check(PropertyAddress != NULL);
 		check(!Record.bShowPin);
diff --git a/Engine/Source/Editor/AnimGraph/Private/AnimGraphNode_SequencePlayer.cpp b/Engine/Source/Editor/AnimGraph/Private/AnimGraphNode_SequencePlayer.cpp
--- a/Engine/Source/Editor/AnimGraph/Private/AnimGraphNode_SequencePlayer.cpp
+++ b/Engine/Source/Editor/AnimGraph/Private/AnimGraphNode_SequencePlayer.cpp
@@ -33,7 +33,7 @@ public:
 		Keywords = InAssetInfo.ObjectPath.ToString();
 	}
 
-	virtual UEdGraphNode* PerformAction(class UEdGraph* ParentGraph, UEdGraphPin* FromPin, const FVector2D Location, bool bSelectNewNode = true) OVERRIDE
+	virtual UEdGraphNode* PerformAction(class UEdGraph* ParentGraph, UEdGraphPin* FromPin, const FVector2D Location, bool bSelectNewNode = true) override
 	{
 		UAnimGraphNode_SequencePlayer* SpawnedNode = CastChecked<UAnimGraphNode_SequencePlayer>(FEdGraphSchemaAction_K2NewNode::PerformAction(ParentGraph, FromPin, Location, bSelectNewNode));
 		SpawnedNode->Node.Sequence = Cast<UAnimSequence>(AssetInfo.GetAsset());
diff --git a/Engine/Source/Editor/AnimGraph/Private/AnimStateTransitionNode.cpp b/Engine/Source/Editor/AnimGraph/Private/AnimStateTransitionNode.cpp
--- a/Engine/Source/Editor/AnimGraph/Private/AnimStateTransitionNode.cpp
+++ b/Engine/Source/Editor/AnimGraph/Private/AnimStateTransitionNode.cpp
@@ -39,11 +39,11 @@ protected:
 class ANIMGRAPH_API FAnimStateTransitionNodeSharedRulesHelper : public IAnimStateTransitionNodeSharedDataHelper
 {
 protected:
-	virtual bool CheckIfNodesShouldShareData(const UAnimStateTransitionNode* NodeA, const UAnimStateTransitionNode* NodeB) OVERRIDE;
-	virtual bool CheckIfHasDataToShare(const UAnimStateTransitionNode* Node) OVERRIDE;
-	virtual void ShareData(UAnimStateTransitionNode* NodeWhoWantsToShare, const UAnimStateTransitionNode* ShareFrom) OVERRIDE;
-	virtual FString& AccessShareDataName(UAnimStateTransitionNode* Node) OVERRIDE;
-	virtual FGuid& AccessShareDataGuid(UAnimStateTransitionNode* Node) OVERRIDE;
+	virtual bool CheckIfNodesShouldShareData(const UAnimStateTransitionNode* NodeA, const UAnimStateTransitionNode* NodeB) override;
+	virtual bool CheckIfHasDataToShare(const UAnimStateTransitionNode* Node) override;
+	virtual void ShareData(UAnimStateTransitionNode* NodeWhoWantsToShare, const UAnimStateTransitionNode* ShareFrom) override;
+	virtual FString& AccessShareDataName(UAnimStateTransitionNode* Node) override;
+	virtual FGuid& AccessShareDataGuid(UAnimStateTransitionNode* Node) override;
 };
 
 //////////////////////////////////////////////////////////////////////////
@@ -52,11 +52,11 @@ protected:
 class ANIMGRAPH_API FAnimStateTransitionNodeSharedCrossfadeHelper : public IAnimStateTransitionNodeSharedDataHelper
 {
 protected:
-	virtual bool CheckIfNodesShouldShareData(const UAnimStateTransitionNode* NodeA, const UAnimStateTransitionNode* NodeB) OVERRIDE;
-	virtual bool CheckIfHasDataToShare(const UAnimStateTransitionNode* Node) OVERRIDE;
-	virtual void ShareData(UAnimStateTransitionNode* NodeWhoWantsToShare, const UAnimStateTransitionNode* ShareFrom) OVERRIDE;
-	virtual FString& AccessShareDataName(UAnimStateTransitionNode* Node) OVERRIDE;
-	virtual FGuid& AccessShareDataGuid(UAnimStateTransitionNode* Node) OVERRIDE;
+	virtual bool CheckIfNodesShouldShareData(const UAnimStateTransitionNode* NodeA, const UAnimStateTransitionNode* NodeB) override;
+	virtual bool CheckIfHasDataToShare(const UAnimStateTransitionNode* Node) override;
+	virtual void ShareData(UAnimStateTransitionNode* NodeWhoWantsToShare, const UAnimStateTransitionNode* ShareFrom) override;
+	virtual FString& AccessShareDataName(UAnimStateTransitionNode* Node) override;
+	virtual FGuid& AccessShareDataGuid(UAnimStateTransitionNode* Node) override;
 };
 
 /////////////////////////////////////////////////////
diff --git a/Engine/Source/Editor/BehaviorTreeEditor/Classes/BehaviorTreeDecoratorGraphNode.h b/Engine/Source/Editor/BehaviorTreeEditor/Classes/BehaviorTreeDecoratorGraphNode.h
--- a/Engine/Source/Editor/BehaviorTreeEditor/Classes/BehaviorTreeDecoratorGraphNode.h
+++ b/Engine/Source/Editor/BehaviorTreeEditor/Classes/BehaviorTreeDecoratorGraphNode.h
@@ -15,7 +15,7 @@ class UBehaviorTreeDecoratorGraphNode : public UEdGraphNode
 	/** Get the BT graph that owns this node */
 	virtual class UBehaviorTreeDecoratorGraph* GetDecoratorGraph();
 
-	virtual void AutowireNewNode(UEdGraphPin* FromPin) OVERRIDE;
+	virtual void AutowireNewNode(UEdGraphPin* FromPin) override;
 
 	// @return the input pin for this state
 	virtual UEdGraphPin* GetInputPin(int32 InputIndex=0) const;
@@ -27,9 +27,9 @@ class UBehaviorTreeDecoratorGraphNode : public UEdGraphNode
 	/** allow creating / removing inputs */
 	uint32 bAllowModifingInputs : 1;
 
-	virtual void NodeConnectionListChanged() OVERRIDE;
+	virtual void NodeConnectionListChanged() override;
 
-	virtual bool CanCreateUnderSpecifiedSchema(const UEdGraphSchema* DesiredSchema) const OVERRIDE;
+	virtual bool CanCreateUnderSpecifiedSchema(const UEdGraphSchema* DesiredSchema) const override;
 
-	virtual bool CanUserDeleteNode() const OVERRIDE;
+	virtual bool CanUserDeleteNode() const override;
 };
diff --git a/Engine/Source/Editor/BehaviorTreeEditor/Classes/BehaviorTreeDecoratorGraphNode_Decorator.h b/Engine/Source/Editor/BehaviorTreeEditor/Classes/BehaviorTreeDecoratorGraphNode_Decorator.h
--- a/Engine/Source/Editor/BehaviorTreeEditor/Classes/BehaviorTreeDecoratorGraphNode_Decorator.h
+++ b/Engine/Source/Editor/BehaviorTreeEditor/Classes/BehaviorTreeDecoratorGraphNode_Decorator.h
@@ -14,14 +14,14 @@ class UBehaviorTreeDecoratorGraphNode_Decorator : public UBehaviorTreeDecoratorG
 
 	struct FClassData ClassData;
 
-	virtual void PostPlacedNewNode() OVERRIDE;
-	virtual void AllocateDefaultPins() OVERRIDE;
-	virtual FText GetNodeTitle(ENodeTitleType::Type TitleType) const OVERRIDE;
+	virtual void PostPlacedNewNode() override;
+	virtual void AllocateDefaultPins() override;
+	virtual FText GetNodeTitle(ENodeTitleType::Type TitleType) const override;
 
-	virtual EBTDecoratorLogic::Type GetOperationType() const OVERRIDE;
+	virtual EBTDecoratorLogic::Type GetOperationType() const override;
 
-	virtual void PostEditImport() OVERRIDE;
-	virtual void PrepareForCopying() OVERRIDE;
+	virtual void PostEditImport() override;
+	virtual void PrepareForCopying() override;
 	void PostCopyNode();
 	bool RefreshNodeClass();
 
diff --git a/Engine/Source/Editor/BehaviorTreeEditor/Classes/BehaviorTreeDecoratorGraphNode_Logic.h b/Engine/Source/Editor/BehaviorTreeEditor/Classes/BehaviorTreeDecoratorGraphNode_Logic.h
--- a/Engine/Source/Editor/BehaviorTreeEditor/Classes/BehaviorTreeDecoratorGraphNode_Logic.h
+++ b/Engine/Source/Editor/BehaviorTreeEditor/Classes/BehaviorTreeDecoratorGraphNode_Logic.h
@@ -23,12 +23,12 @@ class UBehaviorTreeDecoratorGraphNode_Logic : public UBehaviorTreeDecoratorGraph
 	UPROPERTY()
 	TEnumAsByte<EDecoratorLogicMode::Type> LogicMode;
 
-	virtual FText GetNodeTitle(ENodeTitleType::Type TitleType) const OVERRIDE;
-	virtual void AllocateDefaultPins() OVERRIDE;
-	virtual bool CanUserDeleteNode() const OVERRIDE;
-	virtual EBTDecoratorLogic::Type GetOperationType() const OVERRIDE;
+	virtual FText GetNodeTitle(ENodeTitleType::Type TitleType) const override;
+	virtual void AllocateDefaultPins() override;
+	virtual bool CanUserDeleteNode() const override;
+	virtual EBTDecoratorLogic::Type GetOperationType() const override;
 
-	virtual void GetMenuEntries(struct FGraphContextMenuBuilder& ContextMenuBuilder) const OVERRIDE;
+	virtual void GetMenuEntries(struct FGraphContextMenuBuilder& ContextMenuBuilder) const override;
 
 	bool CanRemovePins() const;
 	bool CanAddPins() const;
diff --git a/Engine/Source/Editor/BehaviorTreeEditor/Classes/BehaviorTreeFactory.h b/Engine/Source/Editor/BehaviorTreeEditor/Classes/BehaviorTreeFactory.h
--- a/Engine/Source/Editor/BehaviorTreeEditor/Classes/BehaviorTreeFactory.h
+++ b/Engine/Source/Editor/BehaviorTreeEditor/Classes/BehaviorTreeFactory.h
@@ -9,8 +9,8 @@ class UBehaviorTreeFactory : public UFactory
 	GENERATED_UCLASS_BODY()
 
 	// Begin UFactory Interface
-	virtual UObject* FactoryCreateNew(UClass* Class,UObject* InParent,FName Name,EObjectFlags Flags,UObject* Context,FFeedbackContext* Warn) OVERRIDE;
-	virtual bool ShouldShowInNewMenu() const OVERRIDE;
+	virtual UObject* FactoryCreateNew(UClass* Class,UObject* InParent,FName Name,EObjectFlags Flags,UObject* Context,FFeedbackContext* Warn) override;
+	virtual bool ShouldShowInNewMenu() const override;
 	// Begin UFactory Interface	
 };
 
diff --git a/Engine/Source/Editor/BehaviorTreeEditor/Classes/BehaviorTreeGraphNode.h b/Engine/Source/Editor/BehaviorTreeEditor/Classes/BehaviorTreeGraphNode.h
--- a/Engine/Source/Editor/BehaviorTreeEditor/Classes/BehaviorTreeGraphNode.h
+++ b/Engine/Source/Editor/BehaviorTreeEditor/Classes/BehaviorTreeGraphNode.h
@@ -27,21 +27,21 @@ class UBehaviorTreeGraphNode : public UEdGraphNode
 
 	// Begin UEdGraphNode Interface
 	virtual class UBehaviorTreeGraph* GetBehaviorTreeGraph();
-	virtual void AllocateDefaultPins() OVERRIDE;
-	virtual void AutowireNewNode(UEdGraphPin* FromPin) OVERRIDE;
-	virtual void PostPlacedNewNode() OVERRIDE;
-	virtual void PrepareForCopying() OVERRIDE;
-	virtual bool CanDuplicateNode() const OVERRIDE;
-	virtual bool CanUserDeleteNode() const OVERRIDE;
-	virtual void DestroyNode() OVERRIDE;
-	virtual FString GetTooltip() const OVERRIDE;
-	virtual void NodeConnectionListChanged() OVERRIDE;
-	virtual bool CanCreateUnderSpecifiedSchema(const UEdGraphSchema* DesiredSchema) const OVERRIDE;
-	virtual void FindDiffs(class UEdGraphNode* OtherNode, struct FDiffResults& Results) OVERRIDE;
+	virtual void AllocateDefaultPins() override;
+	virtual void AutowireNewNode(UEdGraphPin* FromPin) override;
+	virtual void PostPlacedNewNode() override;
+	virtual void PrepareForCopying() override;
+	virtual bool CanDuplicateNode() const override;
+	virtual bool CanUserDeleteNode() const override;
+	virtual void DestroyNode() override;
+	virtual FString GetTooltip() const override;
+	virtual void NodeConnectionListChanged() override;
+	virtual bool CanCreateUnderSpecifiedSchema(const UEdGraphSchema* DesiredSchema) const override;
+	virtual void FindDiffs(class UEdGraphNode* OtherNode, struct FDiffResults& Results) override;
 	// End UEdGraphNode Interface
 
 	// Begin UObject Interface
-	virtual void PostEditImport() OVERRIDE;
+	virtual void PostEditImport() override;
 	// End UObject
 
 	// @return the input pin for this state
diff --git a/Engine/Source/Editor/BehaviorTreeEditor/Classes/BehaviorTreeGraphNode_Composite.h b/Engine/Source/Editor/BehaviorTreeEditor/Classes/BehaviorTreeGraphNode_Composite.h
--- a/Engine/Source/Editor/BehaviorTreeEditor/Classes/BehaviorTreeGraphNode_Composite.h
+++ b/Engine/Source/Editor/BehaviorTreeEditor/Classes/BehaviorTreeGraphNode_Composite.h
@@ -8,12 +8,12 @@ class UBehaviorTreeGraphNode_Composite : public UBehaviorTreeGraphNode
 {
 	GENERATED_UCLASS_BODY()
 	
-	virtual FText GetNodeTitle(ENodeTitleType::Type TitleType) const OVERRIDE;
-	virtual bool RefreshNodeClass() OVERRIDE{ return false; }
+	virtual FText GetNodeTitle(ENodeTitleType::Type TitleType) const override;
+	virtual bool RefreshNodeClass() override{ return false; }
 
 	/** Gets a list of actions that can be done to this particular node */
-	virtual void GetContextMenuActions(const FGraphNodeContextMenuBuilder& Context) const OVERRIDE;
+	virtual void GetContextMenuActions(const FGraphNodeContextMenuBuilder& Context) const override;
 
 	/** check if node can accept breakpoints */
-	virtual bool CanPlaceBreakpoints() const OVERRIDE { return true; }
+	virtual bool CanPlaceBreakpoints() const override { return true; }
 };
diff --git a/Engine/Source/Editor/BehaviorTreeEditor/Classes/BehaviorTreeGraphNode_CompositeDecorator.h b/Engine/Source/Editor/BehaviorTreeEditor/Classes/BehaviorTreeGraphNode_CompositeDecorator.h
--- a/Engine/Source/Editor/BehaviorTreeEditor/Classes/BehaviorTreeGraphNode_CompositeDecorator.h
+++ b/Engine/Source/Editor/BehaviorTreeEditor/Classes/BehaviorTreeGraphNode_CompositeDecorator.h
@@ -24,18 +24,18 @@ class UBehaviorTreeGraphNode_CompositeDecorator : public UBehaviorTreeGraphNode
 	uint32 bCanAbortFlow : 1;
 
 	FString GetNodeTypeDescription() const;
-	virtual FText GetNodeTitle(ENodeTitleType::Type TitleType) const OVERRIDE;
-	virtual void AllocateDefaultPins() OVERRIDE;
-	virtual FString GetDescription() const OVERRIDE;
-	virtual void PostPlacedNewNode() OVERRIDE;
-	virtual void PostLoad() OVERRIDE;
+	virtual FText GetNodeTitle(ENodeTitleType::Type TitleType) const override;
+	virtual void AllocateDefaultPins() override;
+	virtual FString GetDescription() const override;
+	virtual void PostPlacedNewNode() override;
+	virtual void PostLoad() override;
 	virtual UEdGraph* GetBoundGraph() const { return BoundGraph; }
-	virtual bool IsSubNode() const OVERRIDE;
-	virtual bool HasErrors() const OVERRIDE;
-	virtual bool RefreshNodeClass() OVERRIDE;
+	virtual bool IsSubNode() const override;
+	virtual bool HasErrors() const override;
+	virtual bool RefreshNodeClass() override;
 
-	virtual void PrepareForCopying() OVERRIDE;
-	virtual void PostCopyNode() OVERRIDE;
+	virtual void PrepareForCopying() override;
+	virtual void PostCopyNode() override;
 
 	void CollectDecoratorData(TArray<class UBTDecorator*>& NodeInstances, TArray<struct FBTDecoratorLogic>& Operations) const;
 	void SetDecoratorData(class UBTCompositeNode* InParentNode, uint8 InChildIndex);
@@ -44,7 +44,7 @@ class UBehaviorTreeGraphNode_CompositeDecorator : public UBehaviorTreeGraphNode
 	void OnInnerGraphChanged();
 	void BuildDescription();
 
-	virtual void PostEditChangeProperty(struct FPropertyChangedEvent& PropertyChangedEvent) OVERRIDE;
+	virtual void PostEditChangeProperty(struct FPropertyChangedEvent& PropertyChangedEvent) override;
 
 	void ResetExecutionRange();
 
diff --git a/Engine/Source/Editor/BehaviorTreeEditor/Classes/BehaviorTreeGraphNode_Decorator.h b/Engine/Source/Editor/BehaviorTreeEditor/Classes/BehaviorTreeGraphNode_Decorator.h
--- a/Engine/Source/Editor/BehaviorTreeEditor/Classes/BehaviorTreeGraphNode_Decorator.h
+++ b/Engine/Source/Editor/BehaviorTreeEditor/Classes/BehaviorTreeGraphNode_Decorator.h
@@ -8,9 +8,9 @@ class UBehaviorTreeGraphNode_Decorator : public UBehaviorTreeGraphNode
 {
 	GENERATED_UCLASS_BODY()
 
-	virtual FText GetNodeTitle(ENodeTitleType::Type TitleType) const OVERRIDE;
-	virtual void AllocateDefaultPins() OVERRIDE;
+	virtual FText GetNodeTitle(ENodeTitleType::Type TitleType) const override;
+	virtual void AllocateDefaultPins() override;
 	void CollectDecoratorData(TArray<class UBTDecorator*>& NodeInstances, TArray<struct FBTDecoratorLogic>& Operations) const;
 
-	virtual bool IsSubNode() const OVERRIDE;
+	virtual bool IsSubNode() const override;
 };
diff --git a/Engine/Source/Editor/BehaviorTreeEditor/Classes/BehaviorTreeGraphNode_Root.h b/Engine/Source/Editor/BehaviorTreeEditor/Classes/BehaviorTreeGraphNode_Root.h
--- a/Engine/Source/Editor/BehaviorTreeEditor/Classes/BehaviorTreeGraphNode_Root.h
+++ b/Engine/Source/Editor/BehaviorTreeEditor/Classes/BehaviorTreeGraphNode_Root.h
@@ -11,18 +11,18 @@ class UBehaviorTreeGraphNode_Root : public UBehaviorTreeGraphNode
 	UPROPERTY(EditAnywhere, Category=BehaviorTree)
 	class UBlackboardData* BlackboardAsset;
 
-	virtual void PostPlacedNewNode() OVERRIDE;
-	virtual void AllocateDefaultPins() OVERRIDE;
-	virtual bool CanDuplicateNode() const OVERRIDE { return false; }
-	virtual bool CanUserDeleteNode() const OVERRIDE{ return false; }
-	virtual bool HasErrors() const OVERRIDE { return false; }
-	virtual bool RefreshNodeClass() OVERRIDE { return false; }
-	virtual FText GetNodeTitle(ENodeTitleType::Type TitleType) const OVERRIDE;
+	virtual void PostPlacedNewNode() override;
+	virtual void AllocateDefaultPins() override;
+	virtual bool CanDuplicateNode() const override { return false; }
+	virtual bool CanUserDeleteNode() const override{ return false; }
+	virtual bool HasErrors() const override { return false; }
+	virtual bool RefreshNodeClass() override { return false; }
+	virtual FText GetNodeTitle(ENodeTitleType::Type TitleType) const override;
 	/** gets icon resource name for title bar */
-	virtual FName GetNameIcon() const OVERRIDE;
+	virtual FName GetNameIcon() const override;
 
-	virtual void PostEditChangeProperty( struct FPropertyChangedEvent& PropertyChangedEvent) OVERRIDE;
-	virtual FString	GetDescription() const OVERRIDE;
+	virtual void PostEditChangeProperty( struct FPropertyChangedEvent& PropertyChangedEvent) override;
+	virtual FString	GetDescription() const override;
 
 	/** notify behavior tree about blackboard change */
 	void UpdateBlackboard();
diff --git a/Engine/Source/Editor/BehaviorTreeEditor/Classes/BehaviorTreeGraphNode_Service.h b/Engine/Source/Editor/BehaviorTreeEditor/Classes/BehaviorTreeGraphNode_Service.h
--- a/Engine/Source/Editor/BehaviorTreeEditor/Classes/BehaviorTreeGraphNode_Service.h
+++ b/Engine/Source/Editor/BehaviorTreeEditor/Classes/BehaviorTreeGraphNode_Service.h
@@ -8,8 +8,8 @@ class UBehaviorTreeGraphNode_Service : public UBehaviorTreeGraphNode
 {
 	GENERATED_UCLASS_BODY()
 
-	virtual FText GetNodeTitle(ENodeTitleType::Type TitleType) const OVERRIDE;
-	virtual void AllocateDefaultPins() OVERRIDE;
-	virtual bool IsSubNode() const OVERRIDE;
+	virtual FText GetNodeTitle(ENodeTitleType::Type TitleType) const override;
+	virtual void AllocateDefaultPins() override;
+	virtual bool IsSubNode() const override;
 };
 
diff --git a/Engine/Source/Editor/BehaviorTreeEditor/Classes/BehaviorTreeGraphNode_SimpleParallel.h b/Engine/Source/Editor/BehaviorTreeEditor/Classes/BehaviorTreeGraphNode_SimpleParallel.h
--- a/Engine/Source/Editor/BehaviorTreeEditor/Classes/BehaviorTreeGraphNode_SimpleParallel.h
+++ b/Engine/Source/Editor/BehaviorTreeEditor/Classes/BehaviorTreeGraphNode_SimpleParallel.h
@@ -8,6 +8,6 @@ class UBehaviorTreeGraphNode_SimpleParallel : public UBehaviorTreeGraphNode_Comp
 {
 	GENERATED_UCLASS_BODY()
 
-	virtual void AllocateDefaultPins() OVERRIDE;
-	virtual void GetPinHoverText(const UEdGraphPin& Pin, FString& HoverTextOut) const OVERRIDE;
+	virtual void AllocateDefaultPins() override;
+	virtual void GetPinHoverText(const UEdGraphPin& Pin, FString& HoverTextOut) const override;
 };
diff --git a/Engine/Source/Editor/BehaviorTreeEditor/Classes/BehaviorTreeGraphNode_Task.h b/Engine/Source/Editor/BehaviorTreeEditor/Classes/BehaviorTreeGraphNode_Task.h
--- a/Engine/Source/Editor/BehaviorTreeEditor/Classes/BehaviorTreeGraphNode_Task.h
+++ b/Engine/Source/Editor/BehaviorTreeEditor/Classes/BehaviorTreeGraphNode_Task.h
@@ -9,10 +9,10 @@ class UBehaviorTreeGraphNode_Task : public UBehaviorTreeGraphNode
 {
 	GENERATED_UCLASS_BODY()
 
-	virtual void AllocateDefaultPins() OVERRIDE;
-	virtual FText GetNodeTitle(ENodeTitleType::Type TitleType) const OVERRIDE;
+	virtual void AllocateDefaultPins() override;
+	virtual FText GetNodeTitle(ENodeTitleType::Type TitleType) const override;
 	/** Gets a list of actions that can be done to this particular node */
-	virtual void GetContextMenuActions(const FGraphNodeContextMenuBuilder& Context) const OVERRIDE;
+	virtual void GetContextMenuActions(const FGraphNodeContextMenuBuilder& Context) const override;
 
-	virtual bool CanPlaceBreakpoints() const OVERRIDE { return true; }
+	virtual bool CanPlaceBreakpoints() const override { return true; }
 };
diff --git a/Engine/Source/Editor/BehaviorTreeEditor/Classes/EdGraphSchema_BehaviorTree.h b/Engine/Source/Editor/BehaviorTreeEditor/Classes/EdGraphSchema_BehaviorTree.h
--- a/Engine/Source/Editor/BehaviorTreeEditor/Classes/EdGraphSchema_BehaviorTree.h
+++ b/Engine/Source/Editor/BehaviorTreeEditor/Classes/EdGraphSchema_BehaviorTree.h
@@ -25,9 +25,9 @@ struct FBehaviorTreeSchemaAction_NewNode : public FEdGraphSchemaAction
 	{}
 
 	// FEdGraphSchemaAction interface
-	virtual UEdGraphNode* PerformAction(class UEdGraph* ParentGraph, UEdGraphPin* FromPin, const FVector2D Location, bool bSelectNewNode = true) OVERRIDE;
-	virtual UEdGraphNode* PerformAction(class UEdGraph* ParentGraph, TArray<UEdGraphPin*>& FromPins, const FVector2D Location, bool bSelectNewNode = true) OVERRIDE;
-	virtual void AddReferencedObjects( FReferenceCollector& Collector ) OVERRIDE;
+	virtual UEdGraphNode* PerformAction(class UEdGraph* ParentGraph, UEdGraphPin* FromPin, const FVector2D Location, bool bSelectNewNode = true) override;
+	virtual UEdGraphNode* PerformAction(class UEdGraph* ParentGraph, TArray<UEdGraphPin*>& FromPins, const FVector2D Location, bool bSelectNewNode = true) override;
+	virtual void AddReferencedObjects( FReferenceCollector& Collector ) override;
 	// End of FEdGraphSchemaAction interface
 
 	template <typename NodeType>
@@ -65,9 +65,9 @@ struct FBehaviorTreeSchemaAction_NewSubNode : public FEdGraphSchemaAction
 	{}
 
 	// FEdGraphSchemaAction interface
-	virtual UEdGraphNode* PerformAction(class UEdGraph* ParentGraph, UEdGraphPin* FromPin, const FVector2D Location, bool bSelectNewNode = true) OVERRIDE;
-	virtual UEdGraphNode* PerformAction(class UEdGraph* ParentGraph, TArray<UEdGraphPin*>& FromPins, const FVector2D Location, bool bSelectNewNode = true) OVERRIDE;
-	virtual void AddReferencedObjects( FReferenceCollector& Collector ) OVERRIDE;
+	virtual UEdGraphNode* PerformAction(class UEdGraph* ParentGraph, UEdGraphPin* FromPin, const FVector2D Location, bool bSelectNewNode = true) override;
+	virtual UEdGraphNode* PerformAction(class UEdGraph* ParentGraph, TArray<UEdGraphPin*>& FromPins, const FVector2D Location, bool bSelectNewNode = true) override;
+	virtual void AddReferencedObjects( FReferenceCollector& Collector ) override;
 	// End of FEdGraphSchemaAction interface
 };
 
@@ -85,7 +85,7 @@ struct FBehaviorTreeSchemaAction_AutoArrange : public FEdGraphSchemaAction
 	{}
 
 	// FEdGraphSchemaAction interface
-	virtual UEdGraphNode* PerformAction(class UEdGraph* ParentGraph, UEdGraphPin* FromPin, const FVector2D Location, bool bSelectNewNode = true) OVERRIDE;
+	virtual UEdGraphNode* PerformAction(class UEdGraph* ParentGraph, UEdGraphPin* FromPin, const FVector2D Location, bool bSelectNewNode = true) override;
 	// End of FEdGraphSchemaAction interface
 };
 
@@ -97,15 +97,15 @@ class UEdGraphSchema_BehaviorTree : public UEdGraphSchema
 	void GetBreakLinkToSubMenuActions(class FMenuBuilder& MenuBuilder, class UEdGraphPin* InGraphPin);
 
 	// Begin EdGraphSchema interface
-	virtual void CreateDefaultNodesForGraph(UEdGraph& Graph) const OVERRIDE;
-	virtual void GetGraphContextActions(FGraphContextMenuBuilder& ContextMenuBuilder) const OVERRIDE;
-	virtual void GetContextMenuActions(const UEdGraph* CurrentGraph, const UEdGraphNode* InGraphNode, const UEdGraphPin* InGraphPin, class FMenuBuilder* MenuBuilder, bool bIsDebugging) const OVERRIDE;
-	virtual const FPinConnectionResponse CanCreateConnection(const UEdGraphPin* A, const UEdGraphPin* B) const OVERRIDE;
-	virtual const FPinConnectionResponse CanMergeNodes(const UEdGraphNode* A, const UEdGraphNode* B) const OVERRIDE;
-	virtual FLinearColor GetPinTypeColor(const FEdGraphPinType& PinType) const OVERRIDE;
-	virtual bool ShouldHidePinDefaultValue(UEdGraphPin* Pin) const OVERRIDE;
-	virtual class FConnectionDrawingPolicy* CreateConnectionDrawingPolicy(int32 InBackLayerID, int32 InFrontLayerID, float InZoomFactor, const FSlateRect& InClippingRect, class FSlateWindowElementList& InDrawElements, class UEdGraph* InGraphObj) const OVERRIDE;
-	virtual int32 GetNodeSelectionCount(const UEdGraph* Graph) const OVERRIDE;
+	virtual void CreateDefaultNodesForGraph(UEdGraph& Graph) const override;
+	virtual void GetGraphContextActions(FGraphContextMenuBuilder& ContextMenuBuilder) const override;
+	virtual void GetContextMenuActions(const UEdGraph* CurrentGraph, const UEdGraphNode* InGraphNode, const UEdGraphPin* InGraphPin, class FMenuBuilder* MenuBuilder, bool bIsDebugging) const override;
+	virtual const FPinConnectionResponse CanCreateConnection(const UEdGraphPin* A, const UEdGraphPin* B) const override;
+	virtual const FPinConnectionResponse CanMergeNodes(const UEdGraphNode* A, const UEdGraphNode* B) const override;
+	virtual FLinearColor GetPinTypeColor(const FEdGraphPinType& PinType) const override;
+	virtual bool ShouldHidePinDefaultValue(UEdGraphPin* Pin) const override;
+	virtual class FConnectionDrawingPolicy* CreateConnectionDrawingPolicy(int32 InBackLayerID, int32 InFrontLayerID, float InZoomFactor, const FSlateRect& InClippingRect, class FSlateWindowElementList& InDrawElements, class UEdGraph* InGraphObj) const override;
+	virtual int32 GetNodeSelectionCount(const UEdGraph* Graph) const override;
 	// End EdGraphSchema interface
 
 	/* subnode creation context menu */
diff --git a/Engine/Source/Editor/BehaviorTreeEditor/Classes/EdGraphSchema_BehaviorTreeDecorator.h b/Engine/Source/Editor/BehaviorTreeEditor/Classes/EdGraphSchema_BehaviorTreeDecorator.h
--- a/Engine/Source/Editor/BehaviorTreeEditor/Classes/EdGraphSchema_BehaviorTreeDecorator.h
+++ b/Engine/Source/Editor/BehaviorTreeEditor/Classes/EdGraphSchema_BehaviorTreeDecorator.h
@@ -24,9 +24,9 @@ struct FDecoratorSchemaAction_NewNode : public FEdGraphSchemaAction
 	{}
 
 	// FEdGraphSchemaAction interface
-	virtual UEdGraphNode* PerformAction(class UEdGraph* ParentGraph, UEdGraphPin* FromPin, const FVector2D Location, bool bSelectNewNode = true) OVERRIDE;
-	virtual UEdGraphNode* PerformAction(class UEdGraph* ParentGraph, TArray<UEdGraphPin*>& FromPins, const FVector2D Location, bool bSelectNewNode = true) OVERRIDE;
-	virtual void AddReferencedObjects( FReferenceCollector& Collector ) OVERRIDE;
+	virtual UEdGraphNode* PerformAction(class UEdGraph* ParentGraph, UEdGraphPin* FromPin, const FVector2D Location, bool bSelectNewNode = true) override;
+	virtual UEdGraphNode* PerformAction(class UEdGraph* ParentGraph, TArray<UEdGraphPin*>& FromPins, const FVector2D Location, bool bSelectNewNode = true) override;
+	virtual void AddReferencedObjects( FReferenceCollector& Collector ) override;
 	// End of FEdGraphSchemaAction interface
 
 	template <typename NodeType>
@@ -52,12 +52,12 @@ class UEdGraphSchema_BehaviorTreeDecorator : public UEdGraphSchema
 	void RemovePin(class UEdGraphPin* InGraphPin);
 
 	// Begin EdGraphSchema interface
-	virtual void CreateDefaultNodesForGraph(UEdGraph& Graph) const OVERRIDE;
-	virtual void GetGraphContextActions(FGraphContextMenuBuilder& ContextMenuBuilder) const OVERRIDE;
-	virtual void GetContextMenuActions(const UEdGraph* CurrentGraph, const UEdGraphNode* InGraphNode, const UEdGraphPin* InGraphPin, class FMenuBuilder* MenuBuilder, bool bIsDebugging) const OVERRIDE;
-	virtual const FPinConnectionResponse CanCreateConnection(const UEdGraphPin* A, const UEdGraphPin* B) const OVERRIDE;
-	virtual FLinearColor GetPinTypeColor(const FEdGraphPinType& PinType) const OVERRIDE;
-	virtual bool ShouldHidePinDefaultValue(UEdGraphPin* Pin) const OVERRIDE;
+	virtual void CreateDefaultNodesForGraph(UEdGraph& Graph) const override;
+	virtual void GetGraphContextActions(FGraphContextMenuBuilder& ContextMenuBuilder) const override;
+	virtual void GetContextMenuActions(const UEdGraph* CurrentGraph, const UEdGraphNode* InGraphNode, const UEdGraphPin* InGraphPin, class FMenuBuilder* MenuBuilder, bool bIsDebugging) const override;
+	virtual const FPinConnectionResponse CanCreateConnection(const UEdGraphPin* A, const UEdGraphPin* B) const override;
+	virtual FLinearColor GetPinTypeColor(const FEdGraphPinType& PinType) const override;
+	virtual bool ShouldHidePinDefaultValue(UEdGraphPin* Pin) const override;
 	// End EdGraphSchema interface
 
 	static TSharedPtr<FDecoratorSchemaAction_NewNode> AddNewDecoratorAction(FGraphContextMenuBuilder& ContextMenuBuilder, const FString& Category, const FText& MenuDesc, const FString& Tooltip);
diff --git a/Engine/Source/Editor/BehaviorTreeEditor/Private/AssetTypeActions_BehaviorTree.h b/Engine/Source/Editor/BehaviorTreeEditor/Private/AssetTypeActions_BehaviorTree.h
--- a/Engine/Source/Editor/BehaviorTreeEditor/Private/AssetTypeActions_BehaviorTree.h
+++ b/Engine/Source/Editor/BehaviorTreeEditor/Private/AssetTypeActions_BehaviorTree.h
@@ -7,14 +7,14 @@ class FAssetTypeActions_BehaviorTree : public FAssetTypeActions_Base
 {
 public:
 	// IAssetTypeActions Implementation
-	virtual FText GetName() const OVERRIDE { return NSLOCTEXT("AssetTypeActions", "AssetTypeActions_BehaviorTree", "Behavior Tree"); }
-	virtual FColor GetTypeColor() const OVERRIDE { return FColor(149,70,255); }
-	virtual UClass* GetSupportedClass() const OVERRIDE;
-	virtual bool HasActions ( const TArray<UObject*>& InObjects ) const OVERRIDE;
-	virtual void GetActions( const TArray<UObject*>& InObjects, FMenuBuilder& MenuBuilder ) OVERRIDE;
-	virtual void OpenAssetEditor( const TArray<UObject*>& InObjects, TSharedPtr<class IToolkitHost> EditWithinLevelEditor = TSharedPtr<IToolkitHost>() ) OVERRIDE;
-	virtual uint32 GetCategories() OVERRIDE { return EAssetTypeCategories::Misc; }
-	virtual void PerformAssetDiff(UObject* Asset1, UObject* Asset2, const struct FRevisionInfo& OldRevision, const struct FRevisionInfo& NewRevision) const OVERRIDE;
+	virtual FText GetName() const override { return NSLOCTEXT("AssetTypeActions", "AssetTypeActions_BehaviorTree", "Behavior Tree"); }
+	virtual FColor GetTypeColor() const override { return FColor(149,70,255); }
+	virtual UClass* GetSupportedClass() const override;
+	virtual bool HasActions ( const TArray<UObject*>& InObjects ) const override;
+	virtual void GetActions( const TArray<UObject*>& InObjects, FMenuBuilder& MenuBuilder ) override;
+	virtual void OpenAssetEditor( const TArray<UObject*>& InObjects, TSharedPtr<class IToolkitHost> EditWithinLevelEditor = TSharedPtr<IToolkitHost>() ) override;
+	virtual uint32 GetCategories() override { return EAssetTypeCategories::Misc; }
+	virtual void PerformAssetDiff(UObject* Asset1, UObject* Asset2, const struct FRevisionInfo& OldRevision, const struct FRevisionInfo& NewRevision) const override;
 
 	/** Check whether the behavior tree system is enabled */
 	static bool BehaviorTreeEditorEnabled();
diff --git a/Engine/Source/Editor/BehaviorTreeEditor/Private/BehaviorTreeConnectionDrawingPolicy.h b/Engine/Source/Editor/BehaviorTreeEditor/Private/BehaviorTreeConnectionDrawingPolicy.h
--- a/Engine/Source/Editor/BehaviorTreeEditor/Private/BehaviorTreeConnectionDrawingPolicy.h
+++ b/Engine/Source/Editor/BehaviorTreeEditor/Private/BehaviorTreeConnectionDrawingPolicy.h
@@ -14,12 +14,12 @@ public:
 	FBehaviorTreeConnectionDrawingPolicy(int32 InBackLayerID, int32 InFrontLayerID, float ZoomFactor, const FSlateRect& InClippingRect, FSlateWindowElementList& InDrawElements, UEdGraph* InGraphObj);
 
 	// FConnectionDrawingPolicy interface 
-	virtual void DetermineWiringStyle(UEdGraphPin* OutputPin, UEdGraphPin* InputPin, /*inout*/ float& Thickness, /*inout*/ FLinearColor& WireColor, /*inout*/ bool& bDrawBubbles, /*inout*/ bool& bBidirectional) OVERRIDE;
-	virtual void Draw(TMap<TSharedRef<SWidget>, FArrangedWidget>& PinGeometries, FArrangedChildren& ArrangedNodes) OVERRIDE;
-	virtual void DrawSplineWithArrow(FGeometry& StartGeom, FGeometry& EndGeom, const FLinearColor& WireColor, float WireThickness, bool bDrawBubbles, bool Bidirectional) OVERRIDE;
-	virtual void DrawSplineWithArrow(const FVector2D& StartPoint, const FVector2D& EndPoint, const FLinearColor& WireColor, float WireThickness, bool bDrawBubbles, bool Bidirectional) OVERRIDE;
-	virtual void DrawPreviewConnector(const FGeometry& PinGeometry, const FVector2D& StartPoint, const FVector2D& EndPoint, UEdGraphPin* Pin) OVERRIDE;
-	virtual void DrawConnection(int32 LayerId, const FVector2D& Start, const FVector2D& End, const FLinearColor& InColor, float Thickness, bool bDrawBubbles) OVERRIDE;
+	virtual void DetermineWiringStyle(UEdGraphPin* OutputPin, UEdGraphPin* InputPin, /*inout*/ float& Thickness, /*inout*/ FLinearColor& WireColor, /*inout*/ bool& bDrawBubbles, /*inout*/ bool& bBidirectional) override;
+	virtual void Draw(TMap<TSharedRef<SWidget>, FArrangedWidget>& PinGeometries, FArrangedChildren& ArrangedNodes) override;
+	virtual void DrawSplineWithArrow(FGeometry& StartGeom, FGeometry& EndGeom, const FLinearColor& WireColor, float WireThickness, bool bDrawBubbles, bool Bidirectional) override;
+	virtual void DrawSplineWithArrow(const FVector2D& StartPoint, const FVector2D& EndPoint, const FLinearColor& WireColor, float WireThickness, bool bDrawBubbles, bool Bidirectional) override;
+	virtual void DrawPreviewConnector(const FGeometry& PinGeometry, const FVector2D& StartPoint, const FVector2D& EndPoint, UEdGraphPin* Pin) override;
+	virtual void DrawConnection(int32 LayerId, const FVector2D& Start, const FVector2D& End, const FLinearColor& InColor, float Thickness, bool bDrawBubbles) override;
 	// End of FConnectionDrawingPolicy interface
 
 protected:
diff --git a/Engine/Source/Editor/BehaviorTreeEditor/Private/BehaviorTreeDebugger.h b/Engine/Source/Editor/BehaviorTreeEditor/Private/BehaviorTreeDebugger.h
--- a/Engine/Source/Editor/BehaviorTreeEditor/Private/BehaviorTreeDebugger.h
+++ b/Engine/Source/Editor/BehaviorTreeEditor/Private/BehaviorTreeDebugger.h
@@ -10,11 +10,11 @@ public:
 	FBehaviorTreeDebugger();
 	~FBehaviorTreeDebugger();
 
-	virtual void Tick(float DeltaTime) OVERRIDE;
-	virtual bool IsTickable() const OVERRIDE;
+	virtual void Tick(float DeltaTime) override;
+	virtual bool IsTickable() const override;
 	virtual bool IsTickableWhenPaused() const override { return true; }
-	virtual bool IsTickableInEditor() const OVERRIDE { return true; }
-	virtual TStatId GetStatId() const OVERRIDE { RETURN_QUICK_DECLARE_CYCLE_STAT(FBehaviorTreeEditorTickHelper, STATGROUP_Tickables); }
+	virtual bool IsTickableInEditor() const override { return true; }
+	virtual TStatId GetStatId() const override { RETURN_QUICK_DECLARE_CYCLE_STAT(FBehaviorTreeEditorTickHelper, STATGROUP_Tickables); }
 
 	/** Refresh the debugging information we are displaying (only when paused, as Tick() updates when running) */
 	void Refresh();
diff --git a/Engine/Source/Editor/BehaviorTreeEditor/Private/BehaviorTreeEditor.h b/Engine/Source/Editor/BehaviorTreeEditor/Private/BehaviorTreeEditor.h
--- a/Engine/Source/Editor/BehaviorTreeEditor/Private/BehaviorTreeEditor.h
+++ b/Engine/Source/Editor/BehaviorTreeEditor/Private/BehaviorTreeEditor.h
@@ -12,29 +12,29 @@ public:
 	/** Destructor */
 	virtual ~FBehaviorTreeEditor();
 
-	virtual void RegisterTabSpawners(const TSharedRef<class FTabManager>& TabManager) OVERRIDE;
+	virtual void RegisterTabSpawners(const TSharedRef<class FTabManager>& TabManager) override;
 
 	void InitBehaviorTreeEditor( const EToolkitMode::Type Mode, const TSharedPtr< class IToolkitHost >& InitToolkitHost, UObject* InObject );
 
 	// Begin IToolkit interface
-	virtual FName GetToolkitFName() const OVERRIDE;
-	virtual FText GetBaseToolkitName() const OVERRIDE;
-	virtual FString GetWorldCentricTabPrefix() const OVERRIDE;
-	virtual FLinearColor GetWorldCentricTabColorScale() const OVERRIDE;
-	virtual FText GetToolkitName() const OVERRIDE;
+	virtual FName GetToolkitFName() const override;
+	virtual FText GetBaseToolkitName() const override;
+	virtual FString GetWorldCentricTabPrefix() const override;
+	virtual FLinearColor GetWorldCentricTabColorScale() const override;
+	virtual FText GetToolkitName() const override;
 	// End IToolkit interface
 
 	// Begin IBehaviorTreeEditor interface
-	virtual uint32 GetSelectedNodesCount() const OVERRIDE { return SelectedNodesCount; }
-	virtual void InitializeDebuggerState(class FBehaviorTreeDebugger* ParentDebugger) const OVERRIDE;
-	virtual UEdGraphNode* FindInjectedNode(int32 Index) const OVERRIDE;
-	virtual void DoubleClickNode(class UEdGraphNode* Node) OVERRIDE;
-	virtual void FocusWindow(UObject* ObjectToFocusOn = NULL) OVERRIDE;
+	virtual uint32 GetSelectedNodesCount() const override { return SelectedNodesCount; }
+	virtual void InitializeDebuggerState(class FBehaviorTreeDebugger* ParentDebugger) const override;
+	virtual UEdGraphNode* FindInjectedNode(int32 Index) const override;
+	virtual void DoubleClickNode(class UEdGraphNode* Node) override;
+	virtual void FocusWindow(UObject* ObjectToFocusOn = NULL) override;
 	// End IBehaviorTreeEditor interface
 
 	// Begin FEditorUndoClient Interface
-	virtual void	PostUndo(bool bSuccess) OVERRIDE;
-	virtual void	PostRedo(bool bSuccess) OVERRIDE;
+	virtual void	PostUndo(bool bSuccess) override;
+	virtual void	PostRedo(bool bSuccess) override;
 	// End of FEditorUndoClient
 
 	// Delegates
@@ -142,7 +142,7 @@ public:
 
 protected:
 	/** Called when "Save" is clicked for this asset */
-	virtual void SaveAsset_Execute() OVERRIDE;
+	virtual void SaveAsset_Execute() override;
 
 private:
 	/** Create widget for graph editing */
diff --git a/Engine/Source/Editor/BehaviorTreeEditor/Private/BehaviorTreeEditorModes.h b/Engine/Source/Editor/BehaviorTreeEditor/Private/BehaviorTreeEditorModes.h
--- a/Engine/Source/Editor/BehaviorTreeEditor/Private/BehaviorTreeEditorModes.h
+++ b/Engine/Source/Editor/BehaviorTreeEditor/Private/BehaviorTreeEditorModes.h
@@ -10,9 +10,9 @@ class FBehaviorTreeEditorApplicationMode : public FApplicationMode
 public:
 	FBehaviorTreeEditorApplicationMode(TSharedPtr<class FBehaviorTreeEditor> InBehaviorTreeEditor);
 
-	virtual void RegisterTabFactories(TSharedPtr<class FTabManager> InTabManager) OVERRIDE;
-	virtual void PreDeactivateMode() OVERRIDE;
-	virtual void PostActivateMode() OVERRIDE;
+	virtual void RegisterTabFactories(TSharedPtr<class FTabManager> InTabManager) override;
+	virtual void PreDeactivateMode() override;
+	virtual void PostActivateMode() override;
 
 protected:
 	TWeakPtr<class FBehaviorTreeEditor> BehaviorTreeEditor;
@@ -27,8 +27,8 @@ class FBlackboardEditorApplicationMode : public FApplicationMode
 public:
 	FBlackboardEditorApplicationMode(TSharedPtr<class FBehaviorTreeEditor> InBehaviorTreeEditor);
 
-	virtual void RegisterTabFactories(TSharedPtr<class FTabManager> InTabManager) OVERRIDE;
-	virtual void PostActivateMode() OVERRIDE;
+	virtual void RegisterTabFactories(TSharedPtr<class FTabManager> InTabManager) override;
+	virtual void PostActivateMode() override;
 
 protected:
 	TWeakPtr<class FBehaviorTreeEditor> BehaviorTreeEditor;
diff --git a/Engine/Source/Editor/BehaviorTreeEditor/Private/BehaviorTreeEditorModule.cpp b/Engine/Source/Editor/BehaviorTreeEditor/Private/BehaviorTreeEditorModule.cpp
--- a/Engine/Source/Editor/BehaviorTreeEditor/Private/BehaviorTreeEditorModule.cpp
+++ b/Engine/Source/Editor/BehaviorTreeEditor/Private/BehaviorTreeEditorModule.cpp
@@ -19,7 +19,7 @@ const FName FBehaviorTreeEditorModule::BehaviorTreeEditorAppIdentifier( TEXT( "B
 
 class FGraphPanelNodeFactory_BehaviorTree : public FGraphPanelNodeFactory
 {
-	virtual TSharedPtr<class SGraphNode> CreateNode(UEdGraphNode* Node) const OVERRIDE
+	virtual TSharedPtr<class SGraphNode> CreateNode(UEdGraphNode* Node) const override
 	{
 		if (UBehaviorTreeGraphNode* BTNode = Cast<UBehaviorTreeGraphNode>(Node))
 		{
diff --git a/Engine/Source/Editor/BehaviorTreeEditor/Private/BehaviorTreeEditorToolbar.cpp b/Engine/Source/Editor/BehaviorTreeEditor/Private/BehaviorTreeEditorToolbar.cpp
--- a/Engine/Source/Editor/BehaviorTreeEditor/Private/BehaviorTreeEditorToolbar.cpp
+++ b/Engine/Source/Editor/BehaviorTreeEditor/Private/BehaviorTreeEditorToolbar.cpp
@@ -24,7 +24,7 @@ public:
 	}
 
 	// SWidget interface
-	virtual FVector2D ComputeDesiredSize() const OVERRIDE
+	virtual FVector2D ComputeDesiredSize() const override
 	{
 		const float Height = 20.0f;
 		const float Thickness = 16.0f;
diff --git a/Engine/Source/Editor/BehaviorTreeEditor/Private/DetailCustomizations/BehaviorDecoratorDetails.h b/Engine/Source/Editor/BehaviorTreeEditor/Private/DetailCustomizations/BehaviorDecoratorDetails.h
--- a/Engine/Source/Editor/BehaviorTreeEditor/Private/DetailCustomizations/BehaviorDecoratorDetails.h
+++ b/Engine/Source/Editor/BehaviorTreeEditor/Private/DetailCustomizations/BehaviorDecoratorDetails.h
@@ -12,7 +12,7 @@ public:
 	static TSharedRef<IDetailCustomization> MakeInstance();
 
 	/** IDetailCustomization interface */
-	virtual void CustomizeDetails( IDetailLayoutBuilder& DetailLayout ) OVERRIDE;
+	virtual void CustomizeDetails( IDetailLayoutBuilder& DetailLayout ) override;
 
 	bool IsEditingEnabled() const;
 
diff --git a/Engine/Source/Editor/BehaviorTreeEditor/Private/DetailCustomizations/BlackboardSelectorDetails.h b/Engine/Source/Editor/BehaviorTreeEditor/Private/DetailCustomizations/BlackboardSelectorDetails.h
--- a/Engine/Source/Editor/BehaviorTreeEditor/Private/DetailCustomizations/BlackboardSelectorDetails.h
+++ b/Engine/Source/Editor/BehaviorTreeEditor/Private/DetailCustomizations/BlackboardSelectorDetails.h
@@ -12,8 +12,8 @@ public:
 	static TSharedRef<IPropertyTypeCustomization> MakeInstance();
 
 	/** IPropertyTypeCustomization interface */
-	virtual void CustomizeHeader( TSharedRef<class IPropertyHandle> StructPropertyHandle, class FDetailWidgetRow& HeaderRow, IPropertyTypeCustomizationUtils& StructCustomizationUtils ) OVERRIDE;
-	virtual void CustomizeChildren( TSharedRef<class IPropertyHandle> StructPropertyHandle, class IDetailChildrenBuilder& StructBuilder, IPropertyTypeCustomizationUtils& StructCustomizationUtils ) OVERRIDE;
+	virtual void CustomizeHeader( TSharedRef<class IPropertyHandle> StructPropertyHandle, class FDetailWidgetRow& HeaderRow, IPropertyTypeCustomizationUtils& StructCustomizationUtils ) override;
+	virtual void CustomizeChildren( TSharedRef<class IPropertyHandle> StructPropertyHandle, class IDetailChildrenBuilder& StructBuilder, IPropertyTypeCustomizationUtils& StructCustomizationUtils ) override;
 
 private:
 
diff --git a/Engine/Source/Editor/BehaviorTreeEditor/Private/SBehaviorTreeDiff.cpp b/Engine/Source/Editor/BehaviorTreeEditor/Private/SBehaviorTreeDiff.cpp
--- a/Engine/Source/Editor/BehaviorTreeEditor/Private/SBehaviorTreeDiff.cpp
+++ b/Engine/Source/Editor/BehaviorTreeEditor/Private/SBehaviorTreeDiff.cpp
@@ -58,7 +58,7 @@ public:
 	}
 
 	/** Initialize commands */
-	virtual void RegisterCommands() OVERRIDE
+	virtual void RegisterCommands() override
 	{
 		UI_COMMAND(Previous, "Prev", "Go to previous difference", EUserInterfaceActionType::Button, FInputGesture(EKeys::F7, EModifierKey::Control));
 		UI_COMMAND(Next, "Next", "Go to next difference", EUserInterfaceActionType::Button, FInputGesture(EKeys::F7));
diff --git a/Engine/Source/Editor/BehaviorTreeEditor/Private/SGraphNode_BehaviorTree.cpp b/Engine/Source/Editor/BehaviorTreeEditor/Private/SGraphNode_BehaviorTree.cpp
--- a/Engine/Source/Editor/BehaviorTreeEditor/Private/SGraphNode_BehaviorTree.cpp
+++ b/Engine/Source/Editor/BehaviorTreeEditor/Private/SGraphNode_BehaviorTree.cpp
@@ -26,10 +26,10 @@ public:
 	void Construct(const FArguments& InArgs, UEdGraphPin* InPin);
 protected:
 	// Begin SGraphPin interface
-	virtual TSharedRef<SWidget>	GetDefaultValueWidget() OVERRIDE;
+	virtual TSharedRef<SWidget>	GetDefaultValueWidget() override;
 
 	/** @return The color that we should use to draw this pin */
-	virtual FSlateColor GetPinColor() const OVERRIDE;
+	virtual FSlateColor GetPinColor() const override;
 
 	// End SGraphPin interface
 
diff --git a/Engine/Source/Editor/BehaviorTreeEditor/Private/SGraphNode_BehaviorTree.h b/Engine/Source/Editor/BehaviorTreeEditor/Private/SGraphNode_BehaviorTree.h
--- a/Engine/Source/Editor/BehaviorTreeEditor/Private/SGraphNode_BehaviorTree.h
+++ b/Engine/Source/Editor/BehaviorTreeEditor/Private/SGraphNode_BehaviorTree.h
@@ -11,19 +11,19 @@ public:
 	void Construct(const FArguments& InArgs, UBehaviorTreeGraphNode* InNode);
 
 	// SGraphNode interface
-	virtual void UpdateGraphNode() OVERRIDE;
-	virtual void CreatePinWidgets() OVERRIDE;
-	virtual void AddPin(const TSharedRef<SGraphPin>& PinToAdd) OVERRIDE;
-	virtual TSharedPtr<SToolTip> GetComplexTooltip() OVERRIDE;
-	virtual void OnDragEnter( const FGeometry& MyGeometry, const FDragDropEvent& DragDropEvent ) OVERRIDE;
-	virtual FReply OnDragOver( const FGeometry& MyGeometry, const FDragDropEvent& DragDropEvent ) OVERRIDE;
-	virtual FReply OnDrop( const FGeometry& MyGeometry, const FDragDropEvent& DragDropEvent ) OVERRIDE;
-	virtual void OnDragLeave( const FDragDropEvent& DragDropEvent ) OVERRIDE;
-	virtual void GetOverlayBrushes(bool bSelected, const FVector2D WidgetSize, TArray<FOverlayBrushInfo>& Brushes) const OVERRIDE;
+	virtual void UpdateGraphNode() override;
+	virtual void CreatePinWidgets() override;
+	virtual void AddPin(const TSharedRef<SGraphPin>& PinToAdd) override;
+	virtual TSharedPtr<SToolTip> GetComplexTooltip() override;
+	virtual void OnDragEnter( const FGeometry& MyGeometry, const FDragDropEvent& DragDropEvent ) override;
+	virtual FReply OnDragOver( const FGeometry& MyGeometry, const FDragDropEvent& DragDropEvent ) override;
+	virtual FReply OnDrop( const FGeometry& MyGeometry, const FDragDropEvent& DragDropEvent ) override;
+	virtual void OnDragLeave( const FDragDropEvent& DragDropEvent ) override;
+	virtual void GetOverlayBrushes(bool bSelected, const FVector2D WidgetSize, TArray<FOverlayBrushInfo>& Brushes) const override;
 	// End of SGraphNode interface
 
 	/** handle double click */
-	virtual FReply OnMouseButtonDoubleClick( const FGeometry& InMyGeometry, const FPointerEvent& InMouseEvent ) OVERRIDE;
+	virtual FReply OnMouseButtonDoubleClick( const FGeometry& InMyGeometry, const FPointerEvent& InMouseEvent ) override;
 
 	/** handle mouse down on the node */
 	FReply OnMouseDown(const FGeometry& SenderGeometry, const FPointerEvent& MouseEvent);
@@ -32,7 +32,7 @@ public:
 	FReply OnMouseUp(const FGeometry& SenderGeometry, const FPointerEvent& MouseEvent);
 
 	/** handle mouse move on the node - activate drag and drop when mouse button is being held */
-	virtual FReply OnMouseMove(const FGeometry& SenderGeometry, const FPointerEvent& MouseEvent) OVERRIDE;
+	virtual FReply OnMouseMove(const FGeometry& SenderGeometry, const FPointerEvent& MouseEvent) override;
 
 	/**
 	 * Ticks this widget.  Override in derived classes, but always call the parent implementation.
@@ -41,7 +41,7 @@ public:
 	 * @param  InCurrentTime  Current absolute real time
 	 * @param  InDeltaTime  Real time passed since last tick
 	 */
-	void Tick( const FGeometry& AllottedGeometry, const double InCurrentTime, const float InDeltaTime ) OVERRIDE;
+	void Tick( const FGeometry& AllottedGeometry, const double InCurrentTime, const float InDeltaTime ) override;
 	
 	/** gets decorator or service node if one is found under mouse cursor */
 	TSharedPtr<SGraphNode> GetSubNodeUnderCursor(const FGeometry& WidgetGeometry, const FPointerEvent& MouseEvent);
diff --git a/Engine/Source/Editor/BehaviorTreeEditor/Private/SGraphNode_Decorator.h b/Engine/Source/Editor/BehaviorTreeEditor/Private/SGraphNode_Decorator.h
--- a/Engine/Source/Editor/BehaviorTreeEditor/Private/SGraphNode_Decorator.h
+++ b/Engine/Source/Editor/BehaviorTreeEditor/Private/SGraphNode_Decorator.h
@@ -10,5 +10,5 @@ public:
 
 	void Construct(const FArguments& InArgs, UBehaviorTreeDecoratorGraphNode_Decorator* InNode);
 
-	virtual FString GetNodeComment() const OVERRIDE;
+	virtual FString GetNodeComment() const override;
 };
diff --git a/Engine/Source/Editor/BlueprintGraph/Classes/EdGraphSchema_K2_Actions.h b/Engine/Source/Editor/BlueprintGraph/Classes/EdGraphSchema_K2_Actions.h
--- a/Engine/Source/Editor/BlueprintGraph/Classes/EdGraphSchema_K2_Actions.h
+++ b/Engine/Source/Editor/BlueprintGraph/Classes/EdGraphSchema_K2_Actions.h
@@ -33,7 +33,7 @@ struct BLUEPRINTGRAPH_API FEdGraphSchemaAction_K2NewNode : public FEdGraphSchema
 
 	// Simple type info
 	static FString StaticGetTypeId() {static FString Type = TEXT("FEdGraphSchemaAction_K2NewNode"); return Type;}
-	virtual FString GetTypeId() const OVERRIDE { return StaticGetTypeId(); } 
+	virtual FString GetTypeId() const override { return StaticGetTypeId(); } 
 
 	/** Template of node we want to create */
 	UPROPERTY()
@@ -51,9 +51,9 @@ struct BLUEPRINTGRAPH_API FEdGraphSchemaAction_K2NewNode : public FEdGraphSchema
 	{}
 
 	// FEdGraphSchemaAction interface
-	virtual UEdGraphNode* PerformAction(class UEdGraph* ParentGraph, UEdGraphPin* FromPin, const FVector2D Location, bool bSelectNewNode = true) OVERRIDE;
-	virtual UEdGraphNode* PerformAction(class UEdGraph* ParentGraph, TArray<UEdGraphPin*>& FromPins, const FVector2D Location, bool bSelectNewNode = true) OVERRIDE;
-	virtual void AddReferencedObjects(FReferenceCollector& Collector) OVERRIDE;
+	virtual UEdGraphNode* PerformAction(class UEdGraph* ParentGraph, UEdGraphPin* FromPin, const FVector2D Location, bool bSelectNewNode = true) override;
+	virtual UEdGraphNode* PerformAction(class UEdGraph* ParentGraph, TArray<UEdGraphPin*>& FromPins, const FVector2D Location, bool bSelectNewNode = true) override;
+	virtual void AddReferencedObjects(FReferenceCollector& Collector) override;
 	// End of FEdGraphSchemaAction interface
 
 	template <typename NodeType>
@@ -80,7 +80,7 @@ struct BLUEPRINTGRAPH_API FEdGraphSchemaAction_K2ViewNode : public FEdGraphSchem
 
 	// Simple type info
 	static FString StaticGetTypeId() {static FString Type = TEXT("FEdGraphSchemaAction_K2ViewNode"); return Type;}
-	virtual FString GetTypeId() const OVERRIDE { return StaticGetTypeId(); } 
+	virtual FString GetTypeId() const override { return StaticGetTypeId(); } 
 
 	/** node we want to view */
 	UPROPERTY()
@@ -97,8 +97,8 @@ struct BLUEPRINTGRAPH_API FEdGraphSchemaAction_K2ViewNode : public FEdGraphSchem
 	{}
 
 	// FEdGraphSchemaAction interface
-	virtual UEdGraphNode* PerformAction(class UEdGraph* ParentGraph, UEdGraphPin* FromPin, const FVector2D Location, bool bSelectNewNode = true) OVERRIDE;
-	virtual UEdGraphNode* PerformAction(class UEdGraph* ParentGraph, TArray<UEdGraphPin*>& FromPins, const FVector2D Location, bool bSelectNewNode = true) OVERRIDE;
+	virtual UEdGraphNode* PerformAction(class UEdGraph* ParentGraph, UEdGraphPin* FromPin, const FVector2D Location, bool bSelectNewNode = true) override;
+	virtual UEdGraphNode* PerformAction(class UEdGraph* ParentGraph, TArray<UEdGraphPin*>& FromPins, const FVector2D Location, bool bSelectNewNode = true) override;
 	// End of FEdGraphSchemaAction interface
 
 };
@@ -114,7 +114,7 @@ struct BLUEPRINTGRAPH_API FEdGraphSchemaAction_K2AssignDelegate : public FEdGrap
 	GENERATED_USTRUCT_BODY()
 
 	static FString StaticGetTypeId() {static FString Type = TEXT("FEdGraphSchemaAction_K2AssignDelegate"); return Type;}
-	virtual FString GetTypeId() const OVERRIDE { return StaticGetTypeId(); } 
+	virtual FString GetTypeId() const override { return StaticGetTypeId(); } 
 
 	FEdGraphSchemaAction_K2AssignDelegate() 
 		:FEdGraphSchemaAction_K2NewNode()
@@ -125,7 +125,7 @@ struct BLUEPRINTGRAPH_API FEdGraphSchemaAction_K2AssignDelegate : public FEdGrap
 	{}
 
 	// FEdGraphSchemaAction interface
-	virtual UEdGraphNode* PerformAction(class UEdGraph* ParentGraph, UEdGraphPin* FromPin, const FVector2D Location, bool bSelectNewNode = true) OVERRIDE;
+	virtual UEdGraphNode* PerformAction(class UEdGraph* ParentGraph, UEdGraphPin* FromPin, const FVector2D Location, bool bSelectNewNode = true) override;
 	// End of FEdGraphSchemaAction interface
 
 	static UEdGraphNode* AssignDelegate(class UK2Node* NodeTemplate, class UEdGraph* ParentGraph, UEdGraphPin* FromPin, const FVector2D Location, bool bSelectNewNode);
@@ -142,7 +142,7 @@ struct BLUEPRINTGRAPH_API FEdGraphSchemaAction_EventFromFunction : public FEdGra
 	GENERATED_USTRUCT_BODY()
 
 	static FString StaticGetTypeId() {static FString Type = TEXT("FEdGraphSchemaAction_EventFromFunction"); return Type;}
-	virtual FString GetTypeId() const OVERRIDE { return StaticGetTypeId(); } 
+	virtual FString GetTypeId() const override { return StaticGetTypeId(); } 
 
 	UPROPERTY()
 	class UFunction* SignatureFunction;
@@ -158,9 +158,9 @@ struct BLUEPRINTGRAPH_API FEdGraphSchemaAction_EventFromFunction : public FEdGra
 	{}
 
 	// FEdGraphSchemaAction interface
-	virtual UEdGraphNode* PerformAction(class UEdGraph* ParentGraph, UEdGraphPin* FromPin, const FVector2D Location, bool bSelectNewNode = true) OVERRIDE;
-	virtual UEdGraphNode* PerformAction(class UEdGraph* ParentGraph, TArray<UEdGraphPin*>& FromPins, const FVector2D Location, bool bSelectNewNode = true) OVERRIDE;
-	virtual void AddReferencedObjects(FReferenceCollector& Collector) OVERRIDE;
+	virtual UEdGraphNode* PerformAction(class UEdGraph* ParentGraph, UEdGraphPin* FromPin, const FVector2D Location, bool bSelectNewNode = true) override;
+	virtual UEdGraphNode* PerformAction(class UEdGraph* ParentGraph, TArray<UEdGraphPin*>& FromPins, const FVector2D Location, bool bSelectNewNode = true) override;
+	virtual void AddReferencedObjects(FReferenceCollector& Collector) override;
 	// End of FEdGraphSchemaAction interface
 };
 
@@ -176,7 +176,7 @@ struct FEdGraphSchemaAction_K2AddComponent : public FEdGraphSchemaAction_K2NewNo
 
 	// Simple type info
 	static FString StaticGetTypeId() {static FString Type = TEXT("FEdGraphSchemaAction_K2AddComponent"); return Type;}
-	virtual FString GetTypeId() const OVERRIDE { return StaticGetTypeId(); } 
+	virtual FString GetTypeId() const override { return StaticGetTypeId(); } 
 
 	/** Class of component we want to add */
 	UPROPERTY()
@@ -199,8 +199,8 @@ struct FEdGraphSchemaAction_K2AddComponent : public FEdGraphSchemaAction_K2NewNo
 	{}
 
 	// FEdGraphSchemaAction interface
-	virtual UEdGraphNode* PerformAction(class UEdGraph* ParentGraph, UEdGraphPin* FromPin, const FVector2D Location, bool bSelectNewNode = true) OVERRIDE;
-	virtual void AddReferencedObjects(FReferenceCollector& Collector) OVERRIDE;
+	virtual UEdGraphNode* PerformAction(class UEdGraph* ParentGraph, UEdGraphPin* FromPin, const FVector2D Location, bool bSelectNewNode = true) override;
+	virtual void AddReferencedObjects(FReferenceCollector& Collector) override;
 	// End of FEdGraphSchemaAction interface
 };
 
@@ -216,7 +216,7 @@ struct FEdGraphSchemaAction_K2AddTimeline : public FEdGraphSchemaAction_K2NewNod
 
 	// Simple type info
 	static FString StaticGetTypeId() {static FString Type = TEXT("FEdGraphSchemaAction_K2AddTimeline"); return Type;}
-	virtual FString GetTypeId() const OVERRIDE { return StaticGetTypeId(); } 
+	virtual FString GetTypeId() const override { return StaticGetTypeId(); } 
 
 	FEdGraphSchemaAction_K2AddTimeline()
 		: FEdGraphSchemaAction_K2NewNode()
@@ -227,7 +227,7 @@ struct FEdGraphSchemaAction_K2AddTimeline : public FEdGraphSchemaAction_K2NewNod
 	{}
 
 	// FEdGraphSchemaAction interface
-	virtual UEdGraphNode* PerformAction(class UEdGraph* ParentGraph, UEdGraphPin* FromPin, const FVector2D Location, bool bSelectNewNode = true) OVERRIDE;
+	virtual UEdGraphNode* PerformAction(class UEdGraph* ParentGraph, UEdGraphPin* FromPin, const FVector2D Location, bool bSelectNewNode = true) override;
 	// End of FEdGraphSchemaAction interface
 };
 
@@ -243,7 +243,7 @@ struct FEdGraphSchemaAction_K2AddEvent : public FEdGraphSchemaAction_K2NewNode
 
 	// Simple type info
 	static FString StaticGetTypeId() {static FString Type = TEXT("FEdGraphSchemaAction_K2AddEvent"); return Type;}
-	virtual FString GetTypeId() const OVERRIDE { return StaticGetTypeId(); } 
+	virtual FString GetTypeId() const override { return StaticGetTypeId(); } 
 
 	FEdGraphSchemaAction_K2AddEvent()
 		: FEdGraphSchemaAction_K2NewNode()
@@ -254,7 +254,7 @@ struct FEdGraphSchemaAction_K2AddEvent : public FEdGraphSchemaAction_K2NewNode
 	{}
 
 	// FEdGraphSchemaAction interface
-	virtual UEdGraphNode* PerformAction(class UEdGraph* ParentGraph, UEdGraphPin* FromPin, const FVector2D Location, bool bSelectNewNode = true) OVERRIDE;
+	virtual UEdGraphNode* PerformAction(class UEdGraph* ParentGraph, UEdGraphPin* FromPin, const FVector2D Location, bool bSelectNewNode = true) override;
 	// End of FEdGraphSchemaAction interface
 
 	/** */
@@ -273,7 +273,7 @@ struct FEdGraphSchemaAction_K2AddCustomEvent : public FEdGraphSchemaAction_K2New
 
 	// Simple type info
 	static FString StaticGetTypeId() {static FString Type = TEXT("FEdGraphSchemaAction_K2AddCustomEvent"); return Type;}
-	virtual FString GetTypeId() const OVERRIDE { return StaticGetTypeId(); } 
+	virtual FString GetTypeId() const override { return StaticGetTypeId(); } 
 
 	FEdGraphSchemaAction_K2AddCustomEvent()
 		: FEdGraphSchemaAction_K2NewNode()
@@ -284,7 +284,7 @@ struct FEdGraphSchemaAction_K2AddCustomEvent : public FEdGraphSchemaAction_K2New
 	{}
 
 	// FEdGraphSchemaAction interface
-	virtual UEdGraphNode* PerformAction(class UEdGraph* ParentGraph, UEdGraphPin* FromPin, const FVector2D Location, bool bSelectNewNode = true) OVERRIDE;
+	virtual UEdGraphNode* PerformAction(class UEdGraph* ParentGraph, UEdGraphPin* FromPin, const FVector2D Location, bool bSelectNewNode = true) override;
 	// End of FEdGraphSchemaAction interface
 };
 
@@ -300,7 +300,7 @@ struct FEdGraphSchemaAction_K2AddCallOnActor : public FEdGraphSchemaAction_K2New
 
 	// Simple type info
 	static FString StaticGetTypeId() {static FString Type = TEXT("FEdGraphSchemaAction_K2AddCallOnActor"); return Type;}
-	virtual FString GetTypeId() const OVERRIDE { return StaticGetTypeId(); } 
+	virtual FString GetTypeId() const override { return StaticGetTypeId(); } 
 
 	/** Pointer to actors in level we want to call function on */
 	UPROPERTY()
@@ -315,8 +315,8 @@ struct FEdGraphSchemaAction_K2AddCallOnActor : public FEdGraphSchemaAction_K2New
 	{}
 
 	// FEdGraphSchemaAction interface
-	virtual UEdGraphNode* PerformAction(class UEdGraph* ParentGraph, UEdGraphPin* FromPin, const FVector2D Location, bool bSelectNewNode = true) OVERRIDE;
-	virtual void AddReferencedObjects(FReferenceCollector& Collector) OVERRIDE;
+	virtual UEdGraphNode* PerformAction(class UEdGraph* ParentGraph, UEdGraphPin* FromPin, const FVector2D Location, bool bSelectNewNode = true) override;
+	virtual void AddReferencedObjects(FReferenceCollector& Collector) override;
 	// End of FEdGraphSchemaAction interface	
 };
 
@@ -343,7 +343,7 @@ struct FEdGraphSchemaAction_K2AddCallOnVariable : public FEdGraphSchemaAction_K2
 	{}
 
 	// FEdGraphSchemaAction interface
-	virtual UEdGraphNode* PerformAction(class UEdGraph* ParentGraph, UEdGraphPin* FromPin, const FVector2D Location, bool bSelectNewNode = true) OVERRIDE;
+	virtual UEdGraphNode* PerformAction(class UEdGraph* ParentGraph, UEdGraphPin* FromPin, const FVector2D Location, bool bSelectNewNode = true) override;
 	// End of FEdGraphSchemaAction interface	
 };
 
@@ -359,7 +359,7 @@ struct FEdGraphSchemaAction_K2AddComment : public FEdGraphSchemaAction
 
 	// Simple type info
 	static FString StaticGetTypeId() {static FString Type = TEXT("FEdGraphSchemaAction_K2AddComment"); return Type;}
-	virtual FString GetTypeId() const OVERRIDE { return StaticGetTypeId(); } 
+	virtual FString GetTypeId() const override { return StaticGetTypeId(); } 
 
 	FEdGraphSchemaAction_K2AddComment()
 		: FEdGraphSchemaAction()
@@ -378,7 +378,7 @@ struct FEdGraphSchemaAction_K2AddComment : public FEdGraphSchemaAction
 	}
 
 	// FEdGraphSchemaAction interface
-	BLUEPRINTGRAPH_API virtual UEdGraphNode* PerformAction(class UEdGraph* ParentGraph, UEdGraphPin* FromPin, const FVector2D Location, bool bSelectNewNode = true) OVERRIDE;
+	BLUEPRINTGRAPH_API virtual UEdGraphNode* PerformAction(class UEdGraph* ParentGraph, UEdGraphPin* FromPin, const FVector2D Location, bool bSelectNewNode = true) override;
 	// End of FEdGraphSchemaAction interface
 };
 
@@ -395,7 +395,7 @@ struct BLUEPRINTGRAPH_API FEdGraphSchemaAction_K2TargetNode : public FEdGraphSch
 
 	// Simple type info
 	static FString StaticGetTypeId() {static FString Type = TEXT("FEdGraphSchemaAction_K2TargetNode"); return Type;}
-	virtual FString GetTypeId() const OVERRIDE { return StaticGetTypeId(); } 
+	virtual FString GetTypeId() const override { return StaticGetTypeId(); } 
 
 	FEdGraphSchemaAction_K2TargetNode()
 		: FEdGraphSchemaAction_K2NewNode()
@@ -406,7 +406,7 @@ struct BLUEPRINTGRAPH_API FEdGraphSchemaAction_K2TargetNode : public FEdGraphSch
 	{}
 
 	// FEdGraphSchemaAction interface
-	virtual UEdGraphNode* PerformAction(class UEdGraph* ParentGraph, UEdGraphPin* FromPin, const FVector2D Location, bool bSelectNewNode = true) OVERRIDE;
+	virtual UEdGraphNode* PerformAction(class UEdGraph* ParentGraph, UEdGraphPin* FromPin, const FVector2D Location, bool bSelectNewNode = true) override;
 	// End of FEdGraphSchemaAction interface
 };
 
@@ -422,7 +422,7 @@ struct FEdGraphSchemaAction_K2PasteHere : public FEdGraphSchemaAction_K2NewNode
 
 	// Simple type info
 	static FString StaticGetTypeId() {static FString Type = TEXT("FEdGraphSchemaAction_K2PasteHere"); return Type;}
-	virtual FString GetTypeId() const OVERRIDE { return StaticGetTypeId(); } 
+	virtual FString GetTypeId() const override { return StaticGetTypeId(); } 
 
 	FEdGraphSchemaAction_K2PasteHere ()
 		: FEdGraphSchemaAction_K2NewNode()
@@ -433,7 +433,7 @@ struct FEdGraphSchemaAction_K2PasteHere : public FEdGraphSchemaAction_K2NewNode
 	{}
 
 	// FEdGraphSchemaAction interface
-	virtual UEdGraphNode* PerformAction(class UEdGraph* ParentGraph, UEdGraphPin* FromPin, const FVector2D Location, bool bSelectNewNode = true) OVERRIDE;
+	virtual UEdGraphNode* PerformAction(class UEdGraph* ParentGraph, UEdGraphPin* FromPin, const FVector2D Location, bool bSelectNewNode = true) override;
 	// End of FEdGraphSchemaAction interface
 };
 
@@ -449,11 +449,11 @@ struct BLUEPRINTGRAPH_API FEdGraphSchemaAction_K2Enum : public FEdGraphSchemaAct
 
 	// Simple type info
 	static FString StaticGetTypeId() {static FString Type = TEXT("FEdGraphSchemaAction_K2Enum"); return Type;}
-	virtual FString GetTypeId() const OVERRIDE { return StaticGetTypeId(); } 
+	virtual FString GetTypeId() const override { return StaticGetTypeId(); } 
 
 	UEnum* Enum;
 
-	void AddReferencedObjects( FReferenceCollector& Collector ) OVERRIDE
+	void AddReferencedObjects( FReferenceCollector& Collector ) override
 	{
 		if( Enum )
 		{
@@ -503,7 +503,7 @@ public:
 
 	// Simple type info
 	static FString StaticGetTypeId() {static FString Type = TEXT("FEdGraphSchemaAction_K2Var"); return Type;}
-	virtual FString GetTypeId() const OVERRIDE { return StaticGetTypeId(); } 
+	virtual FString GetTypeId() const override { return StaticGetTypeId(); } 
 
 	FEdGraphSchemaAction_K2Var() 
 		: FEdGraphSchemaAction()
@@ -561,7 +561,7 @@ public:
 
 	// Simple type info
 	static FString StaticGetTypeId() {static FString Type = TEXT("FEdGraphSchemaAction_K2LocalVar"); return Type;}
-	virtual FString GetTypeId() const OVERRIDE { return StaticGetTypeId(); } 
+	virtual FString GetTypeId() const override { return StaticGetTypeId(); } 
 
 	FEdGraphSchemaAction_K2LocalVar() 
 		: FEdGraphSchemaAction()
@@ -618,7 +618,7 @@ struct BLUEPRINTGRAPH_API FEdGraphSchemaAction_K2Graph : public FEdGraphSchemaAc
 
 	// Simple type info
 	static FString StaticGetTypeId() {static FString Type = TEXT("FEdGraphSchemaAction_K2Graph"); return Type;}
-	virtual FString GetTypeId() const OVERRIDE { return StaticGetTypeId(); } 
+	virtual FString GetTypeId() const override { return StaticGetTypeId(); } 
 
 	/** Name of function or class */
 	FName FuncName;
@@ -638,7 +638,7 @@ struct BLUEPRINTGRAPH_API FEdGraphSchemaAction_K2Graph : public FEdGraphSchemaAc
 	{}
 
 	// FEdGraphSchemaAction interface
-	virtual bool IsParentable() const OVERRIDE { return true; }
+	virtual bool IsParentable() const override { return true; }
 	// End of FEdGraphSchemaAction interface
 };
 
@@ -687,8 +687,8 @@ struct BLUEPRINTGRAPH_API FEdGraphSchemaAction_K2Event : public FEdGraphSchemaAc
 	}
 
 	// FEdGraphSchemaAction interface
-	virtual FString GetTypeId() const OVERRIDE { return StaticGetTypeId(); }
-	virtual bool IsParentable() const OVERRIDE { return true; }
+	virtual FString GetTypeId() const override { return StaticGetTypeId(); }
+	virtual bool IsParentable() const override { return true; }
 	// End of FEdGraphSchemaAction interface
 };
 
@@ -720,7 +720,7 @@ public:
 
 	// Simple type info
 	static FString StaticGetTypeId() {static FString Type = TEXT("FEdGraphSchemaAction_K2Delegate"); return Type;}
-	virtual FString GetTypeId() const OVERRIDE { return StaticGetTypeId(); } 
+	virtual FString GetTypeId() const override { return StaticGetTypeId(); } 
 
 	/** The associated editor graph for this schema */
 	UEdGraph* EdGraph;
diff --git a/Engine/Source/Editor/BlueprintGraph/Classes/K2Node.h b/Engine/Source/Editor/BlueprintGraph/Classes/K2Node.h
--- a/Engine/Source/Editor/BlueprintGraph/Classes/K2Node.h
+++ b/Engine/Source/Editor/BlueprintGraph/Classes/K2Node.h
@@ -128,14 +128,14 @@ class UK2Node : public UEdGraphNode
 	GENERATED_UCLASS_BODY()
 
 	// UEdGraphNode interface
-	BLUEPRINTGRAPH_API virtual void ReconstructNode() OVERRIDE;
-	BLUEPRINTGRAPH_API virtual FLinearColor GetNodeTitleColor() const OVERRIDE;
-	BLUEPRINTGRAPH_API virtual void AutowireNewNode(UEdGraphPin* FromPin) OVERRIDE;
-	BLUEPRINTGRAPH_API void PinConnectionListChanged(UEdGraphPin* Pin) OVERRIDE;
-	virtual UObject* GetJumpTargetForDoubleClick() const OVERRIDE { return GetReferencedLevelActor(); }
-	BLUEPRINTGRAPH_API virtual FString GetDocumentationLink() const OVERRIDE;
-	BLUEPRINTGRAPH_API virtual void GetPinHoverText(const UEdGraphPin& Pin, FString& HoverTextOut) const OVERRIDE;
-	BLUEPRINTGRAPH_API virtual bool ShowPaletteIconOnNode() const OVERRIDE { return true; }
+	BLUEPRINTGRAPH_API virtual void ReconstructNode() override;
+	BLUEPRINTGRAPH_API virtual FLinearColor GetNodeTitleColor() const override;
+	BLUEPRINTGRAPH_API virtual void AutowireNewNode(UEdGraphPin* FromPin) override;
+	BLUEPRINTGRAPH_API void PinConnectionListChanged(UEdGraphPin* Pin) override;
+	virtual UObject* GetJumpTargetForDoubleClick() const override { return GetReferencedLevelActor(); }
+	BLUEPRINTGRAPH_API virtual FString GetDocumentationLink() const override;
+	BLUEPRINTGRAPH_API virtual void GetPinHoverText(const UEdGraphPin& Pin, FString& HoverTextOut) const override;
+	BLUEPRINTGRAPH_API virtual bool ShowPaletteIconOnNode() const override { return true; }
 	// End of UEdGraphNode interface
 
 	// K2Node interface
@@ -237,7 +237,7 @@ class UK2Node : public UEdGraphNode
 	virtual AActor* GetReferencedLevelActor() const { return NULL; }
 
 	// Can this node be created under the specified schema
-	BLUEPRINTGRAPH_API virtual bool CanCreateUnderSpecifiedSchema(const UEdGraphSchema* DesiredSchema) const OVERRIDE;
+	BLUEPRINTGRAPH_API virtual bool CanCreateUnderSpecifiedSchema(const UEdGraphSchema* DesiredSchema) const override;
 
 	/**
 	 * Searches the field redirect map for the specified named field in the scope, and returns the remapped field if found
diff --git a/Engine/Source/Editor/BlueprintGraph/Classes/K2Node_ActorBoundEvent.h b/Engine/Source/Editor/BlueprintGraph/Classes/K2Node_ActorBoundEvent.h
--- a/Engine/Source/Editor/BlueprintGraph/Classes/K2Node_ActorBoundEvent.h
+++ b/Engine/Source/Editor/BlueprintGraph/Classes/K2Node_ActorBoundEvent.h
@@ -20,22 +20,22 @@ class UK2Node_ActorBoundEvent : public UK2Node_Event
 
 
 	// Begin UEdGraphNode interface
-	virtual void ReconstructNode() OVERRIDE;
-	virtual void DestroyNode() OVERRIDE;
-	virtual FText GetNodeTitle(ENodeTitleType::Type TitleType) const OVERRIDE;
-	virtual FString GetNodeNativeTitle(ENodeTitleType::Type TitleType) const OVERRIDE;
-	virtual FString GetTooltip() const OVERRIDE;
-	virtual FString GetDocumentationLink() const OVERRIDE;
-	virtual FString GetDocumentationExcerptName() const OVERRIDE;
+	virtual void ReconstructNode() override;
+	virtual void DestroyNode() override;
+	virtual FText GetNodeTitle(ENodeTitleType::Type TitleType) const override;
+	virtual FString GetNodeNativeTitle(ENodeTitleType::Type TitleType) const override;
+	virtual FString GetTooltip() const override;
+	virtual FString GetDocumentationLink() const override;
+	virtual FString GetDocumentationExcerptName() const override;
 	// End UEdGraphNode interface
 
 	// Begin K2Node interface
-	virtual AActor* GetReferencedLevelActor() const OVERRIDE;
-	virtual bool NodeCausesStructuralBlueprintChange() const OVERRIDE { return true; }
-	virtual FNodeHandlingFunctor* CreateNodeHandler(FKismetCompilerContext& CompilerContext) const OVERRIDE;
+	virtual AActor* GetReferencedLevelActor() const override;
+	virtual bool NodeCausesStructuralBlueprintChange() const override { return true; }
+	virtual FNodeHandlingFunctor* CreateNodeHandler(FKismetCompilerContext& CompilerContext) const override;
 	// End K2Node interface
 
-	virtual bool IsUsedByAuthorityOnlyDelegate() const OVERRIDE;
+	virtual bool IsUsedByAuthorityOnlyDelegate() const override;
 
 	/**
 	 * Initialized the members of the node, given the specified owner and delegate property.  This will fill out all the required members for the event, such as CustomFunctionName
diff --git a/Engine/Source/Editor/BlueprintGraph/Classes/K2Node_AddComponent.h b/Engine/Source/Editor/BlueprintGraph/Classes/K2Node_AddComponent.h
--- a/Engine/Source/Editor/BlueprintGraph/Classes/K2Node_AddComponent.h
+++ b/Engine/Source/Editor/BlueprintGraph/Classes/K2Node_AddComponent.h
@@ -18,21 +18,21 @@ class UK2Node_AddComponent : public UK2Node_CallFunction
 	FString TemplateBlueprint;
 
 	// Begin UEdGraphNode interface
-	virtual void AllocateDefaultPins() OVERRIDE;
-	virtual void DestroyNode() OVERRIDE;
-	virtual void PrepareForCopying() OVERRIDE;
-	virtual void PostPasteNode() OVERRIDE;
-	virtual FText GetNodeTitle(ENodeTitleType::Type TitleType) const OVERRIDE;
-	virtual FString GetDocumentationLink() const OVERRIDE;
-	virtual FString GetDocumentationExcerptName() const OVERRIDE;
+	virtual void AllocateDefaultPins() override;
+	virtual void DestroyNode() override;
+	virtual void PrepareForCopying() override;
+	virtual void PostPasteNode() override;
+	virtual FText GetNodeTitle(ENodeTitleType::Type TitleType) const override;
+	virtual FString GetDocumentationLink() const override;
+	virtual FString GetDocumentationExcerptName() const override;
 	// End UEdGraphNode interface
 
 	// Begin K2Node interface
-	virtual UActorComponent* GetTemplateFromNode() const OVERRIDE;
-	virtual void ValidateNodeDuringCompilation(class FCompilerResultsLog& MessageLog) const OVERRIDE;
-	virtual void ReallocatePinsDuringReconstruction(TArray<UEdGraphPin*>& OldPins) OVERRIDE;
-	virtual class FNodeHandlingFunctor* CreateNodeHandler(class FKismetCompilerContext& CompilerContext) const OVERRIDE;
-	virtual void ExpandNode(class FKismetCompilerContext& CompilerContext, UEdGraph* SourceGraph) OVERRIDE;
+	virtual UActorComponent* GetTemplateFromNode() const override;
+	virtual void ValidateNodeDuringCompilation(class FCompilerResultsLog& MessageLog) const override;
+	virtual void ReallocatePinsDuringReconstruction(TArray<UEdGraphPin*>& OldPins) override;
+	virtual class FNodeHandlingFunctor* CreateNodeHandler(class FKismetCompilerContext& CompilerContext) const override;
+	virtual void ExpandNode(class FKismetCompilerContext& CompilerContext, UEdGraph* SourceGraph) override;
 	// End of UK2Node interface
 
 	BLUEPRINTGRAPH_API void AllocateDefaultPinsWithoutExposedVariables();
diff --git a/Engine/Source/Editor/BlueprintGraph/Classes/K2Node_AddDelegate.h b/Engine/Source/Editor/BlueprintGraph/Classes/K2Node_AddDelegate.h
--- a/Engine/Source/Editor/BlueprintGraph/Classes/K2Node_AddDelegate.h
+++ b/Engine/Source/Editor/BlueprintGraph/Classes/K2Node_AddDelegate.h
@@ -10,13 +10,13 @@ class UK2Node_AddDelegate : public UK2Node_BaseMCDelegate
 
 public:
 	// UEdGraphNode interface
-	virtual void AllocateDefaultPins() OVERRIDE;
-	virtual FText GetNodeTitle(ENodeTitleType::Type TitleType) const OVERRIDE;
-	virtual FString GetNodeNativeTitle(ENodeTitleType::Type TitleType) const OVERRIDE;
+	virtual void AllocateDefaultPins() override;
+	virtual FText GetNodeTitle(ENodeTitleType::Type TitleType) const override;
+	virtual FString GetNodeNativeTitle(ENodeTitleType::Type TitleType) const override;
 	// End of UEdGraphNode interface
 
 	// Begin of UK2Node interface
-	virtual class FNodeHandlingFunctor* CreateNodeHandler(class FKismetCompilerContext& CompilerContext) const OVERRIDE;
-	virtual void GetNodeAttributes( TArray<TKeyValuePair<FString, FString>>& OutNodeAttributes ) const OVERRIDE;
+	virtual class FNodeHandlingFunctor* CreateNodeHandler(class FKismetCompilerContext& CompilerContext) const override;
+	virtual void GetNodeAttributes( TArray<TKeyValuePair<FString, FString>>& OutNodeAttributes ) const override;
 	// End of UK2Node interface
 };
diff --git a/Engine/Source/Editor/BlueprintGraph/Classes/K2Node_AssignmentStatement.h b/Engine/Source/Editor/BlueprintGraph/Classes/K2Node_AssignmentStatement.h
--- a/Engine/Source/Editor/BlueprintGraph/Classes/K2Node_AssignmentStatement.h
+++ b/Engine/Source/Editor/BlueprintGraph/Classes/K2Node_AssignmentStatement.h
@@ -16,15 +16,15 @@ class UK2Node_AssignmentStatement : public UK2Node
 	static FString ValuePinName;
 
 	// Begin UEdGraphNode interface
-	virtual void AllocateDefaultPins() OVERRIDE;
-	virtual FString GetTooltip() const OVERRIDE;
-	virtual FText GetNodeTitle(ENodeTitleType::Type TitleType) const OVERRIDE;
+	virtual void AllocateDefaultPins() override;
+	virtual FString GetTooltip() const override;
+	virtual FText GetNodeTitle(ENodeTitleType::Type TitleType) const override;
 	// End UEdGraphNode interface
 
 	// Begin UK2Node interface
-	virtual void PostReconstructNode() OVERRIDE;
-	virtual void NotifyPinConnectionListChanged(UEdGraphPin* Pin) OVERRIDE;
-	virtual class FNodeHandlingFunctor* CreateNodeHandler(class FKismetCompilerContext& CompilerContext) const OVERRIDE;
+	virtual void PostReconstructNode() override;
+	virtual void NotifyPinConnectionListChanged(UEdGraphPin* Pin) override;
+	virtual class FNodeHandlingFunctor* CreateNodeHandler(class FKismetCompilerContext& CompilerContext) const override;
 	// End UK2Node interface
 
 	/** Get the Then output pin */
diff --git a/Engine/Source/Editor/BlueprintGraph/Classes/K2Node_BaseAsyncTask.h b/Engine/Source/Editor/BlueprintGraph/Classes/K2Node_BaseAsyncTask.h
--- a/Engine/Source/Editor/BlueprintGraph/Classes/K2Node_BaseAsyncTask.h
+++ b/Engine/Source/Editor/BlueprintGraph/Classes/K2Node_BaseAsyncTask.h
@@ -12,16 +12,16 @@ class BLUEPRINTGRAPH_API UK2Node_BaseAsyncTask : public UK2Node
 	GENERATED_UCLASS_BODY()
 
 	// UEdGraphNode interface
-	virtual void AllocateDefaultPins() OVERRIDE;
-	virtual void GetMenuEntries(FGraphContextMenuBuilder& ContextMenuBuilder) const OVERRIDE;
-	virtual FString GetTooltip() const OVERRIDE;
-	virtual FText GetNodeTitle(ENodeTitleType::Type TitleType) const OVERRIDE;
+	virtual void AllocateDefaultPins() override;
+	virtual void GetMenuEntries(FGraphContextMenuBuilder& ContextMenuBuilder) const override;
+	virtual FString GetTooltip() const override;
+	virtual FText GetNodeTitle(ENodeTitleType::Type TitleType) const override;
 	// End of UEdGraphNode interface
 
 	// UK2Node interface
-	virtual void ExpandNode(class FKismetCompilerContext& CompilerContext, UEdGraph* SourceGraph) OVERRIDE;
-	virtual bool HasExternalBlueprintDependencies(TArray<class UStruct*>* OptionalOutput) const OVERRIDE;
-	virtual FName GetCornerIcon() const OVERRIDE;
+	virtual void ExpandNode(class FKismetCompilerContext& CompilerContext, UEdGraph* SourceGraph) override;
+	virtual bool HasExternalBlueprintDependencies(TArray<class UStruct*>* OptionalOutput) const override;
+	virtual FName GetCornerIcon() const override;
 	// End of UK2Node interface
 
 protected:
diff --git a/Engine/Source/Editor/BlueprintGraph/Classes/K2Node_BaseMCDelegate.h b/Engine/Source/Editor/BlueprintGraph/Classes/K2Node_BaseMCDelegate.h
--- a/Engine/Source/Editor/BlueprintGraph/Classes/K2Node_BaseMCDelegate.h
+++ b/Engine/Source/Editor/BlueprintGraph/Classes/K2Node_BaseMCDelegate.h
@@ -15,20 +15,20 @@ class UK2Node_BaseMCDelegate : public UK2Node
 public:
 	
 	// UK2Node interface
-	virtual bool IsNodePure() const OVERRIDE { return false; }
-	virtual ERedirectType DoPinsMatchForReconstruction(const UEdGraphPin* NewPin, int32 NewPinIndex, const UEdGraphPin* OldPin, int32 OldPinIndex) const OVERRIDE;
-	virtual bool NodeCausesStructuralBlueprintChange() const OVERRIDE { return true; }
-	virtual FString GetDocumentationLink() const OVERRIDE;
-	virtual FString GetDocumentationExcerptName() const OVERRIDE;
-	virtual bool AllowMultipleSelfs(bool bInputAsArray) const OVERRIDE { return true; }
-	virtual void ExpandNode(class FKismetCompilerContext& CompilerContext, UEdGraph* SourceGraph) OVERRIDE;
-	virtual bool HasExternalBlueprintDependencies(TArray<class UStruct*>* OptionalOutput) const OVERRIDE;
-	virtual void GetNodeAttributes( TArray<TKeyValuePair<FString, FString>>& OutNodeAttributes ) const OVERRIDE;
+	virtual bool IsNodePure() const override { return false; }
+	virtual ERedirectType DoPinsMatchForReconstruction(const UEdGraphPin* NewPin, int32 NewPinIndex, const UEdGraphPin* OldPin, int32 OldPinIndex) const override;
+	virtual bool NodeCausesStructuralBlueprintChange() const override { return true; }
+	virtual FString GetDocumentationLink() const override;
+	virtual FString GetDocumentationExcerptName() const override;
+	virtual bool AllowMultipleSelfs(bool bInputAsArray) const override { return true; }
+	virtual void ExpandNode(class FKismetCompilerContext& CompilerContext, UEdGraph* SourceGraph) override;
+	virtual bool HasExternalBlueprintDependencies(TArray<class UStruct*>* OptionalOutput) const override;
+	virtual void GetNodeAttributes( TArray<TKeyValuePair<FString, FString>>& OutNodeAttributes ) const override;
 	// End of UK2Node interface
 
 	// UEdGraphNode interface
-	virtual void AllocateDefaultPins() OVERRIDE;
-	virtual void ValidateNodeDuringCompilation(class FCompilerResultsLog& MessageLog) const OVERRIDE;
+	virtual void AllocateDefaultPins() override;
+	virtual void ValidateNodeDuringCompilation(class FCompilerResultsLog& MessageLog) const override;
 	// End of UEdGraphNode interface
 
 	BLUEPRINTGRAPH_API void SetFromProperty(const UProperty* Property, bool bSelfContext)
diff --git a/Engine/Source/Editor/BlueprintGraph/Classes/K2Node_BreakStruct.h b/Engine/Source/Editor/BlueprintGraph/Classes/K2Node_BreakStruct.h
--- a/Engine/Source/Editor/BlueprintGraph/Classes/K2Node_BreakStruct.h
+++ b/Engine/Source/Editor/BlueprintGraph/Classes/K2Node_BreakStruct.h
@@ -12,20 +12,20 @@ class UK2Node_BreakStruct : public UK2Node_StructMemberGet
 	static bool CanBeBroken(const UScriptStruct* Struct);
 
 	// Begin UEdGraphNode interface
-	virtual void AllocateDefaultPins() OVERRIDE;
-	virtual FText GetNodeTitle(ENodeTitleType::Type TitleType) const OVERRIDE;
-	virtual FString GetNodeNativeTitle(ENodeTitleType::Type TitleType) const OVERRIDE;
-	virtual FLinearColor GetNodeTitleColor() const OVERRIDE;
-	virtual FString GetTooltip() const OVERRIDE;
-	virtual void ValidateNodeDuringCompilation(class FCompilerResultsLog& MessageLog) const OVERRIDE;
-	virtual FName GetPaletteIcon(FLinearColor& OutColor) const OVERRIDE{ return TEXT("GraphEditor.BreakStruct_16x"); }
+	virtual void AllocateDefaultPins() override;
+	virtual FText GetNodeTitle(ENodeTitleType::Type TitleType) const override;
+	virtual FString GetNodeNativeTitle(ENodeTitleType::Type TitleType) const override;
+	virtual FLinearColor GetNodeTitleColor() const override;
+	virtual FString GetTooltip() const override;
+	virtual void ValidateNodeDuringCompilation(class FCompilerResultsLog& MessageLog) const override;
+	virtual FName GetPaletteIcon(FLinearColor& OutColor) const override{ return TEXT("GraphEditor.BreakStruct_16x"); }
 	// End  UEdGraphNode interface
 
 	// Begin K2Node interface
-	virtual bool IsNodePure() const OVERRIDE { return true; }
-	virtual bool DrawNodeAsVariable() const OVERRIDE { return false; }
-	virtual ERedirectType DoPinsMatchForReconstruction(const UEdGraphPin* NewPin, int32 NewPinIndex, const UEdGraphPin* OldPin, int32 OldPinIndex) const OVERRIDE;
-	virtual class FNodeHandlingFunctor* CreateNodeHandler(class FKismetCompilerContext& CompilerContext) const OVERRIDE;
+	virtual bool IsNodePure() const override { return true; }
+	virtual bool DrawNodeAsVariable() const override { return false; }
+	virtual ERedirectType DoPinsMatchForReconstruction(const UEdGraphPin* NewPin, int32 NewPinIndex, const UEdGraphPin* OldPin, int32 OldPinIndex) const override;
+	virtual class FNodeHandlingFunctor* CreateNodeHandler(class FKismetCompilerContext& CompilerContext) const override;
 	// End K2Node interface
 };
 
diff --git a/Engine/Source/Editor/BlueprintGraph/Classes/K2Node_CallArrayFunction.h b/Engine/Source/Editor/BlueprintGraph/Classes/K2Node_CallArrayFunction.h
--- a/Engine/Source/Editor/BlueprintGraph/Classes/K2Node_CallArrayFunction.h
+++ b/Engine/Source/Editor/BlueprintGraph/Classes/K2Node_CallArrayFunction.h
@@ -22,12 +22,12 @@ class UK2Node_CallArrayFunction : public UK2Node_CallFunction
 	GENERATED_UCLASS_BODY()
 
 	// Begin UEdGraphNode interface
-	virtual void AllocateDefaultPins() OVERRIDE;
-	virtual void PostReconstructNode() OVERRIDE;
+	virtual void AllocateDefaultPins() override;
+	virtual void PostReconstructNode() override;
 	// End UEdGraphNode interface
 
 	// Begin UK2Node interface
-	virtual void NotifyPinConnectionListChanged(UEdGraphPin* Pin) OVERRIDE;
+	virtual void NotifyPinConnectionListChanged(UEdGraphPin* Pin) override;
 	// End UK2Node interface
 
 	/** Retrieves the target pin for the function */
diff --git a/Engine/Source/Editor/BlueprintGraph/Classes/K2Node_CallDelegate.h b/Engine/Source/Editor/BlueprintGraph/Classes/K2Node_CallDelegate.h
--- a/Engine/Source/Editor/BlueprintGraph/Classes/K2Node_CallDelegate.h
+++ b/Engine/Source/Editor/BlueprintGraph/Classes/K2Node_CallDelegate.h
@@ -11,16 +11,16 @@ class UK2Node_CallDelegate : public UK2Node_BaseMCDelegate
 public:
 
 	// UEdGraphNode interface
-	virtual void AllocateDefaultPins() OVERRIDE;
-	virtual FText GetNodeTitle(ENodeTitleType::Type TitleType) const OVERRIDE;
-	virtual FString GetNodeNativeTitle(ENodeTitleType::Type TitleType) const OVERRIDE;
+	virtual void AllocateDefaultPins() override;
+	virtual FText GetNodeTitle(ENodeTitleType::Type TitleType) const override;
+	virtual FString GetNodeNativeTitle(ENodeTitleType::Type TitleType) const override;
 	// End of UEdGraphNode interface
 
 	// UK2Node interface
-	virtual class FNodeHandlingFunctor* CreateNodeHandler(class FKismetCompilerContext& CompilerContext) const OVERRIDE;
-	virtual FName GetCornerIcon() const OVERRIDE;
+	virtual class FNodeHandlingFunctor* CreateNodeHandler(class FKismetCompilerContext& CompilerContext) const override;
+	virtual FName GetCornerIcon() const override;
 	// End of UK2Node interface
 
-	virtual void ValidateNodeDuringCompilation(class FCompilerResultsLog& MessageLog) const OVERRIDE;
+	virtual void ValidateNodeDuringCompilation(class FCompilerResultsLog& MessageLog) const override;
 	bool CreatePinsForFunctionInputs(const UFunction* Function);
 };
\ No newline at end of file
diff --git a/Engine/Source/Editor/BlueprintGraph/Classes/K2Node_CallFunction.h b/Engine/Source/Editor/BlueprintGraph/Classes/K2Node_CallFunction.h
--- a/Engine/Source/Editor/BlueprintGraph/Classes/K2Node_CallFunction.h
+++ b/Engine/Source/Editor/BlueprintGraph/Classes/K2Node_CallFunction.h
@@ -50,45 +50,45 @@ private:
 public:
 
 	// UObject interface
-	virtual void PostDuplicate(bool bDuplicateForPIE) OVERRIDE;
-	virtual void Serialize(FArchive& Ar) OVERRIDE;
+	virtual void PostDuplicate(bool bDuplicateForPIE) override;
+	virtual void Serialize(FArchive& Ar) override;
 	// End of UObject interface
 
 	// UEdGraphNode interface
-	virtual void AllocateDefaultPins() OVERRIDE;
-	virtual void DestroyNode() OVERRIDE;
-	virtual FLinearColor GetNodeTitleColor() const OVERRIDE;
-	virtual FString GetTooltip() const OVERRIDE;
-	virtual FText GetNodeTitle(ENodeTitleType::Type TitleType) const OVERRIDE;
-	virtual FString GetNodeNativeTitle(ENodeTitleType::Type TitleType) const OVERRIDE;
-	virtual FString GetDescriptiveCompiledName() const OVERRIDE;
-	virtual bool IsDeprecated() const OVERRIDE;
-	virtual bool ShouldWarnOnDeprecation() const OVERRIDE;
-	virtual FString GetDeprecationMessage() const OVERRIDE;
-	virtual void PostPlacedNewNode() OVERRIDE;
-	virtual FString GetDocumentationLink() const OVERRIDE;
-	virtual FString GetDocumentationExcerptName() const OVERRIDE;
-	virtual FName GetPaletteIcon(FLinearColor& OutColor) const OVERRIDE;
+	virtual void AllocateDefaultPins() override;
+	virtual void DestroyNode() override;
+	virtual FLinearColor GetNodeTitleColor() const override;
+	virtual FString GetTooltip() const override;
+	virtual FText GetNodeTitle(ENodeTitleType::Type TitleType) const override;
+	virtual FString GetNodeNativeTitle(ENodeTitleType::Type TitleType) const override;
+	virtual FString GetDescriptiveCompiledName() const override;
+	virtual bool IsDeprecated() const override;
+	virtual bool ShouldWarnOnDeprecation() const override;
+	virtual FString GetDeprecationMessage() const override;
+	virtual void PostPlacedNewNode() override;
+	virtual FString GetDocumentationLink() const override;
+	virtual FString GetDocumentationExcerptName() const override;
+	virtual FName GetPaletteIcon(FLinearColor& OutColor) const override;
 	// End of UEdGraphNode interface
 
 	// UK2Node interface
-	virtual void ReallocatePinsDuringReconstruction(TArray<UEdGraphPin*>& OldPins) OVERRIDE;
-	virtual bool IsNodePure() const OVERRIDE { return bIsPureFunc; }
-	virtual bool HasExternalBlueprintDependencies(TArray<class UStruct*>* OptionalOutput) const OVERRIDE;
-	virtual void PostReconstructNode() OVERRIDE;
-	virtual bool ShouldDrawCompact() const OVERRIDE;
-	virtual bool ShouldDrawAsBead() const OVERRIDE;
-	virtual FText GetCompactNodeTitle() const OVERRIDE;
-	virtual void PostPasteNode() OVERRIDE;
-	virtual void ValidateNodeDuringCompilation(class FCompilerResultsLog& MessageLog) const OVERRIDE;
-	virtual bool ShouldShowNodeProperties() const OVERRIDE;
-	virtual void GetRedirectPinNames(const UEdGraphPin& Pin, TArray<FString>& RedirectPinNames) const OVERRIDE;
-	virtual void NotifyPinConnectionListChanged(UEdGraphPin* Pin) OVERRIDE;
-	virtual class FNodeHandlingFunctor* CreateNodeHandler(class FKismetCompilerContext& CompilerContext) const OVERRIDE;
-	virtual void ExpandNode(class FKismetCompilerContext& CompilerContext, UEdGraph* SourceGraph) OVERRIDE;
-	virtual FName GetCornerIcon() const OVERRIDE;
-	virtual FText GetToolTipHeading() const OVERRIDE;
-	virtual void GetNodeAttributes( TArray<TKeyValuePair<FString, FString>>& OutNodeAttributes ) const OVERRIDE;
+	virtual void ReallocatePinsDuringReconstruction(TArray<UEdGraphPin*>& OldPins) override;
+	virtual bool IsNodePure() const override { return bIsPureFunc; }
+	virtual bool HasExternalBlueprintDependencies(TArray<class UStruct*>* OptionalOutput) const override;
+	virtual void PostReconstructNode() override;
+	virtual bool ShouldDrawCompact() const override;
+	virtual bool ShouldDrawAsBead() const override;
+	virtual FText GetCompactNodeTitle() const override;
+	virtual void PostPasteNode() override;
+	virtual void ValidateNodeDuringCompilation(class FCompilerResultsLog& MessageLog) const override;
+	virtual bool ShouldShowNodeProperties() const override;
+	virtual void GetRedirectPinNames(const UEdGraphPin& Pin, TArray<FString>& RedirectPinNames) const override;
+	virtual void NotifyPinConnectionListChanged(UEdGraphPin* Pin) override;
+	virtual class FNodeHandlingFunctor* CreateNodeHandler(class FKismetCompilerContext& CompilerContext) const override;
+	virtual void ExpandNode(class FKismetCompilerContext& CompilerContext, UEdGraph* SourceGraph) override;
+	virtual FName GetCornerIcon() const override;
+	virtual FText GetToolTipHeading() const override;
+	virtual void GetNodeAttributes( TArray<TKeyValuePair<FString, FString>>& OutNodeAttributes ) const override;
 	// End of UK2Node interface
 
 	/** Returns the UFunction that this class is pointing to */
@@ -103,7 +103,7 @@ public:
 	bool IsLatentFunction() const;
 
 	/** @return true if this function can be called on multiple contexts at once */
-	virtual bool AllowMultipleSelfs(bool bInputAsArray) const OVERRIDE;
+	virtual bool AllowMultipleSelfs(bool bInputAsArray) const override;
 
 	/**
 	 * Creates a self pin for the graph, taking into account the scope of the function call
diff --git a/Engine/Source/Editor/BlueprintGraph/Classes/K2Node_CallFunctionOnMember.h b/Engine/Source/Editor/BlueprintGraph/Classes/K2Node_CallFunctionOnMember.h
--- a/Engine/Source/Editor/BlueprintGraph/Classes/K2Node_CallFunctionOnMember.h
+++ b/Engine/Source/Editor/BlueprintGraph/Classes/K2Node_CallFunctionOnMember.h
@@ -16,10 +16,10 @@ class UK2Node_CallFunctionOnMember : public UK2Node_CallFunction
 	FMemberReference				MemberVariableToCallOn;
 
 	// Begin UK2Node_CallFunction interface
-	virtual UEdGraphPin* CreateSelfPin(const UFunction* Function) OVERRIDE;
-	virtual FString GetFunctionContextString() const OVERRIDE;
-	virtual class FNodeHandlingFunctor* CreateNodeHandler(class FKismetCompilerContext& CompilerContext) const OVERRIDE;
-	virtual void ExpandNode(class FKismetCompilerContext& CompilerContext, UEdGraph* SourceGraph) OVERRIDE;
+	virtual UEdGraphPin* CreateSelfPin(const UFunction* Function) override;
+	virtual FString GetFunctionContextString() const override;
+	virtual class FNodeHandlingFunctor* CreateNodeHandler(class FKismetCompilerContext& CompilerContext) const override;
+	virtual void ExpandNode(class FKismetCompilerContext& CompilerContext, UEdGraph* SourceGraph) override;
 	// End UK2Node_CallFunction interface
 
 };
diff --git a/Engine/Source/Editor/BlueprintGraph/Classes/K2Node_CallMaterialParameterCollectionFunction.h b/Engine/Source/Editor/BlueprintGraph/Classes/K2Node_CallMaterialParameterCollectionFunction.h
--- a/Engine/Source/Editor/BlueprintGraph/Classes/K2Node_CallMaterialParameterCollectionFunction.h
+++ b/Engine/Source/Editor/BlueprintGraph/Classes/K2Node_CallMaterialParameterCollectionFunction.h
@@ -10,6 +10,6 @@ class UK2Node_CallMaterialParameterCollectionFunction : public UK2Node_CallFunct
 	GENERATED_UCLASS_BODY()
 
 	// Begin EdGraphNode interface
-	virtual void PinDefaultValueChanged(UEdGraphPin* Pin) OVERRIDE;
+	virtual void PinDefaultValueChanged(UEdGraphPin* Pin) override;
 	// End EdGraphNode interface
 };
\ No newline at end of file
