UE4 Refactoring. Changed OVERRIDE and FINAL macros to keywords override and final.

diff --git a/Engine/Source/Editor/BlueprintGraph/Classes/K2Node_CallParentFunction.h b/Engine/Source/Editor/BlueprintGraph/Classes/K2Node_CallParentFunction.h
--- a/Engine/Source/Editor/BlueprintGraph/Classes/K2Node_CallParentFunction.h
+++ b/Engine/Source/Editor/BlueprintGraph/Classes/K2Node_CallParentFunction.h
@@ -11,13 +11,13 @@ class UK2Node_CallParentFunction : public UK2Node_CallFunction
 	GENERATED_UCLASS_BODY()
 
 	// Begin EdGraphNode interface
-	virtual FText GetNodeTitle(ENodeTitleType::Type TitleType) const OVERRIDE;
-	virtual FString GetNodeNativeTitle(ENodeTitleType::Type TitleType) const OVERRIDE;
-	virtual void AllocateDefaultPins() OVERRIDE;
-	virtual FLinearColor GetNodeTitleColor() const OVERRIDE;
-	virtual void PostPlacedNewNode() OVERRIDE;
+	virtual FText GetNodeTitle(ENodeTitleType::Type TitleType) const override;
+	virtual FString GetNodeNativeTitle(ENodeTitleType::Type TitleType) const override;
+	virtual void AllocateDefaultPins() override;
+	virtual FLinearColor GetNodeTitleColor() const override;
+	virtual void PostPlacedNewNode() override;
 	// End EdGraphNode interface
 
-	virtual void SetFromFunction(const UFunction* Function) OVERRIDE;
+	virtual void SetFromFunction(const UFunction* Function) override;
 };
 
diff --git a/Engine/Source/Editor/BlueprintGraph/Classes/K2Node_CastByteToEnum.h b/Engine/Source/Editor/BlueprintGraph/Classes/K2Node_CastByteToEnum.h
--- a/Engine/Source/Editor/BlueprintGraph/Classes/K2Node_CastByteToEnum.h
+++ b/Engine/Source/Editor/BlueprintGraph/Classes/K2Node_CastByteToEnum.h
@@ -19,18 +19,18 @@ class UK2Node_CastByteToEnum : public UK2Node
 	static const FString ByteInputPinName;
 
 	// Begin UEdGraphNode interface
-	virtual void AllocateDefaultPins() OVERRIDE;
-	virtual FString GetTooltip() const OVERRIDE;
-	virtual FText GetNodeTitle(ENodeTitleType::Type TitleType) const OVERRIDE;
-	virtual FName GetPaletteIcon(FLinearColor& OutColor) const OVERRIDE{ return TEXT("GraphEditor.Enum_16x"); }
+	virtual void AllocateDefaultPins() override;
+	virtual FString GetTooltip() const override;
+	virtual FText GetNodeTitle(ENodeTitleType::Type TitleType) const override;
+	virtual FName GetPaletteIcon(FLinearColor& OutColor) const override{ return TEXT("GraphEditor.Enum_16x"); }
 	// End UEdGraphNode interface
 
 	// Begin UK2Node interface
-	virtual FText GetCompactNodeTitle() const OVERRIDE;
-	virtual bool ShouldDrawCompact() const OVERRIDE { return true; }
-	virtual bool IsNodePure() const OVERRIDE { return true; }
-	FNodeHandlingFunctor* CreateNodeHandler(FKismetCompilerContext& CompilerContext) const OVERRIDE;
-	virtual void ExpandNode(class FKismetCompilerContext& CompilerContext, UEdGraph* SourceGraph) OVERRIDE;
+	virtual FText GetCompactNodeTitle() const override;
+	virtual bool ShouldDrawCompact() const override { return true; }
+	virtual bool IsNodePure() const override { return true; }
+	FNodeHandlingFunctor* CreateNodeHandler(FKismetCompilerContext& CompilerContext) const override;
+	virtual void ExpandNode(class FKismetCompilerContext& CompilerContext, UEdGraph* SourceGraph) override;
 	// End UK2Node interface
 
 	virtual FName GetFunctionName() const;
diff --git a/Engine/Source/Editor/BlueprintGraph/Classes/K2Node_ClassDynamicCast.h b/Engine/Source/Editor/BlueprintGraph/Classes/K2Node_ClassDynamicCast.h
--- a/Engine/Source/Editor/BlueprintGraph/Classes/K2Node_ClassDynamicCast.h
+++ b/Engine/Source/Editor/BlueprintGraph/Classes/K2Node_ClassDynamicCast.h
@@ -11,15 +11,15 @@ class UK2Node_ClassDynamicCast : public UK2Node_DynamicCast
 	GENERATED_UCLASS_BODY()
 
 	// Begin UEdGraphNode interface
-	virtual void AllocateDefaultPins() OVERRIDE;
-	virtual FLinearColor GetNodeTitleColor() const OVERRIDE;
+	virtual void AllocateDefaultPins() override;
+	virtual FLinearColor GetNodeTitleColor() const override;
 	// End UEdGraphNode interface
 
 	// UK2Node interface
-	virtual class FNodeHandlingFunctor* CreateNodeHandler(class FKismetCompilerContext& CompilerContext) const OVERRIDE;
+	virtual class FNodeHandlingFunctor* CreateNodeHandler(class FKismetCompilerContext& CompilerContext) const override;
 	// End of UK2Node interface
 
 	/** Get the input object to be casted pin */
-	virtual UEdGraphPin* GetCastSourcePin() const OVERRIDE;
+	virtual UEdGraphPin* GetCastSourcePin() const override;
 };
 
diff --git a/Engine/Source/Editor/BlueprintGraph/Classes/K2Node_ClearDelegate.h b/Engine/Source/Editor/BlueprintGraph/Classes/K2Node_ClearDelegate.h
--- a/Engine/Source/Editor/BlueprintGraph/Classes/K2Node_ClearDelegate.h
+++ b/Engine/Source/Editor/BlueprintGraph/Classes/K2Node_ClearDelegate.h
@@ -10,11 +10,11 @@ class UK2Node_ClearDelegate : public UK2Node_BaseMCDelegate
 
 public:
 	// UEdGraphNode interface
-	virtual FText GetNodeTitle(ENodeTitleType::Type TitleType) const OVERRIDE;
-	virtual FString GetNodeNativeTitle(ENodeTitleType::Type TitleType) const OVERRIDE;
+	virtual FText GetNodeTitle(ENodeTitleType::Type TitleType) const override;
+	virtual FString GetNodeNativeTitle(ENodeTitleType::Type TitleType) const override;
 	// End of UEdGraphNode interface
 
 	// Begin of UK2Node interface
-	virtual class FNodeHandlingFunctor* CreateNodeHandler(class FKismetCompilerContext& CompilerContext) const OVERRIDE;
+	virtual class FNodeHandlingFunctor* CreateNodeHandler(class FKismetCompilerContext& CompilerContext) const override;
 	// End of UK2Node interface
 };
diff --git a/Engine/Source/Editor/BlueprintGraph/Classes/K2Node_CommutativeAssociativeBinaryOperator.h b/Engine/Source/Editor/BlueprintGraph/Classes/K2Node_CommutativeAssociativeBinaryOperator.h
--- a/Engine/Source/Editor/BlueprintGraph/Classes/K2Node_CommutativeAssociativeBinaryOperator.h
+++ b/Engine/Source/Editor/BlueprintGraph/Classes/K2Node_CommutativeAssociativeBinaryOperator.h
@@ -35,11 +35,11 @@ public:
 	BLUEPRINTGRAPH_API void RemoveInputPin(UEdGraphPin* Pin);
 
 	// UEdGraphNode interface
-	virtual void AllocateDefaultPins() OVERRIDE;
+	virtual void AllocateDefaultPins() override;
 	// End of UEdGraphNode interface
 
 	// UK2Node interface
-	virtual void GetContextMenuActions(const FGraphNodeContextMenuBuilder& Context) const OVERRIDE;
-	virtual void ExpandNode(class FKismetCompilerContext& CompilerContext, UEdGraph* SourceGraph) OVERRIDE;
+	virtual void GetContextMenuActions(const FGraphNodeContextMenuBuilder& Context) const override;
+	virtual void ExpandNode(class FKismetCompilerContext& CompilerContext, UEdGraph* SourceGraph) override;
 	// End of UK2Node interface
 };
\ No newline at end of file
diff --git a/Engine/Source/Editor/BlueprintGraph/Classes/K2Node_ComponentBoundEvent.h b/Engine/Source/Editor/BlueprintGraph/Classes/K2Node_ComponentBoundEvent.h
--- a/Engine/Source/Editor/BlueprintGraph/Classes/K2Node_ComponentBoundEvent.h
+++ b/Engine/Source/Editor/BlueprintGraph/Classes/K2Node_ComponentBoundEvent.h
@@ -19,20 +19,20 @@ class UK2Node_ComponentBoundEvent : public UK2Node_Event
 	FName ComponentPropertyName;
 
 	// Begin UEdGraphNode interface
-	virtual FText GetNodeTitle(ENodeTitleType::Type TitleType) const OVERRIDE;
-	virtual FString GetNodeNativeTitle(ENodeTitleType::Type TitleType) const OVERRIDE;
-	virtual FString GetTooltip() const OVERRIDE;
-	virtual FString GetDocumentationLink() const OVERRIDE;
-	virtual FString GetDocumentationExcerptName() const OVERRIDE;
+	virtual FText GetNodeTitle(ENodeTitleType::Type TitleType) const override;
+	virtual FString GetNodeNativeTitle(ENodeTitleType::Type TitleType) const override;
+	virtual FString GetTooltip() const override;
+	virtual FString GetDocumentationLink() const override;
+	virtual FString GetDocumentationExcerptName() const override;
 	// End UEdGraphNode interface
 
 	// Begin K2Node interface
-	virtual bool NodeCausesStructuralBlueprintChange() const OVERRIDE { return true; }
-	virtual UClass* GetDynamicBindingClass() const OVERRIDE;
-	virtual void RegisterDynamicBinding(UDynamicBlueprintBinding* BindingObject) const OVERRIDE;
+	virtual bool NodeCausesStructuralBlueprintChange() const override { return true; }
+	virtual UClass* GetDynamicBindingClass() const override;
+	virtual void RegisterDynamicBinding(UDynamicBlueprintBinding* BindingObject) const override;
 	// End K2Node interface
 
-	virtual bool IsUsedByAuthorityOnlyDelegate() const OVERRIDE;
+	virtual bool IsUsedByAuthorityOnlyDelegate() const override;
 
 	/** Return the delegate property that this event is bound to */
 	BLUEPRINTGRAPH_API UMulticastDelegateProperty* GetTargetDelegateProperty() const;
diff --git a/Engine/Source/Editor/BlueprintGraph/Classes/K2Node_Composite.h b/Engine/Source/Editor/BlueprintGraph/Classes/K2Node_Composite.h
--- a/Engine/Source/Editor/BlueprintGraph/Classes/K2Node_Composite.h
+++ b/Engine/Source/Editor/BlueprintGraph/Classes/K2Node_Composite.h
@@ -13,24 +13,24 @@ class UK2Node_Composite : public UK2Node_Tunnel
 	class UEdGraph* BoundGraph;
 
 	// Begin UEdGraphNode interface
-	BLUEPRINTGRAPH_API virtual void AllocateDefaultPins() OVERRIDE;
-	BLUEPRINTGRAPH_API virtual void DestroyNode() OVERRIDE;
-	BLUEPRINTGRAPH_API virtual void PostPasteNode() OVERRIDE;
-	BLUEPRINTGRAPH_API virtual FString GetTooltip() const OVERRIDE;
-	BLUEPRINTGRAPH_API virtual FLinearColor GetNodeTitleColor() const OVERRIDE;
-	BLUEPRINTGRAPH_API virtual FText GetNodeTitle(ENodeTitleType::Type TitleType) const OVERRIDE;
-	BLUEPRINTGRAPH_API virtual FString GetNodeNativeTitle(ENodeTitleType::Type TitleType) const OVERRIDE;
-	BLUEPRINTGRAPH_API virtual bool CanUserDeleteNode() const OVERRIDE;
-	BLUEPRINTGRAPH_API virtual UObject* GetJumpTargetForDoubleClick() const OVERRIDE;
-	BLUEPRINTGRAPH_API virtual void PostPlacedNewNode() OVERRIDE;
-	virtual void OnRenameNode(const FString& NewName) OVERRIDE;
-	virtual TSharedPtr<class INameValidatorInterface> MakeNameValidator() const OVERRIDE;
+	BLUEPRINTGRAPH_API virtual void AllocateDefaultPins() override;
+	BLUEPRINTGRAPH_API virtual void DestroyNode() override;
+	BLUEPRINTGRAPH_API virtual void PostPasteNode() override;
+	BLUEPRINTGRAPH_API virtual FString GetTooltip() const override;
+	BLUEPRINTGRAPH_API virtual FLinearColor GetNodeTitleColor() const override;
+	BLUEPRINTGRAPH_API virtual FText GetNodeTitle(ENodeTitleType::Type TitleType) const override;
+	BLUEPRINTGRAPH_API virtual FString GetNodeNativeTitle(ENodeTitleType::Type TitleType) const override;
+	BLUEPRINTGRAPH_API virtual bool CanUserDeleteNode() const override;
+	BLUEPRINTGRAPH_API virtual UObject* GetJumpTargetForDoubleClick() const override;
+	BLUEPRINTGRAPH_API virtual void PostPlacedNewNode() override;
+	virtual void OnRenameNode(const FString& NewName) override;
+	virtual TSharedPtr<class INameValidatorInterface> MakeNameValidator() const override;
 	// End UEdGraphNode interface
 
 	// Begin UK2Node interface
-	virtual bool DrawNodeAsExit() const OVERRIDE { return false; }
-	virtual bool DrawNodeAsEntry() const OVERRIDE { return false; }
-	virtual bool NodeCausesStructuralBlueprintChange() const OVERRIDE { return true; }
+	virtual bool DrawNodeAsExit() const override { return false; }
+	virtual bool DrawNodeAsEntry() const override { return false; }
+	virtual bool NodeCausesStructuralBlueprintChange() const override { return true; }
 	// End UK2Node interface
 
 	// Get the entry/exit nodes inside this collapsed graph
diff --git a/Engine/Source/Editor/BlueprintGraph/Classes/K2Node_CreateDelegate.h b/Engine/Source/Editor/BlueprintGraph/Classes/K2Node_CreateDelegate.h
--- a/Engine/Source/Editor/BlueprintGraph/Classes/K2Node_CreateDelegate.h
+++ b/Engine/Source/Editor/BlueprintGraph/Classes/K2Node_CreateDelegate.h
@@ -16,18 +16,18 @@ class UK2Node_CreateDelegate : public UK2Node
 
 public:
 	// UEdGraphNode interface
-	virtual void AllocateDefaultPins() OVERRIDE;
-	virtual void PinConnectionListChanged(UEdGraphPin* Pin) OVERRIDE;
-	virtual void PinTypeChanged(UEdGraphPin* Pin) OVERRIDE;
-	virtual void NodeConnectionListChanged() OVERRIDE;
-	virtual FText GetNodeTitle(ENodeTitleType::Type TitleType) const OVERRIDE;
-	virtual UObject* GetJumpTargetForDoubleClick() const OVERRIDE;
+	virtual void AllocateDefaultPins() override;
+	virtual void PinConnectionListChanged(UEdGraphPin* Pin) override;
+	virtual void PinTypeChanged(UEdGraphPin* Pin) override;
+	virtual void NodeConnectionListChanged() override;
+	virtual FText GetNodeTitle(ENodeTitleType::Type TitleType) const override;
+	virtual UObject* GetJumpTargetForDoubleClick() const override;
 	// End of UEdGraphNode interface
 
 	// UK2Node interface
-	virtual bool IsNodePure() const OVERRIDE { return true; }
-	virtual void PostReconstructNode() OVERRIDE;
-	virtual class FNodeHandlingFunctor* CreateNodeHandler(class FKismetCompilerContext& CompilerContext) const OVERRIDE;
+	virtual bool IsNodePure() const override { return true; }
+	virtual void PostReconstructNode() override;
+	virtual class FNodeHandlingFunctor* CreateNodeHandler(class FKismetCompilerContext& CompilerContext) const override;
 	// End of UK2Node interface
 
 	bool IsValid(FString* OutMsg = NULL, bool bDontUseSkeletalClassForSelf = false) const;
diff --git a/Engine/Source/Editor/BlueprintGraph/Classes/K2Node_CustomEvent.h b/Engine/Source/Editor/BlueprintGraph/Classes/K2Node_CustomEvent.h
--- a/Engine/Source/Editor/BlueprintGraph/Classes/K2Node_CustomEvent.h
+++ b/Engine/Source/Editor/BlueprintGraph/Classes/K2Node_CustomEvent.h
@@ -10,29 +10,29 @@ class UK2Node_CustomEvent : public UK2Node_Event
 {
 	GENERATED_UCLASS_BODY()
 
-	virtual bool IsEditable() const OVERRIDE;
+	virtual bool IsEditable() const override;
 
 	// Begin UEdGraphNode interface
-	virtual void ReconstructNode() OVERRIDE;
-	virtual FText GetNodeTitle(ENodeTitleType::Type TitleType) const OVERRIDE;
-	virtual FString GetNodeNativeTitle(ENodeTitleType::Type TitleType) const OVERRIDE;
-	virtual void OnRenameNode(const FString& NewName) OVERRIDE;
-	virtual TSharedPtr<class INameValidatorInterface> MakeNameValidator() const OVERRIDE;
-	virtual FString GetTooltip() const OVERRIDE;
-	virtual FString GetDocumentationLink() const OVERRIDE;
-	virtual FString GetDocumentationExcerptName() const OVERRIDE;
-	virtual FName GetPaletteIcon(FLinearColor& OutColor) const OVERRIDE{ return TEXT("GraphEditor.CustomEvent_16x"); }
+	virtual void ReconstructNode() override;
+	virtual FText GetNodeTitle(ENodeTitleType::Type TitleType) const override;
+	virtual FString GetNodeNativeTitle(ENodeTitleType::Type TitleType) const override;
+	virtual void OnRenameNode(const FString& NewName) override;
+	virtual TSharedPtr<class INameValidatorInterface> MakeNameValidator() const override;
+	virtual FString GetTooltip() const override;
+	virtual FString GetDocumentationLink() const override;
+	virtual FString GetDocumentationExcerptName() const override;
+	virtual FName GetPaletteIcon(FLinearColor& OutColor) const override{ return TEXT("GraphEditor.CustomEvent_16x"); }
 	// End UEdGraphNode interface
 
 	// Begin UK2Node interface
-	BLUEPRINTGRAPH_API virtual void ValidateNodeDuringCompilation(class FCompilerResultsLog& MessageLog) const OVERRIDE;
+	BLUEPRINTGRAPH_API virtual void ValidateNodeDuringCompilation(class FCompilerResultsLog& MessageLog) const override;
 	// End UK2Node interface
 
 	// Begin UK2Node_EditablePinBase interface
-	virtual UEdGraphPin* CreatePinFromUserDefinition(const TSharedPtr<FUserPinInfo> NewPinInfo) OVERRIDE;
+	virtual UEdGraphPin* CreatePinFromUserDefinition(const TSharedPtr<FUserPinInfo> NewPinInfo) override;
 	// Begin UK2Node_EditablePinBase interface
 
-	virtual bool IsUsedByAuthorityOnlyDelegate() const OVERRIDE;
+	virtual bool IsUsedByAuthorityOnlyDelegate() const override;
 
 	// Rename this custom event to have unique name
 	BLUEPRINTGRAPH_API void RenameCustomEventCloseToName(int32 StartIndex = 1);
diff --git a/Engine/Source/Editor/BlueprintGraph/Classes/K2Node_DelegateSet.h b/Engine/Source/Editor/BlueprintGraph/Classes/K2Node_DelegateSet.h
--- a/Engine/Source/Editor/BlueprintGraph/Classes/K2Node_DelegateSet.h
+++ b/Engine/Source/Editor/BlueprintGraph/Classes/K2Node_DelegateSet.h
@@ -19,21 +19,21 @@ class UK2Node_DelegateSet : public UK2Node
 	TSubclassOf<class UObject>  DelegatePropertyClass;
 
 	// Begin UEdGraphNode interface
-	virtual void AllocateDefaultPins() OVERRIDE;
-	virtual FString GetTooltip() const OVERRIDE;
-	virtual FText GetNodeTitle(ENodeTitleType::Type TitleType) const OVERRIDE;
-	virtual FString GetNodeNativeTitle(ENodeTitleType::Type TitleType) const OVERRIDE;
+	virtual void AllocateDefaultPins() override;
+	virtual FString GetTooltip() const override;
+	virtual FText GetNodeTitle(ENodeTitleType::Type TitleType) const override;
+	virtual FString GetNodeNativeTitle(ENodeTitleType::Type TitleType) const override;
 	virtual FLinearColor GetNodeTitleColor() const { return FColor(216,88,88); }
 	// End UEdGraphNode interface
 
 
 	// Begin UK2Node interface
-	virtual bool DrawNodeAsEntry() const OVERRIDE { return true; }
-	virtual bool NodeCausesStructuralBlueprintChange() const OVERRIDE { return true; }
-	virtual void ValidateNodeDuringCompilation(class FCompilerResultsLog& MessageLog) const OVERRIDE;
-	virtual class FNodeHandlingFunctor* CreateNodeHandler(class FKismetCompilerContext& CompilerContext) const OVERRIDE;
-	virtual ERedirectType DoPinsMatchForReconstruction(const UEdGraphPin* NewPin, int32 NewPinIndex, const UEdGraphPin* OldPin, int32 OldPinIndex) const OVERRIDE;
-	virtual void ExpandNode(class FKismetCompilerContext& CompilerContext, UEdGraph* SourceGraph) OVERRIDE;
+	virtual bool DrawNodeAsEntry() const override { return true; }
+	virtual bool NodeCausesStructuralBlueprintChange() const override { return true; }
+	virtual void ValidateNodeDuringCompilation(class FCompilerResultsLog& MessageLog) const override;
+	virtual class FNodeHandlingFunctor* CreateNodeHandler(class FKismetCompilerContext& CompilerContext) const override;
+	virtual ERedirectType DoPinsMatchForReconstruction(const UEdGraphPin* NewPin, int32 NewPinIndex, const UEdGraphPin* OldPin, int32 OldPinIndex) const override;
+	virtual void ExpandNode(class FKismetCompilerContext& CompilerContext, UEdGraph* SourceGraph) override;
 	// End UK2Node interface
 
 	// @todo document
diff --git a/Engine/Source/Editor/BlueprintGraph/Classes/K2Node_DoOnceMultiInput.h b/Engine/Source/Editor/BlueprintGraph/Classes/K2Node_DoOnceMultiInput.h
--- a/Engine/Source/Editor/BlueprintGraph/Classes/K2Node_DoOnceMultiInput.h
+++ b/Engine/Source/Editor/BlueprintGraph/Classes/K2Node_DoOnceMultiInput.h
@@ -13,8 +13,8 @@ class UK2Node_DoOnceMultiInput : public UK2Node
 	UPROPERTY()
 	int32 NumAdditionalInputs;
 
-	virtual void GetMenuEntries(FGraphContextMenuBuilder& ContextMenuBuilder) const OVERRIDE;
-	virtual FText GetNodeTitle(ENodeTitleType::Type TitleType) const OVERRIDE;
+	virtual void GetMenuEntries(FGraphContextMenuBuilder& ContextMenuBuilder) const override;
+	virtual FText GetNodeTitle(ENodeTitleType::Type TitleType) const override;
 
 	/** Reference to the integer that contains */
 	UPROPERTY(transient)
@@ -45,11 +45,11 @@ public:
 	BLUEPRINTGRAPH_API void RemoveInputPin(UEdGraphPin* Pin);
 
 	// UEdGraphNode interface
-	virtual void AllocateDefaultPins() OVERRIDE;
+	virtual void AllocateDefaultPins() override;
 	// End of UEdGraphNode interface
 
 	// UK2Node interface
-	virtual void GetContextMenuActions(const FGraphNodeContextMenuBuilder& Context) const OVERRIDE;
-	virtual void ExpandNode(class FKismetCompilerContext& CompilerContext, UEdGraph* SourceGraph) OVERRIDE;
+	virtual void GetContextMenuActions(const FGraphNodeContextMenuBuilder& Context) const override;
+	virtual void ExpandNode(class FKismetCompilerContext& CompilerContext, UEdGraph* SourceGraph) override;
 	// End of UK2Node interface
 };
\ No newline at end of file
diff --git a/Engine/Source/Editor/BlueprintGraph/Classes/K2Node_DynamicCast.h b/Engine/Source/Editor/BlueprintGraph/Classes/K2Node_DynamicCast.h
--- a/Engine/Source/Editor/BlueprintGraph/Classes/K2Node_DynamicCast.h
+++ b/Engine/Source/Editor/BlueprintGraph/Classes/K2Node_DynamicCast.h
@@ -15,17 +15,17 @@ class UK2Node_DynamicCast : public UK2Node
 	TSubclassOf<class UObject>  TargetType;
 
 	// Begin UEdGraphNode interface
-	virtual void AllocateDefaultPins() OVERRIDE;
-	virtual FLinearColor GetNodeTitleColor() const OVERRIDE;
-	virtual FText GetNodeTitle(ENodeTitleType::Type TitleType) const OVERRIDE;
-	virtual FString GetNodeNativeTitle(ENodeTitleType::Type TitleType) const OVERRIDE;
-	virtual FName GetPaletteIcon(FLinearColor& OutColor) const OVERRIDE{ return TEXT("GraphEditor.Cast_16x"); }
+	virtual void AllocateDefaultPins() override;
+	virtual FLinearColor GetNodeTitleColor() const override;
+	virtual FText GetNodeTitle(ENodeTitleType::Type TitleType) const override;
+	virtual FString GetNodeNativeTitle(ENodeTitleType::Type TitleType) const override;
+	virtual FName GetPaletteIcon(FLinearColor& OutColor) const override{ return TEXT("GraphEditor.Cast_16x"); }
 	// End UEdGraphNode interface
 
 	// UK2Node interface
-	virtual ERedirectType DoPinsMatchForReconstruction(const UEdGraphPin* NewPin, int32 NewPinIndex, const UEdGraphPin* OldPin, int32 OldPinIndex) const OVERRIDE;
-	virtual class FNodeHandlingFunctor* CreateNodeHandler(class FKismetCompilerContext& CompilerContext) const OVERRIDE;
-	virtual bool HasExternalBlueprintDependencies(TArray<class UStruct*>* OptionalOutput) const OVERRIDE;
+	virtual ERedirectType DoPinsMatchForReconstruction(const UEdGraphPin* NewPin, int32 NewPinIndex, const UEdGraphPin* OldPin, int32 OldPinIndex) const override;
+	virtual class FNodeHandlingFunctor* CreateNodeHandler(class FKismetCompilerContext& CompilerContext) const override;
+	virtual bool HasExternalBlueprintDependencies(TArray<class UStruct*>* OptionalOutput) const override;
 	// End of UK2Node interface
 
 	/** Get the 'valid cast' exec pin */
diff --git a/Engine/Source/Editor/BlueprintGraph/Classes/K2Node_EditablePinBase.h b/Engine/Source/Editor/BlueprintGraph/Classes/K2Node_EditablePinBase.h
--- a/Engine/Source/Editor/BlueprintGraph/Classes/K2Node_EditablePinBase.h
+++ b/Engine/Source/Editor/BlueprintGraph/Classes/K2Node_EditablePinBase.h
@@ -88,18 +88,18 @@ class UK2Node_EditablePinBase : public UK2Node
 	BLUEPRINTGRAPH_API virtual bool IsEditable() const { return bIsEditable; }
 
 	// UObject interface
-	BLUEPRINTGRAPH_API virtual void Serialize(FArchive& Ar) OVERRIDE;
+	BLUEPRINTGRAPH_API virtual void Serialize(FArchive& Ar) override;
 	BLUEPRINTGRAPH_API static void AddReferencedObjects(UObject* InThis, FReferenceCollector& Collector);
-	BLUEPRINTGRAPH_API virtual void ExportCustomProperties(FOutputDevice& Out, uint32 Indent) OVERRIDE;
-	BLUEPRINTGRAPH_API virtual void ImportCustomProperties(const TCHAR* SourceText, FFeedbackContext* Warn) OVERRIDE;
+	BLUEPRINTGRAPH_API virtual void ExportCustomProperties(FOutputDevice& Out, uint32 Indent) override;
+	BLUEPRINTGRAPH_API virtual void ImportCustomProperties(const TCHAR* SourceText, FFeedbackContext* Warn) override;
 	// End of UObject interface
 
 	// UEdGraphNode interface
-	BLUEPRINTGRAPH_API virtual void AllocateDefaultPins() OVERRIDE;
+	BLUEPRINTGRAPH_API virtual void AllocateDefaultPins() override;
 	// End of UEdGraphNode interface
 
 	// UK2Node interface
-	BLUEPRINTGRAPH_API virtual bool ShouldShowNodeProperties() const OVERRIDE { return bIsEditable; }
+	BLUEPRINTGRAPH_API virtual bool ShouldShowNodeProperties() const override { return bIsEditable; }
 	// End of UK2Node interface
 
 	/**
diff --git a/Engine/Source/Editor/BlueprintGraph/Classes/K2Node_EnumEquality.h b/Engine/Source/Editor/BlueprintGraph/Classes/K2Node_EnumEquality.h
--- a/Engine/Source/Editor/BlueprintGraph/Classes/K2Node_EnumEquality.h
+++ b/Engine/Source/Editor/BlueprintGraph/Classes/K2Node_EnumEquality.h
@@ -11,22 +11,22 @@ class UK2Node_EnumEquality : public UK2Node
 	GENERATED_UCLASS_BODY()
 
 	// Begin UEdGraphNode interface
-	virtual void AllocateDefaultPins() OVERRIDE;
-	virtual FString GetTooltip() const OVERRIDE;
-	virtual FText GetNodeTitle(ENodeTitleType::Type TitleType) const OVERRIDE;
-	virtual FString GetKeywords() const OVERRIDE;
+	virtual void AllocateDefaultPins() override;
+	virtual FString GetTooltip() const override;
+	virtual FText GetNodeTitle(ENodeTitleType::Type TitleType) const override;
+	virtual FString GetKeywords() const override;
 	// End UEdGraphNode interface
 
 	// Begin UK2Node interface
- 	virtual void PostReconstructNode() OVERRIDE;
-	virtual bool IsNodeSafeToIgnore() const OVERRIDE { return true; }
-	virtual bool IsNodePure() const OVERRIDE { return true; }
-	virtual bool ShouldDrawCompact() const OVERRIDE { return true; }
-	virtual FText GetCompactNodeTitle() const OVERRIDE { return NSLOCTEXT("K2Node", "EqualEqual", "=="); }
-	virtual void GetMenuEntries(FGraphContextMenuBuilder& ContextMenuBuilder) const OVERRIDE;
-	virtual void NotifyPinConnectionListChanged(UEdGraphPin* Pin) OVERRIDE;
-	virtual class FNodeHandlingFunctor* CreateNodeHandler(class FKismetCompilerContext& CompilerContext) const OVERRIDE;
-	virtual void ExpandNode(class FKismetCompilerContext& CompilerContext, UEdGraph* SourceGraph) OVERRIDE;
+ 	virtual void PostReconstructNode() override;
+	virtual bool IsNodeSafeToIgnore() const override { return true; }
+	virtual bool IsNodePure() const override { return true; }
+	virtual bool ShouldDrawCompact() const override { return true; }
+	virtual FText GetCompactNodeTitle() const override { return NSLOCTEXT("K2Node", "EqualEqual", "=="); }
+	virtual void GetMenuEntries(FGraphContextMenuBuilder& ContextMenuBuilder) const override;
+	virtual void NotifyPinConnectionListChanged(UEdGraphPin* Pin) override;
+	virtual class FNodeHandlingFunctor* CreateNodeHandler(class FKismetCompilerContext& CompilerContext) const override;
+	virtual void ExpandNode(class FKismetCompilerContext& CompilerContext, UEdGraph* SourceGraph) override;
 	// End UK2Node interface
 
 	/** Get the return value pin */
diff --git a/Engine/Source/Editor/BlueprintGraph/Classes/K2Node_EnumLiteral.h b/Engine/Source/Editor/BlueprintGraph/Classes/K2Node_EnumLiteral.h
--- a/Engine/Source/Editor/BlueprintGraph/Classes/K2Node_EnumLiteral.h
+++ b/Engine/Source/Editor/BlueprintGraph/Classes/K2Node_EnumLiteral.h
@@ -16,20 +16,20 @@ class UK2Node_EnumLiteral : public UK2Node, public INodeDependingOnEnumInterface
 	static const FString& GetEnumInputPinName();
 
 	// Begin UEdGraphNode interface
-	virtual void AllocateDefaultPins() OVERRIDE;
-	virtual FString GetTooltip() const OVERRIDE;
-	virtual FText GetNodeTitle(ENodeTitleType::Type TitleType) const OVERRIDE;
-	virtual FName GetPaletteIcon(FLinearColor& OutColor) const OVERRIDE;
+	virtual void AllocateDefaultPins() override;
+	virtual FString GetTooltip() const override;
+	virtual FText GetNodeTitle(ENodeTitleType::Type TitleType) const override;
+	virtual FName GetPaletteIcon(FLinearColor& OutColor) const override;
 	// End UEdGraphNode interface
 
 	// Begin UK2Node interface
-	virtual bool IsNodePure() const OVERRIDE { return true; }
-	virtual FNodeHandlingFunctor* CreateNodeHandler(FKismetCompilerContext& CompilerContext) const OVERRIDE;
+	virtual bool IsNodePure() const override { return true; }
+	virtual FNodeHandlingFunctor* CreateNodeHandler(FKismetCompilerContext& CompilerContext) const override;
 	// End UK2Node interface
 
 	// INodeDependingOnEnumInterface
-	virtual class UEnum* GetEnum() const OVERRIDE { return Enum; }
-	virtual bool ShouldBeReconstructedAfterEnumChanged() const OVERRIDE { return true; }
+	virtual class UEnum* GetEnum() const override { return Enum; }
+	virtual bool ShouldBeReconstructedAfterEnumChanged() const override { return true; }
 	// End of INodeDependingOnEnumInterface
 };
 
diff --git a/Engine/Source/Editor/BlueprintGraph/Classes/K2Node_Event.h b/Engine/Source/Editor/BlueprintGraph/Classes/K2Node_Event.h
--- a/Engine/Source/Editor/BlueprintGraph/Classes/K2Node_Event.h
+++ b/Engine/Source/Editor/BlueprintGraph/Classes/K2Node_Event.h
@@ -36,33 +36,33 @@ class UK2Node_Event : public UK2Node_EditablePinBase
 	uint32 FunctionFlags;
 
 	// Begin UEdGraphNode interface
-	BLUEPRINTGRAPH_API virtual void AllocateDefaultPins() OVERRIDE;
-	BLUEPRINTGRAPH_API virtual FString GetTooltip() const OVERRIDE;
-	BLUEPRINTGRAPH_API virtual FString GetKeywords() const OVERRIDE;	
-	BLUEPRINTGRAPH_API virtual FText GetNodeTitle(ENodeTitleType::Type TitleType) const OVERRIDE;
-	BLUEPRINTGRAPH_API virtual FString GetNodeNativeTitle(ENodeTitleType::Type TitleType) const OVERRIDE;
-	BLUEPRINTGRAPH_API virtual FLinearColor GetNodeTitleColor() const OVERRIDE;
-	BLUEPRINTGRAPH_API virtual bool CanPasteHere(const UEdGraph* TargetGraph, const UEdGraphSchema* Schema) const OVERRIDE;
-	BLUEPRINTGRAPH_API virtual FName GetCornerIcon() const OVERRIDE;
-	BLUEPRINTGRAPH_API virtual bool IsDeprecated() const OVERRIDE;
-	BLUEPRINTGRAPH_API virtual FString GetDeprecationMessage() const OVERRIDE;
-	BLUEPRINTGRAPH_API virtual UObject* GetJumpTargetForDoubleClick() const OVERRIDE;
-	BLUEPRINTGRAPH_API virtual FName GetPaletteIcon(FLinearColor& OutColor) const OVERRIDE{ return TEXT("GraphEditor.Event_16x"); }
+	BLUEPRINTGRAPH_API virtual void AllocateDefaultPins() override;
+	BLUEPRINTGRAPH_API virtual FString GetTooltip() const override;
+	BLUEPRINTGRAPH_API virtual FString GetKeywords() const override;	
+	BLUEPRINTGRAPH_API virtual FText GetNodeTitle(ENodeTitleType::Type TitleType) const override;
+	BLUEPRINTGRAPH_API virtual FString GetNodeNativeTitle(ENodeTitleType::Type TitleType) const override;
+	BLUEPRINTGRAPH_API virtual FLinearColor GetNodeTitleColor() const override;
+	BLUEPRINTGRAPH_API virtual bool CanPasteHere(const UEdGraph* TargetGraph, const UEdGraphSchema* Schema) const override;
+	BLUEPRINTGRAPH_API virtual FName GetCornerIcon() const override;
+	BLUEPRINTGRAPH_API virtual bool IsDeprecated() const override;
+	BLUEPRINTGRAPH_API virtual FString GetDeprecationMessage() const override;
+	BLUEPRINTGRAPH_API virtual UObject* GetJumpTargetForDoubleClick() const override;
+	BLUEPRINTGRAPH_API virtual FName GetPaletteIcon(FLinearColor& OutColor) const override{ return TEXT("GraphEditor.Event_16x"); }
 	// End UEdGraphNode interface
 
 	// Begin UK2Node interface
-	virtual bool DrawNodeAsEntry() const OVERRIDE { return true; }
-	virtual bool NodeCausesStructuralBlueprintChange() const OVERRIDE { return true; }
-	BLUEPRINTGRAPH_API virtual void GetRedirectPinNames(const UEdGraphPin& Pin, TArray<FString>& RedirectPinNames) const OVERRIDE;
-	BLUEPRINTGRAPH_API virtual void ValidateNodeDuringCompilation(class FCompilerResultsLog& MessageLog) const OVERRIDE;
-	BLUEPRINTGRAPH_API virtual void PinConnectionListChanged(UEdGraphPin* Pin) OVERRIDE;
-	BLUEPRINTGRAPH_API virtual void PostReconstructNode() OVERRIDE;
-	BLUEPRINTGRAPH_API virtual FString GetDocumentationLink() const OVERRIDE;
-	BLUEPRINTGRAPH_API virtual FString GetDocumentationExcerptName() const OVERRIDE;
-	BLUEPRINTGRAPH_API virtual FNodeHandlingFunctor* CreateNodeHandler(FKismetCompilerContext& CompilerContext) const OVERRIDE;
-	BLUEPRINTGRAPH_API virtual void ExpandNode(class FKismetCompilerContext& CompilerContext, UEdGraph* SourceGraph) OVERRIDE;
+	virtual bool DrawNodeAsEntry() const override { return true; }
+	virtual bool NodeCausesStructuralBlueprintChange() const override { return true; }
+	BLUEPRINTGRAPH_API virtual void GetRedirectPinNames(const UEdGraphPin& Pin, TArray<FString>& RedirectPinNames) const override;
+	BLUEPRINTGRAPH_API virtual void ValidateNodeDuringCompilation(class FCompilerResultsLog& MessageLog) const override;
+	BLUEPRINTGRAPH_API virtual void PinConnectionListChanged(UEdGraphPin* Pin) override;
+	BLUEPRINTGRAPH_API virtual void PostReconstructNode() override;
+	BLUEPRINTGRAPH_API virtual FString GetDocumentationLink() const override;
+	BLUEPRINTGRAPH_API virtual FString GetDocumentationExcerptName() const override;
+	BLUEPRINTGRAPH_API virtual FNodeHandlingFunctor* CreateNodeHandler(FKismetCompilerContext& CompilerContext) const override;
+	BLUEPRINTGRAPH_API virtual void ExpandNode(class FKismetCompilerContext& CompilerContext, UEdGraph* SourceGraph) override;
 	BLUEPRINTGRAPH_API virtual FText GetToolTipHeading() const;
-	BLUEPRINTGRAPH_API virtual void GetNodeAttributes( TArray<TKeyValuePair<FString, FString>>& OutNodeAttributes ) const OVERRIDE;
+	BLUEPRINTGRAPH_API virtual void GetNodeAttributes( TArray<TKeyValuePair<FString, FString>>& OutNodeAttributes ) const override;
 	// End UK2Node interface
 
 	/** Checks whether the parameters for this event node are compatible with the specified function entry node */
diff --git a/Engine/Source/Editor/BlueprintGraph/Classes/K2Node_ExecutionSequence.h b/Engine/Source/Editor/BlueprintGraph/Classes/K2Node_ExecutionSequence.h
--- a/Engine/Source/Editor/BlueprintGraph/Classes/K2Node_ExecutionSequence.h
+++ b/Engine/Source/Editor/BlueprintGraph/Classes/K2Node_ExecutionSequence.h
@@ -11,16 +11,16 @@ class UK2Node_ExecutionSequence : public UK2Node
 	GENERATED_UCLASS_BODY()
 
 	// Begin UEdGraphNode interface
-	virtual void AllocateDefaultPins() OVERRIDE;
-	virtual FString GetTooltip() const OVERRIDE;
-	virtual FLinearColor GetNodeTitleColor() const OVERRIDE;
-	virtual FText GetNodeTitle(ENodeTitleType::Type TitleType) const OVERRIDE;
-	virtual FName GetPaletteIcon(FLinearColor& OutColor) const OVERRIDE{ return TEXT("GraphEditor.Sequence_16x"); }
+	virtual void AllocateDefaultPins() override;
+	virtual FString GetTooltip() const override;
+	virtual FLinearColor GetNodeTitleColor() const override;
+	virtual FText GetNodeTitle(ENodeTitleType::Type TitleType) const override;
+	virtual FName GetPaletteIcon(FLinearColor& OutColor) const override{ return TEXT("GraphEditor.Sequence_16x"); }
 	// End UEdGraphNode interface
 
 	// Begin UK2Node interface
-	void ReallocatePinsDuringReconstruction(TArray<UEdGraphPin*>& OldPins) OVERRIDE;
-	virtual class FNodeHandlingFunctor* CreateNodeHandler(class FKismetCompilerContext& CompilerContext) const OVERRIDE;
+	void ReallocatePinsDuringReconstruction(TArray<UEdGraphPin*>& OldPins) override;
+	virtual class FNodeHandlingFunctor* CreateNodeHandler(class FKismetCompilerContext& CompilerContext) const override;
 	// End UK2Node interface
 
 	// K2Node_ExecutionSequence interface
diff --git a/Engine/Source/Editor/BlueprintGraph/Classes/K2Node_ForEachElementInEnum.h b/Engine/Source/Editor/BlueprintGraph/Classes/K2Node_ForEachElementInEnum.h
--- a/Engine/Source/Editor/BlueprintGraph/Classes/K2Node_ForEachElementInEnum.h
+++ b/Engine/Source/Editor/BlueprintGraph/Classes/K2Node_ForEachElementInEnum.h
@@ -16,17 +16,17 @@ class UK2Node_ForEachElementInEnum : public UK2Node
 	static const FString EnumOuputPinName;
 
 	// Begin UEdGraphNode interface
-	virtual void AllocateDefaultPins() OVERRIDE;
-	virtual FString GetTooltip() const OVERRIDE;
-	virtual FText GetNodeTitle(ENodeTitleType::Type TitleType) const OVERRIDE;
-	virtual FString GetNodeNativeTitle(ENodeTitleType::Type TitleType) const OVERRIDE;
-	virtual FName GetPaletteIcon(FLinearColor& OutColor) const OVERRIDE{ return TEXT("GraphEditor.Macro.Loop_16x"); }
+	virtual void AllocateDefaultPins() override;
+	virtual FString GetTooltip() const override;
+	virtual FText GetNodeTitle(ENodeTitleType::Type TitleType) const override;
+	virtual FString GetNodeNativeTitle(ENodeTitleType::Type TitleType) const override;
+	virtual FName GetPaletteIcon(FLinearColor& OutColor) const override{ return TEXT("GraphEditor.Macro.Loop_16x"); }
 	// End UEdGraphNode interface
 
 	// Begin UK2Node interface
-	virtual bool IsNodePure() const OVERRIDE { return false; }
-	virtual void ExpandNode(class FKismetCompilerContext& CompilerContext, UEdGraph* SourceGraph) OVERRIDE;
-	virtual void ValidateNodeDuringCompilation(class FCompilerResultsLog& MessageLog) const OVERRIDE;
+	virtual bool IsNodePure() const override { return false; }
+	virtual void ExpandNode(class FKismetCompilerContext& CompilerContext, UEdGraph* SourceGraph) override;
+	virtual void ValidateNodeDuringCompilation(class FCompilerResultsLog& MessageLog) const override;
 	// End UK2Node interface
 };
 
diff --git a/Engine/Source/Editor/BlueprintGraph/Classes/K2Node_FormatText.h b/Engine/Source/Editor/BlueprintGraph/Classes/K2Node_FormatText.h
--- a/Engine/Source/Editor/BlueprintGraph/Classes/K2Node_FormatText.h
+++ b/Engine/Source/Editor/BlueprintGraph/Classes/K2Node_FormatText.h
@@ -9,24 +9,24 @@ class UK2Node_FormatText : public UK2Node
 	GENERATED_UCLASS_BODY()
 
 	// UObject interface
-	virtual void PostEditChangeProperty(struct FPropertyChangedEvent& PropertyChangedEvent) OVERRIDE;
+	virtual void PostEditChangeProperty(struct FPropertyChangedEvent& PropertyChangedEvent) override;
 	// End of UObject interface
 
 	// Begin UEdGraphNode interface.
-	virtual void AllocateDefaultPins() OVERRIDE;
-	virtual FText GetNodeTitle(ENodeTitleType::Type TitleType) const OVERRIDE;
-	virtual bool ShouldShowNodeProperties() const OVERRIDE { return true; }
-	virtual void PinConnectionListChanged(UEdGraphPin* Pin) OVERRIDE;
-	virtual void PinDefaultValueChanged(UEdGraphPin* Pin) OVERRIDE;
-	virtual FString GetTooltip() const OVERRIDE;
-	virtual FString GetPinDisplayName(const UEdGraphPin* Pin) const OVERRIDE;
+	virtual void AllocateDefaultPins() override;
+	virtual FText GetNodeTitle(ENodeTitleType::Type TitleType) const override;
+	virtual bool ShouldShowNodeProperties() const override { return true; }
+	virtual void PinConnectionListChanged(UEdGraphPin* Pin) override;
+	virtual void PinDefaultValueChanged(UEdGraphPin* Pin) override;
+	virtual FString GetTooltip() const override;
+	virtual FString GetPinDisplayName(const UEdGraphPin* Pin) const override;
 	// End UEdGraphNode interface.
 
 	// Begin UK2Node interface.
 	//virtual bool IsNodePure() const OVERRIDE { return true; }
-	virtual bool NodeCausesStructuralBlueprintChange() const OVERRIDE { return true; }
-	virtual void ExpandNode(class FKismetCompilerContext& CompilerContext, UEdGraph* SourceGraph) OVERRIDE;
-	virtual ERedirectType DoPinsMatchForReconstruction(const UEdGraphPin* NewPin, int32 NewPinIndex, const UEdGraphPin* OldPin, int32 OldPinIndex) const OVERRIDE;
+	virtual bool NodeCausesStructuralBlueprintChange() const override { return true; }
+	virtual void ExpandNode(class FKismetCompilerContext& CompilerContext, UEdGraph* SourceGraph) override;
+	virtual ERedirectType DoPinsMatchForReconstruction(const UEdGraphPin* NewPin, int32 NewPinIndex, const UEdGraphPin* OldPin, int32 OldPinIndex) const override;
 	// End UK2Node interface.
 
 public:
diff --git a/Engine/Source/Editor/BlueprintGraph/Classes/K2Node_FunctionEntry.h b/Engine/Source/Editor/BlueprintGraph/Classes/K2Node_FunctionEntry.h
--- a/Engine/Source/Editor/BlueprintGraph/Classes/K2Node_FunctionEntry.h
+++ b/Engine/Source/Editor/BlueprintGraph/Classes/K2Node_FunctionEntry.h
@@ -26,25 +26,25 @@ class UK2Node_FunctionEntry : public UK2Node_FunctionTerminator
 	TArray<struct FBPVariableDescription> LocalVariables;
 
 	// Begin UEdGraphNode interface
-	virtual void AllocateDefaultPins() OVERRIDE;
-	virtual FText GetNodeTitle(ENodeTitleType::Type TitleType) const OVERRIDE;
-	virtual FString GetNodeNativeTitle(ENodeTitleType::Type TitleType) const OVERRIDE;
-	virtual bool CanUserDeleteNode() const OVERRIDE { return false; }
-	virtual bool IsDeprecated() const OVERRIDE;
-	virtual FString GetDeprecationMessage() const OVERRIDE;
+	virtual void AllocateDefaultPins() override;
+	virtual FText GetNodeTitle(ENodeTitleType::Type TitleType) const override;
+	virtual FString GetNodeNativeTitle(ENodeTitleType::Type TitleType) const override;
+	virtual bool CanUserDeleteNode() const override { return false; }
+	virtual bool IsDeprecated() const override;
+	virtual FString GetDeprecationMessage() const override;
 	// End UEdGraphNode interface
 
 	// Begin UK2Node interface
-	virtual bool DrawNodeAsEntry() const OVERRIDE { return true; }
-	virtual class FNodeHandlingFunctor* CreateNodeHandler(class FKismetCompilerContext& CompilerContext) const OVERRIDE;
+	virtual bool DrawNodeAsEntry() const override { return true; }
+	virtual class FNodeHandlingFunctor* CreateNodeHandler(class FKismetCompilerContext& CompilerContext) const override;
 	// End UK2Node interface
 
 	// Begin UK2Node_EditablePinBase interface
-	virtual bool CanUseRefParams() const OVERRIDE { return true; }
+	virtual bool CanUseRefParams() const override { return true; }
 	// End UK2Node_EditablePinBase interface
 
 	// Begin K2Node_FunctionTerminator interface
-	virtual UEdGraphPin* CreatePinFromUserDefinition(const TSharedPtr<FUserPinInfo> NewPinInfo) OVERRIDE;
+	virtual UEdGraphPin* CreatePinFromUserDefinition(const TSharedPtr<FUserPinInfo> NewPinInfo) override;
 	// End K2Node_FunctionTerminator interface
 
 	// Removes an output pin from the node
diff --git a/Engine/Source/Editor/BlueprintGraph/Classes/K2Node_FunctionResult.h b/Engine/Source/Editor/BlueprintGraph/Classes/K2Node_FunctionResult.h
--- a/Engine/Source/Editor/BlueprintGraph/Classes/K2Node_FunctionResult.h
+++ b/Engine/Source/Editor/BlueprintGraph/Classes/K2Node_FunctionResult.h
@@ -11,19 +11,19 @@ class UK2Node_FunctionResult : public UK2Node_FunctionTerminator
 	GENERATED_UCLASS_BODY()
 
 	// Begin UEdGraphNode interface
-	virtual void AllocateDefaultPins() OVERRIDE;
-	virtual FText GetNodeTitle(ENodeTitleType::Type TitleType) const OVERRIDE;
-	virtual bool CanUserDeleteNode() const OVERRIDE { return false; }
+	virtual void AllocateDefaultPins() override;
+	virtual FText GetNodeTitle(ENodeTitleType::Type TitleType) const override;
+	virtual bool CanUserDeleteNode() const override { return false; }
 	// End UEdGraphNode interface
 
 	// Begin UK2Node interface
-	virtual bool DrawNodeAsExit() const OVERRIDE { return true; }
-	virtual bool ShouldShowNodeProperties() const OVERRIDE { return true; }
-	virtual class FNodeHandlingFunctor* CreateNodeHandler(class FKismetCompilerContext& CompilerContext) const OVERRIDE;
+	virtual bool DrawNodeAsExit() const override { return true; }
+	virtual bool ShouldShowNodeProperties() const override { return true; }
+	virtual class FNodeHandlingFunctor* CreateNodeHandler(class FKismetCompilerContext& CompilerContext) const override;
 	// End UK2Node interface
 
 	// Begin K2Node_FunctionTerminator interface
-	virtual UEdGraphPin* CreatePinFromUserDefinition(const TSharedPtr<FUserPinInfo> NewPinInfo) OVERRIDE;
+	virtual UEdGraphPin* CreatePinFromUserDefinition(const TSharedPtr<FUserPinInfo> NewPinInfo) override;
 	// End K2Node_FunctionTerminator interface
 };
 
diff --git a/Engine/Source/Editor/BlueprintGraph/Classes/K2Node_FunctionTerminator.h b/Engine/Source/Editor/BlueprintGraph/Classes/K2Node_FunctionTerminator.h
--- a/Engine/Source/Editor/BlueprintGraph/Classes/K2Node_FunctionTerminator.h
+++ b/Engine/Source/Editor/BlueprintGraph/Classes/K2Node_FunctionTerminator.h
@@ -23,12 +23,12 @@ class UK2Node_FunctionTerminator : public UK2Node_EditablePinBase
 
 
 	// Begin UEdGraphNode interface
-	virtual bool CanDuplicateNode() const OVERRIDE { return false; }
-	virtual FLinearColor GetNodeTitleColor() const OVERRIDE;
+	virtual bool CanDuplicateNode() const override { return false; }
+	virtual FLinearColor GetNodeTitleColor() const override;
 	// End UEdGraphNode interface
 
 	// Begin UK2Node interface
-	virtual bool NodeCausesStructuralBlueprintChange() const OVERRIDE { return true; }
+	virtual bool NodeCausesStructuralBlueprintChange() const override { return true; }
 	// End UK2Node interface
 };
 
diff --git a/Engine/Source/Editor/BlueprintGraph/Classes/K2Node_GetDataTableRow.h b/Engine/Source/Editor/BlueprintGraph/Classes/K2Node_GetDataTableRow.h
--- a/Engine/Source/Editor/BlueprintGraph/Classes/K2Node_GetDataTableRow.h
+++ b/Engine/Source/Editor/BlueprintGraph/Classes/K2Node_GetDataTableRow.h
@@ -14,13 +14,13 @@ class UK2Node_GetDataTableRow : public UK2Node
 	FDataTableRowHandle DataTableRowHandle;
 
 	// Begin UEdGraphNode interface.
-	virtual void AllocateDefaultPins() OVERRIDE;
-	virtual FText GetNodeTitle(ENodeTitleType::Type TitleType) const OVERRIDE;
-	virtual FString GetNodeNativeTitle(ENodeTitleType::Type TitleType) const OVERRIDE;
-	virtual void PinDefaultValueChanged(UEdGraphPin* Pin) OVERRIDE;
-	virtual FString GetTooltip() const OVERRIDE;
-	virtual void ExpandNode(class FKismetCompilerContext& CompilerContext, UEdGraph* SourceGraph) OVERRIDE;
-	virtual FName GetPaletteIcon(FLinearColor& OutColor) const OVERRIDE
+	virtual void AllocateDefaultPins() override;
+	virtual FText GetNodeTitle(ENodeTitleType::Type TitleType) const override;
+	virtual FString GetNodeNativeTitle(ENodeTitleType::Type TitleType) const override;
+	virtual void PinDefaultValueChanged(UEdGraphPin* Pin) override;
+	virtual FString GetTooltip() const override;
+	virtual void ExpandNode(class FKismetCompilerContext& CompilerContext, UEdGraph* SourceGraph) override;
+	virtual FName GetPaletteIcon(FLinearColor& OutColor) const override
     {
         OutColor = GetNodeTitleColor();
         return TEXT("Kismet.AllClasses.FunctionIcon");
@@ -28,8 +28,8 @@ class UK2Node_GetDataTableRow : public UK2Node
 	// End UEdGraphNode interface.
 
 	// Begin UK2Node interface
-	virtual bool IsNodeSafeToIgnore() const OVERRIDE { return true; }
-	virtual void ReallocatePinsDuringReconstruction(TArray<UEdGraphPin*>& OldPins) OVERRIDE;
+	virtual bool IsNodeSafeToIgnore() const override { return true; }
+	virtual void ReallocatePinsDuringReconstruction(TArray<UEdGraphPin*>& OldPins) override;
 	// End UK2Node interface
 
 
diff --git a/Engine/Source/Editor/BlueprintGraph/Classes/K2Node_GetEnumeratorName.h b/Engine/Source/Editor/BlueprintGraph/Classes/K2Node_GetEnumeratorName.h
--- a/Engine/Source/Editor/BlueprintGraph/Classes/K2Node_GetEnumeratorName.h
+++ b/Engine/Source/Editor/BlueprintGraph/Classes/K2Node_GetEnumeratorName.h
@@ -12,22 +12,22 @@ class UK2Node_GetEnumeratorName : public UK2Node
 	static FString EnumeratorPinName;
 
 	// Begin UEdGraphNode interface
-	virtual void AllocateDefaultPins() OVERRIDE;
-	virtual FString GetTooltip() const OVERRIDE;
-	virtual FText GetNodeTitle(ENodeTitleType::Type TitleType) const OVERRIDE;
-	virtual void ValidateNodeDuringCompilation(class FCompilerResultsLog& MessageLog) const OVERRIDE;
-	virtual FName GetPaletteIcon(FLinearColor& OutColor) const OVERRIDE{ return TEXT("GraphEditor.Enum_16x"); }
+	virtual void AllocateDefaultPins() override;
+	virtual FString GetTooltip() const override;
+	virtual FText GetNodeTitle(ENodeTitleType::Type TitleType) const override;
+	virtual void ValidateNodeDuringCompilation(class FCompilerResultsLog& MessageLog) const override;
+	virtual FName GetPaletteIcon(FLinearColor& OutColor) const override{ return TEXT("GraphEditor.Enum_16x"); }
 	// End UEdGraphNode interface
 
 	// Begin UK2Node interface
-	virtual FText GetCompactNodeTitle() const OVERRIDE;
-	virtual bool ShouldDrawCompact() const OVERRIDE { return true; }
-	virtual bool IsNodePure() const OVERRIDE { return true; }
-	virtual void ExpandNode(class FKismetCompilerContext& CompilerContext, UEdGraph* SourceGraph) OVERRIDE;
-	virtual bool IsConnectionDisallowed(const UEdGraphPin* MyPin, const UEdGraphPin* OtherPin, FString& OutReason) const OVERRIDE;
-	virtual void EarlyValidation(class FCompilerResultsLog& MessageLog) const OVERRIDE;
-	virtual void PostReconstructNode() OVERRIDE;
-	virtual void NotifyPinConnectionListChanged(UEdGraphPin* Pin) OVERRIDE;
+	virtual FText GetCompactNodeTitle() const override;
+	virtual bool ShouldDrawCompact() const override { return true; }
+	virtual bool IsNodePure() const override { return true; }
+	virtual void ExpandNode(class FKismetCompilerContext& CompilerContext, UEdGraph* SourceGraph) override;
+	virtual bool IsConnectionDisallowed(const UEdGraphPin* MyPin, const UEdGraphPin* OtherPin, FString& OutReason) const override;
+	virtual void EarlyValidation(class FCompilerResultsLog& MessageLog) const override;
+	virtual void PostReconstructNode() override;
+	virtual void NotifyPinConnectionListChanged(UEdGraphPin* Pin) override;
 	// End UK2Node interface
 
 	void UpdatePinType();
diff --git a/Engine/Source/Editor/BlueprintGraph/Classes/K2Node_GetEnumeratorNameAsString.h b/Engine/Source/Editor/BlueprintGraph/Classes/K2Node_GetEnumeratorNameAsString.h
--- a/Engine/Source/Editor/BlueprintGraph/Classes/K2Node_GetEnumeratorNameAsString.h
+++ b/Engine/Source/Editor/BlueprintGraph/Classes/K2Node_GetEnumeratorNameAsString.h
@@ -10,12 +10,12 @@ class UK2Node_GetEnumeratorNameAsString : public UK2Node_GetEnumeratorName
 	GENERATED_UCLASS_BODY()
 
 	// Begin UEdGraphNode interface
-	virtual void AllocateDefaultPins() OVERRIDE;
-	virtual FString GetTooltip() const OVERRIDE;
-	virtual FText GetNodeTitle(ENodeTitleType::Type TitleType) const OVERRIDE;
-	virtual FName GetPaletteIcon(FLinearColor& OutColor) const OVERRIDE{ return TEXT("GraphEditor.Enum_16x"); }
+	virtual void AllocateDefaultPins() override;
+	virtual FString GetTooltip() const override;
+	virtual FText GetNodeTitle(ENodeTitleType::Type TitleType) const override;
+	virtual FName GetPaletteIcon(FLinearColor& OutColor) const override{ return TEXT("GraphEditor.Enum_16x"); }
 	// End UEdGraphNode interface
 
-	virtual FName GetFunctionName() const OVERRIDE;
+	virtual FName GetFunctionName() const override;
 };
 
diff --git a/Engine/Source/Editor/BlueprintGraph/Classes/K2Node_GetInputAxisKeyValue.h b/Engine/Source/Editor/BlueprintGraph/Classes/K2Node_GetInputAxisKeyValue.h
--- a/Engine/Source/Editor/BlueprintGraph/Classes/K2Node_GetInputAxisKeyValue.h
+++ b/Engine/Source/Editor/BlueprintGraph/Classes/K2Node_GetInputAxisKeyValue.h
@@ -21,17 +21,17 @@ class UK2Node_GetInputAxisKeyValue : public UK2Node_CallFunction
 	uint32 bExecuteWhenPaused:1;
 
 	// Begin EdGraphNode interface
-	virtual void AllocateDefaultPins() OVERRIDE;
-	virtual FText GetNodeTitle(ENodeTitleType::Type TitleType) const OVERRIDE;
-	virtual FString GetTooltip() const OVERRIDE;
-	virtual FName GetPaletteIcon(FLinearColor& OutColor) const OVERRIDE;
+	virtual void AllocateDefaultPins() override;
+	virtual FText GetNodeTitle(ENodeTitleType::Type TitleType) const override;
+	virtual FString GetTooltip() const override;
+	virtual FName GetPaletteIcon(FLinearColor& OutColor) const override;
 	// End EdGraphNode interface
 
 	// Begin UK2Node interface
-	virtual void ValidateNodeDuringCompilation(class FCompilerResultsLog& MessageLog) const OVERRIDE;
-	virtual bool ShouldShowNodeProperties() const OVERRIDE { return true; }
-	virtual UClass* GetDynamicBindingClass() const OVERRIDE;
-	virtual void RegisterDynamicBinding(UDynamicBlueprintBinding* BindingObject) const OVERRIDE;
+	virtual void ValidateNodeDuringCompilation(class FCompilerResultsLog& MessageLog) const override;
+	virtual bool ShouldShowNodeProperties() const override { return true; }
+	virtual UClass* GetDynamicBindingClass() const override;
+	virtual void RegisterDynamicBinding(UDynamicBlueprintBinding* BindingObject) const override;
 	// End UK2Node interface
 	
 	void Initialize(const FKey AxisKey);
diff --git a/Engine/Source/Editor/BlueprintGraph/Classes/K2Node_GetInputAxisValue.h b/Engine/Source/Editor/BlueprintGraph/Classes/K2Node_GetInputAxisValue.h
--- a/Engine/Source/Editor/BlueprintGraph/Classes/K2Node_GetInputAxisValue.h
+++ b/Engine/Source/Editor/BlueprintGraph/Classes/K2Node_GetInputAxisValue.h
@@ -21,17 +21,17 @@ class UK2Node_GetInputAxisValue : public UK2Node_CallFunction
 	uint32 bExecuteWhenPaused:1;
 
 	// Begin EdGraphNode interface
-	virtual void AllocateDefaultPins() OVERRIDE;
-	virtual FText GetNodeTitle(ENodeTitleType::Type TitleType) const OVERRIDE;
-	virtual FString GetNodeNativeTitle(ENodeTitleType::Type TitleType) const OVERRIDE;
-	virtual FString GetTooltip() const OVERRIDE;
+	virtual void AllocateDefaultPins() override;
+	virtual FText GetNodeTitle(ENodeTitleType::Type TitleType) const override;
+	virtual FString GetNodeNativeTitle(ENodeTitleType::Type TitleType) const override;
+	virtual FString GetTooltip() const override;
 	// End EdGraphNode interface
 
 	// Begin UK2Node interface
-	virtual void ValidateNodeDuringCompilation(class FCompilerResultsLog& MessageLog) const OVERRIDE;
-	virtual bool ShouldShowNodeProperties() const OVERRIDE { return true; }
-	virtual UClass* GetDynamicBindingClass() const OVERRIDE;
-	virtual void RegisterDynamicBinding(UDynamicBlueprintBinding* BindingObject) const OVERRIDE;
+	virtual void ValidateNodeDuringCompilation(class FCompilerResultsLog& MessageLog) const override;
+	virtual bool ShouldShowNodeProperties() const override { return true; }
+	virtual UClass* GetDynamicBindingClass() const override;
+	virtual void RegisterDynamicBinding(UDynamicBlueprintBinding* BindingObject) const override;
 	// End UK2Node interface
 	
 	void Initialize(const FName AxisName);
diff --git a/Engine/Source/Editor/BlueprintGraph/Classes/K2Node_GetInputVectorAxisValue.h b/Engine/Source/Editor/BlueprintGraph/Classes/K2Node_GetInputVectorAxisValue.h
--- a/Engine/Source/Editor/BlueprintGraph/Classes/K2Node_GetInputVectorAxisValue.h
+++ b/Engine/Source/Editor/BlueprintGraph/Classes/K2Node_GetInputVectorAxisValue.h
@@ -10,14 +10,14 @@ class UK2Node_GetInputVectorAxisValue : public UK2Node_GetInputAxisKeyValue
 	GENERATED_UCLASS_BODY()
 
 	// Begin EdGraphNode interface
-	virtual FString GetTooltip() const OVERRIDE;
+	virtual FString GetTooltip() const override;
 	// End EdGraphNode interface
 
 	// Begin UK2Node interface
-	virtual void ValidateNodeDuringCompilation(class FCompilerResultsLog& MessageLog) const OVERRIDE;
-	virtual bool ShouldShowNodeProperties() const OVERRIDE { return true; }
-	virtual UClass* GetDynamicBindingClass() const OVERRIDE;
-	virtual void RegisterDynamicBinding(UDynamicBlueprintBinding* BindingObject) const OVERRIDE;
+	virtual void ValidateNodeDuringCompilation(class FCompilerResultsLog& MessageLog) const override;
+	virtual bool ShouldShowNodeProperties() const override { return true; }
+	virtual UClass* GetDynamicBindingClass() const override;
+	virtual void RegisterDynamicBinding(UDynamicBlueprintBinding* BindingObject) const override;
 	// End UK2Node interface
 	
 	void Initialize(const FKey AxisKey);
diff --git a/Engine/Source/Editor/BlueprintGraph/Classes/K2Node_GetNumEnumEntries.h b/Engine/Source/Editor/BlueprintGraph/Classes/K2Node_GetNumEnumEntries.h
--- a/Engine/Source/Editor/BlueprintGraph/Classes/K2Node_GetNumEnumEntries.h
+++ b/Engine/Source/Editor/BlueprintGraph/Classes/K2Node_GetNumEnumEntries.h
@@ -13,21 +13,21 @@ class UK2Node_GetNumEnumEntries : public UK2Node, public INodeDependingOnEnumInt
 	UEnum* Enum;
 
 	// Begin UEdGraphNode interface
-	virtual void AllocateDefaultPins() OVERRIDE;
-	virtual FString GetTooltip() const OVERRIDE;
-	virtual FText GetNodeTitle(ENodeTitleType::Type TitleType) const OVERRIDE;
-	virtual FString GetNodeNativeTitle(ENodeTitleType::Type TitleType) const OVERRIDE;
-	virtual FName GetPaletteIcon(FLinearColor& OutColor) const OVERRIDE{ return TEXT("GraphEditor.Enum_16x"); }
+	virtual void AllocateDefaultPins() override;
+	virtual FString GetTooltip() const override;
+	virtual FText GetNodeTitle(ENodeTitleType::Type TitleType) const override;
+	virtual FString GetNodeNativeTitle(ENodeTitleType::Type TitleType) const override;
+	virtual FName GetPaletteIcon(FLinearColor& OutColor) const override{ return TEXT("GraphEditor.Enum_16x"); }
 	// End UEdGraphNode interface
 
 	// Begin UK2Node interface
-	virtual bool IsNodePure() const OVERRIDE { return true; }
-	virtual void ExpandNode(class FKismetCompilerContext& CompilerContext, UEdGraph* SourceGraph) OVERRIDE;
+	virtual bool IsNodePure() const override { return true; }
+	virtual void ExpandNode(class FKismetCompilerContext& CompilerContext, UEdGraph* SourceGraph) override;
 	// End UK2Node interface
 
 	// INodeDependingOnEnumInterface
-	virtual class UEnum* GetEnum() const OVERRIDE { return Enum; }
-	virtual bool ShouldBeReconstructedAfterEnumChanged() const OVERRIDE {return false;}
+	virtual class UEnum* GetEnum() const override { return Enum; }
+	virtual bool ShouldBeReconstructedAfterEnumChanged() const override {return false;}
 	// End of INodeDependingOnEnumInterface
 };
 
diff --git a/Engine/Source/Editor/BlueprintGraph/Classes/K2Node_IfThenElse.h b/Engine/Source/Editor/BlueprintGraph/Classes/K2Node_IfThenElse.h
--- a/Engine/Source/Editor/BlueprintGraph/Classes/K2Node_IfThenElse.h
+++ b/Engine/Source/Editor/BlueprintGraph/Classes/K2Node_IfThenElse.h
@@ -11,16 +11,16 @@ class UK2Node_IfThenElse : public UK2Node
 	GENERATED_UCLASS_BODY()
 
 	// Begin UEdGraphNode interface
-	virtual void AllocateDefaultPins() OVERRIDE;
-	virtual FLinearColor GetNodeTitleColor() const OVERRIDE;
-	virtual FString GetTooltip() const OVERRIDE;
-	virtual FText GetNodeTitle(ENodeTitleType::Type TitleType) const OVERRIDE;
-	virtual FString GetKeywords() const OVERRIDE;
-	virtual FName GetPaletteIcon(FLinearColor& OutColor) const OVERRIDE{ return TEXT("GraphEditor.Branch_16x"); }
+	virtual void AllocateDefaultPins() override;
+	virtual FLinearColor GetNodeTitleColor() const override;
+	virtual FString GetTooltip() const override;
+	virtual FText GetNodeTitle(ENodeTitleType::Type TitleType) const override;
+	virtual FString GetKeywords() const override;
+	virtual FName GetPaletteIcon(FLinearColor& OutColor) const override{ return TEXT("GraphEditor.Branch_16x"); }
 	// End UEdGraphNode interface
 
 	// Begin K2Node interface.
-	virtual class FNodeHandlingFunctor* CreateNodeHandler(class FKismetCompilerContext& CompilerContext) const OVERRIDE;
+	virtual class FNodeHandlingFunctor* CreateNodeHandler(class FKismetCompilerContext& CompilerContext) const override;
 	// End K2Node interface.
 
 	/** Get the then output pin */
diff --git a/Engine/Source/Editor/BlueprintGraph/Classes/K2Node_InputAction.h b/Engine/Source/Editor/BlueprintGraph/Classes/K2Node_InputAction.h
--- a/Engine/Source/Editor/BlueprintGraph/Classes/K2Node_InputAction.h
+++ b/Engine/Source/Editor/BlueprintGraph/Classes/K2Node_InputAction.h
@@ -26,22 +26,22 @@ class UK2Node_InputAction : public UK2Node
 	uint32 bOverrideParentBinding:1;
 
 	// Begin UObject interface
-	virtual void PostLoad() OVERRIDE;
+	virtual void PostLoad() override;
 	// End UObject interface
 
 	// Begin UEdGraphNode interface.
-	virtual void AllocateDefaultPins() OVERRIDE;
-	virtual FLinearColor GetNodeTitleColor() const OVERRIDE;
-	virtual FText GetNodeTitle(ENodeTitleType::Type TitleType) const OVERRIDE;
-	virtual FString GetNodeNativeTitle(ENodeTitleType::Type TitleType) const OVERRIDE;
-	virtual FString GetTooltip() const OVERRIDE;
-	virtual FName GetPaletteIcon(FLinearColor& OutColor) const OVERRIDE{ return TEXT("GraphEditor.Event_16x"); }
+	virtual void AllocateDefaultPins() override;
+	virtual FLinearColor GetNodeTitleColor() const override;
+	virtual FText GetNodeTitle(ENodeTitleType::Type TitleType) const override;
+	virtual FString GetNodeNativeTitle(ENodeTitleType::Type TitleType) const override;
+	virtual FString GetTooltip() const override;
+	virtual FName GetPaletteIcon(FLinearColor& OutColor) const override{ return TEXT("GraphEditor.Event_16x"); }
 	// End UEdGraphNode interface.
 
 	// Begin UK2Node interface
-	virtual void ValidateNodeDuringCompilation(class FCompilerResultsLog& MessageLog) const OVERRIDE;
-	virtual bool ShouldShowNodeProperties() const OVERRIDE { return true; }
-	virtual void ExpandNode(FKismetCompilerContext& CompilerContext, UEdGraph* SourceGraph) OVERRIDE;
+	virtual void ValidateNodeDuringCompilation(class FCompilerResultsLog& MessageLog) const override;
+	virtual bool ShouldShowNodeProperties() const override { return true; }
+	virtual void ExpandNode(FKismetCompilerContext& CompilerContext, UEdGraph* SourceGraph) override;
 	// End UK2Node interface
 
 	/** Get the 'pressed' input pin */
diff --git a/Engine/Source/Editor/BlueprintGraph/Classes/K2Node_InputActionEvent.h b/Engine/Source/Editor/BlueprintGraph/Classes/K2Node_InputActionEvent.h
--- a/Engine/Source/Editor/BlueprintGraph/Classes/K2Node_InputActionEvent.h
+++ b/Engine/Source/Editor/BlueprintGraph/Classes/K2Node_InputActionEvent.h
@@ -26,7 +26,7 @@ class UK2Node_InputActionEvent : public UK2Node_Event
 	uint32 bOverrideParentBinding:1;
 
 	// Begin UK2Node interface
-	virtual UClass* GetDynamicBindingClass() const OVERRIDE;
-	virtual void RegisterDynamicBinding(UDynamicBlueprintBinding* BindingObject) const OVERRIDE;
+	virtual UClass* GetDynamicBindingClass() const override;
+	virtual void RegisterDynamicBinding(UDynamicBlueprintBinding* BindingObject) const override;
 	// End UK2Node interface
 };
\ No newline at end of file
diff --git a/Engine/Source/Editor/BlueprintGraph/Classes/K2Node_InputAxisEvent.h b/Engine/Source/Editor/BlueprintGraph/Classes/K2Node_InputAxisEvent.h
--- a/Engine/Source/Editor/BlueprintGraph/Classes/K2Node_InputAxisEvent.h
+++ b/Engine/Source/Editor/BlueprintGraph/Classes/K2Node_InputAxisEvent.h
@@ -25,21 +25,21 @@ class UK2Node_InputAxisEvent : public UK2Node_Event
 	uint32 bOverrideParentBinding:1;
 
 	// Begin UObject interface
-	virtual void PostLoad() OVERRIDE;
+	virtual void PostLoad() override;
 	// End UObject interface
 
 	// Begin EdGraphNode interface
-	virtual FText GetNodeTitle(ENodeTitleType::Type TitleType) const OVERRIDE;
-	virtual FString GetNodeNativeTitle(ENodeTitleType::Type TitleType) const OVERRIDE;
-	virtual FString GetTooltip() const OVERRIDE;
-	virtual bool CanPasteHere(const UEdGraph* TargetGraph, const UEdGraphSchema* Schema) const OVERRIDE;
+	virtual FText GetNodeTitle(ENodeTitleType::Type TitleType) const override;
+	virtual FString GetNodeNativeTitle(ENodeTitleType::Type TitleType) const override;
+	virtual FString GetTooltip() const override;
+	virtual bool CanPasteHere(const UEdGraph* TargetGraph, const UEdGraphSchema* Schema) const override;
 	// End EdGraphNode interface
 
 	// Begin UK2Node interface
-	virtual void ValidateNodeDuringCompilation(class FCompilerResultsLog& MessageLog) const OVERRIDE;
-	virtual bool ShouldShowNodeProperties() const OVERRIDE { return true; }
-	virtual UClass* GetDynamicBindingClass() const OVERRIDE;
-	virtual void RegisterDynamicBinding(UDynamicBlueprintBinding* BindingObject) const OVERRIDE;
+	virtual void ValidateNodeDuringCompilation(class FCompilerResultsLog& MessageLog) const override;
+	virtual bool ShouldShowNodeProperties() const override { return true; }
+	virtual UClass* GetDynamicBindingClass() const override;
+	virtual void RegisterDynamicBinding(UDynamicBlueprintBinding* BindingObject) const override;
 	// End UK2Node interface
 
 	void Initialize(const FName AxisName);
diff --git a/Engine/Source/Editor/BlueprintGraph/Classes/K2Node_InputAxisKeyEvent.h b/Engine/Source/Editor/BlueprintGraph/Classes/K2Node_InputAxisKeyEvent.h
--- a/Engine/Source/Editor/BlueprintGraph/Classes/K2Node_InputAxisKeyEvent.h
+++ b/Engine/Source/Editor/BlueprintGraph/Classes/K2Node_InputAxisKeyEvent.h
@@ -25,18 +25,18 @@ class UK2Node_InputAxisKeyEvent : public UK2Node_Event
 	uint32 bOverrideParentBinding : 1;
 
 	// Begin EdGraphNode interface
-	virtual FText GetNodeTitle(ENodeTitleType::Type TitleType) const OVERRIDE;
-	virtual FString GetNodeNativeTitle(ENodeTitleType::Type TitleType) const OVERRIDE;
-	virtual FString GetTooltip() const OVERRIDE;
-	virtual FName GetPaletteIcon(FLinearColor& OutColor) const OVERRIDE;
-	virtual bool CanPasteHere(const UEdGraph* TargetGraph, const UEdGraphSchema* Schema) const OVERRIDE;
+	virtual FText GetNodeTitle(ENodeTitleType::Type TitleType) const override;
+	virtual FString GetNodeNativeTitle(ENodeTitleType::Type TitleType) const override;
+	virtual FString GetTooltip() const override;
+	virtual FName GetPaletteIcon(FLinearColor& OutColor) const override;
+	virtual bool CanPasteHere(const UEdGraph* TargetGraph, const UEdGraphSchema* Schema) const override;
 	// End EdGraphNode interface
 
 	// Begin UK2Node interface
-	virtual void ValidateNodeDuringCompilation(class FCompilerResultsLog& MessageLog) const OVERRIDE;
-	virtual bool ShouldShowNodeProperties() const OVERRIDE{ return true; }
-	virtual UClass* GetDynamicBindingClass() const OVERRIDE;
-	virtual void RegisterDynamicBinding(UDynamicBlueprintBinding* BindingObject) const OVERRIDE;
+	virtual void ValidateNodeDuringCompilation(class FCompilerResultsLog& MessageLog) const override;
+	virtual bool ShouldShowNodeProperties() const override{ return true; }
+	virtual UClass* GetDynamicBindingClass() const override;
+	virtual void RegisterDynamicBinding(UDynamicBlueprintBinding* BindingObject) const override;
 	// End UK2Node interface
 
 	void Initialize(const FKey AxisKey);
diff --git a/Engine/Source/Editor/BlueprintGraph/Classes/K2Node_InputKey.h b/Engine/Source/Editor/BlueprintGraph/Classes/K2Node_InputKey.h
--- a/Engine/Source/Editor/BlueprintGraph/Classes/K2Node_InputKey.h
+++ b/Engine/Source/Editor/BlueprintGraph/Classes/K2Node_InputKey.h
@@ -39,22 +39,22 @@ class UK2Node_InputKey : public UK2Node
 	uint32 bShift:1;
 
 	// Begin UObject interface
-	virtual void PostLoad() OVERRIDE;
+	virtual void PostLoad() override;
 	// End UObject interface
 
 	// Begin UK2Node interface.
-	virtual bool ShouldShowNodeProperties() const OVERRIDE { return true; }
-	virtual void ValidateNodeDuringCompilation(class FCompilerResultsLog& MessageLog) const OVERRIDE;
-	virtual void ExpandNode(FKismetCompilerContext& CompilerContext, UEdGraph* SourceGraph) OVERRIDE;
+	virtual bool ShouldShowNodeProperties() const override { return true; }
+	virtual void ValidateNodeDuringCompilation(class FCompilerResultsLog& MessageLog) const override;
+	virtual void ExpandNode(FKismetCompilerContext& CompilerContext, UEdGraph* SourceGraph) override;
 	// End UK2Node interface
 
 	// Begin UEdGraphNode interface.
-	virtual void AllocateDefaultPins() OVERRIDE;
-	virtual FLinearColor GetNodeTitleColor() const OVERRIDE;
-	virtual FText GetNodeTitle(ENodeTitleType::Type TitleType) const OVERRIDE;
-	virtual FString GetNodeNativeTitle(ENodeTitleType::Type TitleType) const OVERRIDE;
-	virtual FString GetTooltip() const OVERRIDE;
-	virtual FName GetPaletteIcon(FLinearColor& OutColor) const OVERRIDE;
+	virtual void AllocateDefaultPins() override;
+	virtual FLinearColor GetNodeTitleColor() const override;
+	virtual FText GetNodeTitle(ENodeTitleType::Type TitleType) const override;
+	virtual FString GetNodeNativeTitle(ENodeTitleType::Type TitleType) const override;
+	virtual FString GetTooltip() const override;
+	virtual FName GetPaletteIcon(FLinearColor& OutColor) const override;
 	// End UEdGraphNode interface.
 
 	BLUEPRINTGRAPH_API FText GetModifierText() const;
diff --git a/Engine/Source/Editor/BlueprintGraph/Classes/K2Node_InputKeyEvent.h b/Engine/Source/Editor/BlueprintGraph/Classes/K2Node_InputKeyEvent.h
--- a/Engine/Source/Editor/BlueprintGraph/Classes/K2Node_InputKeyEvent.h
+++ b/Engine/Source/Editor/BlueprintGraph/Classes/K2Node_InputKeyEvent.h
@@ -26,7 +26,7 @@ class UK2Node_InputKeyEvent : public UK2Node_Event
 	uint32 bOverrideParentBinding:1;
 
 	// Begin UK2Node interface
-	virtual UClass* GetDynamicBindingClass() const OVERRIDE;
-	virtual void RegisterDynamicBinding(UDynamicBlueprintBinding* BindingObject) const OVERRIDE;
+	virtual UClass* GetDynamicBindingClass() const override;
+	virtual void RegisterDynamicBinding(UDynamicBlueprintBinding* BindingObject) const override;
 	// End UK2Node interface
 };
\ No newline at end of file
diff --git a/Engine/Source/Editor/BlueprintGraph/Classes/K2Node_InputTouch.h b/Engine/Source/Editor/BlueprintGraph/Classes/K2Node_InputTouch.h
--- a/Engine/Source/Editor/BlueprintGraph/Classes/K2Node_InputTouch.h
+++ b/Engine/Source/Editor/BlueprintGraph/Classes/K2Node_InputTouch.h
@@ -23,19 +23,19 @@ class UK2Node_InputTouch : public UK2Node
 	uint32 bOverrideParentBinding:1;
 
 	// Begin UObject interface
-	virtual void PostLoad() OVERRIDE;
+	virtual void PostLoad() override;
 	// End UObject interface
 
 	// Begin UK2Node interface.
-	virtual bool ShouldShowNodeProperties() const OVERRIDE { return true; }
-	virtual void ExpandNode(FKismetCompilerContext& CompilerContext, UEdGraph* SourceGraph) OVERRIDE;
+	virtual bool ShouldShowNodeProperties() const override { return true; }
+	virtual void ExpandNode(FKismetCompilerContext& CompilerContext, UEdGraph* SourceGraph) override;
 	// End UK2Node interface
 
 	// Begin UEdGraphNode interface.
-	virtual void AllocateDefaultPins() OVERRIDE;
-	virtual FLinearColor GetNodeTitleColor() const OVERRIDE;
-	virtual FText GetNodeTitle(ENodeTitleType::Type TitleType) const OVERRIDE;
-	virtual FName GetPaletteIcon(FLinearColor& OutColor) const OVERRIDE{ return TEXT("GraphEditor.TouchEvent_16x"); }
+	virtual void AllocateDefaultPins() override;
+	virtual FLinearColor GetNodeTitleColor() const override;
+	virtual FText GetNodeTitle(ENodeTitleType::Type TitleType) const override;
+	virtual FName GetPaletteIcon(FLinearColor& OutColor) const override{ return TEXT("GraphEditor.TouchEvent_16x"); }
 	// End UEdGraphNode interface.
 
 	BLUEPRINTGRAPH_API static UEnum* GetTouchIndexEnum();
diff --git a/Engine/Source/Editor/BlueprintGraph/Classes/K2Node_InputVectorAxisEvent.h b/Engine/Source/Editor/BlueprintGraph/Classes/K2Node_InputVectorAxisEvent.h
--- a/Engine/Source/Editor/BlueprintGraph/Classes/K2Node_InputVectorAxisEvent.h
+++ b/Engine/Source/Editor/BlueprintGraph/Classes/K2Node_InputVectorAxisEvent.h
@@ -10,8 +10,8 @@ class UK2Node_InputVectorAxisEvent : public UK2Node_InputAxisKeyEvent
 	GENERATED_UCLASS_BODY()
 
 	// Begin UK2Node interface
-	virtual void ValidateNodeDuringCompilation(class FCompilerResultsLog& MessageLog) const OVERRIDE;
-	virtual UClass* GetDynamicBindingClass() const OVERRIDE;
-	virtual void RegisterDynamicBinding(UDynamicBlueprintBinding* BindingObject) const OVERRIDE;
+	virtual void ValidateNodeDuringCompilation(class FCompilerResultsLog& MessageLog) const override;
+	virtual UClass* GetDynamicBindingClass() const override;
+	virtual void RegisterDynamicBinding(UDynamicBlueprintBinding* BindingObject) const override;
 	// End UK2Node interface
 };
\ No newline at end of file
diff --git a/Engine/Source/Editor/BlueprintGraph/Classes/K2Node_Knot.h b/Engine/Source/Editor/BlueprintGraph/Classes/K2Node_Knot.h
--- a/Engine/Source/Editor/BlueprintGraph/Classes/K2Node_Knot.h
+++ b/Engine/Source/Editor/BlueprintGraph/Classes/K2Node_Knot.h
@@ -11,21 +11,21 @@ class UK2Node_Knot : public UK2Node
 
 public:
 	// UEdGraphNode interface
-	virtual void AllocateDefaultPins() OVERRIDE;
-	virtual FString GetTooltip() const OVERRIDE;
-	virtual FText GetNodeTitle(ENodeTitleType::Type TitleType) const OVERRIDE;
-	virtual void GetMenuEntries(struct FGraphContextMenuBuilder& ContextMenuBuilder) const OVERRIDE;
-	virtual bool ShouldOverridePinNames() const OVERRIDE;
-	virtual FString GetPinNameOverride(const UEdGraphPin& Pin) const OVERRIDE;
-	virtual void OnRenameNode(const FString& NewName) OVERRIDE;
-	virtual TSharedPtr<class INameValidatorInterface> MakeNameValidator() const OVERRIDE;
+	virtual void AllocateDefaultPins() override;
+	virtual FString GetTooltip() const override;
+	virtual FText GetNodeTitle(ENodeTitleType::Type TitleType) const override;
+	virtual void GetMenuEntries(struct FGraphContextMenuBuilder& ContextMenuBuilder) const override;
+	virtual bool ShouldOverridePinNames() const override;
+	virtual FString GetPinNameOverride(const UEdGraphPin& Pin) const override;
+	virtual void OnRenameNode(const FString& NewName) override;
+	virtual TSharedPtr<class INameValidatorInterface> MakeNameValidator() const override;
 	// End of UEdGraphNode interface
 
 	// UK2Node interface
-	virtual bool IsNodeSafeToIgnore() const OVERRIDE;
-	virtual void ExpandNode(class FKismetCompilerContext& CompilerContext, UEdGraph* SourceGraph) OVERRIDE;
-	virtual void NotifyPinConnectionListChanged(UEdGraphPin* Pin) OVERRIDE;
-	virtual void PostReconstructNode() OVERRIDE;
+	virtual bool IsNodeSafeToIgnore() const override;
+	virtual void ExpandNode(class FKismetCompilerContext& CompilerContext, UEdGraph* SourceGraph) override;
+	virtual void NotifyPinConnectionListChanged(UEdGraphPin* Pin) override;
+	virtual void PostReconstructNode() override;
 	// End of UK2Node interface
 
 	UEdGraphPin* GetInputPin() const
diff --git a/Engine/Source/Editor/BlueprintGraph/Classes/K2Node_Literal.h b/Engine/Source/Editor/BlueprintGraph/Classes/K2Node_Literal.h
--- a/Engine/Source/Editor/BlueprintGraph/Classes/K2Node_Literal.h
+++ b/Engine/Source/Editor/BlueprintGraph/Classes/K2Node_Literal.h
@@ -18,22 +18,22 @@ private:
 public:
 
 	// Begin UEdGraphNode interface
-	virtual void AllocateDefaultPins() OVERRIDE;
-	virtual FString GetTooltip() const OVERRIDE;
-	virtual FLinearColor GetNodeTitleColor() const OVERRIDE;
-	virtual FText GetNodeTitle(ENodeTitleType::Type TitleType) const OVERRIDE;
-	virtual bool ShouldOverridePinNames() const OVERRIDE { return true; }
-	virtual FName GetPaletteIcon(FLinearColor& OutColor) const OVERRIDE;
+	virtual void AllocateDefaultPins() override;
+	virtual FString GetTooltip() const override;
+	virtual FLinearColor GetNodeTitleColor() const override;
+	virtual FText GetNodeTitle(ENodeTitleType::Type TitleType) const override;
+	virtual bool ShouldOverridePinNames() const override { return true; }
+	virtual FName GetPaletteIcon(FLinearColor& OutColor) const override;
 	// End UEdGraphNode interface
 
 	// Begin UK2Node interface
-	virtual bool IsNodePure() const OVERRIDE { return true; }
-	virtual AActor* GetReferencedLevelActor() const OVERRIDE;
-	virtual bool DrawNodeAsVariable() const OVERRIDE { return true; }
-	virtual ERedirectType DoPinsMatchForReconstruction(const UEdGraphPin* NewPin, int32 NewPinIndex, const UEdGraphPin* OldPin, int32 OldPinIndex) const OVERRIDE;
-	virtual bool NodeCausesStructuralBlueprintChange() const OVERRIDE { return true; }
-	virtual void PostReconstructNode() OVERRIDE;
-	virtual class FNodeHandlingFunctor* CreateNodeHandler(class FKismetCompilerContext& CompilerContext) const OVERRIDE;
+	virtual bool IsNodePure() const override { return true; }
+	virtual AActor* GetReferencedLevelActor() const override;
+	virtual bool DrawNodeAsVariable() const override { return true; }
+	virtual ERedirectType DoPinsMatchForReconstruction(const UEdGraphPin* NewPin, int32 NewPinIndex, const UEdGraphPin* OldPin, int32 OldPinIndex) const override;
+	virtual bool NodeCausesStructuralBlueprintChange() const override { return true; }
+	virtual void PostReconstructNode() override;
+	virtual class FNodeHandlingFunctor* CreateNodeHandler(class FKismetCompilerContext& CompilerContext) const override;
 	// End UK2Node interface
 
 	/** Accessor for the value pin of the node */
diff --git a/Engine/Source/Editor/BlueprintGraph/Classes/K2Node_LocalVariable.h b/Engine/Source/Editor/BlueprintGraph/Classes/K2Node_LocalVariable.h
--- a/Engine/Source/Editor/BlueprintGraph/Classes/K2Node_LocalVariable.h
+++ b/Engine/Source/Editor/BlueprintGraph/Classes/K2Node_LocalVariable.h
@@ -18,20 +18,20 @@ class UDEPRECATED_K2Node_LocalVariable : public UK2Node_TemporaryVariable
 	FText VariableTooltip;
 
 	// Begin UEdGraphNode interface.
- 	virtual FString GetTooltip() const OVERRIDE;
-	virtual FText GetNodeTitle(ENodeTitleType::Type TitleType) const OVERRIDE;
-	virtual FString GetNodeNativeTitle(ENodeTitleType::Type TitleType) const OVERRIDE;
-	virtual void OnRenameNode(const FString& NewName) OVERRIDE;
-	virtual TSharedPtr<class INameValidatorInterface> MakeNameValidator() const OVERRIDE;
-	virtual void PostPlacedNewNode() OVERRIDE;
-	virtual void PostPasteNode() OVERRIDE;
-	virtual bool CanPasteHere(const UEdGraph* TargetGraph, const UEdGraphSchema* Schema) const OVERRIDE;
-	virtual void ReconstructNode() OVERRIDE;
+ 	virtual FString GetTooltip() const override;
+	virtual FText GetNodeTitle(ENodeTitleType::Type TitleType) const override;
+	virtual FString GetNodeNativeTitle(ENodeTitleType::Type TitleType) const override;
+	virtual void OnRenameNode(const FString& NewName) override;
+	virtual TSharedPtr<class INameValidatorInterface> MakeNameValidator() const override;
+	virtual void PostPlacedNewNode() override;
+	virtual void PostPasteNode() override;
+	virtual bool CanPasteHere(const UEdGraph* TargetGraph, const UEdGraphSchema* Schema) const override;
+	virtual void ReconstructNode() override;
 	// End UEdGraphNode interface.
 
 	// Begin UK2Node interface.
-	virtual bool NodeCausesStructuralBlueprintChange() const OVERRIDE { return true; }
-	virtual bool ShouldShowNodeProperties() const OVERRIDE { return true; }
+	virtual bool NodeCausesStructuralBlueprintChange() const override { return true; }
+	virtual bool ShouldShowNodeProperties() const override { return true; }
 	// End UK2Node interface.
 
 	/**
diff --git a/Engine/Source/Editor/BlueprintGraph/Classes/K2Node_MacroInstance.h b/Engine/Source/Editor/BlueprintGraph/Classes/K2Node_MacroInstance.h
--- a/Engine/Source/Editor/BlueprintGraph/Classes/K2Node_MacroInstance.h
+++ b/Engine/Source/Editor/BlueprintGraph/Classes/K2Node_MacroInstance.h
@@ -29,33 +29,33 @@ public:
 	bool bReconstructNode;
 
 	// Begin UObject interface
-	virtual void Serialize(FArchive& Ar) OVERRIDE;
+	virtual void Serialize(FArchive& Ar) override;
 	// End UObject interface
 
 	// Begin UEdGraphNode interface
-	virtual void AllocateDefaultPins() OVERRIDE;
-	virtual FString GetTooltip() const OVERRIDE;
-	virtual void PostPasteNode() OVERRIDE;
-	virtual FLinearColor GetNodeTitleColor() const OVERRIDE;
-	virtual FText GetNodeTitle(ENodeTitleType::Type TitleType) const OVERRIDE;
-	virtual FString GetNodeNativeTitle(ENodeTitleType::Type TitleType) const OVERRIDE;
-	virtual bool CanUserDeleteNode() const OVERRIDE { return true; }
-	void GetContextMenuActions(const FGraphNodeContextMenuBuilder& Context) const OVERRIDE;
-	virtual void ReallocatePinsDuringReconstruction(TArray<UEdGraphPin*>& OldPins) OVERRIDE;
-	virtual void NodeConnectionListChanged() OVERRIDE;
-	virtual FString GetDocumentationLink() const OVERRIDE;
-	virtual FString GetDocumentationExcerptName() const OVERRIDE;
-	virtual FName GetPaletteIcon(FLinearColor& OutColor) const OVERRIDE;
+	virtual void AllocateDefaultPins() override;
+	virtual FString GetTooltip() const override;
+	virtual void PostPasteNode() override;
+	virtual FLinearColor GetNodeTitleColor() const override;
+	virtual FText GetNodeTitle(ENodeTitleType::Type TitleType) const override;
+	virtual FString GetNodeNativeTitle(ENodeTitleType::Type TitleType) const override;
+	virtual bool CanUserDeleteNode() const override { return true; }
+	void GetContextMenuActions(const FGraphNodeContextMenuBuilder& Context) const override;
+	virtual void ReallocatePinsDuringReconstruction(TArray<UEdGraphPin*>& OldPins) override;
+	virtual void NodeConnectionListChanged() override;
+	virtual FString GetDocumentationLink() const override;
+	virtual FString GetDocumentationExcerptName() const override;
+	virtual FName GetPaletteIcon(FLinearColor& OutColor) const override;
 	// End UEdGraphNode interface
 
 	// Begin UK2Node interface
-	virtual bool DrawNodeAsExit() const OVERRIDE { return false; }
-	virtual bool DrawNodeAsEntry() const OVERRIDE { return false; }
-	virtual void NotifyPinConnectionListChanged(UEdGraphPin* Pin) OVERRIDE;
-	virtual void PostReconstructNode() OVERRIDE;
-	virtual FText GetActiveBreakpointToolTipText() const OVERRIDE;
-	virtual bool HasExternalBlueprintDependencies(TArray<class UStruct*>* OptionalOutput) const OVERRIDE;
-	virtual void GetNodeAttributes( TArray<TKeyValuePair<FString, FString>>& OutNodeAttributes ) const OVERRIDE;
+	virtual bool DrawNodeAsExit() const override { return false; }
+	virtual bool DrawNodeAsEntry() const override { return false; }
+	virtual void NotifyPinConnectionListChanged(UEdGraphPin* Pin) override;
+	virtual void PostReconstructNode() override;
+	virtual FText GetActiveBreakpointToolTipText() const override;
+	virtual bool HasExternalBlueprintDependencies(TArray<class UStruct*>* OptionalOutput) const override;
+	virtual void GetNodeAttributes( TArray<TKeyValuePair<FString, FString>>& OutNodeAttributes ) const override;
 	// End UK2Node interface
 
 	void SetMacroGraph(UEdGraph* Graph) { MacroGraphReference.SetGraph(Graph); }
diff --git a/Engine/Source/Editor/BlueprintGraph/Classes/K2Node_MakeArray.h b/Engine/Source/Editor/BlueprintGraph/Classes/K2Node_MakeArray.h
--- a/Engine/Source/Editor/BlueprintGraph/Classes/K2Node_MakeArray.h
+++ b/Engine/Source/Editor/BlueprintGraph/Classes/K2Node_MakeArray.h
@@ -23,19 +23,19 @@ public:
 
 public:
 	// UEdGraphNode interface
-	virtual FText GetNodeTitle(ENodeTitleType::Type TitleType) const OVERRIDE;
-	virtual void AllocateDefaultPins() OVERRIDE;
-	virtual void PostReconstructNode() OVERRIDE;
-	virtual FString GetTooltip() const OVERRIDE;
-	virtual FName GetPaletteIcon(FLinearColor& OutColor) const OVERRIDE{ return TEXT("GraphEditor.MakeArray_16x"); }
+	virtual FText GetNodeTitle(ENodeTitleType::Type TitleType) const override;
+	virtual void AllocateDefaultPins() override;
+	virtual void PostReconstructNode() override;
+	virtual FString GetTooltip() const override;
+	virtual FName GetPaletteIcon(FLinearColor& OutColor) const override{ return TEXT("GraphEditor.MakeArray_16x"); }
 	// End of UEdGraphNode interface
 
 	// UK2Node interface
-	virtual bool IsNodePure() const OVERRIDE { return true; }
-	void NotifyPinConnectionListChanged(UEdGraphPin* Pin) OVERRIDE;
-	void GetContextMenuActions(const FGraphNodeContextMenuBuilder& Context) const OVERRIDE;
-	virtual class FNodeHandlingFunctor* CreateNodeHandler(class FKismetCompilerContext& CompilerContext) const OVERRIDE;
-	virtual bool IsConnectionDisallowed(const UEdGraphPin* MyPin, const UEdGraphPin* OtherPin, FString& OutReason) const OVERRIDE;
+	virtual bool IsNodePure() const override { return true; }
+	void NotifyPinConnectionListChanged(UEdGraphPin* Pin) override;
+	void GetContextMenuActions(const FGraphNodeContextMenuBuilder& Context) const override;
+	virtual class FNodeHandlingFunctor* CreateNodeHandler(class FKismetCompilerContext& CompilerContext) const override;
+	virtual bool IsConnectionDisallowed(const UEdGraphPin* MyPin, const UEdGraphPin* OtherPin, FString& OutReason) const override;
 	// End of UK2Node interface
 
 
diff --git a/Engine/Source/Editor/BlueprintGraph/Classes/K2Node_MakeStruct.h b/Engine/Source/Editor/BlueprintGraph/Classes/K2Node_MakeStruct.h
--- a/Engine/Source/Editor/BlueprintGraph/Classes/K2Node_MakeStruct.h
+++ b/Engine/Source/Editor/BlueprintGraph/Classes/K2Node_MakeStruct.h
@@ -14,20 +14,20 @@ class UK2Node_MakeStruct : public UK2Node_StructMemberSet
 	static bool CanBeExposed(const UProperty* Property);
 
 	// Begin UEdGraphNode interface
-	virtual void AllocateDefaultPins() OVERRIDE;
-	virtual FText GetNodeTitle(ENodeTitleType::Type TitleType) const OVERRIDE;
-	virtual FString GetNodeNativeTitle(ENodeTitleType::Type TitleType) const OVERRIDE;
-	virtual FLinearColor GetNodeTitleColor() const OVERRIDE;
-	virtual FString GetTooltip() const OVERRIDE;
-	virtual void ValidateNodeDuringCompilation(class FCompilerResultsLog& MessageLog) const OVERRIDE;
-	virtual FName GetPaletteIcon(FLinearColor& OutColor) const OVERRIDE{ return TEXT("GraphEditor.MakeStruct_16x"); }
+	virtual void AllocateDefaultPins() override;
+	virtual FText GetNodeTitle(ENodeTitleType::Type TitleType) const override;
+	virtual FString GetNodeNativeTitle(ENodeTitleType::Type TitleType) const override;
+	virtual FLinearColor GetNodeTitleColor() const override;
+	virtual FString GetTooltip() const override;
+	virtual void ValidateNodeDuringCompilation(class FCompilerResultsLog& MessageLog) const override;
+	virtual FName GetPaletteIcon(FLinearColor& OutColor) const override{ return TEXT("GraphEditor.MakeStruct_16x"); }
 	// End  UEdGraphNode interface
 
 	// Begin K2Node interface
-	virtual bool IsNodePure() const OVERRIDE { return true; }
-	virtual bool DrawNodeAsVariable() const OVERRIDE { return false; }
-	virtual class FNodeHandlingFunctor* CreateNodeHandler(class FKismetCompilerContext& CompilerContext) const OVERRIDE;
-	virtual ERedirectType DoPinsMatchForReconstruction(const UEdGraphPin* NewPin, int32 NewPinIndex, const UEdGraphPin* OldPin, int32 OldPinIndex)  const OVERRIDE;
+	virtual bool IsNodePure() const override { return true; }
+	virtual bool DrawNodeAsVariable() const override { return false; }
+	virtual class FNodeHandlingFunctor* CreateNodeHandler(class FKismetCompilerContext& CompilerContext) const override;
+	virtual ERedirectType DoPinsMatchForReconstruction(const UEdGraphPin* NewPin, int32 NewPinIndex, const UEdGraphPin* OldPin, int32 OldPinIndex)  const override;
 	// End K2Node interface
 
 protected:
@@ -37,7 +37,7 @@ protected:
 	public:
 		FMakeStructPinManager(const uint8* InSampleStructMemory);
 	protected:
-		virtual void CustomizePinData(UEdGraphPin* Pin, FName SourcePropertyName, int32 ArrayIndex, UProperty* Property) const OVERRIDE;
-		virtual bool CanTreatPropertyAsOptional(UProperty* TestProperty) const OVERRIDE;
+		virtual void CustomizePinData(UEdGraphPin* Pin, FName SourcePropertyName, int32 ArrayIndex, UProperty* Property) const override;
+		virtual bool CanTreatPropertyAsOptional(UProperty* TestProperty) const override;
 	};
 };
diff --git a/Engine/Source/Editor/BlueprintGraph/Classes/K2Node_MathExpression.h b/Engine/Source/Editor/BlueprintGraph/Classes/K2Node_MathExpression.h
--- a/Engine/Source/Editor/BlueprintGraph/Classes/K2Node_MathExpression.h
+++ b/Engine/Source/Editor/BlueprintGraph/Classes/K2Node_MathExpression.h
@@ -24,21 +24,21 @@ public:
 
 public:
 	// UObject interface
-	virtual void PostEditChangeProperty(struct FPropertyChangedEvent& PropertyChangedEvent) OVERRIDE;
+	virtual void PostEditChangeProperty(struct FPropertyChangedEvent& PropertyChangedEvent) override;
 	// End of UObject interface
 
 	// UEdGraphNode interface
-	virtual TSharedPtr<class INameValidatorInterface> MakeNameValidator() const OVERRIDE;
-	virtual void OnRenameNode(const FString& NewName) OVERRIDE;
-	virtual FText GetNodeTitle(ENodeTitleType::Type TitleType) const OVERRIDE;
-	virtual FString GetNodeNativeTitle(ENodeTitleType::Type TitleType) const OVERRIDE;
-	virtual void PostPlacedNewNode() OVERRIDE;
-	virtual void ReconstructNode() OVERRIDE;
+	virtual TSharedPtr<class INameValidatorInterface> MakeNameValidator() const override;
+	virtual void OnRenameNode(const FString& NewName) override;
+	virtual FText GetNodeTitle(ENodeTitleType::Type TitleType) const override;
+	virtual FString GetNodeNativeTitle(ENodeTitleType::Type TitleType) const override;
+	virtual void PostPlacedNewNode() override;
+	virtual void ReconstructNode() override;
 	// End of UEdGraphNode interface
 
 	// UK2Node interface
-	virtual void GetMenuEntries(FGraphContextMenuBuilder& ContextMenuBuilder) const OVERRIDE;
-	virtual void ValidateNodeDuringCompilation(class FCompilerResultsLog& MessageLog) const OVERRIDE;
+	virtual void GetMenuEntries(FGraphContextMenuBuilder& ContextMenuBuilder) const override;
+	virtual void ValidateNodeDuringCompilation(class FCompilerResultsLog& MessageLog) const override;
 	// End of UK2Node interface
 
 private:
diff --git a/Engine/Source/Editor/BlueprintGraph/Classes/K2Node_MatineeController.h b/Engine/Source/Editor/BlueprintGraph/Classes/K2Node_MatineeController.h
--- a/Engine/Source/Editor/BlueprintGraph/Classes/K2Node_MatineeController.h
+++ b/Engine/Source/Editor/BlueprintGraph/Classes/K2Node_MatineeController.h
@@ -15,21 +15,21 @@ class UK2Node_MatineeController : public UK2Node
 	class AMatineeActor* MatineeActor;
 
 	// Begin UEdGraphNode interface
-	virtual void AllocateDefaultPins() OVERRIDE;
-	virtual FLinearColor GetNodeTitleColor() const OVERRIDE;
-	virtual FText GetNodeTitle(ENodeTitleType::Type TitleType) const OVERRIDE;
-	virtual bool CanDuplicateNode() const OVERRIDE { return false; }
-	virtual bool NodeCausesStructuralBlueprintChange() const OVERRIDE { return true; }
+	virtual void AllocateDefaultPins() override;
+	virtual FLinearColor GetNodeTitleColor() const override;
+	virtual FText GetNodeTitle(ENodeTitleType::Type TitleType) const override;
+	virtual bool CanDuplicateNode() const override { return false; }
+	virtual bool NodeCausesStructuralBlueprintChange() const override { return true; }
 	// End UEdGraphNode interface
 
 	// Begin UK2Node interface
-	virtual bool IsNodeSafeToIgnore() const OVERRIDE { return true; }
-	virtual AActor* GetReferencedLevelActor() const OVERRIDE;
-	virtual void ExpandNode(class FKismetCompilerContext& CompilerContext, UEdGraph* SourceGraph) OVERRIDE;
+	virtual bool IsNodeSafeToIgnore() const override { return true; }
+	virtual AActor* GetReferencedLevelActor() const override;
+	virtual void ExpandNode(class FKismetCompilerContext& CompilerContext, UEdGraph* SourceGraph) override;
 	// End UK2Node interface
 
 	// Begin UObject interface
-	virtual void BeginDestroy() OVERRIDE;
+	virtual void BeginDestroy() override;
 	// End UObject interface
 
 	/** Gets the "finished playing matinee sequence" pin */
diff --git a/Engine/Source/Editor/BlueprintGraph/Classes/K2Node_Message.h b/Engine/Source/Editor/BlueprintGraph/Classes/K2Node_Message.h
--- a/Engine/Source/Editor/BlueprintGraph/Classes/K2Node_Message.h
+++ b/Engine/Source/Editor/BlueprintGraph/Classes/K2Node_Message.h
@@ -12,22 +12,22 @@ class UK2Node_Message : public UK2Node_CallFunction
 
 
 	// Begin UEdGraphNode interface.
-	virtual FText GetNodeTitle(ENodeTitleType::Type TitleType) const OVERRIDE;
-	virtual FString GetNodeNativeTitle(ENodeTitleType::Type TitleType) const OVERRIDE;
+	virtual FText GetNodeTitle(ENodeTitleType::Type TitleType) const override;
+	virtual FString GetNodeNativeTitle(ENodeTitleType::Type TitleType) const override;
 	// End UEdGraphNode interface.
 
 	// Begin K2Node interface.
-	virtual bool IsNodeSafeToIgnore() const OVERRIDE { return true; }
-	virtual class FNodeHandlingFunctor* CreateNodeHandler(class FKismetCompilerContext& CompilerContext) const OVERRIDE;
-	virtual void ExpandNode(class FKismetCompilerContext& CompilerContext, UEdGraph* SourceGraph) OVERRIDE;
-	virtual FName GetCornerIcon() const OVERRIDE;
+	virtual bool IsNodeSafeToIgnore() const override { return true; }
+	virtual class FNodeHandlingFunctor* CreateNodeHandler(class FKismetCompilerContext& CompilerContext) const override;
+	virtual void ExpandNode(class FKismetCompilerContext& CompilerContext, UEdGraph* SourceGraph) override;
+	virtual FName GetCornerIcon() const override;
 	// End K2Node interface.
 
 	// Begin K2Node_CallFunction Interface.
-	virtual UEdGraphPin* CreateSelfPin(const UFunction* Function) OVERRIDE;
+	virtual UEdGraphPin* CreateSelfPin(const UFunction* Function) override;
 
 protected:
-	virtual void EnsureFunctionIsInBlueprint() OVERRIDE;
+	virtual void EnsureFunctionIsInBlueprint() override;
 	// End K2Node_CallFunction Interface.
 };
 
diff --git a/Engine/Source/Editor/BlueprintGraph/Classes/K2Node_MultiGate.h b/Engine/Source/Editor/BlueprintGraph/Classes/K2Node_MultiGate.h
--- a/Engine/Source/Editor/BlueprintGraph/Classes/K2Node_MultiGate.h
+++ b/Engine/Source/Editor/BlueprintGraph/Classes/K2Node_MultiGate.h
@@ -14,16 +14,16 @@ class UK2Node_MultiGate : public UK2Node_ExecutionSequence
 	class UK2Node_TemporaryVariable* DataNode;
 
 	// Begin UEdGraphNode interface
-	virtual void AllocateDefaultPins() OVERRIDE;
-	virtual FString GetTooltip() const OVERRIDE;
-	virtual FLinearColor GetNodeTitleColor() const OVERRIDE;
-	virtual FText GetNodeTitle(ENodeTitleType::Type TitleType) const OVERRIDE;
+	virtual void AllocateDefaultPins() override;
+	virtual FString GetTooltip() const override;
+	virtual FLinearColor GetNodeTitleColor() const override;
+	virtual FText GetNodeTitle(ENodeTitleType::Type TitleType) const override;
 	// End UEdGraphNode interface
 
 	// Begin UK2Node interface
-	virtual void ReallocatePinsDuringReconstruction(TArray<UEdGraphPin*>& OldPins) OVERRIDE;
-	virtual class FNodeHandlingFunctor* CreateNodeHandler(class FKismetCompilerContext& CompilerContext) const OVERRIDE;
-	virtual void ExpandNode(class FKismetCompilerContext& CompilerContext, UEdGraph* SourceGraph) OVERRIDE;
+	virtual void ReallocatePinsDuringReconstruction(TArray<UEdGraphPin*>& OldPins) override;
+	virtual class FNodeHandlingFunctor* CreateNodeHandler(class FKismetCompilerContext& CompilerContext) const override;
+	virtual void ExpandNode(class FKismetCompilerContext& CompilerContext, UEdGraph* SourceGraph) override;
 	// End UK2Node interface
 
 	/** Getting pin access */
@@ -52,6 +52,6 @@ class UK2Node_MultiGate : public UK2Node_ExecutionSequence
 
 private:
 	// Returns the exec output pin name for a given 0-based index
- 	virtual FString GetPinNameGivenIndex(int32 Index) const OVERRIDE;
+ 	virtual FString GetPinNameGivenIndex(int32 Index) const override;
 };
 
diff --git a/Engine/Source/Editor/BlueprintGraph/Classes/K2Node_RemoveDelegate.h b/Engine/Source/Editor/BlueprintGraph/Classes/K2Node_RemoveDelegate.h
--- a/Engine/Source/Editor/BlueprintGraph/Classes/K2Node_RemoveDelegate.h
+++ b/Engine/Source/Editor/BlueprintGraph/Classes/K2Node_RemoveDelegate.h
@@ -10,12 +10,12 @@ class UK2Node_RemoveDelegate : public UK2Node_BaseMCDelegate
 
 public:
 	// UEdGraphNode interface
-	virtual void AllocateDefaultPins() OVERRIDE;
-	virtual FText GetNodeTitle(ENodeTitleType::Type TitleType) const OVERRIDE;
-	virtual FString GetNodeNativeTitle(ENodeTitleType::Type TitleType) const OVERRIDE;
+	virtual void AllocateDefaultPins() override;
+	virtual FText GetNodeTitle(ENodeTitleType::Type TitleType) const override;
+	virtual FString GetNodeNativeTitle(ENodeTitleType::Type TitleType) const override;
 	// End of UEdGraphNode interface
 
 	// Begin of UK2Node interface
-	virtual class FNodeHandlingFunctor* CreateNodeHandler(class FKismetCompilerContext& CompilerContext) const OVERRIDE;
+	virtual class FNodeHandlingFunctor* CreateNodeHandler(class FKismetCompilerContext& CompilerContext) const override;
 	// End of UK2Node interface
 };
diff --git a/Engine/Source/Editor/BlueprintGraph/Classes/K2Node_Select.h b/Engine/Source/Editor/BlueprintGraph/Classes/K2Node_Select.h
--- a/Engine/Source/Editor/BlueprintGraph/Classes/K2Node_Select.h
+++ b/Engine/Source/Editor/BlueprintGraph/Classes/K2Node_Select.h
@@ -35,23 +35,23 @@ class UK2Node_Select : public UK2Node
 	bool bReconstructNode;
 
 	// Begin UEdGraphNode interface
-	virtual void AllocateDefaultPins() OVERRIDE;
-	virtual FString GetTooltip() const OVERRIDE;
-	virtual FText GetNodeTitle(ENodeTitleType::Type TitleType) const OVERRIDE;
-	virtual void NodeConnectionListChanged() OVERRIDE;
-	virtual void PinTypeChanged(UEdGraphPin* Pin) OVERRIDE;
-	virtual void PostPasteNode() OVERRIDE;
-	virtual FName GetPaletteIcon(FLinearColor& OutColor) const OVERRIDE{ return TEXT("GraphEditor.Select_16x"); }
+	virtual void AllocateDefaultPins() override;
+	virtual FString GetTooltip() const override;
+	virtual FText GetNodeTitle(ENodeTitleType::Type TitleType) const override;
+	virtual void NodeConnectionListChanged() override;
+	virtual void PinTypeChanged(UEdGraphPin* Pin) override;
+	virtual void PostPasteNode() override;
+	virtual FName GetPaletteIcon(FLinearColor& OutColor) const override{ return TEXT("GraphEditor.Select_16x"); }
 	// End UEdGraphNode interface
 
 	// Begin UK2Node interface
-	virtual void ReallocatePinsDuringReconstruction(TArray<UEdGraphPin*>& OldPins) OVERRIDE;
- 	virtual void PostReconstructNode() OVERRIDE;
-	virtual bool IsNodeSafeToIgnore() const OVERRIDE { return true; }
-	virtual bool IsNodePure() const OVERRIDE { return true; }
-	virtual void NotifyPinConnectionListChanged(UEdGraphPin* Pin) OVERRIDE;
-	virtual class FNodeHandlingFunctor* CreateNodeHandler(class FKismetCompilerContext& CompilerContext) const OVERRIDE;
-	virtual bool IsConnectionDisallowed(const UEdGraphPin* MyPin, const UEdGraphPin* OtherPin, FString& OutReason) const OVERRIDE;
+	virtual void ReallocatePinsDuringReconstruction(TArray<UEdGraphPin*>& OldPins) override;
+ 	virtual void PostReconstructNode() override;
+	virtual bool IsNodeSafeToIgnore() const override { return true; }
+	virtual bool IsNodePure() const override { return true; }
+	virtual void NotifyPinConnectionListChanged(UEdGraphPin* Pin) override;
+	virtual class FNodeHandlingFunctor* CreateNodeHandler(class FKismetCompilerContext& CompilerContext) const override;
+	virtual bool IsConnectionDisallowed(const UEdGraphPin* MyPin, const UEdGraphPin* OtherPin, FString& OutReason) const override;
 	// End UK2Node interface
 
 	/** Get the return value pin */
diff --git a/Engine/Source/Editor/BlueprintGraph/Classes/K2Node_SelectEnum.h b/Engine/Source/Editor/BlueprintGraph/Classes/K2Node_SelectEnum.h
--- a/Engine/Source/Editor/BlueprintGraph/Classes/K2Node_SelectEnum.h
+++ b/Engine/Source/Editor/BlueprintGraph/Classes/K2Node_SelectEnum.h
@@ -11,23 +11,23 @@ class UDEPRECATED_K2Node_SelectEnum : public UK2Node_Select
 	GENERATED_UCLASS_BODY()
 
 	// Begin UEdGraphNode interface
-	virtual void AllocateDefaultPins() OVERRIDE;
-	virtual FString GetTooltip() const OVERRIDE;
-	virtual FText GetNodeTitle(ENodeTitleType::Type TitleType) const OVERRIDE;
-	virtual FString GetNodeNativeTitle(ENodeTitleType::Type TitleType) const OVERRIDE;
+	virtual void AllocateDefaultPins() override;
+	virtual FString GetTooltip() const override;
+	virtual FText GetNodeTitle(ENodeTitleType::Type TitleType) const override;
+	virtual FString GetNodeNativeTitle(ENodeTitleType::Type TitleType) const override;
 	// End UEdGraphNode interface
 
 	// Begin UK2Node interface
-	virtual void NotifyPinConnectionListChanged(UEdGraphPin* Pin) OVERRIDE;
+	virtual void NotifyPinConnectionListChanged(UEdGraphPin* Pin) override;
 	// End UK2Node interface
 
 	// Begin UK2Node_Select interface
-	virtual void GetOptionPins(TArray<UEdGraphPin*>& OptionPins) const OVERRIDE;
-	virtual bool CanAddOptionPinToNode() const OVERRIDE { return false; }
-	virtual bool CanRemoveOptionPinToNode() const OVERRIDE { return false; }
+	virtual void GetOptionPins(TArray<UEdGraphPin*>& OptionPins) const override;
+	virtual bool CanAddOptionPinToNode() const override { return false; }
+	virtual bool CanRemoveOptionPinToNode() const override { return false; }
 	// End UK2Node_Select interface
 
 	// Bind the option to a named enum 
-	BLUEPRINTGRAPH_API void SetEnum(UEnum* InEnum) OVERRIDE;
+	BLUEPRINTGRAPH_API void SetEnum(UEnum* InEnum) override;
 };
 
diff --git a/Engine/Source/Editor/BlueprintGraph/Classes/K2Node_Self.h b/Engine/Source/Editor/BlueprintGraph/Classes/K2Node_Self.h
--- a/Engine/Source/Editor/BlueprintGraph/Classes/K2Node_Self.h
+++ b/Engine/Source/Editor/BlueprintGraph/Classes/K2Node_Self.h
@@ -15,17 +15,17 @@ class UK2Node_Self : public UK2Node
 	TSubclassOf<class UObject>  SelfClass;
 
 	// Begin UEdGraphNode interface
-	virtual FText GetNodeTitle(ENodeTitleType::Type TitleType) const OVERRIDE;
-	virtual FString GetTooltip() const OVERRIDE;
-	virtual FString GetKeywords() const OVERRIDE;
-	virtual void AllocateDefaultPins() OVERRIDE;
-	virtual void ValidateNodeDuringCompilation(class FCompilerResultsLog& MessageLog) const OVERRIDE;
+	virtual FText GetNodeTitle(ENodeTitleType::Type TitleType) const override;
+	virtual FString GetTooltip() const override;
+	virtual FString GetKeywords() const override;
+	virtual void AllocateDefaultPins() override;
+	virtual void ValidateNodeDuringCompilation(class FCompilerResultsLog& MessageLog) const override;
 	// End UEdGraphNode interface
 
 	// Begin K2Node interface
-	virtual bool IsNodePure() const OVERRIDE { return true; }
-	virtual bool DrawNodeAsVariable() const OVERRIDE { return true; }
-	virtual class FNodeHandlingFunctor* CreateNodeHandler(class FKismetCompilerContext& CompilerContext) const OVERRIDE;
+	virtual bool IsNodePure() const override { return true; }
+	virtual bool DrawNodeAsVariable() const override { return true; }
+	virtual class FNodeHandlingFunctor* CreateNodeHandler(class FKismetCompilerContext& CompilerContext) const override;
 	// End K2Node interface
 };
 
