UE4 Refactoring. Changed OVERRIDE and FINAL macros to keywords override and final.

diff --git a/Engine/Source/Editor/BlueprintGraph/Classes/K2Node_SpawnActor.h b/Engine/Source/Editor/BlueprintGraph/Classes/K2Node_SpawnActor.h
--- a/Engine/Source/Editor/BlueprintGraph/Classes/K2Node_SpawnActor.h
+++ b/Engine/Source/Editor/BlueprintGraph/Classes/K2Node_SpawnActor.h
@@ -11,26 +11,26 @@ class UK2Node_SpawnActor : public UK2Node
 	GENERATED_UCLASS_BODY()
 
 	// Begin UEdGraphNode interface.
-	virtual void AllocateDefaultPins() OVERRIDE;
-	virtual FLinearColor GetNodeTitleColor() const OVERRIDE;
-	virtual FText GetNodeTitle(ENodeTitleType::Type TitleType) const OVERRIDE;
-	virtual FString GetNodeNativeTitle(ENodeTitleType::Type TitleType) const OVERRIDE;
-	virtual void PinDefaultValueChanged(UEdGraphPin* Pin) OVERRIDE;
-	virtual FString GetTooltip() const OVERRIDE;
-	virtual bool IsDeprecated() const OVERRIDE;
-	virtual bool ShouldWarnOnDeprecation() const OVERRIDE;
-	virtual FString GetDeprecationMessage() const OVERRIDE;
-	virtual FName GetPaletteIcon(FLinearColor& OutColor) const OVERRIDE{ return TEXT("GraphEditor.SpawnActor_16x"); }
+	virtual void AllocateDefaultPins() override;
+	virtual FLinearColor GetNodeTitleColor() const override;
+	virtual FText GetNodeTitle(ENodeTitleType::Type TitleType) const override;
+	virtual FString GetNodeNativeTitle(ENodeTitleType::Type TitleType) const override;
+	virtual void PinDefaultValueChanged(UEdGraphPin* Pin) override;
+	virtual FString GetTooltip() const override;
+	virtual bool IsDeprecated() const override;
+	virtual bool ShouldWarnOnDeprecation() const override;
+	virtual FString GetDeprecationMessage() const override;
+	virtual FName GetPaletteIcon(FLinearColor& OutColor) const override{ return TEXT("GraphEditor.SpawnActor_16x"); }
 	// End UEdGraphNode interface.
 
 	// Begin UK2Node interface
-	virtual bool IsNodeSafeToIgnore() const OVERRIDE { return true; }
-	virtual void ReallocatePinsDuringReconstruction(TArray<UEdGraphPin*>& OldPins) OVERRIDE;
-	virtual bool CanPasteHere(const UEdGraph* TargetGraph, const UEdGraphSchema* Schema) const OVERRIDE;
-	virtual class FNodeHandlingFunctor* CreateNodeHandler(class FKismetCompilerContext& CompilerContext) const OVERRIDE;
-	virtual void ExpandNode(class FKismetCompilerContext& CompilerContext, UEdGraph* SourceGraph) OVERRIDE;
-	virtual bool HasExternalBlueprintDependencies(TArray<class UStruct*>* OptionalOutput) const OVERRIDE;
-	virtual void GetNodeAttributes( TArray<TKeyValuePair<FString, FString>>& OutNodeAttributes ) const OVERRIDE;
+	virtual bool IsNodeSafeToIgnore() const override { return true; }
+	virtual void ReallocatePinsDuringReconstruction(TArray<UEdGraphPin*>& OldPins) override;
+	virtual bool CanPasteHere(const UEdGraph* TargetGraph, const UEdGraphSchema* Schema) const override;
+	virtual class FNodeHandlingFunctor* CreateNodeHandler(class FKismetCompilerContext& CompilerContext) const override;
+	virtual void ExpandNode(class FKismetCompilerContext& CompilerContext, UEdGraph* SourceGraph) override;
+	virtual bool HasExternalBlueprintDependencies(TArray<class UStruct*>* OptionalOutput) const override;
+	virtual void GetNodeAttributes( TArray<TKeyValuePair<FString, FString>>& OutNodeAttributes ) const override;
 	// End UK2Node interface
 
 
diff --git a/Engine/Source/Editor/BlueprintGraph/Classes/K2Node_SpawnActorFromClass.h b/Engine/Source/Editor/BlueprintGraph/Classes/K2Node_SpawnActorFromClass.h
--- a/Engine/Source/Editor/BlueprintGraph/Classes/K2Node_SpawnActorFromClass.h
+++ b/Engine/Source/Editor/BlueprintGraph/Classes/K2Node_SpawnActorFromClass.h
@@ -11,22 +11,22 @@ class UK2Node_SpawnActorFromClass : public UK2Node
 	GENERATED_UCLASS_BODY()
 
 	// Begin UEdGraphNode interface.
-	virtual void AllocateDefaultPins() OVERRIDE;
-	virtual FLinearColor GetNodeTitleColor() const OVERRIDE;
-	virtual FText GetNodeTitle(ENodeTitleType::Type TitleType) const OVERRIDE;
-	virtual FString GetNodeNativeTitle(ENodeTitleType::Type TitleType) const OVERRIDE;
-	virtual void PinDefaultValueChanged(UEdGraphPin* Pin) OVERRIDE;
-	virtual FString GetTooltip() const OVERRIDE;
-	virtual void ExpandNode(class FKismetCompilerContext& CompilerContext, UEdGraph* SourceGraph) OVERRIDE;
-	virtual bool HasExternalBlueprintDependencies(TArray<class UStruct*>* OptionalOutput) const OVERRIDE;
-	virtual FName GetPaletteIcon(FLinearColor& OutColor) const OVERRIDE{ return TEXT("GraphEditor.SpawnActor_16x"); }
+	virtual void AllocateDefaultPins() override;
+	virtual FLinearColor GetNodeTitleColor() const override;
+	virtual FText GetNodeTitle(ENodeTitleType::Type TitleType) const override;
+	virtual FString GetNodeNativeTitle(ENodeTitleType::Type TitleType) const override;
+	virtual void PinDefaultValueChanged(UEdGraphPin* Pin) override;
+	virtual FString GetTooltip() const override;
+	virtual void ExpandNode(class FKismetCompilerContext& CompilerContext, UEdGraph* SourceGraph) override;
+	virtual bool HasExternalBlueprintDependencies(TArray<class UStruct*>* OptionalOutput) const override;
+	virtual FName GetPaletteIcon(FLinearColor& OutColor) const override{ return TEXT("GraphEditor.SpawnActor_16x"); }
 	// End UEdGraphNode interface.
 
 	// Begin UK2Node interface
-	virtual bool IsNodeSafeToIgnore() const OVERRIDE { return true; }
-	virtual void ReallocatePinsDuringReconstruction(TArray<UEdGraphPin*>& OldPins) OVERRIDE;
-	virtual bool CanPasteHere(const UEdGraph* TargetGraph, const UEdGraphSchema* Schema) const OVERRIDE;
-	virtual void GetNodeAttributes( TArray<TKeyValuePair<FString, FString>>& OutNodeAttributes ) const OVERRIDE;
+	virtual bool IsNodeSafeToIgnore() const override { return true; }
+	virtual void ReallocatePinsDuringReconstruction(TArray<UEdGraphPin*>& OldPins) override;
+	virtual bool CanPasteHere(const UEdGraph* TargetGraph, const UEdGraphSchema* Schema) const override;
+	virtual void GetNodeAttributes( TArray<TKeyValuePair<FString, FString>>& OutNodeAttributes ) const override;
 	// End UK2Node interface
 
 
diff --git a/Engine/Source/Editor/BlueprintGraph/Classes/K2Node_StructMemberGet.h b/Engine/Source/Editor/BlueprintGraph/Classes/K2Node_StructMemberGet.h
--- a/Engine/Source/Editor/BlueprintGraph/Classes/K2Node_StructMemberGet.h
+++ b/Engine/Source/Editor/BlueprintGraph/Classes/K2Node_StructMemberGet.h
@@ -14,20 +14,20 @@ class UK2Node_StructMemberGet : public UK2Node_StructOperation
 	TArray<FOptionalPinFromProperty> ShowPinForProperties;
 
 	// UObject interface
-	virtual void PostEditChangeProperty(struct FPropertyChangedEvent& PropertyChangedEvent) OVERRIDE;
+	virtual void PostEditChangeProperty(struct FPropertyChangedEvent& PropertyChangedEvent) override;
 	// End of UObject interface
 
 	// UEdGraphNode interface
-	virtual void AllocateDefaultPins() OVERRIDE;
-	virtual FString GetTooltip() const OVERRIDE;
-	virtual FText GetNodeTitle(ENodeTitleType::Type TitleType) const OVERRIDE;
-	virtual FString GetNodeNativeTitle(ENodeTitleType::Type TitleType) const OVERRIDE;
+	virtual void AllocateDefaultPins() override;
+	virtual FString GetTooltip() const override;
+	virtual FText GetNodeTitle(ENodeTitleType::Type TitleType) const override;
+	virtual FString GetNodeNativeTitle(ENodeTitleType::Type TitleType) const override;
 	// End of UEdGraphNode interface
 
 	// UK2Node interface
-	virtual bool IsNodePure() const OVERRIDE { return true; }
-	virtual bool NodeCausesStructuralBlueprintChange() const OVERRIDE { return true; }
-	virtual class FNodeHandlingFunctor* CreateNodeHandler(class FKismetCompilerContext& CompilerContext) const OVERRIDE;
+	virtual bool IsNodePure() const override { return true; }
+	virtual bool NodeCausesStructuralBlueprintChange() const override { return true; }
+	virtual class FNodeHandlingFunctor* CreateNodeHandler(class FKismetCompilerContext& CompilerContext) const override;
 	// End of UK2Node interface
 
 	// AllocateDefaultPins with just one member set
diff --git a/Engine/Source/Editor/BlueprintGraph/Classes/K2Node_StructMemberSet.h b/Engine/Source/Editor/BlueprintGraph/Classes/K2Node_StructMemberSet.h
--- a/Engine/Source/Editor/BlueprintGraph/Classes/K2Node_StructMemberSet.h
+++ b/Engine/Source/Editor/BlueprintGraph/Classes/K2Node_StructMemberSet.h
@@ -14,21 +14,21 @@ class UK2Node_StructMemberSet : public UK2Node_StructOperation
 	TArray<FOptionalPinFromProperty> ShowPinForProperties;
 
 	// UObject interface
-	virtual void PostEditChangeProperty(struct FPropertyChangedEvent& PropertyChangedEvent) OVERRIDE;
+	virtual void PostEditChangeProperty(struct FPropertyChangedEvent& PropertyChangedEvent) override;
 	// End of UObject interface
 
 	// UEdGraphNode interface
-	virtual void AllocateDefaultPins() OVERRIDE;
-	virtual FString GetTooltip() const OVERRIDE;
-	virtual FText GetNodeTitle(ENodeTitleType::Type TitleType) const OVERRIDE;
-	virtual FString GetNodeNativeTitle(ENodeTitleType::Type TitleType) const OVERRIDE;
+	virtual void AllocateDefaultPins() override;
+	virtual FString GetTooltip() const override;
+	virtual FText GetNodeTitle(ENodeTitleType::Type TitleType) const override;
+	virtual FString GetNodeNativeTitle(ENodeTitleType::Type TitleType) const override;
 	// End of UEdGraphNode interface
 
 	// UK2Node interface
-	virtual bool IsNodePure() const OVERRIDE { return false; }
-	virtual bool NodeCausesStructuralBlueprintChange() const OVERRIDE { return true; }
-	virtual ERedirectType DoPinsMatchForReconstruction(const UEdGraphPin* NewPin, int32 NewPinIndex, const UEdGraphPin* OldPin, int32 OldPinIndex)  const OVERRIDE;
-	virtual class FNodeHandlingFunctor* CreateNodeHandler(class FKismetCompilerContext& CompilerContext) const OVERRIDE;
+	virtual bool IsNodePure() const override { return false; }
+	virtual bool NodeCausesStructuralBlueprintChange() const override { return true; }
+	virtual ERedirectType DoPinsMatchForReconstruction(const UEdGraphPin* NewPin, int32 NewPinIndex, const UEdGraphPin* OldPin, int32 OldPinIndex)  const override;
+	virtual class FNodeHandlingFunctor* CreateNodeHandler(class FKismetCompilerContext& CompilerContext) const override;
 	// End of UK2Node interface
 };
 
diff --git a/Engine/Source/Editor/BlueprintGraph/Classes/K2Node_StructOperation.h b/Engine/Source/Editor/BlueprintGraph/Classes/K2Node_StructOperation.h
--- a/Engine/Source/Editor/BlueprintGraph/Classes/K2Node_StructOperation.h
+++ b/Engine/Source/Editor/BlueprintGraph/Classes/K2Node_StructOperation.h
@@ -20,9 +20,9 @@ class UK2Node_StructOperation : public UK2Node_Variable
 
 	// UK2Node interface
 	//virtual bool DrawNodeAsVariable() const OVERRIDE { return true; }
-	virtual bool ShouldShowNodeProperties() const OVERRIDE { return true; }
-	virtual void ValidateNodeDuringCompilation(class FCompilerResultsLog& MessageLog) const OVERRIDE {}
-	virtual bool HasExternalUserDefinedStructDependencies(TArray<class UStruct*>* OptionalOutput) const OVERRIDE;
+	virtual bool ShouldShowNodeProperties() const override { return true; }
+	virtual void ValidateNodeDuringCompilation(class FCompilerResultsLog& MessageLog) const override {}
+	virtual bool HasExternalUserDefinedStructDependencies(TArray<class UStruct*>* OptionalOutput) const override;
 	// End of UK2Node interface
 
 protected:
@@ -30,7 +30,7 @@ protected:
 	struct FStructOperationOptionalPinManager : public FOptionalPinManager
 	{
 		// FOptionalPinsUpdater interface
-		virtual void GetRecordDefaults(UProperty* TestProperty, FOptionalPinFromProperty& Record) const OVERRIDE
+		virtual void GetRecordDefaults(UProperty* TestProperty, FOptionalPinFromProperty& Record) const override
 		{
 			Record.bCanToggleVisibility = true;
 			Record.bShowPin = true;
diff --git a/Engine/Source/Editor/BlueprintGraph/Classes/K2Node_Switch.h b/Engine/Source/Editor/BlueprintGraph/Classes/K2Node_Switch.h
--- a/Engine/Source/Editor/BlueprintGraph/Classes/K2Node_Switch.h
+++ b/Engine/Source/Editor/BlueprintGraph/Classes/K2Node_Switch.h
@@ -23,18 +23,18 @@ class UK2Node_Switch : public UK2Node
 	TSubclassOf<class UObject> FunctionClass;
 
 	// UObject interface
-	virtual void PostEditChangeProperty(struct FPropertyChangedEvent& PropertyChangedEvent) OVERRIDE;
+	virtual void PostEditChangeProperty(struct FPropertyChangedEvent& PropertyChangedEvent) override;
 	// End of UObject interface
 
 	// UEdGraphNode interface
-	virtual void AllocateDefaultPins() OVERRIDE;
-	virtual FLinearColor GetNodeTitleColor() const OVERRIDE;
-	virtual FName GetPaletteIcon(FLinearColor& OutColor) const OVERRIDE{ return TEXT("GraphEditor.Switch_16x"); }
+	virtual void AllocateDefaultPins() override;
+	virtual FLinearColor GetNodeTitleColor() const override;
+	virtual FName GetPaletteIcon(FLinearColor& OutColor) const override{ return TEXT("GraphEditor.Switch_16x"); }
 	// End of UEdGraphNode interface
 
 	// UK2Node interface
-	virtual ERedirectType DoPinsMatchForReconstruction(const UEdGraphPin* NewPin, int32 NewPinIndex, const UEdGraphPin* OldPin, int32 OldPinIndex) const OVERRIDE;
-	virtual class FNodeHandlingFunctor* CreateNodeHandler(class FKismetCompilerContext& CompilerContext) const OVERRIDE;
+	virtual ERedirectType DoPinsMatchForReconstruction(const UEdGraphPin* NewPin, int32 NewPinIndex, const UEdGraphPin* OldPin, int32 OldPinIndex) const override;
+	virtual class FNodeHandlingFunctor* CreateNodeHandler(class FKismetCompilerContext& CompilerContext) const override;
 	// End of UK2Node interface
 
 	// UK2Node_Switch interface
diff --git a/Engine/Source/Editor/BlueprintGraph/Classes/K2Node_SwitchEnum.h b/Engine/Source/Editor/BlueprintGraph/Classes/K2Node_SwitchEnum.h
--- a/Engine/Source/Editor/BlueprintGraph/Classes/K2Node_SwitchEnum.h
+++ b/Engine/Source/Editor/BlueprintGraph/Classes/K2Node_SwitchEnum.h
@@ -24,23 +24,23 @@ class UK2Node_SwitchEnum : public UK2Node_Switch, public INodeDependingOnEnumInt
 	TArray<FString> EnumFriendlyNames;
 
 	// INodeDependingOnEnumInterface
-	virtual class UEnum* GetEnum() const OVERRIDE { return Enum; }
-	virtual bool ShouldBeReconstructedAfterEnumChanged() const OVERRIDE {return true;}
+	virtual class UEnum* GetEnum() const override { return Enum; }
+	virtual bool ShouldBeReconstructedAfterEnumChanged() const override {return true;}
 	// End of INodeDependingOnEnumInterface
 
 	// UEdGraphNode interface
-	virtual FString GetTooltip() const OVERRIDE;
-	virtual FText GetNodeTitle(ENodeTitleType::Type TitleType) const OVERRIDE;
-	virtual FString GetNodeNativeTitle(ENodeTitleType::Type TitleType) const OVERRIDE;
-	virtual bool ShouldShowNodeProperties() const OVERRIDE { return true; }	
+	virtual FString GetTooltip() const override;
+	virtual FText GetNodeTitle(ENodeTitleType::Type TitleType) const override;
+	virtual FString GetNodeNativeTitle(ENodeTitleType::Type TitleType) const override;
+	virtual bool ShouldShowNodeProperties() const override { return true; }	
 	// End of UEdGraphNode interface
 
 	// UK2Node_Switch Interface
-	virtual FString GetUniquePinName() OVERRIDE;
-	virtual const FString& GetPinType(const UEdGraphSchema_K2* Schema) const OVERRIDE { return Schema->PC_Byte; }
-	virtual void AddPinToSwitchNode() OVERRIDE;
-	virtual void RemovePinFromSwitchNode(UEdGraphPin* TargetPin) OVERRIDE;
-	virtual ERedirectType DoPinsMatchForReconstruction(const UEdGraphPin* NewPin, int32 NewPinIndex, const UEdGraphPin* OldPin, int32 OldPinIndex) const OVERRIDE;
+	virtual FString GetUniquePinName() override;
+	virtual const FString& GetPinType(const UEdGraphSchema_K2* Schema) const override { return Schema->PC_Byte; }
+	virtual void AddPinToSwitchNode() override;
+	virtual void RemovePinFromSwitchNode(UEdGraphPin* TargetPin) override;
+	virtual ERedirectType DoPinsMatchForReconstruction(const UEdGraphPin* NewPin, int32 NewPinIndex, const UEdGraphPin* OldPin, int32 OldPinIndex) const override;
 	// End of UK2Node_Switch Interface
 
 	/** Bind the switch to a named enum */
@@ -48,10 +48,10 @@ class UK2Node_SwitchEnum : public UK2Node_Switch, public INodeDependingOnEnumInt
 
 protected:
 	/** Helper method to set-up pins */
-	virtual void CreateCasePins() OVERRIDE;
+	virtual void CreateCasePins() override;
 	/** Helper method to set-up correct selection pin */
-	virtual void CreateSelectionPin() OVERRIDE;
+	virtual void CreateSelectionPin() override;
 	
 	/** Don't support removing pins from an enum */
-	virtual void RemovePin(UEdGraphPin* TargetPin) OVERRIDE {}
+	virtual void RemovePin(UEdGraphPin* TargetPin) override {}
 };
diff --git a/Engine/Source/Editor/BlueprintGraph/Classes/K2Node_SwitchInteger.h b/Engine/Source/Editor/BlueprintGraph/Classes/K2Node_SwitchInteger.h
--- a/Engine/Source/Editor/BlueprintGraph/Classes/K2Node_SwitchInteger.h
+++ b/Engine/Source/Editor/BlueprintGraph/Classes/K2Node_SwitchInteger.h
@@ -16,24 +16,24 @@ class UK2Node_SwitchInteger : public UK2Node_Switch
 	int32 StartIndex;
 
 	// UObject interface
-	void PostEditChangeProperty(struct FPropertyChangedEvent& PropertyChangedEvent) OVERRIDE;
+	void PostEditChangeProperty(struct FPropertyChangedEvent& PropertyChangedEvent) override;
 	// End of UObject interface
 
 	// UEdGraphNode interface
-	virtual FText GetNodeTitle(ENodeTitleType::Type TitleType) const OVERRIDE;
-	virtual FString GetTooltip() const OVERRIDE;
-	virtual bool ShouldShowNodeProperties() const OVERRIDE { return true; }
+	virtual FText GetNodeTitle(ENodeTitleType::Type TitleType) const override;
+	virtual FString GetTooltip() const override;
+	virtual bool ShouldShowNodeProperties() const override { return true; }
 	// End of UEdGraphNode interface
 
 	// UK2Node_Switch Interface
-	virtual FString GetUniquePinName() OVERRIDE;
-	virtual FString GetPinNameGivenIndex(int32 Index) OVERRIDE;
-	virtual void ReallocatePinsDuringReconstruction(TArray<UEdGraphPin*>& OldPins) OVERRIDE;
-	virtual const FString& GetPinType(const UEdGraphSchema_K2* Schema) const OVERRIDE { return Schema->PC_Int; }
+	virtual FString GetUniquePinName() override;
+	virtual FString GetPinNameGivenIndex(int32 Index) override;
+	virtual void ReallocatePinsDuringReconstruction(TArray<UEdGraphPin*>& OldPins) override;
+	virtual const FString& GetPinType(const UEdGraphSchema_K2* Schema) const override { return Schema->PC_Int; }
 	// End of UK2Node_Switch Interface
 
 protected:
-	virtual void CreateCasePins() OVERRIDE;
-	virtual void CreateSelectionPin() OVERRIDE;
-	virtual void RemovePin(UEdGraphPin* TargetPin) OVERRIDE;
+	virtual void CreateCasePins() override;
+	virtual void CreateSelectionPin() override;
+	virtual void RemovePin(UEdGraphPin* TargetPin) override;
 };
diff --git a/Engine/Source/Editor/BlueprintGraph/Classes/K2Node_SwitchName.h b/Engine/Source/Editor/BlueprintGraph/Classes/K2Node_SwitchName.h
--- a/Engine/Source/Editor/BlueprintGraph/Classes/K2Node_SwitchName.h
+++ b/Engine/Source/Editor/BlueprintGraph/Classes/K2Node_SwitchName.h
@@ -14,25 +14,25 @@ class UK2Node_SwitchName : public UK2Node_Switch
 	TArray<FName> PinNames;
 
 	// UObject interface
-	virtual void PostEditChangeProperty(struct FPropertyChangedEvent& PropertyChangedEvent) OVERRIDE;
+	virtual void PostEditChangeProperty(struct FPropertyChangedEvent& PropertyChangedEvent) override;
 	// End of UObject interface
 
 	// UEdGraphNode interface
-	virtual FString GetTooltip() const OVERRIDE;
-	virtual FText GetNodeTitle(ENodeTitleType::Type TitleType) const OVERRIDE;
-	virtual bool ShouldShowNodeProperties() const OVERRIDE { return true; }
+	virtual FString GetTooltip() const override;
+	virtual FText GetNodeTitle(ENodeTitleType::Type TitleType) const override;
+	virtual bool ShouldShowNodeProperties() const override { return true; }
 	// End of UEdGraphNode interface
 
 	// UK2Node_Switch Interface
-	BLUEPRINTGRAPH_API virtual void AddPinToSwitchNode() OVERRIDE;
-	virtual FString GetUniquePinName() OVERRIDE;
-	virtual const FString& GetPinType(const UEdGraphSchema_K2* Schema) const OVERRIDE { return Schema->PC_Name; }
+	BLUEPRINTGRAPH_API virtual void AddPinToSwitchNode() override;
+	virtual FString GetUniquePinName() override;
+	virtual const FString& GetPinType(const UEdGraphSchema_K2* Schema) const override { return Schema->PC_Name; }
 	// End of UK2Node_Switch Interface
 
-	virtual FString GetPinNameGivenIndex(int32 Index) OVERRIDE;
+	virtual FString GetPinNameGivenIndex(int32 Index) override;
 
 protected:
-	virtual void CreateSelectionPin() OVERRIDE;
-	virtual void CreateCasePins() OVERRIDE;
-	virtual void RemovePin(UEdGraphPin* TargetPin) OVERRIDE;
+	virtual void CreateSelectionPin() override;
+	virtual void CreateCasePins() override;
+	virtual void RemovePin(UEdGraphPin* TargetPin) override;
 };
diff --git a/Engine/Source/Editor/BlueprintGraph/Classes/K2Node_SwitchString.h b/Engine/Source/Editor/BlueprintGraph/Classes/K2Node_SwitchString.h
--- a/Engine/Source/Editor/BlueprintGraph/Classes/K2Node_SwitchString.h
+++ b/Engine/Source/Editor/BlueprintGraph/Classes/K2Node_SwitchString.h
@@ -17,25 +17,25 @@ class UK2Node_SwitchString : public UK2Node_Switch
 	uint32 bIsCaseSensitive : 1;
 
 	// UObject interface
-	virtual void PostEditChangeProperty(struct FPropertyChangedEvent& PropertyChangedEvent) OVERRIDE;
+	virtual void PostEditChangeProperty(struct FPropertyChangedEvent& PropertyChangedEvent) override;
 	// End of UObject interface
 
 	// UEdGraphNode interface
-	virtual FString GetTooltip() const OVERRIDE;
-	virtual FText GetNodeTitle(ENodeTitleType::Type TitleType) const OVERRIDE;
-	virtual bool ShouldShowNodeProperties() const OVERRIDE { return true; }
+	virtual FString GetTooltip() const override;
+	virtual FText GetNodeTitle(ENodeTitleType::Type TitleType) const override;
+	virtual bool ShouldShowNodeProperties() const override { return true; }
 	// End of UEdGraphNode interface
 
 	// UK2Node_Switch Interface
-	BLUEPRINTGRAPH_API virtual void AddPinToSwitchNode() OVERRIDE;
-	virtual FString GetUniquePinName() OVERRIDE;
-	virtual const FString& GetPinType(const UEdGraphSchema_K2* Schema) const OVERRIDE { return Schema->PC_String; }
+	BLUEPRINTGRAPH_API virtual void AddPinToSwitchNode() override;
+	virtual FString GetUniquePinName() override;
+	virtual const FString& GetPinType(const UEdGraphSchema_K2* Schema) const override { return Schema->PC_String; }
 	// End of UK2Node_Switch Interface
 
-	virtual FString GetPinNameGivenIndex(int32 Index) OVERRIDE;
+	virtual FString GetPinNameGivenIndex(int32 Index) override;
 
 protected:
-	virtual void CreateSelectionPin() OVERRIDE;
-	virtual void CreateCasePins() OVERRIDE;
-	virtual void RemovePin(UEdGraphPin* TargetPin) OVERRIDE;
+	virtual void CreateSelectionPin() override;
+	virtual void CreateCasePins() override;
+	virtual void RemovePin(UEdGraphPin* TargetPin) override;
 };
diff --git a/Engine/Source/Editor/BlueprintGraph/Classes/K2Node_TemporaryVariable.h b/Engine/Source/Editor/BlueprintGraph/Classes/K2Node_TemporaryVariable.h
--- a/Engine/Source/Editor/BlueprintGraph/Classes/K2Node_TemporaryVariable.h
+++ b/Engine/Source/Editor/BlueprintGraph/Classes/K2Node_TemporaryVariable.h
@@ -20,16 +20,16 @@ class UK2Node_TemporaryVariable : public UK2Node
 	BLUEPRINTGRAPH_API UEdGraphPin* GetVariablePin();
 
 	// Begin UEdGraphNode interface.
-	virtual void AllocateDefaultPins() OVERRIDE;
-	virtual FString GetTooltip() const OVERRIDE;
-	virtual FText GetNodeTitle(ENodeTitleType::Type TitleType) const OVERRIDE;
-	virtual FString GetNodeNativeTitle(ENodeTitleType::Type TitleType) const OVERRIDE;
-	virtual FString GetDescriptiveCompiledName() const OVERRIDE;
+	virtual void AllocateDefaultPins() override;
+	virtual FString GetTooltip() const override;
+	virtual FText GetNodeTitle(ENodeTitleType::Type TitleType) const override;
+	virtual FString GetNodeNativeTitle(ENodeTitleType::Type TitleType) const override;
+	virtual FString GetDescriptiveCompiledName() const override;
 	// End UEdGraphNode interface.
 
 	// Begin UK2Node interface.
-	virtual bool IsNodePure() const OVERRIDE;
-	virtual class FNodeHandlingFunctor* CreateNodeHandler(class FKismetCompilerContext& CompilerContext) const OVERRIDE;
+	virtual bool IsNodePure() const override;
+	virtual class FNodeHandlingFunctor* CreateNodeHandler(class FKismetCompilerContext& CompilerContext) const override;
 	// End UK2Node interface.
 };
 
diff --git a/Engine/Source/Editor/BlueprintGraph/Classes/K2Node_Timeline.h b/Engine/Source/Editor/BlueprintGraph/Classes/K2Node_Timeline.h
--- a/Engine/Source/Editor/BlueprintGraph/Classes/K2Node_Timeline.h
+++ b/Engine/Source/Editor/BlueprintGraph/Classes/K2Node_Timeline.h
@@ -31,27 +31,27 @@ class UK2Node_Timeline : public UK2Node
 	uint32 bReplicated:1;
 
 	// Begin UEdGraphNode interface.
-	virtual void AllocateDefaultPins() OVERRIDE;
-	virtual void DestroyNode() OVERRIDE;
-	virtual void PostPasteNode() OVERRIDE;
-	virtual void PrepareForCopying() OVERRIDE;
-	virtual FLinearColor GetNodeTitleColor() const OVERRIDE;
-	virtual FText GetNodeTitle(ENodeTitleType::Type TitleType) const OVERRIDE;
-	virtual bool CanPasteHere(const UEdGraph* TargetGraph, const UEdGraphSchema* Schema) const OVERRIDE;
-	virtual void FindDiffs(class UEdGraphNode* OtherNode, struct FDiffResults& Results )  OVERRIDE;
-	virtual void OnRenameNode(const FString& NewName) OVERRIDE;
-	virtual TSharedPtr<class INameValidatorInterface> MakeNameValidator() const OVERRIDE;
-	virtual FString GetTooltip() const OVERRIDE;
-	virtual FString GetDocumentationExcerptName() const OVERRIDE;
-	virtual FName GetPaletteIcon(FLinearColor& OutColor) const OVERRIDE{ return TEXT("GraphEditor.Timeline_16x"); }
-	virtual bool ShouldShowNodeProperties() const OVERRIDE { return true; }
+	virtual void AllocateDefaultPins() override;
+	virtual void DestroyNode() override;
+	virtual void PostPasteNode() override;
+	virtual void PrepareForCopying() override;
+	virtual FLinearColor GetNodeTitleColor() const override;
+	virtual FText GetNodeTitle(ENodeTitleType::Type TitleType) const override;
+	virtual bool CanPasteHere(const UEdGraph* TargetGraph, const UEdGraphSchema* Schema) const override;
+	virtual void FindDiffs(class UEdGraphNode* OtherNode, struct FDiffResults& Results )  override;
+	virtual void OnRenameNode(const FString& NewName) override;
+	virtual TSharedPtr<class INameValidatorInterface> MakeNameValidator() const override;
+	virtual FString GetTooltip() const override;
+	virtual FString GetDocumentationExcerptName() const override;
+	virtual FName GetPaletteIcon(FLinearColor& OutColor) const override{ return TEXT("GraphEditor.Timeline_16x"); }
+	virtual bool ShouldShowNodeProperties() const override { return true; }
 	// End UEdGraphNode interface.
 
 	// Begin UK2Node interface.
-	virtual bool NodeCausesStructuralBlueprintChange() const OVERRIDE { return true; }
-	virtual class FNodeHandlingFunctor* CreateNodeHandler(class FKismetCompilerContext& CompilerContext) const OVERRIDE;
-	virtual void ExpandNode(class FKismetCompilerContext& CompilerContext, UEdGraph* SourceGraph) OVERRIDE;
-	virtual void GetNodeAttributes( TArray<TKeyValuePair<FString, FString>>& OutNodeAttributes ) const OVERRIDE;
+	virtual bool NodeCausesStructuralBlueprintChange() const override { return true; }
+	virtual class FNodeHandlingFunctor* CreateNodeHandler(class FKismetCompilerContext& CompilerContext) const override;
+	virtual void ExpandNode(class FKismetCompilerContext& CompilerContext, UEdGraph* SourceGraph) override;
+	virtual void GetNodeAttributes( TArray<TKeyValuePair<FString, FString>>& OutNodeAttributes ) const override;
 	// End UK2Node interface.
 
 	/** Get the 'play' input pin */
diff --git a/Engine/Source/Editor/BlueprintGraph/Classes/K2Node_Tunnel.h b/Engine/Source/Editor/BlueprintGraph/Classes/K2Node_Tunnel.h
--- a/Engine/Source/Editor/BlueprintGraph/Classes/K2Node_Tunnel.h
+++ b/Engine/Source/Editor/BlueprintGraph/Classes/K2Node_Tunnel.h
@@ -34,27 +34,27 @@ class UK2Node_Tunnel : public UK2Node_EditablePinBase
 	struct FKismetUserDeclaredFunctionMetadata MetaData;
 
 	// Begin UEdGraphNode interface.
-	virtual void DestroyNode() OVERRIDE;
-	virtual void PostPasteNode() OVERRIDE;
-	virtual FString GetTooltip() const OVERRIDE;
-	virtual FText GetNodeTitle(ENodeTitleType::Type TitleType) const OVERRIDE;
-	virtual bool CanUserDeleteNode() const OVERRIDE;
-	BLUEPRINTGRAPH_API virtual bool CanDuplicateNode() const OVERRIDE;
-	virtual UObject* GetJumpTargetForDoubleClick() const OVERRIDE;
+	virtual void DestroyNode() override;
+	virtual void PostPasteNode() override;
+	virtual FString GetTooltip() const override;
+	virtual FText GetNodeTitle(ENodeTitleType::Type TitleType) const override;
+	virtual bool CanUserDeleteNode() const override;
+	BLUEPRINTGRAPH_API virtual bool CanDuplicateNode() const override;
+	virtual UObject* GetJumpTargetForDoubleClick() const override;
 	// End UEdGraphNode interface
 
 	// Begin UK2Node interface.
-	BLUEPRINTGRAPH_API virtual bool IsNodeSafeToIgnore() const OVERRIDE;
-	virtual bool DrawNodeAsEntry() const OVERRIDE;
-	virtual bool DrawNodeAsExit() const OVERRIDE;
-	virtual bool NodeCausesStructuralBlueprintChange() const OVERRIDE { return true; }
+	BLUEPRINTGRAPH_API virtual bool IsNodeSafeToIgnore() const override;
+	virtual bool DrawNodeAsEntry() const override;
+	virtual bool DrawNodeAsExit() const override;
+	virtual bool NodeCausesStructuralBlueprintChange() const override { return true; }
 	// End UK2Node interface
 
 	// Begin UK2Node_EditablePinBase interface.
-	BLUEPRINTGRAPH_API virtual UEdGraphPin* CreatePinFromUserDefinition(const TSharedPtr<FUserPinInfo> NewPinInfo) OVERRIDE;
-	BLUEPRINTGRAPH_API virtual bool CanModifyExecutionWires() OVERRIDE;
-	BLUEPRINTGRAPH_API virtual ERenamePinResult RenameUserDefinedPin(const FString& OldName, const FString& NewName, bool bTest = false) OVERRIDE;
-	virtual bool CanUseRefParams() const OVERRIDE { return true; }
+	BLUEPRINTGRAPH_API virtual UEdGraphPin* CreatePinFromUserDefinition(const TSharedPtr<FUserPinInfo> NewPinInfo) override;
+	BLUEPRINTGRAPH_API virtual bool CanModifyExecutionWires() override;
+	BLUEPRINTGRAPH_API virtual ERenamePinResult RenameUserDefinedPin(const FString& OldName, const FString& NewName, bool bTest = false) override;
+	virtual bool CanUseRefParams() const override { return true; }
 	// End UK2Node_EditablePinBase interface
 public:
 	// The input pins of this tunnel go to the output pins of InputSinkNode (can be NULL).
diff --git a/Engine/Source/Editor/BlueprintGraph/Classes/K2Node_Variable.h b/Engine/Source/Editor/BlueprintGraph/Classes/K2Node_Variable.h
--- a/Engine/Source/Editor/BlueprintGraph/Classes/K2Node_Variable.h
+++ b/Engine/Source/Editor/BlueprintGraph/Classes/K2Node_Variable.h
@@ -42,25 +42,25 @@ protected:
 
 public:
 	// Begin UObject interface
-	virtual void Serialize(FArchive& Ar) OVERRIDE;
+	virtual void Serialize(FArchive& Ar) override;
 	// End UObject interface
 
 	// Begin UEdGraphNode interface
-	virtual FLinearColor GetNodeTitleColor() const OVERRIDE;
-	virtual void ReconstructNode() OVERRIDE;
-	virtual FString GetDocumentationLink() const OVERRIDE;
-	virtual FString GetDocumentationExcerptName() const OVERRIDE;
-	virtual FName GetCornerIcon() const OVERRIDE;
-	virtual FName GetPaletteIcon(FLinearColor& OutColor) const OVERRIDE;
+	virtual FLinearColor GetNodeTitleColor() const override;
+	virtual void ReconstructNode() override;
+	virtual FString GetDocumentationLink() const override;
+	virtual FString GetDocumentationExcerptName() const override;
+	virtual FName GetCornerIcon() const override;
+	virtual FName GetPaletteIcon(FLinearColor& OutColor) const override;
 	// End UEdGraphNode interface
 
 	// Begin K2Node interface
-	virtual bool HasExternalBlueprintDependencies(TArray<class UStruct*>* OptionalOutput) const OVERRIDE;
-	virtual bool DrawNodeAsVariable() const OVERRIDE { return true; }
-	virtual ERedirectType DoPinsMatchForReconstruction(const UEdGraphPin* NewPin, int32 NewPinIndex, const UEdGraphPin* OldPin, int32 OldPinIndex)  const OVERRIDE;
-	virtual void ValidateNodeDuringCompilation(class FCompilerResultsLog& MessageLog) const OVERRIDE;
-	virtual FText GetToolTipHeading() const OVERRIDE;
-	virtual void GetNodeAttributes( TArray<TKeyValuePair<FString, FString>>& OutNodeAttributes ) const OVERRIDE;
+	virtual bool HasExternalBlueprintDependencies(TArray<class UStruct*>* OptionalOutput) const override;
+	virtual bool DrawNodeAsVariable() const override { return true; }
+	virtual ERedirectType DoPinsMatchForReconstruction(const UEdGraphPin* NewPin, int32 NewPinIndex, const UEdGraphPin* OldPin, int32 OldPinIndex)  const override;
+	virtual void ValidateNodeDuringCompilation(class FCompilerResultsLog& MessageLog) const override;
+	virtual FText GetToolTipHeading() const override;
+	virtual void GetNodeAttributes( TArray<TKeyValuePair<FString, FString>>& OutNodeAttributes ) const override;
 	// End K2Node interface
 
 	/** Set up this variable node from the supplied UProperty */
diff --git a/Engine/Source/Editor/BlueprintGraph/Classes/K2Node_VariableGet.h b/Engine/Source/Editor/BlueprintGraph/Classes/K2Node_VariableGet.h
--- a/Engine/Source/Editor/BlueprintGraph/Classes/K2Node_VariableGet.h
+++ b/Engine/Source/Editor/BlueprintGraph/Classes/K2Node_VariableGet.h
@@ -11,17 +11,17 @@ class UK2Node_VariableGet : public UK2Node_Variable
 	GENERATED_UCLASS_BODY()
 
 	// Begin UEdGraphNode interface
-	virtual void AllocateDefaultPins() OVERRIDE;
-	virtual FString GetTooltip() const OVERRIDE;
-	virtual FText GetNodeTitle(ENodeTitleType::Type TitleType) const OVERRIDE;
-	virtual FString GetNodeNativeTitle(ENodeTitleType::Type TitleType) const OVERRIDE;
+	virtual void AllocateDefaultPins() override;
+	virtual FString GetTooltip() const override;
+	virtual FText GetNodeTitle(ENodeTitleType::Type TitleType) const override;
+	virtual FString GetNodeNativeTitle(ENodeTitleType::Type TitleType) const override;
 	// End UEdGraphNode interface
 
 	// Begin K2Node interface
-	virtual bool IsNodePure() const OVERRIDE { return true; }
-	virtual void ReallocatePinsDuringReconstruction(TArray<UEdGraphPin*>& OldPins) OVERRIDE;
-	virtual bool ShouldShowNodeProperties() const OVERRIDE { return true; }
-	virtual class FNodeHandlingFunctor* CreateNodeHandler(class FKismetCompilerContext& CompilerContext) const OVERRIDE;
+	virtual bool IsNodePure() const override { return true; }
+	virtual void ReallocatePinsDuringReconstruction(TArray<UEdGraphPin*>& OldPins) override;
+	virtual bool ShouldShowNodeProperties() const override { return true; }
+	virtual class FNodeHandlingFunctor* CreateNodeHandler(class FKismetCompilerContext& CompilerContext) const override;
 	// End K2Node interface
 };
 
diff --git a/Engine/Source/Editor/BlueprintGraph/Classes/K2Node_VariableSet.h b/Engine/Source/Editor/BlueprintGraph/Classes/K2Node_VariableSet.h
--- a/Engine/Source/Editor/BlueprintGraph/Classes/K2Node_VariableSet.h
+++ b/Engine/Source/Editor/BlueprintGraph/Classes/K2Node_VariableSet.h
@@ -11,16 +11,16 @@ class UK2Node_VariableSet : public UK2Node_Variable
 	GENERATED_UCLASS_BODY()
 
 	// Begin UEdGraphNode interface
-	virtual void AllocateDefaultPins() OVERRIDE;
-	virtual void ReallocatePinsDuringReconstruction(TArray<UEdGraphPin*>& OldPins) OVERRIDE;
-	virtual FString GetTooltip() const OVERRIDE;
-	virtual FText GetNodeTitle(ENodeTitleType::Type TitleType) const OVERRIDE;
-	virtual FString GetNodeNativeTitle(ENodeTitleType::Type TitleType) const OVERRIDE;
+	virtual void AllocateDefaultPins() override;
+	virtual void ReallocatePinsDuringReconstruction(TArray<UEdGraphPin*>& OldPins) override;
+	virtual FString GetTooltip() const override;
+	virtual FText GetNodeTitle(ENodeTitleType::Type TitleType) const override;
+	virtual FString GetNodeNativeTitle(ENodeTitleType::Type TitleType) const override;
 	// End UEdGraphNode interface
 
 	// Begin K2Node interface
-	virtual bool ShouldShowNodeProperties() const OVERRIDE { return true; }
-	virtual class FNodeHandlingFunctor* CreateNodeHandler(class FKismetCompilerContext& CompilerContext) const OVERRIDE;
+	virtual bool ShouldShowNodeProperties() const override { return true; }
+	virtual class FNodeHandlingFunctor* CreateNodeHandler(class FKismetCompilerContext& CompilerContext) const override;
 	// End K2Node interface
 
 
diff --git a/Engine/Source/Editor/BlueprintGraph/Classes/K2Node_VariableSetRef.h b/Engine/Source/Editor/BlueprintGraph/Classes/K2Node_VariableSetRef.h
--- a/Engine/Source/Editor/BlueprintGraph/Classes/K2Node_VariableSetRef.h
+++ b/Engine/Source/Editor/BlueprintGraph/Classes/K2Node_VariableSetRef.h
@@ -10,16 +10,16 @@ class UK2Node_VariableSetRef : public UK2Node
 	GENERATED_UCLASS_BODY()
 
 	// Begin UEdGraphNode interface
-	virtual void AllocateDefaultPins() OVERRIDE;
-	virtual void ReallocatePinsDuringReconstruction(TArray<UEdGraphPin*>& OldPins) OVERRIDE;
-	virtual FString GetTooltip() const OVERRIDE;
-	virtual FText GetNodeTitle(ENodeTitleType::Type TitleType) const OVERRIDE;
-	virtual FString GetNodeNativeTitle(ENodeTitleType::Type TitleType) const OVERRIDE;
+	virtual void AllocateDefaultPins() override;
+	virtual void ReallocatePinsDuringReconstruction(TArray<UEdGraphPin*>& OldPins) override;
+	virtual FString GetTooltip() const override;
+	virtual FText GetNodeTitle(ENodeTitleType::Type TitleType) const override;
+	virtual FString GetNodeNativeTitle(ENodeTitleType::Type TitleType) const override;
 	// End UEdGraphNode interface
 
 	// Begin UK2Node interface
-	virtual void NotifyPinConnectionListChanged(UEdGraphPin* Pin) OVERRIDE;
-	virtual class FNodeHandlingFunctor* CreateNodeHandler(class FKismetCompilerContext& CompilerContext) const OVERRIDE;
+	virtual void NotifyPinConnectionListChanged(UEdGraphPin* Pin) override;
+	virtual class FNodeHandlingFunctor* CreateNodeHandler(class FKismetCompilerContext& CompilerContext) const override;
 	// End UK2Node interface
 
 	/** 
diff --git a/Engine/Source/Editor/BlueprintGraph/Private/CallFunctionHandler.cpp b/Engine/Source/Editor/BlueprintGraph/Private/CallFunctionHandler.cpp
--- a/Engine/Source/Editor/BlueprintGraph/Private/CallFunctionHandler.cpp
+++ b/Engine/Source/Editor/BlueprintGraph/Private/CallFunctionHandler.cpp
@@ -26,7 +26,7 @@ public:
 	{
 	}
 
-	virtual void Serialize(const TCHAR* V, ELogVerbosity::Type Verbosity, const class FName& Category) OVERRIDE
+	virtual void Serialize(const TCHAR* V, ELogVerbosity::Type Verbosity, const class FName& Category) override
 	{
 		if (TargetObject == NULL)
 		{
diff --git a/Engine/Source/Editor/BlueprintGraph/Private/DelegateNodeHandlers.h b/Engine/Source/Editor/BlueprintGraph/Private/DelegateNodeHandlers.h
--- a/Engine/Source/Editor/BlueprintGraph/Private/DelegateNodeHandlers.h
+++ b/Engine/Source/Editor/BlueprintGraph/Private/DelegateNodeHandlers.h
@@ -46,8 +46,8 @@ public:
 	{
 	}
 
-	virtual void RegisterNets(FKismetFunctionContext& Context, UEdGraphNode* Node) OVERRIDE;
-	virtual void Compile(FKismetFunctionContext& Context, UEdGraphNode* Node) OVERRIDE;
+	virtual void RegisterNets(FKismetFunctionContext& Context, UEdGraphNode* Node) override;
+	virtual void Compile(FKismetFunctionContext& Context, UEdGraphNode* Node) override;
 };
 
 //////////////////////////////////////////////////////////////////////////
@@ -61,8 +61,8 @@ public:
 	{
 	}
 
-	virtual void RegisterNets(FKismetFunctionContext& Context, UEdGraphNode* Node) OVERRIDE;
-	virtual void Compile(FKismetFunctionContext& Context, UEdGraphNode* Node) OVERRIDE;
+	virtual void RegisterNets(FKismetFunctionContext& Context, UEdGraphNode* Node) override;
+	virtual void Compile(FKismetFunctionContext& Context, UEdGraphNode* Node) override;
 };
 
 //////////////////////////////////////////////////////////////////////////
@@ -78,8 +78,8 @@ public:
 	{
 	}
 
-	virtual void RegisterNets(FKismetFunctionContext& Context, UEdGraphNode* Node) OVERRIDE;
-	virtual void Compile(FKismetFunctionContext& Context, UEdGraphNode* Node) OVERRIDE;
+	virtual void RegisterNets(FKismetFunctionContext& Context, UEdGraphNode* Node) override;
+	virtual void Compile(FKismetFunctionContext& Context, UEdGraphNode* Node) override;
 };
 
 //////////////////////////////////////////////////////////////////////////
@@ -95,9 +95,9 @@ public:
 	{
 	}
 
-	virtual void RegisterNets(FKismetFunctionContext& Context, UEdGraphNode* Node) OVERRIDE;
-	virtual void Compile(FKismetFunctionContext& Context, UEdGraphNode* Node) OVERRIDE;
-	virtual UFunction* FindFunction(FKismetFunctionContext& Context, UEdGraphNode* Node) OVERRIDE;
-	virtual void CheckIfFunctionIsCallable(UFunction* Function, FKismetFunctionContext& Context, UEdGraphNode* Node) OVERRIDE { }
-	virtual void AdditionalCompiledStatementHandling(FKismetFunctionContext& Context, UEdGraphNode* Node, FBlueprintCompiledStatement& Statement) OVERRIDE;
+	virtual void RegisterNets(FKismetFunctionContext& Context, UEdGraphNode* Node) override;
+	virtual void Compile(FKismetFunctionContext& Context, UEdGraphNode* Node) override;
+	virtual UFunction* FindFunction(FKismetFunctionContext& Context, UEdGraphNode* Node) override;
+	virtual void CheckIfFunctionIsCallable(UFunction* Function, FKismetFunctionContext& Context, UEdGraphNode* Node) override { }
+	virtual void AdditionalCompiledStatementHandling(FKismetFunctionContext& Context, UEdGraphNode* Node, FBlueprintCompiledStatement& Statement) override;
 };
\ No newline at end of file
diff --git a/Engine/Source/Editor/BlueprintGraph/Private/DynamicCastHandler.h b/Engine/Source/Editor/BlueprintGraph/Private/DynamicCastHandler.h
--- a/Engine/Source/Editor/BlueprintGraph/Private/DynamicCastHandler.h
+++ b/Engine/Source/Editor/BlueprintGraph/Private/DynamicCastHandler.h
@@ -20,7 +20,7 @@ public:
 	{
 	}
 
-	virtual void RegisterNets(FKismetFunctionContext& Context, UEdGraphNode* Node) OVERRIDE;
-	virtual void RegisterNet(FKismetFunctionContext& Context, UEdGraphPin* Net) OVERRIDE;
-	virtual void Compile(FKismetFunctionContext& Context, UEdGraphNode* Node) OVERRIDE;
+	virtual void RegisterNets(FKismetFunctionContext& Context, UEdGraphNode* Node) override;
+	virtual void RegisterNet(FKismetFunctionContext& Context, UEdGraphPin* Net) override;
+	virtual void Compile(FKismetFunctionContext& Context, UEdGraphNode* Node) override;
 };
diff --git a/Engine/Source/Editor/BlueprintGraph/Private/K2Node_ActorBoundEvent.cpp b/Engine/Source/Editor/BlueprintGraph/Private/K2Node_ActorBoundEvent.cpp
--- a/Engine/Source/Editor/BlueprintGraph/Private/K2Node_ActorBoundEvent.cpp
+++ b/Engine/Source/Editor/BlueprintGraph/Private/K2Node_ActorBoundEvent.cpp
@@ -18,7 +18,7 @@ public:
 	{
 	}
 
-	virtual void Compile(FKismetFunctionContext& Context, UEdGraphNode* Node) OVERRIDE
+	virtual void Compile(FKismetFunctionContext& Context, UEdGraphNode* Node) override
 	{
 		// Check to make sure that the object the event is bound to is valid
 		const UK2Node_ActorBoundEvent* BoundEventNode = Cast<UK2Node_ActorBoundEvent>(Node);
diff --git a/Engine/Source/Editor/BlueprintGraph/Private/K2Node_AddComponent.cpp b/Engine/Source/Editor/BlueprintGraph/Private/K2Node_AddComponent.cpp
--- a/Engine/Source/Editor/BlueprintGraph/Private/K2Node_AddComponent.cpp
+++ b/Engine/Source/Editor/BlueprintGraph/Private/K2Node_AddComponent.cpp
@@ -21,7 +21,7 @@ public:
 	{
 	}
 
-	virtual void RegisterNets(FKismetFunctionContext& Context, UEdGraphNode* Node) OVERRIDE
+	virtual void RegisterNets(FKismetFunctionContext& Context, UEdGraphNode* Node) override
 	{
 		FKCHandler_CallFunction::RegisterNets(Context, Node);
 
diff --git a/Engine/Source/Editor/BlueprintGraph/Private/K2Node_AssignmentStatement.cpp b/Engine/Source/Editor/BlueprintGraph/Private/K2Node_AssignmentStatement.cpp
--- a/Engine/Source/Editor/BlueprintGraph/Private/K2Node_AssignmentStatement.cpp
+++ b/Engine/Source/Editor/BlueprintGraph/Private/K2Node_AssignmentStatement.cpp
@@ -18,7 +18,7 @@ public:
 	{
 	}
 
-	virtual void RegisterNets(FKismetFunctionContext& Context, UEdGraphNode* Node) OVERRIDE
+	virtual void RegisterNets(FKismetFunctionContext& Context, UEdGraphNode* Node) override
 	{
 		UEdGraphPin* VariablePin = Node->FindPin(TEXT("Variable"));
 		UEdGraphPin* ValuePin = Node->FindPin(TEXT("Value"));
@@ -37,7 +37,7 @@ public:
 	}
 
 
-	virtual void Compile(FKismetFunctionContext& Context, UEdGraphNode* Node) OVERRIDE
+	virtual void Compile(FKismetFunctionContext& Context, UEdGraphNode* Node) override
 	{
 		UEdGraphPin* VariablePin = Node->FindPin(TEXT("Variable"));
 		UEdGraphPin* ValuePin = Node->FindPin(TEXT("Value"));
diff --git a/Engine/Source/Editor/BlueprintGraph/Private/K2Node_BreakStruct.cpp b/Engine/Source/Editor/BlueprintGraph/Private/K2Node_BreakStruct.cpp
--- a/Engine/Source/Editor/BlueprintGraph/Private/K2Node_BreakStruct.cpp
+++ b/Engine/Source/Editor/BlueprintGraph/Private/K2Node_BreakStruct.cpp
@@ -94,7 +94,7 @@ public:
 		}
 	}
 
-	virtual void RegisterNets(FKismetFunctionContext& Context, UEdGraphNode* InNode) OVERRIDE
+	virtual void RegisterNets(FKismetFunctionContext& Context, UEdGraphNode* InNode) override
 	{
 		UK2Node_BreakStruct* Node = Cast<UK2Node_BreakStruct>(InNode);
 		check(NULL != Node);
diff --git a/Engine/Source/Editor/BlueprintGraph/Private/K2Node_CastByteToEnum.cpp b/Engine/Source/Editor/BlueprintGraph/Private/K2Node_CastByteToEnum.cpp
--- a/Engine/Source/Editor/BlueprintGraph/Private/K2Node_CastByteToEnum.cpp
+++ b/Engine/Source/Editor/BlueprintGraph/Private/K2Node_CastByteToEnum.cpp
@@ -103,7 +103,7 @@ public:
 	{
 	}
 
-	virtual void RegisterNets(FKismetFunctionContext& Context, UEdGraphNode* Node) OVERRIDE
+	virtual void RegisterNets(FKismetFunctionContext& Context, UEdGraphNode* Node) override
 	{
 		FNodeHandlingFunctor::RegisterNets(Context, Node); //handle literals
 
diff --git a/Engine/Source/Editor/BlueprintGraph/Private/K2Node_CustomEvent.cpp b/Engine/Source/Editor/BlueprintGraph/Private/K2Node_CustomEvent.cpp
--- a/Engine/Source/Editor/BlueprintGraph/Private/K2Node_CustomEvent.cpp
+++ b/Engine/Source/Editor/BlueprintGraph/Private/K2Node_CustomEvent.cpp
@@ -64,7 +64,7 @@ public:
 	}
 
 	// Begin INameValidatorInterface
-	virtual EValidatorResult IsValid(FString const& Name, bool bOriginal = false) OVERRIDE
+	virtual EValidatorResult IsValid(FString const& Name, bool bOriginal = false) override
 	{
 		EValidatorResult NameValidity = FKismetNameValidator::IsValid(Name, bOriginal);
 		if ((NameValidity == EValidatorResult::Ok) || (NameValidity == EValidatorResult::ExistingName))
diff --git a/Engine/Source/Editor/BlueprintGraph/Private/K2Node_DelegateSet.cpp b/Engine/Source/Editor/BlueprintGraph/Private/K2Node_DelegateSet.cpp
--- a/Engine/Source/Editor/BlueprintGraph/Private/K2Node_DelegateSet.cpp
+++ b/Engine/Source/Editor/BlueprintGraph/Private/K2Node_DelegateSet.cpp
@@ -20,7 +20,7 @@ public:
 	{
 	}
 
-	virtual void RegisterNets(FKismetFunctionContext& Context, UEdGraphNode* Node) OVERRIDE
+	virtual void RegisterNets(FKismetFunctionContext& Context, UEdGraphNode* Node) override
 	{
 		UK2Node_DelegateSet* DelegateNode = Cast<UK2Node_DelegateSet>(Node);
 		if( DelegateNode )
@@ -83,7 +83,7 @@ public:
 		}
 	}
 
-	virtual void Compile(FKismetFunctionContext& Context, UEdGraphNode* Node) OVERRIDE
+	virtual void Compile(FKismetFunctionContext& Context, UEdGraphNode* Node) override
 	{
 		const UK2Node_DelegateSet* DelegateNode = Cast<UK2Node_DelegateSet>(Node);
 		check(DelegateNode);
diff --git a/Engine/Source/Editor/BlueprintGraph/Private/K2Node_ExecutionSequence.cpp b/Engine/Source/Editor/BlueprintGraph/Private/K2Node_ExecutionSequence.cpp
--- a/Engine/Source/Editor/BlueprintGraph/Private/K2Node_ExecutionSequence.cpp
+++ b/Engine/Source/Editor/BlueprintGraph/Private/K2Node_ExecutionSequence.cpp
@@ -18,7 +18,7 @@ public:
 	{
 	}
 
-	virtual void Compile(FKismetFunctionContext& Context, UEdGraphNode* Node) OVERRIDE
+	virtual void Compile(FKismetFunctionContext& Context, UEdGraphNode* Node) override
 	{
 		// Make sure that the input pin is connected and valid for this block
 		UEdGraphPin* ExecTriggeringPin = Context.FindRequiredPinByName(Node, CompilerContext.GetSchema()->PN_Execute, EGPD_Input);
diff --git a/Engine/Source/Editor/BlueprintGraph/Private/K2Node_FunctionResult.cpp b/Engine/Source/Editor/BlueprintGraph/Private/K2Node_FunctionResult.cpp
--- a/Engine/Source/Editor/BlueprintGraph/Private/K2Node_FunctionResult.cpp
+++ b/Engine/Source/Editor/BlueprintGraph/Private/K2Node_FunctionResult.cpp
@@ -59,7 +59,7 @@ public:
 	{
 	}
 
-	virtual void RegisterNet(FKismetFunctionContext& Context, UEdGraphPin* Net) OVERRIDE
+	virtual void RegisterNet(FKismetFunctionContext& Context, UEdGraphPin* Net) override
 	{
 		FBPTerminal* Term = new (Context.Results) FBPTerminal();
 		Term->CopyFromPin(Net, Net->PinName);
diff --git a/Engine/Source/Editor/BlueprintGraph/Private/K2Node_Literal.cpp b/Engine/Source/Editor/BlueprintGraph/Private/K2Node_Literal.cpp
--- a/Engine/Source/Editor/BlueprintGraph/Private/K2Node_Literal.cpp
+++ b/Engine/Source/Editor/BlueprintGraph/Private/K2Node_Literal.cpp
@@ -15,7 +15,7 @@ public:
 	{
 	}
 
-	virtual void RegisterNet(FKismetFunctionContext& Context, UEdGraphPin* Net) OVERRIDE
+	virtual void RegisterNet(FKismetFunctionContext& Context, UEdGraphPin* Net) override
 	{
 		const UEdGraphSchema_K2* Schema = GetDefault<UEdGraphSchema_K2>();
 		UK2Node_Literal* LiteralNode = Cast<UK2Node_Literal>(Net->GetOwningNode());
diff --git a/Engine/Source/Editor/BlueprintGraph/Private/K2Node_MakeStruct.cpp b/Engine/Source/Editor/BlueprintGraph/Private/K2Node_MakeStruct.cpp
--- a/Engine/Source/Editor/BlueprintGraph/Private/K2Node_MakeStruct.cpp
+++ b/Engine/Source/Editor/BlueprintGraph/Private/K2Node_MakeStruct.cpp
@@ -34,7 +34,7 @@ public:
 	{
 	}
 
-	virtual void RegisterNets(FKismetFunctionContext& Context, UEdGraphNode* InNode) OVERRIDE
+	virtual void RegisterNets(FKismetFunctionContext& Context, UEdGraphNode* InNode) override
 	{
 		UK2Node_MakeStruct* Node = CastChecked<UK2Node_MakeStruct>(InNode);
 		if (NULL == Node->StructType)
@@ -64,14 +64,14 @@ public:
 		}
 	}
 
-	virtual void RegisterNet(FKismetFunctionContext& Context, UEdGraphPin* Net) OVERRIDE
+	virtual void RegisterNet(FKismetFunctionContext& Context, UEdGraphPin* Net) override
 	{
 		FBPTerminal* Term = new (Context.IsEventGraph() ? Context.EventGraphLocals : Context.Locals) FBPTerminal();
 		Term->CopyFromPin(Net, Context.NetNameMap->MakeValidName(Net));
 		Context.NetMap.Add(Net, Term);
 	}
 
-	virtual void Compile(FKismetFunctionContext& Context, UEdGraphNode* InNode) OVERRIDE
+	virtual void Compile(FKismetFunctionContext& Context, UEdGraphNode* InNode) override
 	{
 		UK2Node_MakeStruct* Node = CastChecked<UK2Node_MakeStruct>(InNode);
 		if (NULL == Node->StructType)
diff --git a/Engine/Source/Editor/BlueprintGraph/Private/K2Node_MathExpression.cpp b/Engine/Source/Editor/BlueprintGraph/Private/K2Node_MathExpression.cpp
--- a/Engine/Source/Editor/BlueprintGraph/Private/K2Node_MathExpression.cpp
+++ b/Engine/Source/Editor/BlueprintGraph/Private/K2Node_MathExpression.cpp
@@ -215,13 +215,13 @@ public:
 	 *
 	 * @return True to continue traversing the tree, false to abort. 
 	 */
-	virtual bool Accept(FExpressionVisitor& Visitor) OVERRIDE
+	virtual bool Accept(FExpressionVisitor& Visitor) override
 	{
 		return Visitor.Visit(*this, FExpressionVisitor::VISIT_Leaf);
 	}
 
 	/** For debug purposes, constructs a textual representation of this expression */
-	virtual FString ToString() const OVERRIDE
+	virtual FString ToString() const override
 	{
 		if (Token.TokenType == FBasicToken::TOKEN_Identifier)
 		{
@@ -267,7 +267,7 @@ public:
 	 *
 	 * @return True to continue traversing the tree, false to abort. 
 	 */
-	virtual bool Accept(FExpressionVisitor& Visitor) OVERRIDE
+	virtual bool Accept(FExpressionVisitor& Visitor) override
 	{
 		bool bAbort = !Visitor.Visit(*this, FExpressionVisitor::VISIT_Pre);
 		if (bAbort || !LHS->Accept(Visitor) || !RHS->Accept(Visitor))
@@ -278,7 +278,7 @@ public:
 	}
 
 	/** For debug purposes, constructs a textual representation of this expression */
-	virtual FString ToString() const OVERRIDE
+	virtual FString ToString() const override
 	{
 		const FString LeftStr = LHS->ToString();
 		const FString RightStr = RHS->ToString();
@@ -315,7 +315,7 @@ public:
 	 *
 	 * @return True to continue traversing the tree, false to abort. 
 	 */
-	virtual bool Accept(FExpressionVisitor& Visitor) OVERRIDE
+	virtual bool Accept(FExpressionVisitor& Visitor) override
 	{
 		bool bAbort = !Visitor.Visit(*this, FExpressionVisitor::VISIT_Pre);
 		if (bAbort || !RHS->Accept(Visitor))
@@ -326,7 +326,7 @@ public:
 	}
 
 	/** For debug purposes, constructs a textual representation of this expression */
-	virtual FString ToString() const OVERRIDE
+	virtual FString ToString() const override
 	{
 		const FString RightStr = RHS->ToString();
 		return FString::Printf(TEXT("(%s%s)"), *Operator, *RightStr);
@@ -363,7 +363,7 @@ public:
 	 *
 	 * @return True to continue traversing the tree, false to abort.  
 	 */
-	virtual bool Accept(FExpressionVisitor& Visitor) OVERRIDE
+	virtual bool Accept(FExpressionVisitor& Visitor) override
 	{
 		bool bAbort = !Visitor.Visit(*this, FExpressionVisitor::VISIT_Pre);
 		// @TODO: what about the Condition?
@@ -375,7 +375,7 @@ public:
 	}
 
 	/** For debug purposes, constructs a textual representation of this expression */
-	virtual FString ToString() const OVERRIDE
+	virtual FString ToString() const override
 	{
 		const FString ConditionStr = Condition->ToString();
 		const FString TrueStr = TruePart->ToString();
@@ -408,7 +408,7 @@ public:
 	 *
 	 * @return True to continue traversing the tree, false to abort. 
 	 */
-	virtual bool Accept(FExpressionVisitor& Visitor) OVERRIDE
+	virtual bool Accept(FExpressionVisitor& Visitor) override
 	{
 		bool bAbort = !Visitor.Visit(*this, FExpressionVisitor::VISIT_Pre);
 		for (TSharedRef<IFExpressionNode> Child : Children)
@@ -422,7 +422,7 @@ public:
 	}
     
 	/** For debug purposes, constructs a textual representation of this expression */
-	virtual FString ToString() const OVERRIDE
+	virtual FString ToString() const override
 	{
 		FString AsString("(");
 		for (TSharedRef<IFExpressionNode> Child : Children)
@@ -464,7 +464,7 @@ public:
 	 *
 	 * @return True to continue traversing the tree, false to abort. 
 	 */
-	virtual bool Accept(FExpressionVisitor& Visitor) OVERRIDE
+	virtual bool Accept(FExpressionVisitor& Visitor) override
 	{
 		bool bAbort = !Visitor.Visit(*this, FExpressionVisitor::VISIT_Pre);
 		if (bAbort || !ParamList->Accept(Visitor))
@@ -475,7 +475,7 @@ public:
 	}
 
 	/** For debug purposes, constructs a textual representation of this expression */
-	virtual FString ToString() const OVERRIDE
+	virtual FString ToString() const override
 	{
 		FString const ParamsString = ParamList->ToString();
 		return FString::Printf(TEXT("(%s%s)"), *FuncName, *ParamsString);
@@ -550,7 +550,7 @@ private:
 	 *
 	 * @return True to continue traversing the tree, false to abort. 
 	 */
-	virtual bool VisitUnhandled(class IFExpressionNode& Node, EVisitPhase Phase) OVERRIDE
+	virtual bool VisitUnhandled(class IFExpressionNode& Node, EVisitPhase Phase) override
 	{
 		if (Phase == FExpressionVisitor::VISIT_Pre)
 		{
@@ -1012,7 +1012,7 @@ public:
 	}
 
 	/// Begin FCodeGenFragment Interface
-	virtual bool ConnectToInput(UEdGraphPin* InputPin, FCompilerResultsLog& MessageLog) OVERRIDE
+	virtual bool ConnectToInput(UEdGraphPin* InputPin, FCompilerResultsLog& MessageLog) override
 	{
 		bool bSuccess = false;
 		if (UEdGraphPin* VariablePin = GeneratedNode->FindPin(GeneratedNode->VariableReference.GetMemberName().ToString()))
@@ -1049,7 +1049,7 @@ public:
 	}
 
 	/// Begin FCodeGenFragment Interface
-	virtual bool ConnectToInput(UEdGraphPin* InputPin, FCompilerResultsLog& MessageLog) OVERRIDE
+	virtual bool ConnectToInput(UEdGraphPin* InputPin, FCompilerResultsLog& MessageLog) override
 	{
 		bool bSuccess = false;
 		if (UEdGraphPin* ResultPin = GeneratedNode->GetReturnValuePin())
@@ -1083,7 +1083,7 @@ public:
 	{}
 
 	/// Begin FCodeGenFragment Interface
-	virtual bool ConnectToInput(UEdGraphPin* InputPin, FCompilerResultsLog& MessageLog) OVERRIDE
+	virtual bool ConnectToInput(UEdGraphPin* InputPin, FCompilerResultsLog& MessageLog) override
 	{
 		UEdGraphSchema_K2 const* K2Schema = Cast<UEdGraphSchema_K2>(InputPin->GetSchema());
 		bool bSuccess = true;//K2Schema->ArePinTypesCompatible(GetOutputType(), InputPin->PinType);
@@ -1120,7 +1120,7 @@ public:
 	}
 
 	/// Begin FCodeGenFragment Interface
-	virtual bool ConnectToInput(UEdGraphPin* InputPin, FCompilerResultsLog& MessageLog) OVERRIDE
+	virtual bool ConnectToInput(UEdGraphPin* InputPin, FCompilerResultsLog& MessageLog) override
 	{
 		return SafeConnectPins(TunnelInputPin, InputPin, MessageLog);
 	}
@@ -1229,7 +1229,7 @@ public:
 	 *
 	 * @return True to continue travesing the expression tree, false to stop.
 	 */
-	virtual bool Visit(FTokenWrapperNode& ExpressionNode, EVisitPhase Phase) OVERRIDE
+	virtual bool Visit(FTokenWrapperNode& ExpressionNode, EVisitPhase Phase) override
 	{
 		check(ActiveMessageLog != nullptr);
 		const UEdGraphSchema_K2* K2Schema = GetDefault<UEdGraphSchema_K2>();
@@ -1278,7 +1278,7 @@ public:
 	 *
 	 * @return True to continue traversing the expression tree, false to stop.
 	 */
-	virtual bool Visit(FBinaryOperator& ExpressionNode, EVisitPhase Phase) OVERRIDE
+	virtual bool Visit(FBinaryOperator& ExpressionNode, EVisitPhase Phase) override
 	{
 		check(ActiveMessageLog != nullptr);
 		
@@ -1312,7 +1312,7 @@ public:
 	 *
 	 * @return Always true, it is expected that cascading errors are handled (and all should be logged).
 	 */
-	virtual bool Visit(FExpressionList& ExpressionNode, EVisitPhase Phase) OVERRIDE
+	virtual bool Visit(FExpressionList& ExpressionNode, EVisitPhase Phase) override
 	{
 		// no fragments are generated from a list node, it mostly acts as a link
 		// from a parent node to some set of sub-expressions
@@ -1330,7 +1330,7 @@ public:
 	 *
 	 * @return True to continue traversing the expression tree, false to stop.
 	 */
-	virtual bool Visit(FFunctionExpression& ExpressionNode, EVisitPhase Phase) OVERRIDE
+	virtual bool Visit(FFunctionExpression& ExpressionNode, EVisitPhase Phase) override
 	{
 		check(ActiveMessageLog != nullptr);
 		
@@ -1364,7 +1364,7 @@ public:
 	 *
 	 * @return Always true, it is expected that cascading errors are handled (and all should be logged).
 	 */
-	virtual bool Visit(FUnaryOperator& ExpressionNode, EVisitPhase Phase) OVERRIDE
+	virtual bool Visit(FUnaryOperator& ExpressionNode, EVisitPhase Phase) override
 	{
 		// don't want to double up on the error message (in the "Post" phase)
 		if (Phase == VISIT_Pre)
@@ -1387,7 +1387,7 @@ public:
 	 *
 	 * @return Always true, it is expected that cascading errors are handled (and all should be logged).
 	 */
-	virtual bool Visit(FConditionalOperator& ExpressionNode, EVisitPhase Phase) OVERRIDE
+	virtual bool Visit(FConditionalOperator& ExpressionNode, EVisitPhase Phase) override
 	{
 		check(ActiveMessageLog != nullptr);
 		
@@ -1413,7 +1413,7 @@ private:
 	 *
 	 * @return Always true, it is expected that cascading errors can be handled (and all should be logged).
 	 */
-	virtual bool VisitUnhandled(IFExpressionNode& ExpressionNode, EVisitPhase Phase) OVERRIDE
+	virtual bool VisitUnhandled(IFExpressionNode& ExpressionNode, EVisitPhase Phase) override
 	{
 		check(ActiveMessageLog != nullptr);
 		if (Phase == VISIT_Leaf || Phase == VISIT_Pre)
diff --git a/Engine/Source/Editor/BlueprintGraph/Private/K2Node_StructMemberGet.cpp b/Engine/Source/Editor/BlueprintGraph/Private/K2Node_StructMemberGet.cpp
--- a/Engine/Source/Editor/BlueprintGraph/Private/K2Node_StructMemberGet.cpp
+++ b/Engine/Source/Editor/BlueprintGraph/Private/K2Node_StructMemberGet.cpp
@@ -52,7 +52,7 @@ void UK2Node_StructMemberGet::AllocatePinsForSingleMemberGet(FName MemberName)
 		}
 
 		// FOptionalPinsUpdater interface
-		virtual void GetRecordDefaults(UProperty* TestProperty, FOptionalPinFromProperty& Record) const OVERRIDE
+		virtual void GetRecordDefaults(UProperty* TestProperty, FOptionalPinFromProperty& Record) const override
 		{
 			Record.bCanToggleVisibility = false;
 			Record.bShowPin = TestProperty->GetFName() == MatchName;
diff --git a/Engine/Source/Editor/BlueprintGraph/Private/K2Node_VariableGet.cpp b/Engine/Source/Editor/BlueprintGraph/Private/K2Node_VariableGet.cpp
--- a/Engine/Source/Editor/BlueprintGraph/Private/K2Node_VariableGet.cpp
+++ b/Engine/Source/Editor/BlueprintGraph/Private/K2Node_VariableGet.cpp
@@ -17,13 +17,13 @@ public:
 	{
 	}
 
-	virtual void RegisterNet(FKismetFunctionContext& Context, UEdGraphPin* Net) OVERRIDE
+	virtual void RegisterNet(FKismetFunctionContext& Context, UEdGraphPin* Net) override
 	{
 		// This net is a variable read
 		ResolveAndRegisterScopedTerm(Context, Net, Context.VariableReferences);
 	}
 
-	virtual void RegisterNets(FKismetFunctionContext& Context, UEdGraphNode* Node) OVERRIDE
+	virtual void RegisterNets(FKismetFunctionContext& Context, UEdGraphNode* Node) override
 	{
 		UK2Node_Variable* VarNode = Cast<UK2Node_Variable>(Node);
 		if (VarNode)
diff --git a/Engine/Source/Editor/BlueprintGraph/Private/StructMemberNodeHandlers.h b/Engine/Source/Editor/BlueprintGraph/Private/StructMemberNodeHandlers.h
--- a/Engine/Source/Editor/BlueprintGraph/Private/StructMemberNodeHandlers.h
+++ b/Engine/Source/Editor/BlueprintGraph/Private/StructMemberNodeHandlers.h
@@ -14,8 +14,8 @@ public:
 	{
 	}
 
-	virtual void RegisterNet(FKismetFunctionContext& Context, UEdGraphPin* Net) OVERRIDE;
-	virtual void RegisterNets(FKismetFunctionContext& Context, UEdGraphNode* InNode) OVERRIDE;
+	virtual void RegisterNet(FKismetFunctionContext& Context, UEdGraphPin* Net) override;
+	virtual void RegisterNets(FKismetFunctionContext& Context, UEdGraphNode* InNode) override;
 };
 
 class FKCHandler_StructMemberVariableSet : public FKCHandler_VariableSet
@@ -26,5 +26,5 @@ public:
 	{
 	}
 
-	virtual void RegisterNets(FKismetFunctionContext& Context, UEdGraphNode* InNode) OVERRIDE;
+	virtual void RegisterNets(FKismetFunctionContext& Context, UEdGraphNode* InNode) override;
 };
\ No newline at end of file
diff --git a/Engine/Source/Editor/BlueprintGraph/Public/CallFunctionHandler.h b/Engine/Source/Editor/BlueprintGraph/Public/CallFunctionHandler.h
--- a/Engine/Source/Editor/BlueprintGraph/Public/CallFunctionHandler.h
+++ b/Engine/Source/Editor/BlueprintGraph/Public/CallFunctionHandler.h
@@ -56,11 +56,11 @@ private:
 
 public:
 
-	virtual void RegisterNets(FKismetFunctionContext& Context, UEdGraphNode* Node) OVERRIDE;
-	virtual void RegisterNet(FKismetFunctionContext& Context, UEdGraphPin* Net) OVERRIDE;
+	virtual void RegisterNets(FKismetFunctionContext& Context, UEdGraphNode* Node) override;
+	virtual void RegisterNet(FKismetFunctionContext& Context, UEdGraphPin* Net) override;
 	virtual UFunction* FindFunction(FKismetFunctionContext& Context, UEdGraphNode* Node);
-	virtual void Transform(FKismetFunctionContext& Context, UEdGraphNode* Node) OVERRIDE;
-	virtual void Compile(FKismetFunctionContext& Context, UEdGraphNode* Node) OVERRIDE;
+	virtual void Transform(FKismetFunctionContext& Context, UEdGraphNode* Node) override;
+	virtual void Compile(FKismetFunctionContext& Context, UEdGraphNode* Node) override;
 	virtual void CheckIfFunctionIsCallable(UFunction* Function, FKismetFunctionContext& Context, UEdGraphNode* Node);
 	virtual void AdditionalCompiledStatementHandling(FKismetFunctionContext& Context, UEdGraphNode* Node, FBlueprintCompiledStatement& Statement) {}
 
diff --git a/Engine/Source/Editor/BlueprintGraph/Public/EventEntryHandler.h b/Engine/Source/Editor/BlueprintGraph/Public/EventEntryHandler.h
--- a/Engine/Source/Editor/BlueprintGraph/Public/EventEntryHandler.h
+++ b/Engine/Source/Editor/BlueprintGraph/Public/EventEntryHandler.h
@@ -15,6 +15,6 @@ public:
 	{
 	}
 
-	virtual void RegisterNet(FKismetFunctionContext& Context, UEdGraphPin* Net) OVERRIDE;
-	virtual void Compile(FKismetFunctionContext& Context, UEdGraphNode* Node) OVERRIDE;
+	virtual void RegisterNet(FKismetFunctionContext& Context, UEdGraphPin* Net) override;
+	virtual void Compile(FKismetFunctionContext& Context, UEdGraphNode* Node) override;
 };
diff --git a/Engine/Source/Editor/BlueprintGraph/Public/VariableSetHandler.h b/Engine/Source/Editor/BlueprintGraph/Public/VariableSetHandler.h
--- a/Engine/Source/Editor/BlueprintGraph/Public/VariableSetHandler.h
+++ b/Engine/Source/Editor/BlueprintGraph/Public/VariableSetHandler.h
@@ -13,9 +13,9 @@ public:
 	{
 	}
 
-	virtual void RegisterNet(FKismetFunctionContext& Context, UEdGraphPin* Net) OVERRIDE;
-	virtual void RegisterNets(FKismetFunctionContext& Context, UEdGraphNode* Node) OVERRIDE;
+	virtual void RegisterNet(FKismetFunctionContext& Context, UEdGraphPin* Net) override;
+	virtual void RegisterNets(FKismetFunctionContext& Context, UEdGraphNode* Node) override;
 	void InnerAssignment(FKismetFunctionContext& Context, UEdGraphNode* Node, UEdGraphPin* VariablePin, UEdGraphPin* ValuePin);
-	virtual void Compile(FKismetFunctionContext& Context, UEdGraphNode* Node) OVERRIDE;
-	virtual void Transform(FKismetFunctionContext& Context, UEdGraphNode* Node) OVERRIDE;
+	virtual void Compile(FKismetFunctionContext& Context, UEdGraphNode* Node) override;
+	virtual void Transform(FKismetFunctionContext& Context, UEdGraphNode* Node) override;
 };
diff --git a/Engine/Source/Editor/Blutility/Classes/EditorUtilityBlueprint.h b/Engine/Source/Editor/Blutility/Classes/EditorUtilityBlueprint.h
--- a/Engine/Source/Editor/Blutility/Classes/EditorUtilityBlueprint.h
+++ b/Engine/Source/Editor/Blutility/Classes/EditorUtilityBlueprint.h
@@ -15,7 +15,7 @@ class UEditorUtilityBlueprint : public UBlueprint
 
 #if WITH_EDITOR
 	// UBlueprint interface
-	virtual bool SupportedByDefaultBlueprintFactory() const OVERRIDE
+	virtual bool SupportedByDefaultBlueprintFactory() const override
 	{
 		return false;
 	}
diff --git a/Engine/Source/Editor/Blutility/Classes/EditorUtilityBlueprintFactory.h b/Engine/Source/Editor/Blutility/Classes/EditorUtilityBlueprintFactory.h
--- a/Engine/Source/Editor/Blutility/Classes/EditorUtilityBlueprintFactory.h
+++ b/Engine/Source/Editor/Blutility/Classes/EditorUtilityBlueprintFactory.h
@@ -18,8 +18,8 @@ class UEditorUtilityBlueprintFactory : public UFactory
 	TSubclassOf<class UObject> ParentClass;
 
 	// UFactory interface
-	virtual bool ConfigureProperties() OVERRIDE;
-	virtual UObject* FactoryCreateNew(UClass* Class, UObject* InParent, FName Name, EObjectFlags Flags, UObject* Context, FFeedbackContext* Warn) OVERRIDE;
+	virtual bool ConfigureProperties() override;
+	virtual UObject* FactoryCreateNew(UClass* Class, UObject* InParent, FName Name, EObjectFlags Flags, UObject* Context, FFeedbackContext* Warn) override;
 	// End of UFactory interface
 
 	/** Handler for when a class is picked in the class picker */
diff --git a/Engine/Source/Editor/Blutility/Classes/PlacedEditorUtilityBase.h b/Engine/Source/Editor/Blutility/Classes/PlacedEditorUtilityBase.h
--- a/Engine/Source/Editor/Blutility/Classes/PlacedEditorUtilityBase.h
+++ b/Engine/Source/Editor/Blutility/Classes/PlacedEditorUtilityBase.h
@@ -17,6 +17,6 @@ class BLUTILITY_API APlacedEditorUtilityBase : public AActor
 	FString HelpText;
 
 	// AActor interface
-	virtual void TickActor(float DeltaSeconds, ELevelTick TickType, FActorTickFunction& ThisTickFunction) OVERRIDE;
+	virtual void TickActor(float DeltaSeconds, ELevelTick TickType, FActorTickFunction& ThisTickFunction) override;
 	// End of AActor interface
 };
diff --git a/Engine/Source/Editor/Blutility/Private/AssetTypeActions_EditorUtilityBlueprint.h b/Engine/Source/Editor/Blutility/Private/AssetTypeActions_EditorUtilityBlueprint.h
--- a/Engine/Source/Editor/Blutility/Private/AssetTypeActions_EditorUtilityBlueprint.h
+++ b/Engine/Source/Editor/Blutility/Private/AssetTypeActions_EditorUtilityBlueprint.h
@@ -8,13 +8,13 @@ class FAssetTypeActions_EditorUtilityBlueprint : public FAssetTypeActions_Base
 {
 public:
 	// IAssetTypeActions interface
-	virtual FText GetName() const OVERRIDE;
-	virtual FColor GetTypeColor() const OVERRIDE;
-	virtual UClass* GetSupportedClass() const OVERRIDE;
-	virtual bool HasActions(const TArray<UObject*>& InObjects) const OVERRIDE;
-	virtual void GetActions(const TArray<UObject*>& InObjects, FMenuBuilder& MenuBuilder) OVERRIDE;
-	virtual void OpenAssetEditor(const TArray<UObject*>& InObjects, TSharedPtr<class IToolkitHost> EditWithinLevelEditor = TSharedPtr<IToolkitHost>()) OVERRIDE;
-	virtual uint32 GetCategories() OVERRIDE;
+	virtual FText GetName() const override;
+	virtual FColor GetTypeColor() const override;
+	virtual UClass* GetSupportedClass() const override;
+	virtual bool HasActions(const TArray<UObject*>& InObjects) const override;
+	virtual void GetActions(const TArray<UObject*>& InObjects, FMenuBuilder& MenuBuilder) override;
+	virtual void OpenAssetEditor(const TArray<UObject*>& InObjects, TSharedPtr<class IToolkitHost> EditWithinLevelEditor = TSharedPtr<IToolkitHost>()) override;
+	virtual uint32 GetCategories() override;
 	// End of IAssetTypeActions interface
 
 protected:
diff --git a/Engine/Source/Editor/Blutility/Private/EditorUtilityBlueprintFactory.cpp b/Engine/Source/Editor/Blutility/Private/EditorUtilityBlueprintFactory.cpp
--- a/Engine/Source/Editor/Blutility/Private/EditorUtilityBlueprintFactory.cpp
+++ b/Engine/Source/Editor/Blutility/Private/EditorUtilityBlueprintFactory.cpp
@@ -17,7 +17,7 @@ public:
 		return InFilterFuncs->IfInChildOfClassesSet(AllowedChildOfClasses, InClass) != EFilterReturn::Failed;
 	}
 
-	virtual bool IsUnloadedClassAllowed(const FClassViewerInitializationOptions& InInitOptions, const TSharedRef< const IUnloadedBlueprintData > InUnloadedClassData, TSharedRef< FClassViewerFilterFuncs > InFilterFuncs) OVERRIDE
+	virtual bool IsUnloadedClassAllowed(const FClassViewerInitializationOptions& InInitOptions, const TSharedRef< const IUnloadedBlueprintData > InUnloadedClassData, TSharedRef< FClassViewerFilterFuncs > InFilterFuncs) override
 	{
 		return InFilterFuncs->IfInChildOfClassesSet(AllowedChildOfClasses, InUnloadedClassData) != EFilterReturn::Failed;
 	}
diff --git a/Engine/Source/Editor/Blutility/Private/GlobalBlutilityDialog.h b/Engine/Source/Editor/Blutility/Private/GlobalBlutilityDialog.h
--- a/Engine/Source/Editor/Blutility/Private/GlobalBlutilityDialog.h
+++ b/Engine/Source/Editor/Blutility/Private/GlobalBlutilityDialog.h
@@ -8,8 +8,8 @@
 class FGlobalBlutilityDialog : public FAssetEditorToolkit, public FGCObject
 {
 public:
-	virtual void RegisterTabSpawners(const TSharedRef<class FTabManager>& TabManager) OVERRIDE;
-	virtual void UnregisterTabSpawners(const TSharedRef<class FTabManager>& TabManager) OVERRIDE;
+	virtual void RegisterTabSpawners(const TSharedRef<class FTabManager>& TabManager) override;
+	virtual void UnregisterTabSpawners(const TSharedRef<class FTabManager>& TabManager) override;
 	TSharedRef<SDockTab> SpawnTab_DetailsPanel( const FSpawnTabArgs& SpawnTabArgs );
 
 	/**
@@ -29,15 +29,15 @@ public:
 	virtual ~FGlobalBlutilityDialog();
 
 	/** FGCObject interface */
-	virtual void AddReferencedObjects( FReferenceCollector& Collector ) OVERRIDE;
+	virtual void AddReferencedObjects( FReferenceCollector& Collector ) override;
 
 	/** IToolkit interface */
-	virtual FName GetToolkitFName() const OVERRIDE;
-	virtual FText GetBaseToolkitName() const OVERRIDE;
-	virtual FString GetWorldCentricTabPrefix() const OVERRIDE;
+	virtual FName GetToolkitFName() const override;
+	virtual FText GetBaseToolkitName() const override;
+	virtual FString GetWorldCentricTabPrefix() const override;
 
 	/** @return Returns the color and opacity to use for the color that appears behind the tab text for this toolkit's tab in world-centric mode. */
-	virtual FLinearColor GetWorldCentricTabColorScale() const OVERRIDE;
+	virtual FLinearColor GetWorldCentricTabColorScale() const override;
 
 private:
 
diff --git a/Engine/Source/Editor/BspMode/Private/BspMode.h b/Engine/Source/Editor/BspMode/Private/BspMode.h
--- a/Engine/Source/Editor/BspMode/Private/BspMode.h
+++ b/Engine/Source/Editor/BspMode/Private/BspMode.h
@@ -10,12 +10,12 @@ public:
 	static TSharedRef< FBspMode > Create();
 
 	~FBspMode();
-	virtual bool UsesToolkits() const OVERRIDE;
+	virtual bool UsesToolkits() const override;
 
-	virtual void Enter() OVERRIDE;
-	virtual void Exit() OVERRIDE;
+	virtual void Enter() override;
+	virtual void Exit() override;
 
-	virtual bool UsesPropertyWidgets() const OVERRIDE { return true; }
+	virtual bool UsesPropertyWidgets() const override { return true; }
 
 private:
 	FBspMode();
diff --git a/Engine/Source/Editor/BspMode/Private/BspModeModule.h b/Engine/Source/Editor/BspMode/Private/BspModeModule.h
--- a/Engine/Source/Editor/BspMode/Private/BspModeModule.h
+++ b/Engine/Source/Editor/BspMode/Private/BspModeModule.h
@@ -32,14 +32,14 @@ class FBspModeModule : public IBspModeModule
 public:
 	
 	/** IModuleInterface interface */
-	virtual void StartupModule() OVERRIDE;
-	virtual void ShutdownModule() OVERRIDE;
+	virtual void StartupModule() override;
+	virtual void ShutdownModule() override;
 
 	/** IBspModeModule interface */
-	virtual TSharedRef< SWidget > CreateBspModeWidget() const OVERRIDE;
-	virtual TSharedRef< FEdMode > GetBspMode() const OVERRIDE;
-	virtual void RegisterBspBuilderType( class UClass* InBuilderClass, const FText& InBuilderName, const FText& InBuilderTooltip, const FSlateBrush* InBuilderIcon ) OVERRIDE;
-	virtual void UnregisterBspBuilderType( class UClass* InBuilderClass ) OVERRIDE;
+	virtual TSharedRef< SWidget > CreateBspModeWidget() const override;
+	virtual TSharedRef< FEdMode > GetBspMode() const override;
+	virtual void RegisterBspBuilderType( class UClass* InBuilderClass, const FText& InBuilderName, const FText& InBuilderTooltip, const FSlateBrush* InBuilderIcon ) override;
+	virtual void UnregisterBspBuilderType( class UClass* InBuilderClass ) override;
 
 	const TArray< TSharedPtr<FBspBuilderType> >& GetBspBuilderTypes();
 	TSharedPtr<FBspBuilderType> FindBspBuilderType(UClass* InBuilderClass) const;
diff --git a/Engine/Source/Editor/BspMode/Private/SBspPalette.cpp b/Engine/Source/Editor/BspMode/Private/SBspPalette.cpp
--- a/Engine/Source/Editor/BspMode/Private/SBspPalette.cpp
+++ b/Engine/Source/Editor/BspMode/Private/SBspPalette.cpp
@@ -12,8 +12,8 @@
 class SBspBuilderListView : public SListView<TSharedPtr<FBspBuilderType>>
 {
 public:
-	virtual bool SupportsKeyboardFocus() const OVERRIDE { return true; }
-	virtual FReply OnKeyDown( const FGeometry& InGeometry, const FKeyboardEvent& InKeyboardEvent ) OVERRIDE
+	virtual bool SupportsKeyboardFocus() const override { return true; }
+	virtual FReply OnKeyDown( const FGeometry& InGeometry, const FKeyboardEvent& InKeyboardEvent ) override
 	{
 		return FReply::Unhandled();
 	}
diff --git a/Engine/Source/Editor/Cascade/Classes/CascadeParticleSystemComponent.h b/Engine/Source/Editor/Cascade/Classes/CascadeParticleSystemComponent.h
--- a/Engine/Source/Editor/Cascade/Classes/CascadeParticleSystemComponent.h
+++ b/Engine/Source/Editor/Cascade/Classes/CascadeParticleSystemComponent.h
@@ -18,7 +18,7 @@ public:
 	class FCascadeEdPreviewViewportClient*	CascadePreviewViewportPtr;
 
 	// Begin UActorComponent Interface
-	virtual void TickComponent(float DeltaTime, enum ELevelTick TickType, FActorComponentTickFunction *ThisTickFunction) OVERRIDE;
+	virtual void TickComponent(float DeltaTime, enum ELevelTick TickType, FActorComponentTickFunction *ThisTickFunction) override;
 	// End UActorComponent Interface
 
 	/**
diff --git a/Engine/Source/Editor/Cascade/Private/Cascade.h b/Engine/Source/Editor/Cascade/Private/Cascade.h
--- a/Engine/Source/Editor/Cascade/Private/Cascade.h
+++ b/Engine/Source/Editor/Cascade/Private/Cascade.h
@@ -25,8 +25,8 @@ DECLARE_LOG_CATEGORY_EXTERN(LogCascade, Log, All);
 class FCascade : public ICascade, public FGCObject, public FTickableEditorObject, public FNotifyHook, public FCurveEdNotifyInterface, public FEditorUndoClient
 {
 public:
-	virtual void RegisterTabSpawners(const TSharedRef<class FTabManager>& TabManager) OVERRIDE;
-	virtual void UnregisterTabSpawners(const TSharedRef<class FTabManager>& TabManager) OVERRIDE;
+	virtual void RegisterTabSpawners(const TSharedRef<class FTabManager>& TabManager) override;
+	virtual void UnregisterTabSpawners(const TSharedRef<class FTabManager>& TabManager) override;
 	TSharedRef<SDockTab> SpawnTab(const FSpawnTabArgs& SpawnTabArgs, FName TabIdentifier);
 
 	/** Destructor */
@@ -150,38 +150,38 @@ public:
 	void OnRedo();
 
 	/** IToolkit interface */
-	virtual FName GetToolkitFName() const OVERRIDE;
-	virtual FText GetBaseToolkitName() const OVERRIDE;
-	virtual FString GetWorldCentricTabPrefix() const OVERRIDE;
-	virtual FLinearColor GetWorldCentricTabColorScale() const OVERRIDE;
+	virtual FName GetToolkitFName() const override;
+	virtual FText GetBaseToolkitName() const override;
+	virtual FString GetWorldCentricTabPrefix() const override;
+	virtual FLinearColor GetWorldCentricTabColorScale() const override;
 	
 	/** @return the documentation location for this editor */
-	virtual FString GetDocumentationLink() const OVERRIDE
+	virtual FString GetDocumentationLink() const override
 	{
 		return FString(TEXT("Engine/Rendering/ParticleSystems"));
 	}
 
 	/** FGCObject interface */
-	virtual void AddReferencedObjects(FReferenceCollector& Collector) OVERRIDE;
+	virtual void AddReferencedObjects(FReferenceCollector& Collector) override;
 	
 	/** FTickableEditorObject interface */
-	virtual void Tick(float DeltaTime) OVERRIDE;
-	virtual bool IsTickable() const OVERRIDE;
-	virtual TStatId GetStatId() const OVERRIDE;
+	virtual void Tick(float DeltaTime) override;
+	virtual bool IsTickable() const override;
+	virtual TStatId GetStatId() const override;
 
 
 
 	/** FNotifyHook interface */
-	virtual void NotifyPostChange( const FPropertyChangedEvent& PropertyChangedEvent, UProperty* PropertyThatChanged) OVERRIDE;
-	virtual void NotifyPreChange(FEditPropertyChain* PropertyChain) OVERRIDE;
-	virtual void NotifyPostChange( const FPropertyChangedEvent& PropertyChangedEvent, FEditPropertyChain* PropertyChain) OVERRIDE;
+	virtual void NotifyPostChange( const FPropertyChangedEvent& PropertyChangedEvent, UProperty* PropertyThatChanged) override;
+	virtual void NotifyPreChange(FEditPropertyChain* PropertyChain) override;
+	virtual void NotifyPostChange( const FPropertyChangedEvent& PropertyChangedEvent, FEditPropertyChain* PropertyChain) override;
 
 	/** FCurveEdNotifyInterface */
-	virtual void PreEditCurve(TArray<UObject*> CurvesAboutToChange) OVERRIDE;
-	virtual void PostEditCurve() OVERRIDE;
-	virtual void MovedKey() OVERRIDE;
-	virtual void DesireUndo() OVERRIDE;
-	virtual void DesireRedo() OVERRIDE;
+	virtual void PreEditCurve(TArray<UObject*> CurvesAboutToChange) override;
+	virtual void PostEditCurve() override;
+	virtual void MovedKey() override;
+	virtual void DesireUndo() override;
+	virtual void DesireRedo() override;
 
 	/**
 	 *	Convert all the modules in this particle system to their random seed variant if available
@@ -192,8 +192,8 @@ public:
 
 private:
 	// Begin FEditorUndoClient Interface
-	virtual void PostUndo(bool bSuccess) OVERRIDE;
-	virtual void PostRedo(bool bSuccess) OVERRIDE { PostUndo(bSuccess); }
+	virtual void PostUndo(bool bSuccess) override;
+	virtual void PostRedo(bool bSuccess) override { PostUndo(bSuccess); }
 	// End of FEditorUndoClient
 
 	/** Generates menu for the AnimSpeed sub menu */
diff --git a/Engine/Source/Editor/Cascade/Private/CascadeEmitterCanvasClient.h b/Engine/Source/Editor/Cascade/Private/CascadeEmitterCanvasClient.h
--- a/Engine/Source/Editor/Cascade/Private/CascadeEmitterCanvasClient.h
+++ b/Engine/Source/Editor/Cascade/Private/CascadeEmitterCanvasClient.h
@@ -19,9 +19,9 @@ public:
 	~FCascadeEmitterCanvasClient();
 
 	/** FViewportClient interface */
-	virtual void Draw(FViewport* Viewport, FCanvas* Canvas) OVERRIDE;
-	virtual bool InputKey(FViewport* Viewport, int32 ControllerId, FKey Key, EInputEvent Event, float AmountDepressed = 1.0f, bool bGamepad = false) OVERRIDE;
-	virtual void CapturedMouseMove(FViewport* Viewport, int32 X, int32 Y) OVERRIDE;
+	virtual void Draw(FViewport* Viewport, FCanvas* Canvas) override;
+	virtual bool InputKey(FViewport* Viewport, int32 ControllerId, FKey Key, EInputEvent Event, float AmountDepressed = 1.0f, bool bGamepad = false) override;
+	virtual void CapturedMouseMove(FViewport* Viewport, int32 X, int32 Y) override;
 
 	/** Returns the ratio of the size of the particle emitters to the size of the viewport */
 	float GetViewportVerticalScrollBarRatio() const;
diff --git a/Engine/Source/Editor/Cascade/Private/CascadeModule.cpp b/Engine/Source/Editor/Cascade/Private/CascadeModule.cpp
--- a/Engine/Source/Editor/Cascade/Private/CascadeModule.cpp
+++ b/Engine/Source/Editor/Cascade/Private/CascadeModule.cpp
@@ -22,20 +22,20 @@ public:
 	}
 
 	/** Called right after the module DLL has been loaded and the module object has been created */
-	virtual void StartupModule() OVERRIDE
+	virtual void StartupModule() override
 	{
 		MenuExtensibilityManager = MakeShareable(new FExtensibilityManager);
 		ToolBarExtensibilityManager = MakeShareable(new FExtensibilityManager);
 	}
 
 	/** Called before the module is unloaded, right before the module object is destroyed. */
-	virtual void ShutdownModule() OVERRIDE
+	virtual void ShutdownModule() override
 	{
 		MenuExtensibilityManager.Reset();
 		ToolBarExtensibilityManager.Reset();
 	}
 
-	virtual TSharedRef<ICascade> CreateCascade(const EToolkitMode::Type Mode, const TSharedPtr< IToolkitHost >& InitToolkitHost, UParticleSystem* ParticleSystem) OVERRIDE
+	virtual TSharedRef<ICascade> CreateCascade(const EToolkitMode::Type Mode, const TSharedPtr< IToolkitHost >& InitToolkitHost, UParticleSystem* ParticleSystem) override
 	{
 		TSharedRef<FCascade> NewCascade(new FCascade());
 		NewCascade->InitCascade(Mode, InitToolkitHost, ParticleSystem);
@@ -43,12 +43,12 @@ public:
 		return NewCascade;
 	}
 
-	virtual void CascadeClosed(FCascade* CascadeInstance) OVERRIDE
+	virtual void CascadeClosed(FCascade* CascadeInstance) override
 	{
 		CascadeToolkits.Remove(CascadeInstance);
 	}
 
-	virtual void RefreshCascade(UParticleSystem* ParticleSystem) OVERRIDE
+	virtual void RefreshCascade(UParticleSystem* ParticleSystem) override
 	{
 		for (int32 Idx = 0; Idx < CascadeToolkits.Num(); ++Idx)
 		{
@@ -59,7 +59,7 @@ public:
 		}
 	}
 
-	virtual void ConvertModulesToSeeded(UParticleSystem* ParticleSystem) OVERRIDE
+	virtual void ConvertModulesToSeeded(UParticleSystem* ParticleSystem) override
 	{
 		FCascade::ConvertAllModulesToSeeded( ParticleSystem );
 	}
diff --git a/Engine/Source/Editor/Cascade/Private/CascadePreviewViewportClient.h b/Engine/Source/Editor/Cascade/Private/CascadePreviewViewportClient.h
--- a/Engine/Source/Editor/Cascade/Private/CascadePreviewViewportClient.h
+++ b/Engine/Source/Editor/Cascade/Private/CascadePreviewViewportClient.h
@@ -19,15 +19,15 @@ public:
 	~FCascadeEdPreviewViewportClient();
 
 	/** FEditorViewportClient interface */
-	virtual void Draw(FViewport* Viewport, FCanvas* Canvas) OVERRIDE;
-	virtual void Draw(const FSceneView* View, FPrimitiveDrawInterface* PDI) OVERRIDE;
-	virtual bool InputKey(FViewport* Viewport, int32 ControllerId, FKey Key, EInputEvent Event, float AmountDepressed = 1.0f, bool bGamepad = false) OVERRIDE;
-	virtual bool InputAxis(FViewport* Viewport, int32 ControllerId, FKey Key, float Delta, float DeltaTime, int32 NumSamples = 1, bool bGamepad = false) OVERRIDE;
-	virtual FSceneInterface* GetScene() const OVERRIDE;
-	virtual FLinearColor GetBackgroundColor() const OVERRIDE;
-	virtual bool ShouldOrbitCamera() const OVERRIDE;
-	virtual void AddReferencedObjects( FReferenceCollector& Collector ) OVERRIDE;
-	virtual bool CanCycleWidgetMode() const OVERRIDE;
+	virtual void Draw(FViewport* Viewport, FCanvas* Canvas) override;
+	virtual void Draw(const FSceneView* View, FPrimitiveDrawInterface* PDI) override;
+	virtual bool InputKey(FViewport* Viewport, int32 ControllerId, FKey Key, EInputEvent Event, float AmountDepressed = 1.0f, bool bGamepad = false) override;
+	virtual bool InputAxis(FViewport* Viewport, int32 ControllerId, FKey Key, float Delta, float DeltaTime, int32 NumSamples = 1, bool bGamepad = false) override;
+	virtual FSceneInterface* GetScene() const override;
+	virtual FLinearColor GetBackgroundColor() const override;
+	virtual bool ShouldOrbitCamera() const override;
+	virtual void AddReferencedObjects( FReferenceCollector& Collector ) override;
+	virtual bool CanCycleWidgetMode() const override;
 
 	/** Sets the position and orientation of the preview camera */
 	void SetPreviewCamera(const FRotator& NewPreviewAngle, float NewPreviewDistance);
diff --git a/Engine/Source/Editor/Cascade/Private/SCascadePreviewViewport.h b/Engine/Source/Editor/Cascade/Private/SCascadePreviewViewport.h
--- a/Engine/Source/Editor/Cascade/Private/SCascadePreviewViewport.h
+++ b/Engine/Source/Editor/Cascade/Private/SCascadePreviewViewport.h
@@ -40,9 +40,9 @@ public:
 	TSharedPtr<SViewport> GetViewportWidget() const;
 	
 protected:
-	virtual TSharedRef<FEditorViewportClient> MakeEditorViewportClient() OVERRIDE;
-	virtual TSharedPtr<SWidget> MakeViewportToolbar() OVERRIDE;
-	virtual void OnFocusViewportToSelection() OVERRIDE;
+	virtual TSharedRef<FEditorViewportClient> MakeEditorViewportClient() override;
+	virtual TSharedPtr<SWidget> MakeViewportToolbar() override;
+	virtual void OnFocusViewportToSelection() override;
 public:
 	/** The parent tab where this viewport resides */
 	TWeakPtr<SDockTab> ParentTab;
diff --git a/Engine/Source/Editor/ClassViewer/Private/SClassViewer.h b/Engine/Source/Editor/ClassViewer/Private/SClassViewer.h
--- a/Engine/Source/Editor/ClassViewer/Private/SClassViewer.h
+++ b/Engine/Source/Editor/ClassViewer/Private/SClassViewer.h
@@ -37,10 +37,10 @@ public:
 	virtual void ClearSelection();
 
 	/** SWidget interface */
-	virtual void Tick( const FGeometry& AllottedGeometry, const double InCurrentTime, const float InDeltaTime ) OVERRIDE;
-	virtual FReply OnKeyDown( const FGeometry& MyGeometry, const FKeyboardEvent& InKeyboardEvent ) OVERRIDE;
-	virtual FReply OnKeyboardFocusReceived( const FGeometry& MyGeometry, const FKeyboardFocusEvent& InKeyboardFocusEvent ) OVERRIDE;
-	virtual bool SupportsKeyboardFocus() const OVERRIDE;
+	virtual void Tick( const FGeometry& AllottedGeometry, const double InCurrentTime, const float InDeltaTime ) override;
+	virtual FReply OnKeyDown( const FGeometry& MyGeometry, const FKeyboardEvent& InKeyboardEvent ) override;
+	virtual FReply OnKeyboardFocusReceived( const FGeometry& MyGeometry, const FKeyboardFocusEvent& InKeyboardFocusEvent ) override;
+	virtual bool SupportsKeyboardFocus() const override;
 
 private:
 	/** Retrieves the children for the input node.
diff --git a/Engine/Source/Editor/ClassViewer/Private/UnloadedBlueprintData.h b/Engine/Source/Editor/ClassViewer/Private/UnloadedBlueprintData.h
--- a/Engine/Source/Editor/ClassViewer/Private/UnloadedBlueprintData.h
+++ b/Engine/Source/Editor/ClassViewer/Private/UnloadedBlueprintData.h
@@ -8,17 +8,17 @@ class FUnloadedBlueprintData : public IUnloadedBlueprintData
 public:
 	FUnloadedBlueprintData( TWeakPtr< class FClassViewerNode > InClassViewerNode );
 
-	virtual bool HasAnyClassFlags( uint32 InFlagsToCheck ) const OVERRIDE;
+	virtual bool HasAnyClassFlags( uint32 InFlagsToCheck ) const override;
 
-	virtual bool HasAllClassFlags( uint32 InFlagsToCheck ) const OVERRIDE;
+	virtual bool HasAllClassFlags( uint32 InFlagsToCheck ) const override;
 
-	virtual void SetClassFlags(uint32 InFlags) OVERRIDE;
+	virtual void SetClassFlags(uint32 InFlags) override;
 
-	virtual bool ImplementsInterface(const UClass* InInterface) const OVERRIDE;
+	virtual bool ImplementsInterface(const UClass* InInterface) const override;
 
-	virtual bool IsChildOf(const UClass* InClass) const OVERRIDE;
+	virtual bool IsChildOf(const UClass* InClass) const override;
 
-	virtual bool IsA(const UClass* InClass) const OVERRIDE;
+	virtual bool IsA(const UClass* InClass) const override;
 
 	/** Retrieves the Class Viewer node this data is associated with. */
 	const TWeakPtr< class FClassViewerNode > GetClassViewerNode() const;
diff --git a/Engine/Source/Editor/ComponentVisualizers/Private/AudioComponentVisualizer.h b/Engine/Source/Editor/ComponentVisualizers/Private/AudioComponentVisualizer.h
--- a/Engine/Source/Editor/ComponentVisualizers/Private/AudioComponentVisualizer.h
+++ b/Engine/Source/Editor/ComponentVisualizers/Private/AudioComponentVisualizer.h
@@ -8,6 +8,6 @@ class FAudioComponentVisualizer : public FComponentVisualizer
 {
 public:
 	// Begin FComponentVisualizer interface
-	virtual void DrawVisualization(const UActorComponent* Component, const FSceneView* View, FPrimitiveDrawInterface* PDI) OVERRIDE;
+	virtual void DrawVisualization(const UActorComponent* Component, const FSceneView* View, FPrimitiveDrawInterface* PDI) override;
 	// End FComponentVisualizer interface
 };
diff --git a/Engine/Source/Editor/ComponentVisualizers/Private/ConstraintComponentVisualizer.h b/Engine/Source/Editor/ComponentVisualizers/Private/ConstraintComponentVisualizer.h
--- a/Engine/Source/Editor/ComponentVisualizers/Private/ConstraintComponentVisualizer.h
+++ b/Engine/Source/Editor/ComponentVisualizers/Private/ConstraintComponentVisualizer.h
@@ -8,6 +8,6 @@ class FConstraintComponentVisualizer : public FComponentVisualizer
 {
 public:
 	// Begin FComponentVisualizer interface
-	virtual void DrawVisualization(const UActorComponent* Component, const FSceneView* View, FPrimitiveDrawInterface* PDI) OVERRIDE;
+	virtual void DrawVisualization(const UActorComponent* Component, const FSceneView* View, FPrimitiveDrawInterface* PDI) override;
 	// End FComponentVisualizer interface
 };
diff --git a/Engine/Source/Editor/ComponentVisualizers/Private/RadialForceComponentVisualizer.h b/Engine/Source/Editor/ComponentVisualizers/Private/RadialForceComponentVisualizer.h
--- a/Engine/Source/Editor/ComponentVisualizers/Private/RadialForceComponentVisualizer.h
+++ b/Engine/Source/Editor/ComponentVisualizers/Private/RadialForceComponentVisualizer.h
@@ -8,6 +8,6 @@ class FRadialForceComponentVisualizer : public FComponentVisualizer
 {
 public:
 	// Begin FComponentVisualizer interface
-	virtual void DrawVisualization( const UActorComponent* Component, const FSceneView* View, FPrimitiveDrawInterface* PDI ) OVERRIDE;
+	virtual void DrawVisualization( const UActorComponent* Component, const FSceneView* View, FPrimitiveDrawInterface* PDI ) override;
 	// End FComponentVisualizer interface
 };
diff --git a/Engine/Source/Editor/ComponentVisualizers/Private/SplineComponentVisualizer.h b/Engine/Source/Editor/ComponentVisualizers/Private/SplineComponentVisualizer.h
--- a/Engine/Source/Editor/ComponentVisualizers/Private/SplineComponentVisualizer.h
+++ b/Engine/Source/Editor/ComponentVisualizers/Private/SplineComponentVisualizer.h
@@ -37,14 +37,14 @@ public:
 	virtual ~FSplineComponentVisualizer();
 
 	// Begin FComponentVisualizer interface
-	virtual void OnRegister() OVERRIDE;
-	virtual void DrawVisualization(const UActorComponent* Component, const FSceneView* View, FPrimitiveDrawInterface* PDI) OVERRIDE;
-	virtual bool VisProxyHandleClick(HComponentVisProxy* VisProxy) OVERRIDE;
-	virtual void EndEditing() OVERRIDE;
-	virtual bool GetWidgetLocation(FVector& OutLocation) const OVERRIDE;
-	virtual bool HandleInputDelta(FLevelEditorViewportClient* ViewportClient, FViewport* Viewport, FVector& DeltaTranslate, FRotator& DeltaRotate, FVector& DeltaScale) OVERRIDE;
-	virtual bool HandleInputKey(FLevelEditorViewportClient* ViewportClient, FViewport* Viewport, FKey Key, EInputEvent Event) OVERRIDE;
-	virtual TSharedPtr<SWidget> GenerateContextMenu() const OVERRIDE;
+	virtual void OnRegister() override;
+	virtual void DrawVisualization(const UActorComponent* Component, const FSceneView* View, FPrimitiveDrawInterface* PDI) override;
+	virtual bool VisProxyHandleClick(HComponentVisProxy* VisProxy) override;
+	virtual void EndEditing() override;
+	virtual bool GetWidgetLocation(FVector& OutLocation) const override;
+	virtual bool HandleInputDelta(FLevelEditorViewportClient* ViewportClient, FViewport* Viewport, FVector& DeltaTranslate, FRotator& DeltaRotate, FVector& DeltaScale) override;
+	virtual bool HandleInputKey(FLevelEditorViewportClient* ViewportClient, FViewport* Viewport, FKey Key, EInputEvent Event) override;
+	virtual TSharedPtr<SWidget> GenerateContextMenu() const override;
 	// End FComponentVisualizer interface
 
 	/** Get the spline component we are currently editing */
diff --git a/Engine/Source/Editor/ComponentVisualizers/Public/ComponentVisualizers.h b/Engine/Source/Editor/ComponentVisualizers/Public/ComponentVisualizers.h
--- a/Engine/Source/Editor/ComponentVisualizers/Public/ComponentVisualizers.h
+++ b/Engine/Source/Editor/ComponentVisualizers/Public/ComponentVisualizers.h
@@ -22,7 +22,7 @@ public:
 	 *
 	 * @return	Whether the module supports shutdown separate from the rest of the engine.
 	 */
-	virtual bool SupportsDynamicReloading() OVERRIDE
+	virtual bool SupportsDynamicReloading() override
 	{
 		return true;
 	}
diff --git a/Engine/Source/Editor/ContentBrowser/Private/AssetViewTypes.h b/Engine/Source/Editor/ContentBrowser/Private/AssetViewTypes.h
--- a/Engine/Source/Editor/ContentBrowser/Private/AssetViewTypes.h
+++ b/Engine/Source/Editor/ContentBrowser/Private/AssetViewTypes.h
@@ -59,12 +59,12 @@ struct FAssetViewAsset : public FAssetViewItem
 	}
 
 	// FAssetViewItem interface
-	virtual EAssetItemType::Type GetType() const OVERRIDE
+	virtual EAssetItemType::Type GetType() const override
 	{
 		return EAssetItemType::Normal;
 	}
 
-	virtual bool IsTemporaryItem() const OVERRIDE
+	virtual bool IsTemporaryItem() const override
 	{
 		return false;
 	}
@@ -102,12 +102,12 @@ struct FAssetViewFolder : public FAssetViewItem
 	}
 
 	// FAssetViewItem interface
-	virtual EAssetItemType::Type GetType() const OVERRIDE
+	virtual EAssetItemType::Type GetType() const override
 	{
 		return EAssetItemType::Folder;
 	}
 
-	virtual bool IsTemporaryItem() const OVERRIDE
+	virtual bool IsTemporaryItem() const override
 	{
 		return false;
 	}
@@ -129,18 +129,18 @@ struct FAssetViewCreation : public FAssetViewAsset, public FGCObject
 	{}
 
 	// FAssetViewItem interface
-	virtual EAssetItemType::Type GetType() const OVERRIDE
+	virtual EAssetItemType::Type GetType() const override
 	{
 		return EAssetItemType::Creation;
 	}
 
-	virtual bool IsTemporaryItem() const OVERRIDE
+	virtual bool IsTemporaryItem() const override
 	{
 		return true;
 	}
 
 	// FGCObject interface
-	virtual void AddReferencedObjects(FReferenceCollector& Collector) OVERRIDE
+	virtual void AddReferencedObjects(FReferenceCollector& Collector) override
 	{
 		Collector.AddReferencedObject(AssetClass);
 		Collector.AddReferencedObject(Factory);
@@ -159,12 +159,12 @@ struct FAssetViewDuplication : public FAssetViewAsset
 	{}
 
 	// FAssetViewItem interface
-	virtual EAssetItemType::Type GetType() const OVERRIDE
+	virtual EAssetItemType::Type GetType() const override
 	{
 		return EAssetItemType::Duplication;
 	}
 
-	virtual bool IsTemporaryItem() const OVERRIDE
+	virtual bool IsTemporaryItem() const override
 	{
 		return true;
 	}
