UE4 Refactoring. Changed OVERRIDE and FINAL macros to keywords override and final.
Some missing changes from last submit plus deprecated message when using macros.

diff --git a/Engine/Plugins/Experimental/Paper2D/Source/Paper2DEditor/Private/TileMapEditing/EdModeTileMap.h b/Engine/Plugins/Experimental/Paper2D/Source/Paper2DEditor/Private/TileMapEditing/EdModeTileMap.h
--- a/Engine/Plugins/Experimental/Paper2D/Source/Paper2DEditor/Private/TileMapEditing/EdModeTileMap.h
+++ b/Engine/Plugins/Experimental/Paper2D/Source/Paper2DEditor/Private/TileMapEditing/EdModeTileMap.h
@@ -41,28 +41,28 @@ public:
 	virtual bool UsesToolkits() const override;
 	virtual void Enter() override;
 	virtual void Exit() override;
-	//virtual bool BoxSelect(FBox& InBox, bool bInSelect) OVERRIDE;
+	//virtual bool BoxSelect(FBox& InBox, bool bInSelect) override;
 
-	//	virtual void PostUndo() OVERRIDE;
+	//	virtual void PostUndo() override;
 	virtual bool MouseMove(FLevelEditorViewportClient* ViewportClient, FViewport* Viewport, int32 x, int32 y) override;
 	virtual bool CapturedMouseMove(FLevelEditorViewportClient* InViewportClient, FViewport* InViewport, int32 InMouseX, int32 InMouseY) override;
 	virtual bool StartTracking(FLevelEditorViewportClient* InViewportClient, FViewport* InViewport) override;
 	virtual bool EndTracking(FLevelEditorViewportClient* InViewportClient, FViewport* InViewport) override;
-	// 	virtual void Tick(FLevelEditorViewportClient* ViewportClient,float DeltaTime) OVERRIDE;
+	// 	virtual void Tick(FLevelEditorViewportClient* ViewportClient,float DeltaTime) override;
 	virtual bool InputKey(FLevelEditorViewportClient* InViewportClient, FViewport* InViewport, FKey InKey, EInputEvent InEvent) override;
 	virtual bool InputDelta(FLevelEditorViewportClient* InViewportClient, FViewport* InViewport, FVector& InDrag, FRotator& InRot, FVector& InScale) override;
 	virtual void Render(const FSceneView* View, FViewport* Viewport, FPrimitiveDrawInterface* PDI) override;
 	virtual void DrawHUD(FLevelEditorViewportClient* ViewportClient, FViewport* Viewport, const FSceneView* View, FCanvas* Canvas) override;
-	// 	virtual bool Select(AActor* InActor, bool bInSelected) OVERRIDE;
-	// 	virtual bool IsSelectionAllowed(AActor* InActor) const OVERRIDE;
-	// 	virtual void ActorSelectionChangeNotify() OVERRIDE;
-	// 	virtual FVector GetWidgetLocation() const OVERRIDE;
+	// 	virtual bool Select(AActor* InActor, bool bInSelected) override;
+	// 	virtual bool IsSelectionAllowed(AActor* InActor) const override;
+	// 	virtual void ActorSelectionChangeNotify() override;
+	// 	virtual FVector GetWidgetLocation() const override;
 	virtual bool AllowWidgetMove();
 	virtual bool ShouldDrawWidget() const override;
 	virtual bool UsesTransformWidget() const override;
 	virtual void PeekAtSelectionChangedEvent(UObject* ItemUndergoingChange) override;
-	// 	virtual int32 GetWidgetAxisToDraw(FWidget::EWidgetMode InWidgetMode) const OVERRIDE;
-	// 	virtual bool DisallowMouseDeltaTracking() const OVERRIDE;
+	// 	virtual int32 GetWidgetAxisToDraw(FWidget::EWidgetMode InWidgetMode) const override;
+	// 	virtual bool DisallowMouseDeltaTracking() const override;
 	// End of FEdMode interface
 
 	void SetActiveTool(ETileMapEditorTool::Type NewTool);
diff --git a/Engine/Plugins/OculusRift/Source/OculusRift/Private/OculusRiftHMD.h b/Engine/Plugins/OculusRift/Source/OculusRift/Private/OculusRiftHMD.h
--- a/Engine/Plugins/OculusRift/Source/OculusRift/Private/OculusRiftHMD.h
+++ b/Engine/Plugins/OculusRift/Source/OculusRift/Private/OculusRiftHMD.h
@@ -68,7 +68,7 @@ public:
 
 	virtual void SetInterpupillaryDistance(float NewInterpupillaryDistance) override;
 	virtual float GetInterpupillaryDistance() const override;
-    //virtual float GetFieldOfViewInRadians() const OVERRIDE;
+    //virtual float GetFieldOfViewInRadians() const override;
 	virtual void GetFieldOfView(float& OutHFOVInDegrees, float& OutVFOVInDegrees) const override;
 
     virtual void GetCurrentOrientationAndPosition(FQuat& CurrentOrientation, FVector& CurrentPosition) const override;
diff --git a/Engine/Source/Developer/AutomationWindow/Private/SAutomationTestItem.cpp b/Engine/Source/Developer/AutomationWindow/Private/SAutomationTestItem.cpp
--- a/Engine/Source/Developer/AutomationWindow/Private/SAutomationTestItem.cpp
+++ b/Engine/Source/Developer/AutomationWindow/Private/SAutomationTestItem.cpp
@@ -37,7 +37,7 @@ public:
 	void RebuildContentArea();
 
 	// Begin SWidget Interface
-	virtual void Tick(const FGeometry& AllottedGeometry, const double InCurrentTime, const float InDeltaTime) OVERRIDE;
+	virtual void Tick(const FGeometry& AllottedGeometry, const double InCurrentTime, const float InDeltaTime) override;
 	// End SWidget Interface
 
 private:
diff --git a/Engine/Source/Developer/LogVisualizer/Private/SLogVisualizer.h b/Engine/Source/Developer/LogVisualizer/Private/SLogVisualizer.h
--- a/Engine/Source/Developer/LogVisualizer/Private/SLogVisualizer.h
+++ b/Engine/Source/Developer/LogVisualizer/Private/SLogVisualizer.h
@@ -78,7 +78,7 @@ public:
 
 	virtual void Tick( const FGeometry& AllottedGeometry, const double InCurrentTime, const float InDeltaTime ) override;
 
-	// OVERRIDES
+	// overrides
 	virtual FReply OnMouseWheel( const FGeometry& MyGeometry, const FPointerEvent& MouseEvent ) override;
 	virtual FReply OnKeyDown( const FGeometry& MyGeometry, const FKeyboardEvent& InKeyboardEvent) override;
 
diff --git a/Engine/Source/Editor/BlueprintGraph/Classes/EdGraphSchema_K2.h b/Engine/Source/Editor/BlueprintGraph/Classes/EdGraphSchema_K2.h
--- a/Engine/Source/Editor/BlueprintGraph/Classes/EdGraphSchema_K2.h
+++ b/Engine/Source/Editor/BlueprintGraph/Classes/EdGraphSchema_K2.h
@@ -12,11 +12,11 @@ struct BLUEPRINTGRAPH_API FEdGraphSchemaAction_K2Struct : public FEdGraphSchemaA
 
 	// Simple type info
 	static FString StaticGetTypeId() {static FString Type = TEXT("FEdGraphSchemaAction_K2Struct"); return Type;}
-	virtual FString GetTypeId() const OVERRIDE { return StaticGetTypeId(); } 
+	virtual FString GetTypeId() const override { return StaticGetTypeId(); } 
 
 	UStruct* Struct;
 
-	void AddReferencedObjects( FReferenceCollector& Collector ) OVERRIDE
+	void AddReferencedObjects( FReferenceCollector& Collector ) override
 	{
 		if( Struct )
 		{
@@ -456,43 +456,43 @@ public:
 
 public:
 	// Begin EdGraphSchema Interface
-	virtual void GetGraphContextActions(FGraphContextMenuBuilder& ContextMenuBuilder) const OVERRIDE;
-	virtual void GetContextMenuActions(const UEdGraph* CurrentGraph, const UEdGraphNode* InGraphNode, const UEdGraphPin* InGraphPin, class FMenuBuilder* MenuBuilder, bool bIsDebugging) const OVERRIDE;
-	virtual const FPinConnectionResponse CanCreateConnection(const UEdGraphPin* A, const UEdGraphPin* B) const OVERRIDE;
-	virtual bool TryCreateConnection(UEdGraphPin* A, UEdGraphPin* B) const OVERRIDE;
-	virtual bool CreateAutomaticConversionNodeAndConnections(UEdGraphPin* A, UEdGraphPin* B) const OVERRIDE;
-	virtual FString IsPinDefaultValid(const UEdGraphPin* Pin, const FString& NewDefaultValue, UObject* NewDefaultObject, const FText& InNewDefaultText) const OVERRIDE;
-	virtual bool DoesSupportPinWatching() const	OVERRIDE;
-	virtual bool IsPinBeingWatched(UEdGraphPin const* Pin) const OVERRIDE;
-	virtual void ClearPinWatch(UEdGraphPin const* Pin) const OVERRIDE;
-	virtual void TrySetDefaultValue(UEdGraphPin& Pin, const FString& NewDefaultValue) const OVERRIDE;
-	virtual void TrySetDefaultObject(UEdGraphPin& Pin, UObject* NewDefaultObject) const OVERRIDE;
-	virtual void TrySetDefaultText(UEdGraphPin& InPin, const FText& InNewDefaultText) const OVERRIDE;
-	virtual bool ShouldHidePinDefaultValue(UEdGraphPin* Pin) const OVERRIDE;
-	virtual FLinearColor GetPinTypeColor(const FEdGraphPinType& PinType) const OVERRIDE;
-	virtual FString GetPinDisplayName(const UEdGraphPin* Pin) const OVERRIDE;
-	virtual void ConstructBasicPinTooltip(const UEdGraphPin& Pin, const FString& PinDescription, FString& TooltipOut) const OVERRIDE;
-	virtual EGraphType GetGraphType(const UEdGraph* TestEdGraph) const OVERRIDE;
-	virtual bool IsTitleBarPin(const UEdGraphPin& Pin) const OVERRIDE;
-	virtual void BreakNodeLinks(UEdGraphNode& TargetNode) const OVERRIDE;
-	virtual void BreakPinLinks(UEdGraphPin& TargetPin, bool bSendsNodeNotifcation) const OVERRIDE;
-	virtual void BreakSinglePinLink(UEdGraphPin* SourcePin, UEdGraphPin* TargetPin) OVERRIDE;
-	virtual void ReconstructNode(UEdGraphNode& TargetNode, bool bIsBatchRequest=false) const OVERRIDE;
-	virtual bool CanEncapuslateNode(UEdGraphNode const& TestNode) const OVERRIDE;
-	virtual void HandleGraphBeingDeleted(UEdGraph& GraphBeingRemoved) const OVERRIDE;
-	virtual void GetGraphDisplayInformation(const UEdGraph& Graph, /*out*/ FGraphDisplayInfo& DisplayInfo) const OVERRIDE;
-	virtual void DroppedAssetsOnGraph(const TArray<FAssetData>& Assets, const FVector2D& GraphPosition, UEdGraph* Graph) const OVERRIDE;
-	virtual void DroppedAssetsOnNode(const TArray<FAssetData>& Assets, const FVector2D& GraphPosition, UEdGraphNode* Node) const OVERRIDE;
-	virtual void DroppedAssetsOnPin(const TArray<FAssetData>& Assets, const FVector2D& GraphPosition, UEdGraphPin* Pin) const OVERRIDE;
-	virtual void GetAssetsNodeHoverMessage(const TArray<FAssetData>& Assets, const UEdGraphNode* HoverNode, FString& OutTooltipText, bool& OutOkIcon) const OVERRIDE;
-	virtual void GetAssetsPinHoverMessage(const TArray<FAssetData>& Assets, const UEdGraphPin* HoverPin, FString& OutTooltipText, bool& OutOkIcon) const OVERRIDE;
-	virtual bool CanDuplicateGraph(UEdGraph* InSourceGraph) const OVERRIDE;
-	virtual UEdGraph* DuplicateGraph(UEdGraph* GraphToDuplicate) const OVERRIDE;
-	virtual UEdGraphNode* CreateSubstituteNode(UEdGraphNode* Node, const UEdGraph* Graph, FObjectInstancingGraph* InstanceGraph) const OVERRIDE;
-	virtual int32 GetNodeSelectionCount(const UEdGraph* Graph) const OVERRIDE;
-	virtual TSharedPtr<FEdGraphSchemaAction> GetCreateCommentAction() const OVERRIDE;
-	virtual bool FadeNodeWhenDraggingOffPin(const UEdGraphNode* Node, const UEdGraphPin* Pin) const OVERRIDE;
-	virtual void BackwardCompatibilityNodeConversion(UEdGraph* Graph, bool bOnlySafeChanges) const OVERRIDE;
+	virtual void GetGraphContextActions(FGraphContextMenuBuilder& ContextMenuBuilder) const override;
+	virtual void GetContextMenuActions(const UEdGraph* CurrentGraph, const UEdGraphNode* InGraphNode, const UEdGraphPin* InGraphPin, class FMenuBuilder* MenuBuilder, bool bIsDebugging) const override;
+	virtual const FPinConnectionResponse CanCreateConnection(const UEdGraphPin* A, const UEdGraphPin* B) const override;
+	virtual bool TryCreateConnection(UEdGraphPin* A, UEdGraphPin* B) const override;
+	virtual bool CreateAutomaticConversionNodeAndConnections(UEdGraphPin* A, UEdGraphPin* B) const override;
+	virtual FString IsPinDefaultValid(const UEdGraphPin* Pin, const FString& NewDefaultValue, UObject* NewDefaultObject, const FText& InNewDefaultText) const override;
+	virtual bool DoesSupportPinWatching() const	override;
+	virtual bool IsPinBeingWatched(UEdGraphPin const* Pin) const override;
+	virtual void ClearPinWatch(UEdGraphPin const* Pin) const override;
+	virtual void TrySetDefaultValue(UEdGraphPin& Pin, const FString& NewDefaultValue) const override;
+	virtual void TrySetDefaultObject(UEdGraphPin& Pin, UObject* NewDefaultObject) const override;
+	virtual void TrySetDefaultText(UEdGraphPin& InPin, const FText& InNewDefaultText) const override;
+	virtual bool ShouldHidePinDefaultValue(UEdGraphPin* Pin) const override;
+	virtual FLinearColor GetPinTypeColor(const FEdGraphPinType& PinType) const override;
+	virtual FString GetPinDisplayName(const UEdGraphPin* Pin) const override;
+	virtual void ConstructBasicPinTooltip(const UEdGraphPin& Pin, const FString& PinDescription, FString& TooltipOut) const override;
+	virtual EGraphType GetGraphType(const UEdGraph* TestEdGraph) const override;
+	virtual bool IsTitleBarPin(const UEdGraphPin& Pin) const override;
+	virtual void BreakNodeLinks(UEdGraphNode& TargetNode) const override;
+	virtual void BreakPinLinks(UEdGraphPin& TargetPin, bool bSendsNodeNotifcation) const override;
+	virtual void BreakSinglePinLink(UEdGraphPin* SourcePin, UEdGraphPin* TargetPin) override;
+	virtual void ReconstructNode(UEdGraphNode& TargetNode, bool bIsBatchRequest=false) const override;
+	virtual bool CanEncapuslateNode(UEdGraphNode const& TestNode) const override;
+	virtual void HandleGraphBeingDeleted(UEdGraph& GraphBeingRemoved) const override;
+	virtual void GetGraphDisplayInformation(const UEdGraph& Graph, /*out*/ FGraphDisplayInfo& DisplayInfo) const override;
+	virtual void DroppedAssetsOnGraph(const TArray<FAssetData>& Assets, const FVector2D& GraphPosition, UEdGraph* Graph) const override;
+	virtual void DroppedAssetsOnNode(const TArray<FAssetData>& Assets, const FVector2D& GraphPosition, UEdGraphNode* Node) const override;
+	virtual void DroppedAssetsOnPin(const TArray<FAssetData>& Assets, const FVector2D& GraphPosition, UEdGraphPin* Pin) const override;
+	virtual void GetAssetsNodeHoverMessage(const TArray<FAssetData>& Assets, const UEdGraphNode* HoverNode, FString& OutTooltipText, bool& OutOkIcon) const override;
+	virtual void GetAssetsPinHoverMessage(const TArray<FAssetData>& Assets, const UEdGraphPin* HoverPin, FString& OutTooltipText, bool& OutOkIcon) const override;
+	virtual bool CanDuplicateGraph(UEdGraph* InSourceGraph) const override;
+	virtual UEdGraph* DuplicateGraph(UEdGraph* GraphToDuplicate) const override;
+	virtual UEdGraphNode* CreateSubstituteNode(UEdGraphNode* Node, const UEdGraph* Graph, FObjectInstancingGraph* InstanceGraph) const override;
+	virtual int32 GetNodeSelectionCount(const UEdGraph* Graph) const override;
+	virtual TSharedPtr<FEdGraphSchemaAction> GetCreateCommentAction() const override;
+	virtual bool FadeNodeWhenDraggingOffPin(const UEdGraphNode* Node, const UEdGraphPin* Pin) const override;
+	virtual void BackwardCompatibilityNodeConversion(UEdGraph* Graph, bool bOnlySafeChanges) const override;
 	// End EdGraphSchema Interface
 
 	// Do validation, that doesn't require a knowledge about actual pin. 
@@ -537,7 +537,7 @@ public:
 	 * @param	Pin	The pin to check.
 	 * @return	true if it is a Self pin.
 	 */
-	virtual bool IsSelfPin(const UEdGraphPin& Pin) const OVERRIDE;
+	virtual bool IsSelfPin(const UEdGraphPin& Pin) const override;
 
 	/**
 	 * Checks to see if a pin is a meta-pin (either a Self or Exec pin)
@@ -551,7 +551,7 @@ public:
 	}
 
 	/** Is given string a delegate category name ? */
-	virtual bool IsDelegateCategory(const FString& Category) const OVERRIDE;
+	virtual bool IsDelegateCategory(const FString& Category) const override;
 
 	/** Returns whether a pin category is compatible with an Index Wildcard (PC_Wildcard and PSC_Index) */
 	inline bool IsIndexWildcardCompatible(const FEdGraphPinType& PinType) const
diff --git a/Engine/Source/Editor/BlueprintGraph/Classes/K2Node_FormatText.h b/Engine/Source/Editor/BlueprintGraph/Classes/K2Node_FormatText.h
--- a/Engine/Source/Editor/BlueprintGraph/Classes/K2Node_FormatText.h
+++ b/Engine/Source/Editor/BlueprintGraph/Classes/K2Node_FormatText.h
@@ -23,7 +23,7 @@ class UK2Node_FormatText : public UK2Node
 	// End UEdGraphNode interface.
 
 	// Begin UK2Node interface.
-	//virtual bool IsNodePure() const OVERRIDE { return true; }
+	//virtual bool IsNodePure() const override { return true; }
 	virtual bool NodeCausesStructuralBlueprintChange() const override { return true; }
 	virtual void ExpandNode(class FKismetCompilerContext& CompilerContext, UEdGraph* SourceGraph) override;
 	virtual ERedirectType DoPinsMatchForReconstruction(const UEdGraphPin* NewPin, int32 NewPinIndex, const UEdGraphPin* OldPin, int32 OldPinIndex) const override;
diff --git a/Engine/Source/Editor/BlueprintGraph/Classes/K2Node_StructOperation.h b/Engine/Source/Editor/BlueprintGraph/Classes/K2Node_StructOperation.h
--- a/Engine/Source/Editor/BlueprintGraph/Classes/K2Node_StructOperation.h
+++ b/Engine/Source/Editor/BlueprintGraph/Classes/K2Node_StructOperation.h
@@ -15,11 +15,11 @@ class UK2Node_StructOperation : public UK2Node_Variable
 	UScriptStruct* StructType;
 
 	// Begin UEdGraphNode interface
-	//virtual FLinearColor GetNodeTitleColor() const OVERRIDE;
+	//virtual FLinearColor GetNodeTitleColor() const override;
 	// End UEdGraphNode interface
 
 	// UK2Node interface
-	//virtual bool DrawNodeAsVariable() const OVERRIDE { return true; }
+	//virtual bool DrawNodeAsVariable() const override { return true; }
 	virtual bool ShouldShowNodeProperties() const override { return true; }
 	virtual void ValidateNodeDuringCompilation(class FCompilerResultsLog& MessageLog) const override {}
 	virtual bool HasExternalUserDefinedStructDependencies(TArray<class UStruct*>* OptionalOutput) const override;
diff --git a/Engine/Source/Editor/GraphEditor/Private/DragNode.h b/Engine/Source/Editor/GraphEditor/Private/DragNode.h
--- a/Engine/Source/Editor/GraphEditor/Private/DragNode.h
+++ b/Engine/Source/Editor/GraphEditor/Private/DragNode.h
@@ -15,7 +15,7 @@ public:
 	virtual void HoverTargetChanged() override;
 	virtual FReply DroppedOnNode(FVector2D ScreenPosition, FVector2D GraphPosition) override;
 	virtual FReply DroppedOnPanel( const TSharedRef< SWidget >& Panel, FVector2D ScreenPosition, FVector2D GraphPosition, UEdGraph& Graph) override;
-	//virtual void OnDragBegin( const TSharedRef<class SGraphPin>& InPin) OVERRIDE;
+	//virtual void OnDragBegin( const TSharedRef<class SGraphPin>& InPin) override;
 	virtual void OnDragged (const class FDragDropEvent& DragDropEvent ) override;
 	// End of FGraphEditorDragDropAction interface
 	
diff --git a/Engine/Source/Editor/Kismet/Private/BlueprintEditorTabFactories.h b/Engine/Source/Editor/Kismet/Private/BlueprintEditorTabFactories.h
--- a/Engine/Source/Editor/Kismet/Private/BlueprintEditorTabFactories.h
+++ b/Engine/Source/Editor/Kismet/Private/BlueprintEditorTabFactories.h
@@ -210,7 +210,7 @@ public:
 		}
 	}
 
-//	virtual void OnTabClosed(TSharedRef<SDockTab> Tab, TSharedPtr<FTabPayload> Payload) const OVERRIDE
+//	virtual void OnTabClosed(TSharedRef<SDockTab> Tab, TSharedPtr<FTabPayload> Payload) const override
 //	{
 //		BlueprintEditorPtr.Pin()->RequestSaveEditedObjectState();
 //	}
@@ -273,7 +273,7 @@ protected:
 		BlueprintEditorPtr.Pin()->GetBlueprintObj()->LastEditedDocuments.Add(FEditedDocumentInfo(Timeline));
 	}
 
-// 	virtual void OnTabClosed(TSharedRef<SDockTab> Tab, TSharedPtr<FTabPayload> Payload) const OVERRIDE
+// 	virtual void OnTabClosed(TSharedRef<SDockTab> Tab, TSharedPtr<FTabPayload> Payload) const override
 // 	{
 // 		BlueprintEditorPtr.Pin()->RequestSaveEditedObjectState();
 // 	}
diff --git a/Engine/Source/Editor/LandscapeEditor/Classes/ActorFactoryLandscape.h b/Engine/Source/Editor/LandscapeEditor/Classes/ActorFactoryLandscape.h
--- a/Engine/Source/Editor/LandscapeEditor/Classes/ActorFactoryLandscape.h
+++ b/Engine/Source/Editor/LandscapeEditor/Classes/ActorFactoryLandscape.h
@@ -10,10 +10,10 @@ class UActorFactoryLandscape : public UActorFactory
 	GENERATED_UCLASS_BODY()
 
 	// Begin UActorFactory Interface
-	//virtual bool CanCreateActorFrom( const FAssetData& AssetData, FText& OutErrorMsg ) OVERRIDE;
+	//virtual bool CanCreateActorFrom( const FAssetData& AssetData, FText& OutErrorMsg ) override;
 	virtual AActor* SpawnActor( UObject* Asset, ULevel* InLevel, const FVector& Location, const FRotator& Rotation, EObjectFlags ObjectFlags, const FName& Name ) override;
-	//virtual void PostSpawnActor( UObject* Asset, AActor* NewActor) OVERRIDE;
-	//virtual void PostCreateBlueprint( UObject* Asset, AActor* CDO ) OVERRIDE;
-	//virtual UObject* GetAssetFromActorInstance(AActor* ActorInstance) OVERRIDE;
+	//virtual void PostSpawnActor( UObject* Asset, AActor* NewActor) override;
+	//virtual void PostCreateBlueprint( UObject* Asset, AActor* CDO ) override;
+	//virtual UObject* GetAssetFromActorInstance(AActor* ActorInstance) override;
 	// End UActorFactory Interface
 };
diff --git a/Engine/Source/Editor/LandscapeEditor/Classes/LandscapePlaceholder.h b/Engine/Source/Editor/LandscapeEditor/Classes/LandscapePlaceholder.h
--- a/Engine/Source/Editor/LandscapeEditor/Classes/LandscapePlaceholder.h
+++ b/Engine/Source/Editor/LandscapeEditor/Classes/LandscapePlaceholder.h
@@ -9,7 +9,7 @@ class ALandscapePlaceholder : public AActor
 {
 	GENERATED_UCLASS_BODY()
 
-	//virtual void PostEditMove(bool bFinished) OVERRIDE;
+	//virtual void PostEditMove(bool bFinished) override;
 	virtual bool TeleportTo(const FVector& DestLocation, const FRotator& DestRotation, bool bIsATest = false, bool bNoCheck = false) override;
 
 	void PostEditChangeProperty(FPropertyChangedEvent& PropertyChangedEvent);
diff --git a/Engine/Source/Editor/Persona/Private/Persona.h b/Engine/Source/Editor/Persona/Private/Persona.h
--- a/Engine/Source/Editor/Persona/Private/Persona.h
+++ b/Engine/Source/Editor/Persona/Private/Persona.h
@@ -220,7 +220,7 @@ public:
 
 protected:
 	// FBlueprintEditor interface
-	//virtual void CreateDefaultToolbar() OVERRIDE;
+	//virtual void CreateDefaultToolbar() override;
 	virtual void CreateDefaultCommands() override;
 	virtual void OnSelectBone() override;
 	virtual bool CanSelectBone() const override;
diff --git a/Engine/Source/Editor/UMGEditor/Classes/WidgetGraphNode_Base.h b/Engine/Source/Editor/UMGEditor/Classes/WidgetGraphNode_Base.h
--- a/Engine/Source/Editor/UMGEditor/Classes/WidgetGraphNode_Base.h
+++ b/Engine/Source/Editor/UMGEditor/Classes/WidgetGraphNode_Base.h
@@ -19,8 +19,8 @@ class UWidgetGraphNode_Base : public UK2Node
 	//// UEdGraphNode interface
 	UMGEDITOR_API virtual void AllocateDefaultPins() override;
 	UMGEDITOR_API virtual FLinearColor GetNodeTitleColor() const override;
-	//UMGEDITOR_API virtual FString GetDocumentationLink() const OVERRIDE;
-	//UMGEDITOR_API virtual void GetPinHoverText(const UEdGraphPin& Pin, FString& HoverTextOut) const OVERRIDE;
+	//UMGEDITOR_API virtual FString GetDocumentationLink() const override;
+	//UMGEDITOR_API virtual void GetPinHoverText(const UEdGraphPin& Pin, FString& HoverTextOut) const override;
 	UMGEDITOR_API virtual bool ShowPaletteIconOnNode() const override { return false; }
 	//	// End of UEdGraphNode interface
 
diff --git a/Engine/Source/Editor/UMGEditor/Classes/WidgetGraphNode_Text.h b/Engine/Source/Editor/UMGEditor/Classes/WidgetGraphNode_Text.h
--- a/Engine/Source/Editor/UMGEditor/Classes/WidgetGraphNode_Text.h
+++ b/Engine/Source/Editor/UMGEditor/Classes/WidgetGraphNode_Text.h
@@ -19,7 +19,7 @@ class UWidgetGraphNode_Text : public UWidgetGraphNode_Base
 
 	/*
 	// UAnimGraphNode_Base interface
-	virtual void CustomizePinData(UEdGraphPin* Pin, FName SourcePropertyName, int32 ArrayIndex) const OVERRIDE;
+	virtual void CustomizePinData(UEdGraphPin* Pin, FName SourcePropertyName, int32 ArrayIndex) const override;
 	// End of UAnimGraphNode_Base interface
 	*/
 };
diff --git a/Engine/Source/Editor/UMGEditor/Private/WidgetBlueprintCompiler.h b/Engine/Source/Editor/UMGEditor/Private/WidgetBlueprintCompiler.h
--- a/Engine/Source/Editor/UMGEditor/Private/WidgetBlueprintCompiler.h
+++ b/Engine/Source/Editor/UMGEditor/Private/WidgetBlueprintCompiler.h
@@ -26,7 +26,7 @@ protected:
 	void ValidateWidgetNames();
 
 	// FKismetCompilerContext
-	//virtual UEdGraphSchema_K2* CreateSchema() OVERRIDE;
+	//virtual UEdGraphSchema_K2* CreateSchema() override;
 	virtual void MergeUbergraphPagesIn(UEdGraph* Ubergraph) override;
 	virtual void SpawnNewClass(const FString& NewClassName) override;
 	virtual void CleanAndSanitizeClass(UBlueprintGeneratedClass* ClassToClean, UObject*& OldCDO) override;
diff --git a/Engine/Source/Editor/UMGEditor/Private/WidgetTemplateDragDropOp.h b/Engine/Source/Editor/UMGEditor/Private/WidgetTemplateDragDropOp.h
--- a/Engine/Source/Editor/UMGEditor/Private/WidgetTemplateDragDropOp.h
+++ b/Engine/Source/Editor/UMGEditor/Private/WidgetTemplateDragDropOp.h
@@ -14,5 +14,5 @@ public:
 
 	static TSharedRef<FWidgetTemplateDragDropOp> New(const TSharedPtr<FWidgetTemplate>& InTemplate);
 
-	//virtual TSharedPtr<SWidget> GetDefaultDecorator() const OVERRIDE;
+	//virtual TSharedPtr<SWidget> GetDefaultDecorator() const override;
 };
diff --git a/Engine/Source/Editor/UnrealEd/Public/DragAndDrop/ExportTextDragDropOp.h b/Engine/Source/Editor/UnrealEd/Public/DragAndDrop/ExportTextDragDropOp.h
--- a/Engine/Source/Editor/UnrealEd/Public/DragAndDrop/ExportTextDragDropOp.h
+++ b/Engine/Source/Editor/UnrealEd/Public/DragAndDrop/ExportTextDragDropOp.h
@@ -59,7 +59,7 @@ public:
 	int32 NumActors;
 
 	/** The widget decorator to use */
-	//virtual TSharedPtr<SWidget> GetDefaultDecorator() const OVERRIDE
+	//virtual TSharedPtr<SWidget> GetDefaultDecorator() const override
 	//{
 		
 	//}
diff --git a/Engine/Source/Programs/UnrealHeaderTool/Private/HeaderParser.cpp b/Engine/Source/Programs/UnrealHeaderTool/Private/HeaderParser.cpp
--- a/Engine/Source/Programs/UnrealHeaderTool/Private/HeaderParser.cpp
+++ b/Engine/Source/Programs/UnrealHeaderTool/Private/HeaderParser.cpp
@@ -5068,7 +5068,7 @@ void FHeaderParser::CompileFunctionDeclaration(FClasses& AllClasses)
 
 	if (FuncInfo.FunctionFlags & FUNC_Net)
 	{
-		// Network replicated functions are always events, and are only FINAL if sealed
+		// Network replicated functions are always events, and are only final if sealed
 		TypeOfFunction = TEXT("event");
 		bAutomaticallyFinal = false;
 	}
@@ -5197,7 +5197,7 @@ void FHeaderParser::CompileFunctionDeclaration(FClasses& AllClasses)
 	// Process the virtualness
 	if (bSawVirtual)
 	{
-		// Remove the implicit FINAL, the user can still specifying an explicit FINAL at the end of the declaration
+		// Remove the implicit final, the user can still specifying an explicit final at the end of the declaration
 		bAutomaticallyFinal = false;
 
 		// if this is a BlueprintNativeEvent or BlueprintImplementableEvent in an interface, make sure it's not "virtual"
@@ -5228,7 +5228,7 @@ void FHeaderParser::CompileFunctionDeclaration(FClasses& AllClasses)
 	}
 
 	// Handle the initial implicit/explicit final
-	// A user can still specify an explicit FINAL after the parameter list as well.
+	// A user can still specify an explicit final after the parameter list as well.
 	if (bAutomaticallyFinal || FuncInfo.bSealedEvent)
 	{
 		FuncInfo.FunctionFlags |= FUNC_Final;
@@ -5392,8 +5392,8 @@ void FHeaderParser::CompileFunctionDeclaration(FClasses& AllClasses)
 		}
 	}
 
-	// Look for the FINAL macro to indicate this function is sealed
-	if (MatchIdentifier(TEXT("FINAL")))
+	// Look for the final keyword to indicate this function is sealed
+	if (MatchIdentifier(TEXT("final")))
 	{
 		// This is a final (prebinding, non-overridable) function
 		FuncInfo.FunctionFlags |= FUNC_Final;
diff --git a/Engine/Source/Runtime/AIModule/Classes/EnvironmentQuery/EQSRenderingComponent.h b/Engine/Source/Runtime/AIModule/Classes/EnvironmentQuery/EQSRenderingComponent.h
--- a/Engine/Source/Runtime/AIModule/Classes/EnvironmentQuery/EQSRenderingComponent.h
+++ b/Engine/Source/Runtime/AIModule/Classes/EnvironmentQuery/EQSRenderingComponent.h
@@ -14,7 +14,7 @@ public:
 
 	virtual void DrawDebugLabels(UCanvas* Canvas, APlayerController*) override;
 	
-	//virtual void DrawStaticElements(FStaticPrimitiveDrawInterface* PDI) OVERRIDE;
+	//virtual void DrawStaticElements(FStaticPrimitiveDrawInterface* PDI) override;
 	virtual FPrimitiveViewRelevance GetViewRelevance(const FSceneView* View) override;
 
 	static void CollectEQSData(const UPrimitiveComponent* InComponent, const class IEQSQueryResultSourceInterface* QueryDataSource, TArray<FSphere>& Spheres, TArray<FText3d>& Texts);
diff --git a/Engine/Source/Runtime/Core/Private/IOS/IOSPlatformFile.h b/Engine/Source/Runtime/Core/Private/IOS/IOSPlatformFile.h
--- a/Engine/Source/Runtime/Core/Private/IOS/IOSPlatformFile.h
+++ b/Engine/Source/Runtime/Core/Private/IOS/IOSPlatformFile.h
@@ -39,7 +39,7 @@ protected:
 	virtual FString NormalizeDirectory(const TCHAR* Directory);
 
 public:
-	//virtual bool Initialize(IPlatformFile* Inner, const TCHAR* CommandLineParam) OVERRIDE;
+	//virtual bool Initialize(IPlatformFile* Inner, const TCHAR* CommandLineParam) override;
 
 	virtual bool FileExists(const TCHAR* Filename) override;
 	virtual int64 FileSize(const TCHAR* Filename) override;
diff --git a/Engine/Source/Runtime/Core/Public/HAL/MallocAnsi.h b/Engine/Source/Runtime/Core/Public/HAL/MallocAnsi.h
--- a/Engine/Source/Runtime/Core/Public/HAL/MallocAnsi.h
+++ b/Engine/Source/Runtime/Core/Public/HAL/MallocAnsi.h
@@ -122,7 +122,7 @@ public:
 	}
 
 // #if PLATFORM_MAC
-// 	virtual void GetAllocationInfo( FMemoryAllocationStats_DEPRECATED& MemStats ) OVERRIDE
+// 	virtual void GetAllocationInfo( FMemoryAllocationStats_DEPRECATED& MemStats ) override
 // 	{
 // 		// @todo implement in platform memory
 // 		// Just get memory information for the process and report the working set instead
@@ -132,7 +132,7 @@ public:
 // 		MemStats.TotalUsed___UsedPhysical = MemStats.TotalAllocated___UsedVirtual = TaskInfo.resident_size;
 // 	}
 // #elif PLATFORM_IOS
-// 	virtual void GetAllocationInfo( FMemoryAllocationStats_DEPRECATED& MemStats ) OVERRIDE
+// 	virtual void GetAllocationInfo( FMemoryAllocationStats_DEPRECATED& MemStats ) override
 // 	{
 // 		// @todo implement in platform memory
 // 		task_basic_info info;
diff --git a/Engine/Source/Runtime/Core/Public/HAL/Platform.h b/Engine/Source/Runtime/Core/Public/HAL/Platform.h
--- a/Engine/Source/Runtime/Core/Public/HAL/Platform.h
+++ b/Engine/Source/Runtime/Core/Public/HAL/Platform.h
@@ -294,10 +294,14 @@
 
 // Legacy method modifier macros.  You shouldn't use these macros in modern code.  Use the built-in keyword directly.
 #ifndef OVERRIDE
-	#define OVERRIDE override
+	#define OVERRIDE \
+		EMIT_DEPRECATED_WARNING_MESSAGE("OVERRIDE macro is deprecated. Please use override keyword instead.") \
+		override
 #endif
 #ifndef FINAL						
-	#define FINAL final
+	#define FINAL \
+		EMIT_DEPRECATED_WARNING_MESSAGE("FINAL macro is deprecated. Please use final keyword instead.") \
+		final
 #endif
 
 // Method modifiers
diff --git a/Engine/Source/Runtime/Core/Public/Windows/WindowsPlatformCompilerPreSetup.h b/Engine/Source/Runtime/Core/Public/Windows/WindowsPlatformCompilerPreSetup.h
--- a/Engine/Source/Runtime/Core/Public/Windows/WindowsPlatformCompilerPreSetup.h
+++ b/Engine/Source/Runtime/Core/Public/Windows/WindowsPlatformCompilerPreSetup.h
@@ -7,7 +7,15 @@
 #pragma once
 
 #ifndef DISABLE_DEPRECATION
-#define DEPRECATED(VERSION, MESSAGE) __declspec(deprecated(MESSAGE " Please update your code to the new API before upgrading to the next release, otherwise your project will no longer compile."))
+	#define DEPRECATED(VERSION, MESSAGE) __declspec(deprecated(MESSAGE " Please update your code to the new API before upgrading to the next release, otherwise your project will no longer compile."))
+
+	#ifndef EMIT_DEPRECATED_WARNING_MESSAGE_STR
+		#define EMIT_DEPRECATED_WARNING_MESSAGE_STR1(x) #x
+		#define EMIT_DEPRECATED_WARNING_MESSAGE_STR(x) EMIT_DEPRECATED_WARNING_MESSAGE_STR1(x)
+	#endif // EMIT_DEPRECATED_WARNING_MESSAGE_STR
+
+	#define EMIT_DEPRECATED_WARNING_MESSAGE(Msg) __pragma(message(__FILE__ "(" EMIT_DEPRECATED_WARNING_MESSAGE_STR(__LINE__) "): warning C4996: " Msg))
 #else
 	#define DEPRECATED(VERSION, MESSAGE)
+	#define EMIT_DEPRECATED_WARNING_MESSAGE(Msg)
 #endif // DISABLE_DEPRECATION
diff --git a/Engine/Source/Runtime/CoreUObject/Public/UObject/ObjectBase.h b/Engine/Source/Runtime/CoreUObject/Public/UObject/ObjectBase.h
--- a/Engine/Source/Runtime/CoreUObject/Public/UObject/ObjectBase.h
+++ b/Engine/Source/Runtime/CoreUObject/Public/UObject/ObjectBase.h
@@ -641,7 +641,7 @@ namespace UF
 		BlueprintNativeEvent,
 
 		/// This function is sealed and cannot be overridden in subclasses.
-		/// It is only a valid keyword for events; declare other methods as static or FINAL to indicate that they are sealed.
+		/// It is only a valid keyword for events; declare other methods as static or final to indicate that they are sealed.
 		SealedEvent,
 
 		/// This function is executable from the command line.
diff --git a/Engine/Source/Runtime/Engine/Classes/Engine/IntSerialization.h b/Engine/Source/Runtime/Engine/Classes/Engine/IntSerialization.h
--- a/Engine/Source/Runtime/Engine/Classes/Engine/IntSerialization.h
+++ b/Engine/Source/Runtime/Engine/Classes/Engine/IntSerialization.h
@@ -12,7 +12,7 @@ class ENGINE_API UIntSerialization	: public UObject
 public:
 //	ENGINE_API UIntSerialization(const class FPostConstructInitializeProperties& PCIP, uint8 InUnsignedInt8Variable, uint16 InUnsignedInt16Variable, uint32 InUnsignedInt32Variable, uint64 InUnsignedInt64Variable, int8 InSignedInt8Variable, int16 InSignedInt16Variable, int InSignedInt32Variable, int64 InSignedInt64Variable);
 
-//	virtual void Serialize (FArchive& Ar) OVERRIDE;
+//	virtual void Serialize (FArchive& Ar) override;
 	/**
 	 * Serializes the given int Int Serilization Test test from or into the specified archive.
 	 *
diff --git a/Engine/Source/Runtime/Sockets/Public/IPAddress.h b/Engine/Source/Runtime/Sockets/Public/IPAddress.h
--- a/Engine/Source/Runtime/Sockets/Public/IPAddress.h
+++ b/Engine/Source/Runtime/Sockets/Public/IPAddress.h
@@ -101,7 +101,7 @@ public:
 	 *
 	 * @return true if the ip address was parseable, false otherwise
 	 */
-// 	static bool IsValidIp(const TCHAR* IpAddr) OVERRIDE
+// 	static bool IsValidIp(const TCHAR* IpAddr) override
 // 	{
 // 		FInternetAddr LocalAddr;
 // 		bool bIsValid = false;
