Missed file from CL# 2106993

diff --git a/Engine/Source/Editor/BlueprintGraph/Private/EdGraphSchema_K2.cpp b/Engine/Source/Editor/BlueprintGraph/Private/EdGraphSchema_K2.cpp
--- a/Engine/Source/Editor/BlueprintGraph/Private/EdGraphSchema_K2.cpp
+++ b/Engine/Source/Editor/BlueprintGraph/Private/EdGraphSchema_K2.cpp
@@ -507,7 +507,7 @@ bool UEdGraphSchema_K2::IsActorValidForLevelScriptRefs(const AActor* TestActor,
 	return TestActor
 		&& (TestActor->GetLevel() == Blueprint->GetLevel())
 		&& FKismetEditorUtilities::IsActorValidForLevelScript(TestActor);
-}
+		}
 
 void UEdGraphSchema_K2::ReplaceSelectedNode(UEdGraphNode* SourceNode, AActor* TargetActor)
 {
@@ -625,81 +625,81 @@ bool UEdGraphSchema_K2::IsAllowableBlueprintVariableType(const class UClass* InC
 
 		// UObject is an exception, and is always a blueprint-able type
 		if(InClass == UObject::StaticClass())
-		{
+{
 			return true;
-		}
+}
 
 		struct FBlueprintIsNotBlueprintTypeHelper
-		{
+	{
 			bool bNotBlueprintType;
 
 			FBlueprintIsNotBlueprintTypeHelper() : bNotBlueprintType(false)
-			{
+{
 				GConfig->GetBool(TEXT("EditoronlyBP"), TEXT("bBlueprintIsNotBlueprintType"), bNotBlueprintType, GEditorIni);
-			}
+}
 		};
 		static const FBlueprintIsNotBlueprintTypeHelper BPHelper;
 		if (BPHelper.bNotBlueprintType && InClass->IsChildOf(UBlueprint::StaticClass()))
-		{
+{
 			return false;
-		}
-
+	}
+	
 		const UClass* ParentClass = InClass;
 		while(ParentClass)
 		{
 			// Climb up the class hierarchy and look for "BlueprintType" and "NotBlueprintType" to see if this class is allowed.
 			if(ParentClass->GetBoolMetaData(FBlueprintMetadata::MD_AllowableBlueprintVariableType)
 				|| ParentClass->HasMetaData(FBlueprintMetadata::MD_BlueprintSpawnableComponent))
-			{
+		{
 				return true;
-			}
+		}
 			else if(ParentClass->GetBoolMetaData(FBlueprintMetadata::MD_NotAllowableBlueprintVariableType))
-			{
+		{
 				return false;
-			}
+				}
 			ParentClass = ParentClass->GetSuperClass();
+			}
 		}
-	}
-	
+
 	return false;
-}
+		}
 
 bool UEdGraphSchema_K2::IsAllowableBlueprintVariableType(const class UScriptStruct *InStruct)
-{
+		{
 	if (auto UDStruct = Cast<const UUserDefinedStruct>(InStruct))
-	{
+		{
 		if (EUserDefinedStructureStatus::UDSS_UpToDate != UDStruct->Status.GetValue())
 		{
 			return false;
 		}
-	}
+				}
 	return InStruct && (InStruct->GetBoolMetaDataHierarchical(FBlueprintMetadata::MD_AllowableBlueprintVariableType));
-}
+			}
 
 bool UEdGraphSchema_K2::DoesGraphSupportImpureFunctions(const UEdGraph* InGraph) const
-{
+			{
 	const EGraphType GraphType = GetGraphType(InGraph);
 	const bool bAllowImpureFuncs = GraphType != GT_Animation; //@TODO: It's really more nuanced than this (e.g., in a function someone wants to write as pure)
 
 	return bAllowImpureFuncs;
-}
+		}
 
 bool UEdGraphSchema_K2::IsPropertyExposedOnSpawn(const UProperty* Property)
-{
+			{
 	if (Property)
-	{
+				{
 		const bool bMeta = Property->HasMetaData(FBlueprintMetadata::MD_ExposeOnSpawn);
 		const bool bFlag = Property->HasAllPropertyFlags(CPF_ExposeOnSpawn);
 		if (bMeta != bFlag)
-		{
+				{
 			UE_LOG(LogBlueprint, Warning
 				, TEXT("ExposeOnSpawn ambiguity. Property '%s', MetaData '%s', Flag '%s'")
 				, *Property->GetFullName()
 				, bMeta ? *GTrue.ToString() : *GFalse.ToString()
 				, bFlag ? *GTrue.ToString() : *GFalse.ToString());
-		}
+				}
 		return bMeta || bFlag;
-	}
+			}
 	return false;
 }
 
@@ -728,6 +728,39 @@ static bool IsUsingNonExistantVariable(const UEdGraphNode* InGraphNode, UBluepri
 	return bNonExistantVariable;
 }
 
+bool UEdGraphSchema_K2::PinHasSplittableStructType(const UEdGraphPin* InGraphPin) const
+{
+	const FEdGraphPinType& PinType = InGraphPin->PinType;
+	bool bCanSplit = (!PinType.bIsArray && !PinType.bIsReference && InGraphPin->LinkedTo.Num() == 0 && PinType.PinCategory == PC_Struct);
+
+	if (bCanSplit)
+	{
+		UScriptStruct* StructType = CastChecked<UScriptStruct>(InGraphPin->PinType.PinSubCategoryObject.Get());
+		if (InGraphPin->Direction == EGPD_Input)
+		{
+			bCanSplit = UK2Node_MakeStruct::CanBeMade(StructType);
+			if (!bCanSplit)
+			{
+				const FString& MetaData = StructType->GetMetaData(TEXT("HasNativeMake"));
+				UFunction* Function = FindObject<UFunction>(NULL, *MetaData, true);
+				bCanSplit = (Function != NULL);
+			}
+		}
+		else
+		{
+			bCanSplit = UK2Node_BreakStruct::CanBeBroken(StructType);
+			if (!bCanSplit)
+			{
+				const FString& MetaData = StructType->GetMetaData(TEXT("HasNativeBreak"));
+				UFunction* Function = FindObject<UFunction>(NULL, *MetaData, true);
+				bCanSplit = (Function != NULL);
+			}
+		}
+	}
+
+	return bCanSplit;
+}
+
 void UEdGraphSchema_K2::GetContextMenuActions(const UEdGraph* CurrentGraph, const UEdGraphNode* InGraphNode, const UEdGraphPin* InGraphPin, FMenuBuilder* MenuBuilder, bool bIsDebugging) const
 {
 	check(CurrentGraph);
@@ -738,56 +771,66 @@ void UEdGraphSchema_K2::GetContextMenuActions(const UEdGraph* CurrentGraph, cons
 		MenuBuilder->BeginSection("EdGraphSchemaPinActions", LOCTEXT("PinActionsMenuHeader", "Pin Actions"));
 		{
 			if (!bIsDebugging)
-			{
-				// Break pin links
-				if (InGraphPin->LinkedTo.Num() > 1)
-				{
-					MenuBuilder->AddMenuEntry( FGraphEditorCommands::Get().BreakPinLinks );
-				}
-	
-				// Add the change pin type action, if this is a select node
-				if (InGraphNode->IsA(UK2Node_Select::StaticClass()))
-				{
-					MenuBuilder->AddMenuEntry(FGraphEditorCommands::Get().ChangePinType);
-				}
-	
-				// add sub menu for break link to
-				if (InGraphPin->LinkedTo.Num() > 0)
-				{
-					if(InGraphPin->LinkedTo.Num() > 1)
-					{
-						MenuBuilder->AddSubMenu(
-							LOCTEXT("BreakLinkTo", "Break Link To..."),
-							LOCTEXT("BreakSpecificLinks", "Break a specific link..."),
-							FNewMenuDelegate::CreateUObject( (UEdGraphSchema_K2*const)this, &UEdGraphSchema_K2::GetBreakLinkToSubMenuActions, const_cast<UEdGraphPin*>(InGraphPin)));
+		    {
+			    // Break pin links
+			    if (InGraphPin->LinkedTo.Num() > 1)
+			    {
+				    MenuBuilder->AddMenuEntry( FGraphEditorCommands::Get().BreakPinLinks );
+			    }
+    
+			    // Add the change pin type action, if this is a select node
+			    if (InGraphNode->IsA(UK2Node_Select::StaticClass()))
+			    {
+				    MenuBuilder->AddMenuEntry(FGraphEditorCommands::Get().ChangePinType);
+			    }
+    
+			    // add sub menu for break link to
+			    if (InGraphPin->LinkedTo.Num() > 0)
+			    {
+				    if(InGraphPin->LinkedTo.Num() > 1)
+				    {
+					    MenuBuilder->AddSubMenu(
+						    LOCTEXT("BreakLinkTo", "Break Link To..."),
+						    LOCTEXT("BreakSpecificLinks", "Break a specific link..."),
+						    FNewMenuDelegate::CreateUObject( (UEdGraphSchema_K2*const)this, &UEdGraphSchema_K2::GetBreakLinkToSubMenuActions, const_cast<UEdGraphPin*>(InGraphPin)));
 
 						MenuBuilder->AddSubMenu(
 							LOCTEXT("JumpToConnection", "Jump to Connection..."),
 							LOCTEXT("JumpToSpecificConnection", "Jump to specific connection..."),
 							FNewMenuDelegate::CreateUObject( (UEdGraphSchema_K2*const)this, &UEdGraphSchema_K2::GetJumpToConnectionSubMenuActions, const_cast<UEdGraphPin*>(InGraphPin)));
-					}
-					else
-					{
-						((UEdGraphSchema_K2*const)this)->GetBreakLinkToSubMenuActions(*MenuBuilder, const_cast<UEdGraphPin*>(InGraphPin));
+				    }
+				    else
+				    {
+					    ((UEdGraphSchema_K2*const)this)->GetBreakLinkToSubMenuActions(*MenuBuilder, const_cast<UEdGraphPin*>(InGraphPin));
 						((UEdGraphSchema_K2*const)this)->GetJumpToConnectionSubMenuActions(*MenuBuilder, const_cast<UEdGraphPin*>(InGraphPin));
-					}
-				}
-	
-				// Conditionally add the var promotion pin if this is an output pin and it's not an exec pin
-				if (InGraphPin->PinType.PinCategory != PC_Exec)
+				    }
+			    }
+    
+			    // Conditionally add the var promotion pin if this is an output pin and it's not an exec pin
+			    if (InGraphPin->PinType.PinCategory != PC_Exec)
+			    {
+				    MenuBuilder->AddMenuEntry( FGraphEditorCommands::Get().PromoteToVariable );
+			    }
+    
+				if (PinHasSplittableStructType(InGraphPin))
 				{
-					MenuBuilder->AddMenuEntry( FGraphEditorCommands::Get().PromoteToVariable );
+					MenuBuilder->AddMenuEntry( FGraphEditorCommands::Get().SplitStructPin );
 				}
-	
-				// Conditionally add the execution path pin removal if this is an execution branching node
-				if( InGraphPin->Direction == EGPD_Output && InGraphPin->GetOwningNode())
+
+				if (InGraphPin->ParentPin != NULL)
 				{
-					if ( InGraphPin->GetOwningNode()->IsA(UK2Node_ExecutionSequence::StaticClass()) ||  InGraphPin->GetOwningNode()->IsA(UK2Node_Switch::StaticClass()) )
-					{
-						MenuBuilder->AddMenuEntry( FGraphEditorCommands::Get().RemoveExecutionPin );
-					}
+					MenuBuilder->AddMenuEntry( FGraphEditorCommands::Get().RecombineStructPin );
 				}
-			}
+    
+			    // Conditionally add the execution path pin removal if this is an execution branching node
+			    if( InGraphPin->Direction == EGPD_Output && InGraphPin->GetOwningNode())
+			    {
+				    if ( InGraphPin->GetOwningNode()->IsA(UK2Node_ExecutionSequence::StaticClass()) ||  InGraphPin->GetOwningNode()->IsA(UK2Node_Switch::StaticClass()) )
+				    {
+					    MenuBuilder->AddMenuEntry( FGraphEditorCommands::Get().RemoveExecutionPin );
+				    }
+			    }
+		    }
 		}
 		MenuBuilder->EndSection();
 
@@ -917,20 +960,20 @@ void UEdGraphSchema_K2::GetContextMenuActions(const UEdGraph* CurrentGraph, cons
 			if (const UK2Node* K2Node = Cast<const UK2Node>(InGraphNode))
 			{
 				if (!K2Node->IsNodePure())
+			{
+				MenuBuilder->BeginSection("EdGraphSchemaBreakpoints", LOCTEXT("BreakpointsHeader", "Breakpoints"));
 				{
-					MenuBuilder->BeginSection("EdGraphSchemaBreakpoints", LOCTEXT("BreakpointsHeader", "Breakpoints"));
-					{
-						MenuBuilder->AddMenuEntry( FGraphEditorCommands::Get().ToggleBreakpoint );
-						MenuBuilder->AddMenuEntry( FGraphEditorCommands::Get().AddBreakpoint );
-						MenuBuilder->AddMenuEntry( FGraphEditorCommands::Get().RemoveBreakpoint );
-						MenuBuilder->AddMenuEntry( FGraphEditorCommands::Get().EnableBreakpoint );
-						MenuBuilder->AddMenuEntry( FGraphEditorCommands::Get().DisableBreakpoint );
-					}
-					MenuBuilder->EndSection();
+					MenuBuilder->AddMenuEntry( FGraphEditorCommands::Get().ToggleBreakpoint );
+					MenuBuilder->AddMenuEntry( FGraphEditorCommands::Get().AddBreakpoint );
+					MenuBuilder->AddMenuEntry( FGraphEditorCommands::Get().RemoveBreakpoint );
+					MenuBuilder->AddMenuEntry( FGraphEditorCommands::Get().EnableBreakpoint );
+					MenuBuilder->AddMenuEntry( FGraphEditorCommands::Get().DisableBreakpoint );
 				}
+				MenuBuilder->EndSection();
 			}
 		}
 	}
+	}
 
 	Super::GetContextMenuActions(CurrentGraph, InGraphNode, InGraphPin, MenuBuilder, bIsDebugging);
 }
@@ -1062,11 +1105,11 @@ void UEdGraphSchema_K2::GetJumpToConnectionSubMenuActions( class FMenuBuilder& M
 
 	// Add all the links we could break from
 	for(auto PinLink : InGraphPin->LinkedTo )
-	{
+				{
 		FText Title = PinLink->GetOwningNode()->GetNodeTitle(ENodeTitleType::ListView);
 		FString TitleString = Title.ToString();
 		if ( PinLink->PinName != TEXT("") )
-		{
+					{
 			TitleString = FString::Printf(TEXT("%s (%s)"), *TitleString, *PinLink->PinName);
 
 			// Add name of connection if possible
@@ -1084,13 +1127,13 @@ void UEdGraphSchema_K2::GetJumpToConnectionSubMenuActions( class FMenuBuilder& M
 		Args.Add( TEXT("NumberOfNodes"), Count );
 
 		if ( Count == 0 )
-		{
+			{
 			Description = FText::Format( LOCTEXT("JumpDesc", "Jump to {NodeTitle}"), Args );
-		}
-		else
-		{
+							}
+							else
+							{
 			Description = FText::Format( LOCTEXT("JumpDescMulti", "Jump to {NodeTitle} ({NumberOfNodes})"), Args );
-		}
+					}
 		++Count;
 
 		MenuBuilder.AddMenuEntry( Description, Description, FSlateIcon(), FUIAction(
@@ -1099,7 +1142,7 @@ void UEdGraphSchema_K2::GetJumpToConnectionSubMenuActions( class FMenuBuilder& M
 }
 
 void UEdGraphSchema_K2::GetGraphContextActions(FGraphContextMenuBuilder& ContextMenuBuilder) const
-{
+				{
 	FBlueprintGraphActionListBuilder BlueprintContextMenuBuilder(ContextMenuBuilder.CurrentGraph);
 	BlueprintContextMenuBuilder.FromPin = ContextMenuBuilder.FromPin;
 	BlueprintContextMenuBuilder.SelectedObjects.Append(ContextMenuBuilder.SelectedObjects);
@@ -1114,13 +1157,13 @@ void UEdGraphSchema_K2::GetGraphContextActions(FGraphContextMenuBuilder& Context
 }
 
 void UEdGraphSchema_K2::GetAllActions(FBlueprintPaletteListBuilder& PaletteBuilder)
-{
+	{
 	const UEdGraphSchema_K2* K2SchemaInst = GetDefault<UEdGraphSchema_K2>();
 	FK2ActionMenuBuilder(K2SchemaInst).GetAllActions(PaletteBuilder);
-}
+	}
 
 void UEdGraphSchema_K2::GetPaletteActions(FBlueprintPaletteListBuilder& ActionMenuBuilder, TWeakObjectPtr<UClass> FilterClass/* = NULL*/)
-{
+	{
 	const UEdGraphSchema_K2* K2SchemaInst = GetDefault<UEdGraphSchema_K2>();
 	FK2ActionMenuBuilder(K2SchemaInst).GetPaletteActions(ActionMenuBuilder, FilterClass);
 }
@@ -1665,13 +1708,13 @@ FString UEdGraphSchema_K2::IsPinDefaultValid(const UEdGraphPin* Pin, const FStri
 
 	if (OwningBP->BlueprintType != BPTYPE_Interface)
 	{
-		if (!bIsAutoCreateRefTerm)
+		if( !bIsAutoCreateRefTerm )
 		{
-			if (bIsArray)
+			if( bIsArray )
 			{
 				return TEXT("Literal values are not allowed for array parameters.  Use a Make Array node instead");
 			}
-			else if (bIsReference)
+			else if( bIsReference )
 			{
 				return TEXT("Literal values are not allowed for pass-by-reference parameters.");
 			}
@@ -1852,7 +1895,7 @@ bool UEdGraphSchema_K2::DefaultValueSimpleValidation(const FEdGraphPinType& PinT
 		// It should not be possible to get a culture invariant text in editor
 		if(GIsEditor)
 		{
-			if( InNewDefaultText.IsCultureInvariant() )
+		if( InNewDefaultText.IsCultureInvariant() )
 		{
 			DVSV_RETURN_MSG( TEXT("Invalid text literal, text is culture invariant!") );
 		}
@@ -2352,66 +2395,66 @@ FString UEdGraphSchema_K2::TypeToString(UProperty* const Property)
 		if (Struct->Struct)
 		{
 		return FString::Printf(TEXT("struct'%s'"), *Struct->Struct->GetName());
-	}
+		}
 	}
 	else if (UClassProperty* Class = Cast<UClassProperty>(Property))
-	{
-		if (Class->MetaClass != nullptr)
 		{
+		if (Class->MetaClass != nullptr)
+			{
 		return FString::Printf(TEXT("class'%s'"), *Class->MetaClass->GetName());
 	}
-	}
+			}
 	else if (UInterfaceProperty* Interface = Cast<UInterfaceProperty>(Property))
 	{
 		if (Interface->InterfaceClass != nullptr)
-		{
+			{
 		return FString::Printf(TEXT("interface'%s'"), *Interface->InterfaceClass->GetName());
-	}
-	}
+			}
+		}
 	else if (UObjectPropertyBase* Obj = Cast<UObjectPropertyBase>(Property))
 	{
 		if( Obj->PropertyClass )
 		{
 			if( Property->IsA(UWeakObjectProperty::StaticClass()) )
-			{
+	{
 				return FString::Printf(TEXT("weak_ptr_object'%s'"), *Obj->PropertyClass->GetName());
-			}
-			else
-			{
+	}
+	else
+	{
 				return FString::Printf(TEXT("object'%s'"), *Obj->PropertyClass->GetName());
 			}
-		}
+	}
 
 		return TEXT("");
 	}
 	else if (UArrayProperty* Array = Cast<UArrayProperty>(Property))
 	{
 		if (Array->Inner)
-		{
+	{
 		return FString::Printf(TEXT("array[%s]"), *TypeToString(Array->Inner)); 
 	}
 	}
-	
+
 		return Property->GetClass()->GetName();
-	}
+}
 
 FText UEdGraphSchema_K2::TypeToText(const FEdGraphPinType& Type)
 {
 	FText PropertyText;
 
 	if (Type.PinSubCategoryObject != NULL)
-	{
+{
 		const UEdGraphSchema_K2* Schema = GetDefault<UEdGraphSchema_K2>();
 		if (Type.PinCategory == Schema->PC_Byte)
-		{
+	{
 			FFormatNamedArguments Args;
 			Args.Add(TEXT("EnumName"), FText::FromString(Type.PinSubCategoryObject->GetName()));
 			PropertyText = FText::Format(LOCTEXT("EnumAsText", "enum'{EnumName}'"), Args);
-		}
+	}
 		else
 		{
 			if( !Type.bIsWeakPointer )
-			{
+	{
 				UClass* PSCOAsClass = Cast<UClass>(Type.PinSubCategoryObject.Get());
 				const bool bIsInterface = PSCOAsClass && PSCOAsClass->HasAnyClassFlags(CLASS_Interface);
 
@@ -2421,9 +2464,9 @@ FText UEdGraphSchema_K2::TypeToText(const FEdGraphPinType& Type)
 				Args.Add(TEXT("Category"), CategoryDesc);
 				Args.Add(TEXT("ObjectName"), FText::FromString(Type.PinSubCategoryObject.Get()->GetName()));
 				PropertyText = FText::Format(LOCTEXT("ObjectAsText", "{Category}'{ObjectName}'"), Args);
-			}
+	}
 			else
-			{
+	{
 				FFormatNamedArguments Args;
 				Args.Add(TEXT("Category"), FText::FromString(Type.PinCategory));
 				Args.Add(TEXT("ObjectName"), FText::FromString(Type.PinSubCategoryObject.Get()->GetName()));
@@ -2437,11 +2480,11 @@ FText UEdGraphSchema_K2::TypeToText(const FEdGraphPinType& Type)
 		Args.Add(TEXT("Category"), FText::FromString(Type.PinCategory));
 		Args.Add(TEXT("ObjectName"), FText::FromString(Type.PinSubCategory));
 		PropertyText = FText::Format(LOCTEXT("ObjectAsText", "{Category}'{ObjectName}'"), Args);
-	}
-	else
-	{
+			}
+			else
+			{
 		PropertyText = FText::FromString(Type.PinCategory);
-	}
+		}
 
 	if (Type.bIsArray)
 	{
@@ -2481,9 +2524,9 @@ void UEdGraphSchema_K2::GetVariableTypeTree( TArray< TSharedPtr<FPinTypeTreeInfo
 	UScriptStruct* VectorStruct = FindObjectChecked<UScriptStruct>(UObject::StaticClass(), TEXT("Vector"));
 	UScriptStruct* RotatorStruct = FindObjectChecked<UScriptStruct>(UObject::StaticClass(), TEXT("Rotator"));
 	UScriptStruct* TransformStruct = FindObjectChecked<UScriptStruct>(UObject::StaticClass(), TEXT("Transform"));
-	TypeTree.Add( MakeShareable( new FPinTypeTreeInfo(PC_Struct, VectorStruct, LOCTEXT("VectorType", "A 3D vector").ToString()) ) );
-	TypeTree.Add( MakeShareable( new FPinTypeTreeInfo(PC_Struct, RotatorStruct, LOCTEXT("RotatorType", "A 3D rotation").ToString()) ) );
-	TypeTree.Add( MakeShareable( new FPinTypeTreeInfo(PC_Struct, TransformStruct, LOCTEXT("TransformType", "A 3D transformation, including translation, rotation and 3D scale.").ToString()) ) );
+ 	TypeTree.Add( MakeShareable( new FPinTypeTreeInfo(PC_Struct, VectorStruct, LOCTEXT("VectorType", "A 3D vector").ToString()) ) );
+ 	TypeTree.Add( MakeShareable( new FPinTypeTreeInfo(PC_Struct, RotatorStruct, LOCTEXT("RotatorType", "A 3D rotation").ToString()) ) );
+ 	TypeTree.Add( MakeShareable( new FPinTypeTreeInfo(PC_Struct, TransformStruct, LOCTEXT("TransformType", "A 3D transformation, including translation, rotation and 3D scale.").ToString()) ) );
 	
 	// Add wildcard type
 	if (bAllowWildCard)
@@ -2815,7 +2858,7 @@ bool UEdGraphSchema_K2::CanEncapuslateNode(UEdGraphNode const& TestNode) const
 {
 	// Can't encapsulate entry points (may relax this restriction in the future, but it makes sense for now)
 	return !TestNode.IsA(UK2Node_FunctionTerminator::StaticClass()) && 
-			TestNode.GetClass() != UK2Node_Tunnel::StaticClass(); //Tunnel nodes getting sucked into collapsed graphs fails badly, want to allow derived types though(composite node/Macroinstances)
+		    TestNode.GetClass() != UK2Node_Tunnel::StaticClass(); //Tunnel nodes getting sucked into collapsed graphs fails badly, want to allow derived types though(composite node/Macroinstances)
 }
 
 void UEdGraphSchema_K2::HandleGraphBeingDeleted(UEdGraph& GraphBeingRemoved) const
@@ -2964,6 +3007,34 @@ bool UEdGraphSchema_K2::ShouldHidePinDefaultValue(UEdGraphPin* Pin) const
 	return false;
 }
 
+void UEdGraphSchema_K2::SetPinDefaultValue(UEdGraphPin* Pin, const UFunction* Function, const UProperty* Param) const
+{
+	if (Function && Param)
+	{
+		const FString MetadataDefaultValue = Function->GetMetaData(*Param->GetName());
+		if (!MetadataDefaultValue.IsEmpty())
+		{
+			// Specified default value in the metadata
+			Pin->DefaultValue = Pin->AutogeneratedDefaultValue = MetadataDefaultValue;
+		}
+		else
+		{
+			const FName MetadataCppDefaultValueKey( *(FString(TEXT("CPP_Default_")) + Param->GetName()) );
+			const FString MetadataCppDefaultValue = Function->GetMetaData(MetadataCppDefaultValueKey);
+			if(!MetadataCppDefaultValue.IsEmpty())
+			{
+				Pin->DefaultValue = Pin->AutogeneratedDefaultValue = MetadataCppDefaultValue;
+			}
+		}
+	}
+	
+	if (Pin->DefaultValue.Len() == 0)
+	{
+		// Set the default value to (T)0
+		SetPinDefaultValueBasedOnType(Pin);
+	}
+}
+
 void UEdGraphSchema_K2::SetPinDefaultValueBasedOnType(UEdGraphPin* Pin) const
 {
 	UScriptStruct* VectorStruct = FindObjectChecked<UScriptStruct>(UObject::StaticClass(), TEXT("Vector"));
@@ -3175,13 +3246,13 @@ bool UEdGraphSchema_K2::CanPromotePinToVariable( const UEdGraphPin& Pin ) const
 		else if ((PinType.PinCategory == PC_Object) || (PinType.PinCategory == PC_Interface))
 		{
 			if (PinType.PinSubCategoryObject != NULL)
+		{
+			if (UClass* Class = Cast<UClass>(PinType.PinSubCategoryObject.Get()))
 			{
-				if (UClass* Class = Cast<UClass>(PinType.PinSubCategoryObject.Get()))
-				{
-					bCanPromote = UEdGraphSchema_K2::IsAllowableBlueprintVariableType(Class);
-				}	
-			}
+				bCanPromote = UEdGraphSchema_K2::IsAllowableBlueprintVariableType(Class);
+			}	
 		}
+	}
 		else if ((PinType.PinCategory == PC_Struct) && (PinType.PinSubCategoryObject != NULL))
 		{
 			if (UScriptStruct* Struct = Cast<UScriptStruct>(PinType.PinSubCategoryObject.Get()))
@@ -3194,6 +3265,38 @@ bool UEdGraphSchema_K2::CanPromotePinToVariable( const UEdGraphPin& Pin ) const
 	return bCanPromote;
 }
 
+bool UEdGraphSchema_K2::CanSplitStructPin( const UEdGraphPin& Pin ) const
+{
+	return (Pin.LinkedTo.Num() == 0 && PinHasSplittableStructType(&Pin));
+}
+
+bool UEdGraphSchema_K2::CanRecombineStructPin( const UEdGraphPin& Pin ) const
+{
+	bool bCanRecombine = (Pin.ParentPin != NULL && Pin.LinkedTo.Num() == 0);
+	if (bCanRecombine)
+	{
+		// Go through all the other subpins and ensure they also are not connected to anything
+		TArray<UEdGraphPin*> PinsToExamine = Pin.ParentPin->SubPins;
+
+		int32 PinIndex = 0;
+		while (bCanRecombine && PinIndex < PinsToExamine.Num())
+		{
+			UEdGraphPin* SubPin = PinsToExamine[PinIndex];
+			if (SubPin->LinkedTo.Num() > 0)
+			{
+				bCanRecombine = false;
+			}
+			else if (SubPin->SubPins.Num() > 0)
+			{
+				PinsToExamine.Append(SubPin->SubPins);
+			}
+			++PinIndex;
+		}
+	}
+
+	return bCanRecombine;
+}
+
 void UEdGraphSchema_K2::GetGraphDisplayInformation(const UEdGraph& Graph, /*out*/ FGraphDisplayInfo& DisplayInfo) const
 {
 	DisplayInfo.DocLink = TEXT("Shared/Editors/BlueprintEditor/GraphTypes");
@@ -3224,22 +3327,22 @@ void UEdGraphSchema_K2::GetGraphDisplayInformation(const UEdGraph& Graph, /*out*
 	}
 	else if (GraphType == GT_Function)
 	{
-		if ( Graph.GetFName() == FN_UserConstructionScript )
-		{
+	if ( Graph.GetFName() == FN_UserConstructionScript )
+	{
 			DisplayInfo.PlainName = LOCTEXT("GraphDisplayName_ConstructionScript", "ConstructionScript");
 
 			DisplayInfo.Tooltip = LOCTEXT("GraphTooltip_ConstructionScript", "Function executed when Blueprint is placed or modified.").ToString();
 			DisplayInfo.DocExcerptName = TEXT("ConstructionScript");
-		}
+	}
 		else
-		{
+	{
 			// If we found a function from this graph, grab its tooltip
 			if (Function)
-			{
+	{
 				DisplayInfo.Tooltip = UK2Node_CallFunction::GetDefaultTooltipForFunction(Function);
-			}
-			else
-			{
+	}
+	else
+	{
 				DisplayInfo.Tooltip = Graph.GetName();
 			}
 
@@ -3262,16 +3365,16 @@ void UEdGraphSchema_K2::GetGraphDisplayInformation(const UEdGraph& Graph, /*out*
 		DisplayInfo.DocExcerptName = TEXT("AnimGraph");
 	}
 	else if (GraphType == GT_StateMachine)
-	{
+		{
 		DisplayInfo.Tooltip = Graph.GetName();
 		DisplayInfo.DocExcerptName = TEXT("StateMachine");
 	}
 
 	// Add pure to notes if set
 	if (Function && Function->HasAnyFunctionFlags(FUNC_BlueprintPure))
-	{
-		DisplayInfo.Notes.Add(TEXT("pure"));
-	}
+			{
+				DisplayInfo.Notes.Add(TEXT("pure"));
+			}
 
 	// Mark transient graphs as obviously so
 	if (Graph.HasAllFlags(RF_Transient))
@@ -3500,7 +3603,7 @@ struct FBackwardCompatibilityConversionHelper
 
 		auto OldBlueprintPin = OldNode->FindPin(BlueprintPinName);
 		if (!OldBlueprintPin)
-		{
+	{
 			UE_LOG(LogBlueprint, Warning, TEXT("BackwardCompatibilityNodeConversion Error bp: '%s' node: '%s'. No bp pin found '%s'"),
 				Blueprint ? *Blueprint->GetName() : TEXT("Unknown"),
 				*OldNode->GetName(),
@@ -3627,44 +3730,44 @@ struct FBackwardCompatibilityConversionHelper
 			}
 
 			// MOVE OTHER PINS
-			TArray<UEdGraphPin*> OldPins;
-			OldPins.Add(OldBlueprintPin);
+				TArray<UEdGraphPin*> OldPins;
+				OldPins.Add(OldBlueprintPin);
 			for (auto PinIter = NewNode->Pins.CreateIterator(); PinIter; ++PinIter)
-			{
-				UEdGraphPin* const Pin = *PinIter;
-				check(Pin);
-				if (ClassPin != Pin)
 				{
-					const auto OldPin = OldNode->FindPin(Pin->PinName);
-					if (NULL != OldPin)
+					UEdGraphPin* const Pin = *PinIter;
+					check(Pin);
+					if (ClassPin != Pin)
 					{
-						OldPins.Add(OldPin);
-						if (!Schema.MovePinLinks(*OldPin, *Pin).CanSafeConnect())
+					const auto OldPin = OldNode->FindPin(Pin->PinName);
+						if(NULL != OldPin)
 						{
+							OldPins.Add(OldPin);
+						if (!Schema.MovePinLinks(*OldPin, *Pin).CanSafeConnect())
+							{
 							UE_LOG(LogBlueprint, Warning, TEXT("BackwardCompatibilityNodeConversion Error 'cannot connect' in blueprint: %s, pin: %s"),
-								Blueprint ? *Blueprint->GetName() : TEXT("Unknown"),
-								*Pin->PinName);
+									Blueprint ? *Blueprint->GetName() : TEXT("Unknown"), 
+									*Pin->PinName);
+							}
 						}
-					}
-					else
-					{
+						else
+						{
 						UE_LOG(LogBlueprint, Warning, TEXT("BackwardCompatibilityNodeConversion Error 'missing old pin' in blueprint: %s"),
-							Blueprint ? *Blueprint->GetName() : TEXT("Unknown"),
-							Pin ? *Pin->PinName : TEXT("Unknown"));
+								Blueprint ? *Blueprint->GetName() : TEXT("Unknown"), 
+								Pin ? *Pin->PinName : TEXT("Unknown"));
+						}
 					}
 				}
-			}
 			OldNode->BreakAllNodeLinks();
 			for (auto PinIter = OldNode->Pins.CreateIterator(); PinIter; ++PinIter)
-			{
-				if (!OldPins.Contains(*PinIter))
 				{
-					UEdGraphPin* Pin = *PinIter;
+					if(!OldPins.Contains(*PinIter))
+					{
+						UEdGraphPin* Pin = *PinIter;
 					UE_LOG(LogBlueprint, Warning, TEXT("BackwardCompatibilityNodeConversion Error 'missing new pin' in blueprint: %s"),
-						Blueprint ? *Blueprint->GetName() : TEXT("Unknown"),
-						Pin ? *Pin->PinName : TEXT("Unknown"));
+							Blueprint ? *Blueprint->GetName() : TEXT("Unknown"), 
+							Pin ? *Pin->PinName : TEXT("Unknown"));
+					}
 				}
-			}
 			Graph->RemoveNode(OldNode);
 			return true;
 		}
@@ -3773,7 +3876,7 @@ void UEdGraphSchema_K2::BackwardCompatibilityNodeConversion(UEdGraph* Graph, boo
 			}
 
 			// Check to see if the struct has a native make/break that we should try to convert to.
-			if(OldMakeStructNode->StructType && OldMakeStructNode->StructType->GetBoolMetaData(TEXT("HasNativeMakeBreak")))
+			if(OldMakeStructNode->StructType && OldMakeStructNode->StructType->GetBoolMetaData(TEXT("HasNativeMake")))
 			{
 				UFunction* MakeNodeFunction = NULL;
 
@@ -3853,8 +3956,8 @@ void UEdGraphSchema_K2::BackwardCompatibilityNodeConversion(UEdGraph* Graph, boo
 				continue;
 			}
 
-			// Check to see if the struct has a native make/break that we should try to convert to.
-			if(OldBreakStructNode->StructType && OldBreakStructNode->StructType->GetBoolMetaData(TEXT("HasNativeMakeBreak")))
+				// Check to see if the struct has a native make/break that we should try to convert to.
+			if(OldBreakStructNode->StructType && OldBreakStructNode->StructType->GetBoolMetaData(TEXT("HasNativeBreak")))
 			{
 				UFunction* BreakNodeFunction = NULL;
 
@@ -4223,11 +4326,13 @@ float UEdGraphSchema_K2::EstimateNodeHeight( UEdGraphNode* Node )
 }
 
 
-bool UEdGraphSchema_K2::CollapseGatewayNode(UK2Node* InNode, UEdGraphNode* InEntryNode, UEdGraphNode* InResultNode) const
+bool UEdGraphSchema_K2::CollapseGatewayNode(UK2Node* InNode, UEdGraphNode* InEntryNode, UEdGraphNode* InResultNode, FKismetCompilerContext* CompilerContext) const
 {
 	bool bSuccessful = true;
 
-	for (int32 BoundaryPinIndex = 0; BoundaryPinIndex < InNode->Pins.Num(); ++BoundaryPinIndex)
+	// We iterate the array in reverse so we can both remove the subpins safely after we've read them and
+	// so we have split nested structs we combine them back together in the right order
+	for (int32 BoundaryPinIndex = InNode->Pins.Num() - 1; BoundaryPinIndex >= 0; --BoundaryPinIndex)
 	{
 		UEdGraphPin* const BoundaryPin = InNode->Pins[BoundaryPinIndex];
 
@@ -4238,10 +4343,22 @@ bool UEdGraphSchema_K2::CollapseGatewayNode(UK2Node* InNode, UEdGraphNode* InEnt
 		UEdGraphPin* GatewayPin = NULL;
 		if (GatewayNode)
 		{
-			for (int32 PinIdx=0; PinIdx < GatewayNode->Pins.Num(); PinIdx++)
+			// First handle struct combining if necessary
+			if (BoundaryPin->SubPins.Num() > 0)
+			{
+				InNode->ExpandSplitPin(CompilerContext, InNode->GetGraph(), BoundaryPin);
+			}
+
+			for (int32 PinIdx = GatewayNode->Pins.Num() - 1; PinIdx >= 0; --PinIdx)
 			{
 				UEdGraphPin* const Pin = GatewayNode->Pins[PinIdx];
 
+				// Expand any gateway pins as needed
+				if (Pin->SubPins.Num() > 0)
+				{
+					InNode->ExpandSplitPin(CompilerContext, GatewayNode->GetGraph(), Pin);
+				}
+
 				// Function graphs have a single exec path through them, so only one exec pin for input and another for output. In this fashion, they must not be handled by name.
 				if(InNode->GetClass() == UK2Node_CallFunction::StaticClass() && Pin->PinType.PinCategory == PC_Exec && BoundaryPin->PinType.PinCategory == PC_Exec && (Pin->Direction != BoundaryPin->Direction))
 				{
@@ -4262,7 +4379,7 @@ bool UEdGraphSchema_K2::CollapseGatewayNode(UK2Node* InNode, UEdGraphNode* InEnt
 		}
 		else
 		{
-			if (BoundaryPin->LinkedTo.Num() > 0)
+			if (BoundaryPin->LinkedTo.Num() > 0 && BoundaryPin->ParentPin == NULL)
 			{
 				UBlueprint* OwningBP = InNode->GetBlueprint();
 				if( OwningBP )
@@ -4293,7 +4410,7 @@ bool UEdGraphSchema_K2::CollapseGatewayNode(UK2Node* InNode, UEdGraphNode* InEnt
 
 	return bSuccessful;
 }
-
+	
 void UEdGraphSchema_K2::CombineTwoPinNetsAndRemoveOldPins(UEdGraphPin* InPinA, UEdGraphPin* InPinB) const
 {
 	check(InPinA != NULL);
@@ -4357,6 +4474,199 @@ void UEdGraphSchema_K2::CombineTwoPinNetsAndRemoveOldPins(UEdGraphPin* InPinA, U
 	InPinB->BreakAllPinLinks();
 }
 
+UK2Node* UEdGraphSchema_K2::CreateSplitPinNode(UEdGraphPin* Pin, FKismetCompilerContext* CompilerContext, UEdGraph* SourceGraph) const
+{
+	UEdGraphNode* GraphNode = Pin->GetOwningNode();
+	UEdGraph* Graph = GraphNode->GetGraph();
+	UScriptStruct* StructType = CastChecked<UScriptStruct>(Pin->PinType.PinSubCategoryObject.Get());
+	UK2Node* SplitPinNode = NULL;
+
+	if (Pin->Direction == EGPD_Input)
+	{
+		if (UK2Node_MakeStruct::CanBeMade(StructType))
+		{
+			UK2Node_MakeStruct* MakeStructNode = (CompilerContext ? CompilerContext->SpawnIntermediateNode<UK2Node_MakeStruct>(GraphNode, SourceGraph) : NewObject<UK2Node_MakeStruct>(Graph));
+			MakeStructNode->StructType = StructType;
+			SplitPinNode = MakeStructNode;
+		}
+		else
+		{
+			const FString& MetaData = StructType->GetMetaData(TEXT("HasNativeMake"));
+			const UFunction* Function = FindObject<UFunction>(NULL, *MetaData, true);
+
+			UK2Node_CallFunction* CallFunctionNode = (CompilerContext ? CompilerContext->SpawnIntermediateNode<UK2Node_CallFunction>(GraphNode, SourceGraph) : NewObject<UK2Node_CallFunction>(Graph));
+			CallFunctionNode->SetFromFunction(Function);
+			SplitPinNode = CallFunctionNode;
+		}
+	}
+	else
+	{
+		if (UK2Node_BreakStruct::CanBeBroken(StructType))
+		{
+			UK2Node_BreakStruct* BreakStructNode = (CompilerContext ? CompilerContext->SpawnIntermediateNode<UK2Node_BreakStruct>(GraphNode, SourceGraph) : NewObject<UK2Node_BreakStruct>(Graph));
+			BreakStructNode->StructType = StructType;
+			SplitPinNode = BreakStructNode;
+		}
+		else
+		{
+			const FString& MetaData = StructType->GetMetaData(TEXT("HasNativeBreak"));
+			const UFunction* Function = FindObject<UFunction>(NULL, *MetaData, true);
+
+			UK2Node_CallFunction* CallFunctionNode = (CompilerContext ? CompilerContext->SpawnIntermediateNode<UK2Node_CallFunction>(GraphNode, SourceGraph) : NewObject<UK2Node_CallFunction>(Graph));
+			CallFunctionNode->SetFromFunction(Function);
+			SplitPinNode = CallFunctionNode;
+		}
+	}
+
+	SplitPinNode->AllocateDefaultPins();
+
+	return SplitPinNode;
+}
+
+void UEdGraphSchema_K2::SplitPin(UEdGraphPin* Pin) const
+{
+	UScriptStruct* StructType = CastChecked<UScriptStruct>(Pin->PinType.PinSubCategoryObject.Get());
+	UEdGraphNode* GraphNode = Pin->GetOwningNode();
+	UEdGraph* Graph = CastChecked<UEdGraph>(GraphNode->GetOuter());
+
+	const FScopedTransaction Transaction( LOCTEXT("SplitStructPin", "Split Struct Pin") );
+
+	GraphNode->Modify();
+	Pin->Modify();
+
+	Pin->bHidden = true;
+
+	UK2Node* ProtoExpandNode = CreateSplitPinNode(Pin);
+			
+	for (UEdGraphPin* ProtoPin : ProtoExpandNode->Pins)
+	{
+		if (ProtoPin->Direction == Pin->Direction && !ProtoPin->bHidden)
+		{
+			const FString PinName = FString::Printf(TEXT("%s_%s"), *Pin->PinName, *ProtoPin->PinName);
+			const FEdGraphPinType& ProtoPinType = ProtoPin->PinType;
+			UEdGraphPin* SubPin = GraphNode->CreatePin(Pin->Direction, ProtoPinType.PinCategory, ProtoPinType.PinSubCategory, ProtoPinType.PinSubCategoryObject.Get(), false, false, PinName);
+
+			SubPin->DefaultValue = ProtoPin->DefaultValue;
+			SubPin->AutogeneratedDefaultValue = ProtoPin->AutogeneratedDefaultValue;
+
+			SubPin->ParentPin = Pin;
+
+			// CreatePin puts the Pin in the array, but we are going to insert it later, so pop it back out
+			GraphNode->Pins.Pop();
+
+			Pin->SubPins.Add(SubPin);
+		}
+	}
+
+	ProtoExpandNode->DestroyNode();
+
+	if (Pin->Direction == EGPD_Input)
+	{
+		TArray<FString> OriginalDefaults;
+		if (StructType == FindObjectChecked<UScriptStruct>(UObject::StaticClass(), TEXT("Vector")))
+		{
+			Pin->DefaultValue.ParseIntoArray(&OriginalDefaults, TEXT(","), false);
+		}
+		else if (StructType == FindObjectChecked<UScriptStruct>(UObject::StaticClass(), TEXT("Vector2D")))
+		{
+			FVector2D V2D;
+			V2D.InitFromString(Pin->DefaultValue);
+
+			OriginalDefaults.Add(FString::SanitizeFloat(V2D.X));
+			OriginalDefaults.Add(FString::SanitizeFloat(V2D.Y));
+		}
+		else if (StructType == FindObjectChecked<UScriptStruct>(UObject::StaticClass(), TEXT("LinearColor")))
+		{
+			FLinearColor LC;
+			LC.InitFromString(Pin->DefaultValue);
+
+			OriginalDefaults.Add(FString::SanitizeFloat(LC.R));
+			OriginalDefaults.Add(FString::SanitizeFloat(LC.G));
+			OriginalDefaults.Add(FString::SanitizeFloat(LC.B));
+			OriginalDefaults.Add(FString::SanitizeFloat(LC.A));
+		}
+
+		check(OriginalDefaults.Num() == 0 || OriginalDefaults.Num() == Pin->SubPins.Num());
+
+		for (int32 SubPinIndex = 0; SubPinIndex < OriginalDefaults.Num(); ++SubPinIndex)
+		{
+			UEdGraphPin* SubPin = Pin->SubPins[SubPinIndex];
+			SubPin->DefaultValue = OriginalDefaults[SubPinIndex];
+		}
+	}
+
+	GraphNode->Pins.Insert(Pin->SubPins, GraphNode->Pins.Find(Pin) + 1);
+
+	Graph->NotifyGraphChanged();
+
+	UBlueprint* Blueprint = FBlueprintEditorUtils::FindBlueprintForGraphChecked(Graph);
+	FBlueprintEditorUtils::MarkBlueprintAsModified(Blueprint);
+}
+
+void UEdGraphSchema_K2::RecombinePin(UEdGraphPin* Pin) const
+{
+	const FScopedTransaction Transaction( LOCTEXT("RecombineStructPin", "Recombine Struct Pin") );
+
+	UEdGraphNode* GraphNode = Pin->GetOwningNode();
+	UEdGraphPin* ParentPin = Pin->ParentPin;
+
+	GraphNode->Modify();
+	ParentPin->Modify();
+
+	ParentPin->bHidden = false;
+
+	for (int32 SubPinIndex = 0; SubPinIndex < ParentPin->SubPins.Num(); ++SubPinIndex)
+	{
+		UEdGraphPin* SubPin = ParentPin->SubPins[SubPinIndex];
+
+		if (SubPin->SubPins.Num() > 0)
+		{
+			RecombinePin(SubPin->SubPins[0]);
+		}
+
+		GraphNode->Pins.Remove(SubPin);
+	}
+
+	if (Pin->Direction == EGPD_Input)
+	{
+		UScriptStruct* StructType = CastChecked<UScriptStruct>(ParentPin->PinType.PinSubCategoryObject.Get());
+
+		TArray<FString> OriginalDefaults;
+		if (StructType == FindObjectChecked<UScriptStruct>(UObject::StaticClass(), TEXT("Vector")))
+		{
+			ParentPin->DefaultValue = ParentPin->SubPins[0]->DefaultValue + TEXT(",") 
+									+ ParentPin->SubPins[1]->DefaultValue + TEXT(",")
+									+ ParentPin->SubPins[2]->DefaultValue;
+		}
+		else if (StructType == FindObjectChecked<UScriptStruct>(UObject::StaticClass(), TEXT("Vector2D")))
+		{
+			FVector2D V2D;
+			V2D.X = FCString::Atof(*ParentPin->SubPins[0]->DefaultValue);
+			V2D.Y = FCString::Atof(*ParentPin->SubPins[1]->DefaultValue);
+			
+			ParentPin->DefaultValue = V2D.ToString();
+		}
+		else if (StructType == FindObjectChecked<UScriptStruct>(UObject::StaticClass(), TEXT("LinearColor")))
+		{
+			FLinearColor LC;
+			LC.R = FCString::Atof(*ParentPin->SubPins[0]->DefaultValue);
+			LC.G = FCString::Atof(*ParentPin->SubPins[1]->DefaultValue);
+			LC.B = FCString::Atof(*ParentPin->SubPins[2]->DefaultValue);
+			LC.A = FCString::Atof(*ParentPin->SubPins[3]->DefaultValue);
+
+			ParentPin->DefaultValue = LC.ToString();
+		}
+	}
+
+
+	ParentPin->SubPins.Empty();
+
+	UEdGraph* Graph = CastChecked<UEdGraph>(GraphNode->GetOuter());
+	Graph->NotifyGraphChanged();
+
+	UBlueprint* Blueprint = FBlueprintEditorUtils::FindBlueprintForGraphChecked(Graph);
+	FBlueprintEditorUtils::MarkBlueprintAsModified(Blueprint);
+}
 
 /////////////////////////////////////////////////////
 
