Allow reference struct pins to be split
Fix small defaults issue with splitting vector input pins

diff --git a/Engine/Source/Editor/BlueprintGraph/Private/EdGraphSchema_K2.cpp b/Engine/Source/Editor/BlueprintGraph/Private/EdGraphSchema_K2.cpp
--- a/Engine/Source/Editor/BlueprintGraph/Private/EdGraphSchema_K2.cpp
+++ b/Engine/Source/Editor/BlueprintGraph/Private/EdGraphSchema_K2.cpp
@@ -731,7 +731,7 @@ static bool IsUsingNonExistantVariable(const UEdGraphNode* InGraphNode, UBluepri
 bool UEdGraphSchema_K2::PinHasSplittableStructType(const UEdGraphPin* InGraphPin) const
 {
 	const FEdGraphPinType& PinType = InGraphPin->PinType;
-	bool bCanSplit = (!PinType.bIsArray && !PinType.bIsReference && InGraphPin->LinkedTo.Num() == 0 && PinType.PinCategory == PC_Struct);
+	bool bCanSplit = (!PinType.bIsArray && InGraphPin->LinkedTo.Num() == 0 && PinType.PinCategory == PC_Struct);
 
 	if (bCanSplit)
 	{
@@ -4566,6 +4566,10 @@ void UEdGraphSchema_K2::SplitPin(UEdGraphPin* Pin) const
 		if (StructType == FindObjectChecked<UScriptStruct>(UObject::StaticClass(), TEXT("Vector")))
 		{
 			Pin->DefaultValue.ParseIntoArray(&OriginalDefaults, TEXT(","), false);
+			for (FString& Default : OriginalDefaults)
+			{
+				Default.Trim();
+			}
 		}
 		else if (StructType == FindObjectChecked<UScriptStruct>(UObject::StaticClass(), TEXT("Vector2D")))
 		{
