Fix non-unity compile errors created by CL# 2107064

diff --git a/Engine/Source/Editor/BlueprintGraph/Private/EdGraphSchema_K2.cpp b/Engine/Source/Editor/BlueprintGraph/Private/EdGraphSchema_K2.cpp
--- a/Engine/Source/Editor/BlueprintGraph/Private/EdGraphSchema_K2.cpp
+++ b/Engine/Source/Editor/BlueprintGraph/Private/EdGraphSchema_K2.cpp
@@ -15,6 +15,7 @@
 
 #include "Kismet2/KismetEditorUtilities.h"
 #include "Kismet2/KismetDebugUtilities.h"
+#include "KismetCompiler.h"
 #include "ComponentAssetBroker.h"
 #include "AssetData.h"
 #include "Editor/UnrealEd/Public/EdGraphUtilities.h"
diff --git a/Engine/Source/Editor/BlueprintGraph/Private/K2Node.cpp b/Engine/Source/Editor/BlueprintGraph/Private/K2Node.cpp
--- a/Engine/Source/Editor/BlueprintGraph/Private/K2Node.cpp
+++ b/Engine/Source/Editor/BlueprintGraph/Private/K2Node.cpp
@@ -6,6 +6,7 @@
 
 #include "K2Node.h"
 #include "KismetDebugUtilities.h" // for HasDebuggingData(), GetWatchText()
+#include "KismetCompiler.h"
 
 #define LOCTEXT_NAMESPACE "K2Node"
 
