#ue4 BP: Fix for const pass-by-val terms not having default value entry fields

diff --git a/Engine/Source/Editor/BlueprintGraph/Private/K2Node_CallFunction.cpp b/Engine/Source/Editor/BlueprintGraph/Private/K2Node_CallFunction.cpp
--- a/Engine/Source/Editor/BlueprintGraph/Private/K2Node_CallFunction.cpp
+++ b/Engine/Source/Editor/BlueprintGraph/Private/K2Node_CallFunction.cpp
@@ -600,7 +600,7 @@ bool UK2Node_CallFunction::CreatePinsForFunctionCall(const UFunction* Function)
 		if (bPinGood)
 		{
 			//Flag pin as read only for const reference property
-			Pin->bDefaultValueIsIgnored = Param->HasAnyPropertyFlags(CPF_ConstParm | CPF_ReferenceParm) && (!Function->HasMetaData(FBlueprintMetadata::MD_AutoCreateRefTerm) || Pin->PinType.bIsArray);
+			Pin->bDefaultValueIsIgnored = Param->HasAllPropertyFlags(CPF_ConstParm | CPF_ReferenceParm) && (!Function->HasMetaData(FBlueprintMetadata::MD_AutoCreateRefTerm) || Pin->PinType.bIsArray);
 
 			const bool bAdvancedPin = Param->HasAllPropertyFlags(CPF_AdvancedDisplay);
 			Pin->bAdvancedView = bAdvancedPin;
