#ue4 Fix case where literal entry field would be hidden on async task nodes erroneously

diff --git a/Engine/Source/Editor/BlueprintGraph/Private/K2Node_BaseAsyncTask.cpp b/Engine/Source/Editor/BlueprintGraph/Private/K2Node_BaseAsyncTask.cpp
--- a/Engine/Source/Editor/BlueprintGraph/Private/K2Node_BaseAsyncTask.cpp
+++ b/Engine/Source/Editor/BlueprintGraph/Private/K2Node_BaseAsyncTask.cpp
@@ -115,7 +115,7 @@ void UK2Node_BaseAsyncTask::AllocateDefaultPins()
 			if (bPinGood)
 			{
 				//Flag pin as read only for const reference property
-				Pin->bDefaultValueIsIgnored = Param->HasAnyPropertyFlags(CPF_ConstParm | CPF_ReferenceParm) && (!Function->HasMetaData(FBlueprintMetadata::MD_AutoCreateRefTerm) || Pin->PinType.bIsArray);
+				Pin->bDefaultValueIsIgnored = Param->HasAllPropertyFlags(CPF_ConstParm | CPF_ReferenceParm) && (!Function->HasMetaData(FBlueprintMetadata::MD_AutoCreateRefTerm) || Pin->PinType.bIsArray);
 
 				const bool bAdvancedPin = Param->HasAllPropertyFlags(CPF_AdvancedDisplay);
 				Pin->bAdvancedView = bAdvancedPin;
