Allow Rotator pins to be split

diff --git a/Engine/Source/Editor/BlueprintGraph/Private/EdGraphSchema_K2.cpp b/Engine/Source/Editor/BlueprintGraph/Private/EdGraphSchema_K2.cpp
--- a/Engine/Source/Editor/BlueprintGraph/Private/EdGraphSchema_K2.cpp
+++ b/Engine/Source/Editor/BlueprintGraph/Private/EdGraphSchema_K2.cpp
@@ -4579,7 +4579,8 @@ void UEdGraphSchema_K2::SplitPin(UEdGraphPin* Pin) const
 	if (Pin->Direction == EGPD_Input)
 	{
 		TArray<FString> OriginalDefaults;
-		if (StructType == FindObjectChecked<UScriptStruct>(UObject::StaticClass(), TEXT("Vector")))
+		if (   StructType == FindObjectChecked<UScriptStruct>(UObject::StaticClass(), TEXT("Vector"))
+			|| StructType == FindObjectChecked<UScriptStruct>(UObject::StaticClass(), TEXT("Rotator")))
 		{
 			Pin->DefaultValue.ParseIntoArray(&OriginalDefaults, TEXT(","), false);
 			for (FString& Default : OriginalDefaults)
@@ -4652,7 +4653,8 @@ void UEdGraphSchema_K2::RecombinePin(UEdGraphPin* Pin) const
 		UScriptStruct* StructType = CastChecked<UScriptStruct>(ParentPin->PinType.PinSubCategoryObject.Get());
 
 		TArray<FString> OriginalDefaults;
-		if (StructType == FindObjectChecked<UScriptStruct>(UObject::StaticClass(), TEXT("Vector")))
+		if (   StructType == FindObjectChecked<UScriptStruct>(UObject::StaticClass(), TEXT("Vector"))
+			|| StructType == FindObjectChecked<UScriptStruct>(UObject::StaticClass(), TEXT("Rotator")))
 		{
 			ParentPin->DefaultValue = ParentPin->SubPins[0]->DefaultValue + TEXT(",") 
 									+ ParentPin->SubPins[1]->DefaultValue + TEXT(",")
diff --git a/Engine/Source/Runtime/CoreUObject/Classes/Object.h b/Engine/Source/Runtime/CoreUObject/Classes/Object.h
--- a/Engine/Source/Runtime/CoreUObject/Classes/Object.h
+++ b/Engine/Source/Runtime/CoreUObject/Classes/Object.h
@@ -252,7 +252,7 @@ struct FPlane : public FVector
 
 // An orthogonal rotation in 3d space.
 
-USTRUCT(immutable, noexport, BlueprintType, meta=(HasNativeMake="true", HasNativeBreak="true"))
+USTRUCT(immutable, noexport, BlueprintType, meta=(HasNativeMake="Engine.KismetMathLibrary.MakeRot", HasNativeBreak="Engine.KismetMathLibrary.BreakRot"))
 struct FRotator
 {
 	UPROPERTY(EditAnywhere, BlueprintReadWrite, Category=Rotator, SaveGame)
