Don't create properties for blueprint functions with a split input struct

diff --git a/Engine/Source/Editor/BlueprintGraph/Private/K2Node_FunctionEntry.cpp b/Engine/Source/Editor/BlueprintGraph/Private/K2Node_FunctionEntry.cpp
--- a/Engine/Source/Editor/BlueprintGraph/Private/K2Node_FunctionEntry.cpp
+++ b/Engine/Source/Editor/BlueprintGraph/Private/K2Node_FunctionEntry.cpp
@@ -46,10 +46,9 @@ public:
 
 		UFunction* Function = FindField<UFunction>(EntryNode->SignatureClass, EntryNode->SignatureName);
 
-		for (int32 PinIndex = 0; PinIndex < Node->Pins.Num(); ++PinIndex)
+		for (UEdGraphPin* Pin : Node->Pins)
 		{
-			UEdGraphPin* Pin = Node->Pins[PinIndex];
-			if (!CompilerContext.GetSchema()->IsMetaPin(*Pin))
+			if (Pin->ParentPin == nullptr && !CompilerContext.GetSchema()->IsMetaPin(*Pin))
 			{
 				UEdGraphPin* Net = FEdGraphUtilities::GetNetFromPin(Pin);
 
