#ue4 Adding improved deprecation warning on nodes to notify users that the var was deprecated as opposed to just removed.

diff --git a/Engine/Source/Editor/BlueprintGraph/Classes/K2Node.h b/Engine/Source/Editor/BlueprintGraph/Classes/K2Node.h
--- a/Engine/Source/Editor/BlueprintGraph/Classes/K2Node.h
+++ b/Engine/Source/Editor/BlueprintGraph/Classes/K2Node.h
@@ -387,7 +387,10 @@ protected:
 	UPROPERTY()
 	mutable bool bSelfContext;
 
-
+	/** Whether or not this property has been deprecated */
+	UPROPERTY()
+	mutable bool bWasDeprecated;
+	
 public:
 	FMemberReference()
 		: MemberParentClass(NULL)
@@ -403,6 +406,7 @@ public:
 		MemberParentClass = bIsConsideredSelfContext ? NULL : InField->GetOwnerClass();
 		MemberName = InField->GetFName();
 		bSelfContext = bIsConsideredSelfContext;
+		bWasDeprecated = false;
 
 		MemberGuid.Invalidate();
 		if (InField->GetOwnerClass())
@@ -527,6 +531,12 @@ public:
 		return MemberScope;
 	}
 
+	/** Returns whether or not the variable has been deprecated */
+	bool IsDeprecated() const
+	{
+		return bWasDeprecated;
+	}
+
 	/** 
 	 *	Returns the member UProperty/UFunction this reference is pointing to, or NULL if it no longer exists 
 	 *	Derives 'self' scope from supplied Blueprint node if required
@@ -603,6 +613,12 @@ public:
 			}
 		}
 
+		// Check to see if the member has been deprecated
+		if (UProperty* Property = Cast<UProperty>(ReturnField))
+		{
+			bWasDeprecated = Property->HasAnyPropertyFlags(CPF_Deprecated);
+		}
+
 		return ReturnField;
 	}
 
diff --git a/Engine/Source/Editor/BlueprintGraph/Private/K2Node.cpp b/Engine/Source/Editor/BlueprintGraph/Private/K2Node.cpp
--- a/Engine/Source/Editor/BlueprintGraph/Private/K2Node.cpp
+++ b/Engine/Source/Editor/BlueprintGraph/Private/K2Node.cpp
@@ -1213,6 +1213,7 @@ void FMemberReference::SetExternalMember(FName InMemberName, TSubclassOf<class U
 	MemberParentClass = InMemberParentClass;
 	MemberScope.Empty();
 	bSelfContext = false;
+	bWasDeprecated = false;
 }
 
 void FMemberReference::SetSelfMember(FName InMemberName)
@@ -1221,6 +1222,7 @@ void FMemberReference::SetSelfMember(FName InMemberName)
 	MemberParentClass = NULL;
 	MemberScope.Empty();
 	bSelfContext = true;
+	bWasDeprecated = false;
 }
 
 void FMemberReference::SetDirect(const FName InMemberName, const FGuid InMemberGuid, TSubclassOf<class UObject> InMemberParentClass, bool bIsConsideredSelfContext)
@@ -1228,6 +1230,7 @@ void FMemberReference::SetDirect(const FName InMemberName, const FGuid InMemberG
 	MemberName = InMemberName;
 	MemberGuid = InMemberGuid;
 	bSelfContext = bIsConsideredSelfContext;
+	bWasDeprecated = false;
 	MemberParentClass = InMemberParentClass;
 	MemberScope.Empty();
 }
@@ -1239,6 +1242,7 @@ void FMemberReference::SetGivenSelfScope(const FName InMemberName, const FGuid I
 	MemberParentClass = InMemberParentClass;
 	MemberScope.Empty();
 	bSelfContext = (SelfScope->IsChildOf(InMemberParentClass)) || (SelfScope->ClassGeneratedBy == InMemberParentClass->ClassGeneratedBy);
+	bWasDeprecated = false;
 
 	if (bSelfContext)
 	{
diff --git a/Engine/Source/Editor/BlueprintGraph/Private/K2Node_Variable.cpp b/Engine/Source/Editor/BlueprintGraph/Private/K2Node_Variable.cpp
--- a/Engine/Source/Editor/BlueprintGraph/Private/K2Node_Variable.cpp
+++ b/Engine/Source/Editor/BlueprintGraph/Private/K2Node_Variable.cpp
@@ -254,7 +254,14 @@ void UK2Node_Variable::ValidateNodeDuringCompilation(class FCompilerResultsLog&
 	// Local variables do not exist until much later in the compilation than this function can provide
 	if (VariableProperty == NULL && !VariableReference.IsLocalScope())
 	{
-		MessageLog.Warning(*FString::Printf(*LOCTEXT("VariableNotFound", "Unable to find variable with name '%s' for @@").ToString(), *VariableReference.GetMemberName().ToString()), this);
+		if (!VariableReference.IsDeprecated())
+		{
+			MessageLog.Warning(*FString::Printf(*LOCTEXT("VariableNotFound", "Unable to find variable with name '%s' for @@").ToString(), *VariableReference.GetMemberName().ToString()), this);
+		}
+		else
+		{
+			MessageLog.Warning(*FString::Printf(*LOCTEXT("VariableDeprecated", "Variable '%s' for @@ was deprecated.  Please update it.").ToString(), *VariableReference.GetMemberName().ToString()), this);
+		}
 	}
 }
 
