Removed GetNodeNativeTitle from the editor entirely, now uses BuildSourceString to obtain the source string in a more reliable way.

diff --git a/Engine/Plugins/Editor/IntroTutorials/Source/IntroTutorials/Classes/K2Node_TutorialExcerptComplete.h b/Engine/Plugins/Editor/IntroTutorials/Source/IntroTutorials/Classes/K2Node_TutorialExcerptComplete.h
--- a/Engine/Plugins/Editor/IntroTutorials/Source/IntroTutorials/Classes/K2Node_TutorialExcerptComplete.h
+++ b/Engine/Plugins/Editor/IntroTutorials/Source/IntroTutorials/Classes/K2Node_TutorialExcerptComplete.h
@@ -29,7 +29,6 @@ class UK2Node_TutorialExcerptComplete : public UK2Node_CallFunction
 
 	// Begin UEdGraphNode interface
 	virtual FText GetNodeTitle(ENodeTitleType::Type TitleType) const override;
-	virtual FString GetNodeNativeTitle(ENodeTitleType::Type TitleType) const override;
 	virtual void PostPlacedNewNode() override;
 	virtual void ValidateNodeDuringCompilation(class FCompilerResultsLog& MessageLog) const override;
 	// End UEdGraphNode interface
diff --git a/Engine/Plugins/Editor/IntroTutorials/Source/IntroTutorials/Private/K2Node_TutorialExcerptComplete.cpp b/Engine/Plugins/Editor/IntroTutorials/Source/IntroTutorials/Private/K2Node_TutorialExcerptComplete.cpp
--- a/Engine/Plugins/Editor/IntroTutorials/Source/IntroTutorials/Private/K2Node_TutorialExcerptComplete.cpp
+++ b/Engine/Plugins/Editor/IntroTutorials/Source/IntroTutorials/Private/K2Node_TutorialExcerptComplete.cpp
@@ -46,7 +46,6 @@ void UK2Node_TutorialExcerptComplete::GetMenuEntries(FGraphContextMenuBuilder& C
 
 	TSharedPtr<FEdGraphSchemaAction_K2NewNode> NodeAction = FK2ActionMenuBuilder::AddNewNodeAction(ContextMenuBuilder, Category + TEXT("|") + SubCategory, MenuDesc, Tooltip);
 	NodeAction->NodeTemplate = TemplateNode;
-	NodeAction->SearchTitle = TemplateNode->GetNodeSearchTitle();
 }
 
 void UK2Node_TutorialExcerptComplete::PinDefaultValueChanged(UEdGraphPin* Pin)
@@ -106,18 +105,6 @@ FText UK2Node_TutorialExcerptComplete::GetNodeTitle(ENodeTitleType::Type TitleTy
 	}
 }
 
-FString UK2Node_TutorialExcerptComplete::GetNodeNativeTitle(ENodeTitleType::Type TitleType) const
-{
-	if(CachedExcerpt.Len() > 0)
-	{
-		return FString::Printf(TEXT("Tutorial Excerpt '%s' Complete"), *CachedExcerpt);
-	}
-	else
-	{
-		return TEXT("Tutorial Excerpt Complete");
-	}
-}
-
 void UK2Node_TutorialExcerptComplete::TutorialExcerptComplete(const FString& Excerpt)
 {
 	OnTutorialExcerptComplete.Broadcast(Excerpt);
diff --git a/Engine/Plugins/Experimental/LiveEditor/Source/LiveEditor/Classes/K2Node_LiveEditObject.h b/Engine/Plugins/Experimental/LiveEditor/Source/LiveEditor/Classes/K2Node_LiveEditObject.h
--- a/Engine/Plugins/Experimental/LiveEditor/Source/LiveEditor/Classes/K2Node_LiveEditObject.h
+++ b/Engine/Plugins/Experimental/LiveEditor/Source/LiveEditor/Classes/K2Node_LiveEditObject.h
@@ -12,7 +12,6 @@ class UK2Node_LiveEditObject : public UK2Node
 	// Begin UEdGraphNode interface.
 	virtual void AllocateDefaultPins() override;
 	virtual FText GetNodeTitle(ENodeTitleType::Type TitleType) const override;
-	virtual FString GetNodeNativeTitle(ENodeTitleType::Type TitleType) const override;
 	virtual void PinDefaultValueChanged(UEdGraphPin* Pin) override;
 	// End UEdGraphNode interface.
 
diff --git a/Engine/Plugins/Experimental/LiveEditor/Source/LiveEditor/Private/K2Node_LiveEditObject.cpp b/Engine/Plugins/Experimental/LiveEditor/Source/LiveEditor/Private/K2Node_LiveEditObject.cpp
--- a/Engine/Plugins/Experimental/LiveEditor/Source/LiveEditor/Private/K2Node_LiveEditObject.cpp
+++ b/Engine/Plugins/Experimental/LiveEditor/Source/LiveEditor/Private/K2Node_LiveEditObject.cpp
@@ -199,19 +199,6 @@ FText UK2Node_LiveEditObject::GetNodeTitle(ENodeTitleType::Type TitleType) const
 	return FText::Format(NSLOCTEXT("K2Node", "LiveEditObject", "LiveEditObject {SpawnString}_{ID}"), Args );
 }
 
-FString UK2Node_LiveEditObject::GetNodeNativeTitle(ENodeTitleType::Type TitleType) const
-{
-	UEdGraphPin* BaseClassPin = GetBaseClassPin();
-
-	FString SpawnString = TEXT("NONE");
-	if(BaseClassPin != NULL && BaseClassPin->DefaultObject != NULL )
-	{
-		SpawnString = BaseClassPin->DefaultObject->GetName();
-	}
-
-	return FString::Printf(TEXT("LiveEditObject %s_%d"), *SpawnString, GetUniqueID());
-}
-
 void UK2Node_LiveEditObject::PinDefaultValueChanged(UEdGraphPin* Pin) 
 {
 	if(Pin->PinName == UK2Node_LiveEditObjectStatics::BaseClassPinName)
diff --git a/Engine/Plugins/Experimental/LiveEditor/Source/LiveEditor/Private/LiveEditor.cpp b/Engine/Plugins/Experimental/LiveEditor/Source/LiveEditor/Private/LiveEditor.cpp
--- a/Engine/Plugins/Experimental/LiveEditor/Source/LiveEditor/Private/LiveEditor.cpp
+++ b/Engine/Plugins/Experimental/LiveEditor/Source/LiveEditor/Private/LiveEditor.cpp
@@ -102,7 +102,6 @@ void FLiveEditorNodeInjector::AddContextMenuNodes(FBlueprintGraphActionListBuild
 	FString FunctionCategory( TEXT("LiveEditor") );
 	TSharedPtr<FEdGraphSchemaAction_K2NewNode> Action = FK2ActionMenuBuilder::AddNewNodeAction(ContextMenuBuilder, FunctionCategory, NSLOCTEXT("LiveEditorPlugin", "LiveEditObject", "LiveEdit Object"), NodeTemplate->GetTooltip(), 0, NodeTemplate->GetKeywords());
 	Action->NodeTemplate = NodeTemplate;
-	Action->SearchTitle = NodeTemplate->GetNodeSearchTitle();
 }
 
 void FLiveEditorNodeInjector::OnBeginPIE(const bool bIsSimulating)
diff --git a/Engine/Source/Developer/MovieSceneTools/Classes/K2Node_PlayMovieScene.h b/Engine/Source/Developer/MovieSceneTools/Classes/K2Node_PlayMovieScene.h
--- a/Engine/Source/Developer/MovieSceneTools/Classes/K2Node_PlayMovieScene.h
+++ b/Engine/Source/Developer/MovieSceneTools/Classes/K2Node_PlayMovieScene.h
@@ -42,7 +42,6 @@ class UK2Node_PlayMovieScene : public UK2Node
 	virtual void AllocateDefaultPins() override;
 	virtual FLinearColor GetNodeTitleColor() const override;
 	virtual FText GetNodeTitle(ENodeTitleType::Type TitleType) const override;
-	virtual FString GetNodeNativeTitle(ENodeTitleType::Type TitleType) const override;
 
 	// UK2Node interface
 	virtual bool NodeCausesStructuralBlueprintChange() const override
diff --git a/Engine/Source/Developer/MovieSceneTools/Private/K2Node_PlayMovieScene.cpp b/Engine/Source/Developer/MovieSceneTools/Private/K2Node_PlayMovieScene.cpp
--- a/Engine/Source/Developer/MovieSceneTools/Private/K2Node_PlayMovieScene.cpp
+++ b/Engine/Source/Developer/MovieSceneTools/Private/K2Node_PlayMovieScene.cpp
@@ -87,15 +87,6 @@ FText UK2Node_PlayMovieScene::GetNodeTitle(ENodeTitleType::Type TitleType) const
 		NSLOCTEXT("PlayMovieSceneNode", "NodeTitleWithNoMovieScene", "Play Movie Scene (No Asset)");
 }
 
-FString UK2Node_PlayMovieScene::GetNodeNativeTitle(ENodeTitleType::Type TitleType) const
-{
-	UMovieScene* MovieScene = MovieSceneBindings != NULL ? MovieSceneBindings->GetRootMovieScene() : NULL;
-	return MovieScene != NULL ?
-		( FString(TEXT("Play Movie Scene: ")) + MovieScene->GetName() ) 
-		: TEXT("Play Movie Scene (No Asset)");
-}
-
-
 UEdGraphPin* UK2Node_PlayMovieScene::GetPlayPin() const
 {
 	UEdGraphPin* Pin = FindPinChecked( PlayMovieScenePinNames::Play );
diff --git a/Engine/Source/Editor/AnimGraph/Classes/AnimGraphNode_BlendListByEnum.h b/Engine/Source/Editor/AnimGraph/Classes/AnimGraphNode_BlendListByEnum.h
--- a/Engine/Source/Editor/AnimGraph/Classes/AnimGraphNode_BlendListByEnum.h
+++ b/Engine/Source/Editor/AnimGraph/Classes/AnimGraphNode_BlendListByEnum.h
@@ -24,7 +24,6 @@ public:
 	// UEdGraphNode interface
 	virtual FString GetTooltip() const override;
 	virtual FText GetNodeTitle(ENodeTitleType::Type TitleType) const override;
-	virtual FString GetNodeNativeTitle(ENodeTitleType::Type TitleType) const override;
 	virtual void PostPlacedNewNode() override;
 	// End of UEdGraphNode interface
 
diff --git a/Engine/Source/Editor/AnimGraph/Classes/AnimGraphNode_BlendSpaceEvaluator.h b/Engine/Source/Editor/AnimGraph/Classes/AnimGraphNode_BlendSpaceEvaluator.h
--- a/Engine/Source/Editor/AnimGraph/Classes/AnimGraphNode_BlendSpaceEvaluator.h
+++ b/Engine/Source/Editor/AnimGraph/Classes/AnimGraphNode_BlendSpaceEvaluator.h
@@ -19,7 +19,6 @@ class UAnimGraphNode_BlendSpaceEvaluator : public UAnimGraphNode_BlendSpaceBase
 	// UEdGraphNode interface
 	virtual FString GetTooltip() const override;
 	virtual FText GetNodeTitle(ENodeTitleType::Type TitleType) const override;
-	virtual FString GetNodeNativeTitle(ENodeTitleType::Type TitleType) const override;
 	// End of UEdGraphNode interface
 
 	// UAnimGraphNode_Base interface
diff --git a/Engine/Source/Editor/AnimGraph/Classes/AnimGraphNode_BlendSpacePlayer.h b/Engine/Source/Editor/AnimGraph/Classes/AnimGraphNode_BlendSpacePlayer.h
--- a/Engine/Source/Editor/AnimGraph/Classes/AnimGraphNode_BlendSpacePlayer.h
+++ b/Engine/Source/Editor/AnimGraph/Classes/AnimGraphNode_BlendSpacePlayer.h
@@ -20,7 +20,6 @@ class UAnimGraphNode_BlendSpacePlayer : public UAnimGraphNode_BlendSpaceBase
 	// UEdGraphNode interface
 	virtual FString GetTooltip() const override;
 	virtual FText GetNodeTitle(ENodeTitleType::Type TitleType) const override;
-	virtual FString GetNodeNativeTitle(ENodeTitleType::Type TitleType) const override;
 	// End of UEdGraphNode interface
 
 	// UAnimGraphNode_Base interface
diff --git a/Engine/Source/Editor/AnimGraph/Classes/AnimGraphNode_CopyBone.h b/Engine/Source/Editor/AnimGraph/Classes/AnimGraphNode_CopyBone.h
--- a/Engine/Source/Editor/AnimGraph/Classes/AnimGraphNode_CopyBone.h
+++ b/Engine/Source/Editor/AnimGraph/Classes/AnimGraphNode_CopyBone.h
@@ -16,7 +16,6 @@ class UAnimGraphNode_CopyBone : public UAnimGraphNode_SkeletalControlBase
 public:
 	// UEdGraphNode interface
 	virtual FText GetNodeTitle(ENodeTitleType::Type TitleType) const override;
-	virtual FString GetNodeNativeTitle(ENodeTitleType::Type TitleType) const override;
 	virtual FString GetTooltip() const override;
 	// End of UEdGraphNode interface
 
diff --git a/Engine/Source/Editor/AnimGraph/Classes/AnimGraphNode_HandIKRetargeting.h b/Engine/Source/Editor/AnimGraph/Classes/AnimGraphNode_HandIKRetargeting.h
--- a/Engine/Source/Editor/AnimGraph/Classes/AnimGraphNode_HandIKRetargeting.h
+++ b/Engine/Source/Editor/AnimGraph/Classes/AnimGraphNode_HandIKRetargeting.h
@@ -16,7 +16,6 @@ class UAnimGraphNode_HandIKRetargeting : public UAnimGraphNode_SkeletalControlBa
 public:
 	// UEdGraphNode interface
 	virtual FText GetNodeTitle(ENodeTitleType::Type TitleType) const override;
-	virtual FString GetNodeNativeTitle(ENodeTitleType::Type TitleType) const override;
 	virtual FString GetTooltip() const override;
 	// End of UEdGraphNode interface
 
diff --git a/Engine/Source/Editor/AnimGraph/Classes/AnimGraphNode_RotationMultiplier.h b/Engine/Source/Editor/AnimGraph/Classes/AnimGraphNode_RotationMultiplier.h
--- a/Engine/Source/Editor/AnimGraph/Classes/AnimGraphNode_RotationMultiplier.h
+++ b/Engine/Source/Editor/AnimGraph/Classes/AnimGraphNode_RotationMultiplier.h
@@ -16,7 +16,6 @@ class UAnimGraphNode_RotationMultiplier : public UAnimGraphNode_SkeletalControlB
 public:
 	// UEdGraphNode interface
 	virtual FText GetNodeTitle(ENodeTitleType::Type TitleType) const override;
-	virtual FString GetNodeNativeTitle(ENodeTitleType::Type TitleType) const override;
 	virtual FString GetTooltip() const override;
 	// End of UEdGraphNode interface
 
diff --git a/Engine/Source/Editor/AnimGraph/Classes/AnimGraphNode_RotationOffsetBlendSpace.h b/Engine/Source/Editor/AnimGraph/Classes/AnimGraphNode_RotationOffsetBlendSpace.h
--- a/Engine/Source/Editor/AnimGraph/Classes/AnimGraphNode_RotationOffsetBlendSpace.h
+++ b/Engine/Source/Editor/AnimGraph/Classes/AnimGraphNode_RotationOffsetBlendSpace.h
@@ -16,7 +16,6 @@ class UAnimGraphNode_RotationOffsetBlendSpace: public UAnimGraphNode_BlendSpaceB
 	// UEdGraphNode interface
 	virtual FString GetTooltip() const override;
 	virtual FText GetNodeTitle(ENodeTitleType::Type TitleType) const override;
-	virtual FString GetNodeNativeTitle(ENodeTitleType::Type TitleType) const override;
 	// End of UEdGraphNode interface
 
 	// UAnimGraphNode_Base interface
diff --git a/Engine/Source/Editor/AnimGraph/Classes/AnimGraphNode_SaveCachedPose.h b/Engine/Source/Editor/AnimGraph/Classes/AnimGraphNode_SaveCachedPose.h
--- a/Engine/Source/Editor/AnimGraph/Classes/AnimGraphNode_SaveCachedPose.h
+++ b/Engine/Source/Editor/AnimGraph/Classes/AnimGraphNode_SaveCachedPose.h
@@ -19,7 +19,6 @@ class UAnimGraphNode_SaveCachedPose : public UAnimGraphNode_Base
 	// UEdGraphNode interface
 	virtual FString GetTooltip() const override;
 	virtual FText GetNodeTitle(ENodeTitleType::Type TitleType) const override;
-	virtual FString GetNodeNativeTitle(ENodeTitleType::Type TitleType) const override;
 	virtual void OnRenameNode(const FString& NewName) override;
 	virtual TSharedPtr<class INameValidatorInterface> MakeNameValidator() const override;
 	// End of UEdGraphNode interface
diff --git a/Engine/Source/Editor/AnimGraph/Classes/AnimGraphNode_SequenceEvaluator.h b/Engine/Source/Editor/AnimGraph/Classes/AnimGraphNode_SequenceEvaluator.h
--- a/Engine/Source/Editor/AnimGraph/Classes/AnimGraphNode_SequenceEvaluator.h
+++ b/Engine/Source/Editor/AnimGraph/Classes/AnimGraphNode_SequenceEvaluator.h
@@ -15,7 +15,6 @@ class UAnimGraphNode_SequenceEvaluator : public UAnimGraphNode_Base
 
 	// UEdGraphNode interface
 	virtual FText GetNodeTitle(ENodeTitleType::Type TitleType) const override;
-	virtual FString GetNodeNativeTitle(ENodeTitleType::Type TitleType) const override;
 	virtual FString GetTooltip() const override;
 	// End of UEdGraphNode
 
diff --git a/Engine/Source/Editor/AnimGraph/Classes/AnimGraphNode_SequencePlayer.h b/Engine/Source/Editor/AnimGraph/Classes/AnimGraphNode_SequencePlayer.h
--- a/Engine/Source/Editor/AnimGraph/Classes/AnimGraphNode_SequencePlayer.h
+++ b/Engine/Source/Editor/AnimGraph/Classes/AnimGraphNode_SequencePlayer.h
@@ -21,7 +21,6 @@ class UAnimGraphNode_SequencePlayer : public UAnimGraphNode_Base
 	virtual FLinearColor GetNodeTitleColor() const override;
 	virtual FString GetTooltip() const override;
 	virtual FText GetNodeTitle(ENodeTitleType::Type TitleType) const override;
-	virtual FString GetNodeNativeTitle(ENodeTitleType::Type TitleType) const override;
 	// End of UEdGraphNode interface
 
 	// UAnimGraphNode_Base interface
diff --git a/Engine/Source/Editor/AnimGraph/Classes/AnimGraphNode_Slot.h b/Engine/Source/Editor/AnimGraph/Classes/AnimGraphNode_Slot.h
--- a/Engine/Source/Editor/AnimGraph/Classes/AnimGraphNode_Slot.h
+++ b/Engine/Source/Editor/AnimGraph/Classes/AnimGraphNode_Slot.h
@@ -17,7 +17,6 @@ class UAnimGraphNode_Slot : public UAnimGraphNode_Base
 	virtual FLinearColor GetNodeTitleColor() const override;
 	virtual FString GetTooltip() const override;
 	virtual FText GetNodeTitle(ENodeTitleType::Type TitleType) const override;
-	virtual FString GetNodeNativeTitle(ENodeTitleType::Type TitleType) const override;
 	// End UEdGraphNode interface.
 
 	// UAnimGraphNode_Base interface
diff --git a/Engine/Source/Editor/AnimGraph/Classes/AnimGraphNode_SpringBone.h b/Engine/Source/Editor/AnimGraph/Classes/AnimGraphNode_SpringBone.h
--- a/Engine/Source/Editor/AnimGraph/Classes/AnimGraphNode_SpringBone.h
+++ b/Engine/Source/Editor/AnimGraph/Classes/AnimGraphNode_SpringBone.h
@@ -16,7 +16,6 @@ class UAnimGraphNode_SpringBone : public UAnimGraphNode_SkeletalControlBase
 public:
 	// UEdGraphNode interface
 	virtual FText GetNodeTitle(ENodeTitleType::Type TitleType) const override;
-	virtual FString GetNodeNativeTitle(ENodeTitleType::Type TitleType) const override;
 	virtual FString GetTooltip() const override;
 	// End of UEdGraphNode interface
 
diff --git a/Engine/Source/Editor/AnimGraph/Classes/AnimGraphNode_StateMachineBase.h b/Engine/Source/Editor/AnimGraph/Classes/AnimGraphNode_StateMachineBase.h
--- a/Engine/Source/Editor/AnimGraph/Classes/AnimGraphNode_StateMachineBase.h
+++ b/Engine/Source/Editor/AnimGraph/Classes/AnimGraphNode_StateMachineBase.h
@@ -18,7 +18,6 @@ class ANIMGRAPH_API UAnimGraphNode_StateMachineBase : public UAnimGraphNode_Base
 	virtual FLinearColor GetNodeTitleColor() const override;
 	virtual FString GetTooltip() const override;
 	virtual FText GetNodeTitle(ENodeTitleType::Type TitleType) const override;
-	virtual FString GetNodeNativeTitle(ENodeTitleType::Type TitleType) const override;
 	virtual void PostPlacedNewNode() override;
 	virtual UObject* GetJumpTargetForDoubleClick() const override;
 	virtual void DestroyNode() override;
diff --git a/Engine/Source/Editor/AnimGraph/Classes/AnimGraphNode_TwoBoneIK.h b/Engine/Source/Editor/AnimGraph/Classes/AnimGraphNode_TwoBoneIK.h
--- a/Engine/Source/Editor/AnimGraph/Classes/AnimGraphNode_TwoBoneIK.h
+++ b/Engine/Source/Editor/AnimGraph/Classes/AnimGraphNode_TwoBoneIK.h
@@ -16,7 +16,6 @@ class UAnimGraphNode_TwoBoneIK : public UAnimGraphNode_SkeletalControlBase
 public:
 	// UEdGraphNode interface
 	virtual FText GetNodeTitle(ENodeTitleType::Type TitleType) const override;
-	virtual FString GetNodeNativeTitle(ENodeTitleType::Type TitleType) const override;
 	virtual FString GetTooltip() const override;
 	// End of UEdGraphNode interface
 
diff --git a/Engine/Source/Editor/AnimGraph/Classes/AnimGraphNode_UseCachedPose.h b/Engine/Source/Editor/AnimGraph/Classes/AnimGraphNode_UseCachedPose.h
--- a/Engine/Source/Editor/AnimGraph/Classes/AnimGraphNode_UseCachedPose.h
+++ b/Engine/Source/Editor/AnimGraph/Classes/AnimGraphNode_UseCachedPose.h
@@ -21,7 +21,6 @@ public:
 	// UEdGraphNode interface
 	virtual FString GetTooltip() const override;
 	virtual FText GetNodeTitle(ENodeTitleType::Type TitleType) const override;
-	virtual FString GetNodeNativeTitle(ENodeTitleType::Type TitleType) const override;
 	// End of UEdGraphNode interface
 
 	// UAnimGraphNode_Base interface
diff --git a/Engine/Source/Editor/AnimGraph/Classes/AnimStateTransitionNode.h b/Engine/Source/Editor/AnimGraph/Classes/AnimStateTransitionNode.h
--- a/Engine/Source/Editor/AnimGraph/Classes/AnimStateTransitionNode.h
+++ b/Engine/Source/Editor/AnimGraph/Classes/AnimStateTransitionNode.h
@@ -84,7 +84,6 @@ class UAnimStateTransitionNode : public UAnimStateNodeBase
 	// Begin UEdGraphNode interface
 	virtual void AllocateDefaultPins() override;
 	virtual FText GetNodeTitle(ENodeTitleType::Type TitleType) const override;
-	virtual FString GetNodeNativeTitle(ENodeTitleType::Type TitleType) const override;
 	virtual FString GetTooltip() const override;
 	virtual FLinearColor GetNodeTitleColor() const override;
 	virtual void PinConnectionListChanged(UEdGraphPin* Pin) override;
diff --git a/Engine/Source/Editor/AnimGraph/Classes/K2Node_TransitionRuleGetter.h b/Engine/Source/Editor/AnimGraph/Classes/K2Node_TransitionRuleGetter.h
--- a/Engine/Source/Editor/AnimGraph/Classes/K2Node_TransitionRuleGetter.h
+++ b/Engine/Source/Editor/AnimGraph/Classes/K2Node_TransitionRuleGetter.h
@@ -39,7 +39,6 @@ class UK2Node_TransitionRuleGetter : public UK2Node
 	// UEdGraphNode interface
 	virtual void AllocateDefaultPins() override;
 	virtual FText GetNodeTitle(ENodeTitleType::Type TitleType) const override;
-	virtual FString GetNodeNativeTitle(ENodeTitleType::Type TitleType) const override;
 	virtual FString GetTooltip() const override;
 	virtual bool ShowPaletteIconOnNode() const override{ return false; }
 	// End of UEdGraphNode interface
diff --git a/Engine/Source/Editor/AnimGraph/Private/AnimGraphNode_Base.cpp b/Engine/Source/Editor/AnimGraph/Private/AnimGraphNode_Base.cpp
--- a/Engine/Source/Editor/AnimGraph/Private/AnimGraphNode_Base.cpp
+++ b/Engine/Source/Editor/AnimGraph/Private/AnimGraphNode_Base.cpp
@@ -237,7 +237,6 @@ TSharedPtr<FEdGraphSchemaAction_K2NewNode> UAnimGraphNode_Base::CreateDefaultMen
 
 	TSharedPtr<FEdGraphSchemaAction_K2NewNode> NodeAction = FK2ActionMenuBuilder::AddNewNodeAction(ContextMenuBuilder, Category, MenuDesc, Tooltip, 0, Keywords);
 	NodeAction->NodeTemplate = TemplateNode;
-	NodeAction->SearchTitle = TemplateNode->GetNodeSearchTitle();
 
 	return NodeAction;
 }
diff --git a/Engine/Source/Editor/AnimGraph/Private/AnimGraphNode_BlendListByEnum.cpp b/Engine/Source/Editor/AnimGraph/Private/AnimGraphNode_BlendListByEnum.cpp
--- a/Engine/Source/Editor/AnimGraph/Private/AnimGraphNode_BlendListByEnum.cpp
+++ b/Engine/Source/Editor/AnimGraph/Private/AnimGraphNode_BlendListByEnum.cpp
@@ -52,19 +52,6 @@ FText UAnimGraphNode_BlendListByEnum::GetNodeTitle(ENodeTitleType::Type TitleTyp
 	}
 }
 
-FString UAnimGraphNode_BlendListByEnum::GetNodeNativeTitle(ENodeTitleType::Type TitleType) const
-{
-	// Do not setup this function for localization, intentionally left unlocalized!
-	if (BoundEnum != NULL)
-	{
-		return FString::Printf(TEXT("Blend Poses (%s)"), *BoundEnum->GetName());
-	}
-	else
-	{
-		return FString::Printf(TEXT("ERROR: Blend Poses (by missing enum)"));
-	}
-}
-
 void UAnimGraphNode_BlendListByEnum::PostPlacedNewNode()
 {
 	// Make sure we start out with a pin
@@ -88,7 +75,6 @@ void UAnimGraphNode_BlendListByEnum::GetMenuEntries(FGraphContextMenuBuilder& Co
 
 			TSharedPtr<FEdGraphSchemaAction_K2NewNode> Action = FK2ActionMenuBuilder::AddNewNodeAction(ContextMenuBuilder, EnumTemplate->GetNodeCategory(), EnumTemplate->GetNodeTitle(ENodeTitleType::ListView), EnumTemplate->GetTooltip(), 0, EnumTemplate->GetKeywords());
 			Action->NodeTemplate = EnumTemplate;
-			Action->SearchTitle = EnumTemplate->GetNodeSearchTitle();
 		}
 	}
 }
diff --git a/Engine/Source/Editor/AnimGraph/Private/AnimGraphNode_BlendSpaceBase.cpp b/Engine/Source/Editor/AnimGraph/Private/AnimGraphNode_BlendSpaceBase.cpp
--- a/Engine/Source/Editor/AnimGraph/Private/AnimGraphNode_BlendSpaceBase.cpp
+++ b/Engine/Source/Editor/AnimGraph/Private/AnimGraphNode_BlendSpaceBase.cpp
@@ -38,8 +38,6 @@ struct FNewBlendSpacePlayerAction : public FEdGraphSchemaAction_K2NewNode
 
 		// Grab extra keywords
 		Keywords = BlendSpace->GetPathName();
-
-		SearchTitle = NodeTemplate->GetNodeSearchTitle();
 	}
 };
 
diff --git a/Engine/Source/Editor/AnimGraph/Private/AnimGraphNode_BlendSpaceEvaluator.cpp b/Engine/Source/Editor/AnimGraph/Private/AnimGraphNode_BlendSpaceEvaluator.cpp
--- a/Engine/Source/Editor/AnimGraph/Private/AnimGraphNode_BlendSpaceEvaluator.cpp
+++ b/Engine/Source/Editor/AnimGraph/Private/AnimGraphNode_BlendSpaceEvaluator.cpp
@@ -48,29 +48,6 @@ FText UAnimGraphNode_BlendSpaceEvaluator::GetNodeTitle(ENodeTitleType::Type Titl
 	}
 }
 
-FString UAnimGraphNode_BlendSpaceEvaluator::GetNodeNativeTitle(ENodeTitleType::Type TitleType) const
-{
-	// Do not setup this function for localization, intentionally left unlocalized!
-	const FString BlendSpaceName((Node.BlendSpace != NULL) ? *(Node.BlendSpace->GetName()) : TEXT("(None)"));
-
-	if (TitleType == ENodeTitleType::ListView)
-	{
-		return FString::Printf(TEXT("Blendspace Evaluator '%s'"), *BlendSpaceName);
-	}
-	else
-	{
-		FString Title = FString::Printf(TEXT("%s\nBlendspace Evaluator"), *BlendSpaceName);
-
-		if ((TitleType == ENodeTitleType::FullTitle) && (SyncGroup.GroupName != NAME_None))
-		{
-			Title += TEXT("\n");
-			Title += FString::Printf(TEXT("Sync group %s"), *SyncGroup.GroupName.ToString());
-		}
-
-		return Title;
-	}
-}
-
 void UAnimGraphNode_BlendSpaceEvaluator::GetMenuEntries(FGraphContextMenuBuilder& ContextMenuBuilder) const
 {
 	// Intentionally empty so that we don't get duplicate blend space entries.
diff --git a/Engine/Source/Editor/AnimGraph/Private/AnimGraphNode_BlendSpacePlayer.cpp b/Engine/Source/Editor/AnimGraph/Private/AnimGraphNode_BlendSpacePlayer.cpp
--- a/Engine/Source/Editor/AnimGraph/Private/AnimGraphNode_BlendSpacePlayer.cpp
+++ b/Engine/Source/Editor/AnimGraph/Private/AnimGraphNode_BlendSpacePlayer.cpp
@@ -48,29 +48,6 @@ FText UAnimGraphNode_BlendSpacePlayer::GetNodeTitle(ENodeTitleType::Type TitleTy
 	}
 }
 
-FString UAnimGraphNode_BlendSpacePlayer::GetNodeNativeTitle(ENodeTitleType::Type TitleType) const
-{
-	// Do not setup this function for localization, intentionally left unlocalized!
-	const FString BlendSpaceName((Node.BlendSpace != NULL) ? *(Node.BlendSpace->GetName()) : TEXT("(None)"));
-
-	if (TitleType == ENodeTitleType::ListView)
-	{
-		return FString::Printf(TEXT("Blendspace Player'%s'"), *BlendSpaceName);
-	}
-	else
-	{
-		FString Title = FString::Printf(TEXT("%s\nBlendspace Player"), *BlendSpaceName);
-
-		if ((TitleType == ENodeTitleType::FullTitle) && (SyncGroup.GroupName != NAME_None))
-		{
-			Title += TEXT("\n");
-			Title += FString::Printf(TEXT("Sync group %s"), *SyncGroup.GroupName.ToString());
-		}
-
-		return Title;
-	}
-}
-
 void UAnimGraphNode_BlendSpacePlayer::GetMenuEntries(FGraphContextMenuBuilder& ContextMenuBuilder) const
 {
 	const bool bWantAimOffsets = false;
diff --git a/Engine/Source/Editor/AnimGraph/Private/AnimGraphNode_CopyBone.cpp b/Engine/Source/Editor/AnimGraph/Private/AnimGraphNode_CopyBone.cpp
--- a/Engine/Source/Editor/AnimGraph/Private/AnimGraphNode_CopyBone.cpp
+++ b/Engine/Source/Editor/AnimGraph/Private/AnimGraphNode_CopyBone.cpp
@@ -32,10 +32,4 @@ FText UAnimGraphNode_CopyBone::GetNodeTitle(ENodeTitleType::Type TitleType) cons
 	return FText::Format(LOCTEXT("AnimGraphNode_CopyBone_Title", "{ControllerDescription}\nSource Bone: {SourceBoneName}\nTarget Bone: {TargetBoneName}"), Args);
 }
 
-FString UAnimGraphNode_CopyBone::GetNodeNativeTitle(ENodeTitleType::Type TitleType) const
-{
-	// Do not setup this function for localization, intentionally left unlocalized!
-	return FString::Printf(TEXT("%s\nSource Bone: %s\nTarget Bone: %s"), *GetControllerDescription().ToString(), *Node.SourceBone.BoneName.ToString(), *Node.TargetBone.BoneName.ToString());
-}
-
 #undef LOCTEXT_NAMESPACE
diff --git a/Engine/Source/Editor/AnimGraph/Private/AnimGraphNode_HandIKRetargeting.cpp b/Engine/Source/Editor/AnimGraph/Private/AnimGraphNode_HandIKRetargeting.cpp
--- a/Engine/Source/Editor/AnimGraph/Private/AnimGraphNode_HandIKRetargeting.cpp
+++ b/Engine/Source/Editor/AnimGraph/Private/AnimGraphNode_HandIKRetargeting.cpp
@@ -30,10 +30,4 @@ FText UAnimGraphNode_HandIKRetargeting::GetNodeTitle(ENodeTitleType::Type TitleT
 	return FText::Format(LOCTEXT("AnimGraphNode_CopyBone_Title", "{ControllerDescription}"), Args);
 }
 
-FString UAnimGraphNode_HandIKRetargeting::GetNodeNativeTitle(ENodeTitleType::Type TitleType) const
-{
-	// Do not setup this function for localization, intentionally left unlocalized!
-	return FString::Printf(TEXT("%s"), *GetControllerDescription().ToString());
-}
-
 #undef LOCTEXT_NAMESPACE
diff --git a/Engine/Source/Editor/AnimGraph/Private/AnimGraphNode_RotationMultiplier.cpp b/Engine/Source/Editor/AnimGraph/Private/AnimGraphNode_RotationMultiplier.cpp
--- a/Engine/Source/Editor/AnimGraph/Private/AnimGraphNode_RotationMultiplier.cpp
+++ b/Engine/Source/Editor/AnimGraph/Private/AnimGraphNode_RotationMultiplier.cpp
@@ -38,13 +38,4 @@ FText UAnimGraphNode_RotationMultiplier::GetNodeTitle(ENodeTitleType::Type Title
 	}
 }
 
-FString UAnimGraphNode_RotationMultiplier::GetNodeNativeTitle(ENodeTitleType::Type TitleType) const
-{
-	// Do not setup this function for localization, intentionally left unlocalized!
-	FString Result = GetControllerDescription().ToString();
-	Result += (TitleType == ENodeTitleType::ListView) ? TEXT(" - ") : TEXT("\n");
-	Result += FString::Printf(TEXT("Bone: %s"), *Node.TargetBone.BoneName.ToString());
-	return Result;
-}
-
 #undef LOCTEXT_NAMESPACE
diff --git a/Engine/Source/Editor/AnimGraph/Private/AnimGraphNode_RotationOffsetBlendSpace.cpp b/Engine/Source/Editor/AnimGraph/Private/AnimGraphNode_RotationOffsetBlendSpace.cpp
--- a/Engine/Source/Editor/AnimGraph/Private/AnimGraphNode_RotationOffsetBlendSpace.cpp
+++ b/Engine/Source/Editor/AnimGraph/Private/AnimGraphNode_RotationOffsetBlendSpace.cpp
@@ -37,22 +37,6 @@ FText UAnimGraphNode_RotationOffsetBlendSpace::GetNodeTitle(ENodeTitleType::Type
 	}
 }
 
-FString UAnimGraphNode_RotationOffsetBlendSpace::GetNodeNativeTitle(ENodeTitleType::Type TitleType) const
-{
-	// Do not setup this function for localization, intentionally left unlocalized!
-	
-	const FString BlendSpaceName((Node.BlendSpace != NULL) ? *(Node.BlendSpace->GetName()) : TEXT("(None)"));
-
-	if (TitleType == ENodeTitleType::ListView)
-	{
-		return FString::Printf(TEXT("AimOffset '%s'"), *BlendSpaceName);
-	}
-	else
-	{
-		return FString::Printf(TEXT("%s\nAimOffset"), *BlendSpaceName);
-	}
-}
-
 void UAnimGraphNode_RotationOffsetBlendSpace::GetMenuEntries(FGraphContextMenuBuilder& ContextMenuBuilder) const
 {
 	const bool bWantAimOffsets = true;
diff --git a/Engine/Source/Editor/AnimGraph/Private/AnimGraphNode_SaveCachedPose.cpp b/Engine/Source/Editor/AnimGraph/Private/AnimGraphNode_SaveCachedPose.cpp
--- a/Engine/Source/Editor/AnimGraph/Private/AnimGraphNode_SaveCachedPose.cpp
+++ b/Engine/Source/Editor/AnimGraph/Private/AnimGraphNode_SaveCachedPose.cpp
@@ -58,19 +58,6 @@ FText UAnimGraphNode_SaveCachedPose::GetNodeTitle(ENodeTitleType::Type TitleType
 	}
 }
 
-FString UAnimGraphNode_SaveCachedPose::GetNodeNativeTitle(ENodeTitleType::Type TitleType) const
-{
-	// Do not setup this function for localization, intentionally left unlocalized!
-	if (TitleType == ENodeTitleType::EditableTitle)
-	{
-		return CacheName;
-	}
-	else
-	{
-		return FString::Printf(TEXT("Save cached pose '%s'"), *CacheName);
-	}
-}
-
 FString UAnimGraphNode_SaveCachedPose::GetNodeCategory() const
 {
 	return TEXT("Cached Poses");
@@ -92,7 +79,6 @@ void UAnimGraphNode_SaveCachedPose::GetMenuEntries(FGraphContextMenuBuilder& Con
 
 			TSharedPtr<FEdGraphSchemaAction_K2NewNode> SaveCachedPoseAction = FK2ActionMenuBuilder::AddNewNodeAction(ContextMenuBuilder, GetNodeCategory(), LOCTEXT("NewSaveCachedPose", "New Save cached pose..."), SaveCachedPose->GetTooltip(), 0, SaveCachedPose->GetKeywords());
 			SaveCachedPoseAction->NodeTemplate = SaveCachedPose;
-			SaveCachedPoseAction->SearchTitle = SaveCachedPoseAction->NodeTemplate->GetNodeSearchTitle();
 		}
 	}
 }
diff --git a/Engine/Source/Editor/AnimGraph/Private/AnimGraphNode_SequenceEvaluator.cpp b/Engine/Source/Editor/AnimGraph/Private/AnimGraphNode_SequenceEvaluator.cpp
--- a/Engine/Source/Editor/AnimGraph/Private/AnimGraphNode_SequenceEvaluator.cpp
+++ b/Engine/Source/Editor/AnimGraph/Private/AnimGraphNode_SequenceEvaluator.cpp
@@ -65,19 +65,6 @@ FText UAnimGraphNode_SequenceEvaluator::GetNodeTitle(ENodeTitleType::Type TitleT
 	}
 }
 
-FString UAnimGraphNode_SequenceEvaluator::GetNodeNativeTitle(ENodeTitleType::Type TitleType) const
-{
-	// Do not setup this function for localization, intentionally left unlocalized!
-	if (Node.Sequence && Node.Sequence->IsValidAdditive())
-	{
-		return FString::Printf(TEXT("Evaluate %s (additive)"), (Node.Sequence != NULL) ? *(Node.Sequence->GetName()) : TEXT("(None)"));
-	}
-	else
-	{
-		return FString::Printf(TEXT("Evaluate %s"), (Node.Sequence != NULL) ? *(Node.Sequence->GetName()) : TEXT("(None)"));
-	}
-}
-
 void UAnimGraphNode_SequenceEvaluator::GetMenuEntries(FGraphContextMenuBuilder& ContextMenuBuilder) const
 {
 	// Intentionally empty; you can drop down a regular sequence player and convert into a sequence evaluator in the right-click menu.
diff --git a/Engine/Source/Editor/AnimGraph/Private/AnimGraphNode_SequencePlayer.cpp b/Engine/Source/Editor/AnimGraph/Private/AnimGraphNode_SequencePlayer.cpp
--- a/Engine/Source/Editor/AnimGraph/Private/AnimGraphNode_SequencePlayer.cpp
+++ b/Engine/Source/Editor/AnimGraph/Private/AnimGraphNode_SequencePlayer.cpp
@@ -92,21 +92,6 @@ FText UAnimGraphNode_SequencePlayer::GetNodeTitle(ENodeTitleType::Type TitleType
 	return Title;
 }
 
-FString UAnimGraphNode_SequencePlayer::GetNodeNativeTitle(ENodeTitleType::Type TitleType) const
-{
-	// Do not setup this function for localization, intentionally left unlocalized!
-	const bool bAdditive = ((Node.Sequence != NULL) && Node.Sequence->IsValidAdditive());
-	FString Title = GetTitleGivenAssetInfo((Node.Sequence != NULL) ? FText::FromString(Node.Sequence->GetName()) : LOCTEXT("None", "(None)"), bAdditive).ToString();
-
-	if ((TitleType == ENodeTitleType::FullTitle) && (SyncGroup.GroupName != NAME_None))
-	{
-		Title += TEXT("\n");
-		Title += FString::Printf(TEXT("Sync group %s"), *SyncGroup.GroupName.ToString());
-	}
-
-	return Title;
-}
-
 FText UAnimGraphNode_SequencePlayer::GetTitleGivenAssetInfo(const FText& AssetName, bool bKnownToBeAdditive)
 {
 	FFormatNamedArguments Args;
diff --git a/Engine/Source/Editor/AnimGraph/Private/AnimGraphNode_SkeletalControlBase.cpp b/Engine/Source/Editor/AnimGraph/Private/AnimGraphNode_SkeletalControlBase.cpp
--- a/Engine/Source/Editor/AnimGraph/Private/AnimGraphNode_SkeletalControlBase.cpp
+++ b/Engine/Source/Editor/AnimGraph/Private/AnimGraphNode_SkeletalControlBase.cpp
@@ -52,7 +52,6 @@ void UAnimGraphNode_SkeletalControlBase::GetMenuEntries(FGraphContextMenuBuilder
 
 	TSharedPtr<FEdGraphSchemaAction_K2NewNode> NodeAction = FK2ActionMenuBuilder::AddNewNodeAction(ContextMenuBuilder, Category, MenuDesc, Tooltip, 0, Keywords);
 	NodeAction->NodeTemplate = TemplateNode;
-	NodeAction->SearchTitle = NodeAction->NodeTemplate->GetNodeSearchTitle();
 }
 
 #undef LOCTEXT_NAMESPACE
\ No newline at end of file
diff --git a/Engine/Source/Editor/AnimGraph/Private/AnimGraphNode_Slot.cpp b/Engine/Source/Editor/AnimGraph/Private/AnimGraphNode_Slot.cpp
--- a/Engine/Source/Editor/AnimGraph/Private/AnimGraphNode_Slot.cpp
+++ b/Engine/Source/Editor/AnimGraph/Private/AnimGraphNode_Slot.cpp
@@ -43,21 +43,6 @@ FText UAnimGraphNode_Slot::GetNodeTitle(ENodeTitleType::Type TitleType) const
 	}
 }
 
-FString UAnimGraphNode_Slot::GetNodeNativeTitle(ENodeTitleType::Type TitleType) const
-{
-	// Do not setup this function for localization, intentionally left unlocalized!
-	const FString SlotName = (Node.SlotName != NAME_None) ? Node.SlotName.ToString() : TEXT("(No slot name)");
-
-	if (TitleType == ENodeTitleType::ListView)
-	{
-		return FString::Printf(TEXT("Slot '%s'"), *SlotName);
-	}
-	else
-	{
-		return FString::Printf(TEXT("%s\nSlot"), *SlotName);
-	}
-}
-
 FString UAnimGraphNode_Slot::GetNodeCategory() const
 {
 	return TEXT("Blends");
diff --git a/Engine/Source/Editor/AnimGraph/Private/AnimGraphNode_SpringBone.cpp b/Engine/Source/Editor/AnimGraph/Private/AnimGraphNode_SpringBone.cpp
--- a/Engine/Source/Editor/AnimGraph/Private/AnimGraphNode_SpringBone.cpp
+++ b/Engine/Source/Editor/AnimGraph/Private/AnimGraphNode_SpringBone.cpp
@@ -39,13 +39,4 @@ FText UAnimGraphNode_SpringBone::GetNodeTitle(ENodeTitleType::Type TitleType) co
 	}
 }
 
-FString UAnimGraphNode_SpringBone::GetNodeNativeTitle(ENodeTitleType::Type TitleType) const
-{
-	// Do not setup this function for localization, intentionally left unlocalized!
-	FString Result = GetControllerDescription().ToString();
-	Result += (TitleType == ENodeTitleType::ListView) ? TEXT(" - ") : TEXT("\n");
-	Result += FString::Printf(TEXT("Bone: %s"), *Node.SpringBone.BoneName.ToString());
-	return Result;
-}
-
 #undef LOCTEXT_NAMESPACE
\ No newline at end of file
diff --git a/Engine/Source/Editor/AnimGraph/Private/AnimGraphNode_StateMachineBase.cpp b/Engine/Source/Editor/AnimGraph/Private/AnimGraphNode_StateMachineBase.cpp
--- a/Engine/Source/Editor/AnimGraph/Private/AnimGraphNode_StateMachineBase.cpp
+++ b/Engine/Source/Editor/AnimGraph/Private/AnimGraphNode_StateMachineBase.cpp
@@ -71,18 +71,6 @@ FText UAnimGraphNode_StateMachineBase::GetNodeTitle(ENodeTitleType::Type TitleTy
 	return FirstLine;
 }
 
-FString UAnimGraphNode_StateMachineBase::GetNodeNativeTitle(ENodeTitleType::Type TitleType) const
-{
-	// Do not setup this function for localization, intentionally left unlocalized!
-	FString FirstLine = (EditorStateMachineGraph != NULL) ? EditorStateMachineGraph->GetName() : TEXT("Error: No Graph");
-	if(TitleType == ENodeTitleType::FullTitle)
-	{
-		return FirstLine + TEXT("\nState Machine");
-	}
-
-	return FirstLine;
-}
-
 FString UAnimGraphNode_StateMachineBase::GetNodeCategory() const
 {
 	return TEXT("State Machines");
diff --git a/Engine/Source/Editor/AnimGraph/Private/AnimGraphNode_TwoBoneIK.cpp b/Engine/Source/Editor/AnimGraph/Private/AnimGraphNode_TwoBoneIK.cpp
--- a/Engine/Source/Editor/AnimGraph/Private/AnimGraphNode_TwoBoneIK.cpp
+++ b/Engine/Source/Editor/AnimGraph/Private/AnimGraphNode_TwoBoneIK.cpp
@@ -39,15 +39,6 @@ FText UAnimGraphNode_TwoBoneIK::GetNodeTitle(ENodeTitleType::Type TitleType) con
 	}
 }
 
-FString UAnimGraphNode_TwoBoneIK::GetNodeNativeTitle(ENodeTitleType::Type TitleType) const
-{
-	// Do not setup this function for localization, intentionally left unlocalized!
-	FString Result = GetControllerDescription().ToString();
-	Result += (TitleType == ENodeTitleType::ListView) ? TEXT(" - ") : TEXT("\n");
-	Result += FString::Printf(TEXT("Bone: %s"), *Node.IKBone.BoneName.ToString());
-	return Result;
-}
-
 void UAnimGraphNode_TwoBoneIK::Draw( FPrimitiveDrawInterface* PDI, USkeletalMeshComponent* SkelMeshComp ) const
 {
 	if ( SkelMeshComp && SkelMeshComp->SkeletalMesh && SkelMeshComp->SkeletalMesh->Skeleton )
diff --git a/Engine/Source/Editor/AnimGraph/Private/AnimGraphNode_UseCachedPose.cpp b/Engine/Source/Editor/AnimGraph/Private/AnimGraphNode_UseCachedPose.cpp
--- a/Engine/Source/Editor/AnimGraph/Private/AnimGraphNode_UseCachedPose.cpp
+++ b/Engine/Source/Editor/AnimGraph/Private/AnimGraphNode_UseCachedPose.cpp
@@ -30,12 +30,6 @@ FText UAnimGraphNode_UseCachedPose::GetNodeTitle(ENodeTitleType::Type TitleType)
 	return FText::Format(LOCTEXT("AnimGraphNode_UseCachedPose_Title", "Use cached pose '{CachePoseName}'"), Args);
 }
 
-FString UAnimGraphNode_UseCachedPose::GetNodeNativeTitle(ENodeTitleType::Type TitleType) const
-{
-	// Do not setup this function for localization, intentionally left unlocalized!
-	return FString::Printf(TEXT("Use cached pose '%s'"), *NameOfCache);
-}
-
 FString UAnimGraphNode_UseCachedPose::GetNodeCategory() const
 {
 	return TEXT("Cached Poses");
@@ -58,7 +52,6 @@ void UAnimGraphNode_UseCachedPose::GetMenuEntries(FGraphContextMenuBuilder& Cont
 
 			TSharedPtr<FEdGraphSchemaAction_K2NewNode> UseCachedPoseAction = FK2ActionMenuBuilder::AddNewNodeAction(ContextMenuBuilder, GetNodeCategory(), UseCachedPose->GetNodeTitle(ENodeTitleType::ListView), UseCachedPose->GetTooltip(), 0, UseCachedPose->GetKeywords());
 			UseCachedPoseAction->NodeTemplate = UseCachedPose;
-			UseCachedPoseAction->SearchTitle = UseCachedPose->GetNodeSearchTitle();
 		}
 	}
 }
diff --git a/Engine/Source/Editor/AnimGraph/Private/AnimStateTransitionNode.cpp b/Engine/Source/Editor/AnimGraph/Private/AnimStateTransitionNode.cpp
--- a/Engine/Source/Editor/AnimGraph/Private/AnimStateTransitionNode.cpp
+++ b/Engine/Source/Editor/AnimGraph/Private/AnimStateTransitionNode.cpp
@@ -154,26 +154,6 @@ FText UAnimStateTransitionNode::GetNodeTitle(ENodeTitleType::Type TitleType) con
 	}
 }
 
-FString UAnimStateTransitionNode::GetNodeNativeTitle(ENodeTitleType::Type TitleType) const
-{
-	// Do not setup this function for localization, intentionally left unlocalized!
-	UAnimStateNodeBase* PrevState = GetPreviousState();
-	UAnimStateNodeBase* NextState = GetNextState();
-
-	if (!SharedRulesName.IsEmpty())
-	{
-		return SharedRulesName;
-	}
-	else if ((PrevState != NULL) && (NextState != NULL))
-	{
-		return FString::Printf(TEXT("%s to %s"), *(PrevState->GetStateName()), *(NextState->GetStateName()));
-	}
-	else
-	{
-		return FString::Printf(TEXT("Trans %s"), (BoundGraph != NULL) ? *(BoundGraph->GetName()) : TEXT("(null)"));
-	}
-}
-
 FString UAnimStateTransitionNode::GetTooltip() const
 {
 	return TEXT("This is a state transition");
diff --git a/Engine/Source/Editor/AnimGraph/Private/AnimationGraphSchema.cpp b/Engine/Source/Editor/AnimGraph/Private/AnimationGraphSchema.cpp
--- a/Engine/Source/Editor/AnimGraph/Private/AnimationGraphSchema.cpp
+++ b/Engine/Source/Editor/AnimGraph/Private/AnimationGraphSchema.cpp
@@ -64,7 +64,6 @@ void UAnimationGraphSchema::GetStateMachineMenuItems(FGraphContextMenuBuilder& C
 
 			TSharedPtr<FEdGraphSchemaAction_K2NewNode> Action = FK2ActionMenuBuilder::AddNewNodeAction(ContextMenuBuilder, StateCategory, FText::FromString(Title), NodeTemplate->GetTooltip(), 0, NodeTemplate->GetKeywords());
 			Action->NodeTemplate = NodeTemplate;
-			Action->SearchTitle = NodeTemplate->GetNodeSearchTitle();
 		}
 	}
 }
diff --git a/Engine/Source/Editor/AnimGraph/Private/AnimationStateMachineSchema.cpp b/Engine/Source/Editor/AnimGraph/Private/AnimationStateMachineSchema.cpp
--- a/Engine/Source/Editor/AnimGraph/Private/AnimationStateMachineSchema.cpp
+++ b/Engine/Source/Editor/AnimGraph/Private/AnimationStateMachineSchema.cpp
@@ -266,14 +266,12 @@ void UAnimationStateMachineSchema::GetGraphContextActions(FGraphContextMenuBuild
 	{
 		TSharedPtr<FEdGraphSchemaAction_NewStateNode> Action = AddNewStateNodeAction(ContextMenuBuilder, TEXT(""), LOCTEXT("AddState", "Add State..."), TEXT("A new state"));
 		Action->NodeTemplate = NewObject<UAnimStateNode>(ContextMenuBuilder.OwnerOfTemporaries);
-		Action->SearchTitle = Action->NodeTemplate->GetNodeSearchTitle();
 	}
 
 	// Add conduit node
 	{
 		TSharedPtr<FEdGraphSchemaAction_NewStateNode> Action = AddNewStateNodeAction(ContextMenuBuilder, TEXT(""), LOCTEXT("AddConduit", "Add Conduit..."), TEXT("A new conduit state"));
 		Action->NodeTemplate = NewObject<UAnimStateConduitNode>(ContextMenuBuilder.OwnerOfTemporaries);
-		Action->SearchTitle = Action->NodeTemplate->GetNodeSearchTitle();
 	}
 
 	// Entry point (only if doesn't already exist)
@@ -293,7 +291,6 @@ void UAnimationStateMachineSchema::GetGraphContextActions(FGraphContextMenuBuild
 		{
 			TSharedPtr<FEdGraphSchemaAction_NewStateNode> Action = AddNewStateNodeAction(ContextMenuBuilder, TEXT(""), LOCTEXT("AddEntryPoint", "Add Entry Point..."), TEXT("Define State Machine's Entry Point"));
 			Action->NodeTemplate = NewObject<UAnimStateEntryNode>(ContextMenuBuilder.OwnerOfTemporaries);
-			Action->SearchTitle = Action->NodeTemplate->GetNodeSearchTitle();
 		}
 	}
 
diff --git a/Engine/Source/Editor/AnimGraph/Private/AnimationTransitionSchema.cpp b/Engine/Source/Editor/AnimGraph/Private/AnimationTransitionSchema.cpp
--- a/Engine/Source/Editor/AnimGraph/Private/AnimationTransitionSchema.cpp
+++ b/Engine/Source/Editor/AnimGraph/Private/AnimationTransitionSchema.cpp
@@ -96,7 +96,6 @@ void UAnimationTransitionSchema::GetSourceStateActions(FGraphContextMenuBuilder&
 
 								TSharedPtr<FEdGraphSchemaAction_K2NewNode> Action = FK2ActionMenuBuilder::AddNewNodeAction(ContextMenuBuilder, Category_AssetPlayer, Title, NodeTemplate->GetTooltip(), 0, NodeTemplate->GetKeywords());
 								Action->NodeTemplate = NodeTemplate;
-								Action->SearchTitle = Action->NodeTemplate->GetNodeSearchTitle();
 							}
 						}
 					}
@@ -120,7 +119,6 @@ void UAnimationTransitionSchema::GetSourceStateActions(FGraphContextMenuBuilder&
 
 						TSharedPtr<FEdGraphSchemaAction_K2NewNode> Action = FK2ActionMenuBuilder::AddNewNodeAction(ContextMenuBuilder, Category_Transition, Title, NodeTemplate->GetTooltip(), 0, NodeTemplate->GetKeywords());
 						Action->NodeTemplate = NodeTemplate;
-						Action->SearchTitle = Action->NodeTemplate->GetNodeSearchTitle();
 					}
 				}
 			}
diff --git a/Engine/Source/Editor/AnimGraph/Private/K2Node_TransitionRuleGetter.cpp b/Engine/Source/Editor/AnimGraph/Private/K2Node_TransitionRuleGetter.cpp
--- a/Engine/Source/Editor/AnimGraph/Private/K2Node_TransitionRuleGetter.cpp
+++ b/Engine/Source/Editor/AnimGraph/Private/K2Node_TransitionRuleGetter.cpp
@@ -106,39 +106,6 @@ FText UK2Node_TransitionRuleGetter::GetNodeTitle(ENodeTitleType::Type TitleType)
 	return Super::GetNodeTitle(TitleType);
 }
 
-FString UK2Node_TransitionRuleGetter::GetNodeNativeTitle(ENodeTitleType::Type TitleType) const
-{
-	// Do not setup this function for localization, intentionally left unlocalized!
-	if (AssociatedAnimAssetPlayerNode != NULL)
-	{
-		UAnimationAsset* BoundAsset = AssociatedAnimAssetPlayerNode->GetAnimationAsset();
-		if (BoundAsset)
-		{
-			return  FString::Printf(TEXT("%s Asset"), *BoundAsset->GetName());
-		}
-	}
-	else if (AssociatedStateNode != NULL)
-	{
-		if (UAnimStateNode* State = Cast<UAnimStateNode>(AssociatedStateNode))
-		{
-			const FString OwnerName = State->GetOuter()->GetName();
-
-			return FString::Printf(TEXT("%s.%s State"), *OwnerName, *(State->GetStateName()));
-		}
-	}
-	else if (GetterType == ETransitionGetter::CurrentTransitionDuration)
-	{
-		return TEXT("Transition");
-	}
-	else if (GetterType == ETransitionGetter::CurrentState_ElapsedTime ||
-		GetterType == ETransitionGetter::CurrentState_GetBlendWeight)
-	{
-		return TEXT("Current State");
-	}
-
-	return Super::GetNodeNativeTitle(TitleType);
-}
-
 FString UK2Node_TransitionRuleGetter::GetTooltip() const
 {
 	return GetNodeTitle(ENodeTitleType::FullTitle).ToString();
diff --git a/Engine/Source/Editor/BehaviorTreeEditor/Private/BehaviorTreeDecoratorGraphNode_Logic.cpp b/Engine/Source/Editor/BehaviorTreeEditor/Private/BehaviorTreeDecoratorGraphNode_Logic.cpp
--- a/Engine/Source/Editor/BehaviorTreeEditor/Private/BehaviorTreeDecoratorGraphNode_Logic.cpp
+++ b/Engine/Source/Editor/BehaviorTreeEditor/Private/BehaviorTreeDecoratorGraphNode_Logic.cpp
@@ -47,7 +47,6 @@ void UBehaviorTreeDecoratorGraphNode_Logic::GetMenuEntries(struct FGraphContextM
 		UBehaviorTreeDecoratorGraphNode_Logic* OpNode = NewObject<UBehaviorTreeDecoratorGraphNode_Logic>(ContextMenuBuilder.OwnerOfTemporaries);
 		OpNode->LogicMode = Modes[i];
 		AddOpAction->NodeTemplate = OpNode;
-		AddOpAction->SearchTitle = AddOpAction->NodeTemplate->GetNodeSearchTitle();
 	}
 }
 
diff --git a/Engine/Source/Editor/BehaviorTreeEditor/Private/EdGraphSchema_BehaviorTree.cpp b/Engine/Source/Editor/BehaviorTreeEditor/Private/EdGraphSchema_BehaviorTree.cpp
--- a/Engine/Source/Editor/BehaviorTreeEditor/Private/EdGraphSchema_BehaviorTree.cpp
+++ b/Engine/Source/Editor/BehaviorTreeEditor/Private/EdGraphSchema_BehaviorTree.cpp
@@ -278,7 +278,6 @@ void UEdGraphSchema_BehaviorTree::GetGraphNodeContextActions(FGraphContextMenuBu
 			TSharedPtr<FBehaviorTreeSchemaAction_NewSubNode> AddOpAction = AddNewSubNodeAction(ContextMenuBuilder, TEXT(""), FText::FromString(OpNode->GetNodeTypeDescription()), "");
 			AddOpAction->ParentNode = Cast<UBehaviorTreeGraphNode>(ContextMenuBuilder.SelectedObjects[0]);
 			AddOpAction->NodeTemplate = OpNode;
-			AddOpAction->SearchTitle = AddOpAction->NodeTemplate->GetNodeSearchTitle();
 		}
 
 		for (int32 i = 0; i < NodeClasses.Num(); i++)
@@ -291,7 +290,6 @@ void UEdGraphSchema_BehaviorTree::GetGraphNodeContextActions(FGraphContextMenuBu
 			TSharedPtr<FBehaviorTreeSchemaAction_NewSubNode> AddOpAction = AddNewSubNodeAction(ContextMenuBuilder, TEXT(""), NodeTypeName, "");
 			AddOpAction->ParentNode = Cast<UBehaviorTreeGraphNode>(ContextMenuBuilder.SelectedObjects[0]);
 			AddOpAction->NodeTemplate = OpNode;
-			AddOpAction->SearchTitle = AddOpAction->NodeTemplate->GetNodeSearchTitle();
 		}
 	}
 	else if (SubNodeType == ESubNode::Service)
@@ -307,7 +305,6 @@ void UEdGraphSchema_BehaviorTree::GetGraphNodeContextActions(FGraphContextMenuBu
 			TSharedPtr<FBehaviorTreeSchemaAction_NewSubNode> AddOpAction = AddNewSubNodeAction(ContextMenuBuilder, TEXT(""), NodeTypeName, "");
 			AddOpAction->ParentNode = Cast<UBehaviorTreeGraphNode>(ContextMenuBuilder.SelectedObjects[0]);
 			AddOpAction->NodeTemplate = OpNode;
-			AddOpAction->SearchTitle = AddOpAction->NodeTemplate->GetNodeSearchTitle();
 		}
 	}
 }
@@ -343,7 +340,6 @@ void UEdGraphSchema_BehaviorTree::GetGraphContextActions(FGraphContextMenuBuilde
 
 			OpNode->ClassData = NodeClasses[i];
 			AddOpAction->NodeTemplate = OpNode;
-			AddOpAction->SearchTitle = AddOpAction->NodeTemplate->GetNodeSearchTitle();
 		}
 	}
 
@@ -367,7 +363,6 @@ void UEdGraphSchema_BehaviorTree::GetGraphContextActions(FGraphContextMenuBuilde
 			UBehaviorTreeGraphNode* OpNode = ConstructObject<UBehaviorTreeGraphNode>(GraphNodeClass, ContextMenuBuilder.OwnerOfTemporaries);
 			OpNode->ClassData = NodeClasses[i];
 			AddOpAction->NodeTemplate = OpNode;
-			AddOpAction->SearchTitle = AddOpAction->NodeTemplate->GetNodeSearchTitle();
 		}
 	}
 	
diff --git a/Engine/Source/Editor/BehaviorTreeEditor/Private/EdGraphSchema_BehaviorTreeDecorator.cpp b/Engine/Source/Editor/BehaviorTreeEditor/Private/EdGraphSchema_BehaviorTreeDecorator.cpp
--- a/Engine/Source/Editor/BehaviorTreeEditor/Private/EdGraphSchema_BehaviorTreeDecorator.cpp
+++ b/Engine/Source/Editor/BehaviorTreeEditor/Private/EdGraphSchema_BehaviorTreeDecorator.cpp
@@ -138,7 +138,6 @@ void UEdGraphSchema_BehaviorTreeDecorator::GetGraphContextActions(FGraphContextM
 		UBehaviorTreeDecoratorGraphNode_Decorator* OpNode = NewObject<UBehaviorTreeDecoratorGraphNode_Decorator>(ContextMenuBuilder.OwnerOfTemporaries);
 		OpNode->ClassData = NodeClasses[i];
 		AddOpAction->NodeTemplate = OpNode;
-		AddOpAction->SearchTitle = AddOpAction->NodeTemplate->GetNodeSearchTitle();
 	}
 
 #if WITH_EDITOR
diff --git a/Engine/Source/Editor/BlueprintGraph/Classes/K2Node_ActorBoundEvent.h b/Engine/Source/Editor/BlueprintGraph/Classes/K2Node_ActorBoundEvent.h
--- a/Engine/Source/Editor/BlueprintGraph/Classes/K2Node_ActorBoundEvent.h
+++ b/Engine/Source/Editor/BlueprintGraph/Classes/K2Node_ActorBoundEvent.h
@@ -23,7 +23,6 @@ class UK2Node_ActorBoundEvent : public UK2Node_Event
 	virtual void ReconstructNode() override;
 	virtual void DestroyNode() override;
 	virtual FText GetNodeTitle(ENodeTitleType::Type TitleType) const override;
-	virtual FString GetNodeNativeTitle(ENodeTitleType::Type TitleType) const override;
 	virtual FString GetTooltip() const override;
 	virtual FString GetDocumentationLink() const override;
 	virtual FString GetDocumentationExcerptName() const override;
diff --git a/Engine/Source/Editor/BlueprintGraph/Classes/K2Node_AddDelegate.h b/Engine/Source/Editor/BlueprintGraph/Classes/K2Node_AddDelegate.h
--- a/Engine/Source/Editor/BlueprintGraph/Classes/K2Node_AddDelegate.h
+++ b/Engine/Source/Editor/BlueprintGraph/Classes/K2Node_AddDelegate.h
@@ -12,7 +12,6 @@ public:
 	// UEdGraphNode interface
 	virtual void AllocateDefaultPins() override;
 	virtual FText GetNodeTitle(ENodeTitleType::Type TitleType) const override;
-	virtual FString GetNodeNativeTitle(ENodeTitleType::Type TitleType) const override;
 	// End of UEdGraphNode interface
 
 	// Begin of UK2Node interface
diff --git a/Engine/Source/Editor/BlueprintGraph/Classes/K2Node_BreakStruct.h b/Engine/Source/Editor/BlueprintGraph/Classes/K2Node_BreakStruct.h
--- a/Engine/Source/Editor/BlueprintGraph/Classes/K2Node_BreakStruct.h
+++ b/Engine/Source/Editor/BlueprintGraph/Classes/K2Node_BreakStruct.h
@@ -14,7 +14,6 @@ class UK2Node_BreakStruct : public UK2Node_StructMemberGet
 	// Begin UEdGraphNode interface
 	virtual void AllocateDefaultPins() override;
 	virtual FText GetNodeTitle(ENodeTitleType::Type TitleType) const override;
-	virtual FString GetNodeNativeTitle(ENodeTitleType::Type TitleType) const override;
 	virtual FLinearColor GetNodeTitleColor() const override;
 	virtual FString GetTooltip() const override;
 	virtual void ValidateNodeDuringCompilation(class FCompilerResultsLog& MessageLog) const override;
diff --git a/Engine/Source/Editor/BlueprintGraph/Classes/K2Node_CallDelegate.h b/Engine/Source/Editor/BlueprintGraph/Classes/K2Node_CallDelegate.h
--- a/Engine/Source/Editor/BlueprintGraph/Classes/K2Node_CallDelegate.h
+++ b/Engine/Source/Editor/BlueprintGraph/Classes/K2Node_CallDelegate.h
@@ -13,7 +13,6 @@ public:
 	// UEdGraphNode interface
 	virtual void AllocateDefaultPins() override;
 	virtual FText GetNodeTitle(ENodeTitleType::Type TitleType) const override;
-	virtual FString GetNodeNativeTitle(ENodeTitleType::Type TitleType) const override;
 	// End of UEdGraphNode interface
 
 	// UK2Node interface
diff --git a/Engine/Source/Editor/BlueprintGraph/Classes/K2Node_CallFunction.h b/Engine/Source/Editor/BlueprintGraph/Classes/K2Node_CallFunction.h
--- a/Engine/Source/Editor/BlueprintGraph/Classes/K2Node_CallFunction.h
+++ b/Engine/Source/Editor/BlueprintGraph/Classes/K2Node_CallFunction.h
@@ -60,7 +60,6 @@ public:
 	virtual FLinearColor GetNodeTitleColor() const override;
 	virtual FString GetTooltip() const override;
 	virtual FText GetNodeTitle(ENodeTitleType::Type TitleType) const override;
-	virtual FString GetNodeNativeTitle(ENodeTitleType::Type TitleType) const override;
 	virtual FString GetDescriptiveCompiledName() const override;
 	virtual bool IsDeprecated() const override;
 	virtual bool ShouldWarnOnDeprecation() const override;
diff --git a/Engine/Source/Editor/BlueprintGraph/Classes/K2Node_CallParentFunction.h b/Engine/Source/Editor/BlueprintGraph/Classes/K2Node_CallParentFunction.h
--- a/Engine/Source/Editor/BlueprintGraph/Classes/K2Node_CallParentFunction.h
+++ b/Engine/Source/Editor/BlueprintGraph/Classes/K2Node_CallParentFunction.h
@@ -12,7 +12,6 @@ class UK2Node_CallParentFunction : public UK2Node_CallFunction
 
 	// Begin EdGraphNode interface
 	virtual FText GetNodeTitle(ENodeTitleType::Type TitleType) const override;
-	virtual FString GetNodeNativeTitle(ENodeTitleType::Type TitleType) const override;
 	virtual void AllocateDefaultPins() override;
 	virtual FLinearColor GetNodeTitleColor() const override;
 	virtual void PostPlacedNewNode() override;
diff --git a/Engine/Source/Editor/BlueprintGraph/Classes/K2Node_ClearDelegate.h b/Engine/Source/Editor/BlueprintGraph/Classes/K2Node_ClearDelegate.h
--- a/Engine/Source/Editor/BlueprintGraph/Classes/K2Node_ClearDelegate.h
+++ b/Engine/Source/Editor/BlueprintGraph/Classes/K2Node_ClearDelegate.h
@@ -11,7 +11,6 @@ class UK2Node_ClearDelegate : public UK2Node_BaseMCDelegate
 public:
 	// UEdGraphNode interface
 	virtual FText GetNodeTitle(ENodeTitleType::Type TitleType) const override;
-	virtual FString GetNodeNativeTitle(ENodeTitleType::Type TitleType) const override;
 	// End of UEdGraphNode interface
 
 	// Begin of UK2Node interface
diff --git a/Engine/Source/Editor/BlueprintGraph/Classes/K2Node_ComponentBoundEvent.h b/Engine/Source/Editor/BlueprintGraph/Classes/K2Node_ComponentBoundEvent.h
--- a/Engine/Source/Editor/BlueprintGraph/Classes/K2Node_ComponentBoundEvent.h
+++ b/Engine/Source/Editor/BlueprintGraph/Classes/K2Node_ComponentBoundEvent.h
@@ -20,7 +20,6 @@ class UK2Node_ComponentBoundEvent : public UK2Node_Event
 
 	// Begin UEdGraphNode interface
 	virtual FText GetNodeTitle(ENodeTitleType::Type TitleType) const override;
-	virtual FString GetNodeNativeTitle(ENodeTitleType::Type TitleType) const override;
 	virtual FString GetTooltip() const override;
 	virtual FString GetDocumentationLink() const override;
 	virtual FString GetDocumentationExcerptName() const override;
diff --git a/Engine/Source/Editor/BlueprintGraph/Classes/K2Node_Composite.h b/Engine/Source/Editor/BlueprintGraph/Classes/K2Node_Composite.h
--- a/Engine/Source/Editor/BlueprintGraph/Classes/K2Node_Composite.h
+++ b/Engine/Source/Editor/BlueprintGraph/Classes/K2Node_Composite.h
@@ -19,7 +19,6 @@ class UK2Node_Composite : public UK2Node_Tunnel
 	BLUEPRINTGRAPH_API virtual FString GetTooltip() const override;
 	BLUEPRINTGRAPH_API virtual FLinearColor GetNodeTitleColor() const override;
 	BLUEPRINTGRAPH_API virtual FText GetNodeTitle(ENodeTitleType::Type TitleType) const override;
-	BLUEPRINTGRAPH_API virtual FString GetNodeNativeTitle(ENodeTitleType::Type TitleType) const override;
 	BLUEPRINTGRAPH_API virtual bool CanUserDeleteNode() const override;
 	BLUEPRINTGRAPH_API virtual UObject* GetJumpTargetForDoubleClick() const override;
 	BLUEPRINTGRAPH_API virtual void PostPlacedNewNode() override;
diff --git a/Engine/Source/Editor/BlueprintGraph/Classes/K2Node_CustomEvent.h b/Engine/Source/Editor/BlueprintGraph/Classes/K2Node_CustomEvent.h
--- a/Engine/Source/Editor/BlueprintGraph/Classes/K2Node_CustomEvent.h
+++ b/Engine/Source/Editor/BlueprintGraph/Classes/K2Node_CustomEvent.h
@@ -15,7 +15,6 @@ class UK2Node_CustomEvent : public UK2Node_Event
 	// Begin UEdGraphNode interface
 	virtual void ReconstructNode() override;
 	virtual FText GetNodeTitle(ENodeTitleType::Type TitleType) const override;
-	virtual FString GetNodeNativeTitle(ENodeTitleType::Type TitleType) const override;
 	virtual void OnRenameNode(const FString& NewName) override;
 	virtual TSharedPtr<class INameValidatorInterface> MakeNameValidator() const override;
 	virtual FString GetTooltip() const override;
diff --git a/Engine/Source/Editor/BlueprintGraph/Classes/K2Node_DelegateSet.h b/Engine/Source/Editor/BlueprintGraph/Classes/K2Node_DelegateSet.h
--- a/Engine/Source/Editor/BlueprintGraph/Classes/K2Node_DelegateSet.h
+++ b/Engine/Source/Editor/BlueprintGraph/Classes/K2Node_DelegateSet.h
@@ -22,7 +22,6 @@ class UK2Node_DelegateSet : public UK2Node
 	virtual void AllocateDefaultPins() override;
 	virtual FString GetTooltip() const override;
 	virtual FText GetNodeTitle(ENodeTitleType::Type TitleType) const override;
-	virtual FString GetNodeNativeTitle(ENodeTitleType::Type TitleType) const override;
 	virtual FLinearColor GetNodeTitleColor() const { return FColor(216,88,88); }
 	// End UEdGraphNode interface
 
diff --git a/Engine/Source/Editor/BlueprintGraph/Classes/K2Node_DynamicCast.h b/Engine/Source/Editor/BlueprintGraph/Classes/K2Node_DynamicCast.h
--- a/Engine/Source/Editor/BlueprintGraph/Classes/K2Node_DynamicCast.h
+++ b/Engine/Source/Editor/BlueprintGraph/Classes/K2Node_DynamicCast.h
@@ -18,7 +18,6 @@ class UK2Node_DynamicCast : public UK2Node
 	virtual void AllocateDefaultPins() override;
 	virtual FLinearColor GetNodeTitleColor() const override;
 	virtual FText GetNodeTitle(ENodeTitleType::Type TitleType) const override;
-	virtual FString GetNodeNativeTitle(ENodeTitleType::Type TitleType) const override;
 	virtual FName GetPaletteIcon(FLinearColor& OutColor) const override{ return TEXT("GraphEditor.Cast_16x"); }
 	// End UEdGraphNode interface
 
diff --git a/Engine/Source/Editor/BlueprintGraph/Classes/K2Node_Event.h b/Engine/Source/Editor/BlueprintGraph/Classes/K2Node_Event.h
--- a/Engine/Source/Editor/BlueprintGraph/Classes/K2Node_Event.h
+++ b/Engine/Source/Editor/BlueprintGraph/Classes/K2Node_Event.h
@@ -40,7 +40,6 @@ class UK2Node_Event : public UK2Node_EditablePinBase
 	BLUEPRINTGRAPH_API virtual FString GetTooltip() const override;
 	BLUEPRINTGRAPH_API virtual FString GetKeywords() const override;	
 	BLUEPRINTGRAPH_API virtual FText GetNodeTitle(ENodeTitleType::Type TitleType) const override;
-	BLUEPRINTGRAPH_API virtual FString GetNodeNativeTitle(ENodeTitleType::Type TitleType) const override;
 	BLUEPRINTGRAPH_API virtual FLinearColor GetNodeTitleColor() const override;
 	BLUEPRINTGRAPH_API virtual bool CanPasteHere(const UEdGraph* TargetGraph, const UEdGraphSchema* Schema) const override;
 	BLUEPRINTGRAPH_API virtual FName GetCornerIcon() const override;
diff --git a/Engine/Source/Editor/BlueprintGraph/Classes/K2Node_ForEachElementInEnum.h b/Engine/Source/Editor/BlueprintGraph/Classes/K2Node_ForEachElementInEnum.h
--- a/Engine/Source/Editor/BlueprintGraph/Classes/K2Node_ForEachElementInEnum.h
+++ b/Engine/Source/Editor/BlueprintGraph/Classes/K2Node_ForEachElementInEnum.h
@@ -19,7 +19,6 @@ class UK2Node_ForEachElementInEnum : public UK2Node
 	virtual void AllocateDefaultPins() override;
 	virtual FString GetTooltip() const override;
 	virtual FText GetNodeTitle(ENodeTitleType::Type TitleType) const override;
-	virtual FString GetNodeNativeTitle(ENodeTitleType::Type TitleType) const override;
 	virtual FName GetPaletteIcon(FLinearColor& OutColor) const override{ return TEXT("GraphEditor.Macro.Loop_16x"); }
 	// End UEdGraphNode interface
 
diff --git a/Engine/Source/Editor/BlueprintGraph/Classes/K2Node_FunctionEntry.h b/Engine/Source/Editor/BlueprintGraph/Classes/K2Node_FunctionEntry.h
--- a/Engine/Source/Editor/BlueprintGraph/Classes/K2Node_FunctionEntry.h
+++ b/Engine/Source/Editor/BlueprintGraph/Classes/K2Node_FunctionEntry.h
@@ -28,7 +28,6 @@ class UK2Node_FunctionEntry : public UK2Node_FunctionTerminator
 	// Begin UEdGraphNode interface
 	virtual void AllocateDefaultPins() override;
 	virtual FText GetNodeTitle(ENodeTitleType::Type TitleType) const override;
-	virtual FString GetNodeNativeTitle(ENodeTitleType::Type TitleType) const override;
 	virtual bool CanUserDeleteNode() const override { return false; }
 	virtual bool IsDeprecated() const override;
 	virtual FString GetDeprecationMessage() const override;
diff --git a/Engine/Source/Editor/BlueprintGraph/Classes/K2Node_GetDataTableRow.h b/Engine/Source/Editor/BlueprintGraph/Classes/K2Node_GetDataTableRow.h
--- a/Engine/Source/Editor/BlueprintGraph/Classes/K2Node_GetDataTableRow.h
+++ b/Engine/Source/Editor/BlueprintGraph/Classes/K2Node_GetDataTableRow.h
@@ -16,7 +16,6 @@ class UK2Node_GetDataTableRow : public UK2Node
 	// Begin UEdGraphNode interface.
 	virtual void AllocateDefaultPins() override;
 	virtual FText GetNodeTitle(ENodeTitleType::Type TitleType) const override;
-	virtual FString GetNodeNativeTitle(ENodeTitleType::Type TitleType) const override;
 	virtual void PinDefaultValueChanged(UEdGraphPin* Pin) override;
 	virtual FString GetTooltip() const override;
 	virtual void ExpandNode(class FKismetCompilerContext& CompilerContext, UEdGraph* SourceGraph) override;
diff --git a/Engine/Source/Editor/BlueprintGraph/Classes/K2Node_GetInputAxisValue.h b/Engine/Source/Editor/BlueprintGraph/Classes/K2Node_GetInputAxisValue.h
--- a/Engine/Source/Editor/BlueprintGraph/Classes/K2Node_GetInputAxisValue.h
+++ b/Engine/Source/Editor/BlueprintGraph/Classes/K2Node_GetInputAxisValue.h
@@ -23,7 +23,6 @@ class UK2Node_GetInputAxisValue : public UK2Node_CallFunction
 	// Begin EdGraphNode interface
 	virtual void AllocateDefaultPins() override;
 	virtual FText GetNodeTitle(ENodeTitleType::Type TitleType) const override;
-	virtual FString GetNodeNativeTitle(ENodeTitleType::Type TitleType) const override;
 	virtual FString GetTooltip() const override;
 	// End EdGraphNode interface
 
diff --git a/Engine/Source/Editor/BlueprintGraph/Classes/K2Node_GetNumEnumEntries.h b/Engine/Source/Editor/BlueprintGraph/Classes/K2Node_GetNumEnumEntries.h
--- a/Engine/Source/Editor/BlueprintGraph/Classes/K2Node_GetNumEnumEntries.h
+++ b/Engine/Source/Editor/BlueprintGraph/Classes/K2Node_GetNumEnumEntries.h
@@ -16,7 +16,6 @@ class UK2Node_GetNumEnumEntries : public UK2Node, public INodeDependingOnEnumInt
 	virtual void AllocateDefaultPins() override;
 	virtual FString GetTooltip() const override;
 	virtual FText GetNodeTitle(ENodeTitleType::Type TitleType) const override;
-	virtual FString GetNodeNativeTitle(ENodeTitleType::Type TitleType) const override;
 	virtual FName GetPaletteIcon(FLinearColor& OutColor) const override{ return TEXT("GraphEditor.Enum_16x"); }
 	// End UEdGraphNode interface
 
diff --git a/Engine/Source/Editor/BlueprintGraph/Classes/K2Node_InputAction.h b/Engine/Source/Editor/BlueprintGraph/Classes/K2Node_InputAction.h
--- a/Engine/Source/Editor/BlueprintGraph/Classes/K2Node_InputAction.h
+++ b/Engine/Source/Editor/BlueprintGraph/Classes/K2Node_InputAction.h
@@ -33,7 +33,6 @@ class UK2Node_InputAction : public UK2Node
 	virtual void AllocateDefaultPins() override;
 	virtual FLinearColor GetNodeTitleColor() const override;
 	virtual FText GetNodeTitle(ENodeTitleType::Type TitleType) const override;
-	virtual FString GetNodeNativeTitle(ENodeTitleType::Type TitleType) const override;
 	virtual FString GetTooltip() const override;
 	virtual FName GetPaletteIcon(FLinearColor& OutColor) const override{ return TEXT("GraphEditor.Event_16x"); }
 	// End UEdGraphNode interface.
diff --git a/Engine/Source/Editor/BlueprintGraph/Classes/K2Node_InputAxisEvent.h b/Engine/Source/Editor/BlueprintGraph/Classes/K2Node_InputAxisEvent.h
--- a/Engine/Source/Editor/BlueprintGraph/Classes/K2Node_InputAxisEvent.h
+++ b/Engine/Source/Editor/BlueprintGraph/Classes/K2Node_InputAxisEvent.h
@@ -30,7 +30,6 @@ class UK2Node_InputAxisEvent : public UK2Node_Event
 
 	// Begin EdGraphNode interface
 	virtual FText GetNodeTitle(ENodeTitleType::Type TitleType) const override;
-	virtual FString GetNodeNativeTitle(ENodeTitleType::Type TitleType) const override;
 	virtual FString GetTooltip() const override;
 	virtual bool CanPasteHere(const UEdGraph* TargetGraph, const UEdGraphSchema* Schema) const override;
 	// End EdGraphNode interface
diff --git a/Engine/Source/Editor/BlueprintGraph/Classes/K2Node_InputAxisKeyEvent.h b/Engine/Source/Editor/BlueprintGraph/Classes/K2Node_InputAxisKeyEvent.h
--- a/Engine/Source/Editor/BlueprintGraph/Classes/K2Node_InputAxisKeyEvent.h
+++ b/Engine/Source/Editor/BlueprintGraph/Classes/K2Node_InputAxisKeyEvent.h
@@ -26,7 +26,6 @@ class UK2Node_InputAxisKeyEvent : public UK2Node_Event
 
 	// Begin EdGraphNode interface
 	virtual FText GetNodeTitle(ENodeTitleType::Type TitleType) const override;
-	virtual FString GetNodeNativeTitle(ENodeTitleType::Type TitleType) const override;
 	virtual FString GetTooltip() const override;
 	virtual FName GetPaletteIcon(FLinearColor& OutColor) const override;
 	virtual bool CanPasteHere(const UEdGraph* TargetGraph, const UEdGraphSchema* Schema) const override;
diff --git a/Engine/Source/Editor/BlueprintGraph/Classes/K2Node_InputKey.h b/Engine/Source/Editor/BlueprintGraph/Classes/K2Node_InputKey.h
--- a/Engine/Source/Editor/BlueprintGraph/Classes/K2Node_InputKey.h
+++ b/Engine/Source/Editor/BlueprintGraph/Classes/K2Node_InputKey.h
@@ -52,7 +52,6 @@ class UK2Node_InputKey : public UK2Node
 	virtual void AllocateDefaultPins() override;
 	virtual FLinearColor GetNodeTitleColor() const override;
 	virtual FText GetNodeTitle(ENodeTitleType::Type TitleType) const override;
-	virtual FString GetNodeNativeTitle(ENodeTitleType::Type TitleType) const override;
 	virtual FString GetTooltip() const override;
 	virtual FName GetPaletteIcon(FLinearColor& OutColor) const override;
 	// End UEdGraphNode interface.
diff --git a/Engine/Source/Editor/BlueprintGraph/Classes/K2Node_LocalVariable.h b/Engine/Source/Editor/BlueprintGraph/Classes/K2Node_LocalVariable.h
--- a/Engine/Source/Editor/BlueprintGraph/Classes/K2Node_LocalVariable.h
+++ b/Engine/Source/Editor/BlueprintGraph/Classes/K2Node_LocalVariable.h
@@ -20,7 +20,6 @@ class UDEPRECATED_K2Node_LocalVariable : public UK2Node_TemporaryVariable
 	// Begin UEdGraphNode interface.
  	virtual FString GetTooltip() const override;
 	virtual FText GetNodeTitle(ENodeTitleType::Type TitleType) const override;
-	virtual FString GetNodeNativeTitle(ENodeTitleType::Type TitleType) const override;
 	virtual void OnRenameNode(const FString& NewName) override;
 	virtual TSharedPtr<class INameValidatorInterface> MakeNameValidator() const override;
 	virtual void PostPlacedNewNode() override;
diff --git a/Engine/Source/Editor/BlueprintGraph/Classes/K2Node_MacroInstance.h b/Engine/Source/Editor/BlueprintGraph/Classes/K2Node_MacroInstance.h
--- a/Engine/Source/Editor/BlueprintGraph/Classes/K2Node_MacroInstance.h
+++ b/Engine/Source/Editor/BlueprintGraph/Classes/K2Node_MacroInstance.h
@@ -38,7 +38,6 @@ public:
 	virtual void PostPasteNode() override;
 	virtual FLinearColor GetNodeTitleColor() const override;
 	virtual FText GetNodeTitle(ENodeTitleType::Type TitleType) const override;
-	virtual FString GetNodeNativeTitle(ENodeTitleType::Type TitleType) const override;
 	virtual bool CanUserDeleteNode() const override { return true; }
 	void GetContextMenuActions(const FGraphNodeContextMenuBuilder& Context) const override;
 	virtual void ReallocatePinsDuringReconstruction(TArray<UEdGraphPin*>& OldPins) override;
diff --git a/Engine/Source/Editor/BlueprintGraph/Classes/K2Node_MakeStruct.h b/Engine/Source/Editor/BlueprintGraph/Classes/K2Node_MakeStruct.h
--- a/Engine/Source/Editor/BlueprintGraph/Classes/K2Node_MakeStruct.h
+++ b/Engine/Source/Editor/BlueprintGraph/Classes/K2Node_MakeStruct.h
@@ -16,7 +16,6 @@ class UK2Node_MakeStruct : public UK2Node_StructMemberSet
 	// Begin UEdGraphNode interface
 	virtual void AllocateDefaultPins() override;
 	virtual FText GetNodeTitle(ENodeTitleType::Type TitleType) const override;
-	virtual FString GetNodeNativeTitle(ENodeTitleType::Type TitleType) const override;
 	virtual FLinearColor GetNodeTitleColor() const override;
 	virtual FString GetTooltip() const override;
 	virtual void ValidateNodeDuringCompilation(class FCompilerResultsLog& MessageLog) const override;
diff --git a/Engine/Source/Editor/BlueprintGraph/Classes/K2Node_MathExpression.h b/Engine/Source/Editor/BlueprintGraph/Classes/K2Node_MathExpression.h
--- a/Engine/Source/Editor/BlueprintGraph/Classes/K2Node_MathExpression.h
+++ b/Engine/Source/Editor/BlueprintGraph/Classes/K2Node_MathExpression.h
@@ -31,7 +31,6 @@ public:
 	virtual TSharedPtr<class INameValidatorInterface> MakeNameValidator() const override;
 	virtual void OnRenameNode(const FString& NewName) override;
 	virtual FText GetNodeTitle(ENodeTitleType::Type TitleType) const override;
-	virtual FString GetNodeNativeTitle(ENodeTitleType::Type TitleType) const override;
 	virtual void PostPlacedNewNode() override;
 	virtual void ReconstructNode() override;
 	// End of UEdGraphNode interface
diff --git a/Engine/Source/Editor/BlueprintGraph/Classes/K2Node_Message.h b/Engine/Source/Editor/BlueprintGraph/Classes/K2Node_Message.h
--- a/Engine/Source/Editor/BlueprintGraph/Classes/K2Node_Message.h
+++ b/Engine/Source/Editor/BlueprintGraph/Classes/K2Node_Message.h
@@ -13,7 +13,6 @@ class UK2Node_Message : public UK2Node_CallFunction
 
 	// Begin UEdGraphNode interface.
 	virtual FText GetNodeTitle(ENodeTitleType::Type TitleType) const override;
-	virtual FString GetNodeNativeTitle(ENodeTitleType::Type TitleType) const override;
 	// End UEdGraphNode interface.
 
 	// Begin K2Node interface.
diff --git a/Engine/Source/Editor/BlueprintGraph/Classes/K2Node_RemoveDelegate.h b/Engine/Source/Editor/BlueprintGraph/Classes/K2Node_RemoveDelegate.h
--- a/Engine/Source/Editor/BlueprintGraph/Classes/K2Node_RemoveDelegate.h
+++ b/Engine/Source/Editor/BlueprintGraph/Classes/K2Node_RemoveDelegate.h
@@ -12,7 +12,6 @@ public:
 	// UEdGraphNode interface
 	virtual void AllocateDefaultPins() override;
 	virtual FText GetNodeTitle(ENodeTitleType::Type TitleType) const override;
-	virtual FString GetNodeNativeTitle(ENodeTitleType::Type TitleType) const override;
 	// End of UEdGraphNode interface
 
 	// Begin of UK2Node interface
diff --git a/Engine/Source/Editor/BlueprintGraph/Classes/K2Node_SelectEnum.h b/Engine/Source/Editor/BlueprintGraph/Classes/K2Node_SelectEnum.h
--- a/Engine/Source/Editor/BlueprintGraph/Classes/K2Node_SelectEnum.h
+++ b/Engine/Source/Editor/BlueprintGraph/Classes/K2Node_SelectEnum.h
@@ -14,7 +14,6 @@ class UDEPRECATED_K2Node_SelectEnum : public UK2Node_Select
 	virtual void AllocateDefaultPins() override;
 	virtual FString GetTooltip() const override;
 	virtual FText GetNodeTitle(ENodeTitleType::Type TitleType) const override;
-	virtual FString GetNodeNativeTitle(ENodeTitleType::Type TitleType) const override;
 	// End UEdGraphNode interface
 
 	// Begin UK2Node interface
diff --git a/Engine/Source/Editor/BlueprintGraph/Classes/K2Node_SpawnActor.h b/Engine/Source/Editor/BlueprintGraph/Classes/K2Node_SpawnActor.h
--- a/Engine/Source/Editor/BlueprintGraph/Classes/K2Node_SpawnActor.h
+++ b/Engine/Source/Editor/BlueprintGraph/Classes/K2Node_SpawnActor.h
@@ -14,7 +14,6 @@ class UK2Node_SpawnActor : public UK2Node
 	virtual void AllocateDefaultPins() override;
 	virtual FLinearColor GetNodeTitleColor() const override;
 	virtual FText GetNodeTitle(ENodeTitleType::Type TitleType) const override;
-	virtual FString GetNodeNativeTitle(ENodeTitleType::Type TitleType) const override;
 	virtual void PinDefaultValueChanged(UEdGraphPin* Pin) override;
 	virtual FString GetTooltip() const override;
 	virtual bool IsDeprecated() const override;
diff --git a/Engine/Source/Editor/BlueprintGraph/Classes/K2Node_SpawnActorFromClass.h b/Engine/Source/Editor/BlueprintGraph/Classes/K2Node_SpawnActorFromClass.h
--- a/Engine/Source/Editor/BlueprintGraph/Classes/K2Node_SpawnActorFromClass.h
+++ b/Engine/Source/Editor/BlueprintGraph/Classes/K2Node_SpawnActorFromClass.h
@@ -14,7 +14,6 @@ class UK2Node_SpawnActorFromClass : public UK2Node
 	virtual void AllocateDefaultPins() override;
 	virtual FLinearColor GetNodeTitleColor() const override;
 	virtual FText GetNodeTitle(ENodeTitleType::Type TitleType) const override;
-	virtual FString GetNodeNativeTitle(ENodeTitleType::Type TitleType) const override;
 	virtual void PinDefaultValueChanged(UEdGraphPin* Pin) override;
 	virtual FString GetTooltip() const override;
 	virtual void ExpandNode(class FKismetCompilerContext& CompilerContext, UEdGraph* SourceGraph) override;
diff --git a/Engine/Source/Editor/BlueprintGraph/Classes/K2Node_StructMemberGet.h b/Engine/Source/Editor/BlueprintGraph/Classes/K2Node_StructMemberGet.h
--- a/Engine/Source/Editor/BlueprintGraph/Classes/K2Node_StructMemberGet.h
+++ b/Engine/Source/Editor/BlueprintGraph/Classes/K2Node_StructMemberGet.h
@@ -21,7 +21,6 @@ class UK2Node_StructMemberGet : public UK2Node_StructOperation
 	virtual void AllocateDefaultPins() override;
 	virtual FString GetTooltip() const override;
 	virtual FText GetNodeTitle(ENodeTitleType::Type TitleType) const override;
-	virtual FString GetNodeNativeTitle(ENodeTitleType::Type TitleType) const override;
 	// End of UEdGraphNode interface
 
 	// UK2Node interface
diff --git a/Engine/Source/Editor/BlueprintGraph/Classes/K2Node_StructMemberSet.h b/Engine/Source/Editor/BlueprintGraph/Classes/K2Node_StructMemberSet.h
--- a/Engine/Source/Editor/BlueprintGraph/Classes/K2Node_StructMemberSet.h
+++ b/Engine/Source/Editor/BlueprintGraph/Classes/K2Node_StructMemberSet.h
@@ -21,7 +21,6 @@ class UK2Node_StructMemberSet : public UK2Node_StructOperation
 	virtual void AllocateDefaultPins() override;
 	virtual FString GetTooltip() const override;
 	virtual FText GetNodeTitle(ENodeTitleType::Type TitleType) const override;
-	virtual FString GetNodeNativeTitle(ENodeTitleType::Type TitleType) const override;
 	// End of UEdGraphNode interface
 
 	// UK2Node interface
diff --git a/Engine/Source/Editor/BlueprintGraph/Classes/K2Node_SwitchEnum.h b/Engine/Source/Editor/BlueprintGraph/Classes/K2Node_SwitchEnum.h
--- a/Engine/Source/Editor/BlueprintGraph/Classes/K2Node_SwitchEnum.h
+++ b/Engine/Source/Editor/BlueprintGraph/Classes/K2Node_SwitchEnum.h
@@ -31,7 +31,6 @@ class UK2Node_SwitchEnum : public UK2Node_Switch, public INodeDependingOnEnumInt
 	// UEdGraphNode interface
 	virtual FString GetTooltip() const override;
 	virtual FText GetNodeTitle(ENodeTitleType::Type TitleType) const override;
-	virtual FString GetNodeNativeTitle(ENodeTitleType::Type TitleType) const override;
 	virtual bool ShouldShowNodeProperties() const override { return true; }	
 	// End of UEdGraphNode interface
 
diff --git a/Engine/Source/Editor/BlueprintGraph/Classes/K2Node_TemporaryVariable.h b/Engine/Source/Editor/BlueprintGraph/Classes/K2Node_TemporaryVariable.h
--- a/Engine/Source/Editor/BlueprintGraph/Classes/K2Node_TemporaryVariable.h
+++ b/Engine/Source/Editor/BlueprintGraph/Classes/K2Node_TemporaryVariable.h
@@ -23,7 +23,6 @@ class UK2Node_TemporaryVariable : public UK2Node
 	virtual void AllocateDefaultPins() override;
 	virtual FString GetTooltip() const override;
 	virtual FText GetNodeTitle(ENodeTitleType::Type TitleType) const override;
-	virtual FString GetNodeNativeTitle(ENodeTitleType::Type TitleType) const override;
 	virtual FString GetDescriptiveCompiledName() const override;
 	// End UEdGraphNode interface.
 
diff --git a/Engine/Source/Editor/BlueprintGraph/Classes/K2Node_VariableGet.h b/Engine/Source/Editor/BlueprintGraph/Classes/K2Node_VariableGet.h
--- a/Engine/Source/Editor/BlueprintGraph/Classes/K2Node_VariableGet.h
+++ b/Engine/Source/Editor/BlueprintGraph/Classes/K2Node_VariableGet.h
@@ -14,7 +14,6 @@ class UK2Node_VariableGet : public UK2Node_Variable
 	virtual void AllocateDefaultPins() override;
 	virtual FString GetTooltip() const override;
 	virtual FText GetNodeTitle(ENodeTitleType::Type TitleType) const override;
-	virtual FString GetNodeNativeTitle(ENodeTitleType::Type TitleType) const override;
 	// End UEdGraphNode interface
 
 	// Begin K2Node interface
diff --git a/Engine/Source/Editor/BlueprintGraph/Classes/K2Node_VariableSet.h b/Engine/Source/Editor/BlueprintGraph/Classes/K2Node_VariableSet.h
--- a/Engine/Source/Editor/BlueprintGraph/Classes/K2Node_VariableSet.h
+++ b/Engine/Source/Editor/BlueprintGraph/Classes/K2Node_VariableSet.h
@@ -15,7 +15,6 @@ class UK2Node_VariableSet : public UK2Node_Variable
 	virtual void ReallocatePinsDuringReconstruction(TArray<UEdGraphPin*>& OldPins) override;
 	virtual FString GetTooltip() const override;
 	virtual FText GetNodeTitle(ENodeTitleType::Type TitleType) const override;
-	virtual FString GetNodeNativeTitle(ENodeTitleType::Type TitleType) const override;
 	// End UEdGraphNode interface
 
 	// Begin K2Node interface
diff --git a/Engine/Source/Editor/BlueprintGraph/Classes/K2Node_VariableSetRef.h b/Engine/Source/Editor/BlueprintGraph/Classes/K2Node_VariableSetRef.h
--- a/Engine/Source/Editor/BlueprintGraph/Classes/K2Node_VariableSetRef.h
+++ b/Engine/Source/Editor/BlueprintGraph/Classes/K2Node_VariableSetRef.h
@@ -14,7 +14,6 @@ class UK2Node_VariableSetRef : public UK2Node
 	virtual void ReallocatePinsDuringReconstruction(TArray<UEdGraphPin*>& OldPins) override;
 	virtual FString GetTooltip() const override;
 	virtual FText GetNodeTitle(ENodeTitleType::Type TitleType) const override;
-	virtual FString GetNodeNativeTitle(ENodeTitleType::Type TitleType) const override;
 	// End UEdGraphNode interface
 
 	// Begin UK2Node interface
diff --git a/Engine/Source/Editor/BlueprintGraph/Private/K2ActionMenuBuilder.cpp b/Engine/Source/Editor/BlueprintGraph/Private/K2ActionMenuBuilder.cpp
--- a/Engine/Source/Editor/BlueprintGraph/Private/K2ActionMenuBuilder.cpp
+++ b/Engine/Source/Editor/BlueprintGraph/Private/K2ActionMenuBuilder.cpp
@@ -117,7 +117,6 @@ static void GetRequestedAction(FGraphActionListBuilderBase& ActionMenuBuilder, c
 	UK2Node* NodeTemplate = NewObject<ActionType>(ActionMenuBuilder.OwnerOfTemporaries);
 	TSharedPtr<FEdGraphSchemaAction_K2NewNode> Action = FK2ActionMenuBuilder::AddNewNodeAction(ActionMenuBuilder, Category, NodeTemplate->GetNodeTitle(ENodeTitleType::ListView), NodeTemplate->GetTooltip(), 0, NodeTemplate->GetKeywords());
 	Action->NodeTemplate = NodeTemplate;
-	Action->SearchTitle = Action->NodeTemplate->GetNodeSearchTitle();
 }
 
 /** Utility for creating spawn from archetype nodes */
@@ -159,7 +158,6 @@ static void GetInputNodes(FGraphActionListBuilderBase& ActionMenuBuilder, const
 			UK2Node_InputAction* InputActionNode = ActionMenuBuilder.CreateTemplateNode<UK2Node_InputAction>();
 			InputActionNode->InputActionName = InputActionName;
 			Action->NodeTemplate = InputActionNode;
-			Action->SearchTitle = Action->NodeTemplate->GetNodeSearchTitle();
 			Action->Keywords = TEXT("InputAction");
 		}
 
@@ -171,14 +169,12 @@ static void GetInputNodes(FGraphActionListBuilderBase& ActionMenuBuilder, const
 			UK2Node_GetInputAxisValue* GetInputAxisValue = ActionMenuBuilder.CreateTemplateNode<UK2Node_GetInputAxisValue>();
 			GetInputAxisValue->Initialize(InputAxisName);
 			GetAction->NodeTemplate = GetInputAxisValue;
-			GetAction->SearchTitle = GetAction->NodeTemplate->GetNodeSearchTitle();
 			GetAction->Keywords = TEXT("InputAxis");
 
 			TSharedPtr<FEdGraphSchemaAction_K2NewNode> Action = FK2ActionMenuBuilder::AddNewNodeAction(ActionMenuBuilder, K2ActionCategories::AxisEventCategory, FText::FromName(InputAxisName), InputAxisName.ToString());
 			UK2Node_InputAxisEvent* InputAxisEvent = ActionMenuBuilder.CreateTemplateNode<UK2Node_InputAxisEvent>();
 			InputAxisEvent->Initialize(InputAxisName);
 			Action->NodeTemplate = InputAxisEvent;
-			Action->SearchTitle = Action->NodeTemplate->GetNodeSearchTitle();
 			Action->Keywords = TEXT("InputAxis");
 		}
 
@@ -216,7 +212,6 @@ static void GetInputNodes(FGraphActionListBuilderBase& ActionMenuBuilder, const
 					UK2Node_InputAxisKeyEvent* InputKeyAxisNode = ActionMenuBuilder.CreateTemplateNode<UK2Node_InputAxisKeyEvent>();
 					InputKeyAxisNode->Initialize(Key);
 					Action->NodeTemplate = InputKeyAxisNode;
-					Action->SearchTitle = Action->NodeTemplate->GetNodeSearchTitle();
 
 					const FString& KeyValuesCategory = (bGamepad ? K2ActionCategories::GamepadValuesCategory : (bMouse ? K2ActionCategories::MouseValuesCategory : K2ActionCategories::KeyValuesCategory));
 					GetAction = FK2ActionMenuBuilder::AddNewNodeAction(ActionMenuBuilder, KeyValuesCategory, KeyName, KeyName.ToString());
@@ -230,7 +225,6 @@ static void GetInputNodes(FGraphActionListBuilderBase& ActionMenuBuilder, const
 					UK2Node_InputVectorAxisEvent* InputVectorAxisNode = ActionMenuBuilder.CreateTemplateNode<UK2Node_InputVectorAxisEvent>();
 					InputVectorAxisNode->Initialize(Key);
 					Action->NodeTemplate = InputVectorAxisNode;
-					Action->SearchTitle = Action->NodeTemplate->GetNodeSearchTitle();
 
 					const FString& KeyValuesCategory = (bGamepad ? K2ActionCategories::GamepadValuesCategory : (bMouse ? K2ActionCategories::MouseValuesCategory : K2ActionCategories::KeyValuesCategory));
 					GetAction = FK2ActionMenuBuilder::AddNewNodeAction(ActionMenuBuilder, KeyValuesCategory, KeyName, KeyName.ToString());
@@ -244,7 +238,6 @@ static void GetInputNodes(FGraphActionListBuilderBase& ActionMenuBuilder, const
 					UK2Node_InputKey* InputKeyNode = ActionMenuBuilder.CreateTemplateNode<UK2Node_InputKey>();
 					InputKeyNode->InputKey = Key;
 					Action->NodeTemplate = InputKeyNode;
-					Action->SearchTitle = Action->NodeTemplate->GetNodeSearchTitle();
 				}
 			}
 		}
@@ -252,7 +245,6 @@ static void GetInputNodes(FGraphActionListBuilderBase& ActionMenuBuilder, const
 		TSharedPtr<FEdGraphSchemaAction_K2NewNode> Action = FK2ActionMenuBuilder::AddNewNodeAction(ActionMenuBuilder, K2ActionCategories::TouchEventsCategory, LOCTEXT("Touch", "Touch"), LOCTEXT("Touch", "Touch").ToString());
 		UK2Node_InputTouch* InputTouchNode = ActionMenuBuilder.CreateTemplateNode<UK2Node_InputTouch>();
 		Action->NodeTemplate = InputTouchNode;
-		Action->SearchTitle = Action->NodeTemplate->GetNodeSearchTitle();
 		Action->Keywords = TEXT("Input Touch");
 	}
 }
@@ -322,7 +314,6 @@ static void GetEnumUtilitiesNodes(FBlueprintGraphActionListBuilder& ContextMenuB
 					TSharedPtr<FEdGraphSchemaAction_K2NewNode> Action = FK2ActionMenuBuilder::AddNewNodeAction(
 						InContextMenuBuilder, Category, NodeTemplate->GetNodeTitle(ENodeTitleType::ListView),  NodeTemplate->GetTooltip(), 0, NodeTemplate->GetKeywords());
 					Action->NodeTemplate = NodeTemplate;
-					Action->SearchTitle = NodeTemplate->GetNodeSearchTitle();
 				}
 
 				if (bInForEach)
@@ -333,7 +324,6 @@ static void GetEnumUtilitiesNodes(FBlueprintGraphActionListBuilder& ContextMenuB
 					TSharedPtr<FEdGraphSchemaAction_K2NewNode> Action = FK2ActionMenuBuilder::AddNewNodeAction(
 						InContextMenuBuilder, Category, NodeTemplate->GetNodeTitle(ENodeTitleType::ListView),  NodeTemplate->GetTooltip(), 0, NodeTemplate->GetKeywords());
 					Action->NodeTemplate = NodeTemplate;
-					Action->SearchTitle = NodeTemplate->GetNodeSearchTitle();
 				}
 
 				if (bInCastFromByte)
@@ -345,7 +335,6 @@ static void GetEnumUtilitiesNodes(FBlueprintGraphActionListBuilder& ContextMenuB
 					TSharedPtr<FEdGraphSchemaAction_K2NewNode> Action = FK2ActionMenuBuilder::AddNewNodeAction(
 						InContextMenuBuilder, Category, NodeTemplate->GetNodeTitle(ENodeTitleType::ListView),  NodeTemplate->GetTooltip(), 0, NodeTemplate->GetKeywords());
 					Action->NodeTemplate = NodeTemplate;
-					Action->SearchTitle = NodeTemplate->GetNodeSearchTitle();
 				}
 
 				if (bInLiteralByte)
@@ -392,7 +381,6 @@ static void GetInterfaceMessages(UClass* CurrentInterface, const FString& Messag
 				MessageTemplate->FunctionReference.SetExternalMember(Function->GetFName(), CurrentInterface);
 				TSharedPtr<FEdGraphSchemaAction_K2NewNode> Action = FK2ActionMenuBuilder::AddNewNodeAction(ContextMenuBuilder, MessageSubCategory, FText::FromString(Function->GetName()), MessageTemplate->GetTooltip(), 0, MessageTemplate->GetKeywords());
 				Action->NodeTemplate = MessageTemplate;
-				Action->SearchTitle = MessageTemplate->GetNodeSearchTitle();
 			}
 		}
 	}
@@ -407,7 +395,6 @@ static void GetReferenceSetterItems(FBlueprintGraphActionListBuilder& ContextMen
 	TSharedPtr<FEdGraphSchemaAction_K2NewNode> Action = FK2ActionMenuBuilder::AddNewNodeAction(ContextMenuBuilder, K2ActionCategories::VariablesCategory, SetRefTemplate->GetNodeTitle(ENodeTitleType::ListView), SetRefTemplate->GetTooltip(), 0, SetRefTemplate->GetKeywords());
 
 	Action->NodeTemplate = SetRefTemplate;
-	Action->SearchTitle = SetRefTemplate->GetNodeSearchTitle();
 }
 
 /** 
@@ -538,7 +525,6 @@ TSharedPtr<FEdGraphSchemaAction_K2AddComponent> FK2ActionMenuBuilder::CreateAddC
 	TSharedPtr<FEdGraphSchemaAction_K2AddComponent> Action = TSharedPtr<FEdGraphSchemaAction_K2AddComponent>(new FEdGraphSchemaAction_K2AddComponent(MenuCategory, FText::FromString(Description), ToolTip, Asset != NULL ? 2 : 0));
 
 	Action->NodeTemplate = CallFuncNode;
-	Action->SearchTitle = Action->NodeTemplate->GetNodeSearchTitle();
 	Action->ComponentClass = DestinationComponentType;
 	Action->ComponentAsset = Asset;
 
@@ -575,7 +561,6 @@ void FK2ActionMenuBuilder::AddEventForFunction(FGraphActionListBuilderBase& Acti
 		TSharedPtr<FEdGraphSchemaAction_K2AddEvent> Action(new FEdGraphSchemaAction_K2AddEvent(Category, FText::FromString(EventName), EventNodeTemplate->GetTooltip(), 0));
 		Action->Keywords     = EventNodeTemplate->GetKeywords();
 		Action->NodeTemplate = EventNodeTemplate;
-		Action->SearchTitle = EventNodeTemplate->GetNodeSearchTitle();
 
 		ActionMenuBuilder.AddAction( Action );
 	}
@@ -593,7 +578,6 @@ void FK2ActionMenuBuilder::AttachMacroGraphAction(FGraphActionListBuilderBase& A
 		MacroTemplate->SetMacroGraph(MacroGraph);
 		TSharedPtr<FEdGraphSchemaAction_K2NewNode> Action = FK2ActionMenuBuilder::AddNewNodeAction(ActionMenuBuilder, CategoryName, MacroTemplate->GetNodeTitle(ENodeTitleType::ListView), MacroTemplate->GetTooltip());
 		Action->NodeTemplate = MacroTemplate;
-		Action->SearchTitle = MacroTemplate->GetNodeSearchTitle();
 	}
 }
 
@@ -696,7 +680,6 @@ void FK2ActionMenuBuilder::AddSpawnInfoForFunction(const UFunction* Function, bo
 	//@TODO: FUNCTIONREFACTOR:	CallFuncNode->FunctionReference.SetExternalMember(Function->GetFName(), Function->GetOuterUClass());
 
 	Action->NodeTemplate = CallTemplate;
-	Action->SearchTitle = Action->NodeTemplate->GetNodeSearchTitle();
 }
 
 /*******************************************************************************
@@ -940,7 +923,6 @@ void FK2ActionMenuBuilder::GetContextAllowedNodeTypes(FBlueprintGraphActionListB
 				UK2Node* NodeTemplate = ContextMenuBuilder.CreateTemplateNode<UK2Node_AssignmentStatement>();
 				TSharedPtr<FEdGraphSchemaAction_K2NewNode> Action = FK2ActionMenuBuilder::AddNewNodeAction(ContextMenuBuilder, K2ActionCategories::MacroToolsCategory, NodeTemplate->GetNodeTitle(ENodeTitleType::ListView), NodeTemplate->GetTooltip(), 0, NodeTemplate->GetKeywords());
 				Action->NodeTemplate = NodeTemplate;
-				Action->SearchTitle = NodeTemplate->GetNodeSearchTitle();
 			}
 
 			// Temporary variables
@@ -999,7 +981,6 @@ void FK2ActionMenuBuilder::GetContextAllowedNodeTypes(FBlueprintGraphActionListB
 					NodeTemplate->VariableType = *TypeIt;
 					TSharedPtr<FEdGraphSchemaAction_K2NewNode> Action = FK2ActionMenuBuilder::AddNewNodeAction(ContextMenuBuilder, K2ActionCategories::MacroToolsCategory, NodeTemplate->GetNodeTitle(ENodeTitleType::ListView), NodeTemplate->GetTooltip(), 0, NodeTemplate->GetKeywords());
 					Action->NodeTemplate = NodeTemplate;
-					Action->SearchTitle = NodeTemplate->GetNodeSearchTitle();
 				}
 			}
 
@@ -1012,7 +993,6 @@ void FK2ActionMenuBuilder::GetContextAllowedNodeTypes(FBlueprintGraphActionListB
 				PersistentIntNodeTemplate->bIsPersistent = true;
 				TSharedPtr<FEdGraphSchemaAction_K2NewNode> IntAction = FK2ActionMenuBuilder::AddNewNodeAction(ContextMenuBuilder, K2ActionCategories::MacroToolsCategory, PersistentIntNodeTemplate->GetNodeTitle(ENodeTitleType::ListView), PersistentIntNodeTemplate->GetTooltip(), 0, PersistentIntNodeTemplate->GetKeywords());
 				IntAction->NodeTemplate = PersistentIntNodeTemplate;
-				IntAction->SearchTitle = PersistentIntNodeTemplate->GetNodeSearchTitle();
 
 				// Persistent Bool
 				UK2Node_TemporaryVariable* PersistentBoolNodeTemplate = ContextMenuBuilder.CreateTemplateNode<UK2Node_TemporaryVariable>();
@@ -1020,7 +1000,6 @@ void FK2ActionMenuBuilder::GetContextAllowedNodeTypes(FBlueprintGraphActionListB
 				PersistentBoolNodeTemplate->bIsPersistent = true;
 				TSharedPtr<FEdGraphSchemaAction_K2NewNode> BoolAction = FK2ActionMenuBuilder::AddNewNodeAction(ContextMenuBuilder, K2ActionCategories::MacroToolsCategory, PersistentBoolNodeTemplate->GetNodeTitle(ENodeTitleType::ListView), PersistentBoolNodeTemplate->GetTooltip(), 0, PersistentBoolNodeTemplate->GetKeywords());
 				BoolAction->NodeTemplate = PersistentBoolNodeTemplate;
-				BoolAction->SearchTitle = PersistentBoolNodeTemplate->GetNodeSearchTitle();
 			}
 		}
 
@@ -1043,7 +1022,6 @@ void FK2ActionMenuBuilder::GetContextAllowedNodeTypes(FBlueprintGraphActionListB
 			UK2Node* NodeTemplate = ContextMenuBuilder.CreateTemplateNode<UK2Node_Select>();
 			TSharedPtr<FEdGraphSchemaAction_K2NewNode> SelectionAction = FK2ActionMenuBuilder::AddNewNodeAction(ContextMenuBuilder, K2ActionCategories::UtilitiesCategory, NodeTemplate->GetNodeTitle(ENodeTitleType::ListView), NodeTemplate->GetTooltip(), 0, NodeTemplate->GetKeywords());
 			SelectionAction->NodeTemplate = NodeTemplate;
-			SelectionAction->SearchTitle = NodeTemplate->GetNodeSearchTitle();
 		}
 
 		FEditorDelegates::OnBlueprintContextMenuCreated.Broadcast(ContextMenuBuilder);
@@ -1073,7 +1051,6 @@ struct FClassDynamicCastHelper
 
 		TSharedPtr<FEdGraphSchemaAction_K2NewNode> Action = FK2ActionMenuBuilder::AddNewNodeAction(ContextMenuBuilder, CastCategory, CastNode->GetNodeTitle(ENodeTitleType::ListView), TEXT("Dynamic cast to a specific type"), 0, CastNode->GetKeywords());
 		Action->NodeTemplate = CastNode;
-		Action->SearchTitle = CastNode->GetNodeSearchTitle();
 	}
 
 	static void GetClassDynamicCastNodes(FBlueprintGraphActionListBuilder& ContextMenuBuilder, const UEdGraphSchema_K2* K2Schema)
@@ -1279,7 +1256,6 @@ void FK2ActionMenuBuilder::GetPinAllowedNodeTypes(FBlueprintGraphActionListBuild
 
 								TSharedPtr<FEdGraphSchemaAction_K2NewNode> Action = FK2ActionMenuBuilder::AddNewNodeAction(ContextMenuBuilder,  CastCategory, CastNode->GetNodeTitle(ENodeTitleType::ListView), TEXT("Dynamic cast to a specific type"), 0, CastNode->GetKeywords());
 								Action->NodeTemplate = CastNode;
-								Action->SearchTitle = CastNode->GetNodeSearchTitle();
 							}
 						}
 					}
@@ -1339,7 +1315,6 @@ void FK2ActionMenuBuilder::GetPinAllowedNodeTypes(FBlueprintGraphActionListBuild
 					NodeTemplate->StructType = PinStruct;
 					TSharedPtr<FEdGraphSchemaAction_K2NewNode> Action = FK2ActionMenuBuilder::AddNewNodeAction(ContextMenuBuilder, FString(), NodeTemplate->GetNodeTitle(ENodeTitleType::ListView), NodeTemplate->GetTooltip(), 0, NodeTemplate->GetKeywords());
 					Action->NodeTemplate = NodeTemplate;
-					Action->SearchTitle = NodeTemplate->GetNodeSearchTitle();
 				}
 				else if ((FromPin.Direction == EGPD_Input) && UK2Node_MakeStruct::CanBeMade(PinStruct))
 				{
@@ -1347,7 +1322,6 @@ void FK2ActionMenuBuilder::GetPinAllowedNodeTypes(FBlueprintGraphActionListBuild
 					NodeTemplate->StructType = PinStruct;
 					TSharedPtr<FEdGraphSchemaAction_K2NewNode> Action = FK2ActionMenuBuilder::AddNewNodeAction(ContextMenuBuilder, FString(), NodeTemplate->GetNodeTitle(ENodeTitleType::ListView), NodeTemplate->GetTooltip(), 0, NodeTemplate->GetKeywords());
 					Action->NodeTemplate = NodeTemplate;
-					Action->SearchTitle = NodeTemplate->GetNodeSearchTitle();
 				}
 			}
 
@@ -1356,7 +1330,6 @@ void FK2ActionMenuBuilder::GetPinAllowedNodeTypes(FBlueprintGraphActionListBuild
 				UK2Node_CreateDelegate* NodeTemplate = ContextMenuBuilder.CreateTemplateNode<UK2Node_CreateDelegate>();
 				TSharedPtr<FEdGraphSchemaAction_K2NewNode> Action = FK2ActionMenuBuilder::AddNewNodeAction(ContextMenuBuilder, FString(), NodeTemplate->GetNodeTitle(ENodeTitleType::ListView), NodeTemplate->GetTooltip(), 0, NodeTemplate->GetKeywords());
 				Action->NodeTemplate = NodeTemplate;
-				Action->SearchTitle = NodeTemplate->GetNodeSearchTitle();
 
 				UFunction* Signature = FMemberReference::ResolveSimpleMemberReference<UFunction>(FromPin.PinType.PinSubCategoryMemberReference);
 				const bool bSupportsEventGraphs = (Blueprint && FBlueprintEditorUtils::DoesSupportEventGraphs(Blueprint));
@@ -1603,7 +1576,6 @@ void FK2ActionMenuBuilder::GetSwitchMenuItems(FBlueprintPaletteListBuilder& Acti
 				UK2Node* NodeTemplate = EnumTemplate;
 				TSharedPtr<FEdGraphSchemaAction_K2NewNode> Action = FK2ActionMenuBuilder::AddNewNodeAction(ActionMenuBuilder, SwitchControlCategory, NodeTemplate->GetNodeTitle(ENodeTitleType::ListView), NodeTemplate->GetTooltip(), 0, NodeTemplate->GetKeywords());
 				Action->NodeTemplate = NodeTemplate;
-				Action->SearchTitle = NodeTemplate->GetNodeSearchTitle();
 			}
 		}
 	}
@@ -1631,7 +1603,6 @@ void FK2ActionMenuBuilder::GetAnimNotifyMenuItems(FBlueprintGraphActionListBuild
 				EventNode->EventSignatureClass = UAnimInstance::StaticClass();
 				EventNode->CustomFunctionName = EventNode->EventSignatureName;
 				Action->NodeTemplate = EventNode;
-				Action->SearchTitle = Action->NodeTemplate->GetNodeSearchTitle();
 			}
 
 			// @todo anim: fix this to be same as notifies, save same list in the 
@@ -1665,7 +1636,6 @@ void FK2ActionMenuBuilder::GetAnimNotifyMenuItems(FBlueprintGraphActionListBuild
 				EventNode->EventSignatureClass = UAnimInstance::StaticClass();
 				EventNode->CustomFunctionName = EventNode->EventSignatureName;
 				Action->NodeTemplate = EventNode;
-				Action->SearchTitle = Action->NodeTemplate->GetNodeSearchTitle();
 			}
 		}
 
@@ -1687,7 +1657,6 @@ void FK2ActionMenuBuilder::GetAnimNotifyMenuItems(FBlueprintGraphActionListBuild
 					EventNode->EventSignatureClass = UAnimInstance::StaticClass();
 					EventNode->CustomFunctionName = EventNode->EventSignatureName;
 					Action->NodeTemplate = EventNode;
-					Action->SearchTitle = Action->NodeTemplate->GetNodeSearchTitle();
 				}
 			}
 		}
@@ -1775,7 +1744,6 @@ void FK2ActionMenuBuilder::GetVariableGettersSettersForClass(FBlueprintGraphActi
 					UK2Node_VariableGet* VarGetNode = ContextMenuBuilder.CreateTemplateNode<UK2Node_VariableGet>();
 					VarGetNode->SetFromProperty(Property, bSelfContext);
 					ReadVar->NodeTemplate = VarGetNode;
-					ReadVar->SearchTitle = ReadVar->NodeTemplate->GetNodeSearchTitle();
 				}
 
 				if (bCanWrite)
@@ -1785,7 +1753,6 @@ void FK2ActionMenuBuilder::GetVariableGettersSettersForClass(FBlueprintGraphActi
 					UK2Node_VariableSet* VarSetNode = ContextMenuBuilder.CreateTemplateNode<UK2Node_VariableSet>();
 					VarSetNode->SetFromProperty(Property, bSelfContext);
 					WriteVar->NodeTemplate = VarSetNode;
-					WriteVar->SearchTitle = WriteVar->NodeTemplate->GetNodeSearchTitle();
 				}
 			}
 		}
@@ -1837,7 +1804,6 @@ void FK2ActionMenuBuilder::GetVariableGettersSettersForClass(FBlueprintGraphActi
 					UK2Node_VariableGet* VarGetNode = ContextMenuBuilder.CreateTemplateNode<UK2Node_VariableGet>();
 					VarGetNode->VariableReference.SetLocalMember(LocalVariable.VarName, FunctionGraph->GetName(), LocalVariable.VarGuid);
 					ReadVar->NodeTemplate = VarGetNode;
-					ReadVar->SearchTitle = ReadVar->NodeTemplate->GetNodeSearchTitle();
 				}
 
 				if (bCanWrite)
@@ -1847,7 +1813,6 @@ void FK2ActionMenuBuilder::GetVariableGettersSettersForClass(FBlueprintGraphActi
 					UK2Node_VariableSet* VarSetNode = ContextMenuBuilder.CreateTemplateNode<UK2Node_VariableSet>();
 					VarSetNode->VariableReference.SetLocalMember(LocalVariable.VarName, FunctionGraph->GetName(), LocalVariable.VarGuid);
 					WriteVar->NodeTemplate = VarSetNode;
-					WriteVar->SearchTitle = WriteVar->NodeTemplate->GetNodeSearchTitle();
 				}
 			}
 		}
@@ -1867,7 +1832,6 @@ void FK2ActionMenuBuilder::GetVariableGettersSettersForClass(FBlueprintGraphActi
 		UK2Node_Self* SelfNode = ContextMenuBuilder.CreateTemplateNode<UK2Node_Self>();
 		SelfNode->SelfClass = (UClass*)Class;
 		SelfVar->NodeTemplate = SelfNode;
-		SelfVar->SearchTitle = SelfVar->NodeTemplate->GetNodeSearchTitle();
 		SelfVar->Keywords = SelfNode->GetKeywords();
 	}
 }
@@ -1910,13 +1874,10 @@ void FK2ActionMenuBuilder::GetEventDispatcherNodesForClass(FBlueprintGraphAction
 						const FString PropertyDescription = GetDefault<UEditorStyleSettings>()->bShowFriendlyNames ? UEditorEngine::GetFriendlyName(Property) : Property->GetName();
 						TSharedPtr<FEdGraphSchemaAction_K2AssignDelegate> AssignDelegateEvent(new FEdGraphSchemaAction_K2AssignDelegate(K2ActionCategories::DelegateBindingCategory, FText::FromString(BindString + PropertyDescription), PropertyTooltip, 0));
 						AssignDelegateEvent->NodeTemplate = DelegateEventNode;
-						AssignDelegateEvent->SearchTitle = AssignDelegateEvent->NodeTemplate->GetNodeSearchTitle();
 						ContextMenuBuilder.AddAction( AssignDelegateEvent );
 					}
 					TSharedPtr<FEdGraphSchemaAction_K2NewNode> DelegateEvent = FK2ActionMenuBuilder::AddNewNodeAction(ContextMenuBuilder, K2ActionCategories::DelegateCategory, DelegateEventNode->GetNodeTitle(ENodeTitleType::FullTitle), PropertyTooltip);
 					DelegateEvent->NodeTemplate = DelegateEventNode;
-					DelegateEvent->SearchTitle = DelegateEvent->NodeTemplate->GetNodeSearchTitle();
-					DelegateEvent->SearchTitle = DelegateEventNode->GetNodeSearchTitle();
 				}
 
 				{
@@ -1924,8 +1885,6 @@ void FK2ActionMenuBuilder::GetEventDispatcherNodesForClass(FBlueprintGraphAction
 					DelegateEventNode->SetFromProperty(Property, bSelfContext);
 					TSharedPtr<FEdGraphSchemaAction_K2NewNode> DelegateEvent = FK2ActionMenuBuilder::AddNewNodeAction(ContextMenuBuilder, K2ActionCategories::DelegateCategory, DelegateEventNode->GetNodeTitle(ENodeTitleType::FullTitle), PropertyTooltip);
 					DelegateEvent->NodeTemplate = DelegateEventNode;
-					DelegateEvent->SearchTitle = DelegateEvent->NodeTemplate->GetNodeSearchTitle();
-					DelegateEvent->SearchTitle = DelegateEventNode->GetNodeSearchTitle();
 				}
 
 				{
@@ -1933,8 +1892,6 @@ void FK2ActionMenuBuilder::GetEventDispatcherNodesForClass(FBlueprintGraphAction
 					DelegateEventNode->SetFromProperty(Property, bSelfContext);
 					TSharedPtr<FEdGraphSchemaAction_K2NewNode> DelegateEvent = FK2ActionMenuBuilder::AddNewNodeAction(ContextMenuBuilder, K2ActionCategories::DelegateCategory, DelegateEventNode->GetNodeTitle(ENodeTitleType::FullTitle), PropertyTooltip);
 					DelegateEvent->NodeTemplate = DelegateEventNode;
-					DelegateEvent->SearchTitle = DelegateEvent->NodeTemplate->GetNodeSearchTitle();
-					DelegateEvent->SearchTitle = DelegateEventNode->GetNodeSearchTitle();
 				}
 			}
 
@@ -1944,8 +1901,6 @@ void FK2ActionMenuBuilder::GetEventDispatcherNodesForClass(FBlueprintGraphAction
 				DelegateEventNode->SetFromProperty(Property, bSelfContext);
 				TSharedPtr<FEdGraphSchemaAction_K2NewNode> DelegateEvent = FK2ActionMenuBuilder::AddNewNodeAction(ContextMenuBuilder, K2ActionCategories::DelegateCategory, DelegateEventNode->GetNodeTitle(ENodeTitleType::FullTitle), PropertyTooltip);
 				DelegateEvent->NodeTemplate = DelegateEventNode;
-				DelegateEvent->SearchTitle = DelegateEvent->NodeTemplate->GetNodeSearchTitle();
-				DelegateEvent->SearchTitle = DelegateEventNode->GetNodeSearchTitle();
 			}
 		}
 	}
@@ -1969,7 +1924,6 @@ void FK2ActionMenuBuilder::GetStructActions(FBlueprintGraphActionListBuilder& Co
 					BreakNodeTemplate->StructType = Struct;
 					TSharedPtr<FEdGraphSchemaAction_K2NewNode> Action = FK2ActionMenuBuilder::AddNewNodeAction(ContextMenuBuilder, K2ActionCategories::BreakStructCategory, BreakNodeTemplate->GetNodeTitle(ENodeTitleType::ListView), BreakNodeTemplate->GetTooltip(), 0, BreakNodeTemplate->GetKeywords());
 					Action->NodeTemplate = BreakNodeTemplate;
-					Action->SearchTitle = Action->NodeTemplate->GetNodeSearchTitle();
 				}
 			}
 
@@ -1981,7 +1935,6 @@ void FK2ActionMenuBuilder::GetStructActions(FBlueprintGraphActionListBuilder& Co
 					MakeNodeTemplate->StructType = Struct;
 					TSharedPtr<FEdGraphSchemaAction_K2NewNode> Action = FK2ActionMenuBuilder::AddNewNodeAction(ContextMenuBuilder, K2ActionCategories::MakeStructCategory, MakeNodeTemplate->GetNodeTitle(ENodeTitleType::ListView), MakeNodeTemplate->GetTooltip(), 0, MakeNodeTemplate->GetKeywords());
 					Action->NodeTemplate = MakeNodeTemplate;
-					Action->SearchTitle = Action->NodeTemplate->GetNodeSearchTitle();
 				}
 			}
 		}
@@ -2004,7 +1957,6 @@ void FK2ActionMenuBuilder::GetAllEventActions(FBlueprintPaletteListBuilder& Acti
 		UK2Node_Timeline* TimelineNodeTemplate = ActionMenuBuilder.CreateTemplateNode<UK2Node_Timeline>();
 		TimelineNodeTemplate->TimelineName = FBlueprintEditorUtils::FindUniqueTimelineName(ActionMenuBuilder.Blueprint);
 		Action->NodeTemplate = TimelineNodeTemplate;
-		Action->SearchTitle = Action->NodeTemplate->GetNodeSearchTitle();
 	}
 }
 
@@ -2106,7 +2058,6 @@ void FK2ActionMenuBuilder::GetEventsForBlueprint(FBlueprintPaletteListBuilder& A
 	CustomEventNode->CustomFunctionName = FBlueprintEditorUtils::FindUniqueCustomEventName(Blueprint);
 	CustomEventNode->bIsEditable = true;
 	Action->NodeTemplate = CustomEventNode;
-	Action->SearchTitle = Action->NodeTemplate->GetNodeSearchTitle();
 }
 
 //------------------------------------------------------------------------------
@@ -2146,7 +2097,6 @@ void FK2ActionMenuBuilder::GetLiteralsFromActorSelection(FBlueprintGraphActionLi
 			TSharedPtr<FEdGraphSchemaAction_K2NewNode> Action( new FEdGraphSchemaAction_K2NewNode(TEXT(""), FText::FromString(Description), ToolTip, K2Schema->AG_LevelReference) );
 			Action->Keywords = LiteralNodeTemplate->GetKeywords();
 			Action->NodeTemplate = LiteralNodeTemplate;
-			Action->SearchTitle = Action->NodeTemplate->GetNodeSearchTitle();
 
 			NewActions.Add( Action );
 		}
@@ -2231,7 +2181,6 @@ void FK2ActionMenuBuilder::AddBoundEventActionsForClass(FBlueprintGraphActionLis
 						UK2Node_ComponentBoundEvent* NewComponentEvent = ContextMenuBuilder.CreateTemplateNode<UK2Node_ComponentBoundEvent>();
 						NewComponentEvent->InitializeComponentBoundEventParams( ComponentProperty, DelegateProperty );
 						NewDelegateNode->NodeTemplate = NewComponentEvent;
-						NewDelegateNode->SearchTitle = NewDelegateNode->NodeTemplate->GetNodeSearchTitle();
 
 						ContextMenuBuilder.AddAction( NewDelegateNode );
 					}
@@ -2257,7 +2206,6 @@ void FK2ActionMenuBuilder::AddBoundEventActionsForClass(FBlueprintGraphActionLis
 							UK2Node_ActorBoundEvent* NewActorEvent = ContextMenuBuilder.CreateTemplateNode<UK2Node_ActorBoundEvent>();
 							NewActorEvent->InitializeActorBoundEventParams( Actors[ActorIndex], DelegateProperty );
 							NewDelegateNode->NodeTemplate = NewActorEvent;
-							NewDelegateNode->SearchTitle = NewDelegateNode->NodeTemplate->GetNodeSearchTitle();
 
 							NewActions.Add( NewDelegateNode );
 						}
@@ -2331,7 +2279,6 @@ void FK2ActionMenuBuilder::GetMatineeControllers(FBlueprintGraphActionListBuilde
 
 		TSharedPtr<FEdGraphSchemaAction_K2NewNode> Action = FK2ActionMenuBuilder::AddNewNodeAction(ContextMenuBuilder, TEXT(""), Description, ToolTip, K2Schema->AG_LevelReference, MatControlNode->GetKeywords());
 		Action->NodeTemplate = MatControlNode;
-		Action->SearchTitle = Action->NodeTemplate->GetNodeSearchTitle();
 	}
 }
 
@@ -2594,7 +2541,6 @@ void FK2ActionMenuBuilder::GetAllInterfaceMessageActions(FGraphActionListBuilder
 					MessageTemplate->FunctionReference.SetExternalMember(Function->GetFName(), CurrentInterface);
 					TSharedPtr<FEdGraphSchemaAction_K2NewNode> Action = FK2ActionMenuBuilder::AddNewNodeAction(ActionMenuBuilder, MessageCategory, FText::FromString(Function->GetName()), MessageTemplate->GetTooltip(), 0, MessageTemplate->GetKeywords());
 					Action->NodeTemplate = MessageTemplate;
-					Action->SearchTitle = Action->NodeTemplate->GetNodeSearchTitle();
 				}
 			}
 		}
diff --git a/Engine/Source/Editor/BlueprintGraph/Private/K2Node_ActorBoundEvent.cpp b/Engine/Source/Editor/BlueprintGraph/Private/K2Node_ActorBoundEvent.cpp
--- a/Engine/Source/Editor/BlueprintGraph/Private/K2Node_ActorBoundEvent.cpp
+++ b/Engine/Source/Editor/BlueprintGraph/Private/K2Node_ActorBoundEvent.cpp
@@ -92,18 +92,6 @@ FText UK2Node_ActorBoundEvent::GetNodeTitle(ENodeTitleType::Type TitleType) cons
 	return FText::Format(LOCTEXT("ActorBoundEventTitle", "{DelegatePropertyName} ({TargetName})"), Args);
 }
 
-FString UK2Node_ActorBoundEvent::GetNodeNativeTitle(ENodeTitleType::Type TitleType) const
-{
-	// Do not setup this function for localization, intentionally left unlocalized!
-	FString TargetName = TEXT("None");
-	if( EventOwner )
-	{
-		TargetName = EventOwner->GetActorLabel();		
-	}
-
-	return FString::Printf(TEXT("%s (%s)"), *DelegatePropertyName.ToString(), *TargetName);
-}
-
 FString UK2Node_ActorBoundEvent::GetTooltip() const
 {
 	UMulticastDelegateProperty* TargetDelegateProp = GetTargetDelegatePropertyConst();
diff --git a/Engine/Source/Editor/BlueprintGraph/Private/K2Node_BaseAsyncTask.cpp b/Engine/Source/Editor/BlueprintGraph/Private/K2Node_BaseAsyncTask.cpp
--- a/Engine/Source/Editor/BlueprintGraph/Private/K2Node_BaseAsyncTask.cpp
+++ b/Engine/Source/Editor/BlueprintGraph/Private/K2Node_BaseAsyncTask.cpp
@@ -53,7 +53,6 @@ TSharedPtr<FEdGraphSchemaAction_K2NewNode> UK2Node_BaseAsyncTask::CreateDefaultM
 	TSharedPtr<FEdGraphSchemaAction_K2NewNode> NodeAction = FK2ActionMenuBuilder::AddNewNodeAction(ContextMenuBuilder, NodeTemplate->GetCategoryName(), NodeTemplate->GetNodeTitle(ENodeTitleType::ListView), NodeTemplate->GetTooltip(), 0, NodeTemplate->GetKeywords());
 	
 	NodeAction->NodeTemplate = NodeTemplate;
-	NodeAction->SearchTitle = NodeTemplate->GetNodeSearchTitle();
 
 	return NodeAction;
 }
diff --git a/Engine/Source/Editor/BlueprintGraph/Private/K2Node_BreakStruct.cpp b/Engine/Source/Editor/BlueprintGraph/Private/K2Node_BreakStruct.cpp
--- a/Engine/Source/Editor/BlueprintGraph/Private/K2Node_BreakStruct.cpp
+++ b/Engine/Source/Editor/BlueprintGraph/Private/K2Node_BreakStruct.cpp
@@ -180,12 +180,6 @@ FText UK2Node_BreakStruct::GetNodeTitle(ENodeTitleType::Type TitleType) const
 	return FText::Format(LOCTEXT("BreakNodeTitle", "Break {StructName}"), Args);
 }
 
-FString UK2Node_BreakStruct::GetNodeNativeTitle(ENodeTitleType::Type TitleType) const
-{
-	// Do not setup this function for localization, intentionally left unlocalized!
-	return FString(TEXT("Break ")) + (StructType ? StructType->GetName() : FString());
-}
-
 FString UK2Node_BreakStruct::GetTooltip() const
 {
 	return FString::Printf(
diff --git a/Engine/Source/Editor/BlueprintGraph/Private/K2Node_CallFunction.cpp b/Engine/Source/Editor/BlueprintGraph/Private/K2Node_CallFunction.cpp
--- a/Engine/Source/Editor/BlueprintGraph/Private/K2Node_CallFunction.cpp
+++ b/Engine/Source/Editor/BlueprintGraph/Private/K2Node_CallFunction.cpp
@@ -174,32 +174,6 @@ FText UK2Node_CallFunction::GetNodeTitle(ENodeTitleType::Type TitleType) const
 	}
 }
 
-FString UK2Node_CallFunction::GetNodeNativeTitle(ENodeTitleType::Type TitleType) const
-{
-	// Do not setup this function for localization, intentionally left unlocalized!
-	FString FunctionName;
-	FString ContextString;
-	FString RPCString;
-
-	if (UFunction* Function = GetTargetFunction())
-	{
-		RPCString = UK2Node_Event::GetLocalizedNetString(Function->FunctionFlags, true);
-		FunctionName = UK2Node_CallFunction::GetUserFacingFunctionName(Function);
-		ContextString = GetFunctionContextString();
-	}
-	else
-	{
-		FunctionName = FunctionReference.GetMemberName().ToString();
-		if ((GEditor != NULL) && (GetDefault<UEditorStyleSettings>()->bShowFriendlyNames))
-		{
-			FunctionName = FName::NameToDisplayString(FunctionName, false);
-		}
-	}
-
-	const FString Result = (TitleType == ENodeTitleType::ListView) ? FunctionName : FunctionName + ContextString + RPCString;
-	return Result;
-}
-
 void UK2Node_CallFunction::AllocateDefaultPins()
 {
 	UBlueprint* MyBlueprint = GetBlueprint();
diff --git a/Engine/Source/Editor/BlueprintGraph/Private/K2Node_CallParentFunction.cpp b/Engine/Source/Editor/BlueprintGraph/Private/K2Node_CallParentFunction.cpp
--- a/Engine/Source/Editor/BlueprintGraph/Private/K2Node_CallParentFunction.cpp
+++ b/Engine/Source/Editor/BlueprintGraph/Private/K2Node_CallParentFunction.cpp
@@ -25,20 +25,6 @@ FText UK2Node_CallParentFunction::GetNodeTitle(ENodeTitleType::Type TitleType) c
 	return FText::Format( LOCTEXT( "CallSuperFunction", "Parent: {FunctionName}" ), Args);
 }
 
-FString UK2Node_CallParentFunction::GetNodeNativeTitle(ENodeTitleType::Type TitleType) const
-{
-	// Do not setup this function for localization, intentionally left unlocalized!
-	UFunction* Function = FunctionReference.ResolveMember<UFunction>(this);
-	FString FunctionName = Function ? GetUserFacingFunctionName( Function ) : FunctionReference.GetMemberName().ToString();
-	FString Result = FString::Printf( TEXT("Parent: %s" ), *FunctionName);
-
-	if( GEditor && GetDefault<UEditorStyleSettings>()->bShowFriendlyNames )
-	{
-		Result = FName::NameToDisplayString(Result, false);
-	}
-	return Result;
-}
-
 FLinearColor UK2Node_CallParentFunction::GetNodeTitleColor() const
 {
 	return GetDefault<UGraphEditorSettings>()->ParentFunctionCallNodeTitleColor;
diff --git a/Engine/Source/Editor/BlueprintGraph/Private/K2Node_ComponentBoundEvent.cpp b/Engine/Source/Editor/BlueprintGraph/Private/K2Node_ComponentBoundEvent.cpp
--- a/Engine/Source/Editor/BlueprintGraph/Private/K2Node_ComponentBoundEvent.cpp
+++ b/Engine/Source/Editor/BlueprintGraph/Private/K2Node_ComponentBoundEvent.cpp
@@ -17,12 +17,6 @@ FText UK2Node_ComponentBoundEvent::GetNodeTitle(ENodeTitleType::Type TitleType)
 	return FText::Format(LOCTEXT("ComponentBoundEvent_Title", "{DelegatePropertyName} ({ComponentPropertyName})"), Args);
 }
 
-FString UK2Node_ComponentBoundEvent::GetNodeNativeTitle(ENodeTitleType::Type TitleType) const
-{
-	// Do not setup this function for localization, intentionally left unlocalized!
-	return FString::Printf(TEXT("%s (%s)"), *DelegatePropertyName.ToString(), *ComponentPropertyName.ToString());
-}
-
 void UK2Node_ComponentBoundEvent::InitializeComponentBoundEventParams(UObjectProperty* InComponentProperty, const UMulticastDelegateProperty* InDelegateProperty)
 {
 	if( InComponentProperty && InDelegateProperty )
diff --git a/Engine/Source/Editor/BlueprintGraph/Private/K2Node_Composite.cpp b/Engine/Source/Editor/BlueprintGraph/Private/K2Node_Composite.cpp
--- a/Engine/Source/Editor/BlueprintGraph/Private/K2Node_Composite.cpp
+++ b/Engine/Source/Editor/BlueprintGraph/Private/K2Node_Composite.cpp
@@ -176,19 +176,6 @@ FText UK2Node_Composite::GetNodeTitle(ENodeTitleType::Type TitleType) const
 	}
 }
 
-FString UK2Node_Composite::GetNodeNativeTitle(ENodeTitleType::Type TitleType) const
-{
-	// Do not setup this function for localization, intentionally left unlocalized!
-	if(TitleType == ENodeTitleType::FullTitle)
-	{
-		return FString::Printf(TEXT("%s\nCollapsed Graph"), (BoundGraph)?*BoundGraph->GetName() : TEXT("Invalid Graph"));
-	}
-	else
-	{
-		return (BoundGraph)?*BoundGraph->GetName() : TEXT("Invalid Graph");
-	}
-}
-
 bool UK2Node_Composite::CanUserDeleteNode() const
 {
 	return true;
diff --git a/Engine/Source/Editor/BlueprintGraph/Private/K2Node_CustomEvent.cpp b/Engine/Source/Editor/BlueprintGraph/Private/K2Node_CustomEvent.cpp
--- a/Engine/Source/Editor/BlueprintGraph/Private/K2Node_CustomEvent.cpp
+++ b/Engine/Source/Editor/BlueprintGraph/Private/K2Node_CustomEvent.cpp
@@ -122,24 +122,6 @@ FText UK2Node_CustomEvent::GetNodeTitle(ENodeTitleType::Type TitleType) const
 	}
 }
 
-FString UK2Node_CustomEvent::GetNodeNativeTitle(ENodeTitleType::Type TitleType) const
-{
-	// Do not setup this function for localization, intentionally left unlocalized!
-	if (TitleType == ENodeTitleType::EditableTitle)
-	{
-		return CustomFunctionName.ToString();
-	}
-	else if(TitleType == ENodeTitleType::ListView)
-	{
-		return TEXT("Custom Event");
-	}
-	else
-	{
-		FString RPCString = UK2Node_Event::GetLocalizedNetString(FunctionFlags, false);
-		return FString::Printf(TEXT("%s\nCustom Event"), *CustomFunctionName.ToString()) + RPCString;
-	}
-}
-
 UEdGraphPin* UK2Node_CustomEvent::CreatePinFromUserDefinition(const TSharedPtr<FUserPinInfo> NewPinInfo)
 {
 	UEdGraphPin* NewPin = CreatePin(
diff --git a/Engine/Source/Editor/BlueprintGraph/Private/K2Node_DelegateSet.cpp b/Engine/Source/Editor/BlueprintGraph/Private/K2Node_DelegateSet.cpp
--- a/Engine/Source/Editor/BlueprintGraph/Private/K2Node_DelegateSet.cpp
+++ b/Engine/Source/Editor/BlueprintGraph/Private/K2Node_DelegateSet.cpp
@@ -169,12 +169,6 @@ FText UK2Node_DelegateSet::GetNodeTitle(ENodeTitleType::Type TitleType) const
 	return FText::Format(NSLOCTEXT("K2Node", "Assign_Name", "Assign {DelegatePropertyName}"), Args);
 }
 
-FString UK2Node_DelegateSet::GetNodeNativeTitle(ENodeTitleType::Type TitleType) const
-{
-	// Do not setup this function for localization, intentionally left unlocalized!
-	return FString::Printf(TEXT("Assign %s"), *DelegatePropertyName.ToString());
-}
-
 UEdGraphPin* UK2Node_DelegateSet::GetDelegateOwner() const
 {
 	const UEdGraphSchema_K2* K2Schema = GetDefault<UEdGraphSchema_K2>();
diff --git a/Engine/Source/Editor/BlueprintGraph/Private/K2Node_DoOnceMultiInput.cpp b/Engine/Source/Editor/BlueprintGraph/Private/K2Node_DoOnceMultiInput.cpp
--- a/Engine/Source/Editor/BlueprintGraph/Private/K2Node_DoOnceMultiInput.cpp
+++ b/Engine/Source/Editor/BlueprintGraph/Private/K2Node_DoOnceMultiInput.cpp
@@ -21,7 +21,6 @@ void UK2Node_DoOnceMultiInput::GetMenuEntries(FGraphContextMenuBuilder& ContextM
 	TSharedPtr<FEdGraphSchemaAction_K2NewNode> NodeAction = FK2ActionMenuBuilder::AddNewNodeAction(ContextMenuBuilder, LOCTEXT("FlowCategory", "Flow Control").ToString(), NodeTemplate->GetNodeTitle(ENodeTitleType::ListView), NodeTemplate->GetTooltip(), 0, NodeTemplate->GetKeywords());
 
 	NodeAction->NodeTemplate = NodeTemplate;
-	NodeAction->SearchTitle = NodeTemplate->GetNodeSearchTitle();
 }
 
 FText UK2Node_DoOnceMultiInput::GetNodeTitle(ENodeTitleType::Type TitleType) const
diff --git a/Engine/Source/Editor/BlueprintGraph/Private/K2Node_DynamicCast.cpp b/Engine/Source/Editor/BlueprintGraph/Private/K2Node_DynamicCast.cpp
--- a/Engine/Source/Editor/BlueprintGraph/Private/K2Node_DynamicCast.cpp
+++ b/Engine/Source/Editor/BlueprintGraph/Private/K2Node_DynamicCast.cpp
@@ -76,32 +76,6 @@ FText UK2Node_DynamicCast::GetNodeTitle(ENodeTitleType::Type TitleType) const
 	}
 }
 
-FString UK2Node_DynamicCast::GetNodeNativeTitle(ENodeTitleType::Type TitleType) const
-{
-	// Do not setup this function for localization, intentionally left unlocalized!
-	if(TargetType != NULL)
-	{
-		// If casting to BP class, use BP name not class name (ie. remove the _C)
-		FString TargetName;
-		UBlueprint* CastToBP = UBlueprint::GetBlueprintFromClass(TargetType);
-		if(CastToBP != NULL)
-		{
-			TargetName = CastToBP->GetName();
-		}
-		else
-		{
-			TargetName = TargetType->GetName();
-		}
-
-		return FString(TEXT("Cast To ")) + TargetName;
-	}
-	// No target type, bad node
-	else
-	{
-		return TEXT("Bad cast node");
-	}
-}
-
 UEdGraphPin* UK2Node_DynamicCast::GetValidCastPin() const
 {
 	const UEdGraphSchema_K2* K2Schema = GetDefault<UEdGraphSchema_K2>();
diff --git a/Engine/Source/Editor/BlueprintGraph/Private/K2Node_EnumEquality.cpp b/Engine/Source/Editor/BlueprintGraph/Private/K2Node_EnumEquality.cpp
--- a/Engine/Source/Editor/BlueprintGraph/Private/K2Node_EnumEquality.cpp
+++ b/Engine/Source/Editor/BlueprintGraph/Private/K2Node_EnumEquality.cpp
@@ -193,7 +193,6 @@ void UK2Node_EnumEquality::GetMenuEntries(FGraphContextMenuBuilder& Context) con
 
 		TSharedPtr<FEdGraphSchemaAction_K2NewNode> NodeAction = FK2ActionMenuBuilder::AddNewNodeAction(Context, Category, MenuDesc, Tooltip, 0, Keywords);
 		NodeAction->NodeTemplate = EnumNodeTemplate;
-		NodeAction->SearchTitle = EnumNodeTemplate->GetNodeSearchTitle();
 	}
 }
 
diff --git a/Engine/Source/Editor/BlueprintGraph/Private/K2Node_Event.cpp b/Engine/Source/Editor/BlueprintGraph/Private/K2Node_Event.cpp
--- a/Engine/Source/Editor/BlueprintGraph/Private/K2Node_Event.cpp
+++ b/Engine/Source/Editor/BlueprintGraph/Private/K2Node_Event.cpp
@@ -90,41 +90,6 @@ FText UK2Node_Event::GetNodeTitle(ENodeTitleType::Type TitleType) const
 	}
 }
 
-FString UK2Node_Event::GetNodeNativeTitle(ENodeTitleType::Type TitleType) const
-{
-	// Do not setup this function for localization, intentionally left unlocalized!
-	if (bOverrideFunction || (CustomFunctionName == NAME_None))
-	{
-		FString FunctionName = EventSignatureName.ToString(); // If we fail to find the function, still want to write something on the node.
-
-		if (UFunction* Function = FindField<UFunction>(EventSignatureClass, EventSignatureName))
-		{
-			FunctionName = UEdGraphSchema_K2::GetFriendlySignitureName(Function);
-		}
-
-		FString TitleString = FString(TEXT("Event ")) + FunctionName;
-
-		if(TitleType == ENodeTitleType::FullTitle && EventSignatureClass != NULL && EventSignatureClass->IsChildOf(UInterface::StaticClass()))
-		{
-			FString SourceString = EventSignatureClass->GetName();
-
-			// @todo: This action won't be necessary once the new name convention is used.
-			if(SourceString.EndsWith(TEXT("_C")))
-			{
-				SourceString = SourceString.LeftChop(2);
-			}
-
-			TitleString += FString(TEXT("\n")) + FString::Printf(TEXT("From %s"), *SourceString);
-		}
-
-		return TitleString;
-	}
-	else
-	{
-		return CustomFunctionName.ToString();
-	}
-}
-
 FString UK2Node_Event::GetTooltip() const
 {
 	FString Tooltip;
diff --git a/Engine/Source/Editor/BlueprintGraph/Private/K2Node_ForEachElementInEnum.cpp b/Engine/Source/Editor/BlueprintGraph/Private/K2Node_ForEachElementInEnum.cpp
--- a/Engine/Source/Editor/BlueprintGraph/Private/K2Node_ForEachElementInEnum.cpp
+++ b/Engine/Source/Editor/BlueprintGraph/Private/K2Node_ForEachElementInEnum.cpp
@@ -140,14 +140,6 @@ FText UK2Node_ForEachElementInEnum::GetNodeTitle(ENodeTitleType::Type TitleType)
 	return FText::Format(LOCTEXT("ForEachElementInEnum_Title", "ForEach {EnumName}"), Args);
 }
 
-FString UK2Node_ForEachElementInEnum::GetNodeNativeTitle(ENodeTitleType::Type TitleType) const
-{
-	// Do not setup this function for localization, intentionally left unlocalized!
-	return FString::Printf( 
-		TEXT("ForEach %s"),
-		Enum ? *Enum->GetName() : TEXT("UNKNOWN"));
-}
-
 void UK2Node_ForEachElementInEnum::ExpandNode(class FKismetCompilerContext& CompilerContext, UEdGraph* SourceGraph)
 {
 	Super::ExpandNode(CompilerContext, SourceGraph);
diff --git a/Engine/Source/Editor/BlueprintGraph/Private/K2Node_FunctionEntry.cpp b/Engine/Source/Editor/BlueprintGraph/Private/K2Node_FunctionEntry.cpp
--- a/Engine/Source/Editor/BlueprintGraph/Private/K2Node_FunctionEntry.cpp
+++ b/Engine/Source/Editor/BlueprintGraph/Private/K2Node_FunctionEntry.cpp
@@ -122,16 +122,6 @@ FText UK2Node_FunctionEntry::GetNodeTitle(ENodeTitleType::Type TitleType) const
 	return DisplayInfo.DisplayName;
 }
 
-FString UK2Node_FunctionEntry::GetNodeNativeTitle(ENodeTitleType::Type TitleType) const
-{
-	// Do not setup this function for localization, intentionally left unlocalized!
-	UEdGraph* Graph = GetGraph();
-	FGraphDisplayInfo DisplayInfo;
-	Graph->GetSchema()->GetGraphDisplayInformation(*Graph, DisplayInfo);
-
-	return DisplayInfo.DisplayName.ToString();
-}
-
 void UK2Node_FunctionEntry::AllocateDefaultPins()
 {
 	const UEdGraphSchema_K2* K2Schema = GetDefault<UEdGraphSchema_K2>();
diff --git a/Engine/Source/Editor/BlueprintGraph/Private/K2Node_GetDataTableRow.cpp b/Engine/Source/Editor/BlueprintGraph/Private/K2Node_GetDataTableRow.cpp
--- a/Engine/Source/Editor/BlueprintGraph/Private/K2Node_GetDataTableRow.cpp
+++ b/Engine/Source/Editor/BlueprintGraph/Private/K2Node_GetDataTableRow.cpp
@@ -233,29 +233,6 @@ FText UK2Node_GetDataTableRow::GetNodeTitle(ENodeTitleType::Type TitleType) cons
 	return FText::Format(NSLOCTEXT("K2Node", "DataTable", "Get Data Table Row {DataTableName}"), Args);
 }
 
-FString UK2Node_GetDataTableRow::GetNodeNativeTitle(ENodeTitleType::Type TitleType) const
-{
-	// Do not setup this function for localization, intentionally left unlocalized!
-	
-	UEdGraphPin* DataTablePin = GetDataTablePin();
-
-	FString DataTableString = TEXT("NONE");
-	if(DataTablePin != NULL)
-	{
-		if(DataTablePin->LinkedTo.Num() > 0)
-		{
-			// Blueprint will be determined dynamically, so we don't have the name in this case
-			DataTableString = TEXT("");
-		}
-		else if(DataTablePin->DefaultObject != NULL)
-		{
-			DataTableString = DataTablePin->DefaultObject->GetName();
-		}
-	}
-
-	return FString::Printf(TEXT("GetDataTableRow %s"), *DataTableString);
-}
-
 void UK2Node_GetDataTableRow::ExpandNode(class FKismetCompilerContext& CompilerContext, UEdGraph* SourceGraph)
 {
     Super::ExpandNode(CompilerContext, SourceGraph);
diff --git a/Engine/Source/Editor/BlueprintGraph/Private/K2Node_GetInputAxisValue.cpp b/Engine/Source/Editor/BlueprintGraph/Private/K2Node_GetInputAxisValue.cpp
--- a/Engine/Source/Editor/BlueprintGraph/Private/K2Node_GetInputAxisValue.cpp
+++ b/Engine/Source/Editor/BlueprintGraph/Private/K2Node_GetInputAxisValue.cpp
@@ -31,12 +31,6 @@ FText UK2Node_GetInputAxisValue::GetNodeTitle(ENodeTitleType::Type TitleType) co
 	return FText::Format(NSLOCTEXT("K2Node", "GetInputAxis_Name", "Get {InputAxisName}"), Args);
 }
 
-FString UK2Node_GetInputAxisValue::GetNodeNativeTitle(ENodeTitleType::Type TitleType) const
-{
-	// Do not setup this function for localization, intentionally left unlocalized!
-	return FString::Printf(TEXT("Get %s"), *InputAxisName.ToString());
-}
-
 FString UK2Node_GetInputAxisValue::GetTooltip() const
 {
 	return FString::Printf(*NSLOCTEXT("K2Node", "GetInputAxis_Tooltip", "Returns the current value of input axis %s.  If input is disabled for the actor the value will be 0.").ToString(), *InputAxisName.ToString());
diff --git a/Engine/Source/Editor/BlueprintGraph/Private/K2Node_GetNumEnumEntries.cpp b/Engine/Source/Editor/BlueprintGraph/Private/K2Node_GetNumEnumEntries.cpp
--- a/Engine/Source/Editor/BlueprintGraph/Private/K2Node_GetNumEnumEntries.cpp
+++ b/Engine/Source/Editor/BlueprintGraph/Private/K2Node_GetNumEnumEntries.cpp
@@ -37,14 +37,6 @@ FText UK2Node_GetNumEnumEntries::GetNodeTitle(ENodeTitleType::Type TitleType) co
 	return FText::Format(NSLOCTEXT("K2Node", "GetNumEnumEntries_Title", "Get number of entries in {EnumName}"), Args);
 }
 
-FString UK2Node_GetNumEnumEntries::GetNodeNativeTitle(ENodeTitleType::Type TitleType) const
-{
-	// Do not setup this function for localization, intentionally left unlocalized!
-	const FString EnumName = (Enum != NULL) ? Enum->GetName() : TEXT("(bad enum)");
-
-	return FString::Printf(TEXT("Get number of entries in %s"), *EnumName);
-}
-
 void UK2Node_GetNumEnumEntries::ExpandNode(class FKismetCompilerContext& CompilerContext, UEdGraph* SourceGraph)
 {
 	Super::ExpandNode(CompilerContext, SourceGraph);
diff --git a/Engine/Source/Editor/BlueprintGraph/Private/K2Node_InputAction.cpp b/Engine/Source/Editor/BlueprintGraph/Private/K2Node_InputAction.cpp
--- a/Engine/Source/Editor/BlueprintGraph/Private/K2Node_InputAction.cpp
+++ b/Engine/Source/Editor/BlueprintGraph/Private/K2Node_InputAction.cpp
@@ -45,12 +45,6 @@ FText UK2Node_InputAction::GetNodeTitle(ENodeTitleType::Type TitleType) const
 	return FText::Format(NSLOCTEXT("K2Node", "InputAction_Name", "InputAction {InputActionName}"), Args);
 }
 
-FString UK2Node_InputAction::GetNodeNativeTitle(ENodeTitleType::Type TitleType) const
-{
-	// Do not setup this function for localization, intentionally left unlocalized!
-	return FString::Printf(TEXT("InputAction %s"), *InputActionName.ToString());
-}
-
 FString UK2Node_InputAction::GetTooltip() const
 {
 	return FString::Printf(*NSLOCTEXT("K2Node", "InputAction_Tooltip", "Event for when the keys bound to input action %s are pressed or released.").ToString(), *InputActionName.ToString());
diff --git a/Engine/Source/Editor/BlueprintGraph/Private/K2Node_InputAxisEvent.cpp b/Engine/Source/Editor/BlueprintGraph/Private/K2Node_InputAxisEvent.cpp
--- a/Engine/Source/Editor/BlueprintGraph/Private/K2Node_InputAxisEvent.cpp
+++ b/Engine/Source/Editor/BlueprintGraph/Private/K2Node_InputAxisEvent.cpp
@@ -39,12 +39,6 @@ FText UK2Node_InputAxisEvent::GetNodeTitle(ENodeTitleType::Type TitleType) const
 	return FText::Format(NSLOCTEXT("K2Node", "InputAxis_Name", "InputAxis {InputAxisName}"), Args);
 }
 
-FString UK2Node_InputAxisEvent::GetNodeNativeTitle(ENodeTitleType::Type TitleType) const
-{
-	// Do not setup this function for localization, intentionally left unlocalized!
-	return FString::Printf(TEXT("InputAxis %s"), *InputAxisName.ToString());
-}
-
 FString UK2Node_InputAxisEvent::GetTooltip() const
 {
 	return FString::Printf(*NSLOCTEXT("K2Node", "InputAxis_Tooltip", "Event that provides the current value of the %s axis once per frame when input is enabled for the containing actor.").ToString(), *InputAxisName.ToString());
diff --git a/Engine/Source/Editor/BlueprintGraph/Private/K2Node_InputAxisKeyEvent.cpp b/Engine/Source/Editor/BlueprintGraph/Private/K2Node_InputAxisKeyEvent.cpp
--- a/Engine/Source/Editor/BlueprintGraph/Private/K2Node_InputAxisKeyEvent.cpp
+++ b/Engine/Source/Editor/BlueprintGraph/Private/K2Node_InputAxisKeyEvent.cpp
@@ -26,12 +26,6 @@ FText UK2Node_InputAxisKeyEvent::GetNodeTitle(ENodeTitleType::Type TitleType) co
 	return AxisKey.GetDisplayName();
 }
 
-FString UK2Node_InputAxisKeyEvent::GetNodeNativeTitle(ENodeTitleType::Type TitleType) const
-{
-	// Do not setup this function for localization, intentionally left unlocalized!
-	return AxisKey.GetDisplayName().ToString();
-}
-
 FString UK2Node_InputAxisKeyEvent::GetTooltip() const
 {
 	return FString::Printf(*NSLOCTEXT("K2Node", "InputAxisKey_Tooltip", "Event that provides the current value of the %s axis once per frame when input is enabled for the containing actor.").ToString(), *AxisKey.GetDisplayName().ToString());
diff --git a/Engine/Source/Editor/BlueprintGraph/Private/K2Node_InputKey.cpp b/Engine/Source/Editor/BlueprintGraph/Private/K2Node_InputKey.cpp
--- a/Engine/Source/Editor/BlueprintGraph/Private/K2Node_InputKey.cpp
+++ b/Engine/Source/Editor/BlueprintGraph/Private/K2Node_InputKey.cpp
@@ -143,19 +143,6 @@ FText UK2Node_InputKey::GetNodeTitle(ENodeTitleType::Type TitleType) const
 	}
 }
 
-FString UK2Node_InputKey::GetNodeNativeTitle(ENodeTitleType::Type TitleType) const
-{
-	// Do not setup this function for localization, intentionally left unlocalized!
-	if (bControl || bAlt || bShift)
-	{
-		return FString::Printf(TEXT("%s %s"), *GetModifierName().ToString(), *InputKey.ToString());
-	}
-	else
-	{
-		return GetKeyText().ToString();
-	}
-}
-
 FString UK2Node_InputKey::GetTooltip() const
 {
 	FText ModifierText = GetModifierText();
diff --git a/Engine/Source/Editor/BlueprintGraph/Private/K2Node_Knot.cpp b/Engine/Source/Editor/BlueprintGraph/Private/K2Node_Knot.cpp
--- a/Engine/Source/Editor/BlueprintGraph/Private/K2Node_Knot.cpp
+++ b/Engine/Source/Editor/BlueprintGraph/Private/K2Node_Knot.cpp
@@ -59,7 +59,6 @@ void UK2Node_Knot::GetMenuEntries(struct FGraphContextMenuBuilder& ContextMenuBu
 
 		TSharedPtr<FEdGraphSchemaAction_K2NewNode> NodeAction = FK2ActionMenuBuilder::AddNewNodeAction(ContextMenuBuilder, EmptyCategory, MenuDesc, Tooltip, 0, Keywords);
 		NodeAction->NodeTemplate = TemplateNode;
-		NodeAction->SearchTitle = TemplateNode->GetNodeSearchTitle();
 	}
 }
 
diff --git a/Engine/Source/Editor/BlueprintGraph/Private/K2Node_Literal.cpp b/Engine/Source/Editor/BlueprintGraph/Private/K2Node_Literal.cpp
--- a/Engine/Source/Editor/BlueprintGraph/Private/K2Node_Literal.cpp
+++ b/Engine/Source/Editor/BlueprintGraph/Private/K2Node_Literal.cpp
@@ -184,8 +184,8 @@ void UK2Node_Literal::SetObjectRef(UObject* NewValue)
 
 	if( ValuePin )
 	{
-		ValuePin->PinName = GetNodeNativeTitle(ENodeTitleType::ListView);
 		ValuePin->PinFriendlyName = GetNodeTitle(ENodeTitleType::ListView);
+		ValuePin->PinName = ValuePin->PinFriendlyName.BuildSourceString();
 	}
 }
 
diff --git a/Engine/Source/Editor/BlueprintGraph/Private/K2Node_LocalVariable.cpp b/Engine/Source/Editor/BlueprintGraph/Private/K2Node_LocalVariable.cpp
--- a/Engine/Source/Editor/BlueprintGraph/Private/K2Node_LocalVariable.cpp
+++ b/Engine/Source/Editor/BlueprintGraph/Private/K2Node_LocalVariable.cpp
@@ -43,21 +43,6 @@ FText UDEPRECATED_K2Node_LocalVariable::GetNodeTitle(ENodeTitleType::Type TitleT
 	return FText::Format(NSLOCTEXT("K2Node", "LocalVariable_Name", "{Title}\nLocal Variable"), Args);
 }
 
-FString UDEPRECATED_K2Node_LocalVariable::GetNodeNativeTitle(ENodeTitleType::Type TitleType) const
-{
-	// Do not setup this function for localization, intentionally left unlocalized!
-	if (TitleType == ENodeTitleType::EditableTitle)
-	{
-		return CustomVariableName.ToString();
-	}
-	else if(TitleType == ENodeTitleType::ListView)
-	{
-		return FString(TEXT("Local ")) + UEdGraphSchema_K2::TypeToString(VariableType);
-	}
-
-	return FString::Printf(TEXT("%s\nLocal Variable"), *CustomVariableName.ToString());
-}
-
 void UDEPRECATED_K2Node_LocalVariable::OnRenameNode(const FString& NewName)
 {
 	if(CustomVariableName != *NewName)
diff --git a/Engine/Source/Editor/BlueprintGraph/Private/K2Node_MCDelegate.cpp b/Engine/Source/Editor/BlueprintGraph/Private/K2Node_MCDelegate.cpp
--- a/Engine/Source/Editor/BlueprintGraph/Private/K2Node_MCDelegate.cpp
+++ b/Engine/Source/Editor/BlueprintGraph/Private/K2Node_MCDelegate.cpp
@@ -219,13 +219,6 @@ FText UK2Node_AddDelegate::GetNodeTitle(ENodeTitleType::Type TitleType) const
 	return FText::Format(NSLOCTEXT("K2Node", "AddDelegate", "Bind Event to {PropertyName}"), Args);
 }
 
-FString UK2Node_AddDelegate::GetNodeNativeTitle(ENodeTitleType::Type TitleType) const
-{
-	// Do not setup this function for localization, intentionally left unlocalized!
-	
-	return FString::Printf(TEXT("Bind Event to %s"), *GetPropertyName().ToString());
-}
-
 FNodeHandlingFunctor* UK2Node_AddDelegate::CreateNodeHandler(FKismetCompilerContext& CompilerContext) const
 {
 	return new FKCHandler_AddRemoveDelegate(CompilerContext, KCST_AddMulticastDelegate);
@@ -252,13 +245,6 @@ FText UK2Node_ClearDelegate::GetNodeTitle(ENodeTitleType::Type TitleType) const
 	return FText::Format(NSLOCTEXT("K2Node", "ClearDelegate", "Unbind all Events from {PropertyName}"), Args);
 }
 
-FString UK2Node_ClearDelegate::GetNodeNativeTitle(ENodeTitleType::Type TitleType) const
-{
-	// Do not setup this function for localization, intentionally left unlocalized!
-	
-	return FString::Printf(TEXT("Unbind all Events from %s"), *GetPropertyName().ToString());
-}
-
 FNodeHandlingFunctor* UK2Node_ClearDelegate::CreateNodeHandler(FKismetCompilerContext& CompilerContext) const
 {
 	return new FKCHandler_ClearDelegate(CompilerContext);
@@ -291,13 +277,6 @@ FText UK2Node_RemoveDelegate::GetNodeTitle(ENodeTitleType::Type TitleType) const
 	return FText::Format(NSLOCTEXT("K2Node", "RemoveDelegate", "Unbind Event from {PropertyName}"), Args);
 }
 
-FString UK2Node_RemoveDelegate::GetNodeNativeTitle(ENodeTitleType::Type TitleType) const
-{
-	// Do not setup this function for localization, intentionally left unlocalized!
-	
-	return FString::Printf(TEXT("Unbind Event from %s"), *GetPropertyName().ToString());
-}
-
 FNodeHandlingFunctor* UK2Node_RemoveDelegate::CreateNodeHandler(FKismetCompilerContext& CompilerContext) const
 {
 	return new FKCHandler_AddRemoveDelegate(CompilerContext, KCST_RemoveMulticastDelegate);
@@ -346,13 +325,6 @@ FText UK2Node_CallDelegate::GetNodeTitle(ENodeTitleType::Type TitleType) const
 	return FText::Format(NSLOCTEXT("K2Node", "CallDelegate", "Call {PropertyName}"), Args);
 }
 
-FString UK2Node_CallDelegate::GetNodeNativeTitle(ENodeTitleType::Type TitleType) const
-{
-	// Do not setup this function for localization, intentionally left unlocalized!
-	
-	return FString::Printf(TEXT("Call %s"), *GetPropertyName().ToString());
-}
-
 void UK2Node_CallDelegate::ValidateNodeDuringCompilation(class FCompilerResultsLog& MessageLog) const
 {
 	UK2Node::ValidateNodeDuringCompilation(MessageLog);
diff --git a/Engine/Source/Editor/BlueprintGraph/Private/K2Node_MacroInstance.cpp b/Engine/Source/Editor/BlueprintGraph/Private/K2Node_MacroInstance.cpp
--- a/Engine/Source/Editor/BlueprintGraph/Private/K2Node_MacroInstance.cpp
+++ b/Engine/Source/Editor/BlueprintGraph/Private/K2Node_MacroInstance.cpp
@@ -134,20 +134,6 @@ FText UK2Node_MacroInstance::GetNodeTitle(ENodeTitleType::Type TitleType) const
 	return Result;
 }
 
-FString UK2Node_MacroInstance::GetNodeNativeTitle(ENodeTitleType::Type TitleType) const
-{
-	// Do not setup this function for localization, intentionally left unlocalized!
-
-	UEdGraph* MacroGraph = MacroGraphReference.GetGraph();
-	FString Result = TEXT("Macro instance");
-	if (MacroGraph)
-	{
-		Result = MacroGraph->GetName();
-	}
-
-	return Result;
-}
-
 FLinearColor UK2Node_MacroInstance::GetNodeTitleColor() const
 {
 	UEdGraph* MacroGraph = MacroGraphReference.GetGraph();
diff --git a/Engine/Source/Editor/BlueprintGraph/Private/K2Node_MakeStruct.cpp b/Engine/Source/Editor/BlueprintGraph/Private/K2Node_MakeStruct.cpp
--- a/Engine/Source/Editor/BlueprintGraph/Private/K2Node_MakeStruct.cpp
+++ b/Engine/Source/Editor/BlueprintGraph/Private/K2Node_MakeStruct.cpp
@@ -233,13 +233,6 @@ FText UK2Node_MakeStruct::GetNodeTitle(ENodeTitleType::Type TitleType) const
 	return FText::Format(LOCTEXT("MakeNodeTitle", "Make {StructName}"), Args);
 }
 
-FString UK2Node_MakeStruct::GetNodeNativeTitle(ENodeTitleType::Type TitleType) const 
-{
-	// Do not setup this function for localization, intentionally left unlocalized!
-	
-	return FString(TEXT("Make ")) + (StructType ? StructType->GetName() : FString());
-}
-
 FString UK2Node_MakeStruct::GetTooltip() const
 {
 	return FString::Printf(
diff --git a/Engine/Source/Editor/BlueprintGraph/Private/K2Node_MathExpression.cpp b/Engine/Source/Editor/BlueprintGraph/Private/K2Node_MathExpression.cpp
--- a/Engine/Source/Editor/BlueprintGraph/Private/K2Node_MathExpression.cpp
+++ b/Engine/Source/Editor/BlueprintGraph/Private/K2Node_MathExpression.cpp
@@ -2278,7 +2278,6 @@ void UK2Node_MathExpression::GetMenuEntries(FGraphContextMenuBuilder& ContextMen
 
 	TSharedPtr<FEdGraphSchemaAction_K2NewNode> NodeAction = FK2ActionMenuBuilder::AddNewNodeAction(ContextMenuBuilder, Category, MenuDesc, Tooltip);
 	NodeAction->NodeTemplate = TemplateNode;
-	NodeAction->SearchTitle  = TemplateNode->GetNodeSearchTitle();
 }
 
 //------------------------------------------------------------------------------
@@ -2424,24 +2423,6 @@ FText UK2Node_MathExpression::GetNodeTitle(ENodeTitleType::Type TitleType) const
 	return FText::FromString(Result);
 }
 
-//------------------------------------------------------------------------------
-FString UK2Node_MathExpression::GetNodeNativeTitle(ENodeTitleType::Type TitleType) const
-{
-	FString Result = Expression;
-	if (Expression.IsEmpty())
-	{
-		//Result = TEXT("[Empty]");
-	}
-
-	if (TitleType == ENodeTitleType::FullTitle)
-	{
-		Result += TEXT("\n");
-		Result += TEXT("Math Expression");
-	}
-
-	return Result;
-}
-
 //------------------------------------------------------------------------------
 void UK2Node_MathExpression::PostPlacedNewNode()
 {
diff --git a/Engine/Source/Editor/BlueprintGraph/Private/K2Node_Message.cpp b/Engine/Source/Editor/BlueprintGraph/Private/K2Node_Message.cpp
--- a/Engine/Source/Editor/BlueprintGraph/Private/K2Node_Message.cpp
+++ b/Engine/Source/Editor/BlueprintGraph/Private/K2Node_Message.cpp
@@ -31,24 +31,6 @@ FText UK2Node_Message::GetNodeTitle(ENodeTitleType::Type TitleType) const
 	return NodeName;
 }
 
-FString UK2Node_Message::GetNodeNativeTitle(ENodeTitleType::Type TitleType) const
-{
-	// Do not setup this function for localization, intentionally left unlocalized!
-	FString NodeName;
-	if (UFunction* Function = GetTargetFunction())
-	{
-		NodeName = UK2Node_CallFunction::GetUserFacingFunctionName(Function);
-		NodeName += FString(TEXT("\n"));
-		NodeName += FString::Printf(TEXT("Using Interface %s"), *(Function->GetOuterUClass()->GetName()));
-	}
-	else
-	{
-		NodeName = TEXT("Invalid Message Node");
-	}
-
-	return NodeName;
-}
-
 UEdGraphPin* UK2Node_Message::CreateSelfPin(const UFunction* Function)
 {
 	const UEdGraphSchema_K2* K2Schema = GetDefault<UEdGraphSchema_K2>();
diff --git a/Engine/Source/Editor/BlueprintGraph/Private/K2Node_SelectEnum.cpp b/Engine/Source/Editor/BlueprintGraph/Private/K2Node_SelectEnum.cpp
--- a/Engine/Source/Editor/BlueprintGraph/Private/K2Node_SelectEnum.cpp
+++ b/Engine/Source/Editor/BlueprintGraph/Private/K2Node_SelectEnum.cpp
@@ -53,14 +53,6 @@ FText  UDEPRECATED_K2Node_SelectEnum::GetNodeTitle(ENodeTitleType::Type TitleTyp
 	return FText::Format(NSLOCTEXT("K2Node", "SelectByEnum", "Select by {EnumName}"), Args);
 }
 
-FString UDEPRECATED_K2Node_SelectEnum::GetNodeNativeTitle(ENodeTitleType::Type TitleType) const
-{
-	// Do not setup this function for localization, intentionally left unlocalized!
-
-	FString EnumName = (Enum != NULL) ? Enum->GetName() : TEXT("(bad enum)");
-	return FString::Printf(TEXT("Select by %s"), *EnumName);
-}
-
 /** Determine if any pins are connected, if so make all the other pins the same type, if not, make sure pins are switched back to wildcards */
 void UDEPRECATED_K2Node_SelectEnum::NotifyPinConnectionListChanged(UEdGraphPin* Pin)
 {
diff --git a/Engine/Source/Editor/BlueprintGraph/Private/K2Node_SpawnActor.cpp b/Engine/Source/Editor/BlueprintGraph/Private/K2Node_SpawnActor.cpp
--- a/Engine/Source/Editor/BlueprintGraph/Private/K2Node_SpawnActor.cpp
+++ b/Engine/Source/Editor/BlueprintGraph/Private/K2Node_SpawnActor.cpp
@@ -249,28 +249,6 @@ FText UK2Node_SpawnActor::GetNodeTitle(ENodeTitleType::Type TitleType) const
 	return FText::Format(NSLOCTEXT("K2Node", "SpawnActor", "SpawnActor {ActorName}"), Args);
 }
 
-FString UK2Node_SpawnActor::GetNodeNativeTitle(ENodeTitleType::Type TitleType) const
-{
-	// Do not setup this function for localization, intentionally left unlocalized!
-	UEdGraphPin* BlueprintPin = GetBlueprintPin();
-
-	FString SpawnString = TEXT("NONE");
-	if(BlueprintPin != NULL)
-	{
-		if(BlueprintPin->LinkedTo.Num() > 0)
-		{
-			// Blueprint will be determined dynamically, so we don't have the name in this case
-			SpawnString = TEXT("");
-		}
-		else if(BlueprintPin->DefaultObject != NULL)
-		{
-			SpawnString = BlueprintPin->DefaultObject->GetName();
-		}
-	}
-
-	return FString::Printf(TEXT("SpawnActor %s"), *SpawnString);
-}
-
 bool UK2Node_SpawnActor::CanPasteHere( const UEdGraph* TargetGraph, const UEdGraphSchema* Schema ) const 
 {
 	UBlueprint* Blueprint = Cast<UBlueprint>(TargetGraph->GetOuter());
diff --git a/Engine/Source/Editor/BlueprintGraph/Private/K2Node_SpawnActorFromClass.cpp b/Engine/Source/Editor/BlueprintGraph/Private/K2Node_SpawnActorFromClass.cpp
--- a/Engine/Source/Editor/BlueprintGraph/Private/K2Node_SpawnActorFromClass.cpp
+++ b/Engine/Source/Editor/BlueprintGraph/Private/K2Node_SpawnActorFromClass.cpp
@@ -273,29 +273,6 @@ FText UK2Node_SpawnActorFromClass::GetNodeTitle(ENodeTitleType::Type TitleType)
 	return FText::Format(NSLOCTEXT("K2Node", "SpawnActor", "SpawnActor {ClassName}"), Args);
 }
 
-FString UK2Node_SpawnActorFromClass::GetNodeNativeTitle(ENodeTitleType::Type TitleType) const
-{
-	// Do not setup this function for localization, intentionally left unlocalized!
-	
-	UEdGraphPin* ClassPin = GetClassPin();
-
-	FString SpawnString = TEXT("NONE");
-	if(ClassPin != NULL)
-	{
-		if(ClassPin->LinkedTo.Num() > 0)
-		{
-			// Blueprint will be determined dynamically, so we don't have the name in this case
-			SpawnString = TEXT("");
-		}
-		else if(ClassPin->DefaultObject != NULL)
-		{
-			SpawnString = ClassPin->DefaultObject->GetName();
-		}
-	}
-
-	return FString::Printf(TEXT("SpawnActor %s"), *SpawnString);
-}
-
 bool UK2Node_SpawnActorFromClass::CanPasteHere( const UEdGraph* TargetGraph, const UEdGraphSchema* Schema ) const 
 {
 	UBlueprint* Blueprint = Cast<UBlueprint>(TargetGraph->GetOuter());
diff --git a/Engine/Source/Editor/BlueprintGraph/Private/K2Node_StructMemberGet.cpp b/Engine/Source/Editor/BlueprintGraph/Private/K2Node_StructMemberGet.cpp
--- a/Engine/Source/Editor/BlueprintGraph/Private/K2Node_StructMemberGet.cpp
+++ b/Engine/Source/Editor/BlueprintGraph/Private/K2Node_StructMemberGet.cpp
@@ -83,13 +83,6 @@ FText UK2Node_StructMemberGet::GetNodeTitle(ENodeTitleType::Type TitleType) cons
 	return FText::Format(LOCTEXT("GetMembersInVariable", "Get members in {VariableName}"), Args);
 }
 
-FString UK2Node_StructMemberGet::GetNodeNativeTitle(ENodeTitleType::Type TitleType) const
-{
-	// Do not setup this function for localization, intentionally left unlocalized!
-
-	return FString::Printf(TEXT("Get members in %s"), *GetVarNameString());
-}
-
 FNodeHandlingFunctor* UK2Node_StructMemberGet::CreateNodeHandler(FKismetCompilerContext& CompilerContext) const
 {
 	return new FKCHandler_StructMemberVariableGet(CompilerContext);
diff --git a/Engine/Source/Editor/BlueprintGraph/Private/K2Node_StructMemberSet.cpp b/Engine/Source/Editor/BlueprintGraph/Private/K2Node_StructMemberSet.cpp
--- a/Engine/Source/Editor/BlueprintGraph/Private/K2Node_StructMemberSet.cpp
+++ b/Engine/Source/Editor/BlueprintGraph/Private/K2Node_StructMemberSet.cpp
@@ -56,13 +56,6 @@ FText UK2Node_StructMemberSet::GetNodeTitle(ENodeTitleType::Type TitleType) cons
 	return FText::Format(LOCTEXT("SetMembersInVariable", "Set members in {VariableName}"), Args);
 }
 
-FString UK2Node_StructMemberSet::GetNodeNativeTitle(ENodeTitleType::Type TitleType) const
-{
-	// Do not setup this function for localization, intentionally left unlocalized!
-
-	return FString::Printf(TEXT("Set members in %s"), *VariableReference.GetMemberName().ToString());
-}
-
 UK2Node::ERedirectType UK2Node_StructMemberSet::DoPinsMatchForReconstruction(const UEdGraphPin* NewPin, int32 NewPinIndex, const UEdGraphPin* OldPin, int32 OldPinIndex)  const
 {
 	return UK2Node::DoPinsMatchForReconstruction(NewPin, NewPinIndex, OldPin, OldPinIndex);
diff --git a/Engine/Source/Editor/BlueprintGraph/Private/K2Node_SwitchEnum.cpp b/Engine/Source/Editor/BlueprintGraph/Private/K2Node_SwitchEnum.cpp
--- a/Engine/Source/Editor/BlueprintGraph/Private/K2Node_SwitchEnum.cpp
+++ b/Engine/Source/Editor/BlueprintGraph/Private/K2Node_SwitchEnum.cpp
@@ -56,15 +56,6 @@ FText UK2Node_SwitchEnum::GetNodeTitle(ENodeTitleType::Type TitleType) const
 	return FText::Format(NSLOCTEXT("K2Node", "Switch_Enum", "Switch on {EnumName}"), Args);
 }
 
-FString UK2Node_SwitchEnum::GetNodeNativeTitle(ENodeTitleType::Type TitleType) const 
-{
-	// Do not setup this function for localization, intentionally left unlocalized!
-
-	FString EnumName = (Enum != NULL) ? Enum->GetName() : TEXT("(bad enum)");
-
-	return FString::Printf(TEXT("Switch on %s"), *EnumName);
-}
-
 FString UK2Node_SwitchEnum::GetTooltip() const
 {
 	return NSLOCTEXT("K2Node", "SwitchEnum_ToolTip", "Selects an output that matches the input value").ToString();
diff --git a/Engine/Source/Editor/BlueprintGraph/Private/K2Node_TemporaryVariable.cpp b/Engine/Source/Editor/BlueprintGraph/Private/K2Node_TemporaryVariable.cpp
--- a/Engine/Source/Editor/BlueprintGraph/Private/K2Node_TemporaryVariable.cpp
+++ b/Engine/Source/Editor/BlueprintGraph/Private/K2Node_TemporaryVariable.cpp
@@ -61,15 +61,6 @@ FText UK2Node_TemporaryVariable::GetNodeTitle(ENodeTitleType::Type TitleType) co
 	return Result;
 }
 
-FString UK2Node_TemporaryVariable::GetNodeNativeTitle(ENodeTitleType::Type TitleType) const
-{
-	// Do not setup this function for localization, intentionally left unlocalized!
-
-	FString Result = !bIsPersistent ? TEXT("Local ") : TEXT("Persistent Local ");
-	Result += UEdGraphSchema_K2::TypeToString(VariableType);
-	return Result;
-}
-
 bool UK2Node_TemporaryVariable::IsNodePure() const
 {
 	return true;
diff --git a/Engine/Source/Editor/BlueprintGraph/Private/K2Node_VariableGet.cpp b/Engine/Source/Editor/BlueprintGraph/Private/K2Node_VariableGet.cpp
--- a/Engine/Source/Editor/BlueprintGraph/Private/K2Node_VariableGet.cpp
+++ b/Engine/Source/Editor/BlueprintGraph/Private/K2Node_VariableGet.cpp
@@ -158,33 +158,6 @@ FText UK2Node_VariableGet::GetNodeTitle(ENodeTitleType::Type TitleType) const
 	return Result;
 }
 
-FString UK2Node_VariableGet::GetNodeNativeTitle(ENodeTitleType::Type TitleType) const
-{
-	// Do not setup this function for localization, intentionally left unlocalized!
-	FString Result = TEXT("Get");
-
-	// If there is only one variable being read, the title can be made the variable name
-	FString OutputPinName;
-	int32 NumOutputsFound = 0;
-
-	for (int32 PinIndex = 0; PinIndex < Pins.Num(); ++PinIndex)
-	{
-		UEdGraphPin* Pin = Pins[PinIndex];
-		if (Pin->Direction == EGPD_Output)
-		{
-			++NumOutputsFound;
-			OutputPinName = Pin->PinName;
-		}
-	}
-
-	if (NumOutputsFound == 1)
-	{
-		Result = Result + TEXT(" ") + OutputPinName;
-	}
-
-	return Result;
-}
-
 FNodeHandlingFunctor* UK2Node_VariableGet::CreateNodeHandler(FKismetCompilerContext& CompilerContext) const
 {
 	return new FKCHandler_VariableGet(CompilerContext);
diff --git a/Engine/Source/Editor/BlueprintGraph/Private/K2Node_VariableSet.cpp b/Engine/Source/Editor/BlueprintGraph/Private/K2Node_VariableSet.cpp
--- a/Engine/Source/Editor/BlueprintGraph/Private/K2Node_VariableSet.cpp
+++ b/Engine/Source/Editor/BlueprintGraph/Private/K2Node_VariableSet.cpp
@@ -126,37 +126,6 @@ FText UK2Node_VariableSet::GetNodeTitle(ENodeTitleType::Type TitleType) const
 	return Result;
 }
 
-FString UK2Node_VariableSet::GetNodeNativeTitle(ENodeTitleType::Type TitleType) const
-{
-	// Do not setup this function for localization, intentionally left unlocalized!
-
-	const UEdGraphSchema_K2* K2Schema = GetDefault<UEdGraphSchema_K2>();
-
-	FString Result = HasLocalRepNotify() ? TEXT("Set with Notify") : TEXT("Set");
-
-	// If there is only one variable being written (one non-meta input pin), the title can be made the variable name
-	FString InputPinName;
-	int32 NumInputsFound = 0;
-
-	for (int32 PinIndex = 0; PinIndex < Pins.Num(); ++PinIndex)
-	{
-		UEdGraphPin* Pin = Pins[PinIndex];
-		if ((Pin->Direction == EGPD_Input) && (!K2Schema->IsMetaPin(*Pin)))
-		{
-			++NumInputsFound;
-			InputPinName = Pin->PinName;
-		}
-	}
-
-	if (NumInputsFound == 1)
-	{
-		Result = Result + TEXT(" ") + InputPinName;
-	}
-
-	return Result;
-}
-
-
 /** Returns true if the variable we are setting has a RepNotify AND was defined in a blueprint
  *		The 'defined in a blueprint' is to avoid natively defined RepNotifies being called unintentionally.
  *		Most (all?) native rep notifies are intended to be client only. We are moving away from this paradigm in blueprints
diff --git a/Engine/Source/Editor/BlueprintGraph/Private/K2Node_VariableSetRef.cpp b/Engine/Source/Editor/BlueprintGraph/Private/K2Node_VariableSetRef.cpp
--- a/Engine/Source/Editor/BlueprintGraph/Private/K2Node_VariableSetRef.cpp
+++ b/Engine/Source/Editor/BlueprintGraph/Private/K2Node_VariableSetRef.cpp
@@ -141,24 +141,6 @@ FText UK2Node_VariableSetRef::GetNodeTitle(ENodeTitleType::Type TitleType) const
 	}
 }
 
-FString UK2Node_VariableSetRef::GetNodeNativeTitle(ENodeTitleType::Type TitleType) const
-{
-	// Do not setup this function for localization, intentionally left unlocalized!
-	
-	const UEdGraphSchema_K2* Schema = GetDefault<UEdGraphSchema_K2>();
-
-	UEdGraphPin* TargetPin = GetTargetPin();
-
-	if( TargetPin && TargetPin->PinType.PinCategory != Schema->PC_Wildcard )
-	{
-		return FString::Printf(TEXT("Set %s"), *Schema->TypeToString(TargetPin->PinType));
-	}
-	else
-	{
-		return TEXT("Set By-Ref Var");
-	}
-}
-
 void UK2Node_VariableSetRef::NotifyPinConnectionListChanged(UEdGraphPin* Pin)
 {
 	Super::NotifyPinConnectionListChanged(Pin);
diff --git a/Engine/Source/Editor/EnvironmentQueryEditor/Private/EdGraphSchema_EnvironmentQuery.cpp b/Engine/Source/Editor/EnvironmentQueryEditor/Private/EdGraphSchema_EnvironmentQuery.cpp
--- a/Engine/Source/Editor/EnvironmentQueryEditor/Private/EdGraphSchema_EnvironmentQuery.cpp
+++ b/Engine/Source/Editor/EnvironmentQueryEditor/Private/EdGraphSchema_EnvironmentQuery.cpp
@@ -229,7 +229,6 @@ void UEdGraphSchema_EnvironmentQuery::GetGraphContextActions(FGraphContextMenuBu
 		UEnvironmentQueryGraphNode_Option* OpNode = NewObject<UEnvironmentQueryGraphNode_Option>(ContextMenuBuilder.OwnerOfTemporaries);
 		OpNode->EnvQueryNodeClass = AllowedClasses[i];
 		AddOpAction->NodeTemplate = OpNode;		
-		AddOpAction->SearchTitle = AddOpAction->NodeTemplate->GetNodeSearchTitle();
 	}
 }
 
@@ -249,7 +248,6 @@ void UEdGraphSchema_EnvironmentQuery::GetGraphNodeContextActions(FGraphContextMe
 		TSharedPtr<FEnvironmentQuerySchemaAction_NewSubNode> AddOpAction = AddNewSubNodeAction(ContextMenuBuilder, TEXT(""), FText::FromString(GetNodeDescriptionHelper(AllowedClasses[i])), "");
 		AddOpAction->ParentNode = SelectedOption;
 		AddOpAction->NodeTemplate = OpNode;
-		AddOpAction->SearchTitle = AddOpAction->NodeTemplate->GetNodeSearchTitle();
 	}
 }
 
diff --git a/Engine/Source/Editor/GameplayAbilitiesEditor/Classes/K2Node_GameplayCueEvent.h b/Engine/Source/Editor/GameplayAbilitiesEditor/Classes/K2Node_GameplayCueEvent.h
--- a/Engine/Source/Editor/GameplayAbilitiesEditor/Classes/K2Node_GameplayCueEvent.h
+++ b/Engine/Source/Editor/GameplayAbilitiesEditor/Classes/K2Node_GameplayCueEvent.h
@@ -13,7 +13,6 @@ class UK2Node_GameplayCueEvent : public UK2Node_Event
 	// Begin UEdGraphNode interface
 	virtual FString GetTooltip() const override;
 	virtual FText GetNodeTitle(ENodeTitleType::Type TitleType) const override;
-	virtual FString GetNodeNativeTitle(ENodeTitleType::Type TitleType) const override;
 	// End UEdGraphNode interface
 
 	virtual FString GetCategoryName();
diff --git a/Engine/Source/Editor/GameplayAbilitiesEditor/Private/K2Node_GameplayCueEvent.cpp b/Engine/Source/Editor/GameplayAbilitiesEditor/Private/K2Node_GameplayCueEvent.cpp
--- a/Engine/Source/Editor/GameplayAbilitiesEditor/Private/K2Node_GameplayCueEvent.cpp
+++ b/Engine/Source/Editor/GameplayAbilitiesEditor/Private/K2Node_GameplayCueEvent.cpp
@@ -33,12 +33,6 @@ FText UK2Node_GameplayCueEvent::GetNodeTitle(ENodeTitleType::Type TitleType) con
 	//return LOCTEXT("HandleGameplayCueEvent", "HandleGameplaCueEvent");
 }
 
-FString UK2Node_GameplayCueEvent::GetNodeNativeTitle(ENodeTitleType::Type TitleType) const
-{
-	return CustomFunctionName.ToString();
-	//return LOCTEXT("HandleGameplayCueEvent", "HandleGameplaCueEvent").ToString();
-}
-
 void UK2Node_GameplayCueEvent::GetMenuEntries(FGraphContextMenuBuilder& Context) const
 {
 	Super::GetMenuEntries(Context);
@@ -65,7 +59,6 @@ void UK2Node_GameplayCueEvent::GetMenuEntries(FGraphContextMenuBuilder& Context)
 
 		TSharedPtr<FEdGraphSchemaAction_K2NewNode> NodeAction = FK2ActionMenuBuilder::AddNewNodeAction(Context, Category, MenuDesc, Tooltip, 0, Keywords);
 		NodeAction->NodeTemplate = NodeTemplate;
-		NodeAction->SearchTitle = NodeAction->NodeTemplate->GetNodeSearchTitle();
 	}	
 }
 
diff --git a/Engine/Source/Editor/GameplayTagsEditor/Private/GameplayTagsK2Node_LiteralGameplayTag.cpp b/Engine/Source/Editor/GameplayTagsEditor/Private/GameplayTagsK2Node_LiteralGameplayTag.cpp
--- a/Engine/Source/Editor/GameplayTagsEditor/Private/GameplayTagsK2Node_LiteralGameplayTag.cpp
+++ b/Engine/Source/Editor/GameplayTagsEditor/Private/GameplayTagsK2Node_LiteralGameplayTag.cpp
@@ -49,7 +49,6 @@ void UGameplayTagsK2Node_LiteralGameplayTag::GetMenuEntries( FGraphContextMenuBu
 
 	TSharedPtr<FEdGraphSchemaAction_K2NewNode> NodeAction = FK2ActionMenuBuilder::AddNewNodeAction(ContextMenuBuilder, Category, MenuDesc, Tooltip, 0, Keywords);
 	NodeAction->NodeTemplate = EnumNodeTemplate;
-	NodeAction->SearchTitle = EnumNodeTemplate->GetNodeSearchTitle();
 }
 
 void UGameplayTagsK2Node_LiteralGameplayTag::ExpandNode(class FKismetCompilerContext& CompilerContext, UEdGraph* SourceGraph)
diff --git a/Engine/Source/Editor/GraphEditor/Private/SGraphNode.cpp b/Engine/Source/Editor/GraphEditor/Private/SGraphNode.cpp
--- a/Engine/Source/Editor/GraphEditor/Private/SGraphNode.cpp
+++ b/Engine/Source/Editor/GraphEditor/Private/SGraphNode.cpp
@@ -482,7 +482,7 @@ FText SGraphNode::GetNodeTooltip() const
 		// Display the native title of the node when alt is held
 		if(FSlateApplication::Get().GetModifierKeys().IsAltDown())
 		{
-			return FText::FromString(GraphNode->GetNodeNativeTitle(ENodeTitleType::ListView));
+			return FText::FromString(GraphNode->GetNodeTitle(ENodeTitleType::ListView).BuildSourceString());
 		}
 
 		FText TooltipText = FText::FromString(GraphNode->GetTooltip());
diff --git a/Engine/Source/Editor/Kismet/Private/BlueprintEditorCommands.cpp b/Engine/Source/Editor/Kismet/Private/BlueprintEditorCommands.cpp
--- a/Engine/Source/Editor/Kismet/Private/BlueprintEditorCommands.cpp
+++ b/Engine/Source/Editor/Kismet/Private/BlueprintEditorCommands.cpp
@@ -174,7 +174,6 @@ public:
 
 		// the NodeTemplate UPROPERTY takes ownership of GraphNode's lifetime (hence it being a weak-pointer here)
 		NewActionNode->NodeTemplate = GraphNode.Get();
-		NewActionNode->SearchTitle = NewActionNode->NodeTemplate->GetNodeSearchTitle();
 
 		return NewActionNode;
 	}
diff --git a/Engine/Source/Editor/Kismet/Private/SBlueprintPalette.cpp b/Engine/Source/Editor/Kismet/Private/SBlueprintPalette.cpp
--- a/Engine/Source/Editor/Kismet/Private/SBlueprintPalette.cpp
+++ b/Engine/Source/Editor/Kismet/Private/SBlueprintPalette.cpp
@@ -1270,7 +1270,7 @@ FText SBlueprintPaletteItem::GetToolTipText() const
 			// Display the native title of the node when alt is held
 			if(FSlateApplication::Get().GetModifierKeys().IsAltDown())
 			{
-				return FText::FromString(NodeTemplate->GetNodeNativeTitle(ENodeTitleType::ListView));
+				return FText::FromString(NodeTemplate->GetNodeTitle(ENodeTitleType::ListView).BuildSourceString());
 			}
 
 			// If the node wants to create tooltip text, use that instead, because its probably more detailed
diff --git a/Engine/Source/Editor/Kismet/Private/SMyBlueprint.cpp b/Engine/Source/Editor/Kismet/Private/SMyBlueprint.cpp
--- a/Engine/Source/Editor/Kismet/Private/SMyBlueprint.cpp
+++ b/Engine/Source/Editor/Kismet/Private/SMyBlueprint.cpp
@@ -539,7 +539,6 @@ void SMyBlueprint::GetChildEvents(UEdGraph const* EdGraph, int32 const SectionId
 
 		TSharedPtr<FEdGraphSchemaAction_K2Event> EventNodeAction = MakeShareable(new FEdGraphSchemaAction_K2Event(ActionCategory, Description, Tooltip, 0));
 		EventNodeAction->NodeTemplate = EventNode;
-		EventNodeAction->SearchTitle = EventNodeAction->NodeTemplate->GetNodeSearchTitle();
 		EventNodeAction->SectionID = SectionId;
 		OutAllActions.AddAction(EventNodeAction);
 	}
diff --git a/Engine/Source/Editor/UMGEditor/Private/WidgetGraphNode_Base.cpp b/Engine/Source/Editor/UMGEditor/Private/WidgetGraphNode_Base.cpp
--- a/Engine/Source/Editor/UMGEditor/Private/WidgetGraphNode_Base.cpp
+++ b/Engine/Source/Editor/UMGEditor/Private/WidgetGraphNode_Base.cpp
@@ -184,7 +184,6 @@ TSharedPtr<FEdGraphSchemaAction_K2NewNode> UWidgetGraphNode_Base::CreateDefaultM
 
 	TSharedPtr<FEdGraphSchemaAction_K2NewNode> NodeAction = FK2ActionMenuBuilder::AddNewNodeAction(ContextMenuBuilder, Category, MenuDesc, Tooltip, 0, Keywords);
 	NodeAction->NodeTemplate = TemplateNode;
-	NodeAction->SearchTitle = TemplateNode->GetNodeSearchTitle();
 
 	return NodeAction;
 }
diff --git a/Engine/Source/Editor/UnrealEd/Classes/MaterialGraph/MaterialGraphNode.h b/Engine/Source/Editor/UnrealEd/Classes/MaterialGraph/MaterialGraphNode.h
--- a/Engine/Source/Editor/UnrealEd/Classes/MaterialGraph/MaterialGraphNode.h
+++ b/Engine/Source/Editor/UnrealEd/Classes/MaterialGraph/MaterialGraphNode.h
@@ -52,7 +52,6 @@ public:
 	// Begin UEdGraphNode interface.
 	virtual bool CanPasteHere(const UEdGraph* TargetGraph, const UEdGraphSchema* Schema) const override;
 	virtual FText GetNodeTitle(ENodeTitleType::Type TitleType) const override;
-	virtual FString GetNodeNativeTitle(ENodeTitleType::Type TitleType) const override;
 	virtual FLinearColor GetNodeTitleColor() const override;
 	virtual FString GetTooltip() const override;
 	virtual void PrepareForCopying() override;
diff --git a/Engine/Source/Editor/UnrealEd/Classes/Niagara/NiagaraNodeGetAttr.h b/Engine/Source/Editor/UnrealEd/Classes/Niagara/NiagaraNodeGetAttr.h
--- a/Engine/Source/Editor/UnrealEd/Classes/Niagara/NiagaraNodeGetAttr.h
+++ b/Engine/Source/Editor/UnrealEd/Classes/Niagara/NiagaraNodeGetAttr.h
@@ -15,7 +15,6 @@ class UNiagaraNodeGetAttr : public UNiagaraNode
 	// Begin EdGraphNode interface
 	virtual void AllocateDefaultPins() override;
 	virtual FText GetNodeTitle(ENodeTitleType::Type TitleType) const override;
-	virtual FString GetNodeNativeTitle(ENodeTitleType::Type TitleType) const override;
 	virtual FLinearColor GetNodeTitleColor() const override;
 	// End EdGraphNode interface
 };
diff --git a/Engine/Source/Editor/UnrealEd/Private/MaterialGraphNode.cpp b/Engine/Source/Editor/UnrealEd/Private/MaterialGraphNode.cpp
--- a/Engine/Source/Editor/UnrealEd/Private/MaterialGraphNode.cpp
+++ b/Engine/Source/Editor/UnrealEd/Private/MaterialGraphNode.cpp
@@ -186,49 +186,6 @@ FText UMaterialGraphNode::GetNodeTitle(ENodeTitleType::Type TitleType) const
 	}
 }
 
-FString UMaterialGraphNode::GetNodeNativeTitle(ENodeTitleType::Type TitleType) const
-{
-	// Do not setup this function for localization, intentionally left unlocalized!
-	
-	TArray<FString> Captions;
-	MaterialExpression->GetCaption(Captions);
-
-	if (TitleType == ENodeTitleType::EditableTitle)
-	{
-		return GetParameterName();
-	}
-	else if (TitleType == ENodeTitleType::ListView)
-	{
-		return MaterialExpression->GetDescription();
-	}
-	else
-	{
-		// More useful to display multi line parameter captions in reverse order
-		// TODO: May have to choose order based on expression type if others need correct order
-		int32 CaptionIndex = Captions.Num() -1;
-
-		FString NodeTitle = Captions[CaptionIndex];
-		for (; CaptionIndex > 0; )
-		{
-			CaptionIndex--;
-			NodeTitle += TEXT("\n");
-			NodeTitle += Captions[CaptionIndex];
-		}
-
-		if ( MaterialExpression->bShaderInputData && (MaterialExpression->bHidePreviewWindow || MaterialExpression->bCollapsed))
-		{
-			NodeTitle += TEXT("\nInput Data");
-		}
-
-		if (bIsPreviewExpression)
-		{
-			NodeTitle += TEXT("\nPreviewing");
-		}
-
-		return NodeTitle;
-	}
-}
-
 FLinearColor UMaterialGraphNode::GetNodeTitleColor() const
 {
 	UMaterial* Material = CastChecked<UMaterialGraph>(GetGraph())->Material;
diff --git a/Engine/Source/Editor/UnrealEd/Private/Niagara/EdGraphSchema_Niagara.cpp b/Engine/Source/Editor/UnrealEd/Private/Niagara/EdGraphSchema_Niagara.cpp
--- a/Engine/Source/Editor/UnrealEd/Private/Niagara/EdGraphSchema_Niagara.cpp
+++ b/Engine/Source/Editor/UnrealEd/Private/Niagara/EdGraphSchema_Niagara.cpp
@@ -125,7 +125,6 @@ void UEdGraphSchema_Niagara::GetGraphContextActions(FGraphContextMenuBuilder& Co
 		UNiagaraNodeGetAttr* GetAttrNode = NewObject<UNiagaraNodeGetAttr>(ContextMenuBuilder.OwnerOfTemporaries);
 		GetAttrNode->AttrName = AttrName;
 		GetAttrAction->NodeTemplate = GetAttrNode;
-		GetAttrAction->SearchTitle = GetAttrAction->NodeTemplate->GetNodeSearchTitle();
 	}
 
 	// Then get ops.
@@ -140,7 +139,6 @@ void UEdGraphSchema_Niagara::GetGraphContextActions(FGraphContextMenuBuilder& Co
 			UNiagaraNodeOp* OpNode = NewObject<UNiagaraNodeOp>(ContextMenuBuilder.OwnerOfTemporaries);
 			OpNode->OpIndex = OpIdx;
 			AddOpAction->NodeTemplate = OpNode;
-			AddOpAction->SearchTitle = AddOpAction->NodeTemplate->GetNodeSearchTitle();
 		}
 	}
 }
diff --git a/Engine/Source/Editor/UnrealEd/Private/Niagara/NiagaraNodeGetAttr.cpp b/Engine/Source/Editor/UnrealEd/Private/Niagara/NiagaraNodeGetAttr.cpp
--- a/Engine/Source/Editor/UnrealEd/Private/Niagara/NiagaraNodeGetAttr.cpp
+++ b/Engine/Source/Editor/UnrealEd/Private/Niagara/NiagaraNodeGetAttr.cpp
@@ -21,13 +21,6 @@ FText UNiagaraNodeGetAttr::GetNodeTitle(ENodeTitleType::Type TitleType) const
 	return FText::Format(NSLOCTEXT("Niagara", "GetAttribute", "Get {Attribute}"), Args);
 }
 
-FString UNiagaraNodeGetAttr::GetNodeNativeTitle(ENodeTitleType::Type TitleType) const
-{
-	// Do not setup this function for localization, intentionally left unlocalized!
-	
-	return FString::Printf(TEXT("Get %s"), *AttrName.ToString());
-}
-
 FLinearColor UNiagaraNodeGetAttr::GetNodeTitleColor() const
 {
 	const UEdGraphSchema_Niagara* Schema = GetDefault<UEdGraphSchema_Niagara>();
diff --git a/Engine/Source/Runtime/Engine/Classes/EdGraph/EdGraphNode.h b/Engine/Source/Runtime/Engine/Classes/EdGraph/EdGraphNode.h
--- a/Engine/Source/Runtime/Engine/Classes/EdGraph/EdGraphNode.h
+++ b/Engine/Source/Runtime/Engine/Classes/EdGraph/EdGraphNode.h
@@ -250,13 +250,6 @@ public:
 	/** Gets the name of this node, shown in title bar */
 	virtual FText GetNodeTitle(ENodeTitleType::Type TitleType) const;
 	
-	/** Gets the name of this node in the native language. This should always be overridden when a node's title is built from concatenation
-		By default this will return the source string from GetNodeTitle, which is not guaranteed to be correct */
-	virtual FString GetNodeNativeTitle(ENodeTitleType::Type TitleType) const;
-
-	/** Gets the searchable metadata of this node */
-	virtual FText GetNodeSearchTitle() const;
-
 	/** 
 	 * Gets the draw color of a node's title bar
 	 */
diff --git a/Engine/Source/Runtime/Engine/Private/EdGraph/EdGraphNode.cpp b/Engine/Source/Runtime/Engine/Private/EdGraph/EdGraphNode.cpp
--- a/Engine/Source/Runtime/Engine/Private/EdGraph/EdGraphNode.cpp
+++ b/Engine/Source/Runtime/Engine/Private/EdGraph/EdGraphNode.cpp
@@ -291,34 +291,6 @@ FText UEdGraphNode::GetNodeTitle(ENodeTitleType::Type TitleType) const
 	return FText::FromString(GetClass()->GetName());
 }
 
-FString UEdGraphNode::GetNodeNativeTitle(ENodeTitleType::Type TitleType) const
-{
-	FText NodeTitle = GetNodeTitle(TitleType);
-	if(const FString* SourceString = FTextInspector::GetSourceString(NodeTitle))
-	{
-		return *SourceString;
-	}
-	return NodeTitle.ToString();
-}
-
-FText UEdGraphNode::GetNodeSearchTitle() const
-{
-	TArray< FText > SearchMetadata;
-
-	FText LocalizedTitle = GetNodeTitle(ENodeTitleType::ListView);
-	FText NativeTitle = FText::FromString(GetNodeNativeTitle(ENodeTitleType::ListView));
-
-	if(FInternationalization::Get().GetCurrentCulture()->GetName() == TEXT("en_US"))
-	{
-		ensure( LocalizedTitle.CompareTo(NativeTitle) == 0 );
-	}
-
-	FFormatNamedArguments Args;
-	Args.Add(TEXT("LocalizedTitle"), LocalizedTitle);
-	Args.Add(TEXT("NativeTitle"), NativeTitle);
-	return FText::Format(LOCTEXT("NodeSearchTitle", "{LocalizedTitle} {NativeTitle}"), Args);
-}
-
 UObject* UEdGraphNode::GetJumpTargetForDoubleClick() const
 {
 	return NULL;
