#ue4 Enabling literal entry on Expose on Spawn pins for SpawnFromClass nodes in BPs

diff --git a/Engine/Source/Editor/BlueprintGraph/Private/K2Node_SpawnActor.cpp b/Engine/Source/Editor/BlueprintGraph/Private/K2Node_SpawnActor.cpp
--- a/Engine/Source/Editor/BlueprintGraph/Private/K2Node_SpawnActor.cpp
+++ b/Engine/Source/Editor/BlueprintGraph/Private/K2Node_SpawnActor.cpp
@@ -73,7 +73,6 @@ void UK2Node_SpawnActor::CreatePinsForClass(UClass* InClass)
 		{
 			UEdGraphPin* Pin = CreatePin(EGPD_Input, TEXT(""), TEXT(""), NULL, false, false, Property->GetName());
 			const bool bPinGood = (Pin != NULL) && K2Schema->ConvertPropertyToPinType(Property, /*out*/ Pin->PinType);	
-			Pin->bDefaultValueIsIgnored = true;
 		}
 	}
 
