Blueprint: Structure Setter Node. First version.

diff --git a/Engine/Source/Editor/BlueprintGraph/Classes/K2Node_SetFieldsInStruct.h b/Engine/Source/Editor/BlueprintGraph/Classes/K2Node_SetFieldsInStruct.h
new file mode 100644
--- /dev/null
+++ b/Engine/Source/Editor/BlueprintGraph/Classes/K2Node_SetFieldsInStruct.h
@@ -0,0 +1,23 @@
+// Copyright 1998-2014 Epic Games, Inc. All Rights Reserved.
+
+#pragma once
+#include "K2Node_MakeStruct.h"
+#include "K2Node_SetFieldsInStruct.generated.h"
+
+// Pure kismet node that creates a struct with specified values for each member
+UCLASS(MinimalAPI)
+class UK2Node_SetFieldsInStruct : public UK2Node_MakeStruct
+{
+	GENERATED_UCLASS_BODY()
+	// Begin UEdGraphNode interface
+	virtual void AllocateDefaultPins() override;
+	virtual FText GetNodeTitle(ENodeTitleType::Type TitleType) const override;
+	virtual FString GetTooltip() const override;
+	virtual FName GetPaletteIcon(FLinearColor& OutColor) const override;
+	// End  UEdGraphNode interface
+
+	// Begin K2Node interface
+	virtual bool IsNodePure() const override { return false; }
+	virtual class FNodeHandlingFunctor* CreateNodeHandler(class FKismetCompilerContext& CompilerContext) const override;
+	// End K2Node interface
+};
diff --git a/Engine/Source/Editor/BlueprintGraph/Private/K2Node_MakeStruct.cpp b/Engine/Source/Editor/BlueprintGraph/Private/K2Node_MakeStruct.cpp
--- a/Engine/Source/Editor/BlueprintGraph/Private/K2Node_MakeStruct.cpp
+++ b/Engine/Source/Editor/BlueprintGraph/Private/K2Node_MakeStruct.cpp
@@ -2,119 +2,10 @@
 
 #include "BlueprintGraphPrivatePCH.h"
 #include "KismetCompiler.h"
-#include "Editor/UnrealEd/Public/Kismet2/StructureEditorUtils.h"
+#include "MakeStructHandler.h"
 
 #define LOCTEXT_NAMESPACE "K2Node_MakeStruct"
 
-//////////////////////////////////////////////////////////////////////////
-// FKCHandler_MakeStruct
-
-class FKCHandler_MakeStruct : public FNodeHandlingFunctor
-{
-public:
-	UEdGraphPin* FindOutputStructPinChecked(UEdGraphNode* Node)
-	{
-		check(NULL != Node);
-		UEdGraphPin* OutputPin = NULL;
-		for (int32 PinIndex = 0; PinIndex < Node->Pins.Num(); ++PinIndex)
-		{
-			UEdGraphPin* Pin = Node->Pins[PinIndex];
-			if (Pin && (EGPD_Output == Pin->Direction) && !CompilerContext.GetSchema()->IsMetaPin(*Pin))
-			{
-				OutputPin = Pin;
-				break;
-			}
-		}
-		check(NULL != OutputPin);
-		return OutputPin;
-	}
-
-	FKCHandler_MakeStruct(FKismetCompilerContext& InCompilerContext)
-		: FNodeHandlingFunctor(InCompilerContext)
-	{
-	}
-
-	virtual void RegisterNets(FKismetFunctionContext& Context, UEdGraphNode* InNode) override
-	{
-		UK2Node_MakeStruct* Node = CastChecked<UK2Node_MakeStruct>(InNode);
-		if (NULL == Node->StructType)
-		{
-			CompilerContext.MessageLog.Error(*LOCTEXT("MakeStruct_UnknownStructure_Error", "Unknown structure to break for @@").ToString(), Node);
-			return;
-		}
-
-		if(!UK2Node_MakeStruct::CanBeMade(Node->StructType))
-		{
-			CompilerContext.MessageLog.Warning(*LOCTEXT("MakeStruct_Error", "The structure contains read-only members @@. Try use specialized 'make' function if available. ").ToString(), Node);
-		}
-
-		FNodeHandlingFunctor::RegisterNets(Context, Node);
-
-		UEdGraphPin* OutputPin = FindOutputStructPinChecked(Node);
-		UEdGraphPin* Net = FEdGraphUtilities::GetNetFromPin(OutputPin);
-		check(NULL != Net);
-		FBPTerminal** FoundTerm = Context.NetMap.Find(Net);
-		FBPTerminal* Term = FoundTerm ? *FoundTerm : NULL;
-		check(NULL != Term);
-
-		UStruct* StructInTerm = Cast<UStruct>(Term->Type.PinSubCategoryObject.Get());
-		if (NULL == StructInTerm || !StructInTerm->IsChildOf(Node->StructType))
-		{
-			CompilerContext.MessageLog.Error(*LOCTEXT("MakeStruct_NoMatch_Error", "Structures don't match for @@").ToString(), Node);
-		}
-	}
-
-	virtual void RegisterNet(FKismetFunctionContext& Context, UEdGraphPin* Net) override
-	{
-		FBPTerminal* Term = new (Context.IsEventGraph() ? Context.EventGraphLocals : Context.Locals) FBPTerminal();
-		Term->CopyFromPin(Net, Context.NetNameMap->MakeValidName(Net));
-		Context.NetMap.Add(Net, Term);
-	}
-
-	virtual void Compile(FKismetFunctionContext& Context, UEdGraphNode* InNode) override
-	{
-		UK2Node_MakeStruct* Node = CastChecked<UK2Node_MakeStruct>(InNode);
-		if (NULL == Node->StructType)
-		{
-			CompilerContext.MessageLog.Error(*LOCTEXT("MakeStruct_UnknownStructure_Error", "Unknown structure to break for @@").ToString(), Node);
-			return;
-		}
-
-		UEdGraphPin* OutputStructNet = FEdGraphUtilities::GetNetFromPin(FindOutputStructPinChecked(Node));
-		FBPTerminal** FoundTerm = Context.NetMap.Find(OutputStructNet);
-		FBPTerminal* OutputStructTerm = FoundTerm ? *FoundTerm : NULL;
-		check(NULL != OutputStructTerm);
-
-		for (int32 PinIndex = 0; PinIndex < Node->Pins.Num(); ++PinIndex)
-		{
-			UEdGraphPin* Pin = Node->Pins[PinIndex];
-			if (Pin && !CompilerContext.GetSchema()->IsMetaPin(*Pin) && (Pin->Direction == EGPD_Input))
-			{
-				FBPTerminal** FoundSrcTerm = Context.NetMap.Find(FEdGraphUtilities::GetNetFromPin(Pin));
-				FBPTerminal* SrcTerm = FoundSrcTerm ? *FoundSrcTerm : NULL;
-				check(NULL != SrcTerm);
-
-				UProperty* BoundProperty = FindField<UProperty>(Node->StructType, *(Pin->PinName));
-				check(NULL != BoundProperty);
-
-				FBPTerminal* DstTerm = new (Context.IsEventGraph() ? Context.EventGraphLocals : Context.Locals) FBPTerminal();
-				DstTerm->CopyFromPin(Pin, Context.NetNameMap->MakeValidName(Pin));
-				DstTerm->AssociatedVarProperty = BoundProperty;
-				DstTerm->Context = OutputStructTerm;
-
-				FBlueprintCompiledStatement& Statement = Context.AppendStatementForNode(Node);
-				Statement.Type = KCST_Assignment;
-				Statement.LHS = DstTerm;
-				Statement.RHS.Add(SrcTerm);
-			}
-		}
-		if(!Node->IsNodePure())
-		{
-			GenerateSimpleThenGoto(Context, *Node);
-		}
-	}
-};
-
 //////////////////////////////////////////////////////////////////////////
 // UK2Node_MakeStruct
 
@@ -136,43 +27,43 @@ void UK2Node_MakeStruct::FMakeStructPinManager::CustomizePinData(UEdGraphPin* Pi
 				InPin.DefaultTextValue = FText::FromString(Value);
 			}
 			else
-	{
+			{
 				InPin.DefaultValue = Value;
 			}
-	}
+		}
 	};
 
 	UK2Node_StructOperation::FStructOperationOptionalPinManager::CustomizePinData(Pin, SourcePropertyName, ArrayIndex, Property);
-		if(Pin && Property)
-		{
-			const UEdGraphSchema_K2* Schema = GetDefault<UEdGraphSchema_K2>();
-			check(Schema);
+	if (Pin && Property)
+	{
+		const UEdGraphSchema_K2* Schema = GetDefault<UEdGraphSchema_K2>();
+		check(Schema);
 		const bool bIsText = Property->IsA<UTextProperty>();
 		checkSlow(bIsText == (Schema->PC_Text == Pin->PinType.PinCategory));
 
-			const FString& MetadataDefaultValue = Property->GetMetaData(TEXT("MakeStructureDefaultValue"));
-			if (!MetadataDefaultValue.IsEmpty())
-			{
+		const FString& MetadataDefaultValue = Property->GetMetaData(TEXT("MakeStructureDefaultValue"));
+		if (!MetadataDefaultValue.IsEmpty())
+		{
 			FPinDefaultValueHelper::Set(*Pin, MetadataDefaultValue, bIsText);
-				return;
-			}
+			return;
+		}
 
-			if(NULL != SampleStructMemory)
+		if (NULL != SampleStructMemory)
+		{
+			FString NewDefaultValue;
+			if (Property->ExportText_InContainer(0, NewDefaultValue, SampleStructMemory, SampleStructMemory, NULL, PPF_None))
 			{
-				FString NewDefaultValue;
-				if(Property->ExportText_InContainer(0, NewDefaultValue, SampleStructMemory, SampleStructMemory, NULL, PPF_None))
+				if (Schema->IsPinDefaultValid(Pin, NewDefaultValue, NULL, FText::GetEmpty()).IsEmpty())
 				{
-					if(Schema->IsPinDefaultValid(Pin, NewDefaultValue, NULL, FText::GetEmpty()).IsEmpty())
-					{
 					FPinDefaultValueHelper::Set(*Pin, NewDefaultValue, bIsText);
-						return;
-					}
+					return;
 				}
 			}
-			
-			Schema->SetPinDefaultValueBasedOnType(Pin);
 		}
+
+		Schema->SetPinDefaultValueBasedOnType(Pin);
 	}
+}
 
 bool UK2Node_MakeStruct::FMakeStructPinManager::CanTreatPropertyAsOptional(UProperty* TestProperty) const
 	{
@@ -299,13 +190,13 @@ UK2Node::ERedirectType UK2Node_MakeStruct::DoPinsMatchForReconstruction(const UE
 	{
 		const UEdGraphSchema_K2* K2Schema = GetDefault<UEdGraphSchema_K2>();
 		if ((EGPD_Output == NewPin->Direction) && (EGPD_Output == OldPin->Direction))
-	{
-		const UEdGraphSchema_K2* K2Schema = GetDefault<UEdGraphSchema_K2>();
-		if (K2Schema->ArePinTypesCompatible( NewPin->PinType, OldPin->PinType))
 		{
-			Result = ERedirectType_Custom;
+			const UEdGraphSchema_K2* K2Schema = GetDefault<UEdGraphSchema_K2>();
+			if (K2Schema->ArePinTypesCompatible(NewPin->PinType, OldPin->PinType))
+			{
+				Result = ERedirectType_Custom;
+			}
 		}
-	}
 		else if ((EGPD_Input == NewPin->Direction) && (EGPD_Input == OldPin->Direction))
 		{
 			TMap<FName, FName>* StructRedirects = UStruct::TaggedPropertyRedirects.Find(StructType->GetFName());
diff --git a/Engine/Source/Editor/BlueprintGraph/Private/K2Node_SetFieldsInStruct.cpp b/Engine/Source/Editor/BlueprintGraph/Private/K2Node_SetFieldsInStruct.cpp
new file mode 100644
--- /dev/null
+++ b/Engine/Source/Editor/BlueprintGraph/Private/K2Node_SetFieldsInStruct.cpp
@@ -0,0 +1,80 @@
+// Copyright 1998-2014 Epic Games, Inc. All Rights Reserved.
+
+#include "BlueprintGraphPrivatePCH.h"
+#include "K2Node_SetFieldsInStruct.h"
+#include "MakeStructHandler.h"
+
+#define LOCTEXT_NAMESPACE "K2Node_MakeStruct"
+
+struct SetFieldsInStructHelper
+{
+	static const TCHAR* StructRefPinName()
+	{
+		return TEXT("StructRef");
+	}
+};
+
+class FKCHandler_SetFieldsInStruct : public FKCHandler_MakeStruct
+{
+public:
+	FKCHandler_SetFieldsInStruct(FKismetCompilerContext& InCompilerContext) : FKCHandler_MakeStruct(InCompilerContext) {}
+
+	virtual UEdGraphPin* FindStructPinChecked(UEdGraphNode* InNode) const override
+	{
+		check(CastChecked<UK2Node_SetFieldsInStruct>(InNode));
+		UEdGraphPin* FoundPin = InNode->FindPinChecked(SetFieldsInStructHelper::StructRefPinName());
+		check(EGPD_Input == FoundPin->Direction);
+		return FoundPin;
+	}
+};
+
+UK2Node_SetFieldsInStruct::UK2Node_SetFieldsInStruct(const class FPostConstructInitializeProperties& PCIP)
+	: Super(PCIP)
+{
+}
+
+void UK2Node_SetFieldsInStruct::AllocateDefaultPins()
+{
+	const UEdGraphSchema_K2* Schema = GetDefault<UEdGraphSchema_K2>();
+	if (Schema && StructType)
+	{
+		CreatePin(EGPD_Input, Schema->PC_Exec, TEXT(""), NULL, false, false, Schema->PN_Execute);
+		CreatePin(EGPD_Output, Schema->PC_Exec, TEXT(""), NULL, false, false, Schema->PN_Then);
+
+		CreatePin(EGPD_Input, Schema->PC_Struct, TEXT(""), StructType, false, true, SetFieldsInStructHelper::StructRefPinName());
+
+		{
+			FStructOnScope StructOnScope(Cast<UScriptStruct>(StructType));
+			FMakeStructPinManager OptionalPinManager(StructOnScope.GetStructMemory());
+			OptionalPinManager.RebuildPropertyList(ShowPinForProperties, StructType);
+			OptionalPinManager.CreateVisiblePins(ShowPinForProperties, StructType, EGPD_Input, this);
+		}
+	}
+}
+
+FText UK2Node_SetFieldsInStruct::GetNodeTitle(ENodeTitleType::Type TitleType) const
+{
+	FFormatNamedArguments Args;
+	Args.Add(TEXT("StructName"), FText::FromString(StructType ? StructType->GetName() : FString()));
+	return FText::Format(LOCTEXT("SetFieldsInStructNodeTitle", "Set members in {StructName}"), Args);
+}
+
+FString UK2Node_SetFieldsInStruct::GetTooltip() const
+{
+	return FString::Printf(
+		*LOCTEXT("SetFieldsInStruct_Tooltip", "Adds a node that modifies a '%s'").ToString(),
+		*(StructType ? StructType->GetName() : FString())
+		);
+}
+
+FName UK2Node_SetFieldsInStruct::GetPaletteIcon(FLinearColor& OutColor) const
+{ 
+	return UK2Node::GetPaletteIcon(OutColor); 
+}
+
+FNodeHandlingFunctor* UK2Node_SetFieldsInStruct::CreateNodeHandler(class FKismetCompilerContext& CompilerContext) const
+{
+	return new FKCHandler_SetFieldsInStruct(CompilerContext);
+}
+
+#undef LOCTEXT_NAMESPACE
\ No newline at end of file
diff --git a/Engine/Source/Editor/BlueprintGraph/Private/MakeStructHandler.cpp b/Engine/Source/Editor/BlueprintGraph/Private/MakeStructHandler.cpp
new file mode 100644
--- /dev/null
+++ b/Engine/Source/Editor/BlueprintGraph/Private/MakeStructHandler.cpp
@@ -0,0 +1,115 @@
+// Copyright 1998-2014 Epic Games, Inc. All Rights Reserved.
+
+#include "BlueprintGraphPrivatePCH.h"
+#include "MakeStructHandler.h"
+#include "K2Node_MakeStruct.h"
+
+#define LOCTEXT_NAMESPACE "FKCHandler_MakeStruct"
+
+//////////////////////////////////////////////////////////////////////////
+// FKCHandler_MakeStruct
+
+UEdGraphPin* FKCHandler_MakeStruct::FindStructPinChecked(UEdGraphNode* Node) const
+{
+	check(NULL != Node);
+	UEdGraphPin* OutputPin = NULL;
+	for (int32 PinIndex = 0; PinIndex < Node->Pins.Num(); ++PinIndex)
+	{
+		UEdGraphPin* Pin = Node->Pins[PinIndex];
+		if (Pin && (EGPD_Output == Pin->Direction) && !CompilerContext.GetSchema()->IsMetaPin(*Pin))
+		{
+			OutputPin = Pin;
+			break;
+		}
+	}
+	check(NULL != OutputPin);
+	return OutputPin;
+}
+
+FKCHandler_MakeStruct::FKCHandler_MakeStruct(FKismetCompilerContext& InCompilerContext)
+	: FNodeHandlingFunctor(InCompilerContext)
+{
+}
+
+void FKCHandler_MakeStruct::RegisterNets(FKismetFunctionContext& Context, UEdGraphNode* InNode)
+{
+	UK2Node_MakeStruct* Node = CastChecked<UK2Node_MakeStruct>(InNode);
+	if (NULL == Node->StructType)
+	{
+		CompilerContext.MessageLog.Error(*LOCTEXT("MakeStruct_UnknownStructure_Error", "Unknown structure to break for @@").ToString(), Node);
+		return;
+	}
+
+	if (!UK2Node_MakeStruct::CanBeMade(Node->StructType))
+	{
+		CompilerContext.MessageLog.Warning(*LOCTEXT("MakeStruct_Error", "The structure contains read-only members @@. Try use specialized 'make' function if available. ").ToString(), Node);
+	}
+
+	FNodeHandlingFunctor::RegisterNets(Context, Node);
+
+	UEdGraphPin* OutputPin = FindStructPinChecked(Node);
+	UEdGraphPin* Net = FEdGraphUtilities::GetNetFromPin(OutputPin);
+	check(NULL != Net);
+	FBPTerminal** FoundTerm = Context.NetMap.Find(Net);
+	FBPTerminal* Term = FoundTerm ? *FoundTerm : NULL;
+	check(NULL != Term);
+
+	UStruct* StructInTerm = Cast<UStruct>(Term->Type.PinSubCategoryObject.Get());
+	if (NULL == StructInTerm || !StructInTerm->IsChildOf(Node->StructType))
+	{
+		CompilerContext.MessageLog.Error(*LOCTEXT("MakeStruct_NoMatch_Error", "Structures don't match for @@").ToString(), Node);
+	}
+}
+
+void FKCHandler_MakeStruct::RegisterNet(FKismetFunctionContext& Context, UEdGraphPin* Net)
+{
+	FBPTerminal* Term = new (Context.IsEventGraph() ? Context.EventGraphLocals : Context.Locals) FBPTerminal();
+	Term->CopyFromPin(Net, Context.NetNameMap->MakeValidName(Net));
+	Context.NetMap.Add(Net, Term);
+}
+
+void FKCHandler_MakeStruct::Compile(FKismetFunctionContext& Context, UEdGraphNode* InNode)
+{
+	UK2Node_MakeStruct* Node = CastChecked<UK2Node_MakeStruct>(InNode);
+	if (NULL == Node->StructType)
+	{
+		CompilerContext.MessageLog.Error(*LOCTEXT("MakeStruct_UnknownStructure_Error", "Unknown structure to break for @@").ToString(), Node);
+		return;
+	}
+
+	UEdGraphPin* StructPin = FindStructPinChecked(Node);
+	UEdGraphPin* OutputStructNet = FEdGraphUtilities::GetNetFromPin(StructPin);
+	FBPTerminal** FoundTerm = Context.NetMap.Find(OutputStructNet);
+	FBPTerminal* OutputStructTerm = FoundTerm ? *FoundTerm : NULL;
+	check(NULL != OutputStructTerm);
+
+	for (int32 PinIndex = 0; PinIndex < Node->Pins.Num(); ++PinIndex)
+	{
+		UEdGraphPin* Pin = Node->Pins[PinIndex];
+		if (Pin && (Pin != StructPin) && !CompilerContext.GetSchema()->IsMetaPin(*Pin) && (Pin->Direction == EGPD_Input))
+		{
+			FBPTerminal** FoundSrcTerm = Context.NetMap.Find(FEdGraphUtilities::GetNetFromPin(Pin));
+			FBPTerminal* SrcTerm = FoundSrcTerm ? *FoundSrcTerm : NULL;
+			check(NULL != SrcTerm);
+
+			UProperty* BoundProperty = FindField<UProperty>(Node->StructType, *(Pin->PinName));
+			check(NULL != BoundProperty);
+
+			FBPTerminal* DstTerm = new (Context.IsEventGraph() ? Context.EventGraphLocals : Context.Locals) FBPTerminal();
+			DstTerm->CopyFromPin(Pin, Context.NetNameMap->MakeValidName(Pin));
+			DstTerm->AssociatedVarProperty = BoundProperty;
+			DstTerm->Context = OutputStructTerm;
+
+			FBlueprintCompiledStatement& Statement = Context.AppendStatementForNode(Node);
+			Statement.Type = KCST_Assignment;
+			Statement.LHS = DstTerm;
+			Statement.RHS.Add(SrcTerm);
+		}
+	}
+	if (!Node->IsNodePure())
+	{
+		GenerateSimpleThenGoto(Context, *Node);
+	}
+}
+
+#undef LOCTEXT_NAMESPACE
\ No newline at end of file
diff --git a/Engine/Source/Editor/BlueprintGraph/Private/MakeStructHandler.h b/Engine/Source/Editor/BlueprintGraph/Private/MakeStructHandler.h
new file mode 100644
--- /dev/null
+++ b/Engine/Source/Editor/BlueprintGraph/Private/MakeStructHandler.h
@@ -0,0 +1,19 @@
+#pragma once
+#include "KismetCompilerMisc.h"
+
+//////////////////////////////////////////////////////////////////////////
+// FKCHandler_MakeStruct
+
+class FKCHandler_MakeStruct : public FNodeHandlingFunctor
+{
+public:
+	FKCHandler_MakeStruct(FKismetCompilerContext& InCompilerContext);
+
+	virtual UEdGraphPin* FindStructPinChecked(UEdGraphNode* Node) const;
+
+	virtual void RegisterNets(FKismetFunctionContext& Context, UEdGraphNode* InNode) override;
+
+	virtual void RegisterNet(FKismetFunctionContext& Context, UEdGraphPin* Net) override;
+
+	virtual void Compile(FKismetFunctionContext& Context, UEdGraphNode* InNode) override;
+};
\ No newline at end of file
