CIS build fix.

diff --git a/Engine/Source/Editor/BlueprintGraph/Private/MakeStructHandler.cpp b/Engine/Source/Editor/BlueprintGraph/Private/MakeStructHandler.cpp
--- a/Engine/Source/Editor/BlueprintGraph/Private/MakeStructHandler.cpp
+++ b/Engine/Source/Editor/BlueprintGraph/Private/MakeStructHandler.cpp
@@ -3,6 +3,7 @@
 #include "BlueprintGraphPrivatePCH.h"
 #include "MakeStructHandler.h"
 #include "K2Node_MakeStruct.h"
+#include "KismetCompiler.h"
 
 #define LOCTEXT_NAMESPACE "FKCHandler_MakeStruct"
 
