#ue4 fix for 338364, we can now directly set values for variables exposed to a spawn actor from class node. Also fixed assignment of enum constants via an add component or spawn actor from class node, #codereview Maciej.Mroz

diff --git a/Engine/Source/Editor/BlueprintGraph/Private/K2Node_AddComponent.cpp b/Engine/Source/Editor/BlueprintGraph/Private/K2Node_AddComponent.cpp
--- a/Engine/Source/Editor/BlueprintGraph/Private/K2Node_AddComponent.cpp
+++ b/Engine/Source/Editor/BlueprintGraph/Private/K2Node_AddComponent.cpp
@@ -398,46 +398,54 @@ void UK2Node_AddComponent::ExpandNode(class FKismetCompilerContext& CompilerCont
 		// exec in
 		CompilerContext.MovePinLinksToIntermediate(*GetExecPin(), *NewNode->GetExecPin());
 
+#if 1
+		UEdGraphPin* LastThen = FKismetCompilerUtilities::GenerateAssignmentNodes( CompilerContext, SourceGraph, NewNode, this, ReturnPin );
+#else
 		UEdGraphPin* LastThen = NewNode->GetThenPin();
 		for(int32 PinIndex = 0; PinIndex < Pins.Num(); PinIndex++)
 		{
 			UEdGraphPin* OrgPin = Pins[PinIndex];
-			UFunction* SetByNameFunction = Schema->FindSetVariableByNameFunction(OrgPin->PinType);
-			if((NULL == NewNode->FindPin(OrgPin->PinName)) && SetByNameFunction)
+			if( NULL == NewNode->FindPin(OrgPin->PinName))
 			{
-				UK2Node_CallFunction* SetVarNode = NULL;
-				if(OrgPin->PinType.bIsArray)
+				UFunction* SetByNameFunction = Schema->FindSetVariableByNameFunction(OrgPin->PinType);
+				if(SetByNameFunction)
 				{
-					SetVarNode = CompilerContext.SpawnIntermediateNode<UK2Node_CallArrayFunction>(this, SourceGraph);
-				}
-				else
-				{
-					SetVarNode = CompilerContext.SpawnIntermediateNode<UK2Node_CallFunction>(this, SourceGraph);
-				}
-				SetVarNode->SetFromFunction(SetByNameFunction);
-				SetVarNode->AllocateDefaultPins();
+					UK2Node_CallFunction* SetVarNode = NULL;
+					if(OrgPin->PinType.bIsArray)
+					{
+						SetVarNode = CompilerContext.SpawnIntermediateNode<UK2Node_CallArrayFunction>(this, SourceGraph);
+					}
+					else
+					{
+						SetVarNode = CompilerContext.SpawnIntermediateNode<UK2Node_CallFunction>(this, SourceGraph);
+					}
+					SetVarNode->SetFromFunction(SetByNameFunction);
+					SetVarNode->AllocateDefaultPins();
 
-				// Connect this node into the exec chain
-				Schema->TryCreateConnection(LastThen, SetVarNode->GetExecPin());
-				LastThen = SetVarNode->GetThenPin();
+					// Connect this node into the exec chain
+					Schema->TryCreateConnection(LastThen, SetVarNode->GetExecPin());
+					LastThen = SetVarNode->GetThenPin();
 
-				// Connect the new actor to the 'object' pin
-				UEdGraphPin* ObjectPin = SetVarNode->FindPinChecked(ObjectParamName);
-				Schema->TryCreateConnection(ReturnPin, ObjectPin);
+					// Connect the new actor to the 'object' pin
+					UEdGraphPin* ObjectPin = SetVarNode->FindPinChecked(ObjectParamName);
+					Schema->TryCreateConnection(ReturnPin, ObjectPin);
 
-				// Fill in literal for 'property name' pin - name of pin is property name
-				UEdGraphPin* PropertyNamePin = SetVarNode->FindPinChecked(PropertyNameParamName);
-				PropertyNamePin->DefaultValue = OrgPin->PinName;
+					// Fill in literal for 'property name' pin - name of pin is property name
+					UEdGraphPin* PropertyNamePin = SetVarNode->FindPinChecked(PropertyNameParamName);
+					PropertyNamePin->DefaultValue = OrgPin->PinName;
 
-				// Move connection from the variable pin on the spawn node to the 'value' pin
-				UEdGraphPin* ValuePin = SetVarNode->FindPinChecked(ValueParamName);
-				CompilerContext.MovePinLinksToIntermediate(*OrgPin, *ValuePin);
-				if(OrgPin->PinType.bIsArray)
-				{
+					// Move connection from the variable pin on the spawn node to the 'value' pin
+					UEdGraphPin* ValuePin = SetVarNode->FindPinChecked(ValueParamName);
+					CompilerContext.MovePinLinksToIntermediate(*OrgPin, *ValuePin);
 					SetVarNode->PinConnectionListChanged(ValuePin);
+
+					// PinSubCategoryObject is lost when we go from function back to pin when we have type information that is BP only (IE an enum created in the editor):
+					ValuePin->PinType.PinSubCategoryObject = OrgPin->PinType.PinSubCategoryObject;
 				}
 			}
 		}
+#endif
+
 		CompilerContext.MovePinLinksToIntermediate(*GetThenPin(), *LastThen);
 		BreakAllNodeLinks();
 	}
diff --git a/Engine/Source/Editor/BlueprintGraph/Private/K2Node_SpawnActorFromClass.cpp b/Engine/Source/Editor/BlueprintGraph/Private/K2Node_SpawnActorFromClass.cpp
--- a/Engine/Source/Editor/BlueprintGraph/Private/K2Node_SpawnActorFromClass.cpp
+++ b/Engine/Source/Editor/BlueprintGraph/Private/K2Node_SpawnActorFromClass.cpp
@@ -96,7 +96,6 @@ void UK2Node_SpawnActorFromClass::CreatePinsForClass(UClass* InClass)
 		{
 			UEdGraphPin* Pin = CreatePin(EGPD_Input, TEXT(""), TEXT(""), NULL, false, false, Property->GetName());
 			const bool bPinGood = (Pin != NULL) && K2Schema->ConvertPropertyToPinType(Property, /*out*/ Pin->PinType);	
-			Pin->bDefaultValueIsIgnored = true;
 		}
 	}
 
@@ -309,6 +308,7 @@ void UK2Node_SpawnActorFromClass::ExpandNode(class FKismetCompilerContext& Compi
 		static FString PropertyNameParamName = FString(TEXT("PropertyName"));
 
 		const UEdGraphSchema_K2* Schema = CompilerContext.GetSchema();
+		(void)Schema;
 
 		UK2Node_SpawnActorFromClass* SpawnNode = this;
 		UEdGraphPin* SpawnNodeExec = SpawnNode->GetExecPin();
@@ -396,20 +396,23 @@ void UK2Node_SpawnActorFromClass::ExpandNode(class FKismetCompilerContext& Compi
 		// create 'set var' nodes
 
 		// Get 'result' pin from 'begin spawn', this is the actual actor we want to set properties on
-		UK2Node_CallFunction* LastNode = CallBeginSpawnNode;
+		UEdGraphPin* LastThen = FKismetCompilerUtilities::GenerateAssignmentNodes(CompilerContext, SourceGraph, CallBeginSpawnNode, SpawnNode, CallBeginResult );
+
+#if 0
+		UEdGraphPin* LastThen = CallBeginSpawnNode->GetThenPin();
 
 		// Create 'set var by name' nodes and hook them up
 		for(int32 PinIdx=0; PinIdx < SpawnNode->Pins.Num(); PinIdx++)
 		{
 			// Only create 'set param by name' node if this pin is linked to something
-			UEdGraphPin* SpawnVarPin = SpawnNode->Pins[PinIdx];
-			if(SpawnVarPin->LinkedTo.Num() > 0)
+			UEdGraphPin* OrgPin = SpawnNode->Pins[PinIdx];
+			if(OrgPin->LinkedTo.Num() > 0 || OrgPin->DefaultValue != FString() )
 			{
-				UFunction* SetByNameFunction = Schema->FindSetVariableByNameFunction(SpawnVarPin->PinType);
+				UFunction* SetByNameFunction = Schema->FindSetVariableByNameFunction(OrgPin->PinType);
 				if(SetByNameFunction)
 				{
 					UK2Node_CallFunction* SetVarNode = NULL;
-					if(SpawnVarPin->PinType.bIsArray)
+					if(OrgPin->PinType.bIsArray)
 					{
 						SetVarNode = CompilerContext.SpawnIntermediateNode<UK2Node_CallArrayFunction>(SpawnNode, SourceGraph);
 					}
@@ -421,9 +424,8 @@ void UK2Node_SpawnActorFromClass::ExpandNode(class FKismetCompilerContext& Compi
 					SetVarNode->AllocateDefaultPins();
 
 					// Connect this node into the exec chain
-					UEdGraphPin* LastThen = LastNode->GetThenPin();
-					UEdGraphPin* SetVarExec = SetVarNode->GetExecPin();
-					LastThen->MakeLinkTo(SetVarExec);
+					Schema->TryCreateConnection(LastThen, SetVarNode->GetExecPin());
+					LastThen = SetVarNode->GetThenPin();
 
 					// Connect the new actor to the 'object' pin
 					UEdGraphPin* ObjectPin = SetVarNode->FindPinChecked(ObjectParamName);
@@ -431,21 +433,20 @@ void UK2Node_SpawnActorFromClass::ExpandNode(class FKismetCompilerContext& Compi
 
 					// Fill in literal for 'property name' pin - name of pin is property name
 					UEdGraphPin* PropertyNamePin = SetVarNode->FindPinChecked(PropertyNameParamName);
-					PropertyNamePin->DefaultValue = SpawnVarPin->PinName;
+					PropertyNamePin->DefaultValue = OrgPin->PinName;
 
 					// Move connection from the variable pin on the spawn node to the 'value' pin
 					UEdGraphPin* ValuePin = SetVarNode->FindPinChecked(ValueParamName);
-					CompilerContext.MovePinLinksToIntermediate(*SpawnVarPin, *ValuePin);
+					CompilerContext.MovePinLinksToIntermediate(*OrgPin, *ValuePin);
 					SetVarNode->PinConnectionListChanged(ValuePin);
 
-					// Update 'last node in sequence' var
-					LastNode = SetVarNode;
+					// PinSubCategoryObject is lost when we go from function back to pin when we have type information that is BP only (IE an enum created in the editor):
+					ValuePin->PinType.PinSubCategoryObject = OrgPin->PinType.PinSubCategoryObject;
 				}
 			}
 		}
-
+#endif
 		// Make exec connection between 'then' on last node and 'finish'
-		UEdGraphPin* LastThen = LastNode->GetThenPin();
 		LastThen->MakeLinkTo(CallFinishExec);
 
 		// Break any links to the expanded node
diff --git a/Engine/Source/Editor/KismetCompiler/Private/KismetCompilerMisc.cpp b/Engine/Source/Editor/KismetCompiler/Private/KismetCompilerMisc.cpp
--- a/Engine/Source/Editor/KismetCompiler/Private/KismetCompilerMisc.cpp
+++ b/Engine/Source/Editor/KismetCompiler/Private/KismetCompilerMisc.cpp
@@ -502,6 +502,64 @@ void FKismetCompilerUtilities::ValidateEnumProperties(UObject* DefaultObject, FC
 	}
 }
 
+UEdGraphPin* FKismetCompilerUtilities::GenerateAssignmentNodes(class FKismetCompilerContext& CompilerContext, UEdGraph* SourceGraph, UK2Node_CallFunction* CallBeginSpawnNode, UEdGraphNode* SpawnNode, UEdGraphPin* CallBeginResult)
+{
+	static FString ObjectParamName = FString(TEXT("Object"));
+	static FString ValueParamName = FString(TEXT("Value"));
+	static FString PropertyNameParamName = FString(TEXT("PropertyName"));
+
+	const UEdGraphSchema_K2* Schema = CompilerContext.GetSchema();
+	UEdGraphPin* LastThen = CallBeginSpawnNode->GetThenPin();
+
+	// Create 'set var by name' nodes and hook them up
+	for (int32 PinIdx = 0; PinIdx < SpawnNode->Pins.Num(); PinIdx++)
+	{
+		// Only create 'set param by name' node if this pin is linked to something
+		UEdGraphPin* OrgPin = SpawnNode->Pins[PinIdx];
+		if (NULL == CallBeginSpawnNode->FindPin(OrgPin->PinName) &&
+			(OrgPin->LinkedTo.Num() > 0 || OrgPin->DefaultValue != FString()))
+		{
+			UFunction* SetByNameFunction = Schema->FindSetVariableByNameFunction(OrgPin->PinType);
+			if (SetByNameFunction)
+			{
+				UK2Node_CallFunction* SetVarNode = NULL;
+				if (OrgPin->PinType.bIsArray)
+				{
+					SetVarNode = CompilerContext.SpawnIntermediateNode<UK2Node_CallArrayFunction>(SpawnNode, SourceGraph);
+				}
+				else
+				{
+					SetVarNode = CompilerContext.SpawnIntermediateNode<UK2Node_CallFunction>(SpawnNode, SourceGraph);
+				}
+				SetVarNode->SetFromFunction(SetByNameFunction);
+				SetVarNode->AllocateDefaultPins();
+
+				// Connect this node into the exec chain
+				Schema->TryCreateConnection(LastThen, SetVarNode->GetExecPin());
+				LastThen = SetVarNode->GetThenPin();
+
+				// Connect the new actor to the 'object' pin
+				UEdGraphPin* ObjectPin = SetVarNode->FindPinChecked(ObjectParamName);
+				CallBeginResult->MakeLinkTo(ObjectPin);
+
+				// Fill in literal for 'property name' pin - name of pin is property name
+				UEdGraphPin* PropertyNamePin = SetVarNode->FindPinChecked(PropertyNameParamName);
+				PropertyNamePin->DefaultValue = OrgPin->PinName;
+
+				// Move connection from the variable pin on the spawn node to the 'value' pin
+				UEdGraphPin* ValuePin = SetVarNode->FindPinChecked(ValueParamName);
+				CompilerContext.MovePinLinksToIntermediate(*OrgPin, *ValuePin);
+				SetVarNode->PinConnectionListChanged(ValuePin);
+
+				// PinSubCategoryObject is lost when we go from function back to pin when we have type information that is BP only (IE an enum created in the editor):
+				ValuePin->PinType.PinSubCategoryObject = OrgPin->PinType.PinSubCategoryObject;
+			}
+		}
+	}
+
+	return LastThen;
+}
+
 /** Creates a property named PropertyName of type PropertyType in the Scope or returns NULL if the type is unknown, but does *not* link that property in */
 UProperty* FKismetCompilerUtilities::CreatePropertyOnScope(UStruct* Scope, const FName& PropertyName, const FEdGraphPinType& Type, UClass* SelfClass, uint64 PropertyFlags, const UEdGraphSchema_K2* Schema, FCompilerResultsLog& MessageLog)
 {
diff --git a/Engine/Source/Editor/KismetCompiler/Public/KismetCompilerMisc.h b/Engine/Source/Editor/KismetCompiler/Public/KismetCompilerMisc.h
--- a/Engine/Source/Editor/KismetCompiler/Public/KismetCompilerMisc.h
+++ b/Engine/Source/Editor/KismetCompiler/Public/KismetCompilerMisc.h
@@ -75,6 +75,9 @@ public:
 
 	/* checks if enum variables from given object store proper indexes */
 	static void ValidateEnumProperties(UObject* DefaultObject, FCompilerResultsLog& MessageLog);
+
+	/** Create 'set var by name' nodes and hook them up - used to set values when components are added or actor are created at run time. Returns the 'last then' pin of the assignment nodes */
+	static UEdGraphPin* GenerateAssignmentNodes( class FKismetCompilerContext& CompilerContext, UEdGraph* SourceGraph, UK2Node_CallFunction* CallBeginSpawnNode, UEdGraphNode* SpawnNode, UEdGraphPin* CallBeginResult );
 };
 
 //////////////////////////////////////////////////////////////////////////
