BP inner: Function UK2Node_CallFunctionOnMember::HasExternalBlueprintDependencies added. More robust dependency gathering.

diff --git a/Engine/Source/Editor/BlueprintGraph/Classes/K2Node_CallFunctionOnMember.h b/Engine/Source/Editor/BlueprintGraph/Classes/K2Node_CallFunctionOnMember.h
--- a/Engine/Source/Editor/BlueprintGraph/Classes/K2Node_CallFunctionOnMember.h
+++ b/Engine/Source/Editor/BlueprintGraph/Classes/K2Node_CallFunctionOnMember.h
@@ -15,6 +15,8 @@ class UK2Node_CallFunctionOnMember : public UK2Node_CallFunction
 	UPROPERTY()
 	FMemberReference				MemberVariableToCallOn;
 
+	virtual bool HasExternalBlueprintDependencies(TArray<class UStruct*>* OptionalOutput) const override;
+
 	// Begin UK2Node_CallFunction interface
 	virtual UEdGraphPin* CreateSelfPin(const UFunction* Function) override;
 	virtual FString GetFunctionContextString() const override;
diff --git a/Engine/Source/Editor/BlueprintGraph/Private/K2Node_CallFunctionOnMember.cpp b/Engine/Source/Editor/BlueprintGraph/Private/K2Node_CallFunctionOnMember.cpp
--- a/Engine/Source/Editor/BlueprintGraph/Private/K2Node_CallFunctionOnMember.cpp
+++ b/Engine/Source/Editor/BlueprintGraph/Private/K2Node_CallFunctionOnMember.cpp
@@ -111,4 +111,16 @@ void UK2Node_CallFunctionOnMember::ExpandNode(class FKismetCompilerContext& Comp
 	}
 }
 
+bool UK2Node_CallFunctionOnMember::HasExternalBlueprintDependencies(TArray<class UStruct*>* OptionalOutput) const
+{
+	UClass* SourceClass = MemberVariableToCallOn.GetMemberParentClass(this);
+	const UBlueprint* SourceBlueprint = GetBlueprint();
+	const bool bResult = (SourceClass != NULL) && (SourceClass->ClassGeneratedBy != NULL) && (SourceClass->ClassGeneratedBy != SourceBlueprint);
+	if (bResult && OptionalOutput)
+	{
+		OptionalOutput->Add(SourceClass);
+	}
+	return bResult || Super::HasExternalBlueprintDependencies(OptionalOutput);
+}
+
 #undef LOCTEXT_NAMESPACE
