Display friendly names of custom struct members

diff --git a/Engine/Source/Editor/BlueprintGraph/Private/EdGraphSchema_K2.cpp b/Engine/Source/Editor/BlueprintGraph/Private/EdGraphSchema_K2.cpp
--- a/Engine/Source/Editor/BlueprintGraph/Private/EdGraphSchema_K2.cpp
+++ b/Engine/Source/Editor/BlueprintGraph/Private/EdGraphSchema_K2.cpp
@@ -4574,6 +4574,13 @@ void UEdGraphSchema_K2::SplitPin(UEdGraphPin* Pin) const
 			const FEdGraphPinType& ProtoPinType = ProtoPin->PinType;
 			UEdGraphPin* SubPin = GraphNode->CreatePin(Pin->Direction, ProtoPinType.PinCategory, ProtoPinType.PinSubCategory, ProtoPinType.PinSubCategoryObject.Get(), false, false, PinName);
 
+			if (!Pin->PinFriendlyName.IsEmpty() || !ProtoPin->PinFriendlyName.IsEmpty())
+			{
+				SubPin->PinFriendlyName = FText::FromString(FString::Printf(TEXT("%s %s")
+                                                                                            , (Pin->PinFriendlyName.IsEmpty() ? *Pin->PinName : *Pin->PinFriendlyName.ToString())
+                                                                                            , (ProtoPin->PinFriendlyName.IsEmpty() ? *ProtoPin->PinName : *ProtoPin->PinFriendlyName.ToString())));
+			}
+
 			SubPin->DefaultValue = ProtoPin->DefaultValue;
 			SubPin->AutogeneratedDefaultValue = ProtoPin->AutogeneratedDefaultValue;
 
