#ue4 Adding safe fallbacks for delegate nodes when the native class that defines the delegate removes them

diff --git a/Engine/Source/Editor/BlueprintGraph/Private/DelegateNodeHandlers.cpp b/Engine/Source/Editor/BlueprintGraph/Private/DelegateNodeHandlers.cpp
--- a/Engine/Source/Editor/BlueprintGraph/Private/DelegateNodeHandlers.cpp
+++ b/Engine/Source/Editor/BlueprintGraph/Private/DelegateNodeHandlers.cpp
@@ -35,10 +35,17 @@ struct FKCHandlerDelegateHelper
 		check(NULL != DelegateNode);
 
 		UEdGraphPin* Pin = Schema->FindSelfPin(*DelegateNode, EEdGraphPinDirection::EGPD_Input);
-		check(Pin);
+		UStruct* DelegateScope = Context.GetScopeFromPinType(Pin->PinType, Context.NewClass);
+
+		// Early out if we have no pin.  That means the delegate is no longer valid, and we need to terminate gracefully
+		if (!Pin || !DelegateScope)
+		{
+			MessageLog.Error(*LOCTEXT("NoDelegateProperty", "Event Dispatcher has no property @@").ToString(), DelegateNode);
+			return NULL;
+		}
 
 		// Don't use DelegateNode->GetProperty(), because we don't want any property from skeletal class
-		UClass* PropertyOwnerClass = CastChecked<UClass>(Context.GetScopeFromPinType(Pin->PinType, Context.NewClass));
+		UClass* PropertyOwnerClass = CastChecked<UClass>(DelegateScope);
 		UMulticastDelegateProperty* BoundProperty = NULL;
 		for (TFieldIterator<UMulticastDelegateProperty> It(PropertyOwnerClass); It; ++It)
 		{
diff --git a/Engine/Source/Editor/BlueprintGraph/Private/K2Node_MCDelegate.cpp b/Engine/Source/Editor/BlueprintGraph/Private/K2Node_MCDelegate.cpp
--- a/Engine/Source/Editor/BlueprintGraph/Private/K2Node_MCDelegate.cpp
+++ b/Engine/Source/Editor/BlueprintGraph/Private/K2Node_MCDelegate.cpp
@@ -55,9 +55,9 @@ void UK2Node_BaseMCDelegate::AllocateDefaultPins()
 	CreatePin(EGPD_Output, K2Schema->PC_Exec, TEXT(""), NULL, false, false, K2Schema->PN_Then);
 
 	const auto Property = DelegateReference.ResolveMember<UMulticastDelegateProperty>(this);
-	const auto PropertyOwnerClass = Property->GetOwnerClass();
+	const auto PropertyOwnerClass = (Property ? Property->GetOwnerClass() : NULL);
 	const auto Blueprint = GetBlueprint();
-	ensure(PropertyOwnerClass && Blueprint);
+	
 	const bool bUseSelf = Blueprint && (PropertyOwnerClass == Blueprint->GeneratedClass || PropertyOwnerClass == Blueprint->SkeletonGeneratedClass);
 
 	UEdGraphPin* SelfPin = NULL;
