#ue4 Adding safe fallbacks for delegate nodes when the native class that defines the delegate removes them - fix

diff --git a/Engine/Source/Editor/BlueprintGraph/Private/DelegateNodeHandlers.cpp b/Engine/Source/Editor/BlueprintGraph/Private/DelegateNodeHandlers.cpp
--- a/Engine/Source/Editor/BlueprintGraph/Private/DelegateNodeHandlers.cpp
+++ b/Engine/Source/Editor/BlueprintGraph/Private/DelegateNodeHandlers.cpp
@@ -35,7 +35,7 @@ struct FKCHandlerDelegateHelper
 		check(NULL != DelegateNode);
 
 		UEdGraphPin* Pin = Schema->FindSelfPin(*DelegateNode, EEdGraphPinDirection::EGPD_Input);
-		UStruct* DelegateScope = Context.GetScopeFromPinType(Pin->PinType, Context.NewClass);
+		UStruct* DelegateScope = Pin ? Context.GetScopeFromPinType(Pin->PinType, Context.NewClass) : NULL;
 
 		// Early out if we have no pin.  That means the delegate is no longer valid, and we need to terminate gracefully
 		if (!Pin || !DelegateScope)
