Maintain decimal point when splitting vector/rotator structs

diff --git a/Engine/Source/Editor/BlueprintGraph/Private/EdGraphSchema_K2.cpp b/Engine/Source/Editor/BlueprintGraph/Private/EdGraphSchema_K2.cpp
--- a/Engine/Source/Editor/BlueprintGraph/Private/EdGraphSchema_K2.cpp
+++ b/Engine/Source/Editor/BlueprintGraph/Private/EdGraphSchema_K2.cpp
@@ -4604,7 +4604,7 @@ void UEdGraphSchema_K2::SplitPin(UEdGraphPin* Pin) const
 			Pin->DefaultValue.ParseIntoArray(&OriginalDefaults, TEXT(","), false);
 			for (FString& Default : OriginalDefaults)
 			{
-				Default.Trim();
+				Default = FString::SanitizeFloat(FCString::Atof(*Default));
 			}
 		}
 		else if (StructType == FindObjectChecked<UScriptStruct>(UObject::StaticClass(), TEXT("Vector2D")))
