Move FBaseAsyncTaskHelper to protected struct so that child classes can use it

diff --git a/Engine/Source/Editor/BlueprintGraph/Classes/K2Node_BaseAsyncTask.h b/Engine/Source/Editor/BlueprintGraph/Classes/K2Node_BaseAsyncTask.h
--- a/Engine/Source/Editor/BlueprintGraph/Classes/K2Node_BaseAsyncTask.h
+++ b/Engine/Source/Editor/BlueprintGraph/Classes/K2Node_BaseAsyncTask.h
@@ -6,6 +6,10 @@
 #include "EdGraphSchema_K2_Actions.h"
 #include "K2Node_BaseAsyncTask.generated.h"
 
+class UK2Node_CustomEvent;
+class UEdGraphPin;
+class UK2Node_TemporaryVariable;
+
 UCLASS(Abstract)
 class BLUEPRINTGRAPH_API UK2Node_BaseAsyncTask : public UK2Node
 {
@@ -49,4 +53,23 @@ protected:
 	// The name of the 'go' function on the proxy object that will be called after delegates are in place, can be NAME_None
 	UPROPERTY()
 	FName ProxyActivateFunctionName;
+
+	struct BLUEPRINTGRAPH_API FBaseAsyncTaskHelper
+	{
+		struct FOutputPinAndLocalVariable
+		{
+			UEdGraphPin* OutputPin;
+			UK2Node_TemporaryVariable* TempVar;
+
+			FOutputPinAndLocalVariable(UEdGraphPin* Pin, UK2Node_TemporaryVariable* Var) : OutputPin(Pin), TempVar(Var) {}
+		};
+
+		static bool ValidDataPin(const UEdGraphPin* Pin, EEdGraphPinDirection Direction, const UEdGraphSchema_K2* Schema);
+		static bool CreateDelegateForNewFunction(UEdGraphPin* DelegateInputPin, FName FunctionName, UK2Node* CurrentNode, UEdGraph* SourceGraph, FKismetCompilerContext& CompilerContext);
+		static bool CopyEventSignature(UK2Node_CustomEvent* CENode, UFunction* Function, const UEdGraphSchema_K2* Schema);
+		static bool HandleDelegateImplementation(
+			UMulticastDelegateProperty* CurrentProperty, const TArray<FBaseAsyncTaskHelper::FOutputPinAndLocalVariable>& VariableOutputs,
+			UEdGraphPin* ProxyObjectPin, UEdGraphPin*& InOutLastThenPin,
+			UK2Node* CurrentNode, UEdGraph* SourceGraph, FKismetCompilerContext& CompilerContext);
+	};
 };
diff --git a/Engine/Source/Editor/BlueprintGraph/Private/K2Node_BaseAsyncTask.cpp b/Engine/Source/Editor/BlueprintGraph/Private/K2Node_BaseAsyncTask.cpp
--- a/Engine/Source/Editor/BlueprintGraph/Private/K2Node_BaseAsyncTask.cpp
+++ b/Engine/Source/Editor/BlueprintGraph/Private/K2Node_BaseAsyncTask.cpp
@@ -140,130 +140,120 @@ void UK2Node_BaseAsyncTask::AllocateDefaultPins()
 	Super::AllocateDefaultPins();
 }
 
-struct FBaseAsyncTaskHelper
-{
-	static bool ValidDataPin(const UEdGraphPin* Pin, EEdGraphPinDirection Direction, const UEdGraphSchema_K2* Schema)
+
+bool UK2Node_BaseAsyncTask::FBaseAsyncTaskHelper::ValidDataPin(const UEdGraphPin* Pin, EEdGraphPinDirection Direction, const UEdGraphSchema_K2* Schema)
 {
-		check(Schema);
-		const bool bValidDataPin = Pin
-			&& (Pin->PinName != Schema->PN_Execute)
-			&& (Pin->PinName != Schema->PN_Then)
-			&& (Pin->PinType.PinCategory != Schema->PC_Exec);
+	check(Schema);
+	const bool bValidDataPin = Pin
+		&& (Pin->PinName != Schema->PN_Execute)
+		&& (Pin->PinName != Schema->PN_Then)
+		&& (Pin->PinType.PinCategory != Schema->PC_Exec);
 
-		const bool bProperDirection = Pin && (Pin->Direction == Direction);
+	const bool bProperDirection = Pin && (Pin->Direction == Direction);
 
-		return bValidDataPin && bProperDirection;
+	return bValidDataPin && bProperDirection;
 }
 
-	static bool CreateDelegateForNewFunction(UEdGraphPin* DelegateInputPin, FName FunctionName, UK2Node* CurrentNode, UEdGraph* SourceGraph, FKismetCompilerContext& CompilerContext)
-	{
+bool UK2Node_BaseAsyncTask::FBaseAsyncTaskHelper::CreateDelegateForNewFunction(UEdGraphPin* DelegateInputPin, FName FunctionName, UK2Node* CurrentNode, UEdGraph* SourceGraph, FKismetCompilerContext& CompilerContext)
+{
 	const UEdGraphSchema_K2* Schema = CompilerContext.GetSchema();
-		check(DelegateInputPin && Schema && CurrentNode && SourceGraph && (FunctionName != NAME_None));
-		bool bResult = true;
+	check(DelegateInputPin && Schema && CurrentNode && SourceGraph && (FunctionName != NAME_None));
+	bool bResult = true;
 
-		// WORKAROUND, so we can create delegate from nonexistent function by avoiding check at expanding step
-		// instead simply: Schema->TryCreateConnection(AddDelegateNode->GetDelegatePin(), CurrentCENode->FindPinChecked(UK2Node_CustomEvent::DelegateOutputName));
-		UK2Node_Self* SelfNode = CompilerContext.SpawnIntermediateNode<UK2Node_Self>(CurrentNode, SourceGraph);
-		SelfNode->AllocateDefaultPins();
+	// WORKAROUND, so we can create delegate from nonexistent function by avoiding check at expanding step
+	// instead simply: Schema->TryCreateConnection(AddDelegateNode->GetDelegatePin(), CurrentCENode->FindPinChecked(UK2Node_CustomEvent::DelegateOutputName));
+	UK2Node_Self* SelfNode = CompilerContext.SpawnIntermediateNode<UK2Node_Self>(CurrentNode, SourceGraph);
+	SelfNode->AllocateDefaultPins();
 
-		UK2Node_CreateDelegate* CreateDelegateNode = CompilerContext.SpawnIntermediateNode<UK2Node_CreateDelegate>(CurrentNode, SourceGraph);
-		CreateDelegateNode->AllocateDefaultPins();
-		bResult &= Schema->TryCreateConnection(DelegateInputPin, CreateDelegateNode->GetDelegateOutPin());
-		bResult &= Schema->TryCreateConnection(SelfNode->FindPinChecked(Schema->PN_Self), CreateDelegateNode->GetObjectInPin());
-		CreateDelegateNode->SetFunction(FunctionName);
+	UK2Node_CreateDelegate* CreateDelegateNode = CompilerContext.SpawnIntermediateNode<UK2Node_CreateDelegate>(CurrentNode, SourceGraph);
+	CreateDelegateNode->AllocateDefaultPins();
+	bResult &= Schema->TryCreateConnection(DelegateInputPin, CreateDelegateNode->GetDelegateOutPin());
+	bResult &= Schema->TryCreateConnection(SelfNode->FindPinChecked(Schema->PN_Self), CreateDelegateNode->GetObjectInPin());
+	CreateDelegateNode->SetFunction(FunctionName);
 
-		return bResult;
-	}
-
-	struct FOutputPinAndLocalVariable
-	{
-		UEdGraphPin* OutputPin;
-		UK2Node_TemporaryVariable* TempVar;
+	return bResult;
+}
 
-		FOutputPinAndLocalVariable(UEdGraphPin* Pin, UK2Node_TemporaryVariable* Var) : OutputPin(Pin), TempVar(Var) {}
-	};
+bool UK2Node_BaseAsyncTask::FBaseAsyncTaskHelper::CopyEventSignature(UK2Node_CustomEvent* CENode, UFunction* Function, const UEdGraphSchema_K2* Schema)
+{
+	check(CENode && Function && Schema);
 
-	static bool CopyEventSignature(UK2Node_CustomEvent* CENode, UFunction* Function, const UEdGraphSchema_K2* Schema)
+	bool bResult = true;
+	for (TFieldIterator<UProperty> PropIt(Function); PropIt && (PropIt->PropertyFlags & CPF_Parm); ++PropIt)
 	{
-		check(CENode && Function && Schema);
-
-		bool bResult = true;
-		for (TFieldIterator<UProperty> PropIt(Function); PropIt && (PropIt->PropertyFlags & CPF_Parm); ++PropIt)
+		const UProperty* Param = *PropIt;
+		if (!Param->HasAnyPropertyFlags(CPF_OutParm) || Param->HasAnyPropertyFlags(CPF_ReferenceParm))
 		{
-			const UProperty* Param = *PropIt;
-			if (!Param->HasAnyPropertyFlags(CPF_OutParm) || Param->HasAnyPropertyFlags(CPF_ReferenceParm))
-		{
-				FEdGraphPinType PinType;
-				bResult &= Schema->ConvertPropertyToPinType(Param, /*out*/ PinType);
-				bResult &= (NULL != CENode->CreateUserDefinedPin(Param->GetName(), PinType));
+			FEdGraphPinType PinType;
+			bResult &= Schema->ConvertPropertyToPinType(Param, /*out*/ PinType);
+			bResult &= (NULL != CENode->CreateUserDefinedPin(Param->GetName(), PinType));
 		}
 	}
+	return bResult;
+}
+
+bool UK2Node_BaseAsyncTask::FBaseAsyncTaskHelper::HandleDelegateImplementation(
+	UMulticastDelegateProperty* CurrentProperty, const TArray<FBaseAsyncTaskHelper::FOutputPinAndLocalVariable>& VariableOutputs,
+	UEdGraphPin* ProxyObjectPin, UEdGraphPin*& InOutLastThenPin,
+	UK2Node* CurrentNode, UEdGraph* SourceGraph, FKismetCompilerContext& CompilerContext)
+{
+	bool bIsErrorFree = true;
+	const UEdGraphSchema_K2* Schema = CompilerContext.GetSchema();
+	check(CurrentProperty && ProxyObjectPin && InOutLastThenPin && CurrentNode && SourceGraph && Schema);
 
-		return bResult;
+	UEdGraphPin* PinForCurrentDelegateProperty = CurrentNode->FindPin(CurrentProperty->GetName());
+	if (!PinForCurrentDelegateProperty || (Schema->PC_Exec != PinForCurrentDelegateProperty->PinType.PinCategory))
+	{
+		FText ErrorMessage = FText::Format(LOCTEXT("WrongDelegateProperty", "BaseAsyncTask: Cannot find execution pin for delegate "), FText::FromString(CurrentProperty->GetName()));
+		CompilerContext.MessageLog.Error(*ErrorMessage.ToString(), CurrentNode);
+		return false;
 	}
 
-	static bool HandleDelegateImplementation(
-		UMulticastDelegateProperty* CurrentProperty, const TArray<FBaseAsyncTaskHelper::FOutputPinAndLocalVariable>& VariableOutputs,
-		UEdGraphPin* ProxyObjectPin, UEdGraphPin*& InOutLastThenPin,
-		UK2Node* CurrentNode, UEdGraph* SourceGraph, FKismetCompilerContext& CompilerContext)
+	UK2Node_CustomEvent* CurrentCENode = CompilerContext.SpawnIntermediateNode<UK2Node_CustomEvent>(CurrentNode, SourceGraph);
 	{
-		bool bIsErrorFree = true;
-		const UEdGraphSchema_K2* Schema = CompilerContext.GetSchema();
-		check(CurrentProperty && ProxyObjectPin && InOutLastThenPin && CurrentNode && SourceGraph && Schema);
+	UK2Node_AddDelegate* AddDelegateNode = CompilerContext.SpawnIntermediateNode<UK2Node_AddDelegate>(CurrentNode, SourceGraph);
+	AddDelegateNode->SetFromProperty(CurrentProperty, false);
+	AddDelegateNode->AllocateDefaultPins();
+		bIsErrorFree &= Schema->TryCreateConnection(AddDelegateNode->FindPinChecked(Schema->PN_Self), ProxyObjectPin);
+		bIsErrorFree &= Schema->TryCreateConnection(InOutLastThenPin, AddDelegateNode->FindPinChecked(Schema->PN_Execute));
+		InOutLastThenPin = AddDelegateNode->FindPinChecked(Schema->PN_Then);
+		CurrentCENode->CustomFunctionName = *FString::Printf(TEXT("%s_%s"), *CurrentProperty->GetName(), *CurrentCENode->NodeGuid.ToString());
+		CurrentCENode->AllocateDefaultPins();
+
+		bIsErrorFree &= FBaseAsyncTaskHelper::CreateDelegateForNewFunction(AddDelegateNode->GetDelegatePin(), CurrentCENode->GetFunctionName(), CurrentNode, SourceGraph, CompilerContext);
+		bIsErrorFree &= FBaseAsyncTaskHelper::CopyEventSignature(CurrentCENode, AddDelegateNode->GetDelegateSignature(), Schema);
+	}
 
-		UEdGraphPin* PinForCurrentDelegateProperty = CurrentNode->FindPin(CurrentProperty->GetName());
-		if (!PinForCurrentDelegateProperty || (Schema->PC_Exec != PinForCurrentDelegateProperty->PinType.PinCategory))
-		{
-			FText ErrorMessage = FText::Format(LOCTEXT("WrongDelegateProperty", "BaseAsyncTask: Cannot find execution pin for delegate "), FText::FromString(CurrentProperty->GetName()));
-			CompilerContext.MessageLog.Error(*ErrorMessage.ToString(), CurrentNode);
+	UEdGraphPin* LastActivatedNodeThen = CurrentCENode->FindPinChecked(Schema->PN_Then);
+	for (auto OutputPair : VariableOutputs) // CREATE CHAIN OF ASSIGMENTS
+	{
+		UEdGraphPin* PinWithData = CurrentCENode->FindPin(OutputPair.OutputPin->PinName);
+		if (PinWithData == NULL)
+	{
+			FText ErrorMessage = FText::Format(LOCTEXT("MissingDataPin", "ICE: Pin @@ was expecting a data output pin named {0} on @@ (each delegate must have the same signature)"), FText::FromString(OutputPair.OutputPin->PinName));
+			CompilerContext.MessageLog.Error(*ErrorMessage.ToString(), OutputPair.OutputPin, CurrentCENode);
 			return false;
-		}
-
-		UK2Node_CustomEvent* CurrentCENode = CompilerContext.SpawnIntermediateNode<UK2Node_CustomEvent>(CurrentNode, SourceGraph);
-		{
-		UK2Node_AddDelegate* AddDelegateNode = CompilerContext.SpawnIntermediateNode<UK2Node_AddDelegate>(CurrentNode, SourceGraph);
-		AddDelegateNode->SetFromProperty(CurrentProperty, false);
-		AddDelegateNode->AllocateDefaultPins();
-			bIsErrorFree &= Schema->TryCreateConnection(AddDelegateNode->FindPinChecked(Schema->PN_Self), ProxyObjectPin);
-			bIsErrorFree &= Schema->TryCreateConnection(InOutLastThenPin, AddDelegateNode->FindPinChecked(Schema->PN_Execute));
-			InOutLastThenPin = AddDelegateNode->FindPinChecked(Schema->PN_Then);
-			CurrentCENode->CustomFunctionName = *FString::Printf(TEXT("%s_%s"), *CurrentProperty->GetName(), *CurrentCENode->NodeGuid.ToString());
-			CurrentCENode->AllocateDefaultPins();
-
-			bIsErrorFree &= FBaseAsyncTaskHelper::CreateDelegateForNewFunction(AddDelegateNode->GetDelegatePin(), CurrentCENode->GetFunctionName(), CurrentNode, SourceGraph, CompilerContext);
-			bIsErrorFree &= FBaseAsyncTaskHelper::CopyEventSignature(CurrentCENode, AddDelegateNode->GetDelegateSignature(), Schema);
-		}
+	}
 
-		UEdGraphPin* LastActivatedNodeThen = CurrentCENode->FindPinChecked(Schema->PN_Then);
-		for (auto OutputPair : VariableOutputs) // CREATE CHAIN OF ASSIGMENTS
-		{
-			UEdGraphPin* PinWithData = CurrentCENode->FindPin(OutputPair.OutputPin->PinName);
-			if (PinWithData == NULL)
-		{
-				FText ErrorMessage = FText::Format(LOCTEXT("MissingDataPin", "ICE: Pin @@ was expecting a data output pin named {0} on @@ (each delegate must have the same signature)"), FText::FromString(OutputPair.OutputPin->PinName));
-				CompilerContext.MessageLog.Error(*ErrorMessage.ToString(), OutputPair.OutputPin, CurrentCENode);
-				return false;
-		}
+		UK2Node_AssignmentStatement* AssignNode = CompilerContext.SpawnIntermediateNode<UK2Node_AssignmentStatement>(CurrentNode, SourceGraph);
+		AssignNode->AllocateDefaultPins();
+		bIsErrorFree &= Schema->TryCreateConnection(LastActivatedNodeThen, AssignNode->GetExecPin());
+		bIsErrorFree &= Schema->TryCreateConnection(OutputPair.TempVar->GetVariablePin(), AssignNode->GetVariablePin());
+		AssignNode->NotifyPinConnectionListChanged(AssignNode->GetVariablePin());
+		bIsErrorFree &= Schema->TryCreateConnection(AssignNode->GetValuePin(), PinWithData);
+		AssignNode->NotifyPinConnectionListChanged(AssignNode->GetValuePin());
 
-			UK2Node_AssignmentStatement* AssignNode = CompilerContext.SpawnIntermediateNode<UK2Node_AssignmentStatement>(CurrentNode, SourceGraph);
-			AssignNode->AllocateDefaultPins();
-			bIsErrorFree &= Schema->TryCreateConnection(LastActivatedNodeThen, AssignNode->GetExecPin());
-			bIsErrorFree &= Schema->TryCreateConnection(OutputPair.TempVar->GetVariablePin(), AssignNode->GetVariablePin());
-			AssignNode->NotifyPinConnectionListChanged(AssignNode->GetVariablePin());
-			bIsErrorFree &= Schema->TryCreateConnection(AssignNode->GetValuePin(), PinWithData);
-			AssignNode->NotifyPinConnectionListChanged(AssignNode->GetValuePin());
+		LastActivatedNodeThen = AssignNode->GetThenPin();
+	}
 
-			LastActivatedNodeThen = AssignNode->GetThenPin();
-		}
+	bIsErrorFree &= CompilerContext.MovePinLinksToIntermediate(*PinForCurrentDelegateProperty, *LastActivatedNodeThen).CanSafeConnect();
+	return bIsErrorFree;
+}
 
-		bIsErrorFree &= CompilerContext.MovePinLinksToIntermediate(*PinForCurrentDelegateProperty, *LastActivatedNodeThen).CanSafeConnect();
-		return bIsErrorFree;
-	}
-};
 
 void UK2Node_BaseAsyncTask::ExpandNode(class FKismetCompilerContext& CompilerContext, UEdGraph* SourceGraph)
 {
-        Super::ExpandNode(CompilerContext, SourceGraph);
+    Super::ExpandNode(CompilerContext, SourceGraph);
 
 	if (!CompilerContext.bIsFullCompile)
 	{
