Custom events are now searchable by name in Find-in-Blueprints

#ttp 340065 - BLUEPRINT: Unable to search custom event name

diff --git a/Engine/Source/Editor/BlueprintGraph/Private/K2Node_CustomEvent.cpp b/Engine/Source/Editor/BlueprintGraph/Private/K2Node_CustomEvent.cpp
--- a/Engine/Source/Editor/BlueprintGraph/Private/K2Node_CustomEvent.cpp
+++ b/Engine/Source/Editor/BlueprintGraph/Private/K2Node_CustomEvent.cpp
@@ -103,14 +103,10 @@ UK2Node_CustomEvent::UK2Node_CustomEvent(const class FPostConstructInitializePro
 
 FText UK2Node_CustomEvent::GetNodeTitle(ENodeTitleType::Type TitleType) const
 {
-	if (TitleType == ENodeTitleType::EditableTitle)
+	if (TitleType == ENodeTitleType::EditableTitle || TitleType == ENodeTitleType::ListView)
 	{
 		return FText::FromName(CustomFunctionName);
 	}
-	else if(TitleType == ENodeTitleType::ListView)
-	{
-		return NSLOCTEXT("K2Node", "CustomEvent_Title", "Custom Event");
-	}
 	else
 	{
 		FString RPCString = UK2Node_Event::GetLocalizedNetString(FunctionFlags, false);
