Making it so blueprints that inherit directly from Object (and any blueprint ever for that matter) can be used as a blueprint variable type.
(TTP #340627)

diff --git a/Engine/Source/Editor/BlueprintGraph/Private/EdGraphSchema_K2.cpp b/Engine/Source/Editor/BlueprintGraph/Private/EdGraphSchema_K2.cpp
--- a/Engine/Source/Editor/BlueprintGraph/Private/EdGraphSchema_K2.cpp
+++ b/Engine/Source/Editor/BlueprintGraph/Private/EdGraphSchema_K2.cpp
@@ -627,24 +627,24 @@ bool UEdGraphSchema_K2::IsAllowableBlueprintVariableType(const class UClass* InC
 
 		// UObject is an exception, and is always a blueprint-able type
 		if(InClass == UObject::StaticClass())
-{
+		{
 			return true;
-}
+		}
 
 		struct FBlueprintIsNotBlueprintTypeHelper
-	{
+		{
 			bool bNotBlueprintType;
 
 			FBlueprintIsNotBlueprintTypeHelper() : bNotBlueprintType(false)
-{
+			{
 				GConfig->GetBool(TEXT("EditoronlyBP"), TEXT("bBlueprintIsNotBlueprintType"), bNotBlueprintType, GEditorIni);
-}
+			}
 		};
 		static const FBlueprintIsNotBlueprintTypeHelper BPHelper;
 		if (BPHelper.bNotBlueprintType && InClass->IsChildOf(UBlueprint::StaticClass()))
-{
+		{
 			return false;
-	}
+		}
 	
 		const UClass* ParentClass = InClass;
 		while(ParentClass)
@@ -652,19 +652,19 @@ bool UEdGraphSchema_K2::IsAllowableBlueprintVariableType(const class UClass* InC
 			// Climb up the class hierarchy and look for "BlueprintType" and "NotBlueprintType" to see if this class is allowed.
 			if(ParentClass->GetBoolMetaData(FBlueprintMetadata::MD_AllowableBlueprintVariableType)
 				|| ParentClass->HasMetaData(FBlueprintMetadata::MD_BlueprintSpawnableComponent))
-		{
+			{
 				return true;
-		}
+			}
 			else if(ParentClass->GetBoolMetaData(FBlueprintMetadata::MD_NotAllowableBlueprintVariableType))
-		{
+			{
 				return false;
-				}
-			ParentClass = ParentClass->GetSuperClass();
 			}
+			ParentClass = ParentClass->GetSuperClass();
 		}
+	}
 
 	return false;
-		}
+}
 
 bool UEdGraphSchema_K2::IsAllowableBlueprintVariableType(const class UScriptStruct *InStruct)
 		{
diff --git a/Engine/Source/Editor/KismetCompiler/Private/KismetCompiler.cpp b/Engine/Source/Editor/KismetCompiler/Private/KismetCompiler.cpp
--- a/Engine/Source/Editor/KismetCompiler/Private/KismetCompiler.cpp
+++ b/Engine/Source/Editor/KismetCompiler/Private/KismetCompiler.cpp
@@ -1630,6 +1630,8 @@ void FKismetCompilerContext::FinishCompilingClass(UClass* Class)
 		{
 			Class->SetMetaData(TEXT("Category"), *Blueprint->BlueprintCategory);
 		}
+
+		Class->SetMetaData(FBlueprintMetadata::MD_AllowableBlueprintVariableType, TEXT("true"));
 #endif
 
 		// Add in additional flags implied by the blueprint
