#ue4 code cleanup.. I checked a block of code in commented out, it needed to simply be removed.

diff --git a/Engine/Source/Editor/BlueprintGraph/Private/K2Node_AddComponent.cpp b/Engine/Source/Editor/BlueprintGraph/Private/K2Node_AddComponent.cpp
--- a/Engine/Source/Editor/BlueprintGraph/Private/K2Node_AddComponent.cpp
+++ b/Engine/Source/Editor/BlueprintGraph/Private/K2Node_AddComponent.cpp
@@ -379,8 +379,6 @@ void UK2Node_AddComponent::ExpandNode(class FKismetCompilerContext& CompilerCont
 		static FString ValueParamName = FString(TEXT("Value"));
 		static FString PropertyNameParamName = FString(TEXT("PropertyName"));
 
-		const UEdGraphSchema_K2* Schema = CompilerContext.GetSchema();
-
 		UK2Node_AddComponent* NewNode = CompilerContext.SpawnIntermediateNode<UK2Node_AddComponent>(this, SourceGraph); 
 		NewNode->SetFromFunction(GetTargetFunction());
 		NewNode->AllocateDefaultPinsWithoutExposedVariables();
@@ -398,53 +396,7 @@ void UK2Node_AddComponent::ExpandNode(class FKismetCompilerContext& CompilerCont
 		// exec in
 		CompilerContext.MovePinLinksToIntermediate(*GetExecPin(), *NewNode->GetExecPin());
 
-#if 1
 		UEdGraphPin* LastThen = FKismetCompilerUtilities::GenerateAssignmentNodes( CompilerContext, SourceGraph, NewNode, this, ReturnPin );
-#else
-		UEdGraphPin* LastThen = NewNode->GetThenPin();
-		for(int32 PinIndex = 0; PinIndex < Pins.Num(); PinIndex++)
-		{
-			UEdGraphPin* OrgPin = Pins[PinIndex];
-			if( NULL == NewNode->FindPin(OrgPin->PinName))
-			{
-				UFunction* SetByNameFunction = Schema->FindSetVariableByNameFunction(OrgPin->PinType);
-				if(SetByNameFunction)
-				{
-					UK2Node_CallFunction* SetVarNode = NULL;
-					if(OrgPin->PinType.bIsArray)
-					{
-						SetVarNode = CompilerContext.SpawnIntermediateNode<UK2Node_CallArrayFunction>(this, SourceGraph);
-					}
-					else
-					{
-						SetVarNode = CompilerContext.SpawnIntermediateNode<UK2Node_CallFunction>(this, SourceGraph);
-					}
-					SetVarNode->SetFromFunction(SetByNameFunction);
-					SetVarNode->AllocateDefaultPins();
-
-					// Connect this node into the exec chain
-					Schema->TryCreateConnection(LastThen, SetVarNode->GetExecPin());
-					LastThen = SetVarNode->GetThenPin();
-
-					// Connect the new actor to the 'object' pin
-					UEdGraphPin* ObjectPin = SetVarNode->FindPinChecked(ObjectParamName);
-					Schema->TryCreateConnection(ReturnPin, ObjectPin);
-
-					// Fill in literal for 'property name' pin - name of pin is property name
-					UEdGraphPin* PropertyNamePin = SetVarNode->FindPinChecked(PropertyNameParamName);
-					PropertyNamePin->DefaultValue = OrgPin->PinName;
-
-					// Move connection from the variable pin on the spawn node to the 'value' pin
-					UEdGraphPin* ValuePin = SetVarNode->FindPinChecked(ValueParamName);
-					CompilerContext.MovePinLinksToIntermediate(*OrgPin, *ValuePin);
-					SetVarNode->PinConnectionListChanged(ValuePin);
-
-					// PinSubCategoryObject is lost when we go from function back to pin when we have type information that is BP only (IE an enum created in the editor):
-					ValuePin->PinType.PinSubCategoryObject = OrgPin->PinType.PinSubCategoryObject;
-				}
-			}
-		}
-#endif
 
 		CompilerContext.MovePinLinksToIntermediate(*GetThenPin(), *LastThen);
 		BreakAllNodeLinks();
diff --git a/Engine/Source/Editor/BlueprintGraph/Private/K2Node_SpawnActorFromClass.cpp b/Engine/Source/Editor/BlueprintGraph/Private/K2Node_SpawnActorFromClass.cpp
--- a/Engine/Source/Editor/BlueprintGraph/Private/K2Node_SpawnActorFromClass.cpp
+++ b/Engine/Source/Editor/BlueprintGraph/Private/K2Node_SpawnActorFromClass.cpp
@@ -307,9 +307,6 @@ void UK2Node_SpawnActorFromClass::ExpandNode(class FKismetCompilerContext& Compi
 		static FString ValueParamName = FString(TEXT("Value"));
 		static FString PropertyNameParamName = FString(TEXT("PropertyName"));
 
-		const UEdGraphSchema_K2* Schema = CompilerContext.GetSchema();
-		(void)Schema;
-
 		UK2Node_SpawnActorFromClass* SpawnNode = this;
 		UEdGraphPin* SpawnNodeExec = SpawnNode->GetExecPin();
 		UEdGraphPin* SpawnNodeTransform = SpawnNode->GetSpawnTransformPin();
@@ -398,54 +395,6 @@ void UK2Node_SpawnActorFromClass::ExpandNode(class FKismetCompilerContext& Compi
 		// Get 'result' pin from 'begin spawn', this is the actual actor we want to set properties on
 		UEdGraphPin* LastThen = FKismetCompilerUtilities::GenerateAssignmentNodes(CompilerContext, SourceGraph, CallBeginSpawnNode, SpawnNode, CallBeginResult );
 
-#if 0
-		UEdGraphPin* LastThen = CallBeginSpawnNode->GetThenPin();
-
-		// Create 'set var by name' nodes and hook them up
-		for(int32 PinIdx=0; PinIdx < SpawnNode->Pins.Num(); PinIdx++)
-		{
-			// Only create 'set param by name' node if this pin is linked to something
-			UEdGraphPin* OrgPin = SpawnNode->Pins[PinIdx];
-			if(OrgPin->LinkedTo.Num() > 0 || OrgPin->DefaultValue != FString() )
-			{
-				UFunction* SetByNameFunction = Schema->FindSetVariableByNameFunction(OrgPin->PinType);
-				if(SetByNameFunction)
-				{
-					UK2Node_CallFunction* SetVarNode = NULL;
-					if(OrgPin->PinType.bIsArray)
-					{
-						SetVarNode = CompilerContext.SpawnIntermediateNode<UK2Node_CallArrayFunction>(SpawnNode, SourceGraph);
-					}
-					else
-					{
-						SetVarNode = CompilerContext.SpawnIntermediateNode<UK2Node_CallFunction>(SpawnNode, SourceGraph);
-					}
-					SetVarNode->SetFromFunction(SetByNameFunction);
-					SetVarNode->AllocateDefaultPins();
-
-					// Connect this node into the exec chain
-					Schema->TryCreateConnection(LastThen, SetVarNode->GetExecPin());
-					LastThen = SetVarNode->GetThenPin();
-
-					// Connect the new actor to the 'object' pin
-					UEdGraphPin* ObjectPin = SetVarNode->FindPinChecked(ObjectParamName);
-					CallBeginResult->MakeLinkTo(ObjectPin);
-
-					// Fill in literal for 'property name' pin - name of pin is property name
-					UEdGraphPin* PropertyNamePin = SetVarNode->FindPinChecked(PropertyNameParamName);
-					PropertyNamePin->DefaultValue = OrgPin->PinName;
-
-					// Move connection from the variable pin on the spawn node to the 'value' pin
-					UEdGraphPin* ValuePin = SetVarNode->FindPinChecked(ValueParamName);
-					CompilerContext.MovePinLinksToIntermediate(*OrgPin, *ValuePin);
-					SetVarNode->PinConnectionListChanged(ValuePin);
-
-					// PinSubCategoryObject is lost when we go from function back to pin when we have type information that is BP only (IE an enum created in the editor):
-					ValuePin->PinType.PinSubCategoryObject = OrgPin->PinType.PinSubCategoryObject;
-				}
-			}
-		}
-#endif
 		// Make exec connection between 'then' on last node and 'finish'
 		LastThen->MakeLinkTo(CallFinishExec);
 
