Display labels on split struct array get node

diff --git a/Engine/Source/Editor/BlueprintGraph/Private/EdGraphSchema_K2.cpp b/Engine/Source/Editor/BlueprintGraph/Private/EdGraphSchema_K2.cpp
--- a/Engine/Source/Editor/BlueprintGraph/Private/EdGraphSchema_K2.cpp
+++ b/Engine/Source/Editor/BlueprintGraph/Private/EdGraphSchema_K2.cpp
@@ -4555,6 +4555,7 @@ void UEdGraphSchema_K2::SplitPin(UEdGraphPin* Pin) const
 {
 	UScriptStruct* StructType = CastChecked<UScriptStruct>(Pin->PinType.PinSubCategoryObject.Get());
 	UEdGraphNode* GraphNode = Pin->GetOwningNode();
+	UK2Node* K2Node = Cast<UK2Node>(GraphNode);
 	UEdGraph* Graph = CastChecked<UEdGraph>(GraphNode->GetOuter());
 
 	const FScopedTransaction Transaction( LOCTEXT("SplitStructPin", "Split Struct Pin") );
@@ -4574,11 +4575,24 @@ void UEdGraphSchema_K2::SplitPin(UEdGraphPin* Pin) const
 			const FEdGraphPinType& ProtoPinType = ProtoPin->PinType;
 			UEdGraphPin* SubPin = GraphNode->CreatePin(Pin->Direction, ProtoPinType.PinCategory, ProtoPinType.PinSubCategory, ProtoPinType.PinSubCategoryObject.Get(), false, false, PinName);
 
-			if (!Pin->PinFriendlyName.IsEmpty() || !ProtoPin->PinFriendlyName.IsEmpty())
+			if (K2Node != nullptr && K2Node->ShouldDrawCompact())
+			{
+				if (Pin->ParentPin)
+				{
+					SubPin->PinFriendlyName = FText::FromString(FString::Printf(TEXT("%s %s")
+																, *Pin->PinFriendlyName.ToString()
+																, (ProtoPin->PinFriendlyName.IsEmpty() ? *ProtoPin->PinName : *ProtoPin->PinFriendlyName.ToString())));
+				}
+				else
+				{
+					SubPin->PinFriendlyName = FText::FromString(ProtoPin->PinFriendlyName.IsEmpty() ? ProtoPin->PinName : ProtoPin->PinFriendlyName.ToString());
+				}
+			}
+			else
 			{
 				SubPin->PinFriendlyName = FText::FromString(FString::Printf(TEXT("%s %s")
-                                                                                            , (Pin->PinFriendlyName.IsEmpty() ? *Pin->PinName : *Pin->PinFriendlyName.ToString())
-                                                                                            , (ProtoPin->PinFriendlyName.IsEmpty() ? *ProtoPin->PinName : *ProtoPin->PinFriendlyName.ToString())));
+                                                            , (Pin->PinFriendlyName.IsEmpty() ? *Pin->PinName : *Pin->PinFriendlyName.ToString())
+                                                            , (ProtoPin->PinFriendlyName.IsEmpty() ? *ProtoPin->PinName : *ProtoPin->PinFriendlyName.ToString())));
 			}
 
 			SubPin->DefaultValue = ProtoPin->DefaultValue;
diff --git a/Engine/Source/Editor/GraphEditor/Private/KismetNodes/SGraphNodeK2ArrayFunction.cpp b/Engine/Source/Editor/GraphEditor/Private/KismetNodes/SGraphNodeK2ArrayFunction.cpp
--- a/Engine/Source/Editor/GraphEditor/Private/KismetNodes/SGraphNodeK2ArrayFunction.cpp
+++ b/Engine/Source/Editor/GraphEditor/Private/KismetNodes/SGraphNodeK2ArrayFunction.cpp
@@ -167,14 +167,17 @@ void SGraphNodeK2ArrayFunction::UpdateGraphNode()
 		CreatePinWidgets();
 
 		// Hide pin labels
-		for (int32 PinIndex=0; PinIndex < this->InputPins.Num(); ++PinIndex)
+		for (auto InputPin : this->InputPins)
 		{
-			InputPins[PinIndex]->SetShowLabel(false);
+			InputPin->SetShowLabel(false);
 		}
 
-		for (int32 PinIndex = 0; PinIndex < this->OutputPins.Num(); ++PinIndex)
+		for (auto OutputPin : this->OutputPins)
 		{
-			OutputPins[PinIndex]->SetShowLabel(false);
+			if (OutputPin->GetPinObj()->ParentPin == nullptr)
+			{
+				OutputPin->SetShowLabel(false);
+			}
 		}
 	}
 }
