Missed Schema changes from CL 2170750

diff --git a/Engine/Source/Editor/BlueprintGraph/Private/EdGraphSchema_K2.cpp b/Engine/Source/Editor/BlueprintGraph/Private/EdGraphSchema_K2.cpp
--- a/Engine/Source/Editor/BlueprintGraph/Private/EdGraphSchema_K2.cpp
+++ b/Engine/Source/Editor/BlueprintGraph/Private/EdGraphSchema_K2.cpp
@@ -4723,6 +4723,43 @@ void UEdGraphSchema_K2::RecombinePin(UEdGraphPin* Pin) const
 	FBlueprintEditorUtils::MarkBlueprintAsModified(Blueprint);
 }
 
+void UEdGraphSchema_K2::ConfigureVarNode(UK2Node_Variable* InVarNode, FName InVariableName, UStruct* InVariableSource, UBlueprint* InTargetBlueprint)
+{
+	// See if this is a 'self context' (ie. blueprint class is owner (or child of owner) of dropped var class)
+	if ((InVariableSource == NULL) || InTargetBlueprint->SkeletonGeneratedClass->IsChildOf(InVariableSource))
+	{
+		InVarNode->VariableReference.SetSelfMember(InVariableName);
+	}
+	else if (InVariableSource->IsA(UStruct::StaticClass()))
+	{
+		InVarNode->VariableReference.SetLocalMember(InVariableName, InVariableSource, FBlueprintEditorUtils::FindLocalVariableGuidByName(InTargetBlueprint, InVariableName));
+	}
+	else
+	{
+		InVarNode->VariableReference.SetExternalMember(InVariableName, CastChecked<UClass>(InVariableSource));
+	}
+}
+
+UK2Node_VariableGet* UEdGraphSchema_K2::SpawnVariableGetNode(const FVector2D GraphPosition, class UEdGraph* ParentGraph, FName VariableName, UStruct* Source) const
+{
+	UK2Node_VariableGet* NodeTemplate = NewObject<UK2Node_VariableGet>();
+	UBlueprint* Blueprint = FBlueprintEditorUtils::FindBlueprintForGraph(ParentGraph);
+
+	UEdGraphSchema_K2::ConfigureVarNode(NodeTemplate, VariableName, Source, Blueprint);
+
+	return FEdGraphSchemaAction_K2NewNode::SpawnNodeFromTemplate<UK2Node_VariableGet>(ParentGraph, NodeTemplate, GraphPosition);
+}
+
+UK2Node_VariableSet* UEdGraphSchema_K2::SpawnVariableSetNode(const FVector2D GraphPosition, class UEdGraph* ParentGraph, FName VariableName, UStruct* Source) const
+{
+	UK2Node_VariableSet* NodeTemplate = NewObject<UK2Node_VariableSet>();
+	UBlueprint* Blueprint = FBlueprintEditorUtils::FindBlueprintForGraph(ParentGraph);
+
+	UEdGraphSchema_K2::ConfigureVarNode(NodeTemplate, VariableName, Source, Blueprint);
+
+	return FEdGraphSchemaAction_K2NewNode::SpawnNodeFromTemplate<UK2Node_VariableSet>(ParentGraph, NodeTemplate, GraphPosition);
+}
+
 /////////////////////////////////////////////////////
 
 #undef LOCTEXT_NAMESPACE
