Making the "RestrictedToClasses" metadata key a constant in FBlueprintMetadata (so it can be used in various places).

diff --git a/Engine/Source/Editor/BlueprintGraph/Classes/EdGraphSchema_K2.h b/Engine/Source/Editor/BlueprintGraph/Classes/EdGraphSchema_K2.h
--- a/Engine/Source/Editor/BlueprintGraph/Classes/EdGraphSchema_K2.h
+++ b/Engine/Source/Editor/BlueprintGraph/Classes/EdGraphSchema_K2.h
@@ -55,6 +55,9 @@ public:
 
 	/** If true, the class will be usable as a base for blueprints */
 	static const FName MD_IsBlueprintBase;
+	
+	/** A listing of classes that this class is accessible from (and only those classes, if present) */
+	static const FName MD_RestrictedToClasses;
 
 	//    function metadata
 	/** Specifies a UFUNCTION as Kismet protected, which can only be called from itself */
diff --git a/Engine/Source/Editor/BlueprintGraph/Private/EdGraphSchema_K2.cpp b/Engine/Source/Editor/BlueprintGraph/Private/EdGraphSchema_K2.cpp
--- a/Engine/Source/Editor/BlueprintGraph/Private/EdGraphSchema_K2.cpp
+++ b/Engine/Source/Editor/BlueprintGraph/Private/EdGraphSchema_K2.cpp
@@ -42,6 +42,7 @@ const FName FBlueprintMetadata::MD_NotAllowableBlueprintVariableType(TEXT("NotBl
 
 const FName FBlueprintMetadata::MD_BlueprintSpawnableComponent(TEXT("BlueprintSpawnableComponent"));
 const FName FBlueprintMetadata::MD_IsBlueprintBase(TEXT("IsBlueprintBase"));
+const FName FBlueprintMetadata::MD_RestrictedToClasses(TEXT("RestrictedToClasses"));
 
 const FName FBlueprintMetadata::MD_Protected(TEXT("BlueprintProtected"));
 const FName FBlueprintMetadata::MD_Latent(TEXT("Latent"));
@@ -667,24 +668,24 @@ bool UEdGraphSchema_K2::IsAllowableBlueprintVariableType(const class UClass* InC
 }
 
 bool UEdGraphSchema_K2::IsAllowableBlueprintVariableType(const class UScriptStruct *InStruct)
-		{
+{
 	if (auto UDStruct = Cast<const UUserDefinedStruct>(InStruct))
-		{
+	{
 		if (EUserDefinedStructureStatus::UDSS_UpToDate != UDStruct->Status.GetValue())
 		{
 			return false;
 		}
-				}
+	}
 	return InStruct && (InStruct->GetBoolMetaDataHierarchical(FBlueprintMetadata::MD_AllowableBlueprintVariableType));
-			}
+}
 
 bool UEdGraphSchema_K2::DoesGraphSupportImpureFunctions(const UEdGraph* InGraph) const
-			{
+{
 	const EGraphType GraphType = GetGraphType(InGraph);
 	const bool bAllowImpureFuncs = GraphType != GT_Animation; //@TODO: It's really more nuanced than this (e.g., in a function someone wants to write as pure)
 
 	return bAllowImpureFuncs;
-		}
+}
 
 bool UEdGraphSchema_K2::IsPropertyExposedOnSpawn(const UProperty* Property)
 			{
diff --git a/Engine/Source/Editor/BlueprintGraph/Private/K2ActionMenuBuilder.cpp b/Engine/Source/Editor/BlueprintGraph/Private/K2ActionMenuBuilder.cpp
--- a/Engine/Source/Editor/BlueprintGraph/Private/K2ActionMenuBuilder.cpp
+++ b/Engine/Source/Editor/BlueprintGraph/Private/K2ActionMenuBuilder.cpp
@@ -406,10 +406,9 @@ static void GetReferenceSetterItems(FBlueprintGraphActionListBuilder& ContextMen
  */
 static bool CanUseLibraryFunctionsForScope(UClass* LibraryClass, const UClass* ScopeClass)
 {
-	static const FName NAME_RestristedToClasses(TEXT("RestrictedToClasses"));
-	if (LibraryClass->HasMetaData(NAME_RestristedToClasses))
+	if (LibraryClass->HasMetaData(FBlueprintMetadata::MD_RestrictedToClasses))
 	{
-		const FString& ClassRestrictions = LibraryClass->GetMetaData(NAME_RestristedToClasses);
+		const FString& ClassRestrictions = LibraryClass->GetMetaData(FBlueprintMetadata::MD_RestrictedToClasses);
 		for (UClass* TestOwnerClass = (UClass*)ScopeClass; TestOwnerClass; TestOwnerClass = TestOwnerClass->GetSuperClass())
 		{
 			const bool bFound = (ClassRestrictions == TestOwnerClass->GetName()) || !!FCString::StrfindDelim(*ClassRestrictions, *ScopeClass->GetName(), TEXT(" "));
