Don't lose split transform pins when changing class on Spawn Actor from Class node

diff --git a/Engine/Source/Editor/BlueprintGraph/Private/K2Node_SpawnActorFromClass.cpp b/Engine/Source/Editor/BlueprintGraph/Private/K2Node_SpawnActorFromClass.cpp
--- a/Engine/Source/Editor/BlueprintGraph/Private/K2Node_SpawnActorFromClass.cpp
+++ b/Engine/Source/Editor/BlueprintGraph/Private/K2Node_SpawnActorFromClass.cpp
@@ -128,11 +128,20 @@ void UK2Node_SpawnActorFromClass::ReallocatePinsDuringReconstruction(TArray<UEdG
 		CreatePinsForClass(UseSpawnClass);
 	}
 }
-
 bool UK2Node_SpawnActorFromClass::IsSpawnVarPin(UEdGraphPin* Pin)
 {
 	const UEdGraphSchema_K2* K2Schema = GetDefault<UEdGraphSchema_K2>();
 
+	UEdGraphPin* ParentPin = Pin->ParentPin;
+	while (ParentPin)
+	{
+		if (ParentPin->PinName == FK2Node_SpawnActorFromClassHelper::SpawnTransformPinName)
+		{
+			return false;
+		}
+		ParentPin = ParentPin->ParentPin;
+	}
+
 	return(	Pin->PinName != K2Schema->PN_Execute &&
 			Pin->PinName != K2Schema->PN_Then &&
 			Pin->PinName != K2Schema->PN_ReturnValue &&
