(WIP) Start of the blueprint menu refactor. Building a database of blueprint actions to refrence when constructiong various menus.

diff --git a/Engine/Source/Editor/BlueprintGraph/Classes/BlueprintComponentNodeSpawner.h b/Engine/Source/Editor/BlueprintGraph/Classes/BlueprintComponentNodeSpawner.h
new file mode 100644
--- /dev/null
+++ b/Engine/Source/Editor/BlueprintGraph/Classes/BlueprintComponentNodeSpawner.h
@@ -0,0 +1,49 @@
+// Copyright 1998-2014 Epic Games, Inc. All Rights Reserved.
+
+#pragma once
+
+#include "BlueprintNodeSpawner.h"
+#include "BlueprintComponentNodeSpawner.generated.h"
+
+// Forward declarations
+class UActorComponent;
+
+/**
+ * Takes care of spawning UK2Node_AddComponent nodes. Acts as the "action" 
+ * portion of certain FBlueprintActionMenuItems. Evolved from 
+ * FEdGraphSchemaAction_K2AddComponent.
+ */
+UCLASS(Transient)
+class UBlueprintComponentNodeSpawner : public UBlueprintNodeSpawner
+{
+	GENERATED_UCLASS_BODY()
+
+public:
+	/**
+	 * Creates a new UBlueprintComponentNodeSpawner for the specified class. 
+	 * Does not do any compatibility checking to ensure that the class is 
+	 * viable as a spawnable component (do that before calling this).
+	 *
+	 * @param  ComponentClass	The component type you want spawned nodes to spawn.
+	 * @return A newly allocated instance of this class.
+	 */
+	static UBlueprintComponentNodeSpawner* Create(TSubclassOf<UActorComponent> const ComponentClass);
+
+	// UBlueprintNodeSpawner interface
+	virtual UEdGraphNode* Invoke(UEdGraph* ParentGraph) const override;
+	virtual FText    GetDefaultMenuName() const override;
+	virtual FText    GetDefaultMenuCategory() const override;
+	// End UBlueprintNodeSpawner interface
+	
+	/**
+	 * Retrieves the component class that this configures spawned nodes with.
+	 *
+	 * @return The component type that this class was initialized with.
+	 */
+	TSubclassOf<UActorComponent> GetComponentClass() const;
+
+private:
+	/** The component class to configure new nodes with. */
+	UPROPERTY()
+	TSubclassOf<UActorComponent> ComponentClass;
+};
diff --git a/Engine/Source/Editor/BlueprintGraph/Classes/BlueprintEventNodeSpawner.h b/Engine/Source/Editor/BlueprintGraph/Classes/BlueprintEventNodeSpawner.h
new file mode 100644
--- /dev/null
+++ b/Engine/Source/Editor/BlueprintGraph/Classes/BlueprintEventNodeSpawner.h
@@ -0,0 +1,49 @@
+// Copyright 1998-2014 Epic Games, Inc. All Rights Reserved.
+
+#pragma once
+
+#include "BlueprintNodeSpawner.h"
+#include "BlueprintEventNodeSpawner.generated.h"
+
+/**
+ * Takes care of spawning UK2Node_Event nodes. Acts as the "action" portion of
+ * certain FBlueprintActionMenuItems. Will not spawn a new event node if one
+ * associated with the specified function already exits (instead, Invoke() will
+ * return the existing one). Evolved from FEdGraphSchemaAction_K2AddEvent and 
+ * FEdGraphSchemaAction_K2ViewNode.
+ */
+UCLASS(Transient)
+class UBlueprintEventNodeSpawner : public UBlueprintNodeSpawner
+{
+	GENERATED_UCLASS_BODY()
+
+public:
+	/**
+	 * Creates a new UBlueprintEventNodeSpawner for the specified function.
+	 * Does not do any compatibility checking to ensure that the function is
+	 * viable as a blueprint event (do that before calling this).
+	 *
+	 * @param  EventFunc	The function you want assigned to new nodes.
+	 * @return A newly allocated instance of this class.
+	 */
+	static UBlueprintEventNodeSpawner* Create(UFunction const* const EventFunc);
+
+	// UBlueprintNodeSpawner interface
+	virtual UEdGraphNode* Invoke(UEdGraph* ParentGraph) const override;
+	virtual FText    GetDefaultMenuName() const override;
+	virtual FText    GetDefaultMenuCategory() const override;
+	// End UBlueprintNodeSpawner interface
+	
+	/**
+	 * Retrieves the function that this assigns to spawned nodes (defines the
+	 * event's signature).
+	 *
+	 * @return The event function that this class was initialized with.
+	 */
+	UFunction const* GetEventFunction() const;
+
+private:
+	/** The function to configure new nodes with. */
+	UPROPERTY()
+    UFunction const* EventFunc;
+};
diff --git a/Engine/Source/Editor/BlueprintGraph/Classes/BlueprintFunctionNodeSpawner.h b/Engine/Source/Editor/BlueprintGraph/Classes/BlueprintFunctionNodeSpawner.h
new file mode 100644
--- /dev/null
+++ b/Engine/Source/Editor/BlueprintGraph/Classes/BlueprintFunctionNodeSpawner.h
@@ -0,0 +1,48 @@
+// Copyright 1998-2014 Epic Games, Inc. All Rights Reserved.
+
+#pragma once
+
+#include "BlueprintNodeSpawner.h"
+#include "BlueprintFunctionNodeSpawner.generated.h"
+
+/**
+ * Takes care of spawning various UK2Node_CallFunction nodes. Acts as the 
+ * "action" portion of certain FBlueprintActionMenuItems. 
+ */
+UCLASS(Transient)
+class UBlueprintFunctionNodeSpawner : public UBlueprintNodeSpawner
+{
+	GENERATED_UCLASS_BODY()
+
+public:
+	/**
+	 * Creates a new UBlueprintFunctionNodeSpawner for the specified function. 
+	 * Does not do any compatibility checking to ensure that the function is 
+	 * viable as a blueprint function call (do that before calling this).
+	 *
+	 * @param  Function		The function you want assigned to new nodes.
+	 * @return A newly allocated instance of this class.
+	 */
+	BLUEPRINTGRAPH_API static UBlueprintFunctionNodeSpawner* Create(UFunction const* const Function);
+
+	// UBlueprintNodeSpawner interface
+	virtual UEdGraphNode* Invoke(UEdGraph* ParentGraph) const override;
+	virtual FText GetDefaultMenuName() const override;
+	virtual FText GetDefaultMenuCategory() const override;
+	virtual FText GetDefaultMenuTooltip() const override;
+	virtual FText GetDefaultSearchKeywords() const override;
+	// End UBlueprintNodeSpawner interface
+	
+	/**
+	 * Retrieves the function that this assigns to spawned nodes (defines the 
+	 * node's signature).
+	 *
+	 * @return The function that this class was initialized with.
+	 */
+	UFunction const* GetFunction() const;
+
+private:
+	/** The function to configure new nodes with. */
+	UPROPERTY()
+	UFunction const* Function;
+};
diff --git a/Engine/Source/Editor/BlueprintGraph/Classes/BlueprintPropertyNodeSpawner.h b/Engine/Source/Editor/BlueprintGraph/Classes/BlueprintPropertyNodeSpawner.h
new file mode 100644
--- /dev/null
+++ b/Engine/Source/Editor/BlueprintGraph/Classes/BlueprintPropertyNodeSpawner.h
@@ -0,0 +1,55 @@
+// Copyright 1998-2014 Epic Games, Inc. All Rights Reserved.
+
+#pragma once
+
+#include "BlueprintNodeSpawner.h"
+#include "BlueprintPropertyNodeSpawner.generated.h"
+
+/**
+ * Takes care of spawning various nodes associated with properties (like 
+ * variable getter/setters, or one of the various delegate nodes). Acts as the 
+ * "action" portion of certain FBlueprintActionMenuItems. Evolved from 
+ * FEdGraphSchemaAction_K2Var, FEdGraphSchemaAction_K2Delegate, etc.
+ */
+UCLASS(Transient)
+class UBlueprintPropertyNodeSpawner : public UBlueprintNodeSpawner
+{
+	GENERATED_UCLASS_BODY()
+
+public:
+	/**
+	 * Creates a new UBlueprintPropertyNodeSpawner for the specified property.
+	 * Does not do any compatibility checking to ensure that the property is
+	 * accessible from blueprints (do that before calling this).
+	 *
+	 * @param  Property		The property you want assigned to spawned nodes.
+	 * @param  NodeClass	The type of node you want spawned (can't distinct a getter from a setter from the property alone).
+	 * @return A newly allocated instance of this class.
+	 */
+	static UBlueprintPropertyNodeSpawner* Create(UProperty const* const Property, UClass* const NodeClass = nullptr);
+
+	// UBlueprintNodeSpawner interface
+	virtual UEdGraphNode* Invoke(UEdGraph* ParentGraph) const override;
+	virtual FText GetDefaultMenuName() const override;
+	virtual FText GetDefaultMenuCategory() const override;
+	// End UBlueprintNodeSpawner interface
+
+	/**
+	 * Retrieves the property that this assigns to spawned nodes.
+	 *
+	 * @return The property that this class was initialized with.
+	 */
+	UProperty const* GetProperty() const;
+
+	/**
+	 * Checks to see if the associated property is a delegate property.
+	 *
+	 * @return True if the property is a UMulticastDelegateProperty (otherwise false).
+	 */
+	bool IsDelegateProperty() const;
+	
+private:
+	/** The property to configure new nodes with. */
+	UPROPERTY()
+	UProperty const* Property;
+};
\ No newline at end of file
diff --git a/Engine/Source/Editor/BlueprintGraph/Classes/K2Node.h b/Engine/Source/Editor/BlueprintGraph/Classes/K2Node.h
--- a/Engine/Source/Editor/BlueprintGraph/Classes/K2Node.h
+++ b/Engine/Source/Editor/BlueprintGraph/Classes/K2Node.h
@@ -4,6 +4,8 @@
 #include "EdGraph/EdGraphNode.h"
 #include "K2Node.generated.h"
 
+class UBlueprintNodeSpawner;
+
 /** Helper struct to allow us to redirect properties and functions through renames and additionally between classes if necessary */
 struct FFieldRemapInfo
 {
@@ -284,6 +286,28 @@ class UK2Node : public UEdGraphNode
 	/** This function returns an arbitrary number of attributes that describe this node for analytics events */
 	BLUEPRINTGRAPH_API virtual void GetNodeAttributes( TArray<TKeyValuePair<FString, FString>>& OutNodeAttributes ) const;
 
+	/** 
+	 * Replacement for GetMenuEntries(). Override to add specific 
+	 * UBlueprintNodeSpawners pertaining to the sub-class type. Serves as an 
+	 * extensible way for new nodes, and game module nodes to add themselves to
+	 * context menus.
+	 *
+	 * @TODO: Not operational yet (stubbed in to prep for blueprint action menu refactor).
+	 *
+	 * @param  ActionListOut	The list to be populated with new spawners.
+	 */
+	virtual void GetMenuActions(TArray<UBlueprintNodeSpawner*>& ActionListOut) const {}
+
+	/**
+	 * Override to provide a default category for specific node types to be 
+	 * listed under.
+	 *
+	 * @TODO: Not operational yet (stubbed in to prep for blueprint action menu refactor).
+	 *
+	 * @return A localized category string (or an empty string if you want this node listed at the menu's root).
+	 */
+	virtual FText GetMenuCategory() const { return FText::GetEmpty(); }
+
 protected:
 	/** 
 	 * A mapping from old property and function names to new ones.  Get primed from INI files, and should contain entries for properties, functions, and delegates that get moved, so they can be fixed up
diff --git a/Engine/Source/Editor/BlueprintGraph/Private/BlueprintActionDatabase.cpp b/Engine/Source/Editor/BlueprintGraph/Private/BlueprintActionDatabase.cpp
new file mode 100644
--- /dev/null
+++ b/Engine/Source/Editor/BlueprintGraph/Private/BlueprintActionDatabase.cpp
@@ -0,0 +1,696 @@
+// Copyright 1998-2014 Epic Games, Inc. All Rights Reserved.
+
+#include "BlueprintGraphPrivatePCH.h"
+#include "BlueprintActionDatabase.h"
+#include "EdGraphSchema_K2.h"       // for CanUserKismetCallFunction()
+#include "KismetEditorUtilities.h"	// for IsClassABlueprintSkeleton(), IsClassABlueprintInterface(), GetBoundsForSelectedNodes(), etc.
+#include "BlueprintNodeSpawner.h"
+#include "BlueprintFunctionNodeSpawner.h"
+#include "BlueprintPropertyNodeSpawner.h"
+#include "BlueprintEventNodeSpawner.h"
+#include "BlueprintComponentNodeSpawner.h"
+
+// used below in FBlueprintNodeSpawnerFactory::MakeMacroNodeSpawner()
+#include "K2Node_MacroInstance.h"
+// used below in FBlueprintActionDatabaseImpl::AddClassEnumActions()
+#include "K2Node_GetNumEnumEntries.h"
+#include "K2Node_ForEachElementInEnum.h"
+#include "K2Node_EnumLiteral.h"
+#include "K2Node_CastByteToEnum.h"
+#include "K2Node_SwitchEnum.h"
+// used below in FBlueprintActionDatabaseImpl::AddClassStructActions()
+#include "K2Node_BreakStruct.h"
+#include "K2Node_MakeStruct.h"
+#include "K2Node_SetFieldsInStruct.h"
+// used below in FBlueprintNodeSpawnerFactory::MakeMessageNodeSpawner()
+#include "K2Node_Message.h"
+// used below in FBlueprintActionDatabaseImpl::AddClassPropertyActions()
+#include "K2Node_VariableGet.h"
+#include "K2Node_VariableSet.h"
+// used below in FBlueprintActionDatabaseImpl::AddClassCastActions()
+#include "K2Node_DynamicCast.h"
+#include "K2Node_ClassDynamicCast.h"
+// used below in FBlueprintNodeSpawnerFactory::MakeCommentNodeSpawner()
+#include "EdGraph/EdGraphNode_Comment.h"
+
+/*******************************************************************************
+ * FBlueprintNodeSpawnerFactory
+ ******************************************************************************/
+
+namespace FBlueprintNodeSpawnerFactory
+{
+	/**
+	 * Constructs a UK2Node_MacroInstance spawner. Evolved from 
+	 * FK2ActionMenuBuilder::AttachMacroGraphAction(). Sets up the spawner to 
+	 * set spawned nodes with the supplied macro.
+	 *
+	 * @param  MacroGraph	The macro you want spawned nodes referencing.
+	 * @return A new node-spawner, setup to spawn a UK2Node_MacroInstance.
+	 */
+	static UBlueprintNodeSpawner* MakeMacroNodeSpawner(UEdGraph* MacroGraph);
+
+	/**
+	 * A templatized method which constructs a node-spawner for various enum 
+	 * node types (any node with a public UEnum* Enum field). Takes the 
+	 * specified enum and applies it to the node post-spawn.
+	 *
+	 * @param  Enum					The enum you want set for the spawned node.
+	 * @param  ExtraSetupCallback	A callback for any further post-spawn customization (other than setting the node's Enum field).
+	 * @return A new node-spawner, setup to spawn some enum node (defined by the EnumNodeType template param).
+	 */
+	template<class EnumNodeType>
+	static UBlueprintNodeSpawner* MakeEnumNodeSpawner(UEnum* Enum, void(*ExtraSetupCallback)(EnumNodeType*) = nullptr);
+
+	/**
+	 * A templatized method which constructs a node-spawner for various struct 
+	 * node types (any node with a public UScriptStruct* StructType field). 
+	 * Takes the specified struct and applies it to the node post-spawn.
+	 *
+	 * @param  Struct	The struct you want set for the spawned node.
+	 * @return A new node-spawner, setup to spawn some struct node (defined by the StructNodeType template param).
+	 */
+	template<class StructNodeType>
+	static UBlueprintNodeSpawner* MakeStructNodeSpawner(UScriptStruct* Struct);
+
+	/**
+	 * Constructs a UK2Node_Message spawner. Sets up the spawner to set 
+	 * spawned nodes with the supplied function.
+	 *
+	 * @param  InterfaceFunction	The function you want spawned nodes referencing.
+	 * @return A new node-spawner, setup to spawn a UK2Node_Message.
+	 */
+	static UBlueprintNodeSpawner* MakeMessageNodeSpawner(UFunction* InterfaceFunction);
+
+	/**
+	 * Constructs a UEdGraphNode_Comment spawner. Since UEdGraphNode_Comment is
+	 * not a UK2Node then we can't have it create a spawner for itself (using 
+	 * UK2Node's GetMenuActions() method).
+	 *
+	 * @TODO:  Fix it so comment nodes spawned this way will properly position 
+	 *         themselves (FBlueprintActionMenuItem overrides positioning).
+	 *
+	 * @return A new node-spawner, setup to spawn a UEdGraphNode_Comment.
+	 */
+	static UBlueprintNodeSpawner* MakeCommentNodeSpawner();
+};
+
+//------------------------------------------------------------------------------
+static UBlueprintNodeSpawner* FBlueprintNodeSpawnerFactory::MakeMacroNodeSpawner(UEdGraph* MacroGraph)
+{
+	check(MacroGraph != nullptr);
+	check(MacroGraph->GetSchema()->GetGraphType(MacroGraph) == GT_Macro);
+
+	UBlueprintNodeSpawner* NodeSpawner = UBlueprintNodeSpawner::Create(UK2Node_MacroInstance::StaticClass());
+	check(NodeSpawner != nullptr);
+
+	auto CustomizeMacroNodeLambda = [](UEdGraphNode* NewNode, bool bIsTemplateNode, TWeakObjectPtr<UEdGraph> MacroGraph)
+	{
+		UK2Node_MacroInstance* MacroNode = CastChecked<UK2Node_MacroInstance>(NewNode);
+		if (MacroGraph.IsValid())
+		{
+			MacroNode->SetMacroGraph(MacroGraph.Get());
+		}
+	};
+
+	TWeakObjectPtr<UEdGraph> GraphPtr = MacroGraph;
+	NodeSpawner->CustomizeNodeDelegate = UBlueprintNodeSpawner::FCustomizeNodeDelegate::CreateStatic(CustomizeMacroNodeLambda, GraphPtr);
+
+	return NodeSpawner;
+}
+
+//------------------------------------------------------------------------------
+template<class EnumNodeType>
+static UBlueprintNodeSpawner* FBlueprintNodeSpawnerFactory::MakeEnumNodeSpawner(UEnum* Enum, void(*ExtraSetupCallback)(EnumNodeType*))
+{
+	UBlueprintNodeSpawner* NodeSpawner = UBlueprintNodeSpawner::Create(EnumNodeType::StaticClass());
+	check(NodeSpawner != nullptr);
+
+	DECLARE_DELEGATE_OneParam(FFurtherCustomizeNodeDelegate, EnumNodeType*);
+
+	auto CustomizeEnumNodeLambda = [](UEdGraphNode* NewNode, bool bIsTemplateNode, TWeakObjectPtr<UEnum> EnumPtr, FFurtherCustomizeNodeDelegate FurtherCustomizeDelegate)
+	{
+		EnumNodeType* EnumNode = CastChecked<EnumNodeType>(NewNode);
+
+		if (EnumPtr.IsValid())
+		{
+			EnumNode->Enum = EnumPtr.Get();
+		}
+
+		if (FurtherCustomizeDelegate.IsBound())
+		{
+			FurtherCustomizeDelegate.Execute(EnumNode);
+		}
+	};
+
+	FFurtherCustomizeNodeDelegate FurtherCustomizeDelegate;
+	if (ExtraSetupCallback != nullptr)
+	{
+		FurtherCustomizeDelegate = FFurtherCustomizeNodeDelegate::CreateStatic(ExtraSetupCallback);
+	}
+
+	TWeakObjectPtr<UEnum> EnumPtr = Enum;
+	NodeSpawner->CustomizeNodeDelegate = UBlueprintNodeSpawner::FCustomizeNodeDelegate::CreateStatic(CustomizeEnumNodeLambda, EnumPtr, FurtherCustomizeDelegate);
+
+	return NodeSpawner;
+}
+
+//------------------------------------------------------------------------------
+template<class StructNodeType>
+static UBlueprintNodeSpawner* FBlueprintNodeSpawnerFactory::MakeStructNodeSpawner(UScriptStruct* Struct)
+{
+	UBlueprintNodeSpawner* NodeSpawner = UBlueprintNodeSpawner::Create(StructNodeType::StaticClass());
+	check(NodeSpawner != nullptr);
+
+	auto CustomizeStructNodeLambda = [](UEdGraphNode* NewNode, bool bIsTemplateNode, TWeakObjectPtr<UScriptStruct> StructPtr)
+	{
+		StructNodeType* StructNode = CastChecked<StructNodeType>(NewNode);
+		if (StructPtr.IsValid())
+		{
+			StructNode->StructType = StructPtr.Get();
+		}
+	};
+
+	TWeakObjectPtr<UScriptStruct> StructPtr = Struct;
+	NodeSpawner->CustomizeNodeDelegate = UBlueprintNodeSpawner::FCustomizeNodeDelegate::CreateStatic(CustomizeStructNodeLambda, StructPtr);
+
+	return NodeSpawner;
+}
+
+//------------------------------------------------------------------------------
+static UBlueprintNodeSpawner* FBlueprintNodeSpawnerFactory::MakeMessageNodeSpawner(UFunction* InterfaceFunction)
+{
+	check(InterfaceFunction != nullptr);
+	check(FKismetEditorUtilities::IsClassABlueprintInterface(CastChecked<UClass>(InterfaceFunction->GetOuter())));
+
+	UBlueprintNodeSpawner* NodeSpawner = UBlueprintNodeSpawner::Create(UK2Node_Message::StaticClass());
+	check(NodeSpawner != nullptr);
+
+	auto CustomizeMessageNodeLambda = [](UEdGraphNode* NewNode, bool bIsTemplateNode, TWeakObjectPtr<UFunction> FunctionPtr)
+	{
+		UK2Node_Message* MessageNode = CastChecked<UK2Node_Message>(NewNode);
+		if (FunctionPtr.IsValid())
+		{
+			UFunction* Function = FunctionPtr.Get();
+			MessageNode->FunctionReference.SetExternalMember(Function->GetFName(), Cast<UClass>(Function->GetOuter()));
+		}
+	};
+
+	TWeakObjectPtr<UFunction> FunctionPtr = InterfaceFunction;
+	NodeSpawner->CustomizeNodeDelegate = UBlueprintNodeSpawner::FCustomizeNodeDelegate::CreateStatic(CustomizeMessageNodeLambda, FunctionPtr);
+
+	return NodeSpawner;
+}
+
+//------------------------------------------------------------------------------
+static UBlueprintNodeSpawner* FBlueprintNodeSpawnerFactory::MakeCommentNodeSpawner()
+{
+	UBlueprintNodeSpawner* NodeSpawner = UBlueprintNodeSpawner::Create(UEdGraphNode_Comment::StaticClass());
+	check(NodeSpawner != nullptr);
+
+	auto CustomizeMessageNodeLambda = [](UEdGraphNode* NewNode, bool bIsTemplateNode)
+	{
+		UEdGraphNode_Comment* CommentNode = CastChecked<UEdGraphNode_Comment>(NewNode);
+
+		UEdGraph* OuterGraph = CommentNode->GetGraph();
+		check(OuterGraph != nullptr);
+		UBlueprint* Blueprint = FBlueprintEditorUtils::FindBlueprintForGraph(OuterGraph);
+		check(Blueprint != nullptr);
+
+		FSlateRect Bounds;
+		FKismetEditorUtilities::GetBoundsForSelectedNodes(Blueprint, Bounds, 50.0f);
+		CommentNode->SetBounds(Bounds);
+	};
+
+	NodeSpawner->CustomizeNodeDelegate = UBlueprintNodeSpawner::FCustomizeNodeDelegate::CreateStatic(CustomizeMessageNodeLambda);
+
+	return NodeSpawner;
+}
+
+/*******************************************************************************
+ * Static FBlueprintActionDatabase Helpers
+ ******************************************************************************/
+
+namespace FBlueprintActionDatabaseImpl
+{
+	typedef FBlueprintActionDatabase::FActionList FActionList;
+
+	/**
+	 * Mimics UEdGraphSchema_K2::CanUserKismetAccessVariable(); however, this 
+	 * omits the filtering that CanUserKismetAccessVariable() does (saves that 
+	 * for later with FBlueprintActionFilter).
+	 * 
+	 * @param  Property		The property you want to check.
+	 * @return True if the property can be seen from a blueprint.
+	 */
+	static bool IsPropertyBlueprintVisible(UProperty const* const Property);
+
+	/**
+	 * Loops over all of the class's functions and creates a node-spawners for
+	 * any that are viable for blueprint use. Evolved from 
+	 * FK2ActionMenuBuilder::GetFuncNodesForClass(), plus a series of other
+	 * FK2ActionMenuBuilder methods (GetAllInterfaceMessageActions,
+	 * GetEventsForBlueprint, etc). 
+	 *
+	 * Ideally, any node that is constructed from a UFunction should go in here 
+	 * (so we only ever loop through the class's functions once). We handle
+	 * UK2Node_CallFunction alongside UK2Node_Event.
+	 *
+	 * @param  Class			The class whose functions you want node-spawners for.
+	 * @param  ActionListOut	The list you want populated with the new spawners.
+	 */
+	static void AddClassFunctionActions(UClass const* const Class, FActionList& ActionListOut);
+
+	/**
+	 * Loops over all of the class's properties and creates node-spawners for 
+	 * any that are viable for blueprint use. Evolved from certain parts of 
+	 * FK2ActionMenuBuilder::GetAllActionsForClass().
+	 *
+	 * @param  Class			The class whose properties you want node-spawners for.
+	 * @param  ActionListOut	The list you want populated with the new spawners.
+	 */
+	static void AddClassPropertyActions(UClass const* const Class, FActionList& ActionListOut);
+
+	/**
+	 * Evolved from FClassDynamicCastHelper::GetClassDynamicCastNodes(). If the
+	 * specified class is a viable blueprint variable type, then two cast nodes
+	 * are added for it (UK2Node_DynamicCast, and UK2Node_ClassDynamicCast).
+	 *
+	 * @param  Class			The class who you want cast nodes for (they cast to this class).
+	 * @param  ActionListOut	The list you want populated with the new spawners.
+	 */
+	static void AddClassCastActions(UClass* const Class, FActionList& ActionListOut);
+
+	/**
+	 * Evolved from K2ActionMenuBuilder's GetAddComponentClasses(). If the
+	 * specified class is a component type (and can be spawned), then a 
+	 * UK2Node_AddComponent spawner is created and added to ActionListOut.
+	 *
+	 * @param  Class			The class who you want a spawner for (the component class).
+	 * @param  ActionListOut	The list you want populated with the new spawner.
+	 */
+	static void AddComponentClassActions(UClass* const Class, FActionList& ActionListOut);
+
+	/**
+	 * Loops over the class's enums and creates node-spawners for any that are 
+	 * viable for blueprint use. Evolved from K2ActionMenuBuilder's 
+	 * GetEnumUtilitiesNodes() (as well as snippets from GetSwitchMenuItems).
+	 *
+	 * NOTE: This only accounts for enums that belong to this specific class, 
+	 *       not autonomous globally-scoped enums (ones with a UPackage outer). 
+	 *       Those enums should instead be accounted for in the appropriate 
+	 *       node's GetMenuActions(). For example, UK2Node_EnumLiteral should
+	 *       add UK2Node_EnumLiteral spawners for any global enums.
+	 *
+	 * @param  Class			The class whose enums you want node-spawners for.
+	 * @param  ActionListOut	The list you want populated with new spawners.
+	 */
+	static void AddClassEnumActions(UClass const* const Class, FActionList& ActionListOut);
+
+	/**
+	 * Loops over the class's structs and creates node-spawners for any that are 
+	 * viable for blueprint use. Evolved from 
+	 * FK2ActionMenuBuilder::GetStructActions().
+	 *
+	 * NOTE: This only captures structs that belong to this specific class, 
+	 *       not autonomous globally-scoped ones (those with a UPackage outer). 
+	 *       Those struct should instead be accounted for in the appropriate 
+	 *       node's GetMenuActions(). For example, UK2Node_MakeStruct should
+	 *       add UK2Node_MakeStruct spawners for any global structs.
+	 *
+	 * @param  Class			The class whose struct you want node-spawners for.
+	 * @param  ActionListOut	The list you want populated with new spawners.
+	 */
+	static void AddClassStructActions(UClass const* const Class, FActionList& ActionListOut);
+
+	/**
+	 * If the associated class is a blueprint generated class, then this will
+	 * loop over the blueprint's graphs and create any node-spawners associated
+	 * with those graphs (like UK2Node_MacroInstance spawners for macro graphs).
+	 *
+	 * @param  Class			The class which you want graph associated node-spawners for.
+	 * @param  ActionListOut	The list you want populated with new spawners.
+	 */
+	static void AddBlueprintGraphActions(UClass const* const Class, FActionList& ActionListOut);
+
+	/**
+	 * Emulates UEdGraphSchema::GetGraphContextActions(). If the supplied class  
+	 * is a node type, then it will query the node's CDO for any actions it 
+	 * wishes to add. This helps us keep the code in this file paired down, and 
+	 * makes it easily extensible for new node types. At the high level, this is 
+	 * for node types that aren't associated with an other class other, 
+	 * "autonomous" nodes that don't have any other way of being listed (this is
+	 * where enum/struct nodes missed by AddClassEnumActions()/AddClassStructActions()
+	 * would be added).
+	 * 
+	 * @param  Class			The class which you want node-spawners for.
+	 * @param  ActionListOut	The list you want populated with new spawners.
+	 */
+	static void AddAutonomousNodeActions(UClass* const Class, FActionList& ActionListOut);
+}
+
+//------------------------------------------------------------------------------
+static bool FBlueprintActionDatabaseImpl::IsPropertyBlueprintVisible(UProperty const* const Property)
+{
+	bool const bIsAccessible = Property->HasAllPropertyFlags(CPF_BlueprintVisible);
+
+	bool const bIsDelegate = Property->IsA(UMulticastDelegateProperty::StaticClass());
+	bool const bIsAssignableOrCallable = Property->HasAnyPropertyFlags(CPF_BlueprintAssignable | CPF_BlueprintCallable);
+
+	return !Property->HasAnyPropertyFlags(CPF_Parm) && (bIsAccessible || (bIsDelegate && bIsAssignableOrCallable));
+}
+
+//------------------------------------------------------------------------------
+static void FBlueprintActionDatabaseImpl::AddClassFunctionActions(UClass const* const Class, FActionList& ActionListOut)
+{
+	using namespace FBlueprintNodeSpawnerFactory; // for MakeMessageNodeSpawner()
+	check(Class != nullptr);
+
+	// loop over all the functions in the specified class; exclude-super because 
+	// we can always get the super functions by looking up that class separately 
+	for (TFieldIterator<UFunction> FunctionIt(Class, EFieldIteratorFlags::ExcludeSuper); FunctionIt; ++FunctionIt)
+	{
+		UFunction* Function = *FunctionIt;
+
+		if (UEdGraphSchema_K2::FunctionCanBePlacedAsEvent(Function))
+		{
+			if (UBlueprintEventNodeSpawner* NodeSpawner = UBlueprintEventNodeSpawner::Create(Function))
+			{
+				ActionListOut.Add(NodeSpawner);
+			}
+		}
+		
+		if (UEdGraphSchema_K2::CanUserKismetCallFunction(Function))
+		{
+			if (UBlueprintFunctionNodeSpawner* NodeSpawner = UBlueprintFunctionNodeSpawner::Create(Function))
+			{
+				ActionListOut.Add(NodeSpawner);
+			}
+
+			if (FKismetEditorUtilities::IsClassABlueprintInterface(Class))
+			{
+				ActionListOut.Add(MakeMessageNodeSpawner(Function));
+			}
+		}
+	}
+}
+
+//------------------------------------------------------------------------------
+static void FBlueprintActionDatabaseImpl::AddClassPropertyActions(UClass const* const Class, FActionList& ActionListOut)
+{
+	// loop over all the properties in the specified class; exclude-super because 
+	// we can always get the super properties by looking up that class separately 
+	for (TFieldIterator<UProperty> PropertyIt(Class, EFieldIteratorFlags::ExcludeSuper); PropertyIt; ++PropertyIt)
+	{
+		UProperty* Property = *PropertyIt;
+		if (!IsPropertyBlueprintVisible(Property))
+		{
+			continue;
+		}
+
+		// add an empty node spawner (that actually won't spawn anything) to 
+		// represent the property as a whole (used for the "collapsed" property
+		// entries in the palette; the ones that don't immediatly spawn a node, 
+		// but instead a sub-menu when you drop them in the graph)
+		//
+		// @TODO: consider moving this out into FBlueprintActionMenuBuilder
+		if (UBlueprintPropertyNodeSpawner* NodeSpawner = UBlueprintPropertyNodeSpawner::Create(Property))
+		{
+			ActionListOut.Add(NodeSpawner);
+		}
+
+ 		bool bIsDelegate = Property->IsA(UMulticastDelegateProperty::StaticClass());
+ 		if (bIsDelegate)
+ 		{
+			// @TODO: GetEventDispatcherNodesForClass()
+			//   UK2Node_AddDelegate
+			//   UK2Node_RemoveDelegate
+			//   UK2Node_ClearDelegate
+			//   UK2Node_CallDelegate
+
+			// @TODO: AddBoundEventActionsForClass()
+			//   UK2Node_ComponentBoundEvent 
+			//   UK2Node_ActorBoundEvent
+ 		}
+		else
+		{
+			UBlueprintPropertyNodeSpawner* GetterSpawner = UBlueprintPropertyNodeSpawner::Create(Property, UK2Node_VariableGet::StaticClass());
+			ActionListOut.Add(GetterSpawner);
+			UBlueprintPropertyNodeSpawner* SetterSpawner = UBlueprintPropertyNodeSpawner::Create(Property, UK2Node_VariableSet::StaticClass());
+			ActionListOut.Add(SetterSpawner);
+		}
+	}
+
+	// @TODO: if blueprint class, loop over function graphs and get local variables
+}
+
+//------------------------------------------------------------------------------
+static void FBlueprintActionDatabaseImpl::AddClassCastActions(UClass* const Class, FActionList& ActionListOut)
+{
+	UEdGraphSchema_K2 const* K2Schema = GetDefault<UEdGraphSchema_K2>();
+	bool bIsCastPermitted  = UEdGraphSchema_K2::IsAllowableBlueprintVariableType(Class);
+
+	if (bIsCastPermitted)
+	{
+		auto CustomizeCastNodeLambda = [](UEdGraphNode* NewNode, bool bIsTemplateNode, UClass* TargetType)
+		{
+			UK2Node_DynamicCast* CastNode = CastChecked<UK2Node_DynamicCast>(NewNode);
+			CastNode->TargetType = TargetType;
+		};
+
+		UBlueprintNodeSpawner* CastObjNodeSpawner = UBlueprintNodeSpawner::Create(UK2Node_DynamicCast::StaticClass());
+		CastObjNodeSpawner->CustomizeNodeDelegate = UBlueprintNodeSpawner::FCustomizeNodeDelegate::CreateStatic(CustomizeCastNodeLambda, Class);
+		ActionListOut.Add(CastObjNodeSpawner);
+
+		UBlueprintNodeSpawner* CastClassNodeSpawner = UBlueprintNodeSpawner::Create(UK2Node_ClassDynamicCast::StaticClass());
+		CastClassNodeSpawner->CustomizeNodeDelegate = CastObjNodeSpawner->CustomizeNodeDelegate;
+		ActionListOut.Add(CastClassNodeSpawner);
+	}
+}
+
+//------------------------------------------------------------------------------
+static void FBlueprintActionDatabaseImpl::AddComponentClassActions(UClass* const Class, FActionList& ActionListOut)
+{
+	bool bIsSpawnable = !Class->HasAnyClassFlags(CLASS_Abstract) && Class->HasMetaData(FBlueprintMetadata::MD_BlueprintSpawnableComponent);
+
+	if (bIsSpawnable && Class->IsChildOf(UActorComponent::StaticClass()))
+	{
+		if (UBlueprintComponentNodeSpawner* NodeSpawner = UBlueprintComponentNodeSpawner::Create(Class))
+		{
+			ActionListOut.Add(NodeSpawner);
+		}
+	}
+}
+
+//------------------------------------------------------------------------------
+static void FBlueprintActionDatabaseImpl::AddClassEnumActions(UClass const* const Class, FActionList& ActionListOut)
+{
+	using namespace FBlueprintNodeSpawnerFactory; // for MakeEnumNodeSpawner()
+
+	for (TFieldIterator<UEnum> EnumIt(Class, EFieldIteratorFlags::ExcludeSuper); EnumIt; ++EnumIt)
+	{
+		UEnum* Enum = *EnumIt;
+		if (!UEdGraphSchema_K2::IsAllowableBlueprintVariableType(Enum))
+		{
+			continue;
+		}
+
+		ActionListOut.Add(MakeEnumNodeSpawner<UK2Node_GetNumEnumEntries>(Enum));
+		ActionListOut.Add(MakeEnumNodeSpawner<UK2Node_ForEachElementInEnum>(Enum));
+		ActionListOut.Add(MakeEnumNodeSpawner<UK2Node_EnumLiteral>(Enum));
+
+		auto SetupEnumByteCastLambda = [](UK2Node_CastByteToEnum* NewNode)
+		{
+			NewNode->bSafe = true;
+		};
+		ActionListOut.Add(MakeEnumNodeSpawner<UK2Node_CastByteToEnum>(Enum, SetupEnumByteCastLambda));
+
+
+		auto SetupEnumSwitchLambda = [](UK2Node_SwitchEnum* NewNode)
+		{
+			// the node's 'Enum' field should have already been setup generically by MakeEnumNodeSpawner()
+			check(NewNode->Enum != nullptr);
+			NewNode->SetEnum(NewNode->Enum); // SetEnum() does more setup work than just setting the field
+		};
+		ActionListOut.Add(MakeEnumNodeSpawner<UK2Node_SwitchEnum>(Enum, SetupEnumSwitchLambda));
+	}
+
+	// @TODO: what about enum assets? how do we detect newly added ones and refresh those?
+	//    FEditorDelegates::LoadSelectedAssetsIfNeeded.Broadcast();
+}
+
+//------------------------------------------------------------------------------
+static void FBlueprintActionDatabaseImpl::AddClassStructActions(UClass const* const Class, FActionList& ActionListOut)
+{
+	using namespace FBlueprintNodeSpawnerFactory; // for MakeStructNodeSpawner()
+
+	for (TFieldIterator<UScriptStruct> StructIt(Class, EFieldIteratorFlags::ExcludeSuper); StructIt; ++StructIt)
+	{
+		UScriptStruct* Struct = (*StructIt);
+		if (!UEdGraphSchema_K2::IsAllowableBlueprintVariableType(Struct))
+		{
+			continue;
+		}
+
+		if (UK2Node_BreakStruct::CanBeBroken(Struct))
+		{
+			ActionListOut.Add(MakeStructNodeSpawner<UK2Node_BreakStruct>(Struct));
+		}
+
+		if (UK2Node_MakeStruct::CanBeMade(Struct))
+		{
+			ActionListOut.Add(MakeStructNodeSpawner<UK2Node_MakeStruct>(Struct));
+			ActionListOut.Add(MakeStructNodeSpawner<UK2Node_SetFieldsInStruct>(Struct));
+		}
+	}
+
+	// @TODO: what about struct assets? how do we detect newly added ones and refresh those?
+	//    FEditorDelegates::LoadSelectedAssetsIfNeeded.Broadcast();
+}
+
+//------------------------------------------------------------------------------
+static void FBlueprintActionDatabaseImpl::AddBlueprintGraphActions(UClass const* const Class, FActionList& ActionListOut)
+{
+	using namespace FBlueprintNodeSpawnerFactory; // for MakeMacroNodeSpawner()
+
+	if (UBlueprint const* const Blueprint = Cast<UBlueprint const>(Class->ClassGeneratedBy))
+	{
+		for (auto GraphIt = Blueprint->MacroGraphs.CreateConstIterator(); GraphIt; ++GraphIt)
+		{
+			ActionListOut.Add(MakeMacroNodeSpawner(*GraphIt));
+		}
+
+		// @TODO: local variables
+// 		for (auto GraphIt = Blueprint->FunctionGraphs.CreateConstIterator(); GraphIt; ++GraphIt)
+// 		{
+// 			UEdGraph* FunctionGraph = (*GraphIt);
+// 
+// 			TArray<UK2Node_FunctionEntry*> GraphEntryNodes;
+// 			FunctionGraph->GetNodesOfClass<UK2Node_FunctionEntry>(GraphEntryNodes);
+// 
+// 			for (UK2Node_FunctionEntry* FunctionEntry : GraphEntryNodes)
+// 			{
+// 				for (FBPVariableDescription const& LocalVar : FunctionEntry->LocalVariables)
+// 				{
+// 					ActionListOut.Add(MakeLocalVarNodeSpawner(LocalVar, VT_Getter));
+// 				}
+// 			}
+// 		}
+	}	
+}
+
+//------------------------------------------------------------------------------
+static void FBlueprintActionDatabaseImpl::AddAutonomousNodeActions(UClass* const Class, FActionList& ActionListOut)
+{
+	using namespace FBlueprintNodeSpawnerFactory; // for MakeCommentNodeSpawner()
+
+	if (Class->IsChildOf(UK2Node::StaticClass()) && !Class->HasAnyClassFlags(CLASS_Abstract))
+	{
+		UK2Node const* const NodeCDO = Class->GetDefaultObject<UK2Node>();
+		check(NodeCDO != nullptr);
+		NodeCDO->GetMenuActions(ActionListOut);
+	}
+	// unfortunately, UEdGraphNode_Comment is not a UK2Node and therefore cannot
+	// leverage UK2Node's GetMenuActions() function, so here we HACK it in
+	//
+	// @TODO: DO NOT follow this example! Do as I say, not as I do! If we need
+	//        to support other nodes in a similar way, then we should come up
+	//        with a better (more generalized) solution.
+	else if (Class == UEdGraphNode_Comment::StaticClass())
+	{
+		ActionListOut.Add(MakeCommentNodeSpawner());
+	}
+}
+
+/*******************************************************************************
+ * FBlueprintActionDatabase
+ ******************************************************************************/
+
+//------------------------------------------------------------------------------
+FBlueprintActionDatabase& FBlueprintActionDatabase::Get()
+{
+	static FBlueprintActionDatabase DatabaseInst;
+	if (DatabaseInst.ClassActions.Num() == 0)
+	{
+		// prime the database the first time we access it
+		DatabaseInst.RefreshAll();
+	}
+
+	return DatabaseInst;
+}
+
+//------------------------------------------------------------------------------
+void FBlueprintActionDatabase::RefreshAll()
+{
+	ClassActions.Empty();
+	for (TObjectIterator<UClass> ClassIt; ClassIt; ++ClassIt)
+	{
+		UClass* const Class = (*ClassIt);
+		RefreshClassActions(Class);
+	}
+}
+
+//------------------------------------------------------------------------------
+void FBlueprintActionDatabase::RefreshClassActions(UClass* const Class)
+{
+	using namespace FBlueprintActionDatabaseImpl;
+
+	check(Class != nullptr);
+	bool const bIsSkelClass   = FKismetEditorUtilities::IsClassABlueprintSkeleton(Class);
+	bool const bIsReinstClass = Class->HasAnyClassFlags(CLASS_NewerVersionExists);
+	
+	// make sure this class is not an intermediate that was part of some
+	// blueprint compile (if it is a SKEL or REINST class, then don't bother)
+	if (!bIsSkelClass && !bIsReinstClass)
+	{
+		FActionList& ClassActionList = ClassActions.FindOrAdd(Class);
+		ClassActionList.Empty();
+		
+		// class field actions (nodes that represent and perform actions on
+		// specific fields of the class... functions, properties, etc.)
+		{
+			AddClassFunctionActions(Class, ClassActionList);
+			AddClassPropertyActions(Class, ClassActionList);
+			AddClassEnumActions(Class, ClassActionList);
+			AddClassStructActions(Class, ClassActionList);
+		}
+
+		AddClassCastActions(Class, ClassActionList);
+		AddBlueprintGraphActions(Class, ClassActionList);
+		AddComponentClassActions(Class, ClassActionList);
+		
+		// accounts for the "autonomous" standalone nodes that can't be strongly
+		// associated with a particular class (besides the node's class)...
+		// think things like: comment nodes, custom events, the self node, etc.
+		//
+		// also should catch any actions dealing with global UFields (like
+		// global structs, enums, etc.; elements that wouldn't be caught
+		// normally when sifting through fields on all known classes)
+		AddAutonomousNodeActions(Class, ClassActionList);
+	}
+
+	// @TODO: account for all K2ActionMenuBuilder mathods...
+	//   GetLiteralsFromActorSelection() - UK2Node_Literal
+	//   GetAnimNotifyMenuItems()
+	//   GetMatineeControllers() - UK2Node_MatineeController
+	//   GetEventDispatcherNodesForClass()
+	//   GetBoundEventsFromActorSelection() - handle with filter
+	//   GetFunctionCallsOnSelectedActors() - handle with filter
+	//   GetAddComponentActionsUsingSelectedAssets() 
+	//   GetFunctionCallsOnSelectedComponents() - handle with filter
+	//   GetBoundEventsFromComponentSelection() - handle with filter 
+	//   GetPinAllowedMacros()
+	//   GetFuncNodesWithPinType()
+	//   GetVariableGettersSettersForClass()
+}
+
+//------------------------------------------------------------------------------
+FBlueprintActionDatabase::FClassActionMap const& FBlueprintActionDatabase::GetAllActions()
+{
+	// if this is the first time that we're querying for actions, generate the
+	// list before returning it
+	if (ClassActions.Num() == 0)
+	{
+		RefreshAll();
+	}
+	return ClassActions;
+}
diff --git a/Engine/Source/Editor/BlueprintGraph/Private/BlueprintComponentNodeSpawner.cpp b/Engine/Source/Editor/BlueprintGraph/Private/BlueprintComponentNodeSpawner.cpp
new file mode 100644
--- /dev/null
+++ b/Engine/Source/Editor/BlueprintGraph/Private/BlueprintComponentNodeSpawner.cpp
@@ -0,0 +1,122 @@
+// Copyright 1998-2014 Epic Games, Inc. All Rights Reserved.
+
+#include "BlueprintGraphPrivatePCH.h"
+#include "BlueprintComponentNodeSpawner.h"
+#include "K2Node_AddComponent.h"
+
+#define LOCTEXT_NAMESPACE "BlueprintComponenetNodeSpawner"
+
+/*******************************************************************************
+ * UBlueprintComponentNodeSpawner
+ ******************************************************************************/
+
+//------------------------------------------------------------------------------
+UBlueprintComponentNodeSpawner* UBlueprintComponentNodeSpawner::Create(TSubclassOf<UActorComponent> const ComponentClass)
+{
+	check(ComponentClass != nullptr);
+
+	UBlueprintComponentNodeSpawner* NodeSpawner = NewObject<UBlueprintComponentNodeSpawner>(GetTransientPackage());
+	NodeSpawner->ComponentClass = ComponentClass;
+	NodeSpawner->NodeClass      = UK2Node_AddComponent::StaticClass();
+
+	return NodeSpawner;
+}
+
+//------------------------------------------------------------------------------
+UBlueprintComponentNodeSpawner::UBlueprintComponentNodeSpawner(class FPostConstructInitializeProperties const& PCIP)
+	: Super(PCIP)
+{
+}
+
+//------------------------------------------------------------------------------
+// Evolved from a combination of FK2ActionMenuBuilder::CreateAddComponentAction()
+// and FEdGraphSchemaAction_K2AddComponent::PerformAction().
+UEdGraphNode* UBlueprintComponentNodeSpawner::Invoke(UEdGraph* ParentGraph) const
+{
+	UK2Node_AddComponent* NewNode = NewObject<UK2Node_AddComponent>(ParentGraph);
+	UBlueprint* Blueprint = NewNode->GetBlueprint();
+
+	UFunction* AddComponentFunc = FindFieldChecked<UFunction>(AActor::StaticClass(), UK2Node_AddComponent::GetAddComponentFunctionName());
+	NewNode->FunctionReference.SetFromField<UFunction>(AddComponentFunc, FBlueprintEditorUtils::IsActorBased(Blueprint));
+
+	bool bIsTemplateNode = ParentGraph->HasAnyFlags(RF_Transient);
+	if (CustomizeNodeDelegate.IsBound())
+	{
+		CustomizeNodeDelegate.Execute(NewNode, bIsTemplateNode);
+	}
+
+	if (!bIsTemplateNode)
+	{
+		NewNode->SetFlags(RF_Transactional);
+		NewNode->AllocateDefaultPins();
+
+		check(ComponentClass != nullptr);
+		UActorComponent* ComponentTemplate = ConstructObject<UActorComponent>(ComponentClass, Blueprint->GeneratedClass);
+		ComponentTemplate->SetFlags(RF_ArchetypeObject);
+
+		// @TODO: what is this is a template?
+		Blueprint->ComponentTemplates.Add(ComponentTemplate);
+
+		// set the name of the template as the default for the TemplateName param
+		UEdGraphPin* TemplateNamePin = NewNode->GetTemplateNamePinChecked();
+		if (TemplateNamePin != nullptr)
+		{
+			TemplateNamePin->DefaultValue = ComponentTemplate->GetName();
+		}
+
+		// set the return type to be the type of the template
+		UEdGraphPin* ReturnPin = NewNode->GetReturnValuePin();
+		if (ReturnPin != nullptr)
+		{
+			ReturnPin->PinType.PinSubCategoryObject = *ComponentClass;
+		}
+
+		// @TODO:
+// 		if (ComponentAsset != NULL)
+// 		{
+// 			FComponentAssetBrokerage::AssignAssetToComponent(ComponentTemplate, ComponentAsset);
+// 		}
+
+		NewNode->ReconstructNode();
+	}
+
+	return NewNode;
+}
+
+//------------------------------------------------------------------------------
+FText UBlueprintComponentNodeSpawner::GetDefaultMenuName() const
+{
+	check(ComponentClass != nullptr);
+	return FText::Format(LOCTEXT("AddComponentMenuName", "Add {0}"), FText::FromName(ComponentClass->GetFName()));
+}
+
+//------------------------------------------------------------------------------
+FText UBlueprintComponentNodeSpawner::GetDefaultMenuCategory() const
+{
+	check(ComponentClass != nullptr);
+
+	FText ClassGroup;
+	TArray<FString> ClassGroupNames;
+	ComponentClass->GetClassGroupNames(ClassGroupNames);
+
+	static FText const DefaultClassGroup(LOCTEXT("DefaultClassGroup", "Common"));
+	// 'Common' takes priority over other class groups
+	if (ClassGroupNames.Contains(DefaultClassGroup.ToString()) || (ClassGroupNames.Num() == 0)) 
+	{
+		ClassGroup = DefaultClassGroup;
+	}
+	else
+	{
+		ClassGroup = FText::FromString(ClassGroupNames[0]);
+	}
+
+	return FText::Format(LOCTEXT("ComponentCategory", "Add Component|{0}"), ClassGroup);
+}
+
+//------------------------------------------------------------------------------
+TSubclassOf<UActorComponent> UBlueprintComponentNodeSpawner::GetComponentClass() const
+{
+	return ComponentClass;
+}
+
+#undef LOCTEXT_NAMESPACE
diff --git a/Engine/Source/Editor/BlueprintGraph/Private/BlueprintEventNodeSpawner.cpp b/Engine/Source/Editor/BlueprintGraph/Private/BlueprintEventNodeSpawner.cpp
new file mode 100644
--- /dev/null
+++ b/Engine/Source/Editor/BlueprintGraph/Private/BlueprintEventNodeSpawner.cpp
@@ -0,0 +1,88 @@
+// Copyright 1998-2014 Epic Games, Inc. All Rights Reserved.
+
+#include "BlueprintGraphPrivatePCH.h"
+#include "BlueprintEventNodeSpawner.h"
+#include "EdGraphSchema_K2.h"	// for GetFriendlySignitureName()
+
+#define LOCTEXT_NAMESPACE "BlueprintEventNodeSpawner"
+
+/*******************************************************************************
+ * UBlueprintEventNodeSpawner
+ ******************************************************************************/
+
+//------------------------------------------------------------------------------
+UBlueprintEventNodeSpawner* UBlueprintEventNodeSpawner::Create(UFunction const* const EventFunc)
+{
+	check(EventFunc != nullptr);
+
+	UBlueprintEventNodeSpawner* NodeSpawner = NewObject<UBlueprintEventNodeSpawner>(GetTransientPackage());
+	NodeSpawner->EventFunc = EventFunc;
+	NodeSpawner->NodeClass = UK2Node_Event::StaticClass();
+
+	return NodeSpawner;
+}
+
+//------------------------------------------------------------------------------
+UBlueprintEventNodeSpawner::UBlueprintEventNodeSpawner(class FPostConstructInitializeProperties const& PCIP)
+	: Super(PCIP)
+	, EventFunc(nullptr)
+{
+}
+
+//------------------------------------------------------------------------------
+UEdGraphNode* UBlueprintEventNodeSpawner::Invoke(UEdGraph* ParentGraph) const
+{
+	check(ParentGraph != nullptr);
+	UBlueprint* Blueprint = FBlueprintEditorUtils::FindBlueprintForGraphChecked(ParentGraph);
+
+	check(EventFunc != nullptr);
+
+	UClass* ClassOwner = CastChecked<UClass>(EventFunc->GetOuter());
+	// look to see if a node for this event already exists (only one node is
+	// allowed per event, per blueprint)
+	UK2Node_Event* EventNode = FBlueprintEditorUtils::FindOverrideForFunction(Blueprint, ClassOwner, EventFunc->GetFName());
+
+	// if there is no existing node, then we can happily spawn one into the graph
+	if (EventNode == nullptr)
+	{
+		EventNode = NewObject<UK2Node_Event>(ParentGraph);
+		EventNode->EventSignatureName  = EventFunc->GetFName();
+		EventNode->EventSignatureClass = ClassOwner;
+		EventNode->bOverrideFunction   = true;
+
+		bool bIsTemplateNode = ParentGraph->HasAnyFlags(RF_Transient);
+		if (!bIsTemplateNode)
+		{
+			EventNode->SetFlags(RF_Transactional);
+			EventNode->AllocateDefaultPins();
+		}
+	}
+	// else, a node for this event already exists, and we should return that 
+	// (the FBlueprintActionMenuItem should detect this and focus in on it).
+
+	return EventNode;
+}
+
+//------------------------------------------------------------------------------
+FText UBlueprintEventNodeSpawner::GetDefaultMenuName() const
+{
+	check(EventFunc != nullptr);
+	
+	FString const EventName = UEdGraphSchema_K2::GetFriendlySignitureName(EventFunc);
+	return FText::Format(LOCTEXT("EventWithSignatureName", "Event {0}"), FText::FromString(EventName));
+}
+
+//------------------------------------------------------------------------------
+FText UBlueprintEventNodeSpawner::GetDefaultMenuCategory() const
+{
+	// @TODO: 
+	return Super::GetDefaultMenuCategory();
+}
+
+//------------------------------------------------------------------------------
+UFunction const* UBlueprintEventNodeSpawner::GetEventFunction() const
+{
+	return EventFunc;
+}
+
+#undef LOCTEXT_NAMESPACE
diff --git a/Engine/Source/Editor/BlueprintGraph/Private/BlueprintFunctionNodeSpawner.cpp b/Engine/Source/Editor/BlueprintGraph/Private/BlueprintFunctionNodeSpawner.cpp
new file mode 100644
--- /dev/null
+++ b/Engine/Source/Editor/BlueprintGraph/Private/BlueprintFunctionNodeSpawner.cpp
@@ -0,0 +1,107 @@
+// Copyright 1998-2014 Epic Games, Inc. All Rights Reserved.
+
+#include "BlueprintGraphPrivatePCH.h"
+#include "BlueprintFunctionNodeSpawner.h"
+#include "EdGraphSchema_K2.h"	// for FBlueprintMetadata
+
+/*******************************************************************************
+ * UBlueprintFunctionNodeSpawner
+ ******************************************************************************/
+
+//------------------------------------------------------------------------------
+// Evolved from FK2ActionMenuBuilder::AddSpawnInfoForFunction()
+UBlueprintFunctionNodeSpawner* UBlueprintFunctionNodeSpawner::Create(UFunction const* const Function)
+{
+	check(Function != nullptr);
+
+	bool const bIsPure = Function->HasAllFunctionFlags(FUNC_BlueprintPure);
+	bool const bHasArrayPointerParms = Function->HasMetaData(TEXT("ArrayParm"));
+	bool const bIsCommutativeAssociativeBinaryOp = Function->HasMetaData(FBlueprintMetadata::MD_CommutativeAssociativeBinaryOperator);
+	bool const bIsMaterialParamCollectionFunc = Function->HasMetaData(FBlueprintMetadata::MD_MaterialParameterCollectionFunction);
+
+	UBlueprintFunctionNodeSpawner* NodeSpawner = NewObject<UBlueprintFunctionNodeSpawner>(GetTransientPackage());
+	NodeSpawner->Function  = Function;
+
+	if (bIsCommutativeAssociativeBinaryOp && bIsPure)
+	{
+		NodeSpawner->NodeClass = UK2Node_CommutativeAssociativeBinaryOperator::StaticClass();
+	}
+	else if (bIsMaterialParamCollectionFunc)
+	{
+		NodeSpawner->NodeClass = UK2Node_CallMaterialParameterCollectionFunction::StaticClass();
+	}
+	// @TODO:
+	//   else if CallOnMember => UK2Node_CallFunctionOnMember
+	//   else if bIsParentContext => UK2Node_CallParentFunction
+	else if (bHasArrayPointerParms)
+	{
+		NodeSpawner->NodeClass = UK2Node_CallArrayFunction::StaticClass();
+	}
+	else
+	{
+		NodeSpawner->NodeClass = UK2Node_CallFunction::StaticClass();
+	}
+
+	return NodeSpawner;
+}
+
+//------------------------------------------------------------------------------
+UBlueprintFunctionNodeSpawner::UBlueprintFunctionNodeSpawner(class FPostConstructInitializeProperties const& PCIP)
+	: Super(PCIP)
+	, Function(nullptr)
+{
+}
+
+//------------------------------------------------------------------------------
+UEdGraphNode* UBlueprintFunctionNodeSpawner::Invoke(UEdGraph* ParentGraph) const
+{
+	UK2Node_CallFunction* NewNode = CastChecked<UK2Node_CallFunction>(Super::Invoke(ParentGraph));
+	NewNode->SetFromFunction(Function);
+	
+	bool bIsTemplateNode = ParentGraph->HasAnyFlags(RF_Transient);
+	if (CustomizeNodeDelegate.IsBound())
+	{
+		CustomizeNodeDelegate.Execute(NewNode, bIsTemplateNode);
+	}
+	if (!bIsTemplateNode)
+	{
+		NewNode->SetFlags(RF_Transactional);
+		NewNode->AllocateDefaultPins();
+	}
+ 	
+	return NewNode;
+}
+
+//------------------------------------------------------------------------------
+FText UBlueprintFunctionNodeSpawner::GetDefaultMenuName() const
+{
+	check(Function != nullptr);
+	return FText::FromString(UK2Node_CallFunction::GetUserFacingFunctionName(Function));
+}
+
+//------------------------------------------------------------------------------
+FText UBlueprintFunctionNodeSpawner::GetDefaultMenuCategory() const
+{
+	check(Function != nullptr);
+	return FText::FromString(UK2Node_CallFunction::GetDefaultCategoryForFunction(Function, TEXT("")));
+}
+
+//------------------------------------------------------------------------------
+FText UBlueprintFunctionNodeSpawner::GetDefaultMenuTooltip() const
+{
+	check(Function != nullptr);
+	return FText::FromString(UK2Node_CallFunction::GetDefaultTooltipForFunction(Function));
+}
+
+//------------------------------------------------------------------------------
+FText UBlueprintFunctionNodeSpawner::GetDefaultSearchKeywords() const
+{
+	check(Function != nullptr);
+	return FText::FromString(UK2Node_CallFunction::GetKeywordsForFunction(Function));
+}
+
+//------------------------------------------------------------------------------
+UFunction const* UBlueprintFunctionNodeSpawner::GetFunction() const
+{
+	return Function;
+}
diff --git a/Engine/Source/Editor/BlueprintGraph/Private/BlueprintNodeSpawner.cpp b/Engine/Source/Editor/BlueprintGraph/Private/BlueprintNodeSpawner.cpp
new file mode 100644
--- /dev/null
+++ b/Engine/Source/Editor/BlueprintGraph/Private/BlueprintNodeSpawner.cpp
@@ -0,0 +1,112 @@
+// Copyright 1998-2014 Epic Games, Inc. All Rights Reserved.
+
+#include "BlueprintGraphPrivatePCH.h"
+#include "BlueprintNodeSpawner.h"
+
+/*******************************************************************************
+ * UBlueprintNodeSpawner
+*******************************************************************************/
+
+//------------------------------------------------------------------------------
+UBlueprintNodeSpawner* UBlueprintNodeSpawner::Create(UClass* const NodeClass, FCustomizeNodeDelegate CustomizeNodeDelegate/* = FCustomizeNodeDelegate()*/)
+{
+	check(NodeClass != nullptr);
+	check(NodeClass->IsChildOf<UEdGraphNode>());
+	
+	UBlueprintNodeSpawner* NodeSpawner = NewObject<UBlueprintNodeSpawner>(GetTransientPackage());
+	NodeSpawner->NodeClass = NodeClass;
+	NodeSpawner->CustomizeNodeDelegate = CustomizeNodeDelegate;
+	
+	return NodeSpawner;
+}
+
+//------------------------------------------------------------------------------
+UBlueprintNodeSpawner::UBlueprintNodeSpawner(class FPostConstructInitializeProperties const& PCIP)
+	: Super(PCIP)
+	, CachedNodeTemplate(nullptr)
+{
+}
+
+//------------------------------------------------------------------------------
+UEdGraphNode* UBlueprintNodeSpawner::Invoke(UEdGraph* ParentGraph) const
+{
+	UEdGraphNode* NewNode = nullptr;	
+	if (NodeClass != nullptr)
+	{
+		NewNode = NewObject<UEdGraphNode>(ParentGraph, NodeClass);
+		check(NewNode != nullptr);
+
+		bool bIsTemplateNode = ParentGraph->HasAnyFlags(RF_Transient);		
+		if (CustomizeNodeDelegate.IsBound())
+		{
+			CustomizeNodeDelegate.Execute(NewNode, bIsTemplateNode);
+		}
+
+		if (!bIsTemplateNode)
+		{
+			NewNode->SetFlags(RF_Transactional);
+			NewNode->AllocateDefaultPins();
+		}
+	}
+
+	return NewNode;
+}
+
+//------------------------------------------------------------------------------
+FText UBlueprintNodeSpawner::GetDefaultMenuName() const
+{
+	// DO NOT make a template node and query it here (the separate ui building
+	// code can do that if it likes)... this is meant to be an overridable
+	// alternative to that (for perf reasons)
+	return FText::GetEmpty();
+}
+
+//------------------------------------------------------------------------------
+FText UBlueprintNodeSpawner::GetDefaultMenuCategory() const
+{
+	// DO NOT make a template node and query it here (the separate ui building
+	// code can do that if it likes)... this is meant to be an overridable
+	// alternative to that (for perf reasons)
+	return FText::GetEmpty();
+}
+
+//------------------------------------------------------------------------------
+FText UBlueprintNodeSpawner::GetDefaultMenuTooltip() const
+{
+	// DO NOT make a template node and query it here (the separate ui building
+	// code can do that if it likes)... this is meant to be an overridable
+	// alternative to that (for perf reasons)
+	return FText::GetEmpty();
+}
+
+//------------------------------------------------------------------------------
+FText UBlueprintNodeSpawner::GetDefaultSearchKeywords() const
+{
+	// DO NOT make a template node and query it here (the separate ui building
+	// code can do that if it likes)... this is meant to be an overridable
+	// alternative to that (for perf reasons)
+	return FText::GetEmpty();
+}
+
+//------------------------------------------------------------------------------
+UEdGraphNode* UBlueprintNodeSpawner::MakeTemplateNode(UEdGraph* Outer) const
+{
+	check(Outer != nullptr);
+	check(Outer->HasAnyFlags(RF_Transient));
+	
+	// @TODO: Do we need to respawn when the outers don't match? If not, this'll
+	//        save time for subsequent menu generation... if the old outer gets
+	//        destroyed then CachedNodeTemplate should be nulled (so we should
+	//        be covered on that front)... maybe we should just check the
+	//        graph's class type (and graph type)
+	if (CachedNodeTemplate == nullptr)// || (CachedNodeTemplate->GetOuter() != Outer))
+	{
+		CachedNodeTemplate = Invoke(Outer);
+		if (CachedNodeTemplate != nullptr)
+		{
+			CachedNodeTemplate->SetFlags(RF_Transient | RF_ArchetypeObject);
+		}
+	}
+	
+	return CachedNodeTemplate;
+}
diff --git a/Engine/Source/Editor/BlueprintGraph/Private/BlueprintPropertyNodeSpawner.cpp b/Engine/Source/Editor/BlueprintGraph/Private/BlueprintPropertyNodeSpawner.cpp
new file mode 100644
--- /dev/null
+++ b/Engine/Source/Editor/BlueprintGraph/Private/BlueprintPropertyNodeSpawner.cpp
@@ -0,0 +1,117 @@
+// Copyright 1998-2014 Epic Games, Inc. All Rights Reserved.
+
+#include "BlueprintGraphPrivatePCH.h"
+#include "BlueprintPropertyNodeSpawner.h"
+#include "K2Node_Variable.h"
+#include "EditorStyleSettings.h"	// for bShowFriendlyNames
+#include "ObjectEditorUtils.h"		// for GetCategory()
+
+#define LOCTEXT_NAMESPACE "BlueprintPropertyNodeSpawner"
+
+/*******************************************************************************
+ * UBlueprintPropertyNodeSpawner
+ ******************************************************************************/
+
+//------------------------------------------------------------------------------
+UBlueprintPropertyNodeSpawner* UBlueprintPropertyNodeSpawner::Create(UProperty const* const Property, UClass* const NodeClass)
+{
+	check(Property != nullptr);
+
+	UBlueprintPropertyNodeSpawner* NodeSpawner = NewObject<UBlueprintPropertyNodeSpawner>(GetTransientPackage());
+	NodeSpawner->Property  = Property;
+	NodeSpawner->NodeClass = NodeClass;
+
+	return NodeSpawner;
+}
+
+//------------------------------------------------------------------------------
+UBlueprintPropertyNodeSpawner::UBlueprintPropertyNodeSpawner(class FPostConstructInitializeProperties const& PCIP)
+	: Super(PCIP)
+	, Property(nullptr)
+{
+}
+
+//------------------------------------------------------------------------------
+UEdGraphNode* UBlueprintPropertyNodeSpawner::Invoke(UEdGraph* ParentGraph) const
+{
+	// if a NodeClass was set, then this should spawn something
+	UEdGraphNode* NewNode = Super::Invoke(ParentGraph);
+	check(Property != nullptr);	
+	
+	if (UK2Node_Variable* VarNode = Cast<UK2Node_Variable>(NewNode))
+	{
+		UBlueprint* Blueprint = VarNode->GetBlueprint();
+
+		UClass* BlueprintClass = Blueprint->GeneratedClass;
+		UClass* PropertyClass = Cast<UClass>(Property->GetOuterUField());
+
+		bool const bIsSelfContext = BlueprintClass->IsChildOf(PropertyClass);
+		VarNode->SetFromProperty(Property, bIsSelfContext);
+	}
+	// else, without a NodeClass explicitly specified, we don't know what kind  
+	// of property node to spawn (for example: a getter versus a setter)... 
+	// we can use a null NodeClass to represent the property itself (an 
+	// accumulation of all the nodes that it COULD spawn), and the ui can then 
+	// provide the user with a sub-menu to pick from (but it is up to some 
+	// external system to interpret/use it that way)
+
+	return NewNode;
+}
+
+//------------------------------------------------------------------------------
+FText UBlueprintPropertyNodeSpawner::GetDefaultMenuName() const
+{
+	check(Property != nullptr);
+	
+	FText MenuName = FText::FromName(Property->GetFName());
+	if (GetDefault<UEditorStyleSettings>()->bShowFriendlyNames)
+	{
+		MenuName = FText::FromString(UEditorEngine::GetFriendlyName(Property));
+	}
+
+	if (NodeClass != nullptr)
+	{
+		if (NodeClass->IsChildOf<UK2Node_VariableGet>())
+		{
+			MenuName = FText::Format(LOCTEXT("GetterMenuName", "Get {0}"), MenuName);
+		}
+		else if (NodeClass->IsChildOf<UK2Node_VariableSet>())
+		{
+			MenuName = FText::Format(LOCTEXT("SetterMenuName", "Set {0}"), MenuName);
+		}
+	}
+
+	return MenuName;
+}
+
+//------------------------------------------------------------------------------
+FText UBlueprintPropertyNodeSpawner::GetDefaultMenuCategory() const
+{
+	check(Property != nullptr);
+
+	FText PropertyCategory = FText::FromString(FObjectEditorUtils::GetCategory(Property));
+	if (!IsDelegateProperty())
+	{
+		PropertyCategory = FText::Format(LOCTEXT("VariableCategory", "Variables|{0}"), PropertyCategory);
+	}
+	else if (PropertyCategory.IsEmpty())
+	{
+		PropertyCategory = LOCTEXT("DefaultDelegateCategory", "Event Dispatchers");
+	}
+	return PropertyCategory;
+}
+
+//------------------------------------------------------------------------------
+UProperty const* UBlueprintPropertyNodeSpawner::GetProperty() const
+{
+	return Property;
+}
+
+//------------------------------------------------------------------------------
+bool UBlueprintPropertyNodeSpawner::IsDelegateProperty() const
+{
+	check(Property != nullptr);
+	return Property->IsA(UMulticastDelegateProperty::StaticClass());
+}
+
+#undef LOCTEXT_NAMESPACE
diff --git a/Engine/Source/Editor/BlueprintGraph/Public/BlueprintActionDatabase.h b/Engine/Source/Editor/BlueprintGraph/Public/BlueprintActionDatabase.h
new file mode 100644
--- /dev/null
+++ b/Engine/Source/Editor/BlueprintGraph/Public/BlueprintActionDatabase.h
@@ -0,0 +1,79 @@
+// Copyright 1998-2014 Epic Games, Inc. All Rights Reserved.
+
+#pragma once
+
+// Forward declarations
+class UBlueprint;
+class FBlueprintActionFilter;
+class UBlueprintNodeSpawner;
+
+/**
+ * Serves as a container for all available blueprint actions (no matter the 
+ * type of blueprint/graph they belong in). The actions stored here are not tied
+ * to a specific ui menu; each action is a UBlueprintNodeSpawner which is 
+ * charged with spawning a specific node type. Should be set up in a way where
+ * class specific actions are refreshed when the associated class is regenerated.
+ * 
+ * @TODO:  Hook up to handle class recompile events, along with enum/struct asset creation events.
+ */
+class BLUEPRINTGRAPH_API FBlueprintActionDatabase
+{
+public:
+	typedef TArray<UBlueprintNodeSpawner*> FActionList;
+	typedef TMap<UClass*, FActionList>     FClassActionMap;
+
+	/**
+	 * Ensures that the database singleton is initialized and populated.
+	 */
+	static void Prime() { Get(); }
+
+	/**
+	 * Getter to access the database singleton. Will populate the database first 
+	 * if this is the first time accessing it.
+	 *
+	 * @return The singleton instance of FBlueprintActionDatabase.
+	 */
+	static FBlueprintActionDatabase& Get();
+
+	/**
+	 * Populates the action database from scratch. Loops over every known class
+	 * and records a set of node-spawners associated with each.
+	 */
+	void RefreshAll();
+
+	/**
+	 * Finds the database entry for the specified class and wipes it, 
+	 * repopulating it with a fresh set of associated node-spawners.
+	 *
+	 * @param  Class	The class entry you want rebuilt.
+	 */
+	void RefreshClassActions(UClass* const Class);
+	
+	/**
+	 * Will populate the database first if it hasn't been created yet, and then 
+	 * returns it in its entirety. 
+	 * 
+	 * Each node spawner is categorized by a class. A spawner that corresponds
+	 * to a specific class field (like a function, property, enum, etc.) will
+	 * be listed under that field's class outer. Remaining spawners that can't 
+	 * be categorized this way will be filed by node type; for example: with 
+	 * global enums (enums that don't belong to a specific class), the enum's 
+	 * UK2Node_EnumLiteral spawner is found under the UK2Node_EnumLiteral class.
+	 *
+	 * @return The entire action database, which maps specific classes to arrays of associated node-spawners.
+	 */
+	FClassActionMap const& GetAllActions();
+	
+private:
+	/** Private constructor for singleton purposes. */
+	FBlueprintActionDatabase() {} 
+
+	/** 
+	 * A map of associated node-spawners for each class. A spawner that 
+	 * corresponds to a specific class field (like a function, property, enum, 
+	 * etc.) will be mapped under that field's class outer. Other spawners (that
+	 * can't be associated with a class outer), will be filed under the desired
+	 * node's type.
+	 */
+	FClassActionMap ClassActions;
+};
diff --git a/Engine/Source/Editor/BlueprintGraph/Public/BlueprintNodeSpawner.h b/Engine/Source/Editor/BlueprintGraph/Public/BlueprintNodeSpawner.h
new file mode 100644
--- /dev/null
+++ b/Engine/Source/Editor/BlueprintGraph/Public/BlueprintNodeSpawner.h
@@ -0,0 +1,147 @@
+// Copyright 1998-2014 Epic Games, Inc. All Rights Reserved.
+
+#pragma once
+
+#include "BlueprintNodeSpawner.generated.h"
+
+// Forward declarations
+class UEdGraph;
+class UEdGraphNode;
+
+/**
+ * Intended to be wrapped and used by FBlueprintActionMenuItem. Rather than 
+ * sub-classing the menu item, we choose to subclass this instead (for 
+ * different node types). That way, we get the type inference that comes with  
+ * UObjects (and we don't have to continuously compare identification strings). 
+ *
+ * Additionally, this decouples the "action" from ui settings (like the tooltip,
+ * menu category, etc.). It gives us more leeway in how we choose to present 
+ * these "actions".
+ */
+UCLASS(Transient)
+class UBlueprintNodeSpawner : public UObject
+{
+	GENERATED_UCLASS_BODY()
+	DECLARE_DELEGATE_TwoParams(FCustomizeNodeDelegate, UEdGraphNode*, bool);
+
+public:
+	/**
+	 * Creates a new UBlueprintNodeSpawner for the specified node class. Sets
+	 * the allocated spawner's NodeClass and CustomizeNodeDelegate fields from
+	 * the supplied parameters.
+	 *
+	 * @param  NodeClass                The node type that you want the spawner to spawn.
+	 * @param  CustomizeNodeDelegate    A delegate to perform specialized node setup post-spawn.
+	 * @return A newly allocated instance of this class.
+	 */
+	BLUEPRINTGRAPH_API static UBlueprintNodeSpawner* Create(UClass* const NodeClass, FCustomizeNodeDelegate CustomizeNodeDelegate = FCustomizeNodeDelegate());
+	
+	/**
+	 * Holds the class of node to spawn. May be null for sub-classes that know
+	 * specifically what node types to spawn (if null for an instance of this 
+	 * class, then nothing will be spawned).
+	 */
+	UClass* NodeClass;
+	
+	/**
+	 * A delegate to perform specialized node setup post-spawn (so we don't have
+	 * to sub-class this for every node type, and so we don't have to spawn a
+	 * template node for each to hold on to).
+	 */
+	FCustomizeNodeDelegate CustomizeNodeDelegate;
+
+	/**
+	 * Takes care of spawning a node for the specified graph. Looks to see if 
+	 * the supplied graph is transient, and if so, spawns a NOT fully formed
+	 * node (intended for template use).
+	 * 
+	 * This function is intended to be overridden; sub-classes may return
+	 * a pre-existing node, instead of a newly allocated one (for cases where
+	 * only one instance of the node type can exist). Callers should check for 
+	 * this case upon use.
+	 * 
+	 * @param  ParentGraph	The graph you want the node spawned into.
+	 * @return Null if it failed to spawn a node, otherwise a newly spawned node or posibly one that already existed.
+	 */
+	virtual UEdGraphNode* Invoke(UEdGraph* ParentGraph) const;
+	
+	// @TODO: for favorites, generated for the spawner type + data (function, property, node, etc.)
+	//virtual FGuid GetActionHash() const;
+	
+	/**
+	 * To save on performance, certain sub-classes can concoct menu names from
+	 * static data (like function names, etc.). This unfortunately doesn't 
+	 * completely decouple these spawners from the ui we sought to separate 
+	 * from. However, this can be thought of as a suggestion that the ui 
+	 * building code is free to do with as it pleases.
+	 *
+	 * @return An empty text string (could be a localized name for sub-classes).
+	 */
+	virtual FText GetDefaultMenuName() const;
+	
+	/**
+	 * To save on performance, certain sub-classes can concoct menu categories
+	 * from static data (like function metadata, etc.). This unfortunately
+	 * doesn't completely decouple these spawners from the ui we sought to
+	 * separate from. However, this can be thought of as a suggestion that the
+	 * ui building code is free to do with as it pleases (like concatenate with 
+	 * parent categories, etc.).
+	 *
+	 * @return An empty text string (could be a localized category for sub-classes).
+	 */
+	virtual FText GetDefaultMenuCategory() const;
+	
+	/**
+	 * To save on performance, certain sub-classes can construct tooltips from
+	 * static data (like function metadata, etc.). This unfortunately doesn't
+	 * completely decouple these spawners from the ui we sought to separate
+	 * them from. However, this can be thought of as a suggestion that the ui
+	 * building code is free to do with as it pleases.
+	 *
+	 * @return An empty text string (could be a localized tooltip for sub-classes).
+	 */
+	virtual FText GetDefaultMenuTooltip() const;
+	
+	/**
+	 * To save on performance, certain sub-classes can pull search keywords from
+	 * static data (like function metadata, etc.). This unfortunately doesn't
+	 * completely decouple these spawners from the ui we sought to separate
+	 * them from. Although, this can be thought of as a suggestion that the ui
+	 * building code is free to do with as it pleases (doesn't have to use it).
+	 *
+	 * @return An empty text string (could be a localized keywords for sub-classes).
+	 */
+	virtual FText GetDefaultSearchKeywords() const;
+
+	/**
+	 * Expects the supplied outer to be transient. Will spawn a new template-
+	 * node if the cached one is null, or if its graph outer doesn't match the
+	 * one specified here. Once a new node is spawned, this class will cache it
+	 * to save on spawning a new one with later requests (and to keep the node
+	 * from being GC'd).
+	 *
+	 * This should not have to be overridden (it calls Invoke() with the
+	 * specified outer; the transient-ness of the graph should signify that it
+	 * is for a template).
+	 *
+	 * Unfortunately, the supplied graph requires a Blueprint outer. Certain 
+	 * node types rely on this assumption; this is why the template-node and the
+	 * supplied graph cannot belong to the transient package themselves.
+	 *
+	 * @param  Outer    A transient graph (with a Blueprint outer), to act as the template-node's outer.
+	 * @return Should return a new/cached template-node (but could be null, or some pre-existing node... depends on the sub-class's Invoke() method).
+	 */
+	UEdGraphNode* MakeTemplateNode(UEdGraph* Outer) const;
+
+private:
+	/**
+	 * We try to avoid the spawning of node-templates (save on perf/mem where we
+	 * can), but sometimes users need to pull non-static data from the node 
+	 * itself (before it's spawned).
+	 *
+	 * Here we store the node template once it's created. Stored so that it can
+	 * be reused (and not reallocated) and not GC'd.
+	 */
+	UPROPERTY(Transient)
+	mutable UEdGraphNode* CachedNodeTemplate;
+};
