Properly display array properties when splitting a struct

diff --git a/Engine/Source/Editor/BlueprintGraph/Private/EdGraphSchema_K2.cpp b/Engine/Source/Editor/BlueprintGraph/Private/EdGraphSchema_K2.cpp
--- a/Engine/Source/Editor/BlueprintGraph/Private/EdGraphSchema_K2.cpp
+++ b/Engine/Source/Editor/BlueprintGraph/Private/EdGraphSchema_K2.cpp
@@ -4587,7 +4587,7 @@ void UEdGraphSchema_K2::SplitPin(UEdGraphPin* Pin) const
 		{
 			const FString PinName = FString::Printf(TEXT("%s_%s"), *Pin->PinName, *ProtoPin->PinName);
 			const FEdGraphPinType& ProtoPinType = ProtoPin->PinType;
-			UEdGraphPin* SubPin = GraphNode->CreatePin(Pin->Direction, ProtoPinType.PinCategory, ProtoPinType.PinSubCategory, ProtoPinType.PinSubCategoryObject.Get(), false, false, PinName);
+			UEdGraphPin* SubPin = GraphNode->CreatePin(Pin->Direction, ProtoPinType.PinCategory, ProtoPinType.PinSubCategory, ProtoPinType.PinSubCategoryObject.Get(), ProtoPinType.bIsArray, false, PinName);
 
 			if (K2Node != nullptr && K2Node->ShouldDrawCompact())
 			{
