(WIP) Stubbing in new BlueprintActionMenuItem for the blueprint menu refactor (wraps a node-spawner).

diff --git a/Engine/Source/Editor/BlueprintGraph/Private/BlueprintActionMenuItem.cpp b/Engine/Source/Editor/BlueprintGraph/Private/BlueprintActionMenuItem.cpp
new file mode 100644
--- /dev/null
+++ b/Engine/Source/Editor/BlueprintGraph/Private/BlueprintActionMenuItem.cpp
@@ -0,0 +1,153 @@
+// Copyright 1998-2014 Epic Games, Inc. All Rights Reserved.
+
+#include "BlueprintGraphPrivatePCH.h"
+#include "BlueprintActionMenuItem.h"
+#include "BlueprintNodeSpawner.h"
+#include "KismetEditorUtilities.h"	// for BringKismetToFocusAttentionOnObject()
+#include "BlueprintEditorUtils.h"	// for AnalyticsTrackNewNode(), MarkBlueprintAsModified(), etc.
+#include "ScopedTransaction.h"
+
+#define LOCTEXT_NAMESPACE "BlueprintActionMenuItem"
+
+/*******************************************************************************
+ * Static FBlueprintMenuActionItem Helpers
+ ******************************************************************************/
+
+namespace FBlueprintMenuActionItemImpl
+{
+	/**
+	 * Utility function for marking blueprints dirty and recompiling them after 
+	 * a node has been added.
+	 * 
+	 * @param  SpawnedNode	The node that was just added to the blueprint.
+	 */
+	static void DirtyBlueprintFromNewNode(UEdGraphNode* SpawnedNode);
+}
+
+//------------------------------------------------------------------------------
+static void FBlueprintMenuActionItemImpl::DirtyBlueprintFromNewNode(UEdGraphNode* SpawnedNode)
+{
+	UEdGraph const* const NodeGraph = SpawnedNode->GetGraph();
+	check(NodeGraph != nullptr);
+	
+	UBlueprint* Blueprint = FBlueprintEditorUtils::FindBlueprintForGraphChecked(NodeGraph);
+	check(Blueprint != nullptr);
+	
+	UK2Node* K2Node = Cast<UK2Node>(SpawnedNode);
+	// see if we need to recompile skeleton after adding this node, or just mark
+	// it dirty (default to rebuilding the skel, since there is no way to if
+	// non-k2 nodes structurally modify the blueprint)
+	if ((K2Node == nullptr) || K2Node->NodeCausesStructuralBlueprintChange())
+	{
+		FBlueprintEditorUtils::MarkBlueprintAsStructurallyModified(Blueprint);
+	}
+	else
+	{
+		FBlueprintEditorUtils::MarkBlueprintAsModified(Blueprint);
+	}
+}
+
+/*******************************************************************************
+ * FBlueprintMenuActionItem
+ ******************************************************************************/
+
+//------------------------------------------------------------------------------
+FBlueprintActionMenuItem::FBlueprintActionMenuItem(UBlueprintNodeSpawner* NodeSpawner, int32 MenuGrouping/* = 0*/)
+	: Action(NodeSpawner)
+{
+	check(Action != nullptr);
+	Grouping = MenuGrouping;
+}
+
+//------------------------------------------------------------------------------
+UEdGraphNode* FBlueprintActionMenuItem::PerformAction(UEdGraph* ParentGraph, UEdGraphPin* FromPin, FVector2D const Location, bool bSelectNewNode/* = true*/)
+{
+	using namespace FBlueprintMenuActionItemImpl;
+	FScopedTransaction Transaction(LOCTEXT("AddNodeTransaction", "Add Node"));
+	
+	// this could return an existing node
+	UEdGraphNode* SpawnedNode = Action->Invoke(ParentGraph);
+	
+	// if a returned node hasn't been added to the graph yet (it must have been freshly spawned)
+	if (ParentGraph->Nodes.Find(SpawnedNode) == INDEX_NONE)
+	{
+		check(SpawnedNode != nullptr);
+		
+		ParentGraph->Modify();
+		ParentGraph->AddNode(SpawnedNode, /*bool bFromUI =*/true, bSelectNewNode);
+		
+		SpawnedNode->PostPlacedNewNode();
+		SpawnedNode->NodePosX = Location.X;
+		SpawnedNode->NodePosY = Location.Y;
+		
+		if (FromPin != nullptr)
+		{
+			// for input pins, a new node will generally overlap the node being
+			// dragged from... work out if we want add in some spacing from the connecting node
+			if (FromPin->Direction == EGPD_Input)
+			{
+				UEdGraphNode* FromNode = FromPin->GetOwningNode();
+				check(FromNode != nullptr);
+				float const FromNodeX = FromNode->NodePosX;
+				
+				static const float MinNodeDistance = 60.f; // min distance between spawned nodes (to keep them from overlapping)
+				if (MinNodeDistance > FMath::Abs(FromNodeX - Location.X))
+				{
+					SpawnedNode->NodePosX = FromNodeX - MinNodeDistance;
+				}
+			}
+			
+			// modify before the call to AutowireNewNode() below
+			FromPin->Modify();
+		}
+		
+		static const float GridSnapSize = 16.f; // @TODO: ensure this is the same as SNodePanel::GetSnapGridSize()
+		SpawnedNode->SnapToGrid(GridSnapSize);
+		
+		// make sure to auto-wire after we position the new node (in case
+		// the auto-wire creates a conversion node to put between them)
+		SpawnedNode->AutowireNewNode(FromPin);
+		
+		FBlueprintEditorUtils::AnalyticsTrackNewNode(SpawnedNode);
+		DirtyBlueprintFromNewNode(SpawnedNode);
+	}
+	// if this node already existed, then we just want to focus on that node...
+	// some node types are only allowed one instance per blueprint (like events)
+	else if (SpawnedNode != nullptr)
+	{
+		FKismetEditorUtilities::BringKismetToFocusAttentionOnObject(SpawnedNode);
+	}
+	
+	return SpawnedNode;
+}
+
+//------------------------------------------------------------------------------
+UEdGraphNode* FBlueprintActionMenuItem::PerformAction(UEdGraph* ParentGraph, TArray<UEdGraphPin*>& FromPins, FVector2D const Location, bool bSelectNewNode/* = true*/)
+{
+	UEdGraphPin* FromPin = nullptr;
+	if (FromPins.Num() > 0)
+	{
+		FromPin = FromPins[0];
+	}
+	
+	UEdGraphNode* SpawnedNode = PerformAction(ParentGraph, FromPin, Location, bSelectNewNode);
+	// try auto-wiring the rest of the pins (if there are any)
+	for (int32 PinIndex = 1; PinIndex < FromPins.Num(); ++PinIndex)
+	{
+		SpawnedNode->AutowireNewNode(FromPins[PinIndex]);
+	}
+	
+	return SpawnedNode;
+}
+
+//------------------------------------------------------------------------------
+void FBlueprintActionMenuItem::AddReferencedObjects(FReferenceCollector& Collector)
+{
+	FEdGraphSchemaAction::AddReferencedObjects(Collector);
+	
+	// these don't get saved to disk, but we want to make sure the objects don't
+	// get GC'd while the action array is around
+	Collector.AddReferencedObject(Action);
+}
+
+#undef LOCTEXT_NAMESPACE
diff --git a/Engine/Source/Editor/BlueprintGraph/Public/BlueprintActionMenuItem.h b/Engine/Source/Editor/BlueprintGraph/Public/BlueprintActionMenuItem.h
new file mode 100644
--- /dev/null
+++ b/Engine/Source/Editor/BlueprintGraph/Public/BlueprintActionMenuItem.h
@@ -0,0 +1,41 @@
+// Copyright 1998-2014 Epic Games, Inc. All Rights Reserved.
+
+#pragma once
+
+#include "EdGraph/EdGraphSchema.h" // for FEdGraphSchemaAction
+#include "BlueprintActionMenuItem.generated.h"
+
+// Forward declarations
+class UBlueprintNodeSpawner;
+class FDragDropOperation;
+
+/**
+ * Wrapper around a UBlueprintNodeSpawner, which takes care of specialized
+ * node spawning. This class should not be sub-classed, any special handling
+ * should be done inside a UBlueprintNodeSpawner subclass, which will be
+ * invoked from this class (separated to divide ui and node-spawning).
+ */
+USTRUCT()
+struct FBlueprintActionMenuItem : public FEdGraphSchemaAction
+{
+	GENERATED_USTRUCT_BODY()
+	
+	static FString StaticGetTypeId() { static FString const TypeId = TEXT("FBlueprintActionMenuItem"); return TypeId; }
+	
+	/** Constructors */
+	FBlueprintActionMenuItem() : Action(nullptr) {}
+	FBlueprintActionMenuItem(UBlueprintNodeSpawner* NodeSpawner, int32 MenuGrouping = 0);
+	
+	// FEdGraphSchemaAction interface
+	virtual FString       GetTypeId() const final { return StaticGetTypeId(); }
+	virtual UEdGraphNode* PerformAction(UEdGraph* ParentGraph, UEdGraphPin* FromPin, FVector2D const Location, bool bSelectNewNode = true) final;
+	virtual UEdGraphNode* PerformAction(UEdGraph* ParentGraph, TArray<UEdGraphPin*>& FromPins, FVector2D const Location, bool bSelectNewNode = true) final;
+	virtual void          AddReferencedObjects(FReferenceCollector& Collector) final;
+	// End FEdGraphSchemaAction interface
+
+	// @TODO:
+	//TSharedRef<FDragDropOperation> OnDragged() const;
+
+	/** Specialized node-spawner, that comprises the action portion of this menu entry. */
+	UBlueprintNodeSpawner* Action;
+};
