Prevent knot nodes for structs from being split

diff --git a/Engine/Source/Editor/BlueprintGraph/Classes/K2Node.h b/Engine/Source/Editor/BlueprintGraph/Classes/K2Node.h
--- a/Engine/Source/Editor/BlueprintGraph/Classes/K2Node.h
+++ b/Engine/Source/Editor/BlueprintGraph/Classes/K2Node.h
@@ -138,6 +138,7 @@ class UK2Node : public UEdGraphNode
 	BLUEPRINTGRAPH_API virtual FString GetDocumentationLink() const override;
 	BLUEPRINTGRAPH_API virtual void GetPinHoverText(const UEdGraphPin& Pin, FString& HoverTextOut) const override;
 	BLUEPRINTGRAPH_API virtual bool ShowPaletteIconOnNode() const override { return true; }
+	virtual bool AllowSplitPins() const override { return true; }
 	// End of UEdGraphNode interface
 
 	// K2Node interface
diff --git a/Engine/Source/Editor/BlueprintGraph/Classes/K2Node_Knot.h b/Engine/Source/Editor/BlueprintGraph/Classes/K2Node_Knot.h
--- a/Engine/Source/Editor/BlueprintGraph/Classes/K2Node_Knot.h
+++ b/Engine/Source/Editor/BlueprintGraph/Classes/K2Node_Knot.h
@@ -19,6 +19,7 @@ public:
 	virtual FString GetPinNameOverride(const UEdGraphPin& Pin) const override;
 	virtual void OnRenameNode(const FString& NewName) override;
 	virtual TSharedPtr<class INameValidatorInterface> MakeNameValidator() const override;
+	virtual bool AllowSplitPins() const override;
 	// End of UEdGraphNode interface
 
 	// UK2Node interface
diff --git a/Engine/Source/Editor/BlueprintGraph/Private/EdGraphSchema_K2.cpp b/Engine/Source/Editor/BlueprintGraph/Private/EdGraphSchema_K2.cpp
--- a/Engine/Source/Editor/BlueprintGraph/Private/EdGraphSchema_K2.cpp
+++ b/Engine/Source/Editor/BlueprintGraph/Private/EdGraphSchema_K2.cpp
@@ -806,7 +806,7 @@ void UEdGraphSchema_K2::GetContextMenuActions(const UEdGraph* CurrentGraph, cons
 					MenuBuilder->AddMenuEntry( FGraphEditorCommands::Get().PromoteToVariable );
 				}
 	
-				if (PinHasSplittableStructType(InGraphPin))
+				if (PinHasSplittableStructType(InGraphPin) && InGraphNode->AllowSplitPins())
 				{
 					MenuBuilder->AddMenuEntry( FGraphEditorCommands::Get().SplitStructPin );
 				}
@@ -3294,7 +3294,7 @@ bool UEdGraphSchema_K2::CanPromotePinToVariable( const UEdGraphPin& Pin ) const
 
 bool UEdGraphSchema_K2::CanSplitStructPin( const UEdGraphPin& Pin ) const
 {
-	return (Pin.LinkedTo.Num() == 0 && PinHasSplittableStructType(&Pin));
+	return (Pin.LinkedTo.Num() == 0 && PinHasSplittableStructType(&Pin) && Pin.GetOwningNode()->AllowSplitPins());
 }
 
 bool UEdGraphSchema_K2::CanRecombineStructPin( const UEdGraphPin& Pin ) const
diff --git a/Engine/Source/Editor/BlueprintGraph/Private/K2Node_Knot.cpp b/Engine/Source/Editor/BlueprintGraph/Private/K2Node_Knot.cpp
--- a/Engine/Source/Editor/BlueprintGraph/Private/K2Node_Knot.cpp
+++ b/Engine/Source/Editor/BlueprintGraph/Private/K2Node_Knot.cpp
@@ -79,6 +79,11 @@ bool UK2Node_Knot::IsNodeSafeToIgnore() const
 	return true;
 }
 
+bool UK2Node_Knot::AllowSplitPins() const
+{
+	return false;
+}
+
 void UK2Node_Knot::NotifyPinConnectionListChanged(UEdGraphPin* Pin)
 {
 	const UEdGraphSchema_K2* K2Schema = GetDefault<UEdGraphSchema_K2>();
diff --git a/Engine/Source/Runtime/Engine/Classes/EdGraph/EdGraphNode.h b/Engine/Source/Runtime/Engine/Classes/EdGraph/EdGraphNode.h
--- a/Engine/Source/Runtime/Engine/Classes/EdGraph/EdGraphNode.h
+++ b/Engine/Source/Runtime/Engine/Classes/EdGraph/EdGraphNode.h
@@ -165,6 +165,9 @@ public:
 	/** Whether or not this node should be given the chance to override pin names.  If this returns true, then GetPinNameOverride() will be called for each pin, each frame */
 	virtual bool ShouldOverridePinNames() const { return false; }
 
+	/** Whether or not struct pins belonging to this node should be allowed to be split or not. */
+	virtual bool AllowSplitPins() const { return false; }
+
 	/** Gets the overridden name for the specified pin, if any */
 	virtual FString GetPinNameOverride(const UEdGraphPin& Pin) const { return FString(TEXT("")); }
 
