Make blueprint struct splitting an experimental feature

diff --git a/Engine/Source/Editor/BlueprintGraph/Classes/K2Node.h b/Engine/Source/Editor/BlueprintGraph/Classes/K2Node.h
--- a/Engine/Source/Editor/BlueprintGraph/Classes/K2Node.h
+++ b/Engine/Source/Editor/BlueprintGraph/Classes/K2Node.h
@@ -138,7 +138,7 @@ class UK2Node : public UEdGraphNode
 	BLUEPRINTGRAPH_API virtual FString GetDocumentationLink() const override;
 	BLUEPRINTGRAPH_API virtual void GetPinHoverText(const UEdGraphPin& Pin, FString& HoverTextOut) const override;
 	BLUEPRINTGRAPH_API virtual bool ShowPaletteIconOnNode() const override { return true; }
-	virtual bool AllowSplitPins() const override { return true; }
+	BLUEPRINTGRAPH_API virtual bool AllowSplitPins() const override;
 	// End of UEdGraphNode interface
 
 	// K2Node interface
diff --git a/Engine/Source/Editor/BlueprintGraph/Private/K2Node.cpp b/Engine/Source/Editor/BlueprintGraph/Private/K2Node.cpp
--- a/Engine/Source/Editor/BlueprintGraph/Private/K2Node.cpp
+++ b/Engine/Source/Editor/BlueprintGraph/Private/K2Node.cpp
@@ -575,6 +575,11 @@ void UK2Node::ReconstructSinglePin(UEdGraphPin* NewPin, UEdGraphPin* OldPin, ERe
 	OldPin->Rename(NULL, GetTransientPackage(), (REN_DontCreateRedirectors|(Blueprint->bIsRegeneratingOnLoad ? REN_ForceNoResetLoaders : REN_None)));
 }
 
+bool UK2Node::AllowSplitPins() const
+{
+	return GetDefault<UEditorExperimentalSettings>()->bAllowSplitStructPins;
+}
+
 void UK2Node::ExpandSplitPin(FKismetCompilerContext* CompilerContext, UEdGraph* SourceGraph, UEdGraphPin* Pin)
 {
 	const UEdGraphSchema_K2* Schema = CastChecked<UEdGraphSchema_K2>(CompilerContext ? CompilerContext->GetSchema() : SourceGraph->GetSchema());
diff --git a/Engine/Source/Editor/UnrealEd/Classes/Settings/EditorExperimentalSettings.h b/Engine/Source/Editor/UnrealEd/Classes/Settings/EditorExperimentalSettings.h
--- a/Engine/Source/Editor/UnrealEd/Classes/Settings/EditorExperimentalSettings.h
+++ b/Engine/Source/Editor/UnrealEd/Classes/Settings/EditorExperimentalSettings.h
@@ -63,6 +63,9 @@ public:
 	UPROPERTY(EditAnywhere, config, Category=Blueprints, meta=(DisplayName="Draw midpoint arrows in Blueprints"))
 	bool bDrawMidpointArrowsInBlueprints;
 
+	UPROPERTY(EditAnywhere, config, Category=Blueprints, meta=(DisplayName="Allow splitting of struct pins"))
+	bool bAllowSplitStructPins;
+
 	/**
 	 * Returns an event delegate that is executed when a setting has changed.
 	 *
