(WIP) First pass at blueprint action menu builder and filter (for the impending action menu repactor).

diff --git a/Engine/Source/Editor/BlueprintGraph/Private/BlueprintActionFilter.cpp b/Engine/Source/Editor/BlueprintGraph/Private/BlueprintActionFilter.cpp
new file mode 100644
--- /dev/null
+++ b/Engine/Source/Editor/BlueprintGraph/Private/BlueprintActionFilter.cpp
@@ -0,0 +1,880 @@
+// Copyright 1998-2014 Epic Games, Inc. All Rights Reserved.
+
+#include "BlueprintGraphPrivatePCH.h"
+#include "BlueprintActionFilter.h"
+#include "BlueprintNodeSpawner.h"
+#include "BlueprintFunctionNodeSpawner.h"
+#include "BlueprintEventNodeSpawner.h"
+#include "BlueprintComponentNodeSpawner.h"
+#include "BlueprintPropertyNodeSpawner.h"
+#include "EdGraphSchema_K2.h"		// for FBlueprintMetadata
+#include "BlueprintEditorUtils.h"	// for FindBlueprintForGraph()
+#include "ObjectEditorUtils.h"		// for IsFunctionHiddenFromClass()/IsVariableCategoryHiddenFromClass()
+#include "K2Node_VariableSet.h"
+// "impure" node types (utilized in BlueprintActionFilterImpl::IsImpure)
+#include "K2Node_IfThenElse.h"
+#include "K2Node_MultiGate.h"
+#include "K2Node_MakeArray.h"
+#include "K2Node_Message.h"
+#include "K2Node_ExecutionSequence.h"
+
+/*******************************************************************************
+ * Static BlueprintActionFilter Helpers
+ ******************************************************************************/
+
+namespace BlueprintActionFilterImpl
+{	
+	/**
+	 * All actions are associated with some class, whether it's the class of the
+	 * node it will spawn, or the outer of the associated field (property,
+	 * function, etc.).
+	 * 
+	 * @param  BlueprintAction	The node-spawner you want an associated class for.
+	 * @return The action's associated class (null if it doesn't have one).
+	 */
+	static UClass* GetActionClass(UBlueprintNodeSpawner const* BlueprintAction);
+	
+	/**
+	 * Certain node-spawners are associated with specific UFields (functions, 
+	 * properties, etc.). This attempts to retrieve the field from the spawner.
+	 * 
+	 * @param  BlueprintAction	The node-spawner you want an associated field for.
+	 * @return The action's associated field (null if it doesn't have one).
+	 */
+	static UField const* GetAssociatedField(UBlueprintNodeSpawner const* BlueprintAction);
+
+	/**
+	 * Certain node-spawners are associated with specific UFunction (call-
+	 * function, and event spawners). This attempts to retrieve the function 
+	 * from the spawner.
+	 * 
+	 * @param  BlueprintAction	The node-spawner you want an associated function for.
+	 * @return The action's associated function (null if it doesn't have one).
+	 */
+	static UFunction const* GetAssociatedFunction(UBlueprintNodeSpawner const* BlueprintAction);
+	
+	/**
+	 * Utility method to check and see if the specified node-spawner would 
+	 * produce an impure node.
+	 * 
+	 * @param  NodeSpawner	The action you want to query.
+	 * @return True if the action will spawn an impure node, false if not (or unknown).
+	 */
+	static bool IsImpure(UBlueprintNodeSpawner const* NodeSpawner);
+
+	/**
+	 * Utility method to check and see if the specified node-spawner would 
+	 * produce a latent node.
+	 * 
+	 * @param  NodeSpawner	The action you want to query.
+	 * @return True if the action will spawn a latent node, false if not (or unknown).
+	 */
+	static bool IsLatent(UBlueprintNodeSpawner const* NodeSpawner);
+	
+	/**
+	 * Rejection test that checks to see if the supplied node-spawner would
+	 * produce an event that does NOT belong to the specified blueprint.
+	 * 
+	 * @param  Filter			Holds the blueprint context for this test.
+	 * @param  BlueprintAction	The action you wish to query.
+	 * @return True if the node-spawner would produce an event incompatible with the specified blueprint(s).
+	 */
+	static bool IsEventUnimplementable(FBlueprintActionFilter const& Filter, UBlueprintNodeSpawner const* BlueprintAction);
+
+	/**
+	 * Rejection test that checks to see if the supplied node-spawner has an 
+	 * associated field that is not accessible by the blueprint (it's private or
+	 * protected).
+	 * 
+	 * @param  Filter			Holds the blueprint context for this test.
+	 * @param  BlueprintAction	The action you wish to query.
+	 * @return True if the node-spawner is associated with a private field.
+	 */
+	static bool IsFieldInaccessible(FBlueprintActionFilter const& Filter, UBlueprintNodeSpawner const* BlueprintAction);
+
+	/**
+	 * Rejection test that checks to see if the supplied node-spawner has an 
+	 * associated class that is "restricted" and thusly, hidden from the 
+	 * blueprint.
+	 * 
+	 * @param  Filter			Holds the blueprint context for this test.
+	 * @param  BlueprintAction	The action you wish to query.
+	 * @return True if the node-spawner belongs to a class that is restricted to certain blueprints (not this one).
+	 */
+	static bool IsRestrictedClassMember(FBlueprintActionFilter const& Filter, UBlueprintNodeSpawner const* BlueprintAction);
+	
+	/**
+	 * Rejection test that checks to see if the supplied node-spawner would 
+	 * produce a variable-set node when the property is read-only (in this
+	 * blueprint).
+	 * 
+	 * @param  Filter			Holds the blueprint context for this test.
+	 * @param  BlueprintAction	The action you wish to query.
+	 * @return True if the action would spawn a variable-set node for a read-only property.
+	 */
+	static bool IsPermissionNotGranted(FBlueprintActionFilter const& Filter, UBlueprintNodeSpawner const* BlueprintAction);
+
+	/**
+	 * Rejection test that checks to see if the supplied node-spawner would 
+	 * produce a node (or comes from an associated class) that is deprecated.
+	 * 
+	 * @param  Filter			Filter context (unused) for this test.
+	 * @param  BlueprintAction	The action you wish to query.
+	 * @return True if the action would spawn a node that is deprecated.
+	 */
+	static bool IsDeprecated(FBlueprintActionFilter const& Filter, UBlueprintNodeSpawner const* BlueprintAction);
+	
+	/**
+	 * Rejection test that checks to see if the supplied node-spawner would 
+	 * produce an impure node, incompatible with the specified graphs.
+	 * 
+	 * @param  Filter			Holds the graph context for this test.
+	 * @param  BlueprintAction	The action you wish to query.
+	 * @return True if the action would spawn an impure node (and the graph doesn't allow them).
+	 */
+	static bool IsIncompatibleImpureNode(FBlueprintActionFilter const& Filter, UBlueprintNodeSpawner const* BlueprintAction);
+	
+	/**
+	 * Rejection test that checks to see if the supplied node-spawner would 
+	 * produce a node incompatible with the specified graph.
+	 * 
+	 * @param  Filter			Holds the graph context for this test.
+	 * @param  BlueprintAction	The action you wish to query.
+	 * @return True if the action would spawn a node (and the graph wouldn't allow it).
+	 */
+	static bool IsIncompatibleWithGraphType(FBlueprintActionFilter const& Filter, UBlueprintNodeSpawner const* BlueprintAction);
+
+	/**
+	 * Rejection test that checks to see if the supplied node-spawner would 
+	 * produce a node incompatible with an ubergraph.
+	 * 
+	 * @param  Context			Holds the blueprint/graph/pin context for this test.
+	 * @param  BlueprintAction	The action you wish to query.
+	 * @return True if the action would spawn a node not allowed in ubergraphs.
+	 */
+	static bool IsIncompatibleWithUbergraph(FBlueprintActionContext const& Context, UBlueprintNodeSpawner const* BlueprintAction);
+
+	/**
+	 * Rejection test that checks to see if the supplied node-spawner would 
+	 * produce a node incompatible with an animation graph.
+	 * 
+	 * @param  Context			Holds the blueprint/graph/pin context for this test.
+	 * @param  BlueprintAction	The action you wish to query.
+	 * @return True if the action would spawn a node not allowed in animation graphs.
+	 */
+	static bool IsIncompatibleWithAnimGraph(FBlueprintActionContext const& Context, UBlueprintNodeSpawner const* BlueprintAction);
+
+	/**
+	 * Rejection test that checks to see if the supplied node-spawner would 
+	 * produce a node incompatible with a macro graph.
+	 * 
+	 * @param  Context			Holds the blueprint/graph/pin context for this test.
+	 * @param  BlueprintAction	The action you wish to query.
+	 * @return True if the action would spawn a node not allowed in macro graphs.
+	 */
+	static bool IsIncompatibleWithMacroGraph(FBlueprintActionContext const& Context, UBlueprintNodeSpawner const* BlueprintAction);
+
+	/**
+	 * Rejection test that checks to see if the supplied node-spawner would 
+	 * produce an node incompatible with a function graph.
+	 * 
+	 * @param  Context			Holds the blueprint/graph/pin context for this test.
+	 * @param  BlueprintAction	The action you wish to query.
+	 * @return True if the action would spawn a node not allowed in function graphs.
+	 */
+	static bool IsIncompatibleWithFuncGraph(FBlueprintActionContext const& Context, UBlueprintNodeSpawner const* BlueprintAction);
+	
+	/**
+	 * Rejection test that checks to see if the node-spawner has any associated
+	 * fields that are global (structs, enums, etc.), or if the field belongs to
+	 * a static library class.
+	 * 
+	 * @param  Context			Holds the context (unused) for this test.
+	 * @param  BlueprintAction	The action you wish to query.
+	 * @return True if the action would spawn a node not allowed in function graphs.
+	 */
+	static bool IsGlobalField(FBlueprintActionFilter const& Filter, UBlueprintNodeSpawner const* BlueprintAction);
+
+	/**
+	 * Rejection test that checks to see if the node-spawner is associated with 
+	 * a field that belongs to a class that is not white-listed (ignores global
+	 * fields).
+	 * 
+	 * @param  Context			Holds the class context for this test.
+	 * @param  BlueprintAction	The action you wish to query.
+	 * @return True if the action is associated with a non-whitelisted class member.
+	 */
+	static bool IsExternalField(FBlueprintActionFilter const& Filter, UBlueprintNodeSpawner const* BlueprintAction);
+
+	/**
+	 * Rejection test that checks to see if the node-spawner is associated with 
+	 * a field that is hidden from the specified blueprint (via metadata).
+	 * 
+	 * @param  Context			Holds the class context for this test.
+	 * @param  BlueprintAction	The action you wish to query.
+	 * @return True if the action is associated with a hidden field.
+	 */
+	static bool IsFieldCategoryHidden(FBlueprintActionFilter const& Filter, UBlueprintNodeSpawner const* BlueprintAction);
+
+	/**
+	 * Rejection test that checks to see if the node-spawner would produce a 
+	 * node that couldn't be pasted into the specified graph(s) (check against 
+	 * the node's CDO to save on perf, to keep from spawning a temp node).
+	 * 
+	 * @param  Context			Holds the graph context for this test.
+	 * @param  BlueprintAction	The action you wish to query.
+	 * @return True if the action would produce a node that cannot be pasted in the graph.
+	 */
+	static bool IsNodeUnpastable(FBlueprintActionFilter const& Filter, UBlueprintNodeSpawner const* BlueprintAction);
+
+	/**
+	 * Rejection test that checks to see if the node-spawner would produce a 
+	 * node type that isn't white-listed. 
+	 * 
+	 * @param  Context			Holds the class context for this test.
+	 * @param  BlueprintAction	The action you wish to query.
+	 * @param  bMatchExplicitly When true, the node's class has to perfectly match a white-listed type (sub-classes don't count).
+	 * @return True if the action would produce a non-whitelisted node.
+	 */
+	static bool IsFilteredNodeType(FBlueprintActionFilter const& Filter, UBlueprintNodeSpawner const* BlueprintAction, bool const bMatchExplicitly = false);
+};
+
+//------------------------------------------------------------------------------
+static UClass* BlueprintActionFilterImpl::GetActionClass(UBlueprintNodeSpawner const* BlueprintAction)
+{
+	UClass* ActionClass = BlueprintAction->NodeClass;
+	
+	if (UBlueprintFunctionNodeSpawner const* FuncNodeSpawner = Cast<UBlueprintFunctionNodeSpawner>(BlueprintAction))
+	{
+		UFunction const* Function = FuncNodeSpawner->GetFunction();
+		check(Function != nullptr);
+		ActionClass = Function->GetOuterUClass();
+	}
+	else if (UBlueprintEventNodeSpawner const* EventSpawner = Cast<UBlueprintEventNodeSpawner>(BlueprintAction))
+	{
+		UFunction const* EventFunc = EventSpawner->GetEventFunction();
+		check(EventFunc != nullptr);
+		ActionClass = EventFunc->GetOuterUClass();
+	}
+	else if (UBlueprintPropertyNodeSpawner const* PropertySpawner = Cast<UBlueprintPropertyNodeSpawner>(BlueprintAction))
+	{
+		UProperty const* Property = PropertySpawner->GetProperty();
+		check(Property != nullptr);
+		ActionClass = Cast<UClass>(Property->GetOuterUField());
+	}
+	else if (UBlueprintComponentNodeSpawner const* ComponentSpawner = Cast<UBlueprintComponentNodeSpawner>(BlueprintAction))
+	{
+		ActionClass = ComponentSpawner->GetComponentClass();
+	}
+	
+	return ActionClass;
+}
+
+//------------------------------------------------------------------------------
+static UField const* BlueprintActionFilterImpl::GetAssociatedField(UBlueprintNodeSpawner const* BlueprintAction)
+{
+	UField const* ClassField = nullptr;
+
+	if (UFunction const* Function = GetAssociatedFunction(BlueprintAction))
+	{
+		ClassField = Function;
+	}
+	else if (UBlueprintPropertyNodeSpawner const* PropertySpawner = Cast<UBlueprintPropertyNodeSpawner>(BlueprintAction))
+	{
+		ClassField = PropertySpawner->GetProperty();
+	}
+
+	return ClassField;
+}
+
+//------------------------------------------------------------------------------
+static UFunction const* BlueprintActionFilterImpl::GetAssociatedFunction(UBlueprintNodeSpawner const* BlueprintAction)
+{
+	UFunction const* Function = nullptr;
+
+	if (UBlueprintFunctionNodeSpawner const* FuncNodeSpawner = Cast<UBlueprintFunctionNodeSpawner>(BlueprintAction))
+	{
+		Function = FuncNodeSpawner->GetFunction();
+	}
+	else if (UBlueprintEventNodeSpawner const* EventSpawner = Cast<UBlueprintEventNodeSpawner>(BlueprintAction))
+	{
+		Function = EventSpawner->GetEventFunction();
+	}
+
+	return Function;
+}
+
+//------------------------------------------------------------------------------
+static bool BlueprintActionFilterImpl::IsImpure(UBlueprintNodeSpawner const* NodeSpawner)
+{
+	bool bIsImpure = false;
+
+	UClass* NodeClass = NodeSpawner->NodeClass;
+	if (NodeClass != nullptr)
+	{
+		// TODO: why are some of these "impure"?
+		bIsImpure = (NodeClass == UK2Node_IfThenElse::StaticClass()) ||
+			(NodeClass == UK2Node_MultiGate::StaticClass()) ||
+			(NodeClass == UK2Node_MakeArray::StaticClass()) ||
+			(NodeClass == UK2Node_Message::StaticClass()) ||
+			(NodeClass == UK2Node_ExecutionSequence::StaticClass());
+	}
+	else if (UFunction const* Function = GetAssociatedFunction(NodeSpawner))
+	{
+		bIsImpure = (Function->HasAnyFunctionFlags(FUNC_BlueprintPure) != false);
+	}
+
+	return bIsImpure;
+}
+
+//------------------------------------------------------------------------------
+static bool BlueprintActionFilterImpl::IsLatent(UBlueprintNodeSpawner const* NodeSpawner)
+{
+	bool bIsLatent = false;
+
+	UClass* NodeClass = NodeSpawner->NodeClass;
+	if ((NodeClass != nullptr) && NodeClass->IsChildOf<UK2Node_BaseAsyncTask>())
+	{
+		bIsLatent = true;
+	}
+	else if (UFunction const* Function = GetAssociatedFunction(NodeSpawner))
+	{
+		bIsLatent = (Function->HasMetaData(FBlueprintMetadata::MD_Latent) != false);
+	}
+
+	return bIsLatent;
+}
+
+//------------------------------------------------------------------------------
+static bool BlueprintActionFilterImpl::IsEventUnimplementable(FBlueprintActionFilter const& Filter, UBlueprintNodeSpawner const* BlueprintAction)
+{
+	bool bIsFilteredOut = false;
+	FBlueprintActionContext const& FilterContext = Filter.Context;
+	
+	if (UBlueprintEventNodeSpawner const* EventSpawner = Cast<UBlueprintEventNodeSpawner>(BlueprintAction))
+	{
+		UFunction const* EventFunc = EventSpawner->GetEventFunction();
+		check(EventFunc != nullptr);
+		
+		UClass* FuncOwner = EventFunc->GetOuterUClass();
+		for (UBlueprint const* Blueprint : FilterContext.Blueprints)
+		{
+			UClass* BpClass = Blueprint->GeneratedClass;
+			check(BpClass != nullptr);
+			
+			// if this function belongs directly to this blueprint, then it is
+			// already implemented here (this action however is valid for sub-
+			// classes, as they can override the event's functionality)
+			if (FuncOwner == BpClass)
+			{
+				bIsFilteredOut = true;
+				break;
+			}
+			
+			// you can only implement events that you inherit; so if this
+			// blueprint is not a subclass of the event's owner, then we're not
+			// allowed to implement it
+			if (!BpClass->IsChildOf(FuncOwner))
+			{
+				bIsFilteredOut = true;
+				break;
+			}
+		}
+	}
+	
+	return bIsFilteredOut;
+}
+
+//------------------------------------------------------------------------------
+static bool BlueprintActionFilterImpl::IsFieldInaccessible(FBlueprintActionFilter const& Filter, UBlueprintNodeSpawner const* BlueprintAction)
+{
+	bool bIsFilteredOut = false;
+	FBlueprintActionContext const& FilterContext = Filter.Context;
+	
+	if (UField const* ClassField = GetAssociatedField(BlueprintAction))
+	{
+		bool bIsProtected = ClassField->HasMetaData(FBlueprintMetadata::MD_Protected);
+		bool bIsPrivate   = ClassField->HasMetaData(FBlueprintMetadata::MD_Private);
+
+		UClass* FieldOwner = Cast<UClass>(ClassField->GetOuter());
+		for (UBlueprint const* Blueprint : FilterContext.Blueprints)
+		{
+			UClass* BpClass = Blueprint->GeneratedClass;
+			check(BpClass != nullptr);
+			
+			// private functions are only accessible from the class they belong to
+			if (bIsPrivate && (FieldOwner != BpClass))
+			{
+				bIsFilteredOut = true;
+				break;
+			}
+			else if (bIsProtected && !BpClass->IsChildOf(FieldOwner))
+			{
+				bIsFilteredOut = true;
+				break;
+			}
+		}
+	}
+	
+	return bIsFilteredOut;
+}
+
+//------------------------------------------------------------------------------
+static bool BlueprintActionFilterImpl::IsRestrictedClassMember(FBlueprintActionFilter const& Filter, UBlueprintNodeSpawner const* BlueprintAction)
+{
+	bool bIsFilteredOut = false;
+	FBlueprintActionContext const& FilterContext = Filter.Context;
+	
+	if (UClass const* ActionClass = GetActionClass(BlueprintAction))
+	{
+		if (ActionClass->HasMetaData(FBlueprintMetadata::MD_RestrictedToClasses))
+		{
+			FString const& ClassRestrictions = ActionClass->GetMetaData(FBlueprintMetadata::MD_RestrictedToClasses);
+			for (UBlueprint const* Blueprint : FilterContext.Blueprints)
+			{
+				bool bIsClassListed = false;
+				
+				UClass* BpClass = Blueprint->SkeletonGeneratedClass;
+				// walk the class inheritance chain to see if this class is one
+				// of the allowed
+				while (!bIsClassListed && (BpClass != nullptr))
+				{
+					FString const ClassName = BpClass->GetName();
+					bIsClassListed = (ClassName == ClassRestrictions) || !!FCString::StrfindDelim(*ClassRestrictions, *ClassName, TEXT(" "));
+					
+					BpClass = BpClass->GetSuperClass();
+				}
+				
+				// if the blueprint class wasn't listed as one of the few
+				// classes that this can be accessed from, then filter it out
+				if (!bIsClassListed)
+				{
+					bIsFilteredOut = true;
+					break;
+				}
+			}
+		}
+	}
+	
+	return bIsFilteredOut;
+}
+
+//------------------------------------------------------------------------------
+static bool BlueprintActionFilterImpl::IsPermissionNotGranted(FBlueprintActionFilter const& Filter, UBlueprintNodeSpawner const* BlueprintAction)
+{
+	bool bIsFilteredOut = false;
+	FBlueprintActionContext const& FilterContext = Filter.Context;
+
+	if (UBlueprintPropertyNodeSpawner const* PropertySpawner = Cast<UBlueprintPropertyNodeSpawner>(BlueprintAction))
+	{
+		UProperty const* const Property = PropertySpawner->GetProperty();
+		check(Property != nullptr);
+		UClass const* const NodeClass = PropertySpawner->NodeClass;
+
+		for (UBlueprint const* Blueprint : FilterContext.Blueprints)
+		{
+			bool const bIsReadOnly = FBlueprintEditorUtils::IsPropertyReadOnlyInCurrentBlueprint(Blueprint, Property);
+			if (bIsReadOnly && NodeClass && NodeClass->IsChildOf<UK2Node_VariableSet>())
+			{
+				bIsFilteredOut = true;
+			}
+		}
+	}
+
+	return bIsFilteredOut;
+}
+
+//------------------------------------------------------------------------------
+static bool BlueprintActionFilterImpl::IsDeprecated(FBlueprintActionFilter const& /*Filter*/, UBlueprintNodeSpawner const* BlueprintAction)
+{
+	bool bIsFilteredOut = false;	
+	if ((BlueprintAction->NodeClass != nullptr) && BlueprintAction->NodeClass->HasAnyClassFlags(CLASS_Deprecated))
+	{
+		bIsFilteredOut = true;
+	}
+	else if (UClass* ActionClass = GetActionClass(BlueprintAction))
+	{
+		bIsFilteredOut = ActionClass->HasAnyClassFlags(CLASS_Deprecated);
+	}
+	return bIsFilteredOut;
+}
+
+//------------------------------------------------------------------------------
+static bool BlueprintActionFilterImpl::IsGlobalField(FBlueprintActionFilter const& Filter, UBlueprintNodeSpawner const* BlueprintAction)
+{
+	bool bIsFilteredOut = false;
+	if (UField const* ClassField = GetAssociatedField(BlueprintAction))
+	{
+		UClass* FieldClass = Cast<UClass>(ClassField->GetOuter());
+		bIsFilteredOut = (FieldClass == nullptr) || FieldClass->IsChildOf<UBlueprintFunctionLibrary>();
+	}
+
+	return bIsFilteredOut;
+}
+
+//------------------------------------------------------------------------------
+static bool BlueprintActionFilterImpl::IsExternalField(FBlueprintActionFilter const& Filter, UBlueprintNodeSpawner const* BlueprintAction)
+{
+	bool bIsFilteredOut = false;
+	if (UField const* ClassField = GetAssociatedField(BlueprintAction))
+	{		
+		UClass* FieldClass = Cast<UClass>(ClassField->GetOuter());
+		bool const bIsGlobal = (FieldClass == nullptr) || FieldClass->IsChildOf<UBlueprintFunctionLibrary>();
+
+		// global (and static library) fields can stay (unless explicitly excluded... save that for a separate test)
+		if (!bIsGlobal)
+		{
+			for (UClass const* Class : Filter.OwnerClasses)
+			{
+				bool const bIsInternalFunc = Class->IsChildOf(FieldClass);	
+				if (!bIsInternalFunc)
+				{
+					bIsFilteredOut = true;
+					break;
+				}
+			}
+		}
+	}
+	
+	return bIsFilteredOut;
+}
+
+//------------------------------------------------------------------------------
+static bool BlueprintActionFilterImpl::IsFieldCategoryHidden(FBlueprintActionFilter const& Filter, UBlueprintNodeSpawner const* BlueprintAction)
+{
+	bool bIsFilteredOut = false;
+	FBlueprintActionContext const& FilterContext = Filter.Context;
+
+	if (UField const* AssociatedField = GetAssociatedField(BlueprintAction))
+	{
+		DECLARE_DELEGATE_RetVal_OneParam(bool, FIsFieldHiddenDelegate, UClass*);
+		FIsFieldHiddenDelegate IsFieldHiddenDelegate;
+
+		if (UFunction const* Function = Cast<UFunction>(AssociatedField))
+		{
+			auto IsFunctionHiddenLambda = [](UClass* Class, UFunction const* Function)->bool
+			{
+				return FObjectEditorUtils::IsFunctionHiddenFromClass(Function, Class);
+			};
+			IsFieldHiddenDelegate = FIsFieldHiddenDelegate::CreateStatic(IsFunctionHiddenLambda, Function);
+		}
+		else if (UProperty const* Property = Cast<UProperty>(AssociatedField))
+		{
+			auto IsPropertyHiddenLambda = [](UClass* Class, UProperty const* Property)->bool
+			{
+				return FObjectEditorUtils::IsVariableCategoryHiddenFromClass(Property, Class);
+			};
+			IsFieldHiddenDelegate = FIsFieldHiddenDelegate::CreateStatic(IsPropertyHiddenLambda, Property);
+		}
+
+		if (IsFieldHiddenDelegate.IsBound())
+		{
+			for (UBlueprint* Blueprint : FilterContext.Blueprints)
+			{
+				if (IsFieldHiddenDelegate.Execute(Blueprint->GeneratedClass))
+				{
+					bIsFilteredOut = true;
+					break;
+				}
+			}
+
+			for (int32 ClassIndex = 0; !bIsFilteredOut && (ClassIndex < Filter.OwnerClasses.Num()); ++ClassIndex)
+			{
+				bIsFilteredOut = IsFieldHiddenDelegate.Execute(Filter.OwnerClasses[ClassIndex]);
+			}
+		}
+	}
+
+	return bIsFilteredOut;
+}
+
+//------------------------------------------------------------------------------
+static bool BlueprintActionFilterImpl::IsIncompatibleImpureNode(FBlueprintActionFilter const& Filter, UBlueprintNodeSpawner const* BlueprintAction)
+{
+	bool bAllowImpureNodes = true;
+	FBlueprintActionContext const& FilterContext = Filter.Context;
+
+	for (UEdGraph* Graph : FilterContext.Graphs)
+	{
+		if (UEdGraphSchema_K2* K2Schema = Graph->Schema->GetDefaultObject<UEdGraphSchema_K2>())
+		{
+			bAllowImpureNodes &= K2Schema->DoesGraphSupportImpureFunctions(Graph);
+		}
+	}
+	
+	bool bIsFilteredOut = !bAllowImpureNodes && IsImpure(BlueprintAction);
+	return bIsFilteredOut;
+}
+
+//------------------------------------------------------------------------------
+static bool BlueprintActionFilterImpl::IsIncompatibleWithUbergraph(FBlueprintActionContext const& /*Context*/, UBlueprintNodeSpawner const* BlueprintAction)
+{
+	bool bIsFilteredOut = false;
+	return bIsFilteredOut;
+}
+
+//------------------------------------------------------------------------------
+static bool BlueprintActionFilterImpl::IsIncompatibleWithAnimGraph(FBlueprintActionContext const& /*Context*/, UBlueprintNodeSpawner const* BlueprintAction)
+{
+	bool bIsFilteredOut = IsLatent(BlueprintAction); // latent actions are only allowed in ubergraphs
+	return bIsFilteredOut;
+}
+
+//------------------------------------------------------------------------------
+static bool BlueprintActionFilterImpl::IsIncompatibleWithMacroGraph(FBlueprintActionContext const& /*Context*/, UBlueprintNodeSpawner const* BlueprintAction)
+{
+	bool bIsFilteredOut = IsLatent(BlueprintAction); // latent actions are only allowed in ubergraphs
+	return bIsFilteredOut;
+}
+
+//------------------------------------------------------------------------------
+static bool BlueprintActionFilterImpl::IsIncompatibleWithFuncGraph(FBlueprintActionContext const& /*Context*/, UBlueprintNodeSpawner const* BlueprintAction)
+{
+	bool bIsFilteredOut = IsLatent(BlueprintAction); // latent actions are only allowed in ubergraphs
+	return bIsFilteredOut;
+}
+
+//------------------------------------------------------------------------------
+static bool BlueprintActionFilterImpl::IsIncompatibleWithGraphType(FBlueprintActionFilter const& Filter, UBlueprintNodeSpawner const* BlueprintAction)
+{
+	bool bIsFilteredOut = false;
+	FBlueprintActionContext const& FilterContext = Filter.Context;
+	
+	uint32 GraphTypesTestedMask = 0x00;
+	for (UEdGraph* Graph : FilterContext.Graphs)
+	{
+		EGraphType GraphType = GT_MAX;
+		if (UEdGraphSchema_K2* K2Schema = Graph->Schema->GetDefaultObject<UEdGraphSchema_K2>())
+		{
+			GraphType = K2Schema->GetGraphType(Graph);
+		}
+		
+		uint32 GraphTypeBit = (1 << GraphType);
+		if (GraphTypesTestedMask & GraphTypeBit)
+		{
+			continue;
+		}
+		
+		switch (GraphType)
+		{
+		case GT_Ubergraph:
+			{
+				bIsFilteredOut = IsIncompatibleWithUbergraph(FilterContext, BlueprintAction);
+				break;
+			}
+		case GT_Animation:
+			{
+				bIsFilteredOut = IsIncompatibleWithAnimGraph(FilterContext, BlueprintAction);
+				break;
+			}
+		case GT_Macro:
+			{
+				bIsFilteredOut = IsIncompatibleWithMacroGraph(FilterContext, BlueprintAction);
+				break;
+			}
+		case GT_Function:
+			{
+				bIsFilteredOut = IsIncompatibleWithFuncGraph(FilterContext, BlueprintAction);
+				break;
+			}
+		}
+		
+		if (bIsFilteredOut)
+		{
+			break;
+		}
+		GraphTypesTestedMask |= GraphTypeBit;
+	}
+
+	return bIsFilteredOut;
+}
+
+//------------------------------------------------------------------------------
+static bool BlueprintActionFilterImpl::IsNodeUnpastable(FBlueprintActionFilter const& Filter, UBlueprintNodeSpawner const* BlueprintAction)
+{
+	bool bIsFilteredOut = false;
+	FBlueprintActionContext const& FilterContext = Filter.Context;
+	
+	for (UEdGraph* Graph : FilterContext.Graphs)
+	{		
+		UClass* NodeClass = BlueprintAction->NodeClass;
+		if (NodeClass == nullptr)
+		{
+			continue;
+		}
+		else if (NodeClass->IsChildOf<UK2Node_Event>())// @TODO: make all event nodes, this: BlueprintAction->IsA<UBlueprintEventNodeSpawner>())
+		{
+			// events need a valid function to check CanPasteHere() (we can't use the CDO)
+			// @TODO: modify the CDO's function reference, and then restore it
+			continue;
+		}
+
+		// wouldn't want to spawn a temp node here (too costly to be doing while
+		// filtering)... if a node that is unpasteble slips through, then the 
+		// user should be notified through ui feedback way (we purposely do this
+		// with the palette, so that it can have an all encompassing list)
+		if (UEdGraphNode* NodeCDO = NodeClass->GetDefaultObject<UEdGraphNode>())
+		{
+			if (!NodeCDO->CanPasteHere(Graph, Graph->Schema->GetDefaultObject<UEdGraphSchema>()))
+			{
+				bIsFilteredOut = true;
+				break;
+			}
+		}
+	}
+	
+	return bIsFilteredOut;
+}
+
+//------------------------------------------------------------------------------
+static bool BlueprintActionFilterImpl::IsFilteredNodeType(FBlueprintActionFilter const& Filter, UBlueprintNodeSpawner const* BlueprintAction, bool const bExcludeChildClasses)
+{
+	bool bIsFilteredOut = false;
+
+	UClass* NodeClass = BlueprintAction->NodeClass;
+	if (NodeClass != nullptr)
+	{
+		for (UClass* AllowedClass : Filter.NodeTypes)
+		{
+			if ((bExcludeChildClasses && (AllowedClass != NodeClass)) ||
+				!NodeClass->IsChildOf(AllowedClass))
+			{
+				bIsFilteredOut = false;
+				break;
+			}
+		}
+	}
+	else
+	{
+		// this node's class type is unknown, filter it out!
+		bIsFilteredOut = (Filter.NodeTypes.Num() > 0);
+	}
+
+	return bIsFilteredOut;
+}
+
+/*******************************************************************************
+ * FBlueprintActionFilter
+ ******************************************************************************/
+
+//------------------------------------------------------------------------------
+FBlueprintActionFilter::FBlueprintActionFilter(uint32 Flags/*= 0x00*/)
+{
+	using namespace BlueprintActionFilterImpl;
+
+	//
+	// NOTE: The order of these tests can have perf implications, the more one
+	//       rejects on average the sooner it should be added (they're executed
+	//       in order added)
+	//
+
+	if (!(Flags & BPFILTER_IncludeDeprecated))
+	{
+		AddIsFilteredTest(FIsFilteredDelegate::CreateStatic(IsDeprecated));
+	}
+	AddIsFilteredTest(FIsFilteredDelegate::CreateStatic(IsFilteredNodeType, (Flags & BPFILTER_ExcludeChildNodeTypes) != 0));
+	AddIsFilteredTest(FIsFilteredDelegate::CreateStatic(IsEventUnimplementable));
+	AddIsFilteredTest(FIsFilteredDelegate::CreateStatic(IsFieldInaccessible));
+	AddIsFilteredTest(FIsFilteredDelegate::CreateStatic(IsRestrictedClassMember));
+	if (Flags & BPFILTER_ExcludeGlobalFields)
+	{
+		AddIsFilteredTest(FIsFilteredDelegate::CreateStatic(IsGlobalField));
+	}
+	AddIsFilteredTest(FIsFilteredDelegate::CreateStatic(IsExternalField));
+	AddIsFilteredTest(FIsFilteredDelegate::CreateStatic(IsIncompatibleImpureNode));
+	AddIsFilteredTest(FIsFilteredDelegate::CreateStatic(IsPermissionNotGranted));
+	//AddIsFilteredTest(FIsFilteredDelegate::CreateStatic(IsIncompatibleWithGraphType));
+	AddIsFilteredTest(FIsFilteredDelegate::CreateStatic(IsNodeUnpastable));
+	AddIsFilteredTest(FIsFilteredDelegate::CreateStatic(IsFieldCategoryHidden));
+
+	// @TODO: account for all K2ActionMenuBuilder checks...
+	//    NodeCDO->CanCreateUnderSpecifiedSchema(this)
+	//    FBlueprintEditorUtils::DoesSupportTimelines()
+	//    FKismetEditorUtilities::IsClassABlueprintInterface()
+	//    FK2ActionMenuBuilder::GetPaletteActions()
+	//       is allowed in level-script?
+	//       is allowed in non-actor?
+	//       is allowed in actor?
+	//    FK2ActionMenuBuilder::GetAllActionsForClass()
+	//       show variables?
+	//       show protected?
+	//       show delegates?
+	//    FK2ActionMenuBuilder::GetFuncNodesForClass()
+	//       show inherited?
+	//    FK2ActionMenuBuilder::GetContextAllowedNodeTypes()
+	//       impure funcs?
+	//       construction script?
+	//    UEdGraphSchema_K2::CanFunctionBeUsedInClass()
+
+	// @TODO: account for implemented interface methods
+}
+
+//------------------------------------------------------------------------------
+void FBlueprintActionFilter::AddIsFilteredTest(FIsFilteredDelegate IsFilteredDelegate)
+{
+	if (IsFilteredDelegate.IsBound())
+	{
+		FilterTests.Add(IsFilteredDelegate);
+	}
+}
+
+//------------------------------------------------------------------------------
+bool FBlueprintActionFilter::IsFiltered(UBlueprintNodeSpawner const* BlueprintAction)
+{
+	bool bIsFiltered = IsFilteredByThis(BlueprintAction);
+	if (!bIsFiltered)
+	{
+		for (FBlueprintActionFilter& OrFilter : OrFilters)
+		{
+			if (OrFilter.IsFiltered(BlueprintAction))
+			{
+				bIsFiltered = true;
+				break;
+			}
+		}
+	}
+
+	if (bIsFiltered)
+	{
+		for (FBlueprintActionFilter& AndFilter : AndFilters)
+		{
+			bIsFiltered &= AndFilter.IsFiltered(BlueprintAction);
+		}
+	}
+
+	return bIsFiltered;
+}
+
+//------------------------------------------------------------------------------
+FBlueprintActionFilter const& FBlueprintActionFilter::operator|=(FBlueprintActionFilter const& Rhs)
+{
+	OrFilters.Add(Rhs);
+	return *this;
+}
+
+//------------------------------------------------------------------------------
+FBlueprintActionFilter const& FBlueprintActionFilter::operator&=(FBlueprintActionFilter const& Rhs)
+{
+	AndFilters.Add(Rhs);
+	return *this;
+}
+
+//------------------------------------------------------------------------------
+bool FBlueprintActionFilter::IsFilteredByThis(UBlueprintNodeSpawner const* BlueprintAction) const
+{
+	FBlueprintActionFilter const& FilterRef = *this;
+
+	bool bIsFiltered = false;
+	for (int32 TestIndex = 0; TestIndex < FilterTests.Num(); ++TestIndex)
+	{
+		FIsFilteredDelegate const& IsFilteredDelegate = FilterTests[TestIndex];
+		check(IsFilteredDelegate.IsBound());
+
+		if (IsFilteredDelegate.Execute(FilterRef, BlueprintAction))
+		{
+			bIsFiltered = true;
+			break;
+		}
+	}
+	return bIsFiltered;
+}
diff --git a/Engine/Source/Editor/BlueprintGraph/Private/BlueprintActionMenuBuilder.cpp b/Engine/Source/Editor/BlueprintGraph/Private/BlueprintActionMenuBuilder.cpp
new file mode 100644
--- /dev/null
+++ b/Engine/Source/Editor/BlueprintGraph/Private/BlueprintActionMenuBuilder.cpp
@@ -0,0 +1,453 @@
+// Copyright 1998-2014 Epic Games, Inc. All Rights Reserved.
+
+#include "BlueprintGraphPrivatePCH.h"
+#include "BlueprintActionMenuBuilder.h"
+#include "BlueprintActionMenuItem.h"
+#include "BlueprintActionDatabase.h"
+#include "BlueprintActionFilter.h"
+#include "BlueprintNodeSpawner.h"
+#include "BlueprintFunctionNodeSpawner.h"
+#include "K2ActionMenuBuilder.h"			// for FBlueprintPaletteListBuilder/FBlueprintGraphActionListBuilder
+#include "EdGraphSchema_K2.h"				// for StaticClass(), bUseLegacyActionMenus, etc.
+
+#define LOCTEXT_NAMESPACE "BlueprintActionMenuBuilder"
+
+/*******************************************************************************
+ * FBlueprintActionMenuItemFactory
+ ******************************************************************************/
+
+class FBlueprintActionMenuItemFactory
+{
+public:
+	/** 
+	 * Menu item factory constructor. Sets up the blueprint context, which
+	 * is utilized when configuring blueprint menu items' names/tooltips/etc.
+	 *
+	 * @param  Context	The blueprint context for the menu being built.
+	 */
+	FBlueprintActionMenuItemFactory(FBlueprintActionContext const& Context);
+
+	/** A root category to perpend every menu item with */
+	FText RootCategory;
+	/** The menu sort order to set every menu item with */
+	int32 MenuGrouping;
+	
+	/**
+	 * Spawns a new FBlueprintActionMenuItem with the node-spawner. Constructs
+	 * the menu item's category, name, tooltip, etc.
+	 * 
+	 * @param  ItemOwner	The menu item owner that this will be assigned to.
+	 * @param  Action		The node-spawner that the new menu item should wrap.
+	 * @return A newly allocated FBlueprintActionMenuItem (which wraps the supplied action).
+	 */
+	TSharedPtr<FEdGraphSchemaAction> MakeMenuItem(FBlueprintActionMenuBuilder& ItemOwner, UBlueprintNodeSpawner* Action);
+
+private:
+	/**
+	 * Attempts to pull a menu name from the supplied spawner. If one isn't 
+	 * provided, then it spawns a temporary node and pulls one from that node's 
+	 * title.
+	 * 
+	 * @param  ItemOwner	The menu builder that this name will be listed from.
+	 * @param  Action		The action you want to suss name information from.
+	 * @return A name for the menu item wrapping this action.
+	 */
+	FText GetMenuNameForAction(FBlueprintActionMenuBuilder& ItemOwner, UBlueprintNodeSpawner* Action);
+
+	/**
+	 * Attempts to pull a menu category from the supplied spawner. If one isn't 
+	 * provided, then it spawns a temporary node and pulls one from that node.
+	 * 
+	 * @param  ItemOwner	The menu builder that the item will be listed from.
+	 * @param  Action		The action you want to suss category information from.
+	 * @return A category for the menu item wrapping this action.
+	 */
+	FText GetCategoryForAction(FBlueprintActionMenuBuilder& ItemOwner, UBlueprintNodeSpawner* Action);
+
+	/**
+	 * Attempts to pull a menu tooltip from the supplied spawner. If one isn't 
+	 * provided, then it spawns a temporary node and pulls one from that node.
+	 * 
+	 * @param  ItemOwner	The menu builder that the item will be listed from.
+	 * @param  Action		The action you want to suss tooltip information from.
+	 * @return A tooltip for the menu item wrapping this action.
+	 */
+	FText GetTooltipForAction(FBlueprintActionMenuBuilder& ItemOwner, UBlueprintNodeSpawner* Action);
+
+	/**
+	 * Attempts to pull a keywords from the supplied spawner. If one isn't 
+	 * provided, then it spawns a temporary node and pulls them from that.
+	 * 
+	 * @param  ItemOwner	The menu builder that the item will be listed from.
+	 * @param  Action		The action you want to suss keyword information from.
+	 * @return A keyword text string for the menu item wrapping this action.
+	 */
+	FText GetSearchKeywordsForAction(FBlueprintActionMenuBuilder& ItemOwner, UBlueprintNodeSpawner* Action);
+
+	/**
+	 * Utility getter function that retrieves the blueprint context for the menu
+	 * items being made.
+	 * 
+	 * @return The first blueprint out of the cached FBlueprintActionContext.
+	 */
+	UBlueprint* GetBlueprint() const;
+
+	/**
+	 * Utility function that will pull the "OwnerOfTemporaries" field from 
+	 * ItemOwner. If the object is null, then this will allocate one on behalf
+	 * of ItemOwner.
+	 * 
+	 * @param  The FBlueprintActionMenuBuilder you want an "OwnerOfTemporaries" from. 
+	 * @return The FBlueprintActionMenuBuilder's OwnerOfTemporaries.
+	 */
+	UEdGraph* GetOwnerOfTemporaries(FBlueprintActionMenuBuilder& ItemOwner) const;
+
+	/** Cached context for the blueprint menu being built */
+	FBlueprintActionContext const& Context;
+};
+
+//------------------------------------------------------------------------------
+FBlueprintActionMenuItemFactory::FBlueprintActionMenuItemFactory(FBlueprintActionContext const& ContextIn)
+	: Context(ContextIn)
+	, RootCategory(FText::GetEmpty())
+	, MenuGrouping(0)
+{
+	
+}
+
+//------------------------------------------------------------------------------
+TSharedPtr<FEdGraphSchemaAction> FBlueprintActionMenuItemFactory::MakeMenuItem(FBlueprintActionMenuBuilder& ItemOwner, UBlueprintNodeSpawner* Action)
+{
+	FBlueprintActionMenuItem* NewMenuItem = new FBlueprintActionMenuItem(Action, MenuGrouping);
+	
+	NewMenuItem->MenuDescription    = GetMenuNameForAction(ItemOwner, Action);
+	NewMenuItem->TooltipDescription = GetTooltipForAction(ItemOwner, Action).ToString();
+	NewMenuItem->Category           = GetCategoryForAction(ItemOwner, Action).ToString();
+	NewMenuItem->Keywords           = GetSearchKeywordsForAction(ItemOwner, Action).ToString();
+
+	if (NewMenuItem->Category.IsEmpty())
+	{
+		NewMenuItem->Category = RootCategory.ToString();
+	}
+	else if (!RootCategory.IsEmpty())
+	{
+		NewMenuItem->Category = FString::Printf(TEXT("%s|%s"), *RootCategory.ToString(), *NewMenuItem->Category);
+	}
+	
+	return MakeShareable(NewMenuItem);
+}
+
+//------------------------------------------------------------------------------
+FText FBlueprintActionMenuItemFactory::GetMenuNameForAction(FBlueprintActionMenuBuilder& ItemOwner, UBlueprintNodeSpawner* Action)
+{
+	check(Action != nullptr);
+	// give the action the chance to save on performance (to keep from having to spawn a template node)
+	FText MenuName = Action->GetDefaultMenuName();
+	
+	if (MenuName.IsEmpty())
+	{
+		if (UEdGraphNode* NodeTemplate = Action->MakeTemplateNode(GetOwnerOfTemporaries(ItemOwner)))
+		{
+			MenuName = NodeTemplate->GetNodeTitle(ENodeTitleType::ListView);
+		}
+		else
+		{
+			// need to give it some name, this is as good as any I guess
+			MenuName = FText::FromName(Action->GetFName());
+		}
+	}
+	
+	return MenuName;
+}
+
+//------------------------------------------------------------------------------
+FText FBlueprintActionMenuItemFactory::GetCategoryForAction(FBlueprintActionMenuBuilder& ItemOwner, UBlueprintNodeSpawner* Action)
+{
+	check(Action != nullptr);
+	// give the action the chance to save on performance (to keep from having to spawn a template node)
+	FText MenuCategory = Action->GetDefaultMenuCategory();
+	
+	if (MenuCategory.IsEmpty())
+	{
+		UEdGraph* OwnerOfTemporaries = GetOwnerOfTemporaries(ItemOwner);
+		// @TODO: consider moving GetMenuCategory() up into UEdGraphNode
+		if (UK2Node* NodeTemplate = Cast<UK2Node>(Action->MakeTemplateNode(OwnerOfTemporaries)))
+		{
+			MenuCategory = NodeTemplate->GetMenuCategory();
+		}
+		
+		// if the category is still empty, then try and construct the category
+		// from the spawner's type
+		if (MenuCategory.IsEmpty())
+		{
+			// put uncategorized function calls in a member function category
+			// (sorted by their respective classes)
+			if (UBlueprintFunctionNodeSpawner* FuncSpawner = Cast<UBlueprintFunctionNodeSpawner>(Action))
+			{
+				UFunction const* Function = FuncSpawner->GetFunction();
+				check(Function != nullptr);
+				UClass* FuncOwner = Function->GetOuterUClass();
+				
+				UBlueprint* Blueprint = GetBlueprint();
+				check(Blueprint != nullptr);
+				
+				// if this is NOT a self function call (self function calls
+				// don't get nested any deeper)
+				if (!Blueprint->GeneratedClass->IsChildOf(FuncOwner))
+				{
+					MenuCategory = FText::FromName(FuncOwner->GetFName());
+				}
+				MenuCategory = FText::Format(LOCTEXT("MemberFunctionsCategory", "Call Funtions|{0}"), MenuCategory);
+			}
+		}
+	}
+	
+	return MenuCategory;
+}
+
+//------------------------------------------------------------------------------
+FText FBlueprintActionMenuItemFactory::GetTooltipForAction(FBlueprintActionMenuBuilder& ItemOwner, UBlueprintNodeSpawner* Action)
+{
+	check(Action != nullptr);
+	// give the action the chance to save on performance (to keep from having to spawn a template node)
+	FText Tooltip = Action->GetDefaultMenuTooltip();
+	
+	if (Tooltip.IsEmpty())
+	{
+		if (UEdGraphNode* NodeTemplate = Action->MakeTemplateNode(GetOwnerOfTemporaries(ItemOwner)))
+		{
+			Tooltip = FText::FromString(NodeTemplate->GetTooltip());
+		}
+	}
+	
+	return Tooltip;
+}
+
+//------------------------------------------------------------------------------
+FText FBlueprintActionMenuItemFactory::GetSearchKeywordsForAction(FBlueprintActionMenuBuilder& ItemOwner, UBlueprintNodeSpawner* Action)
+{
+	check(Action != nullptr);
+	// give the action the chance to save on performance (to keep from having to spawn a template node)
+	FText SearchKeywords = Action->GetDefaultSearchKeywords();
+	
+	if (SearchKeywords.IsEmpty())
+	{
+		if (UEdGraphNode* NodeTemplate = Action->MakeTemplateNode(GetOwnerOfTemporaries(ItemOwner)))
+		{
+			SearchKeywords = FText::FromString(NodeTemplate->GetKeywords());
+		}
+	}
+	
+	return SearchKeywords;
+}
+
+//------------------------------------------------------------------------------
+UBlueprint* FBlueprintActionMenuItemFactory::GetBlueprint() const
+{
+	UBlueprint* Blueprint = nullptr;
+	if (Context.Blueprints.Num() > 0)
+	{
+		Blueprint = Context.Blueprints[0];
+	}
+	return Blueprint;
+}
+
+//------------------------------------------------------------------------------
+UEdGraph* FBlueprintActionMenuItemFactory::GetOwnerOfTemporaries(FBlueprintActionMenuBuilder& ItemOwner) const
+{
+	UBlueprint* Blueprint = GetBlueprint();
+	check(Blueprint != nullptr);
+	
+	UEdGraph*& OwnerOfTemporaries = ItemOwner.OwnerOfTemporaries;
+	if (OwnerOfTemporaries == nullptr)
+	{
+		// unfortunately the OwnerOfTemporaries cannont be directly assigned to
+		// the transient package because various nodes relay on the blueprint
+		// outer (just like they rely on this graph outer); for example, the
+		// call-function node uses the blueprint to see if the function has
+		// self-context
+		OwnerOfTemporaries = NewObject<UEdGraph>(Blueprint);
+		OwnerOfTemporaries->SetFlags(RF_Transient);
+		
+		OwnerOfTemporaries->Schema = UEdGraphSchema_K2::StaticClass();
+		if (Context.Graphs.Num() > 0)
+		{
+			OwnerOfTemporaries->Schema = Context.Graphs[0]->Schema;
+		}
+		
+	}
+	return OwnerOfTemporaries;
+}
+
+/*******************************************************************************
+ * Static FBlueprintActionMenuBuilder Helpers
+ ******************************************************************************/
+
+namespace FBlueprintActionMenuBuilderImpl
+{
+	/** 
+	 * Defines a sub-section of the overall blueprint menu (filter, heading, etc.)
+	 */
+	struct FMenuSectionDefinition
+	{
+		FMenuSectionDefinition(FBlueprintActionFilter const& SectionFilter);
+
+		/** A filter for this section of the menu */
+		FBlueprintActionFilter Filter;
+		/** In charge of spawning menu items for this section (holds category/ordering information)*/
+		FBlueprintActionMenuItemFactory ItemFactory;
+	};
+
+	/**
+	 * To offer a fallback in case the action menu refactor is unstable, this
+	 * method implements the old way we used collect blueprint menu actions (for
+	 * both the palette and context menu).
+	 * 
+	 * @param  MenuSections	All the sections that were defined for FBlueprintActionMenuBuilder (we attempt to combine them)
+	 * @param  MenuBuilder	The menu builder we want all the legacy actions appended to.
+	 * @return 
+	 */
+	static void AppendLegacyActions(TArray< TSharedRef<FMenuSectionDefinition> > const& MenuSections, FBlueprintActionMenuBuilder& MenuBuilder);
+}
+
+//------------------------------------------------------------------------------
+FBlueprintActionMenuBuilderImpl::FMenuSectionDefinition::FMenuSectionDefinition(FBlueprintActionFilter const& SectionFilterIn)
+	: ItemFactory(SectionFilterIn.Context)
+	, Filter(SectionFilterIn)
+{
+}
+
+//------------------------------------------------------------------------------
+void FBlueprintActionMenuBuilderImpl::AppendLegacyActions(TArray< TSharedRef<FMenuSectionDefinition> > const& MenuSections, FBlueprintActionMenuBuilder& MenuBuilder)
+{
+	check(MenuSections.Num() > 0);
+	FBlueprintActionContext const& Context = MenuSections[0]->Filter.Context;
+
+	if (Context.Graphs.Num() > 0)
+	{
+		FBlueprintGraphActionListBuilder ContextMenuBuilder(Context.Graphs[0]);
+		if (Context.Pins.Num() > 0)
+		{
+			ContextMenuBuilder.FromPin = Context.Pins[0];
+		}
+		
+		// @TODO: ContextMenuBuilder.SelectedObjects
+// 		for (TSharedRef<FMenuSectionDefinition> MenuSection : MenuSections)
+// 		{
+// 			FBlueprintActionFilter const& SectionFilter = MenuSection->Filter;
+// 		}
+
+		UEdGraphSchema const* GraphSchema = GetDefault<UEdGraphSchema>(ContextMenuBuilder.CurrentGraph->Schema);
+		GraphSchema->GetGraphContextActions(ContextMenuBuilder);
+
+		MenuBuilder.Append(ContextMenuBuilder);
+	}
+	else if (Context.Blueprints.Num() > 0)
+	{
+		FBlueprintPaletteListBuilder PaletteBuilder(Context.Blueprints[0], MenuSections[0]->ItemFactory.RootCategory.ToString());
+
+		UClass* FilterClass = nullptr;
+
+		FBlueprintActionFilter const& SectionFilter = MenuSections[0]->Filter;
+		if (SectionFilter.OwnerClasses.Num() > 0)
+		{
+			FilterClass = SectionFilter.OwnerClasses[0];
+		}
+
+		UEdGraphSchema_K2 const* K2Schema = GetDefault<UEdGraphSchema_K2>();
+		FK2ActionMenuBuilder(K2Schema).GetPaletteActions(PaletteBuilder, FilterClass); 
+
+		MenuBuilder.Append(PaletteBuilder);
+	}
+}
+
+/*******************************************************************************
+ * FBlueprintActionMenuBuilder
+ ******************************************************************************/
+
+//------------------------------------------------------------------------------
+FBlueprintActionMenuBuilder::FBlueprintActionMenuBuilder(FBlueprintActionFilter const& Filter, bool bAutoBuild/*= true*/)
+{
+	AddMenuSection(Filter);
+	if (bAutoBuild)
+	{
+		RebuildActionList();
+	}
+}
+
+//------------------------------------------------------------------------------
+void FBlueprintActionMenuBuilder::AddMenuSection(FBlueprintActionFilter const& Filter, FText const& Heading/* = FText::GetEmpty()*/, int32 MenuOrder/* = 0*/)
+{
+	using namespace FBlueprintActionMenuBuilderImpl;
+	
+	TSharedRef<FMenuSectionDefinition> SectionDescRef = MakeShareable(new FMenuSectionDefinition(Filter));
+	SectionDescRef->ItemFactory.RootCategory = Heading;
+	SectionDescRef->ItemFactory.MenuGrouping = MenuOrder;
+
+	MenuSections.Add(SectionDescRef);
+}
+
+//------------------------------------------------------------------------------
+void FBlueprintActionMenuBuilder::RebuildActionList()
+{
+	using namespace FBlueprintActionMenuBuilderImpl;
+
+	Empty();
+	if (true)//!GetDefault<UEdGraphSchema_K2>()->bUseLegacyActionMenus)
+	{
+		GenerateMenuSections();
+	}
+	else if (MenuSections.Num() > 0)
+	{
+		// use the context from the first section (the legacy menu building did
+		// it all in one sweep and knew how to mark off sections)
+		AppendLegacyActions(MenuSections, *this);
+	}
+
+	// @TODO: account for all K2ActionMenuBuilder action types...
+	// - FEdGraphSchemaAction_K2AddTimeline
+	// - FEdGraphSchemaAction_K2ViewNode
+	// - FEdGraphSchemaAction_K2AddCustomEvent
+	//   FEdGraphSchemaAction_EventFromFunction
+	//   FEdGraphSchemaAction_K2Var
+	//   FEdGraphSchemaAction_K2Delegate
+	//   FEdGraphSchemaAction_K2AssignDelegate
+	// - FEdGraphSchemaAction_K2AddComment
+	//   FEdGraphSchemaAction_K2PasteHere
+	// - FEdGraphSchemaAction_K2NewNode
+	//   FEdGraphSchemaAction_Dummy
+	//   FEdGraphSchemaAction_K2AddCallOnActor
+	//   FEdGraphSchemaAction_K2AddCallOnVariable
+	// - FEdGraphSchemaAction_K2AddComponent
+	//   FEdGraphSchemaAction_K2AddComment
+}
+
+//------------------------------------------------------------------------------
+void FBlueprintActionMenuBuilder::GenerateMenuSections()
+{
+	using namespace FBlueprintActionMenuBuilderImpl;
+
+	FBlueprintActionDatabase::FClassActionMap const& ClassActions = FBlueprintActionDatabase::Get().GetAllActions();
+	for (auto ActionEntry : ClassActions)
+	{
+		UClass* Class = ActionEntry.Key;
+		for (UBlueprintNodeSpawner* NodeSpawner : ActionEntry.Value)
+		{
+			for (TSharedRef<FMenuSectionDefinition> MenuSection : MenuSections)
+			{
+				if (MenuSection->Filter.IsFiltered(NodeSpawner))
+				{
+					// this node spawner doesn't belong in this section of the menu
+					continue;
+				}
+
+				TSharedPtr<FEdGraphSchemaAction> NewActionPtr = MenuSection->ItemFactory.MakeMenuItem(*this, NodeSpawner);
+				check(NewActionPtr.IsValid());
+
+				AddAction(NewActionPtr);
+			}
+		}
+	}
+}
+
+#undef LOCTEXT_NAMESPACE
diff --git a/Engine/Source/Editor/BlueprintGraph/Public/BlueprintActionFilter.h b/Engine/Source/Editor/BlueprintGraph/Public/BlueprintActionFilter.h
new file mode 100644
--- /dev/null
+++ b/Engine/Source/Editor/BlueprintGraph/Public/BlueprintActionFilter.h
@@ -0,0 +1,151 @@
+// Copyright 1998-2014 Epic Games, Inc. All Rights Reserved.
+
+#pragma once
+
+// Forward declarations
+class UBlueprintNodeSpawner;
+class UBlueprint;
+
+/*******************************************************************************
+ * FBlueprintActionContext
+ ******************************************************************************/
+
+struct FBlueprintActionContext
+{
+	/** 
+	 * A list of all blueprints you want actions for. Generally, this will
+	 * only contain a single blueprint, but it can have many (where an action
+	 * has to be available for every blueprint listed to pass the filter).
+	 */
+	TArray<UBlueprint*> Blueprints;
+
+	/** 
+	 * A list of graphs you want compatible actions for. Generally, this will
+	 * contain a single graph, but it can have several (where an action has to 
+	 * be viable for every graph to pass the filter).
+	 */
+	TArray<UEdGraph*> Graphs;
+
+	/** 
+	 * A list of pins you want compatible actions for. Generally, this will
+	 * contain a single pin, but it can have several (where an action has to 
+	 * be viable for every pin to pass the filter).
+	 */
+	TArray<UEdGraphPin*> Pins;
+};
+
+
+/*******************************************************************************
+ * FBlueprintActionFilter
+ ******************************************************************************/
+
+class BLUEPRINTGRAPH_API FBlueprintActionFilter
+{
+public:
+	/** The filter uses a series of rejection tests matching */
+	DECLARE_DELEGATE_RetVal_TwoParams(bool, FIsFilteredDelegate, FBlueprintActionFilter const&, UBlueprintNodeSpawner const*);
+
+public:
+	enum EFlags // Flags, which configure certain rejection tests.
+	{
+		/** Deprecated class actions will not be filtered out*/
+		BPFILTER_IncludeDeprecated		= 0x00000001,
+
+		/** 
+		 * Excludes nodes associated with global field types (functions, 
+		 * structs, enums, etc.); NOTE: that static function library members are 
+		 * counted as "global". 
+		 */
+		BPFILTER_ExcludeGlobalFields	= 0x00000002,
+
+		/**
+		 * Makes NodeTypes tests more aggressive by excluding nodes sub-classes, 
+		 * nodes have to explicitly match the class listed. 
+		 */
+		BPFILTER_ExcludeChildNodeTypes	= 0x00000004,
+	};
+	FBlueprintActionFilter(uint32 const Flags = 0x00);
+	
+	/**
+	 * A list of allowed node types. If an action's NodeClass isn't one of
+	 * these types, then it is filtered out. Use the "ExcludeChildNodeTypes"
+	 * flag to aggressively filter out child classes (must be an explicit match).
+	 */
+	FBlueprintActionContext Context;	
+	
+	/**
+	 * A list of allowed node types. If an action's NodeClass isn't one of
+	 * these types, then it is filtered out. Use the "ExcludeChildNodeTypes"
+	 * flag to aggressively filter out child classes (must be an explicit match).
+	 */
+	TArray<UClass*> NodeTypes;
+	
+	/**
+	 * A list of classes that you want members for. If an action does not have 
+	 * an associated field, or the associated field doesn't belong to any of 
+	 * these classes, then it is filtered out.
+	 */
+	TArray<UClass*> OwnerClasses;
+
+	/**
+	 * Users can extend the filter and add their own rejection tests with this
+	 * method. We use rejection "IsFiltered" tests rather than inclusive tests 
+	 * because it is more optimal to whittle down the list of actions early.
+	 * 
+	 * @param  IsFilteredDelegate	The rejection test you wish to add to this filter.
+	 */
+	void AddIsFilteredTest(FIsFilteredDelegate IsFilteredDelegate);
+
+	/**
+	 * Query to check and see if the specified action gets filtered out by this 
+	 * (and any and'd/or'd filters). NOT marked const to keep 
+	 * FIsFilteredDelegates from recursively calling.
+	 * 
+	 * @param  BlueprintAction	The node-spawner you wish to test.
+	 * @return False if the action passes the filter, otherwise false (the action got filtered out).
+	 */
+	bool IsFiltered(UBlueprintNodeSpawner const* BlueprintAction);
+
+	/**
+	 * Appends another filter to be utilized in IsFiltered() queries, extending  
+	 * the query to be: IsFilteredByThis() || Rhs.IsFiltered()
+	 *
+	 * NOTE: Appending additional filters can hinder performance (as the 
+	 *       they have to loop through all their tests on top of these)
+	 * 
+	 * @param  Rhs	The filter to append to this one.
+	 * @return This.
+	 */
+	FBlueprintActionFilter const& operator|=(FBlueprintActionFilter const& Rhs);
+
+	/**
+	 * Appends another filter to be utilized in IsFiltered() queries, extending  
+	 * the query to be: IsFilteredByThis() && Rhs.IsFiltered()
+	 *
+	 * NOTE: Appending additional filters can hinder performance (as the 
+	 *       they have to loop through all their tests on top of these)
+	 * 
+	 * @param  Rhs	The filter to append to this one.
+	 * @return This.
+	 */
+	FBlueprintActionFilter const& operator&=(FBlueprintActionFilter const& Rhs);
+
+private:
+	/**
+	 * Query to check and see if the specified action gets filtered out by this 
+	 * (does not take into consideration any and'd/or'd filters).
+	 * 
+	 * @param  BlueprintAction	The node-spawner you wish to test.
+	 * @return False if the action passes the filter, otherwise false (the action got filtered out).
+	 */
+	bool IsFilteredByThis(UBlueprintNodeSpawner const* BlueprintAction) const;
+
+	/** Set of rejection tests for this specific filter. */
+	TArray<FIsFilteredDelegate> FilterTests;
+
+	/** Filters to be logically and'd in with the IsFilteredByThis() result. */
+	TArray<FBlueprintActionFilter> AndFilters;
+
+	/** Alternative filters to be logically or'd in with the IsFilteredByThis() result. */
+	TArray<FBlueprintActionFilter> OrFilters;
+};
\ No newline at end of file
diff --git a/Engine/Source/Editor/BlueprintGraph/Public/BlueprintActionMenuBuilder.h b/Engine/Source/Editor/BlueprintGraph/Public/BlueprintActionMenuBuilder.h
new file mode 100644
--- /dev/null
+++ b/Engine/Source/Editor/BlueprintGraph/Public/BlueprintActionMenuBuilder.h
@@ -0,0 +1,69 @@
+// Copyright 1998-2014 Epic Games, Inc. All Rights Reserved.
+
+#pragma once
+
+#include "EdGraph/EdGraphSchema.h" // for FGraphActionListBuilderBase
+
+// Forward declarations
+namespace FBlueprintActionMenuBuilderImpl
+{
+	struct FMenuSectionDefinition; // internal type, forward declared to hide implementation details
+};
+class  FBlueprintActionFilter;
+struct FBlueprintActionContext;
+class  UActorComponent;
+
+/**
+ * Responsible for constructing a list of viable blueprint actions. Runs the 
+ * blueprint actions database through a filter and spawns a series of 
+ * FBlueprintActionMenuItems for actions that pass. Takes care of generating the 
+ * each menu item's category/name/etc.
+ */
+struct BLUEPRINTGRAPH_API FBlueprintActionMenuBuilder : public FGraphActionListBuilderBase
+{
+public:
+	/** Default constructor */
+	FBlueprintActionMenuBuilder() {}
+
+	/**
+	 * Good for constructing a menu built from a single filter. If the 
+	 * bAutoBuild param is set to false, then know you have to manually call
+	 * RebuildActionList() to generate the menu.
+	 *
+	 * @param  Filter		The filter you want applied for this main section of the menu.
+	 * @param  bAutoBuild	If set, will automatically call RebuildActionList() here in the constructor.
+	 */
+	FBlueprintActionMenuBuilder(FBlueprintActionFilter const& Filter, bool bAutoBuild = true);
+
+	/**
+	 * Some action menus require multiple sections. One option is to create 
+	 * multiple FBlueprintActionMenuBuilders and append them together, but that
+	 * can be unperformant (each builder will run through the entire database
+	 * separately)... This method provides an alternative, where you can specify
+	 * a separate filter/heading/ordering for a sub-section of the menu.
+	 *
+	 * @param  Filter	 The filter you want applied to this section of the menu.
+	 * @param  Heading	 The root category for this section of the menu (can be left blank).
+	 * @param  MenuOrder The sort order to assign this section of the menu (higher numbers come first).
+	 */
+	void AddMenuSection(FBlueprintActionFilter const& Filter, FText const& Heading = FText::GetEmpty(), int32 MenuOrder = 0);
+	
+	/**
+	 * Regenerates the entire menu list from the cached menu sections.
+	 */
+	void RebuildActionList();
+
+private:
+	/**
+	 * Filters and adds action items from the blueprint action database (as 
+	 * defined by the MenuSections list).
+	 */
+	void GenerateMenuSections();
+	
+	/** 
+	 * Defines all the separate sections of the menu (filter, sort order, etc.).
+	 * Defined as a TSharedRef<> so as to hide the implementation details (keep 
+	 * this API clean).
+	 */
+	TArray< TSharedRef<FBlueprintActionMenuBuilderImpl::FMenuSectionDefinition> > MenuSections;
+};
\ No newline at end of file
