(WIP) Setting up make/break struct nodes so that entries for global structs now get added to the blueprint action database (for the blueprint action menu refactor).

diff --git a/Engine/Source/Editor/BlueprintGraph/Classes/K2Node_BreakStruct.h b/Engine/Source/Editor/BlueprintGraph/Classes/K2Node_BreakStruct.h
--- a/Engine/Source/Editor/BlueprintGraph/Classes/K2Node_BreakStruct.h
+++ b/Engine/Source/Editor/BlueprintGraph/Classes/K2Node_BreakStruct.h
@@ -25,6 +25,8 @@ class UK2Node_BreakStruct : public UK2Node_StructMemberGet
 	virtual bool DrawNodeAsVariable() const override { return false; }
 	virtual ERedirectType DoPinsMatchForReconstruction(const UEdGraphPin* NewPin, int32 NewPinIndex, const UEdGraphPin* OldPin, int32 OldPinIndex) const override;
 	virtual class FNodeHandlingFunctor* CreateNodeHandler(class FKismetCompilerContext& CompilerContext) const override;
+	virtual void GetMenuActions(TArray<UBlueprintNodeSpawner*>& ActionListOut) const override;
+	virtual FText GetMenuCategory() const override;
 	// End K2Node interface
 };
 
diff --git a/Engine/Source/Editor/BlueprintGraph/Classes/K2Node_MakeStruct.h b/Engine/Source/Editor/BlueprintGraph/Classes/K2Node_MakeStruct.h
--- a/Engine/Source/Editor/BlueprintGraph/Classes/K2Node_MakeStruct.h
+++ b/Engine/Source/Editor/BlueprintGraph/Classes/K2Node_MakeStruct.h
@@ -27,6 +27,8 @@ class UK2Node_MakeStruct : public UK2Node_StructMemberSet
 	virtual bool DrawNodeAsVariable() const override { return false; }
 	virtual class FNodeHandlingFunctor* CreateNodeHandler(class FKismetCompilerContext& CompilerContext) const override;
 	virtual ERedirectType DoPinsMatchForReconstruction(const UEdGraphPin* NewPin, int32 NewPinIndex, const UEdGraphPin* OldPin, int32 OldPinIndex)  const override;
+	virtual void GetMenuActions(TArray<UBlueprintNodeSpawner*>& ActionListOut) const override;
+	virtual FText GetMenuCategory() const override;
 	// End K2Node interface
 
 protected:
diff --git a/Engine/Source/Editor/BlueprintGraph/Private/K2Node_BreakStruct.cpp b/Engine/Source/Editor/BlueprintGraph/Private/K2Node_BreakStruct.cpp
--- a/Engine/Source/Editor/BlueprintGraph/Private/K2Node_BreakStruct.cpp
+++ b/Engine/Source/Editor/BlueprintGraph/Private/K2Node_BreakStruct.cpp
@@ -1,8 +1,9 @@
 // Copyright 1998-2014 Epic Games, Inc. All Rights Reserved.
 
 #include "BlueprintGraphPrivatePCH.h"
-
 #include "KismetCompiler.h"
+#include "BlueprintNodeSpawner.h"
+#include "EditorCategoryUtils.h"
 
 #define LOCTEXT_NAMESPACE "K2Node_BreakStruct"
 
@@ -242,4 +243,42 @@ FNodeHandlingFunctor* UK2Node_BreakStruct::CreateNodeHandler(class FKismetCompil
 	return new FKCHandler_BreakStruct(CompilerContext);
 }
 
+void UK2Node_BreakStruct::GetMenuActions(TArray<UBlueprintNodeSpawner*>& ActionListOut) const
+{
+	for (TObjectIterator<UScriptStruct> StructIt; StructIt; ++StructIt)
+	{
+		UScriptStruct const* Struct = (*StructIt);
+		// we only want to add autonomous structs here; those belonging to a 
+		// certain class should instead be associated with that class (so when 
+		// the class is modified we can easily handle any structs that were changed).
+		bool bIsStandaloneStruct = Struct->GetOuter()->IsA(UPackage::StaticClass());
+
+		if (!bIsStandaloneStruct || !UEdGraphSchema_K2::IsAllowableBlueprintVariableType(Struct) || !CanBeBroken(Struct))
+		{
+			continue;
+		}
+
+		auto CustomizeBreakNodeLambda = [](UEdGraphNode* NewNode, bool bIsTemplateNode, TWeakObjectPtr<UScriptStruct> StructPtr)
+		{
+			UK2Node_BreakStruct* BreakNode = CastChecked<UK2Node_BreakStruct>(NewNode);
+			if (StructPtr.IsValid())
+			{
+				BreakNode->StructType = StructPtr.Get();
+			}
+		};
+
+		UBlueprintNodeSpawner* NodeSpawner = UBlueprintNodeSpawner::Create(GetClass());
+		check(NodeSpawner != nullptr);
+		ActionListOut.Add(NodeSpawner);
+
+		TWeakObjectPtr<UScriptStruct> StructPtr = Struct;
+		NodeSpawner->CustomizeNodeDelegate = UBlueprintNodeSpawner::FCustomizeNodeDelegate::CreateStatic(CustomizeBreakNodeLambda, StructPtr);
+	}	
+}
+
+FText UK2Node_BreakStruct::GetMenuCategory() const
+{
+	return FEditorCategoryUtils::GetCommonCategory(FCommonEditorCategory::Struct);
+}
+
 #undef LOCTEXT_NAMESPACE
diff --git a/Engine/Source/Editor/BlueprintGraph/Private/K2Node_MakeStruct.cpp b/Engine/Source/Editor/BlueprintGraph/Private/K2Node_MakeStruct.cpp
--- a/Engine/Source/Editor/BlueprintGraph/Private/K2Node_MakeStruct.cpp
+++ b/Engine/Source/Editor/BlueprintGraph/Private/K2Node_MakeStruct.cpp
@@ -3,6 +3,8 @@
 #include "BlueprintGraphPrivatePCH.h"
 #include "KismetCompiler.h"
 #include "MakeStructHandler.h"
+#include "BlueprintNodeSpawner.h"
+#include "EditorCategoryUtils.h"
 
 #define LOCTEXT_NAMESPACE "K2Node_MakeStruct"
 
@@ -213,4 +215,42 @@ UK2Node::ERedirectType UK2Node_MakeStruct::DoPinsMatchForReconstruction(const UE
 	return Result;
 }
 
+void UK2Node_MakeStruct::GetMenuActions(TArray<UBlueprintNodeSpawner*>& ActionListOut) const
+{
+	for (TObjectIterator<UScriptStruct> StructIt; StructIt; ++StructIt)
+	{
+		UScriptStruct const* Struct = (*StructIt);
+		// we only want to add autonomous structs here; those belonging to a 
+		// certain class should instead be associated with that class (so when 
+		// the class is modified we can easily handle any structs that were changed).
+		bool bIsStandaloneStruct = Struct->GetOuter()->IsA(UPackage::StaticClass());
+
+		if (!bIsStandaloneStruct || !UEdGraphSchema_K2::IsAllowableBlueprintVariableType(Struct) || !CanBeMade(Struct))
+		{
+			continue;
+		}
+
+		auto CustomizeMakeNodeLambda = [](UEdGraphNode* NewNode, bool bIsTemplateNode, TWeakObjectPtr<UScriptStruct> StructPtr)
+		{
+			UK2Node_MakeStruct* MakeNode = CastChecked<UK2Node_MakeStruct>(NewNode);
+			if (StructPtr.IsValid())
+			{
+				MakeNode->StructType = StructPtr.Get();
+			}
+		};
+
+		UBlueprintNodeSpawner* NodeSpawner = UBlueprintNodeSpawner::Create(GetClass());
+		check(NodeSpawner != nullptr);
+		ActionListOut.Add(NodeSpawner);
+
+		TWeakObjectPtr<UScriptStruct> StructPtr = Struct;
+		NodeSpawner->CustomizeNodeDelegate = UBlueprintNodeSpawner::FCustomizeNodeDelegate::CreateStatic(CustomizeMakeNodeLambda, StructPtr);
+	}
+}
+
+FText UK2Node_MakeStruct::GetMenuCategory() const
+{
+	return FEditorCategoryUtils::GetCommonCategory(FCommonEditorCategory::Struct);
+}
+
 #undef LOCTEXT_NAMESPACE
