(WIP) Enum related nodes adding themselves to the blueprint action database (accounting for global enums... all for the blueprint action menu refactor).

diff --git a/Engine/Source/Editor/BlueprintGraph/Classes/K2Node_CastByteToEnum.h b/Engine/Source/Editor/BlueprintGraph/Classes/K2Node_CastByteToEnum.h
--- a/Engine/Source/Editor/BlueprintGraph/Classes/K2Node_CastByteToEnum.h
+++ b/Engine/Source/Editor/BlueprintGraph/Classes/K2Node_CastByteToEnum.h
@@ -31,6 +31,8 @@ class UK2Node_CastByteToEnum : public UK2Node
 	virtual bool IsNodePure() const override { return true; }
 	FNodeHandlingFunctor* CreateNodeHandler(FKismetCompilerContext& CompilerContext) const override;
 	virtual void ExpandNode(class FKismetCompilerContext& CompilerContext, UEdGraph* SourceGraph) override;
+	virtual void GetMenuActions(TArray<UBlueprintNodeSpawner*>& ActionListOut) const override;
+	virtual FText GetMenuCategory() const override;
 	// End UK2Node interface
 
 	virtual FName GetFunctionName() const;
diff --git a/Engine/Source/Editor/BlueprintGraph/Classes/K2Node_EnumEquality.h b/Engine/Source/Editor/BlueprintGraph/Classes/K2Node_EnumEquality.h
--- a/Engine/Source/Editor/BlueprintGraph/Classes/K2Node_EnumEquality.h
+++ b/Engine/Source/Editor/BlueprintGraph/Classes/K2Node_EnumEquality.h
@@ -27,6 +27,8 @@ class UK2Node_EnumEquality : public UK2Node
 	virtual void NotifyPinConnectionListChanged(UEdGraphPin* Pin) override;
 	virtual class FNodeHandlingFunctor* CreateNodeHandler(class FKismetCompilerContext& CompilerContext) const override;
 	virtual void ExpandNode(class FKismetCompilerContext& CompilerContext, UEdGraph* SourceGraph) override;
+	virtual void GetMenuActions(TArray<UBlueprintNodeSpawner*>& ActionListOut) const override;
+	virtual FText GetMenuCategory() const override;
 	// End UK2Node interface
 
 	/** Get the return value pin */
diff --git a/Engine/Source/Editor/BlueprintGraph/Classes/K2Node_EnumLiteral.h b/Engine/Source/Editor/BlueprintGraph/Classes/K2Node_EnumLiteral.h
--- a/Engine/Source/Editor/BlueprintGraph/Classes/K2Node_EnumLiteral.h
+++ b/Engine/Source/Editor/BlueprintGraph/Classes/K2Node_EnumLiteral.h
@@ -25,6 +25,8 @@ class UK2Node_EnumLiteral : public UK2Node, public INodeDependingOnEnumInterface
 	// Begin UK2Node interface
 	virtual bool IsNodePure() const override { return true; }
 	virtual FNodeHandlingFunctor* CreateNodeHandler(FKismetCompilerContext& CompilerContext) const override;
+	virtual void GetMenuActions(TArray<UBlueprintNodeSpawner*>& ActionListOut) const override;
+	virtual FText GetMenuCategory() const override;
 	// End UK2Node interface
 
 	// INodeDependingOnEnumInterface
diff --git a/Engine/Source/Editor/BlueprintGraph/Classes/K2Node_ForEachElementInEnum.h b/Engine/Source/Editor/BlueprintGraph/Classes/K2Node_ForEachElementInEnum.h
--- a/Engine/Source/Editor/BlueprintGraph/Classes/K2Node_ForEachElementInEnum.h
+++ b/Engine/Source/Editor/BlueprintGraph/Classes/K2Node_ForEachElementInEnum.h
@@ -26,6 +26,8 @@ class UK2Node_ForEachElementInEnum : public UK2Node
 	virtual bool IsNodePure() const override { return false; }
 	virtual void ExpandNode(class FKismetCompilerContext& CompilerContext, UEdGraph* SourceGraph) override;
 	virtual void ValidateNodeDuringCompilation(class FCompilerResultsLog& MessageLog) const override;
+	virtual void GetMenuActions(TArray<UBlueprintNodeSpawner*>& ActionListOut) const override;
+	virtual FText GetMenuCategory() const override;
 	// End UK2Node interface
 };
 
diff --git a/Engine/Source/Editor/BlueprintGraph/Classes/K2Node_GetEnumeratorName.h b/Engine/Source/Editor/BlueprintGraph/Classes/K2Node_GetEnumeratorName.h
--- a/Engine/Source/Editor/BlueprintGraph/Classes/K2Node_GetEnumeratorName.h
+++ b/Engine/Source/Editor/BlueprintGraph/Classes/K2Node_GetEnumeratorName.h
@@ -28,6 +28,8 @@ class UK2Node_GetEnumeratorName : public UK2Node
 	virtual void EarlyValidation(class FCompilerResultsLog& MessageLog) const override;
 	virtual void PostReconstructNode() override;
 	virtual void NotifyPinConnectionListChanged(UEdGraphPin* Pin) override;
+	virtual void GetMenuActions(TArray<UBlueprintNodeSpawner*>& ActionListOut) const override;
+	virtual FText GetMenuCategory() const override;
 	// End UK2Node interface
 
 	void UpdatePinType();
diff --git a/Engine/Source/Editor/BlueprintGraph/Classes/K2Node_GetEnumeratorNameAsString.h b/Engine/Source/Editor/BlueprintGraph/Classes/K2Node_GetEnumeratorNameAsString.h
--- a/Engine/Source/Editor/BlueprintGraph/Classes/K2Node_GetEnumeratorNameAsString.h
+++ b/Engine/Source/Editor/BlueprintGraph/Classes/K2Node_GetEnumeratorNameAsString.h
@@ -16,6 +16,11 @@ class UK2Node_GetEnumeratorNameAsString : public UK2Node_GetEnumeratorName
 	virtual FName GetPaletteIcon(FLinearColor& OutColor) const override{ return TEXT("GraphEditor.Enum_16x"); }
 	// End UEdGraphNode interface
 
+	// UK2Node interface
+	virtual void GetMenuActions(TArray<UBlueprintNodeSpawner*>& ActionListOut) const override;
+	virtual FText GetMenuCategory() const override;
+	// End UK2Node interface
+
 	virtual FName GetFunctionName() const override;
 };
 
diff --git a/Engine/Source/Editor/BlueprintGraph/Classes/K2Node_GetNumEnumEntries.h b/Engine/Source/Editor/BlueprintGraph/Classes/K2Node_GetNumEnumEntries.h
--- a/Engine/Source/Editor/BlueprintGraph/Classes/K2Node_GetNumEnumEntries.h
+++ b/Engine/Source/Editor/BlueprintGraph/Classes/K2Node_GetNumEnumEntries.h
@@ -22,6 +22,8 @@ class UK2Node_GetNumEnumEntries : public UK2Node, public INodeDependingOnEnumInt
 	// Begin UK2Node interface
 	virtual bool IsNodePure() const override { return true; }
 	virtual void ExpandNode(class FKismetCompilerContext& CompilerContext, UEdGraph* SourceGraph) override;
+	virtual void GetMenuActions(TArray<UBlueprintNodeSpawner*>& ActionListOut) const override;
+	virtual FText GetMenuCategory() const override;
 	// End UK2Node interface
 
 	// INodeDependingOnEnumInterface
diff --git a/Engine/Source/Editor/BlueprintGraph/Private/K2Node_CastByteToEnum.cpp b/Engine/Source/Editor/BlueprintGraph/Private/K2Node_CastByteToEnum.cpp
--- a/Engine/Source/Editor/BlueprintGraph/Private/K2Node_CastByteToEnum.cpp
+++ b/Engine/Source/Editor/BlueprintGraph/Private/K2Node_CastByteToEnum.cpp
@@ -5,6 +5,8 @@
 #include "KismetCompiler.h"
 #include "Kismet/KismetNodeHelperLibrary.h"
 #include "K2Node_CastByteToEnum.h"
+#include "BlueprintNodeSpawner.h"
+#include "EditorCategoryUtils.h"
 
 const FString UK2Node_CastByteToEnum::ByteInputPinName = TEXT("Byte");
 
@@ -134,4 +136,45 @@ FNodeHandlingFunctor* UK2Node_CastByteToEnum::CreateNodeHandler(FKismetCompilerC
 		return new FKCHandler_CastByteToEnum(CompilerContext);
 	}
 	return NULL;
-}
\ No newline at end of file
+}
+
+void UK2Node_CastByteToEnum::GetMenuActions(TArray<UBlueprintNodeSpawner*>& ActionListOut) const
+{
+	for (TObjectIterator<UEnum> EnumIt; EnumIt; ++EnumIt)
+	{
+		UEnum const* Enum = (*EnumIt);
+		// we only want to add global "standalone" enums here; those belonging to a 
+		// certain class should instead be associated with that class (so when 
+		// the class is modified we can easily handle any enums that were changed).
+		//
+		// @TODO: don't love how this code is essentially duplicated in BlueprintActionDatabase.cpp, for class enums
+		bool bIsStandaloneEnum = Enum->GetOuter()->IsA(UPackage::StaticClass());
+
+		if (!bIsStandaloneEnum || !UEdGraphSchema_K2::IsAllowableBlueprintVariableType(Enum))
+		{
+			continue;
+		}
+
+		auto CustomizeEnumNodeLambda = [](UEdGraphNode* NewNode, bool bIsTemplateNode, TWeakObjectPtr<UEnum> EnumPtr)
+		{
+			UK2Node_CastByteToEnum* EnumNode = CastChecked<UK2Node_CastByteToEnum>(NewNode);
+			if (EnumPtr.IsValid())
+			{
+				EnumNode->Enum = EnumPtr.Get();
+			}
+			EnumNode->bSafe = true;
+		};
+
+		UBlueprintNodeSpawner* NodeSpawner = UBlueprintNodeSpawner::Create(GetClass());
+		check(NodeSpawner != nullptr);
+		ActionListOut.Add(NodeSpawner);
+
+		TWeakObjectPtr<UEnum> EnumPtr = Enum;
+		NodeSpawner->CustomizeNodeDelegate = UBlueprintNodeSpawner::FCustomizeNodeDelegate::CreateStatic(CustomizeEnumNodeLambda, EnumPtr);
+	}
+}
+
+FText UK2Node_CastByteToEnum::GetMenuCategory() const
+{
+	return FEditorCategoryUtils::GetCommonCategory(FCommonEditorCategory::Enum);
+}
diff --git a/Engine/Source/Editor/BlueprintGraph/Private/K2Node_EnumEquality.cpp b/Engine/Source/Editor/BlueprintGraph/Private/K2Node_EnumEquality.cpp
--- a/Engine/Source/Editor/BlueprintGraph/Private/K2Node_EnumEquality.cpp
+++ b/Engine/Source/Editor/BlueprintGraph/Private/K2Node_EnumEquality.cpp
@@ -6,6 +6,8 @@
 #include "Kismet/KismetMathLibrary.h"
 #include "K2ActionMenuBuilder.h" // for FK2ActionMenuBuilder::AddNewNodeAction()
 #include "K2Node_EnumEquality.h"
+#include "BlueprintNodeSpawner.h"
+#include "EditorCategoryUtils.h"
 
 UK2Node_EnumEquality::UK2Node_EnumEquality(const class FPostConstructInitializeProperties& PCIP)
 	: Super(PCIP)
@@ -235,4 +237,17 @@ void UK2Node_EnumEquality::ExpandNode(class FKismetCompilerContext& CompilerCont
 		// Break all links to the Select node so it goes away for at scheduling time
 		BreakAllNodeLinks();
 	}
-}
\ No newline at end of file
+}
+
+void UK2Node_EnumEquality::GetMenuActions(TArray<UBlueprintNodeSpawner*>& ActionListOut) const
+{
+	UBlueprintNodeSpawner* NodeSpawner = UBlueprintNodeSpawner::Create(GetClass());
+	check(NodeSpawner != nullptr);
+	
+	ActionListOut.Add(NodeSpawner);
+}
+
+FText UK2Node_EnumEquality::GetMenuCategory() const
+{
+	return FEditorCategoryUtils::GetCommonCategory(FCommonEditorCategory::Enum);
+}
diff --git a/Engine/Source/Editor/BlueprintGraph/Private/K2Node_EnumLiteral.cpp b/Engine/Source/Editor/BlueprintGraph/Private/K2Node_EnumLiteral.cpp
--- a/Engine/Source/Editor/BlueprintGraph/Private/K2Node_EnumLiteral.cpp
+++ b/Engine/Source/Editor/BlueprintGraph/Private/K2Node_EnumLiteral.cpp
@@ -4,6 +4,8 @@
 #include "BlueprintGraphPrivatePCH.h"
 #include "KismetCompilerMisc.h"
 #include "K2Node_EnumLiteral.h"
+#include "BlueprintNodeSpawner.h"
+#include "EditorCategoryUtils.h"
 
 const FString& UK2Node_EnumLiteral::GetEnumInputPinName()
 {
@@ -86,3 +88,44 @@ FNodeHandlingFunctor* UK2Node_EnumLiteral::CreateNodeHandler(FKismetCompilerCont
 {
 	return new FKCHandler_EnumLiteral(CompilerContext);
 }
+
+void UK2Node_EnumLiteral::GetMenuActions(TArray<UBlueprintNodeSpawner*>& ActionListOut) const
+{
+	for (TObjectIterator<UEnum> EnumIt; EnumIt; ++EnumIt)
+	{
+		UEnum const* Enum = (*EnumIt);
+		// we only want to add global "standalone" enums here; those belonging to a 
+		// certain class should instead be associated with that class (so when 
+		// the class is modified we can easily handle any enums that were changed).
+		//
+		// @TODO: don't love how this code is essentially duplicated in BlueprintActionDatabase.cpp, for class enums
+		bool bIsStandaloneEnum = Enum->GetOuter()->IsA(UPackage::StaticClass());
+
+		if (!bIsStandaloneEnum || !UEdGraphSchema_K2::IsAllowableBlueprintVariableType(Enum))
+		{
+			continue;
+		}
+
+		auto CustomizeEnumNodeLambda = [](UEdGraphNode* NewNode, bool bIsTemplateNode, TWeakObjectPtr<UEnum> EnumPtr)
+		{
+			UK2Node_EnumLiteral* EnumNode = CastChecked<UK2Node_EnumLiteral>(NewNode);
+			if (EnumPtr.IsValid())
+			{
+				EnumNode->Enum = EnumPtr.Get();
+			}
+		};
+
+		UBlueprintNodeSpawner* NodeSpawner = UBlueprintNodeSpawner::Create(GetClass());
+		check(NodeSpawner != nullptr);
+		ActionListOut.Add(NodeSpawner);
+
+		TWeakObjectPtr<UEnum> EnumPtr = Enum;
+		NodeSpawner->CustomizeNodeDelegate = UBlueprintNodeSpawner::FCustomizeNodeDelegate::CreateStatic(CustomizeEnumNodeLambda, EnumPtr);
+	}
+}
+
+FText UK2Node_EnumLiteral::GetMenuCategory() const
+{
+	return FEditorCategoryUtils::GetCommonCategory(FCommonEditorCategory::Enum);
+}
+
diff --git a/Engine/Source/Editor/BlueprintGraph/Private/K2Node_ForEachElementInEnum.cpp b/Engine/Source/Editor/BlueprintGraph/Private/K2Node_ForEachElementInEnum.cpp
--- a/Engine/Source/Editor/BlueprintGraph/Private/K2Node_ForEachElementInEnum.cpp
+++ b/Engine/Source/Editor/BlueprintGraph/Private/K2Node_ForEachElementInEnum.cpp
@@ -8,6 +8,8 @@
 #include "K2Node_CastByteToEnum.h"
 #include "K2Node_ForEachElementInEnum.h"
 #include "K2Node_GetNumEnumEntries.h"
+#include "BlueprintNodeSpawner.h"
+#include "EditorCategoryUtils.h"
 
 #define LOCTEXT_NAMESPACE "K2Node"
 
@@ -191,4 +193,44 @@ void UK2Node_ForEachElementInEnum::ExpandNode(class FKismetCompilerContext& Comp
 	}
 }
 
+void UK2Node_ForEachElementInEnum::GetMenuActions(TArray<UBlueprintNodeSpawner*>& ActionListOut) const
+{
+	for (TObjectIterator<UEnum> EnumIt; EnumIt; ++EnumIt)
+	{
+		UEnum const* Enum = (*EnumIt);
+		// we only want to add global "standalone" enums here; those belonging to a 
+		// certain class should instead be associated with that class (so when 
+		// the class is modified we can easily handle any enums that were changed).
+		//
+		// @TODO: don't love how this code is essentially duplicated in BlueprintActionDatabase.cpp, for class enums
+		bool bIsStandaloneEnum = Enum->GetOuter()->IsA(UPackage::StaticClass());
+
+		if (!bIsStandaloneEnum || !UEdGraphSchema_K2::IsAllowableBlueprintVariableType(Enum))
+		{
+			continue;
+		}
+
+		auto CustomizeEnumNodeLambda = [](UEdGraphNode* NewNode, bool bIsTemplateNode, TWeakObjectPtr<UEnum> EnumPtr)
+		{
+			UK2Node_ForEachElementInEnum* EnumNode = CastChecked<UK2Node_ForEachElementInEnum>(NewNode);
+			if (EnumPtr.IsValid())
+			{
+				EnumNode->Enum = EnumPtr.Get();
+			}
+		};
+
+		UBlueprintNodeSpawner* NodeSpawner = UBlueprintNodeSpawner::Create(GetClass());
+		check(NodeSpawner != nullptr);
+		ActionListOut.Add(NodeSpawner);
+
+		TWeakObjectPtr<UEnum> EnumPtr = Enum;
+		NodeSpawner->CustomizeNodeDelegate = UBlueprintNodeSpawner::FCustomizeNodeDelegate::CreateStatic(CustomizeEnumNodeLambda, EnumPtr);
+	}
+}
+
+FText UK2Node_ForEachElementInEnum::GetMenuCategory() const
+{
+	return FEditorCategoryUtils::GetCommonCategory(FCommonEditorCategory::Enum);
+}
+
 #undef LOCTEXT_NAMESPACE
\ No newline at end of file
diff --git a/Engine/Source/Editor/BlueprintGraph/Private/K2Node_GetEnumeratorName.cpp b/Engine/Source/Editor/BlueprintGraph/Private/K2Node_GetEnumeratorName.cpp
--- a/Engine/Source/Editor/BlueprintGraph/Private/K2Node_GetEnumeratorName.cpp
+++ b/Engine/Source/Editor/BlueprintGraph/Private/K2Node_GetEnumeratorName.cpp
@@ -6,6 +6,8 @@
 #include "Kismet/KismetNodeHelperLibrary.h"
 #include "Kismet/KismetSystemLibrary.h"
 #include "K2Node_GetEnumeratorName.h"
+#include "BlueprintNodeSpawner.h"
+#include "EditorCategoryUtils.h"
 
 FString UK2Node_GetEnumeratorName::EnumeratorPinName = TEXT("Enumerator");
 
@@ -180,4 +182,17 @@ void UK2Node_GetEnumeratorName::ExpandNode(class FKismetCompilerContext& Compile
 
 		BreakAllNodeLinks();
 	}
-}
\ No newline at end of file
+}
+
+void UK2Node_GetEnumeratorName::GetMenuActions(TArray<UBlueprintNodeSpawner*>& ActionListOut) const
+{
+	UBlueprintNodeSpawner* NodeSpawner = UBlueprintNodeSpawner::Create(GetClass());
+	check(NodeSpawner != nullptr);
+
+	ActionListOut.Add(NodeSpawner);
+}
+
+FText UK2Node_GetEnumeratorName::GetMenuCategory() const
+{
+	return FEditorCategoryUtils::GetCommonCategory(FCommonEditorCategory::Name);
+}
diff --git a/Engine/Source/Editor/BlueprintGraph/Private/K2Node_GetEnumeratorNameAsString.cpp b/Engine/Source/Editor/BlueprintGraph/Private/K2Node_GetEnumeratorNameAsString.cpp
--- a/Engine/Source/Editor/BlueprintGraph/Private/K2Node_GetEnumeratorNameAsString.cpp
+++ b/Engine/Source/Editor/BlueprintGraph/Private/K2Node_GetEnumeratorNameAsString.cpp
@@ -5,6 +5,8 @@
 #include "KismetCompiler.h"
 #include "Kismet/KismetNodeHelperLibrary.h"
 #include "K2Node_GetEnumeratorNameAsString.h"
+#include "BlueprintNodeSpawner.h"
+#include "EditorCategoryUtils.h"
 
 UK2Node_GetEnumeratorNameAsString::UK2Node_GetEnumeratorNameAsString(const class FPostConstructInitializeProperties& PCIP)
 	: Super(PCIP)
@@ -29,8 +31,21 @@ FText UK2Node_GetEnumeratorNameAsString::GetNodeTitle(ENodeTitleType::Type Title
 	return NSLOCTEXT("K2Node", "GetEnumeratorNameAsString_Title", "Enum to String");
 }
 
+void UK2Node_GetEnumeratorNameAsString::GetMenuActions(TArray<UBlueprintNodeSpawner*>& ActionListOut) const
+{
+	UBlueprintNodeSpawner* NodeSpawner = UBlueprintNodeSpawner::Create(GetClass());
+	check(NodeSpawner != nullptr);
+
+	ActionListOut.Add(NodeSpawner);
+}
+
 FName UK2Node_GetEnumeratorNameAsString::GetFunctionName() const
 {
 	const FName FunctionName = GET_FUNCTION_NAME_CHECKED(UKismetNodeHelperLibrary, GetEnumeratorUserFriendlyName);
 	return FunctionName;
-}
\ No newline at end of file
+}
+
+FText UK2Node_GetEnumeratorNameAsString::GetMenuCategory() const
+{
+	return FEditorCategoryUtils::GetCommonCategory(FCommonEditorCategory::String);
+}
diff --git a/Engine/Source/Editor/BlueprintGraph/Private/K2Node_GetNumEnumEntries.cpp b/Engine/Source/Editor/BlueprintGraph/Private/K2Node_GetNumEnumEntries.cpp
--- a/Engine/Source/Editor/BlueprintGraph/Private/K2Node_GetNumEnumEntries.cpp
+++ b/Engine/Source/Editor/BlueprintGraph/Private/K2Node_GetNumEnumEntries.cpp
@@ -5,6 +5,8 @@
 #include "KismetCompiler.h"
 #include "../../../Runtime/Engine/Classes/Kismet/KismetSystemLibrary.h"
 #include "K2Node_GetNumEnumEntries.h"
+#include "BlueprintNodeSpawner.h"
+#include "EditorCategoryUtils.h"
 
 UK2Node_GetNumEnumEntries::UK2Node_GetNumEnumEntries(const class FPostConstructInitializeProperties& PCIP)
 	: Super(PCIP)
@@ -72,4 +74,45 @@ void UK2Node_GetNumEnumEntries::ExpandNode(class FKismetCompilerContext& Compile
 
 		BreakAllNodeLinks();
 	}
-}
\ No newline at end of file
+}
+
+void UK2Node_GetNumEnumEntries::GetMenuActions(TArray<UBlueprintNodeSpawner*>& ActionListOut) const
+{
+	for (TObjectIterator<UEnum> EnumIt; EnumIt; ++EnumIt)
+	{
+		UEnum const* Enum = (*EnumIt);
+		// we only want to add global "standalone" enums here; those belonging to a 
+		// certain class should instead be associated with that class (so when 
+		// the class is modified we can easily handle any enums that were changed).
+		//
+		// @TODO: don't love how this code is essentially duplicated in BlueprintActionDatabase.cpp, for class enums
+		bool bIsStandaloneEnum = Enum->GetOuter()->IsA(UPackage::StaticClass());
+
+		if (!bIsStandaloneEnum || !UEdGraphSchema_K2::IsAllowableBlueprintVariableType(Enum))
+		{
+			continue;
+		}
+
+		auto CustomizeEnumNodeLambda = [](UEdGraphNode* NewNode, bool bIsTemplateNode, TWeakObjectPtr<UEnum> EnumPtr)
+		{
+			UK2Node_GetNumEnumEntries* EnumNode = CastChecked<UK2Node_GetNumEnumEntries>(NewNode);
+			if (EnumPtr.IsValid())
+			{
+				EnumNode->Enum = EnumPtr.Get();
+			}
+		};
+
+		UBlueprintNodeSpawner* NodeSpawner = UBlueprintNodeSpawner::Create(GetClass());
+		check(NodeSpawner != nullptr);
+		ActionListOut.Add(NodeSpawner);
+
+		TWeakObjectPtr<UEnum> EnumPtr = Enum;
+		NodeSpawner->CustomizeNodeDelegate = UBlueprintNodeSpawner::FCustomizeNodeDelegate::CreateStatic(CustomizeEnumNodeLambda, EnumPtr);
+	}
+}
+
+FText UK2Node_GetNumEnumEntries::GetMenuCategory() const
+{
+	return FEditorCategoryUtils::GetCommonCategory(FCommonEditorCategory::Enum);
+}
+
