(WIP) Setting up input nodes for the new blueprint action menu system (nodes add themselves to the action database).

diff --git a/Engine/Source/Editor/BlueprintGraph/Classes/K2Node_GetInputAxisKeyValue.h b/Engine/Source/Editor/BlueprintGraph/Classes/K2Node_GetInputAxisKeyValue.h
--- a/Engine/Source/Editor/BlueprintGraph/Classes/K2Node_GetInputAxisKeyValue.h
+++ b/Engine/Source/Editor/BlueprintGraph/Classes/K2Node_GetInputAxisKeyValue.h
@@ -32,6 +32,8 @@ class UK2Node_GetInputAxisKeyValue : public UK2Node_CallFunction
 	virtual bool ShouldShowNodeProperties() const override { return true; }
 	virtual UClass* GetDynamicBindingClass() const override;
 	virtual void RegisterDynamicBinding(UDynamicBlueprintBinding* BindingObject) const override;
+	virtual void GetMenuActions(TArray<UBlueprintNodeSpawner*>& ActionListOut) const override;
+	virtual FText GetMenuCategory() const override;
 	// End UK2Node interface
 	
 	void Initialize(const FKey AxisKey);
diff --git a/Engine/Source/Editor/BlueprintGraph/Classes/K2Node_GetInputAxisValue.h b/Engine/Source/Editor/BlueprintGraph/Classes/K2Node_GetInputAxisValue.h
--- a/Engine/Source/Editor/BlueprintGraph/Classes/K2Node_GetInputAxisValue.h
+++ b/Engine/Source/Editor/BlueprintGraph/Classes/K2Node_GetInputAxisValue.h
@@ -31,7 +31,9 @@ class UK2Node_GetInputAxisValue : public UK2Node_CallFunction
 	virtual bool ShouldShowNodeProperties() const override { return true; }
 	virtual UClass* GetDynamicBindingClass() const override;
 	virtual void RegisterDynamicBinding(UDynamicBlueprintBinding* BindingObject) const override;
+	virtual void GetMenuActions(TArray<UBlueprintNodeSpawner*>& ActionListOut) const override;
+	virtual FText GetMenuCategory() const override;
 	// End UK2Node interface
 	
 	void Initialize(const FName AxisName);
-};
\ No newline at end of file
+};
diff --git a/Engine/Source/Editor/BlueprintGraph/Classes/K2Node_GetInputVectorAxisValue.h b/Engine/Source/Editor/BlueprintGraph/Classes/K2Node_GetInputVectorAxisValue.h
--- a/Engine/Source/Editor/BlueprintGraph/Classes/K2Node_GetInputVectorAxisValue.h
+++ b/Engine/Source/Editor/BlueprintGraph/Classes/K2Node_GetInputVectorAxisValue.h
@@ -18,7 +18,8 @@ class UK2Node_GetInputVectorAxisValue : public UK2Node_GetInputAxisKeyValue
 	virtual bool ShouldShowNodeProperties() const override { return true; }
 	virtual UClass* GetDynamicBindingClass() const override;
 	virtual void RegisterDynamicBinding(UDynamicBlueprintBinding* BindingObject) const override;
+	virtual void GetMenuActions(TArray<UBlueprintNodeSpawner*>& ActionListOut) const override;
 	// End UK2Node interface
 	
 	void Initialize(const FKey AxisKey);
-};
\ No newline at end of file
+};
diff --git a/Engine/Source/Editor/BlueprintGraph/Classes/K2Node_InputAction.h b/Engine/Source/Editor/BlueprintGraph/Classes/K2Node_InputAction.h
--- a/Engine/Source/Editor/BlueprintGraph/Classes/K2Node_InputAction.h
+++ b/Engine/Source/Editor/BlueprintGraph/Classes/K2Node_InputAction.h
@@ -41,6 +41,8 @@ class UK2Node_InputAction : public UK2Node
 	virtual void ValidateNodeDuringCompilation(class FCompilerResultsLog& MessageLog) const override;
 	virtual bool ShouldShowNodeProperties() const override { return true; }
 	virtual void ExpandNode(FKismetCompilerContext& CompilerContext, UEdGraph* SourceGraph) override;
+	virtual void GetMenuActions(TArray<UBlueprintNodeSpawner*>& ActionListOut) const override;
+	virtual FText GetMenuCategory() const override;
 	// End UK2Node interface
 
 	/** Get the 'pressed' input pin */
@@ -51,4 +53,4 @@ class UK2Node_InputAction : public UK2Node
 
 private:
 	void CreateInputActionEvent(FKismetCompilerContext& CompilerContext, UEdGraph* SourceGraph, UEdGraphPin* InputActionPin, const EInputEvent InputKeyEvent);
-};
\ No newline at end of file
+};
diff --git a/Engine/Source/Editor/BlueprintGraph/Classes/K2Node_InputAxisEvent.h b/Engine/Source/Editor/BlueprintGraph/Classes/K2Node_InputAxisEvent.h
--- a/Engine/Source/Editor/BlueprintGraph/Classes/K2Node_InputAxisEvent.h
+++ b/Engine/Source/Editor/BlueprintGraph/Classes/K2Node_InputAxisEvent.h
@@ -39,7 +39,9 @@ class UK2Node_InputAxisEvent : public UK2Node_Event
 	virtual bool ShouldShowNodeProperties() const override { return true; }
 	virtual UClass* GetDynamicBindingClass() const override;
 	virtual void RegisterDynamicBinding(UDynamicBlueprintBinding* BindingObject) const override;
+	virtual void GetMenuActions(TArray<UBlueprintNodeSpawner*>& ActionListOut) const override;
+	virtual FText GetMenuCategory() const override;
 	// End UK2Node interface
 
 	void Initialize(const FName AxisName);
-};
\ No newline at end of file
+};
diff --git a/Engine/Source/Editor/BlueprintGraph/Classes/K2Node_InputAxisKeyEvent.h b/Engine/Source/Editor/BlueprintGraph/Classes/K2Node_InputAxisKeyEvent.h
--- a/Engine/Source/Editor/BlueprintGraph/Classes/K2Node_InputAxisKeyEvent.h
+++ b/Engine/Source/Editor/BlueprintGraph/Classes/K2Node_InputAxisKeyEvent.h
@@ -36,7 +36,9 @@ class UK2Node_InputAxisKeyEvent : public UK2Node_Event
 	virtual bool ShouldShowNodeProperties() const override{ return true; }
 	virtual UClass* GetDynamicBindingClass() const override;
 	virtual void RegisterDynamicBinding(UDynamicBlueprintBinding* BindingObject) const override;
+	virtual void GetMenuActions(TArray<UBlueprintNodeSpawner*>& ActionListOut) const override;
+	virtual FText GetMenuCategory() const override;
 	// End UK2Node interface
 
 	void Initialize(const FKey AxisKey);
-};
\ No newline at end of file
+};
diff --git a/Engine/Source/Editor/BlueprintGraph/Classes/K2Node_InputKey.h b/Engine/Source/Editor/BlueprintGraph/Classes/K2Node_InputKey.h
--- a/Engine/Source/Editor/BlueprintGraph/Classes/K2Node_InputKey.h
+++ b/Engine/Source/Editor/BlueprintGraph/Classes/K2Node_InputKey.h
@@ -46,6 +46,8 @@ class UK2Node_InputKey : public UK2Node
 	virtual bool ShouldShowNodeProperties() const override { return true; }
 	virtual void ValidateNodeDuringCompilation(class FCompilerResultsLog& MessageLog) const override;
 	virtual void ExpandNode(FKismetCompilerContext& CompilerContext, UEdGraph* SourceGraph) override;
+	virtual void GetMenuActions(TArray<UBlueprintNodeSpawner*>& ActionListOut) const override;
+	virtual FText GetMenuCategory() const override;
 	// End UK2Node interface
 
 	// Begin UEdGraphNode interface.
@@ -71,4 +73,4 @@ class UK2Node_InputKey : public UK2Node
 private:
 	void CreateInputKeyEvent(FKismetCompilerContext& CompilerContext, UEdGraph* SourceGraph, UEdGraphPin* InputKeyPin, const EInputEvent KeyEvent);
 
-};
\ No newline at end of file
+};
diff --git a/Engine/Source/Editor/BlueprintGraph/Classes/K2Node_InputTouch.h b/Engine/Source/Editor/BlueprintGraph/Classes/K2Node_InputTouch.h
--- a/Engine/Source/Editor/BlueprintGraph/Classes/K2Node_InputTouch.h
+++ b/Engine/Source/Editor/BlueprintGraph/Classes/K2Node_InputTouch.h
@@ -36,6 +36,8 @@ class UK2Node_InputTouch : public UK2Node
 	virtual FLinearColor GetNodeTitleColor() const override;
 	virtual FText GetNodeTitle(ENodeTitleType::Type TitleType) const override;
 	virtual FName GetPaletteIcon(FLinearColor& OutColor) const override{ return TEXT("GraphEditor.TouchEvent_16x"); }
+	virtual void GetMenuActions(TArray<UBlueprintNodeSpawner*>& ActionListOut) const override;
+	virtual FText GetMenuCategory() const override;
 	// End UEdGraphNode interface.
 
 	BLUEPRINTGRAPH_API static UEnum* GetTouchIndexEnum();
@@ -49,4 +51,4 @@ class UK2Node_InputTouch : public UK2Node
 	BLUEPRINTGRAPH_API UEdGraphPin* GetLocationPin() const;
 	BLUEPRINTGRAPH_API UEdGraphPin* GetFingerIndexPin() const;
 	
-};
\ No newline at end of file
+};
diff --git a/Engine/Source/Editor/BlueprintGraph/Classes/K2Node_InputVectorAxisEvent.h b/Engine/Source/Editor/BlueprintGraph/Classes/K2Node_InputVectorAxisEvent.h
--- a/Engine/Source/Editor/BlueprintGraph/Classes/K2Node_InputVectorAxisEvent.h
+++ b/Engine/Source/Editor/BlueprintGraph/Classes/K2Node_InputVectorAxisEvent.h
@@ -13,5 +13,6 @@ class UK2Node_InputVectorAxisEvent : public UK2Node_InputAxisKeyEvent
 	virtual void ValidateNodeDuringCompilation(class FCompilerResultsLog& MessageLog) const override;
 	virtual UClass* GetDynamicBindingClass() const override;
 	virtual void RegisterDynamicBinding(UDynamicBlueprintBinding* BindingObject) const override;
+	virtual void GetMenuActions(TArray<UBlueprintNodeSpawner*>& ActionListOut) const override;
 	// End UK2Node interface
 };
\ No newline at end of file
diff --git a/Engine/Source/Editor/BlueprintGraph/Private/K2Node_GetInputAxisKeyValue.cpp b/Engine/Source/Editor/BlueprintGraph/Private/K2Node_GetInputAxisKeyValue.cpp
--- a/Engine/Source/Editor/BlueprintGraph/Private/K2Node_GetInputAxisKeyValue.cpp
+++ b/Engine/Source/Editor/BlueprintGraph/Private/K2Node_GetInputAxisKeyValue.cpp
@@ -4,6 +4,8 @@
 #include "K2Node_GetInputAxisKeyValue.h"
 #include "CompilerResultsLog.h"
 
+#define LOCTEXT_NAMESPACE "K2Node_GetInputAxisKeyValue"
+
 UK2Node_GetInputAxisKeyValue::UK2Node_GetInputAxisKeyValue(const class FPostConstructInitializeProperties& PCIP)
 	: Super(PCIP)
 {
@@ -28,7 +30,13 @@ FText UK2Node_GetInputAxisKeyValue::GetNodeTitle(ENodeTitleType::Type TitleType)
 {
 	FFormatNamedArguments Args;
 	Args.Add(TEXT("AxisKey"), InputAxisKey.GetDisplayName());
-	return FText::Format(NSLOCTEXT("K2Node", "GetInputAxisKey_Name", "Get {AxisKey}"), Args);
+	
+	FText TitleFormat = NSLOCTEXT("K2Node", "GetInputAxisKey_Name", "Get {AxisKey}");
+	if (TitleType == ENodeTitleType::ListView)
+	{
+		TitleFormat = NSLOCTEXT("K2Node", "GetInputAxisKey_ListTitle", "{AxisKey}");
+	}
+	return FText::Format(TitleFormat, Args);
 }
 
 FString UK2Node_GetInputAxisKeyValue::GetTooltip() const
@@ -87,4 +95,52 @@ FName UK2Node_GetInputAxisKeyValue::GetPaletteIcon(FLinearColor& OutColor) const
 	{
 		return TEXT("GraphEditor.KeyEvent_16x");
 	}
-}
\ No newline at end of file
+}
+
+
+void UK2Node_GetInputAxisKeyValue::GetMenuActions(TArray<UBlueprintNodeSpawner*>& ActionListOut) const
+{
+	TArray<FKey> AllKeys;
+	EKeys::GetAllKeys(AllKeys);
+	
+	for (FKey const Key : AllKeys)
+	{
+		if (!Key.IsBindableInBlueprints() || !Key.IsFloatAxis())
+		{
+			continue;
+		}
+		
+		UBlueprintNodeSpawner* NodeSpawner = UBlueprintNodeSpawner::Create(GetClass());
+		check(NodeSpawner != nullptr);
+		
+		auto CustomizeInputNodeLambda = [](UEdGraphNode* NewNode, bool bIsTemplateNode, FKey Key)
+		{
+			UK2Node_GetInputAxisKeyValue* InputNode = CastChecked<UK2Node_GetInputAxisKeyValue>(NewNode);
+			InputNode->Initialize(Key);
+		};
+		
+		NodeSpawner->CustomizeNodeDelegate = UBlueprintNodeSpawner::FCustomizeNodeDelegate::CreateStatic(CustomizeInputNodeLambda, Key);
+		ActionListOut.Add(NodeSpawner);
+	}
+}
+
+FText UK2Node_GetInputAxisKeyValue::GetMenuCategory() const
+{
+	FText SubCategory;
+	if (InputAxisKey.IsGamepadKey())
+	{
+		SubCategory = LOCTEXT("GamepadSubCategory", "Gamepad Values");
+	}
+	else if (InputAxisKey.IsMouseButton())
+	{
+		SubCategory = LOCTEXT("MouseSubCategory", "Mouse Values");
+	}
+	else
+	{
+		SubCategory = LOCTEXT("KeySubCategory", "Key Values");
+	}
+	
+	return FEditorCategoryUtils::BuildCategoryString(FCommonEditorCategory::Input, SubCategory);
+}
+
+#undef LOCTEXT_NAMESPACE
\ No newline at end of file
diff --git a/Engine/Source/Editor/BlueprintGraph/Private/K2Node_GetInputAxisValue.cpp b/Engine/Source/Editor/BlueprintGraph/Private/K2Node_GetInputAxisValue.cpp
--- a/Engine/Source/Editor/BlueprintGraph/Private/K2Node_GetInputAxisValue.cpp
+++ b/Engine/Source/Editor/BlueprintGraph/Private/K2Node_GetInputAxisValue.cpp
@@ -3,6 +3,10 @@
 #include "BlueprintGraphPrivatePCH.h"
 #include "K2Node_GetInputAxisValue.h"
 #include "CompilerResultsLog.h"
+#include "BlueprintNodeSpawner.h"
+#include "EditorCategoryUtils.h"
+
+#define LOCTEXT_NAMESPACE "K2Node_GetInputAxisValue"
 
 UK2Node_GetInputAxisValue::UK2Node_GetInputAxisValue(const class FPostConstructInitializeProperties& PCIP)
 	: Super(PCIP)
@@ -28,7 +32,14 @@ FText UK2Node_GetInputAxisValue::GetNodeTitle(ENodeTitleType::Type TitleType) co
 {
 	FFormatNamedArguments Args;
 	Args.Add(TEXT("InputAxisName"), FText::FromName(InputAxisName));
-	return FText::Format(NSLOCTEXT("K2Node", "GetInputAxis_Name", "Get {InputAxisName}"), Args);
+
+	FText LocFormat = NSLOCTEXT("K2Node", "GetInputAxis_Name", "Get {InputAxisName}");
+	if (TitleType == ENodeTitleType::ListView)
+	{
+		LocFormat = NSLOCTEXT("K2Node", "GetInputAxis_ListTitle", "{InputAxisName}");
+	}
+
+	return FText::Format(LocFormat, Args);
 }
 
 FString UK2Node_GetInputAxisValue::GetTooltip() const
@@ -64,3 +75,31 @@ void UK2Node_GetInputAxisValue::RegisterDynamicBinding(UDynamicBlueprintBinding*
 
 	InputAxisBindingObject->InputAxisDelegateBindings.Add(Binding);
 }
+
+void UK2Node_GetInputAxisValue::GetMenuActions(TArray<UBlueprintNodeSpawner*>& ActionListOut) const
+{
+	TArray<FName> AxisNames;
+	GetDefault<UInputSettings>()->GetAxisNames(AxisNames);
+
+	for (FName const InputAxisName : AxisNames)
+	{
+		UBlueprintNodeSpawner* NodeSpawner = UBlueprintNodeSpawner::Create(GetClass());
+		check(NodeSpawner != nullptr);
+
+		auto CustomizeInputNodeLambda = [](UEdGraphNode* NewNode, bool bIsTemplateNode, FName AxisName)
+		{
+			UK2Node_GetInputAxisValue* InputNode = CastChecked<UK2Node_GetInputAxisValue>(NewNode);
+			InputNode->Initialize(AxisName);
+		};
+
+		NodeSpawner->CustomizeNodeDelegate = UBlueprintNodeSpawner::FCustomizeNodeDelegate::CreateStatic(CustomizeInputNodeLambda, InputAxisName);
+		ActionListOut.Add(NodeSpawner);
+	}
+}
+
+FText UK2Node_GetInputAxisValue::GetMenuCategory() const
+{
+	return FEditorCategoryUtils::BuildCategoryString(FCommonEditorCategory::Input, LOCTEXT("ActionMenuCategory", "Axis Values"));
+}
+
+#undef LOCTEXT_NAMESPACE
diff --git a/Engine/Source/Editor/BlueprintGraph/Private/K2Node_GetInputVectorAxisValue.cpp b/Engine/Source/Editor/BlueprintGraph/Private/K2Node_GetInputVectorAxisValue.cpp
--- a/Engine/Source/Editor/BlueprintGraph/Private/K2Node_GetInputVectorAxisValue.cpp
+++ b/Engine/Source/Editor/BlueprintGraph/Private/K2Node_GetInputVectorAxisValue.cpp
@@ -2,8 +2,8 @@
 
 #include "BlueprintGraphPrivatePCH.h"
 #include "K2Node_GetInputVectorAxisValue.h"
-
 #include "CompilerResultsLog.h"
+#include "BlueprintNodeSpawner.h"
 
 UK2Node_GetInputVectorAxisValue::UK2Node_GetInputVectorAxisValue(const class FPostConstructInitializeProperties& PCIP)
 	: Super(PCIP)
@@ -58,4 +58,30 @@ void UK2Node_GetInputVectorAxisValue::RegisterDynamicBinding(UDynamicBlueprintBi
 	Binding.bExecuteWhenPaused = bExecuteWhenPaused;
 
 	InputVectorAxisBindingObject->InputAxisKeyDelegateBindings.Add(Binding);
-}
\ No newline at end of file
+}
+
+void UK2Node_GetInputVectorAxisValue::GetMenuActions(TArray<UBlueprintNodeSpawner*>& ActionListOut) const
+{
+	TArray<FKey> AllKeys;
+	EKeys::GetAllKeys(AllKeys);
+
+	for (FKey const Key : AllKeys)
+	{
+		if (!Key.IsBindableInBlueprints() || !Key.IsVectorAxis())
+		{
+			continue;
+		}
+
+		UBlueprintNodeSpawner* NodeSpawner = UBlueprintNodeSpawner::Create(GetClass());
+		check(NodeSpawner != nullptr);
+
+		auto CustomizeInputNodeLambda = [](UEdGraphNode* NewNode, bool bIsTemplateNode, FKey Key)
+		{
+			UK2Node_GetInputVectorAxisValue* InputNode = CastChecked<UK2Node_GetInputVectorAxisValue>(NewNode);
+			InputNode->Initialize(Key);
+		};
+
+		NodeSpawner->CustomizeNodeDelegate = UBlueprintNodeSpawner::FCustomizeNodeDelegate::CreateStatic(CustomizeInputNodeLambda, Key);
+		ActionListOut.Add(NodeSpawner);
+	}
+}
diff --git a/Engine/Source/Editor/BlueprintGraph/Private/K2Node_InputAction.cpp b/Engine/Source/Editor/BlueprintGraph/Private/K2Node_InputAction.cpp
--- a/Engine/Source/Editor/BlueprintGraph/Private/K2Node_InputAction.cpp
+++ b/Engine/Source/Editor/BlueprintGraph/Private/K2Node_InputAction.cpp
@@ -4,6 +4,10 @@
 #include "K2Node_InputActionEvent.h"
 #include "CompilerResultsLog.h"
 #include "KismetCompiler.h"
+#include "BlueprintNodeSpawner.h"
+#include "EditorCategoryUtils.h"
+
+#define LOCTEXT_NAMESPACE "K2Node_InputAction"
 
 UK2Node_InputAction::UK2Node_InputAction(const class FPostConstructInitializeProperties& PCIP)
 	: Super(PCIP)
@@ -42,7 +46,14 @@ FText UK2Node_InputAction::GetNodeTitle(ENodeTitleType::Type TitleType) const
 {
 	FFormatNamedArguments Args;
 	Args.Add(TEXT("InputActionName"), FText::FromName(InputActionName));
-	return FText::Format(NSLOCTEXT("K2Node", "InputAction_Name", "InputAction {InputActionName}"), Args);
+
+	FText LocFormat = NSLOCTEXT("K2Node", "InputAction_Name", "InputAction {InputActionName}");
+	if (TitleType == ENodeTitleType::ListView)
+	{
+		LocFormat = NSLOCTEXT("K2Node", "InputAction_ListTitle", "{InputActionName}");
+	}
+
+	return FText::Format(LocFormat, Args);
 }
 
 FString UK2Node_InputAction::GetTooltip() const
@@ -108,4 +119,32 @@ void UK2Node_InputAction::ExpandNode(FKismetCompilerContext& CompilerContext, UE
 		CreateInputActionEvent(CompilerContext, SourceGraph, GetPressedPin(), IE_Pressed);
 		CreateInputActionEvent(CompilerContext, SourceGraph, GetReleasedPin(), IE_Released);
 	}
-}
\ No newline at end of file
+}
+
+void UK2Node_InputAction::GetMenuActions(TArray<UBlueprintNodeSpawner*>& ActionListOut) const
+{
+	TArray<FName> ActionNames;
+	GetDefault<UInputSettings>()->GetActionNames(ActionNames);
+
+	for (FName const InputActionName : ActionNames)
+	{
+		UBlueprintNodeSpawner* NodeSpawner = UBlueprintNodeSpawner::Create(GetClass());
+		check(NodeSpawner != nullptr);
+
+		auto CustomizeInputNodeLambda = [](UEdGraphNode* NewNode, bool bIsTemplateNode, FName ActionName)
+		{
+			UK2Node_InputAction* InputNode = CastChecked<UK2Node_InputAction>(NewNode);
+			InputNode->InputActionName = ActionName;
+		};
+
+		NodeSpawner->CustomizeNodeDelegate = UBlueprintNodeSpawner::FCustomizeNodeDelegate::CreateStatic(CustomizeInputNodeLambda, InputActionName);
+		ActionListOut.Add(NodeSpawner);
+	}
+}
+
+FText UK2Node_InputAction::GetMenuCategory() const
+{
+	return FEditorCategoryUtils::BuildCategoryString(FCommonEditorCategory::Input, LOCTEXT("ActionMenuCategory", "Action Events"));
+}
+
+#undef LOCTEXT_NAMESPACE
diff --git a/Engine/Source/Editor/BlueprintGraph/Private/K2Node_InputAxisEvent.cpp b/Engine/Source/Editor/BlueprintGraph/Private/K2Node_InputAxisEvent.cpp
--- a/Engine/Source/Editor/BlueprintGraph/Private/K2Node_InputAxisEvent.cpp
+++ b/Engine/Source/Editor/BlueprintGraph/Private/K2Node_InputAxisEvent.cpp
@@ -1,8 +1,11 @@
 // Copyright 1998-2014 Epic Games, Inc. All Rights Reserved.
 
 #include "BlueprintGraphPrivatePCH.h"
-
 #include "CompilerResultsLog.h"
+#include "BlueprintNodeSpawner.h"
+#include "EditorCategoryUtils.h"
+
+#define LOCTEXT_NAMESPACE "K2Node_InputAxisEvent"
 
 UK2Node_InputAxisEvent::UK2Node_InputAxisEvent(const class FPostConstructInitializeProperties& PCIP)
 	: Super(PCIP)
@@ -36,7 +39,14 @@ FText UK2Node_InputAxisEvent::GetNodeTitle(ENodeTitleType::Type TitleType) const
 {
 	FFormatNamedArguments Args;
 	Args.Add(TEXT("InputAxisName"), FText::FromName(InputAxisName));
-	return FText::Format(NSLOCTEXT("K2Node", "InputAxis_Name", "InputAxis {InputAxisName}"), Args);
+
+	FText LocFormat = NSLOCTEXT("K2Node", "InputAxis_Name", "InputAxis {InputAxisName}");
+	if (TitleType == ENodeTitleType::ListView)
+	{
+		LocFormat = NSLOCTEXT("K2Node", "InputAxis_ListTitle", "{InputAxisName}");
+	}
+
+	return FText::Format(LocFormat, Args);
 }
 
 FString UK2Node_InputAxisEvent::GetTooltip() const
@@ -105,3 +115,31 @@ bool UK2Node_InputAxisEvent::CanPasteHere(const UEdGraph* TargetGraph, const UEd
 
 	return bAllowPaste;
 }
+
+void UK2Node_InputAxisEvent::GetMenuActions(TArray<UBlueprintNodeSpawner*>& ActionListOut) const
+{
+	TArray<FName> AxisNames;
+	GetDefault<UInputSettings>()->GetAxisNames(AxisNames);
+
+	for (FName const InputAxisName : AxisNames)
+	{
+		UBlueprintNodeSpawner* NodeSpawner = UBlueprintNodeSpawner::Create(GetClass());
+		check(NodeSpawner != nullptr);
+
+		auto CustomizeInputNodeLambda = [](UEdGraphNode* NewNode, bool bIsTemplateNode, FName AxisName)
+		{
+			UK2Node_InputAxisEvent* InputNode = CastChecked<UK2Node_InputAxisEvent>(NewNode);
+			InputNode->Initialize(AxisName);
+		};
+
+		NodeSpawner->CustomizeNodeDelegate = UBlueprintNodeSpawner::FCustomizeNodeDelegate::CreateStatic(CustomizeInputNodeLambda, InputAxisName);
+		ActionListOut.Add(NodeSpawner);
+	}
+}
+
+FText UK2Node_InputAxisEvent::GetMenuCategory() const
+{
+	return FEditorCategoryUtils::BuildCategoryString(FCommonEditorCategory::Input, LOCTEXT("ActionMenuCategory", "Axis Events"));
+}
+
+#undef LOCTEXT_NAMESPACE
diff --git a/Engine/Source/Editor/BlueprintGraph/Private/K2Node_InputAxisKeyEvent.cpp b/Engine/Source/Editor/BlueprintGraph/Private/K2Node_InputAxisKeyEvent.cpp
--- a/Engine/Source/Editor/BlueprintGraph/Private/K2Node_InputAxisKeyEvent.cpp
+++ b/Engine/Source/Editor/BlueprintGraph/Private/K2Node_InputAxisKeyEvent.cpp
@@ -3,6 +3,10 @@
 #include "BlueprintGraphPrivatePCH.h"
 #include "K2Node_InputAxisKeyEvent.h"
 #include "CompilerResultsLog.h"
+#include "BlueprintNodeSpawner.h"
+#include "EditorCategoryUtils.h"
+
+#define LOCTEXT_NAMESPACE "UK2Node_InputAxisKeyEvent"
 
 UK2Node_InputAxisKeyEvent::UK2Node_InputAxisKeyEvent(const class FPostConstructInitializeProperties& PCIP)
 : Super(PCIP)
@@ -114,3 +118,51 @@ bool UK2Node_InputAxisKeyEvent::CanPasteHere(const UEdGraph* TargetGraph, const
 
 	return bAllowPaste;
 }
+
+void UK2Node_InputAxisKeyEvent::GetMenuActions(TArray<UBlueprintNodeSpawner*>& ActionListOut) const
+{
+	TArray<FKey> AllKeys;
+	EKeys::GetAllKeys(AllKeys);
+
+	for (FKey const Key : AllKeys)
+	{
+		if (!Key.IsBindableInBlueprints() || !Key.IsFloatAxis())
+		{
+			continue;
+		}
+
+		UBlueprintNodeSpawner* NodeSpawner = UBlueprintNodeSpawner::Create(GetClass());
+		check(NodeSpawner != nullptr);
+
+		auto CustomizeInputNodeLambda = [](UEdGraphNode* NewNode, bool bIsTemplateNode, FKey Key)
+		{
+			UK2Node_InputAxisKeyEvent* InputNode = CastChecked<UK2Node_InputAxisKeyEvent>(NewNode);
+			InputNode->Initialize(Key);
+		};
+
+		NodeSpawner->CustomizeNodeDelegate = UBlueprintNodeSpawner::FCustomizeNodeDelegate::CreateStatic(CustomizeInputNodeLambda, Key);
+		ActionListOut.Add(NodeSpawner);
+	}
+}
+
+FText UK2Node_InputAxisKeyEvent::GetMenuCategory() const
+{
+	FText SubCategory;
+	if (AxisKey.IsGamepadKey())
+	{
+		SubCategory = LOCTEXT("GamepadCategory", "Gamepad Events");
+	}
+	else if (AxisKey.IsMouseButton())
+	{
+		SubCategory = LOCTEXT("MouseCategory", "Mouse Events");
+	}
+	else
+	{
+		SubCategory = LOCTEXT("KeyEventsCategory", "Key Events");
+	}
+
+	return FEditorCategoryUtils::BuildCategoryString(FCommonEditorCategory::Input, SubCategory);
+
+}
+
+#undef LOCTEXT_NAMESPACE
diff --git a/Engine/Source/Editor/BlueprintGraph/Private/K2Node_InputKey.cpp b/Engine/Source/Editor/BlueprintGraph/Private/K2Node_InputKey.cpp
--- a/Engine/Source/Editor/BlueprintGraph/Private/K2Node_InputKey.cpp
+++ b/Engine/Source/Editor/BlueprintGraph/Private/K2Node_InputKey.cpp
@@ -4,6 +4,10 @@
 #include "K2Node_InputKeyEvent.h"
 #include "CompilerResultsLog.h"
 #include "KismetCompiler.h"
+#include "BlueprintNodeSpawner.h"
+#include "EditorCategoryUtils.h"
+
+#define LOCTEXT_NAMESPACE "UK2Node_InputKey"
 
 UK2Node_InputKey::UK2Node_InputKey(const class FPostConstructInitializeProperties& PCIP)
 	: Super(PCIP)
@@ -246,4 +250,53 @@ void UK2Node_InputKey::ExpandNode(FKismetCompilerContext& CompilerContext, UEdGr
 		CreateInputKeyEvent(CompilerContext, SourceGraph, GetPressedPin(), IE_Pressed);
 		CreateInputKeyEvent(CompilerContext, SourceGraph, GetReleasedPin(), IE_Released);
 	}
-}
\ No newline at end of file
+}
+
+void UK2Node_InputKey::GetMenuActions(TArray<UBlueprintNodeSpawner*>& ActionListOut) const
+{
+	TArray<FKey> AllKeys;
+	EKeys::GetAllKeys(AllKeys);
+
+	for (FKey const Key : AllKeys)
+	{
+		// these will be handled by UK2Node_GetInputAxisKeyValue and UK2Node_GetInputVectorAxisValue respectively
+		if (!Key.IsBindableInBlueprints() || Key.IsFloatAxis() || Key.IsVectorAxis())
+		{
+			continue;
+		}
+
+		UBlueprintNodeSpawner* NodeSpawner = UBlueprintNodeSpawner::Create(GetClass());
+		check(NodeSpawner != nullptr);
+
+		auto CustomizeInputNodeLambda = [](UEdGraphNode* NewNode, bool bIsTemplateNode, FKey Key)
+		{
+			UK2Node_InputKey* InputNode = CastChecked<UK2Node_InputKey>(NewNode);
+			InputNode->InputKey = Key;
+		};
+
+		NodeSpawner->CustomizeNodeDelegate = UBlueprintNodeSpawner::FCustomizeNodeDelegate::CreateStatic(CustomizeInputNodeLambda, Key);
+		ActionListOut.Add(NodeSpawner);
+	}
+}
+
+FText UK2Node_InputKey::GetMenuCategory() const
+{
+	FText SubCategory;
+	if (InputKey.IsGamepadKey())
+	{
+		SubCategory = LOCTEXT("GamepadCategory", "Gamepad Events");
+	}
+	else if (InputKey.IsMouseButton())
+	{
+		SubCategory = LOCTEXT("MouseCategory", "Mouse Events");
+	}
+	else
+	{
+		SubCategory = LOCTEXT("KeyEventsCategory", "Key Events");
+	}
+
+	return FEditorCategoryUtils::BuildCategoryString(FCommonEditorCategory::Input, SubCategory);
+
+}
+
+#undef LOCTEXT_NAMESPACE
diff --git a/Engine/Source/Editor/BlueprintGraph/Private/K2Node_InputTouch.cpp b/Engine/Source/Editor/BlueprintGraph/Private/K2Node_InputTouch.cpp
--- a/Engine/Source/Editor/BlueprintGraph/Private/K2Node_InputTouch.cpp
+++ b/Engine/Source/Editor/BlueprintGraph/Private/K2Node_InputTouch.cpp
@@ -4,6 +4,8 @@
 #include "K2Node_InputTouchEvent.h"
 #include "CompilerResultsLog.h"
 #include "KismetCompiler.h"
+#include "BlueprintNodeSpawner.h"
+#include "EditorCategoryUtils.h"
 
 UK2Node_InputTouch::UK2Node_InputTouch(const class FPostConstructInitializeProperties& PCIP)
 	: Super(PCIP)
@@ -57,7 +59,25 @@ FLinearColor UK2Node_InputTouch::GetNodeTitleColor() const
 
 FText UK2Node_InputTouch::GetNodeTitle(ENodeTitleType::Type TitleType) const
 {
-	return NSLOCTEXT("K2Node", "InputTouch_Name", "InputTouch");
+	FText Title = NSLOCTEXT("K2Node", "InputTouch_Name", "InputTouch");
+	if (TitleType == ENodeTitleType::ListView)
+	{
+		Title = NSLOCTEXT("K2Node", "InputTouch_ListTitle", "Touch");
+	}
+	return Title;
+}
+
+void UK2Node_InputTouch::GetMenuActions(TArray<UBlueprintNodeSpawner*>& ActionListOut) const
+{
+	UBlueprintNodeSpawner* NodeSpawner = UBlueprintNodeSpawner::Create(GetClass());
+	check(NodeSpawner != nullptr);
+
+	ActionListOut.Add(NodeSpawner);
+}
+
+FText UK2Node_InputTouch::GetMenuCategory() const
+{
+	return FEditorCategoryUtils::GetCommonCategory(FCommonEditorCategory::Input);
 }
 
 UEdGraphPin* UK2Node_InputTouch::GetPressedPin() const
diff --git a/Engine/Source/Editor/BlueprintGraph/Private/K2Node_InputVectorAxisEvent.cpp b/Engine/Source/Editor/BlueprintGraph/Private/K2Node_InputVectorAxisEvent.cpp
--- a/Engine/Source/Editor/BlueprintGraph/Private/K2Node_InputVectorAxisEvent.cpp
+++ b/Engine/Source/Editor/BlueprintGraph/Private/K2Node_InputVectorAxisEvent.cpp
@@ -3,6 +3,7 @@
 #include "BlueprintGraphPrivatePCH.h"
 #include "K2Node_InputVectorAxisEvent.h"
 #include "CompilerResultsLog.h"
+#include "BlueprintNodeSpawner.h"
 
 UK2Node_InputVectorAxisEvent::UK2Node_InputVectorAxisEvent(const class FPostConstructInitializeProperties& PCIP)
 : Super(PCIP)
@@ -46,4 +47,30 @@ void UK2Node_InputVectorAxisEvent::RegisterDynamicBinding(UDynamicBlueprintBindi
 	Binding.FunctionNameToBind = CustomFunctionName;
 
 	InputVectorAxisBindingObject->InputAxisKeyDelegateBindings.Add(Binding);
-}
\ No newline at end of file
+}
+
+void UK2Node_InputVectorAxisEvent::GetMenuActions(TArray<UBlueprintNodeSpawner*>& ActionListOut) const
+{
+	TArray<FKey> AllKeys;
+	EKeys::GetAllKeys(AllKeys);
+	
+	for (FKey const Key : AllKeys)
+	{
+		if (!Key.IsBindableInBlueprints() || !Key.IsVectorAxis())
+		{
+			continue;
+		}
+		
+		UBlueprintNodeSpawner* NodeSpawner = UBlueprintNodeSpawner::Create(GetClass());
+		check(NodeSpawner != nullptr);
+		
+		auto CustomizeInputNodeLambda = [](UEdGraphNode* NewNode, bool bIsTemplateNode, FKey Key)
+		{
+			UK2Node_InputAxisKeyEvent* InputNode = CastChecked<UK2Node_InputAxisKeyEvent>(NewNode);
+			InputNode->Initialize(Key);
+		};
+		
+		NodeSpawner->CustomizeNodeDelegate = UBlueprintNodeSpawner::FCustomizeNodeDelegate::CreateStatic(CustomizeInputNodeLambda, Key);
+		ActionListOut.Add(NodeSpawner);
+	}
+}
