(WIP) Various flow control node types now add themselves to the blueprint action database (for the new action menu system).

diff --git a/Engine/Source/Editor/BlueprintGraph/Classes/K2Node_DoOnceMultiInput.h b/Engine/Source/Editor/BlueprintGraph/Classes/K2Node_DoOnceMultiInput.h
--- a/Engine/Source/Editor/BlueprintGraph/Classes/K2Node_DoOnceMultiInput.h
+++ b/Engine/Source/Editor/BlueprintGraph/Classes/K2Node_DoOnceMultiInput.h
@@ -51,5 +51,7 @@ public:
 	// UK2Node interface
 	virtual void GetContextMenuActions(const FGraphNodeContextMenuBuilder& Context) const override;
 	virtual void ExpandNode(class FKismetCompilerContext& CompilerContext, UEdGraph* SourceGraph) override;
+	virtual void GetMenuActions(TArray<UBlueprintNodeSpawner*>& ActionListOut) const override;
+	virtual FText GetMenuCategory() const override;
 	// End of UK2Node interface
 };
\ No newline at end of file
diff --git a/Engine/Source/Editor/BlueprintGraph/Classes/K2Node_ExecutionSequence.h b/Engine/Source/Editor/BlueprintGraph/Classes/K2Node_ExecutionSequence.h
--- a/Engine/Source/Editor/BlueprintGraph/Classes/K2Node_ExecutionSequence.h
+++ b/Engine/Source/Editor/BlueprintGraph/Classes/K2Node_ExecutionSequence.h
@@ -21,6 +21,8 @@ class UK2Node_ExecutionSequence : public UK2Node
 	// Begin UK2Node interface
 	void ReallocatePinsDuringReconstruction(TArray<UEdGraphPin*>& OldPins) override;
 	virtual class FNodeHandlingFunctor* CreateNodeHandler(class FKismetCompilerContext& CompilerContext) const override;
+	virtual void GetMenuActions(TArray<UBlueprintNodeSpawner*>& ActionListOut) const override;
+	virtual FText GetMenuCategory() const override;
 	// End UK2Node interface
 
 	// K2Node_ExecutionSequence interface
diff --git a/Engine/Source/Editor/BlueprintGraph/Classes/K2Node_IfThenElse.h b/Engine/Source/Editor/BlueprintGraph/Classes/K2Node_IfThenElse.h
--- a/Engine/Source/Editor/BlueprintGraph/Classes/K2Node_IfThenElse.h
+++ b/Engine/Source/Editor/BlueprintGraph/Classes/K2Node_IfThenElse.h
@@ -21,6 +21,8 @@ class UK2Node_IfThenElse : public UK2Node
 
 	// Begin K2Node interface.
 	virtual class FNodeHandlingFunctor* CreateNodeHandler(class FKismetCompilerContext& CompilerContext) const override;
+	virtual void GetMenuActions(TArray<UBlueprintNodeSpawner*>& ActionListOut) const override;
+	virtual FText GetMenuCategory() const override;
 	// End K2Node interface.
 
 	/** Get the then output pin */
diff --git a/Engine/Source/Editor/BlueprintGraph/Classes/K2Node_MultiGate.h b/Engine/Source/Editor/BlueprintGraph/Classes/K2Node_MultiGate.h
--- a/Engine/Source/Editor/BlueprintGraph/Classes/K2Node_MultiGate.h
+++ b/Engine/Source/Editor/BlueprintGraph/Classes/K2Node_MultiGate.h
@@ -24,6 +24,7 @@ class UK2Node_MultiGate : public UK2Node_ExecutionSequence
 	virtual void ReallocatePinsDuringReconstruction(TArray<UEdGraphPin*>& OldPins) override;
 	virtual class FNodeHandlingFunctor* CreateNodeHandler(class FKismetCompilerContext& CompilerContext) const override;
 	virtual void ExpandNode(class FKismetCompilerContext& CompilerContext, UEdGraph* SourceGraph) override;
+	virtual void GetMenuActions(TArray<UBlueprintNodeSpawner*>& ActionListOut) const override;
 	// End UK2Node interface
 
 	/** Getting pin access */
diff --git a/Engine/Source/Editor/BlueprintGraph/Classes/K2Node_Select.h b/Engine/Source/Editor/BlueprintGraph/Classes/K2Node_Select.h
--- a/Engine/Source/Editor/BlueprintGraph/Classes/K2Node_Select.h
+++ b/Engine/Source/Editor/BlueprintGraph/Classes/K2Node_Select.h
@@ -52,6 +52,8 @@ class UK2Node_Select : public UK2Node
 	virtual void NotifyPinConnectionListChanged(UEdGraphPin* Pin) override;
 	virtual class FNodeHandlingFunctor* CreateNodeHandler(class FKismetCompilerContext& CompilerContext) const override;
 	virtual bool IsConnectionDisallowed(const UEdGraphPin* MyPin, const UEdGraphPin* OtherPin, FString& OutReason) const override;
+	virtual void GetMenuActions(TArray<UBlueprintNodeSpawner*>& ActionListOut) const override;
+	virtual FText GetMenuCategory() const override;
 	// End UK2Node interface
 
 	/** Get the return value pin */
diff --git a/Engine/Source/Editor/BlueprintGraph/Classes/K2Node_SelectEnum.h b/Engine/Source/Editor/BlueprintGraph/Classes/K2Node_SelectEnum.h
--- a/Engine/Source/Editor/BlueprintGraph/Classes/K2Node_SelectEnum.h
+++ b/Engine/Source/Editor/BlueprintGraph/Classes/K2Node_SelectEnum.h
@@ -18,6 +18,7 @@ class UDEPRECATED_K2Node_SelectEnum : public UK2Node_Select
 
 	// Begin UK2Node interface
 	virtual void NotifyPinConnectionListChanged(UEdGraphPin* Pin) override;
+	virtual void GetMenuActions(TArray<UBlueprintNodeSpawner*>& ActionListOut) const override;
 	// End UK2Node interface
 
 	// Begin UK2Node_Select interface
diff --git a/Engine/Source/Editor/BlueprintGraph/Classes/K2Node_Switch.h b/Engine/Source/Editor/BlueprintGraph/Classes/K2Node_Switch.h
--- a/Engine/Source/Editor/BlueprintGraph/Classes/K2Node_Switch.h
+++ b/Engine/Source/Editor/BlueprintGraph/Classes/K2Node_Switch.h
@@ -35,6 +35,7 @@ class UK2Node_Switch : public UK2Node
 	// UK2Node interface
 	virtual ERedirectType DoPinsMatchForReconstruction(const UEdGraphPin* NewPin, int32 NewPinIndex, const UEdGraphPin* OldPin, int32 OldPinIndex) const override;
 	virtual class FNodeHandlingFunctor* CreateNodeHandler(class FKismetCompilerContext& CompilerContext) const override;
+	virtual FText GetMenuCategory() const override;
 	// End of UK2Node interface
 
 	// UK2Node_Switch interface
diff --git a/Engine/Source/Editor/BlueprintGraph/Classes/K2Node_SwitchEnum.h b/Engine/Source/Editor/BlueprintGraph/Classes/K2Node_SwitchEnum.h
--- a/Engine/Source/Editor/BlueprintGraph/Classes/K2Node_SwitchEnum.h
+++ b/Engine/Source/Editor/BlueprintGraph/Classes/K2Node_SwitchEnum.h
@@ -34,6 +34,10 @@ class UK2Node_SwitchEnum : public UK2Node_Switch, public INodeDependingOnEnumInt
 	virtual bool ShouldShowNodeProperties() const override { return true; }	
 	// End of UEdGraphNode interface
 
+	// UK2Node interface
+	virtual void GetMenuActions(TArray<UBlueprintNodeSpawner*>& ActionListOut) const override;
+	// End of UK2Node interface
+
 	// UK2Node_Switch Interface
 	virtual FString GetUniquePinName() override;
 	virtual const FString& GetPinType(const UEdGraphSchema_K2* Schema) const override { return Schema->PC_Byte; }
diff --git a/Engine/Source/Editor/BlueprintGraph/Classes/K2Node_SwitchInteger.h b/Engine/Source/Editor/BlueprintGraph/Classes/K2Node_SwitchInteger.h
--- a/Engine/Source/Editor/BlueprintGraph/Classes/K2Node_SwitchInteger.h
+++ b/Engine/Source/Editor/BlueprintGraph/Classes/K2Node_SwitchInteger.h
@@ -25,6 +25,10 @@ class UK2Node_SwitchInteger : public UK2Node_Switch
 	virtual bool ShouldShowNodeProperties() const override { return true; }
 	// End of UEdGraphNode interface
 
+	// UK2Node interface
+	virtual void GetMenuActions(TArray<UBlueprintNodeSpawner*>& ActionListOut) const override;
+	// End of UK2Node interface
+
 	// UK2Node_Switch Interface
 	virtual FString GetUniquePinName() override;
 	virtual FString GetPinNameGivenIndex(int32 Index) override;
diff --git a/Engine/Source/Editor/BlueprintGraph/Classes/K2Node_SwitchName.h b/Engine/Source/Editor/BlueprintGraph/Classes/K2Node_SwitchName.h
--- a/Engine/Source/Editor/BlueprintGraph/Classes/K2Node_SwitchName.h
+++ b/Engine/Source/Editor/BlueprintGraph/Classes/K2Node_SwitchName.h
@@ -23,6 +23,10 @@ class UK2Node_SwitchName : public UK2Node_Switch
 	virtual bool ShouldShowNodeProperties() const override { return true; }
 	// End of UEdGraphNode interface
 
+	// UK2Node interface
+	virtual void GetMenuActions(TArray<UBlueprintNodeSpawner*>& ActionListOut) const override;
+	// End of UK2Node interface
+
 	// UK2Node_Switch Interface
 	BLUEPRINTGRAPH_API virtual void AddPinToSwitchNode() override;
 	virtual FString GetUniquePinName() override;
diff --git a/Engine/Source/Editor/BlueprintGraph/Classes/K2Node_SwitchString.h b/Engine/Source/Editor/BlueprintGraph/Classes/K2Node_SwitchString.h
--- a/Engine/Source/Editor/BlueprintGraph/Classes/K2Node_SwitchString.h
+++ b/Engine/Source/Editor/BlueprintGraph/Classes/K2Node_SwitchString.h
@@ -26,6 +26,10 @@ class UK2Node_SwitchString : public UK2Node_Switch
 	virtual bool ShouldShowNodeProperties() const override { return true; }
 	// End of UEdGraphNode interface
 
+	// UK2Node interface
+	virtual void GetMenuActions(TArray<UBlueprintNodeSpawner*>& ActionListOut) const override;
+	// End of UK2Node interface
+
 	// UK2Node_Switch Interface
 	BLUEPRINTGRAPH_API virtual void AddPinToSwitchNode() override;
 	virtual FString GetUniquePinName() override;
diff --git a/Engine/Source/Editor/BlueprintGraph/Private/K2Node_DoOnceMultiInput.cpp b/Engine/Source/Editor/BlueprintGraph/Private/K2Node_DoOnceMultiInput.cpp
--- a/Engine/Source/Editor/BlueprintGraph/Private/K2Node_DoOnceMultiInput.cpp
+++ b/Engine/Source/Editor/BlueprintGraph/Private/K2Node_DoOnceMultiInput.cpp
@@ -6,7 +6,8 @@
 #include "ScopedTransaction.h"
 #include "K2ActionMenuBuilder.h"
 #include "KismetCompiler.h"
-
+#include "BlueprintNodeSpawner.h"
+#include "EditorCategoryUtils.h"
 
 #define LOCTEXT_NAMESPACE "K2Node"
 
@@ -374,4 +375,17 @@ void UK2Node_DoOnceMultiInput::ExpandNode(FKismetCompilerContext& CompilerContex
 	}
 }
 
+void UK2Node_DoOnceMultiInput::GetMenuActions(TArray<UBlueprintNodeSpawner*>& ActionListOut) const
+{
+	UBlueprintNodeSpawner* NodeSpawner = UBlueprintNodeSpawner::Create(GetClass());
+	check(NodeSpawner != nullptr);
+	
+	ActionListOut.Add(NodeSpawner);
+}
+
+FText UK2Node_DoOnceMultiInput::GetMenuCategory() const
+{
+	return FEditorCategoryUtils::GetCommonCategory(FCommonEditorCategory::FlowControl);
+}
+
 #undef LOCTEXT_NAMESPACE
diff --git a/Engine/Source/Editor/BlueprintGraph/Private/K2Node_ExecutionSequence.cpp b/Engine/Source/Editor/BlueprintGraph/Private/K2Node_ExecutionSequence.cpp
--- a/Engine/Source/Editor/BlueprintGraph/Private/K2Node_ExecutionSequence.cpp
+++ b/Engine/Source/Editor/BlueprintGraph/Private/K2Node_ExecutionSequence.cpp
@@ -4,6 +4,8 @@
 #include "BlueprintGraphPrivatePCH.h"
 #include "KismetCompiler.h"
 #include "KismetCompilerMisc.h"
+#include "BlueprintNodeSpawner.h"
+#include "EditorCategoryUtils.h"
 
 #define LOCTEXT_NAMESPACE "K2Node_MultiGate"
 
@@ -261,4 +263,17 @@ FNodeHandlingFunctor* UK2Node_ExecutionSequence::CreateNodeHandler(FKismetCompil
 	return new FKCHandler_ExecutionSequence(CompilerContext);
 }
 
+void UK2Node_ExecutionSequence::GetMenuActions(TArray<UBlueprintNodeSpawner*>& ActionListOut) const
+{
+	UBlueprintNodeSpawner* NodeSpawner = UBlueprintNodeSpawner::Create(GetClass());
+	check(NodeSpawner != nullptr);
+
+	ActionListOut.Add(NodeSpawner);
+}
+
+FText UK2Node_ExecutionSequence::GetMenuCategory() const
+{
+	return FEditorCategoryUtils::GetCommonCategory(FCommonEditorCategory::FlowControl);
+}
+
 #undef LOCTEXT_NAMESPACE
diff --git a/Engine/Source/Editor/BlueprintGraph/Private/K2Node_IfThenElse.cpp b/Engine/Source/Editor/BlueprintGraph/Private/K2Node_IfThenElse.cpp
--- a/Engine/Source/Editor/BlueprintGraph/Private/K2Node_IfThenElse.cpp
+++ b/Engine/Source/Editor/BlueprintGraph/Private/K2Node_IfThenElse.cpp
@@ -3,6 +3,8 @@
 
 #include "BlueprintGraphPrivatePCH.h"
 #include "KismetCompiler.h"
+#include "BlueprintNodeSpawner.h"
+#include "EditorCategoryUtils.h"
 
 #define LOCTEXT_NAMESPACE "K2Node"
 
@@ -143,4 +145,17 @@ FNodeHandlingFunctor* UK2Node_IfThenElse::CreateNodeHandler(FKismetCompilerConte
 	return new FKCHandler_Branch(CompilerContext);
 }
 
+void UK2Node_IfThenElse::GetMenuActions(TArray<UBlueprintNodeSpawner*>& ActionListOut) const
+{
+	UBlueprintNodeSpawner* NodeSpawner = UBlueprintNodeSpawner::Create(GetClass());
+	check(NodeSpawner != nullptr);
+
+	ActionListOut.Add(NodeSpawner);
+}
+
+FText UK2Node_IfThenElse::GetMenuCategory() const
+{
+	return FEditorCategoryUtils::GetCommonCategory(FCommonEditorCategory::FlowControl);
+}
+
 #undef LOCTEXT_NAMESPACE
diff --git a/Engine/Source/Editor/BlueprintGraph/Private/K2Node_MultiGate.cpp b/Engine/Source/Editor/BlueprintGraph/Private/K2Node_MultiGate.cpp
--- a/Engine/Source/Editor/BlueprintGraph/Private/K2Node_MultiGate.cpp
+++ b/Engine/Source/Editor/BlueprintGraph/Private/K2Node_MultiGate.cpp
@@ -6,6 +6,8 @@
 #include "Kismet/KismetNodeHelperLibrary.h"
 #include "Kismet/KismetSystemLibrary.h"
 #include "K2Node_MultiGate.h"
+#include "BlueprintNodeSpawner.h"
+
 #define LOCTEXT_NAMESPACE "K2Node_MultiGate"
 
 //////////////////////////////////////////////////////////////////////////
@@ -654,4 +656,12 @@ void UK2Node_MultiGate::ExpandNode(class FKismetCompilerContext& CompilerContext
 	}
 }
 
+void UK2Node_MultiGate::GetMenuActions(TArray<UBlueprintNodeSpawner*>& ActionListOut) const
+{
+	UBlueprintNodeSpawner* NodeSpawner = UBlueprintNodeSpawner::Create(GetClass());
+	check(NodeSpawner != nullptr);
+
+	ActionListOut.Add(NodeSpawner);
+}
+
 #undef LOCTEXT_NAMESPACE
diff --git a/Engine/Source/Editor/BlueprintGraph/Private/K2Node_Select.cpp b/Engine/Source/Editor/BlueprintGraph/Private/K2Node_Select.cpp
--- a/Engine/Source/Editor/BlueprintGraph/Private/K2Node_Select.cpp
+++ b/Engine/Source/Editor/BlueprintGraph/Private/K2Node_Select.cpp
@@ -1,12 +1,11 @@
 // Copyright 1998-2014 Epic Games, Inc. All Rights Reserved.
 
-
 #include "BlueprintGraphPrivatePCH.h"
-
 #include "../../../Runtime/Engine/Classes/Kismet/KismetMathLibrary.h"
 #include "../../../Runtime/Engine/Classes/Kismet/KismetSystemLibrary.h"
-
 #include "KismetCompiler.h"
+#include "BlueprintNodeSpawner.h"
+#include "EditorCategoryUtils.h"
 
 #define LOCTEXT_NAMESPACE "K2Node_Select"
 
@@ -775,4 +774,17 @@ FNodeHandlingFunctor* UK2Node_Select::CreateNodeHandler(FKismetCompilerContext&
 	return new FKCHandler_Select(CompilerContext);
 }
 
+void UK2Node_Select::GetMenuActions(TArray<UBlueprintNodeSpawner*>& ActionListOut) const
+{
+	UBlueprintNodeSpawner* NodeSpawner = UBlueprintNodeSpawner::Create(GetClass());
+	check(NodeSpawner != nullptr);
+
+	ActionListOut.Add(NodeSpawner);
+}
+
+FText UK2Node_Select::GetMenuCategory() const
+{
+	return FEditorCategoryUtils::GetCommonCategory(FCommonEditorCategory::Utilities);
+}
+
 #undef LOCTEXT_NAMESPACE
diff --git a/Engine/Source/Editor/BlueprintGraph/Private/K2Node_SelectEnum.cpp b/Engine/Source/Editor/BlueprintGraph/Private/K2Node_SelectEnum.cpp
--- a/Engine/Source/Editor/BlueprintGraph/Private/K2Node_SelectEnum.cpp
+++ b/Engine/Source/Editor/BlueprintGraph/Private/K2Node_SelectEnum.cpp
@@ -3,6 +3,7 @@
 
 #include "BlueprintGraphPrivatePCH.h"
 #include "K2Node_SelectEnum.h"
+#include "BlueprintNodeSpawner.h"
 
 UDEPRECATED_K2Node_SelectEnum::UDEPRECATED_K2Node_SelectEnum(const class FPostConstructInitializeProperties& PCIP)
 	: Super(PCIP)
@@ -61,6 +62,29 @@ void UDEPRECATED_K2Node_SelectEnum::NotifyPinConnectionListChanged(UEdGraphPin*
 	// don't do anything
 }
 
+void UDEPRECATED_K2Node_SelectEnum::GetMenuActions(TArray<UBlueprintNodeSpawner*>& ActionListOut) const
+{
+	for (TObjectIterator<UEnum> EnumIt; EnumIt; ++EnumIt)
+	{
+		UEnum const* Enum = (*EnumIt);
+		auto CustomizeEnumNodeLambda = [](UEdGraphNode* NewNode, bool bIsTemplateNode, TWeakObjectPtr<UEnum> EnumPtr)
+		{
+			UDEPRECATED_K2Node_SelectEnum* EnumNode = CastChecked<UDEPRECATED_K2Node_SelectEnum>(NewNode);
+			if (EnumPtr.IsValid())
+			{
+				EnumNode->Enum = EnumPtr.Get();
+			}
+		};
+		
+		UBlueprintNodeSpawner* NodeSpawner = UBlueprintNodeSpawner::Create(GetClass());
+		check(NodeSpawner != nullptr);
+		ActionListOut.Add(NodeSpawner);
+		
+		TWeakObjectPtr<UEnum> EnumPtr = Enum;
+		NodeSpawner->CustomizeNodeDelegate = UBlueprintNodeSpawner::FCustomizeNodeDelegate::CreateStatic(CustomizeEnumNodeLambda, EnumPtr);
+	}
+}
+
 void UDEPRECATED_K2Node_SelectEnum::GetOptionPins(TArray<UEdGraphPin*>& OptionPins) const
 {
 	OptionPins.Empty();
diff --git a/Engine/Source/Editor/BlueprintGraph/Private/K2Node_Switch.cpp b/Engine/Source/Editor/BlueprintGraph/Private/K2Node_Switch.cpp
--- a/Engine/Source/Editor/BlueprintGraph/Private/K2Node_Switch.cpp
+++ b/Engine/Source/Editor/BlueprintGraph/Private/K2Node_Switch.cpp
@@ -5,6 +5,7 @@
 #include "KismetCompiler.h"
 #include "../../../Runtime/Engine/Classes/Kismet/KismetMathLibrary.h"
 #include "K2Node_SwitchEnum.h"
+#include "EditorCategoryUtils.h"
 
 #define LOCTEXT_NAMESPACE "K2Node_Switch"
 
@@ -325,4 +326,9 @@ FNodeHandlingFunctor* UK2Node_Switch::CreateNodeHandler(FKismetCompilerContext&
 	return new FKCHandler_Switch(CompilerContext, GetPinType(CompilerContext.GetSchema()));
 }
 
+FText UK2Node_Switch::GetMenuCategory() const
+{
+	return FEditorCategoryUtils::BuildCategoryString(FCommonEditorCategory::FlowControl, LOCTEXT("ActionMenuCategory", "Switch"));
+}
+
 #undef LOCTEXT_NAMESPACE
diff --git a/Engine/Source/Editor/BlueprintGraph/Private/K2Node_SwitchEnum.cpp b/Engine/Source/Editor/BlueprintGraph/Private/K2Node_SwitchEnum.cpp
--- a/Engine/Source/Editor/BlueprintGraph/Private/K2Node_SwitchEnum.cpp
+++ b/Engine/Source/Editor/BlueprintGraph/Private/K2Node_SwitchEnum.cpp
@@ -3,6 +3,7 @@
 #include "BlueprintGraphPrivatePCH.h"
 #include "Kismet/KismetMathLibrary.h"
 #include "K2Node_SwitchEnum.h"
+#include "BlueprintNodeSpawner.h"
 
 #define LOCTEXT_NAMESPACE "K2Node"
 
@@ -61,6 +62,41 @@ FString UK2Node_SwitchEnum::GetTooltip() const
 	return NSLOCTEXT("K2Node", "SwitchEnum_ToolTip", "Selects an output that matches the input value").ToString();
 }
 
+void UK2Node_SwitchEnum::GetMenuActions(TArray<UBlueprintNodeSpawner*>& ActionListOut) const
+{
+	for (TObjectIterator<UEnum> EnumIt; EnumIt; ++EnumIt)
+	{
+		UEnum const* Enum = (*EnumIt);
+		// we only want to add global "standalone" enums here; those belonging to a 
+		// certain class should instead be associated with that class (so when 
+		// the class is modified we can easily handle any enums that were changed).
+		//
+		// @TODO: don't love how this code is essentially duplicated in BlueprintActionDatabase.cpp, for class enums
+		bool bIsStandaloneEnum = Enum->GetOuter()->IsA(UPackage::StaticClass());
+
+		if (!bIsStandaloneEnum || !UEdGraphSchema_K2::IsAllowableBlueprintVariableType(Enum))
+		{
+			continue;
+		}
+
+		auto CustomizeEnumNodeLambda = [](UEdGraphNode* NewNode, bool bIsTemplateNode, TWeakObjectPtr<UEnum> EnumPtr)
+		{
+			UK2Node_SwitchEnum* EnumNode = CastChecked<UK2Node_SwitchEnum>(NewNode);
+			if (EnumPtr.IsValid())
+			{
+				EnumNode->SetEnum(EnumPtr.Get());
+			}
+		};
+
+		UBlueprintNodeSpawner* NodeSpawner = UBlueprintNodeSpawner::Create(GetClass());
+		check(NodeSpawner != nullptr);
+		ActionListOut.Add(NodeSpawner);
+
+		TWeakObjectPtr<UEnum> EnumPtr = Enum;
+		NodeSpawner->CustomizeNodeDelegate = UBlueprintNodeSpawner::FCustomizeNodeDelegate::CreateStatic(CustomizeEnumNodeLambda, EnumPtr);
+	}
+}
+
 void UK2Node_SwitchEnum::CreateSelectionPin()
 {
 	const UEdGraphSchema_K2* K2Schema = GetDefault<UEdGraphSchema_K2>();
@@ -167,4 +203,4 @@ void UK2Node_SwitchEnum::RemovePinFromSwitchNode(UEdGraphPin* Pin)
 	}
 }
 
-#undef LOCTEXT_NAMESPACE
\ No newline at end of file
+#undef LOCTEXT_NAMESPACE
diff --git a/Engine/Source/Editor/BlueprintGraph/Private/K2Node_SwitchInteger.cpp b/Engine/Source/Editor/BlueprintGraph/Private/K2Node_SwitchInteger.cpp
--- a/Engine/Source/Editor/BlueprintGraph/Private/K2Node_SwitchInteger.cpp
+++ b/Engine/Source/Editor/BlueprintGraph/Private/K2Node_SwitchInteger.cpp
@@ -2,6 +2,7 @@
 
 #include "BlueprintGraphPrivatePCH.h"
 #include "../../../Runtime/Engine/Classes/Kismet/KismetMathLibrary.h"
+#include "BlueprintNodeSpawner.h"
 
 UK2Node_SwitchInteger::UK2Node_SwitchInteger(const class FPostConstructInitializeProperties& PCIP)
 	: Super(PCIP)
@@ -22,6 +23,14 @@ FString UK2Node_SwitchInteger::GetTooltip() const
 	return NSLOCTEXT("K2Node", "SwitchInt_ToolTip", "Selects an output that matches the input value").ToString();
 }
 
+void UK2Node_SwitchInteger::GetMenuActions(TArray<UBlueprintNodeSpawner*>& ActionListOut) const
+{
+	UBlueprintNodeSpawner* NodeSpawner = UBlueprintNodeSpawner::Create(GetClass());
+	check(NodeSpawner != nullptr);
+
+	ActionListOut.Add(NodeSpawner);
+}
+
 void UK2Node_SwitchInteger::CreateSelectionPin()
 {
 	const UEdGraphSchema_K2* K2Schema = GetDefault<UEdGraphSchema_K2>();
diff --git a/Engine/Source/Editor/BlueprintGraph/Private/K2Node_SwitchName.cpp b/Engine/Source/Editor/BlueprintGraph/Private/K2Node_SwitchName.cpp
--- a/Engine/Source/Editor/BlueprintGraph/Private/K2Node_SwitchName.cpp
+++ b/Engine/Source/Editor/BlueprintGraph/Private/K2Node_SwitchName.cpp
@@ -2,6 +2,7 @@
 
 #include "BlueprintGraphPrivatePCH.h"
 #include "../../../Runtime/Engine/Classes/Kismet/KismetMathLibrary.h"
+#include "BlueprintNodeSpawner.h"
 
 UK2Node_SwitchName::UK2Node_SwitchName(const class FPostConstructInitializeProperties& PCIP)
 	: Super(PCIP)
@@ -37,6 +38,14 @@ FString UK2Node_SwitchName::GetTooltip() const
 	return NSLOCTEXT("K2Node", "SwitchName_ToolTip", "Selects an output that matches the input value").ToString();
 }
 
+void UK2Node_SwitchName::GetMenuActions(TArray<UBlueprintNodeSpawner*>& ActionListOut) const
+{
+	UBlueprintNodeSpawner* NodeSpawner = UBlueprintNodeSpawner::Create(GetClass());
+	check(NodeSpawner != nullptr);
+
+	ActionListOut.Add(NodeSpawner);
+}
+
 void UK2Node_SwitchName::CreateSelectionPin()
 {
 	const UEdGraphSchema_K2* K2Schema = GetDefault<UEdGraphSchema_K2>();
diff --git a/Engine/Source/Editor/BlueprintGraph/Private/K2Node_SwitchString.cpp b/Engine/Source/Editor/BlueprintGraph/Private/K2Node_SwitchString.cpp
--- a/Engine/Source/Editor/BlueprintGraph/Private/K2Node_SwitchString.cpp
+++ b/Engine/Source/Editor/BlueprintGraph/Private/K2Node_SwitchString.cpp
@@ -3,6 +3,7 @@
 #include "BlueprintGraphPrivatePCH.h"
 #include "Kismet/KismetMathLibrary.h"
 #include "K2Node_SwitchString.h"
+#include "BlueprintNodeSpawner.h"
 
 UK2Node_SwitchString::UK2Node_SwitchString(const class FPostConstructInitializeProperties& PCIP)
 	: Super(PCIP)
@@ -49,6 +50,14 @@ FString UK2Node_SwitchString::GetTooltip() const
 	return NSLOCTEXT("K2Node", "SwitchString_ToolTip", "Selects an output that matches the input value").ToString();
 }
 
+void UK2Node_SwitchString::GetMenuActions(TArray<UBlueprintNodeSpawner*>& ActionListOut) const
+{
+	UBlueprintNodeSpawner* NodeSpawner = UBlueprintNodeSpawner::Create(GetClass());
+	check(NodeSpawner != nullptr);
+
+	ActionListOut.Add(NodeSpawner);
+}
+
 void UK2Node_SwitchString::CreateSelectionPin()
 {
 	const UEdGraphSchema_K2* K2Schema = GetDefault<UEdGraphSchema_K2>();
diff --git a/Engine/Source/Editor/UnrealEd/Private/EditorCategoryUtils.cpp b/Engine/Source/Editor/UnrealEd/Private/EditorCategoryUtils.cpp
--- a/Engine/Source/Editor/UnrealEd/Private/EditorCategoryUtils.cpp
+++ b/Engine/Source/Editor/UnrealEd/Private/EditorCategoryUtils.cpp
@@ -52,7 +52,7 @@ static FFormatNamedArguments& FEditorCategoryUtilsImpl::GetCategoryTable()
 		FEditorCategoryUtils::RegisterCategoryKey(ENUM_STRING(EnumVal), LOCTEXT(#EnumVal L"Category", Text))
 #define REGISTER_ROOT_CATEGORY(EnumVal) ASSIGN_ROOT_CATEGORY(EnumVal, #EnumVal)
 
-		REGISTER_ROOT_CATEGORY(Ai);
+		ASSIGN_ROOT_CATEGORY(Ai, "AI");
 		REGISTER_ROOT_CATEGORY(Animation);
 		REGISTER_ROOT_CATEGORY(Audio);
 		REGISTER_ROOT_CATEGORY(Development);
@@ -73,6 +73,7 @@ static FFormatNamedArguments& FEditorCategoryUtilsImpl::GetCategoryTable()
 		FText EnumVal##_Category = FEditorCategoryUtils::BuildCategoryString(FCommonEditorCategory::RootId, LOCTEXT(#EnumVal L"Category", #EnumVal)); \
 		FEditorCategoryUtils::RegisterCategoryKey(ENUM_STRING(EnumVal), EnumVal##_Category)
 
+		REGISTER_SUB_CATEGORY(Utilities, FlowControl);
 		REGISTER_SUB_CATEGORY(Utilities, Transformation);
 		REGISTER_SUB_CATEGORY(Utilities, String);
 		REGISTER_SUB_CATEGORY(Utilities, Text);
@@ -130,6 +131,7 @@ FText const& FEditorCategoryUtils::GetCommonCategory(FCommonEditorCategory::EVal
 		REGISTER_COMMON_CATEGORY(Rendering);
 		REGISTER_COMMON_CATEGORY(Utilities);
 
+		REGISTER_COMMON_CATEGORY(FlowControl);
 		REGISTER_COMMON_CATEGORY(Transformation);
 		REGISTER_COMMON_CATEGORY(String);
 		REGISTER_COMMON_CATEGORY(Text);
diff --git a/Engine/Source/Editor/UnrealEd/Public/EditorCategoryUtils.h b/Engine/Source/Editor/UnrealEd/Public/EditorCategoryUtils.h
--- a/Engine/Source/Editor/UnrealEd/Public/EditorCategoryUtils.h
+++ b/Engine/Source/Editor/UnrealEd/Public/EditorCategoryUtils.h
@@ -32,8 +32,9 @@ struct FCommonEditorCategory
 		Rendering,
 		Transformation,
 		Utilities,
+		FlowControl,
 
-		// Type categories:
+		// Type library categories:
 		String,
 		Text,
 		Name,
