(WIP) Standalone, variable centric nodes now adding themselves to the blueprint action database (for the new action menu system).

diff --git a/Engine/Source/Editor/BlueprintGraph/Classes/K2Node_AssignmentStatement.h b/Engine/Source/Editor/BlueprintGraph/Classes/K2Node_AssignmentStatement.h
--- a/Engine/Source/Editor/BlueprintGraph/Classes/K2Node_AssignmentStatement.h
+++ b/Engine/Source/Editor/BlueprintGraph/Classes/K2Node_AssignmentStatement.h
@@ -19,12 +19,15 @@ class UK2Node_AssignmentStatement : public UK2Node
 	virtual void AllocateDefaultPins() override;
 	virtual FString GetTooltip() const override;
 	virtual FText GetNodeTitle(ENodeTitleType::Type TitleType) const override;
+	virtual bool CanPasteHere(UEdGraph const* TargetGraph, UEdGraphSchema const* Schema) const override;
 	// End UEdGraphNode interface
 
 	// Begin UK2Node interface
 	virtual void PostReconstructNode() override;
 	virtual void NotifyPinConnectionListChanged(UEdGraphPin* Pin) override;
 	virtual class FNodeHandlingFunctor* CreateNodeHandler(class FKismetCompilerContext& CompilerContext) const override;
+	virtual void GetMenuActions(TArray<UBlueprintNodeSpawner*>& ActionListOut) const override;
+	virtual FText GetMenuCategory() const override;
 	// End UK2Node interface
 
 	/** Get the Then output pin */
diff --git a/Engine/Source/Editor/BlueprintGraph/Classes/K2Node_Self.h b/Engine/Source/Editor/BlueprintGraph/Classes/K2Node_Self.h
--- a/Engine/Source/Editor/BlueprintGraph/Classes/K2Node_Self.h
+++ b/Engine/Source/Editor/BlueprintGraph/Classes/K2Node_Self.h
@@ -26,6 +26,8 @@ class UK2Node_Self : public UK2Node
 	virtual bool IsNodePure() const override { return true; }
 	virtual bool DrawNodeAsVariable() const override { return true; }
 	virtual class FNodeHandlingFunctor* CreateNodeHandler(class FKismetCompilerContext& CompilerContext) const override;
+	virtual void GetMenuActions(TArray<UBlueprintNodeSpawner*>& ActionListOut) const override;
+	virtual FText GetMenuCategory() const override;
 	// End K2Node interface
 };
 
diff --git a/Engine/Source/Editor/BlueprintGraph/Classes/K2Node_TemporaryVariable.h b/Engine/Source/Editor/BlueprintGraph/Classes/K2Node_TemporaryVariable.h
--- a/Engine/Source/Editor/BlueprintGraph/Classes/K2Node_TemporaryVariable.h
+++ b/Engine/Source/Editor/BlueprintGraph/Classes/K2Node_TemporaryVariable.h
@@ -24,11 +24,14 @@ class UK2Node_TemporaryVariable : public UK2Node
 	virtual FString GetTooltip() const override;
 	virtual FText GetNodeTitle(ENodeTitleType::Type TitleType) const override;
 	virtual FString GetDescriptiveCompiledName() const override;
+	virtual bool CanPasteHere(UEdGraph const* TargetGraph, UEdGraphSchema const* Schema) const override;
 	// End UEdGraphNode interface.
 
 	// Begin UK2Node interface.
 	virtual bool IsNodePure() const override;
 	virtual class FNodeHandlingFunctor* CreateNodeHandler(class FKismetCompilerContext& CompilerContext) const override;
+	virtual void GetMenuActions(TArray<UBlueprintNodeSpawner*>& ActionListOut) const override;
+	virtual FText GetMenuCategory() const override;
 	// End UK2Node interface.
 };
 
diff --git a/Engine/Source/Editor/BlueprintGraph/Classes/K2Node_VariableSetRef.h b/Engine/Source/Editor/BlueprintGraph/Classes/K2Node_VariableSetRef.h
--- a/Engine/Source/Editor/BlueprintGraph/Classes/K2Node_VariableSetRef.h
+++ b/Engine/Source/Editor/BlueprintGraph/Classes/K2Node_VariableSetRef.h
@@ -19,6 +19,8 @@ class UK2Node_VariableSetRef : public UK2Node
 	// Begin UK2Node interface
 	virtual void NotifyPinConnectionListChanged(UEdGraphPin* Pin) override;
 	virtual class FNodeHandlingFunctor* CreateNodeHandler(class FKismetCompilerContext& CompilerContext) const override;
+	virtual void GetMenuActions(TArray<UBlueprintNodeSpawner*>& ActionListOut) const override;
+	virtual FText GetMenuCategory() const override;
 	// End UK2Node interface
 
 	/** 
diff --git a/Engine/Source/Editor/BlueprintGraph/Private/BlueprintActionMenuBuilder.cpp b/Engine/Source/Editor/BlueprintGraph/Private/BlueprintActionMenuBuilder.cpp
--- a/Engine/Source/Editor/BlueprintGraph/Private/BlueprintActionMenuBuilder.cpp
+++ b/Engine/Source/Editor/BlueprintGraph/Private/BlueprintActionMenuBuilder.cpp
@@ -197,7 +197,7 @@ FText FBlueprintActionMenuItemFactory::GetCategoryForAction(FBlueprintActionMenu
 				{
 					MenuCategory = FText::FromName(FuncOwner->GetFName());
 				}
-				MenuCategory = FText::Format(LOCTEXT("MemberFunctionsCategory", "Call Funtions|{0}"), MenuCategory);
+				MenuCategory = FText::Format(LOCTEXT("MemberFunctionsCategory", "Call Funtion|{0}"), MenuCategory);
 			}
 		}
 	}
@@ -365,6 +365,9 @@ void FBlueprintActionMenuBuilderImpl::AppendLegacyActions(TArray< TSharedRef<FMe
  * FBlueprintActionMenuBuilder
  ******************************************************************************/
 
+FText const FBlueprintActionMenuBuilder::AddEventCategory  = LOCTEXT("AddEventCategory",  "Add Event");
+FText const FBlueprintActionMenuBuilder::VariablesCategory = LOCTEXT("VariablesCategory", "Variables");
+
 //------------------------------------------------------------------------------
 FBlueprintActionMenuBuilder::FBlueprintActionMenuBuilder(FBlueprintActionFilter const& Filter, bool bAutoBuild/*= true*/)
 {
diff --git a/Engine/Source/Editor/BlueprintGraph/Private/K2Node_AssignmentStatement.cpp b/Engine/Source/Editor/BlueprintGraph/Private/K2Node_AssignmentStatement.cpp
--- a/Engine/Source/Editor/BlueprintGraph/Private/K2Node_AssignmentStatement.cpp
+++ b/Engine/Source/Editor/BlueprintGraph/Private/K2Node_AssignmentStatement.cpp
@@ -4,6 +4,8 @@
 #include "BlueprintGraphPrivatePCH.h"
 #include "KismetCompiler.h"
 #include "VariableSetHandler.h"
+#include "BlueprintNodeSpawner.h"
+#include "EditorCategoryUtils.h"
 
 #define LOCTEXT_NAMESPACE "K2Node_AssignmentStatement"
 
@@ -82,6 +84,19 @@ FText UK2Node_AssignmentStatement::GetNodeTitle(ENodeTitleType::Type TitleType)
 	return LOCTEXT("Assign", "Assign");
 }
 
+bool UK2Node_AssignmentStatement::CanPasteHere(UEdGraph const* TargetGraph, UEdGraphSchema const* Schema) const
+{
+	bool bIsPastable = Super::CanPasteHere(TargetGraph, Schema);
+	if (bIsPastable)
+	{
+		EGraphType const GraphType = Schema->GetGraphType(TargetGraph);
+
+		bIsPastable = (GraphType != GT_Ubergraph) && (GraphType != GT_Animation);
+	}
+
+	return bIsPastable;
+}
+
 void UK2Node_AssignmentStatement::NotifyPinConnectionListChanged(UEdGraphPin* Pin)
 {
 	Super::NotifyPinConnectionListChanged(Pin);
@@ -160,4 +175,17 @@ FNodeHandlingFunctor* UK2Node_AssignmentStatement::CreateNodeHandler(FKismetComp
 	return new FKCHandler_AssignmentStatement(CompilerContext);
 }
 
+void UK2Node_AssignmentStatement::GetMenuActions(TArray<UBlueprintNodeSpawner*>& ActionListOut) const
+{
+	UBlueprintNodeSpawner* NodeSpawner = UBlueprintNodeSpawner::Create(GetClass());
+	check(NodeSpawner != nullptr);
+
+	ActionListOut.Add(NodeSpawner);
+}
+
+FText UK2Node_AssignmentStatement::GetMenuCategory() const
+{
+	return FEditorCategoryUtils::GetCommonCategory(FCommonEditorCategory::Macro);
+}
+
 #undef LOCTEXT_NAMESPACE
diff --git a/Engine/Source/Editor/BlueprintGraph/Private/K2Node_Self.cpp b/Engine/Source/Editor/BlueprintGraph/Private/K2Node_Self.cpp
--- a/Engine/Source/Editor/BlueprintGraph/Private/K2Node_Self.cpp
+++ b/Engine/Source/Editor/BlueprintGraph/Private/K2Node_Self.cpp
@@ -3,6 +3,10 @@
 
 #include "BlueprintGraphPrivatePCH.h"
 #include "KismetCompiler.h"
+#include "BlueprintNodeSpawner.h"
+#include "BlueprintActionMenuBuilder.h"
+
+#define LOCTEXT_NAMESPACE "K2Node_Self"
 
 class FKCHandler_Self : public FNodeHandlingFunctor
 {
@@ -54,7 +58,12 @@ FString UK2Node_Self::GetKeywords() const
 
 FText UK2Node_Self::GetNodeTitle(ENodeTitleType::Type TitleType) const
 {
-	return NSLOCTEXT("K2Node", "SelfReferenceName", "Self-Reference");
+	FText NodeTitle = NSLOCTEXT("K2Node", "SelfReferenceName", "Self-Reference");
+	if (TitleType == ENodeTitleType::ListView)
+	{
+		NodeTitle = LOCTEXT("ListTitle", "Get a reference to self");
+	}
+	return NodeTitle;	
 }
 
 FNodeHandlingFunctor* UK2Node_Self::CreateNodeHandler(FKismetCompilerContext& CompilerContext) const
@@ -72,4 +81,19 @@ void UK2Node_Self::ValidateNodeDuringCompilation(class FCompilerResultsLog& Mess
 	{
 		MessageLog.Warning(*NSLOCTEXT("K2Node", "InvalidSelfNode", "Self node @@ cannot be used in static library.").ToString(), this);
 	}
-}
\ No newline at end of file
+}
+
+void UK2Node_Self::GetMenuActions(TArray<UBlueprintNodeSpawner*>& ActionListOut) const
+{
+	UBlueprintNodeSpawner* NodeSpawner = UBlueprintNodeSpawner::Create(GetClass());
+	check(NodeSpawner != nullptr);
+
+	ActionListOut.Add(NodeSpawner);
+}
+
+FText UK2Node_Self::GetMenuCategory() const
+{
+	return FBlueprintActionMenuBuilder::VariablesCategory;
+}
+
+#undef LOCTEXT_NAMESPACE
diff --git a/Engine/Source/Editor/BlueprintGraph/Private/K2Node_TemporaryVariable.cpp b/Engine/Source/Editor/BlueprintGraph/Private/K2Node_TemporaryVariable.cpp
--- a/Engine/Source/Editor/BlueprintGraph/Private/K2Node_TemporaryVariable.cpp
+++ b/Engine/Source/Editor/BlueprintGraph/Private/K2Node_TemporaryVariable.cpp
@@ -3,6 +3,8 @@
 
 #include "BlueprintGraphPrivatePCH.h"
 #include "KismetCompiler.h"
+#include "BlueprintActionFilter.h"
+#include "EditorCategoryUtils.h"
 
 class FKCHandler_TemporaryVariable : public FNodeHandlingFunctor
 {
@@ -86,6 +88,23 @@ FString UK2Node_TemporaryVariable::GetDescriptiveCompiledName() const
 	return Result;
 }
 
+bool UK2Node_TemporaryVariable::CanPasteHere(UEdGraph const* TargetGraph, UEdGraphSchema const* Schema) const
+{
+	bool bIsPastable = Super::CanPasteHere(TargetGraph, Schema);
+	if (bIsPastable)
+	{
+		EGraphType const GraphType = Schema->GetGraphType(TargetGraph);
+		bIsPastable = (GraphType != GT_Ubergraph) && (GraphType != GT_Animation);
+
+		if (bIsPastable && (GraphType != GT_Macro))
+		{
+			bIsPastable = !bIsPersistent;
+		}
+	}
+
+	return bIsPastable;
+}
+
 // get variable pin
 UEdGraphPin* UK2Node_TemporaryVariable::GetVariablePin()
 {
@@ -96,3 +115,60 @@ FNodeHandlingFunctor* UK2Node_TemporaryVariable::CreateNodeHandler(FKismetCompil
 {
 	return new FKCHandler_TemporaryVariable(CompilerContext);
 }
+
+void UK2Node_TemporaryVariable::GetMenuActions(TArray<UBlueprintNodeSpawner*>& ActionListOut) const
+{
+	auto MakeTempVarNodeSpawner = [](FEdGraphPinType const& VarType, bool bIsPersistent) -> UBlueprintNodeSpawner*
+	{
+		UBlueprintNodeSpawner* NodeSpawner = UBlueprintNodeSpawner::Create(UK2Node_TemporaryVariable::StaticClass());
+		check(NodeSpawner != nullptr);
+
+		auto PostSpawnLambda = [](UEdGraphNode* NewNode, bool bIsTemplateNode, FEdGraphPinType VarType, bool bIsPersistent)
+		{
+			UK2Node_TemporaryVariable* TempVarNode = CastChecked<UK2Node_TemporaryVariable>(NewNode);
+			TempVarNode->VariableType  = VarType;
+			TempVarNode->bIsPersistent = bIsPersistent;
+		};
+
+		NodeSpawner->CustomizeNodeDelegate = UBlueprintNodeSpawner::FCustomizeNodeDelegate::CreateStatic(PostSpawnLambda, VarType, bIsPersistent);
+		return NodeSpawner;
+	};
+
+	UEdGraphSchema_K2 const* K2Schema = GetDefault<UEdGraphSchema_K2>();
+
+	ActionListOut.Add(MakeTempVarNodeSpawner(FEdGraphPinType(K2Schema->PC_Int,     TEXT(""), nullptr, /*bIsArray =*/false, /*bIsReference =*/false), /*bIsPersistent =*/false));
+	ActionListOut.Add(MakeTempVarNodeSpawner(FEdGraphPinType(K2Schema->PC_Int,     TEXT(""), nullptr, /*bIsArray =*/ true, /*bIsReference =*/false), /*bIsPersistent =*/false));
+	ActionListOut.Add(MakeTempVarNodeSpawner(FEdGraphPinType(K2Schema->PC_Float,   TEXT(""), nullptr, /*bIsArray =*/false, /*bIsReference =*/false), /*bIsPersistent =*/false));
+	ActionListOut.Add(MakeTempVarNodeSpawner(FEdGraphPinType(K2Schema->PC_Float,   TEXT(""), nullptr, /*bIsArray =*/ true, /*bIsReference =*/false), /*bIsPersistent =*/false));
+	ActionListOut.Add(MakeTempVarNodeSpawner(FEdGraphPinType(K2Schema->PC_Boolean, TEXT(""), nullptr, /*bIsArray =*/false, /*bIsReference =*/false), /*bIsPersistent =*/false));
+	ActionListOut.Add(MakeTempVarNodeSpawner(FEdGraphPinType(K2Schema->PC_Boolean, TEXT(""), nullptr, /*bIsArray =*/ true, /*bIsReference =*/false), /*bIsPersistent =*/false));
+	ActionListOut.Add(MakeTempVarNodeSpawner(FEdGraphPinType(K2Schema->PC_String,  TEXT(""), nullptr, /*bIsArray =*/false, /*bIsReference =*/false), /*bIsPersistent =*/false));
+	ActionListOut.Add(MakeTempVarNodeSpawner(FEdGraphPinType(K2Schema->PC_String,  TEXT(""), nullptr, /*bIsArray =*/ true, /*bIsReference =*/false), /*bIsPersistent =*/false));
+	ActionListOut.Add(MakeTempVarNodeSpawner(FEdGraphPinType(K2Schema->PC_Text,    TEXT(""), nullptr, /*bIsArray =*/false, /*bIsReference =*/false), /*bIsPersistent =*/false));
+	ActionListOut.Add(MakeTempVarNodeSpawner(FEdGraphPinType(K2Schema->PC_Text,    TEXT(""), nullptr, /*bIsArray =*/ true, /*bIsReference =*/false), /*bIsPersistent =*/false));
+
+	UScriptStruct* VectorStruct  = FindObjectChecked<UScriptStruct>(UObject::StaticClass(), TEXT("Vector"));
+	ActionListOut.Add(MakeTempVarNodeSpawner(FEdGraphPinType(K2Schema->PC_Struct, TEXT("Vector"), VectorStruct, /*bIsArray =*/false, /*bIsReference =*/false), /*bIsPersistent =*/false));
+	ActionListOut.Add(MakeTempVarNodeSpawner(FEdGraphPinType(K2Schema->PC_Struct, TEXT("Vector"), VectorStruct, /*bIsArray =*/ true, /*bIsReference =*/false), /*bIsPersistent =*/false));
+	
+	UScriptStruct* RotatorStruct = FindObjectChecked<UScriptStruct>(UObject::StaticClass(), TEXT("Rotator"));
+	ActionListOut.Add(MakeTempVarNodeSpawner(FEdGraphPinType(K2Schema->PC_Struct, TEXT("Rotator"), RotatorStruct, /*bIsArray =*/false, /*bIsReference =*/false), /*bIsPersistent =*/false));
+	ActionListOut.Add(MakeTempVarNodeSpawner(FEdGraphPinType(K2Schema->PC_Struct, TEXT("Rotator"), RotatorStruct, /*bIsArray =*/ true, /*bIsReference =*/false), /*bIsPersistent =*/false));
+	
+	UScriptStruct* TransformStruct = FindObjectChecked<UScriptStruct>(UObject::StaticClass(), TEXT("Transform"));
+	ActionListOut.Add(MakeTempVarNodeSpawner(FEdGraphPinType(K2Schema->PC_Struct, TEXT("Transform"), TransformStruct, /*bIsArray =*/false, /*bIsReference =*/false), /*bIsPersistent =*/false));
+	ActionListOut.Add(MakeTempVarNodeSpawner(FEdGraphPinType(K2Schema->PC_Struct, TEXT("Transform"), TransformStruct, /*bIsArray =*/ true, /*bIsReference =*/false), /*bIsPersistent =*/false));
+
+	UScriptStruct* BlendSampleStruct = FindObjectChecked<UScriptStruct>(ANY_PACKAGE, TEXT("BlendSampleData"));
+	ActionListOut.Add(MakeTempVarNodeSpawner(FEdGraphPinType(K2Schema->PC_Struct, TEXT("BlendSampleData"), BlendSampleStruct, /*bIsArray =*/false, /*bIsReference =*/false), /*bIsPersistent =*/false));
+	ActionListOut.Add(MakeTempVarNodeSpawner(FEdGraphPinType(K2Schema->PC_Struct, TEXT("BlendSampleData"), BlendSampleStruct, /*bIsArray =*/ true, /*bIsReference =*/false), /*bIsPersistent =*/false));
+
+	// add persistent bool and int types (for macro graphs)
+	ActionListOut.Add(MakeTempVarNodeSpawner(FEdGraphPinType(K2Schema->PC_Int,     TEXT(""), nullptr, /*bIsArray =*/false, /*bIsReference =*/false), /*bIsPersistent =*/true));
+	ActionListOut.Add(MakeTempVarNodeSpawner(FEdGraphPinType(K2Schema->PC_Boolean, TEXT(""), nullptr, /*bIsArray =*/false, /*bIsReference =*/false), /*bIsPersistent =*/true));
+}
+
+FText UK2Node_TemporaryVariable::GetMenuCategory() const
+{
+	return FEditorCategoryUtils::GetCommonCategory(FCommonEditorCategory::Macro);
+}
diff --git a/Engine/Source/Editor/BlueprintGraph/Private/K2Node_VariableSetRef.cpp b/Engine/Source/Editor/BlueprintGraph/Private/K2Node_VariableSetRef.cpp
--- a/Engine/Source/Editor/BlueprintGraph/Private/K2Node_VariableSetRef.cpp
+++ b/Engine/Source/Editor/BlueprintGraph/Private/K2Node_VariableSetRef.cpp
@@ -5,6 +5,8 @@
 #include "KismetCompiler.h"
 #include "VariableSetHandler.h"
 #include "K2Node_VariableSetRef.h"
+#include "BlueprintNodeSpawner.h"
+#include "BlueprintActionMenuBuilder.h" // for VariablesCategory
 
 static FString TargetVarPinName(TEXT("Target"));
 static FString VarValuePinName(TEXT("Value"));
@@ -202,4 +204,18 @@ FNodeHandlingFunctor* UK2Node_VariableSetRef::CreateNodeHandler(FKismetCompilerC
 	return new FKCHandler_VariableSetRef(CompilerContext);
 }
 
+void UK2Node_VariableSetRef::GetMenuActions(TArray<UBlueprintNodeSpawner*>& ActionListOut) const
+{
+	UBlueprintNodeSpawner* NodeSpawner = UBlueprintNodeSpawner::Create(GetClass());
+	check(NodeSpawner != nullptr);
+
+	ActionListOut.Add(NodeSpawner);
+}
+
+FText UK2Node_VariableSetRef::GetMenuCategory() const
+{
+	return FBlueprintActionMenuBuilder::VariablesCategory;
+}
+
+
 #undef LOCTEXT_NAMESPACE
diff --git a/Engine/Source/Editor/BlueprintGraph/Public/BlueprintActionMenuBuilder.h b/Engine/Source/Editor/BlueprintGraph/Public/BlueprintActionMenuBuilder.h
--- a/Engine/Source/Editor/BlueprintGraph/Public/BlueprintActionMenuBuilder.h
+++ b/Engine/Source/Editor/BlueprintGraph/Public/BlueprintActionMenuBuilder.h
@@ -24,6 +24,10 @@ struct BLUEPRINTGRAPH_API FBlueprintActionMenuBuilder : public FGraphActionListB
 public:
 	/** Default constructor */
 	FBlueprintActionMenuBuilder() {}
+	
+	/** Blueprint centric menu categories (not generic enough to make it into FCommonEditorCategory) */
+	static FText const AddEventCategory;
+	static FText const VariablesCategory;
 
 	/**
 	 * Good for constructing a menu built from a single filter. If the 
