(WIP) Various standalone node types, adding themselves to the blueprint action database (for the new action menu system).

diff --git a/Engine/Source/Editor/BlueprintGraph/Classes/K2Node_CustomEvent.h b/Engine/Source/Editor/BlueprintGraph/Classes/K2Node_CustomEvent.h
--- a/Engine/Source/Editor/BlueprintGraph/Classes/K2Node_CustomEvent.h
+++ b/Engine/Source/Editor/BlueprintGraph/Classes/K2Node_CustomEvent.h
@@ -25,6 +25,7 @@ class UK2Node_CustomEvent : public UK2Node_Event
 
 	// Begin UK2Node interface
 	BLUEPRINTGRAPH_API virtual void ValidateNodeDuringCompilation(class FCompilerResultsLog& MessageLog) const override;
+	virtual void GetMenuActions(TArray<UBlueprintNodeSpawner*>& ActionListOut) const override;
 	// End UK2Node interface
 
 	// Begin UK2Node_EditablePinBase interface
diff --git a/Engine/Source/Editor/BlueprintGraph/Classes/K2Node_FormatText.h b/Engine/Source/Editor/BlueprintGraph/Classes/K2Node_FormatText.h
--- a/Engine/Source/Editor/BlueprintGraph/Classes/K2Node_FormatText.h
+++ b/Engine/Source/Editor/BlueprintGraph/Classes/K2Node_FormatText.h
@@ -27,6 +27,8 @@ class UK2Node_FormatText : public UK2Node
 	virtual bool NodeCausesStructuralBlueprintChange() const override { return true; }
 	virtual void ExpandNode(class FKismetCompilerContext& CompilerContext, UEdGraph* SourceGraph) override;
 	virtual ERedirectType DoPinsMatchForReconstruction(const UEdGraphPin* NewPin, int32 NewPinIndex, const UEdGraphPin* OldPin, int32 OldPinIndex) const override;
+	virtual void GetMenuActions(TArray<UBlueprintNodeSpawner*>& ActionListOut) const override;
+	virtual FText GetMenuCategory() const override;
 	// End UK2Node interface.
 
 public:
diff --git a/Engine/Source/Editor/BlueprintGraph/Classes/K2Node_GetDataTableRow.h b/Engine/Source/Editor/BlueprintGraph/Classes/K2Node_GetDataTableRow.h
--- a/Engine/Source/Editor/BlueprintGraph/Classes/K2Node_GetDataTableRow.h
+++ b/Engine/Source/Editor/BlueprintGraph/Classes/K2Node_GetDataTableRow.h
@@ -29,6 +29,8 @@ class UK2Node_GetDataTableRow : public UK2Node
 	// Begin UK2Node interface
 	virtual bool IsNodeSafeToIgnore() const override { return true; }
 	virtual void ReallocatePinsDuringReconstruction(TArray<UEdGraphPin*>& OldPins) override;
+	virtual void GetMenuActions(TArray<UBlueprintNodeSpawner*>& ActionListOut) const override;
+	virtual FText GetMenuCategory() const override;
 	// End UK2Node interface
 
 
diff --git a/Engine/Source/Editor/BlueprintGraph/Classes/K2Node_MakeArray.h b/Engine/Source/Editor/BlueprintGraph/Classes/K2Node_MakeArray.h
--- a/Engine/Source/Editor/BlueprintGraph/Classes/K2Node_MakeArray.h
+++ b/Engine/Source/Editor/BlueprintGraph/Classes/K2Node_MakeArray.h
@@ -36,6 +36,8 @@ public:
 	void GetContextMenuActions(const FGraphNodeContextMenuBuilder& Context) const override;
 	virtual class FNodeHandlingFunctor* CreateNodeHandler(class FKismetCompilerContext& CompilerContext) const override;
 	virtual bool IsConnectionDisallowed(const UEdGraphPin* MyPin, const UEdGraphPin* OtherPin, FString& OutReason) const override;
+	virtual void GetMenuActions(TArray<UBlueprintNodeSpawner*>& ActionListOut) const override;
+	virtual FText GetMenuCategory() const override;
 	// End of UK2Node interface
 
 
diff --git a/Engine/Source/Editor/BlueprintGraph/Classes/K2Node_MathExpression.h b/Engine/Source/Editor/BlueprintGraph/Classes/K2Node_MathExpression.h
--- a/Engine/Source/Editor/BlueprintGraph/Classes/K2Node_MathExpression.h
+++ b/Engine/Source/Editor/BlueprintGraph/Classes/K2Node_MathExpression.h
@@ -38,6 +38,7 @@ public:
 	// UK2Node interface
 	virtual void GetMenuEntries(FGraphContextMenuBuilder& ContextMenuBuilder) const override;
 	virtual void ValidateNodeDuringCompilation(class FCompilerResultsLog& MessageLog) const override;
+	virtual void GetMenuActions(TArray<UBlueprintNodeSpawner*>& ActionListOut) const override;
 	// End of UK2Node interface
 
 private:
diff --git a/Engine/Source/Editor/BlueprintGraph/Classes/K2Node_SpawnActorFromClass.h b/Engine/Source/Editor/BlueprintGraph/Classes/K2Node_SpawnActorFromClass.h
--- a/Engine/Source/Editor/BlueprintGraph/Classes/K2Node_SpawnActorFromClass.h
+++ b/Engine/Source/Editor/BlueprintGraph/Classes/K2Node_SpawnActorFromClass.h
@@ -26,6 +26,8 @@ class UK2Node_SpawnActorFromClass : public UK2Node
 	virtual void ReallocatePinsDuringReconstruction(TArray<UEdGraphPin*>& OldPins) override;
 	virtual bool CanPasteHere(const UEdGraph* TargetGraph, const UEdGraphSchema* Schema) const override;
 	virtual void GetNodeAttributes( TArray<TKeyValuePair<FString, FString>>& OutNodeAttributes ) const override;
+	virtual void GetMenuActions(TArray<UBlueprintNodeSpawner*>& ActionListOut) const override;
+	virtual FText GetMenuCategory() const override;
 	// End UK2Node interface
 
 
diff --git a/Engine/Source/Editor/BlueprintGraph/Classes/K2Node_Timeline.h b/Engine/Source/Editor/BlueprintGraph/Classes/K2Node_Timeline.h
--- a/Engine/Source/Editor/BlueprintGraph/Classes/K2Node_Timeline.h
+++ b/Engine/Source/Editor/BlueprintGraph/Classes/K2Node_Timeline.h
@@ -52,6 +52,7 @@ class UK2Node_Timeline : public UK2Node
 	virtual class FNodeHandlingFunctor* CreateNodeHandler(class FKismetCompilerContext& CompilerContext) const override;
 	virtual void ExpandNode(class FKismetCompilerContext& CompilerContext, UEdGraph* SourceGraph) override;
 	virtual void GetNodeAttributes( TArray<TKeyValuePair<FString, FString>>& OutNodeAttributes ) const override;
+	virtual void GetMenuActions(TArray<UBlueprintNodeSpawner*>& ActionListOut) const override;
 	// End UK2Node interface.
 
 	/** Get the 'play' input pin */
diff --git a/Engine/Source/Editor/BlueprintGraph/Private/K2Node_CustomEvent.cpp b/Engine/Source/Editor/BlueprintGraph/Private/K2Node_CustomEvent.cpp
--- a/Engine/Source/Editor/BlueprintGraph/Private/K2Node_CustomEvent.cpp
+++ b/Engine/Source/Editor/BlueprintGraph/Private/K2Node_CustomEvent.cpp
@@ -103,7 +103,11 @@ UK2Node_CustomEvent::UK2Node_CustomEvent(const class FPostConstructInitializePro
 
 FText UK2Node_CustomEvent::GetNodeTitle(ENodeTitleType::Type TitleType) const
 {
-	if (TitleType == ENodeTitleType::EditableTitle || TitleType == ENodeTitleType::ListView)
+	if (CustomFunctionName.IsNone() && (TitleType == ENodeTitleType::ListView))
+	{
+		return LOCTEXT("ActionMenuTitle", "Custom Event...");
+	}
+	else if ((TitleType == ENodeTitleType::EditableTitle) || (TitleType == ENodeTitleType::ListView))
 	{
 		return FText::FromName(CustomFunctionName);
 	}
@@ -223,6 +227,28 @@ void UK2Node_CustomEvent::ValidateNodeDuringCompilation(class FCompilerResultsLo
 	}
 }
 
+void UK2Node_CustomEvent::GetMenuActions(TArray<UBlueprintNodeSpawner*>& ActionListOut) const
+{
+	UBlueprintNodeSpawner* NodeSpawner = UBlueprintNodeSpawner::Create(GetClass());
+	check(NodeSpawner != nullptr);
+
+	auto SetupCustomEventNodeLambda = [](UEdGraphNode* NewNode, bool bIsTemplateNode)
+	{
+		UK2Node_CustomEvent* EventNode = CastChecked<UK2Node_CustomEvent>(NewNode);
+		UBlueprint* Blueprint = EventNode->GetBlueprint();
+
+		// in GetNodeTitle(), we use an empty CustomFunctionName to identify a menu entry
+		if (!bIsTemplateNode)
+		{
+			EventNode->CustomFunctionName = FBlueprintEditorUtils::FindUniqueCustomEventName(Blueprint);
+		}
+		EventNode->bIsEditable = true;
+	};
+
+	NodeSpawner->CustomizeNodeDelegate = UBlueprintNodeSpawner::FCustomizeNodeDelegate::CreateStatic(SetupCustomEventNodeLambda);
+	ActionListOut.Add(NodeSpawner);
+}
+
 void UK2Node_CustomEvent::ReconstructNode()
 {
 	const UEdGraphPin* DelegateOutPin = FindPin(DelegateOutputName);
diff --git a/Engine/Source/Editor/BlueprintGraph/Private/K2Node_FormatText.cpp b/Engine/Source/Editor/BlueprintGraph/Private/K2Node_FormatText.cpp
--- a/Engine/Source/Editor/BlueprintGraph/Private/K2Node_FormatText.cpp
+++ b/Engine/Source/Editor/BlueprintGraph/Private/K2Node_FormatText.cpp
@@ -4,9 +4,10 @@
 #include "BlueprintGraphPrivatePCH.h"
 #include "Editor/GraphEditor/Public/DiffResults.h"
 #include "Kismet2NameValidators.h"
-
 #include "KismetCompiler.h"
 #include "ScopedTransaction.h"
+#include "BlueprintNodeSpawner.h"
+#include "EditorCategoryUtils.h"
 
 #define LOCTEXT_NAMESPACE "K2Node_FormatText"
 
@@ -363,4 +364,17 @@ UEdGraphPin* UK2Node_FormatText::GetFormatPin() const
 	return CachedFormatPin;
 }
 
+void UK2Node_FormatText::GetMenuActions(TArray<UBlueprintNodeSpawner*>& ActionListOut) const
+{
+	UBlueprintNodeSpawner* NodeSpawner = UBlueprintNodeSpawner::Create(GetClass());
+	check(NodeSpawner != nullptr);
+
+	ActionListOut.Add(NodeSpawner);
+}
+
+FText UK2Node_FormatText::GetMenuCategory() const
+{
+	return FEditorCategoryUtils::GetCommonCategory(FCommonEditorCategory::Text);
+}
+
 #undef LOCTEXT_NAMESPACE
\ No newline at end of file
diff --git a/Engine/Source/Editor/BlueprintGraph/Private/K2Node_GetDataTableRow.cpp b/Engine/Source/Editor/BlueprintGraph/Private/K2Node_GetDataTableRow.cpp
--- a/Engine/Source/Editor/BlueprintGraph/Private/K2Node_GetDataTableRow.cpp
+++ b/Engine/Source/Editor/BlueprintGraph/Private/K2Node_GetDataTableRow.cpp
@@ -4,6 +4,7 @@
 #include "KismetCompiler.h"
 //#include "Runtime/Engine/Classes/Kismet/GameplayStatics.h"
 #include "Kismet/DataTableFunctionLibrary.h"
+#include "EditorCategoryUtils.h"
 
 #define LOCTEXT_NAMESPACE "K2Node_GetDataTableRow"
 
@@ -104,6 +105,19 @@ void UK2Node_GetDataTableRow::ReallocatePinsDuringReconstruction(TArray<UEdGraph
 	}
 }
 
+void UK2Node_GetDataTableRow::GetMenuActions(TArray<UBlueprintNodeSpawner*>& ActionListOut) const
+{
+	UBlueprintNodeSpawner* NodeSpawner = UBlueprintNodeSpawner::Create(GetClass());
+	check(NodeSpawner != nullptr);
+
+	ActionListOut.Add(NodeSpawner);
+}
+
+FText UK2Node_GetDataTableRow::GetMenuCategory() const
+{
+	return FEditorCategoryUtils::GetCommonCategory(FCommonEditorCategory::Utilities);
+}
+
 bool UK2Node_GetDataTableRow::IsDataTablePin(UEdGraphPin* Pin)
 {
 	const UEdGraphSchema_K2* K2Schema = GetDefault<UEdGraphSchema_K2>();
@@ -215,7 +229,7 @@ FText UK2Node_GetDataTableRow::GetNodeTitle(ENodeTitleType::Type TitleType) cons
 	UEdGraphPin* DataTablePin = GetDataTablePin();
 
 	FText DataTableString = NSLOCTEXT("K2Node", "None", "NONE");
-	if(DataTablePin != NULL)
+	if (DataTablePin != NULL)
 	{
 		if(DataTablePin->LinkedTo.Num() > 0)
 		{
@@ -228,9 +242,16 @@ FText UK2Node_GetDataTableRow::GetNodeTitle(ENodeTitleType::Type TitleType) cons
 		}
 	}
 
+	FText LocFormat = NSLOCTEXT("K2Node", "DataTable", "Get Data Table Row {DataTableName}");
+	if (TitleType == ENodeTitleType::ListView)
+	{
+		LocFormat = LOCTEXT("ListViewTitle", "Get Data Table Row");
+	}
+
+
 	FFormatNamedArguments Args;
 	Args.Add(TEXT("DataTableName"), DataTableString);
-	return FText::Format(NSLOCTEXT("K2Node", "DataTable", "Get Data Table Row {DataTableName}"), Args);
+	return FText::Format(LocFormat, Args);
 }
 
 void UK2Node_GetDataTableRow::ExpandNode(class FKismetCompilerContext& CompilerContext, UEdGraph* SourceGraph)
diff --git a/Engine/Source/Editor/BlueprintGraph/Private/K2Node_MakeArray.cpp b/Engine/Source/Editor/BlueprintGraph/Private/K2Node_MakeArray.cpp
--- a/Engine/Source/Editor/BlueprintGraph/Private/K2Node_MakeArray.cpp
+++ b/Engine/Source/Editor/BlueprintGraph/Private/K2Node_MakeArray.cpp
@@ -6,8 +6,9 @@
 #include "Slate.h"
 #include "../../../Runtime/Engine/Classes/Kismet/KismetArrayLibrary.h"
 #include "ScopedTransaction.h"
-
 #include "KismetCompiler.h"
+#include "BlueprintNodeSpawner.h"
+#include "EditorCategoryUtils.h"
 
 static const FString OutputPinName = FString(TEXT("Array"));
 
@@ -402,4 +403,17 @@ bool UK2Node_MakeArray::IsConnectionDisallowed(const UEdGraphPin* MyPin, const U
 	return false;
 }
 
+void UK2Node_MakeArray::GetMenuActions(TArray<UBlueprintNodeSpawner*>& ActionListOut) const
+{
+	UBlueprintNodeSpawner* NodeSpawner = UBlueprintNodeSpawner::Create(GetClass());
+	check(NodeSpawner != nullptr);
+
+	ActionListOut.Add(NodeSpawner);
+}
+
+FText UK2Node_MakeArray::GetMenuCategory() const
+{
+	return FEditorCategoryUtils::BuildCategoryString(FCommonEditorCategory::Utilities, LOCTEXT("ActionMenuCategory", "Array"));
+}
+
 #undef LOCTEXT_NAMESPACE
diff --git a/Engine/Source/Editor/BlueprintGraph/Private/K2Node_MathExpression.cpp b/Engine/Source/Editor/BlueprintGraph/Private/K2Node_MathExpression.cpp
--- a/Engine/Source/Editor/BlueprintGraph/Private/K2Node_MathExpression.cpp
+++ b/Engine/Source/Editor/BlueprintGraph/Private/K2Node_MathExpression.cpp
@@ -2280,6 +2280,15 @@ void UK2Node_MathExpression::GetMenuEntries(FGraphContextMenuBuilder& ContextMen
 	NodeAction->NodeTemplate = TemplateNode;
 }
 
+//------------------------------------------------------------------------------
+void UK2Node_MathExpression::GetMenuActions(TArray<UBlueprintNodeSpawner*>& ActionListOut) const
+{
+	UBlueprintNodeSpawner* NodeSpawner = UBlueprintNodeSpawner::Create(GetClass());
+	check(NodeSpawner != nullptr);
+	
+	ActionListOut.Add(NodeSpawner);
+}
+
 //------------------------------------------------------------------------------
 TSharedPtr<class INameValidatorInterface> UK2Node_MathExpression::MakeNameValidator() const
 {
@@ -2410,7 +2419,7 @@ void UK2Node_MathExpression::ValidateNodeDuringCompilation(FCompilerResultsLog&
 //------------------------------------------------------------------------------
 FText UK2Node_MathExpression::GetNodeTitle(ENodeTitleType::Type TitleType) const
 {
-	FString Result = Expression;
+	FText Result = FText::FromString(Expression);
 	if (Expression.IsEmpty())
 	{
 		//Result = LOCTEXT("NoExpressionTitle", "[Empty]").ToString();
@@ -2418,9 +2427,13 @@ FText UK2Node_MathExpression::GetNodeTitle(ENodeTitleType::Type TitleType) const
 	
 	if (TitleType == ENodeTitleType::FullTitle)
 	{
-		Result = FText::Format(LOCTEXT("MathExpressionSecondTitleLine", "{0}\nMath Expression"), FText::FromString(Result)).ToString();
+		Result = FText::Format(LOCTEXT("MathExpressionSecondTitleLine", "{0}\nMath Expression"), Result);
+	}
+	else if ((TitleType == ENodeTitleType::ListView) && Expression.IsEmpty())
+	{
+		Result = LOCTEXT("AddMathExprMenuOption", "Add Math Expression...");
 	}
-	return FText::FromString(Result);
+	return Result;
 }
 
 //------------------------------------------------------------------------------
diff --git a/Engine/Source/Editor/BlueprintGraph/Private/K2Node_SpawnActorFromClass.cpp b/Engine/Source/Editor/BlueprintGraph/Private/K2Node_SpawnActorFromClass.cpp
--- a/Engine/Source/Editor/BlueprintGraph/Private/K2Node_SpawnActorFromClass.cpp
+++ b/Engine/Source/Editor/BlueprintGraph/Private/K2Node_SpawnActorFromClass.cpp
@@ -3,6 +3,8 @@
 #include "BlueprintGraphPrivatePCH.h"
 #include "KismetCompiler.h"
 #include "Runtime/Engine/Classes/Kismet/GameplayStatics.h"
+#include "BlueprintNodeSpawner.h"
+#include "EditorCategoryUtils.h"
 
 struct FK2Node_SpawnActorFromClassHelper
 {
@@ -260,25 +262,30 @@ FLinearColor UK2Node_SpawnActorFromClass::GetNodeTitleColor() const
 
 FText UK2Node_SpawnActorFromClass::GetNodeTitle(ENodeTitleType::Type TitleType) const
 {
-	UEdGraphPin* ClassPin = GetClassPin();
-
-	FText SpawnString = NSLOCTEXT("K2Node", "None", "NONE");
-	if(ClassPin != NULL)
+	FText NodeTitle = NSLOCTEXT("K2Node", "SpawnActor_BaseTitle", "Spawn Actor from Class");
+	if (TitleType != ENodeTitleType::ListView)
 	{
-		if(ClassPin->LinkedTo.Num() > 0)
-		{
-			// Blueprint will be determined dynamically, so we don't have the name in this case
-			SpawnString = FText::GetEmpty();
-		}
-		else if(ClassPin->DefaultObject != NULL)
+		FText SpawnString = NSLOCTEXT("K2Node", "None", "NONE");
+		if (UEdGraphPin* ClassPin = GetClassPin())
 		{
-			SpawnString = FText::FromString(ClassPin->DefaultObject->GetName());
+			if(ClassPin->LinkedTo.Num() > 0)
+			{
+				// Blueprint will be determined dynamically, so we don't have the name in this case
+				SpawnString = FText::GetEmpty();
+			}
+			else if(ClassPin->DefaultObject != NULL)
+			{
+				SpawnString = FText::FromString(ClassPin->DefaultObject->GetName());
+			}
 		}
+		
+		FFormatNamedArguments Args;
+		Args.Add(TEXT("ClassName"), SpawnString);
+		
+		NodeTitle = FText::Format(NSLOCTEXT("K2Node", "SpawnActor", "SpawnActor {ClassName}"), Args);
 	}
 
-	FFormatNamedArguments Args;
-	Args.Add(TEXT("ClassName"), SpawnString);
-	return FText::Format(NSLOCTEXT("K2Node", "SpawnActor", "SpawnActor {ClassName}"), Args);
+	return NodeTitle;
 }
 
 bool UK2Node_SpawnActorFromClass::CanPasteHere( const UEdGraph* TargetGraph, const UEdGraphSchema* Schema ) const 
@@ -297,6 +304,19 @@ void UK2Node_SpawnActorFromClass::GetNodeAttributes( TArray<TKeyValuePair<FStrin
 	OutNodeAttributes.Add( TKeyValuePair<FString, FString>( TEXT( "ActorClass" ), ClassToSpawnStr ));
 }
 
+void UK2Node_SpawnActorFromClass::GetMenuActions(TArray<UBlueprintNodeSpawner*>& ActionListOut) const
+{
+	UBlueprintNodeSpawner* NodeSpawner = UBlueprintNodeSpawner::Create(GetClass());
+	check(NodeSpawner != nullptr);
+
+	ActionListOut.Add(NodeSpawner);
+}
+
+FText UK2Node_SpawnActorFromClass::GetMenuCategory() const
+{
+	return FEditorCategoryUtils::GetCommonCategory(FCommonEditorCategory::Gameplay);
+}
+
 void UK2Node_SpawnActorFromClass::ExpandNode(class FKismetCompilerContext& CompilerContext, UEdGraph* SourceGraph)
 {
 	Super::ExpandNode(CompilerContext, SourceGraph);
diff --git a/Engine/Source/Editor/BlueprintGraph/Private/K2Node_Timeline.cpp b/Engine/Source/Editor/BlueprintGraph/Private/K2Node_Timeline.cpp
--- a/Engine/Source/Editor/BlueprintGraph/Private/K2Node_Timeline.cpp
+++ b/Engine/Source/Editor/BlueprintGraph/Private/K2Node_Timeline.cpp
@@ -4,8 +4,8 @@
 #include "BlueprintGraphPrivatePCH.h"
 #include "Editor/GraphEditor/Public/DiffResults.h"
 #include "Kismet2NameValidators.h"
-
 #include "KismetCompiler.h"
+#include "BlueprintNodeSpawner.h"
 
 #define LOCTEXT_NAMESPACE "K2Node_Timeline"
 
@@ -262,7 +262,20 @@ FLinearColor UK2Node_Timeline::GetNodeTitleColor() const
 
 FText UK2Node_Timeline::GetNodeTitle(ENodeTitleType::Type TitleType) const
 {
-	return FText::FromName(TimelineName);
+	FText Title = FText::FromName(TimelineName);
+	
+	UBlueprint* Blueprint = GetBlueprint();
+	check(Blueprint != nullptr);
+	
+	UTimelineTemplate* Timeline = Blueprint->FindTimelineTemplateByVariableName(TimelineName);
+	// if a node hasn't been spawned for this node yet, then lets title it
+	// after what it will do (the name would be invalid anyways)
+	if (Timeline == nullptr)
+	{
+		// if this node hasn't spawned a
+		Title = LOCTEXT("NoTimelineTitle", "Add Timeline...");
+	}
+	return Title;
 }
 
 UEdGraphPin* UK2Node_Timeline::GetDirectionPin() const
@@ -553,4 +566,30 @@ void UK2Node_Timeline::GetNodeAttributes( TArray<TKeyValuePair<FString, FString>
 	OutNodeAttributes.Add( TKeyValuePair<FString, FString>( TEXT( "Name" ), GetName() ));
 }
 
+void UK2Node_Timeline::GetMenuActions(TArray<UBlueprintNodeSpawner*>& ActionListOut) const
+{
+	UBlueprintNodeSpawner* NodeSpawner = UBlueprintNodeSpawner::Create(GetClass());
+	check(NodeSpawner != nullptr);
+
+	auto CustomizeTimelineNodeLambda = [](UEdGraphNode* NewNode, bool bIsTemplateNode)
+	{
+		UK2Node_Timeline* TimelineNode = CastChecked<UK2Node_Timeline>(NewNode);
+
+		UBlueprint* Blueprint = TimelineNode->GetBlueprint();
+		if (Blueprint != nullptr)
+		{
+			TimelineNode->TimelineName = FBlueprintEditorUtils::FindUniqueTimelineName(Blueprint);
+			if (!bIsTemplateNode && FBlueprintEditorUtils::AddNewTimeline(Blueprint, TimelineNode->TimelineName))
+			{
+				// clear off any existing error message now that the timeline has been added
+				TimelineNode->ErrorMsg.Empty();
+				TimelineNode->bHasCompilerMessage = false;
+			}
+		}
+	};
+
+	NodeSpawner->CustomizeNodeDelegate = UBlueprintNodeSpawner::FCustomizeNodeDelegate::CreateStatic(CustomizeTimelineNodeLambda);
+	ActionListOut.Add(NodeSpawner);
+}
+
 #undef LOCTEXT_NAMESPACE
diff --git a/Engine/Source/Editor/GameplayAbilitiesEditor/Classes/K2Node_GameplayCueEvent.h b/Engine/Source/Editor/GameplayAbilitiesEditor/Classes/K2Node_GameplayCueEvent.h
--- a/Engine/Source/Editor/GameplayAbilitiesEditor/Classes/K2Node_GameplayCueEvent.h
+++ b/Engine/Source/Editor/GameplayAbilitiesEditor/Classes/K2Node_GameplayCueEvent.h
@@ -13,9 +13,11 @@ class UK2Node_GameplayCueEvent : public UK2Node_Event
 	// Begin UEdGraphNode interface
 	virtual FString GetTooltip() const override;
 	virtual FText GetNodeTitle(ENodeTitleType::Type TitleType) const override;
-	// End UEdGraphNode interface
-
-	virtual FString GetCategoryName();
-
 	virtual void GetMenuEntries(struct FGraphContextMenuBuilder& ContextMenuBuilder) const override;
+	// End UEdGraphNode interface
+	
+	// Begin UK2Node interface
+	virtual void GetMenuActions(TArray<UBlueprintNodeSpawner*>& ActionListOut) const override;
+	virtual FText GetMenuCategory() const override;
+	// End UK2Node interface
 };
\ No newline at end of file
diff --git a/Engine/Source/Editor/GameplayAbilitiesEditor/Private/K2Node_GameplayCueEvent.cpp b/Engine/Source/Editor/GameplayAbilitiesEditor/Private/K2Node_GameplayCueEvent.cpp
--- a/Engine/Source/Editor/GameplayAbilitiesEditor/Private/K2Node_GameplayCueEvent.cpp
+++ b/Engine/Source/Editor/GameplayAbilitiesEditor/Private/K2Node_GameplayCueEvent.cpp
@@ -7,8 +7,10 @@
 #include "K2ActionMenuBuilder.h"
 #include "GameplayTagsModule.h"
 #include "GameplayCueInterface.h"
+#include "BlueprintNodeSpawner.h"
+#include "BlueprintActionMenuBuilder.h" // for AddEventCategory
 
-#define LOCTEXT_NAMESPACE "K2Node_PlayMontageAndWait"
+#define LOCTEXT_NAMESPACE "K2Node_GameplayCueEvent"
 
 UK2Node_GameplayCueEvent::UK2Node_GameplayCueEvent(const class FPostConstructInitializeProperties& PCIP)
 	: Super(PCIP)
@@ -17,11 +19,6 @@ UK2Node_GameplayCueEvent::UK2Node_GameplayCueEvent(const class FPostConstructIni
 	EventSignatureClass = UGameplayCueInterface::StaticClass();
 }
 
-FString UK2Node_GameplayCueEvent::GetCategoryName()
-{
-	return TEXT("GameplayCueEvent");
-}
-
 FString UK2Node_GameplayCueEvent::GetTooltip() const
 {
 	return TEXT("Handle GameplayCue Event");
@@ -47,7 +44,7 @@ void UK2Node_GameplayCueEvent::GetMenuEntries(FGraphContextMenuBuilder& Context)
 		}
 	}
 
-	if (!bValidClass)
+	if (!bValidClass) 
 	{
 		return;
 	}
@@ -85,4 +82,32 @@ void UK2Node_GameplayCueEvent::GetMenuEntries(FGraphContextMenuBuilder& Context)
 	}	
 }
 
+void UK2Node_GameplayCueEvent::GetMenuActions(TArray<UBlueprintNodeSpawner*>& ActionListOut) const
+{
+	auto CustomizeCueNodeLambda = [](UEdGraphNode* NewNode, bool bIsTemplateNode, FName TagName)
+	{
+		UK2Node_GameplayCueEvent* EventNode = CastChecked<UK2Node_GameplayCueEvent>(NewNode);
+		EventNode->CustomFunctionName = TagName;
+	};
+	
+	IGameplayTagsModule& GameplayTagsModule = IGameplayTagsModule::Get();
+	FGameplayTag RootTag = GameplayTagsModule.GetGameplayTagsManager().RequestGameplayTag(FName(TEXT("GameplayCue")));
+	
+	FGameplayTagContainer CueTags = GameplayTagsModule.GetGameplayTagsManager().RequestGameplayTagChildren(RootTag);
+	for (auto TagIt = CueTags.CreateConstIterator(); TagIt; ++TagIt)
+	{
+		UBlueprintNodeSpawner::FCustomizeNodeDelegate PostSpawnDelegate = UBlueprintNodeSpawner::FCustomizeNodeDelegate::CreateStatic(CustomizeCueNodeLambda, TagIt->GetTagName());
+		
+		UBlueprintNodeSpawner* NodeSpawner = UBlueprintNodeSpawner::Create(GetClass(), PostSpawnDelegate);
+		check(NodeSpawner != nullptr);
+		
+		ActionListOut.Add(NodeSpawner);
+	}
+}
+
+FText UK2Node_GameplayCueEvent::GetMenuCategory() const
+{
+	return FText::Format(LOCTEXT("ActionMenuCategory", "{0}|GameplayCue"), FBlueprintActionMenuBuilder::AddEventCategory);
+}
+
 #undef LOCTEXT_NAMESPACE
diff --git a/Engine/Source/Editor/GameplayTagsEditor/Classes/GameplayTagsK2Node_LiteralGameplayTag.h b/Engine/Source/Editor/GameplayTagsEditor/Classes/GameplayTagsK2Node_LiteralGameplayTag.h
--- a/Engine/Source/Editor/GameplayTagsEditor/Classes/GameplayTagsK2Node_LiteralGameplayTag.h
+++ b/Engine/Source/Editor/GameplayTagsEditor/Classes/GameplayTagsK2Node_LiteralGameplayTag.h
@@ -25,6 +25,8 @@ class UGameplayTagsK2Node_LiteralGameplayTag : public UK2Node
 	virtual bool ShouldShowNodeProperties() const override { return true; }
 	virtual bool IsNodeSafeToIgnore() const override { return true; }
 	virtual void ExpandNode(class FKismetCompilerContext& CompilerContext, UEdGraph* SourceGraph) override;
+	virtual void GetMenuActions(TArray<UBlueprintNodeSpawner*>& ActionListOut) const override;
+	virtual FText GetMenuCategory() const override;
 	// End UK2Node interface
 #endif
 
diff --git a/Engine/Source/Editor/GameplayTagsEditor/Private/GameplayTagsK2Node_LiteralGameplayTag.cpp b/Engine/Source/Editor/GameplayTagsEditor/Private/GameplayTagsK2Node_LiteralGameplayTag.cpp
--- a/Engine/Source/Editor/GameplayTagsEditor/Private/GameplayTagsK2Node_LiteralGameplayTag.cpp
+++ b/Engine/Source/Editor/GameplayTagsEditor/Private/GameplayTagsK2Node_LiteralGameplayTag.cpp
@@ -4,6 +4,10 @@
 #include "KismetCompiler.h"
 #include "GameplayTags.h"
 #include "K2ActionMenuBuilder.h" // for FK2ActionMenuBuilder::AddNewNodeAction()
+#include "BlueprintNodeSpawner.h"
+#include "EditorCategoryUtils.h"
+
+#define LOCTEXT_NAMESPACE "GameplayTagsK2Node_LiteralGameplayTag"
 
 UGameplayTagsK2Node_LiteralGameplayTag::UGameplayTagsK2Node_LiteralGameplayTag(const class FPostConstructInitializeProperties& PCIP)
 	: Super(PCIP)
@@ -143,3 +147,18 @@ void UGameplayTagsK2Node_LiteralGameplayTag::ExpandNode(class FKismetCompilerCon
 		BreakAllNodeLinks();
 	}
 }
+
+void UGameplayTagsK2Node_LiteralGameplayTag::GetMenuActions(TArray<UBlueprintNodeSpawner*>& ActionListOut) const
+{
+	UBlueprintNodeSpawner* NodeSpawner = UBlueprintNodeSpawner::Create(GetClass());
+	check(NodeSpawner != nullptr);
+	
+	ActionListOut.Add(NodeSpawner);
+}
+
+FText UGameplayTagsK2Node_LiteralGameplayTag::GetMenuCategory() const
+{
+	return LOCTEXT("ActionMenuCategory", "Gameplay Tags");
+}
+
+#undef LOCTEXT_NAMESPACE
