(WIP) Various node types, defining their own action menu category (for the new blueprint menu system).

diff --git a/Engine/Source/Editor/BlueprintGraph/Classes/K2Node_CallFunction.h b/Engine/Source/Editor/BlueprintGraph/Classes/K2Node_CallFunction.h
--- a/Engine/Source/Editor/BlueprintGraph/Classes/K2Node_CallFunction.h
+++ b/Engine/Source/Editor/BlueprintGraph/Classes/K2Node_CallFunction.h
@@ -88,6 +88,7 @@ public:
 	virtual FName GetCornerIcon() const override;
 	virtual FText GetToolTipHeading() const override;
 	virtual void GetNodeAttributes( TArray<TKeyValuePair<FString, FString>>& OutNodeAttributes ) const override;
+	virtual FText GetMenuCategory() const override;
 	// End of UK2Node interface
 
 	/** Returns the UFunction that this class is pointing to */
diff --git a/Engine/Source/Editor/BlueprintGraph/Classes/K2Node_DynamicCast.h b/Engine/Source/Editor/BlueprintGraph/Classes/K2Node_DynamicCast.h
--- a/Engine/Source/Editor/BlueprintGraph/Classes/K2Node_DynamicCast.h
+++ b/Engine/Source/Editor/BlueprintGraph/Classes/K2Node_DynamicCast.h
@@ -25,6 +25,7 @@ class UK2Node_DynamicCast : public UK2Node
 	virtual ERedirectType DoPinsMatchForReconstruction(const UEdGraphPin* NewPin, int32 NewPinIndex, const UEdGraphPin* OldPin, int32 OldPinIndex) const override;
 	virtual class FNodeHandlingFunctor* CreateNodeHandler(class FKismetCompilerContext& CompilerContext) const override;
 	virtual bool HasExternalBlueprintDependencies(TArray<class UStruct*>* OptionalOutput) const override;
+	virtual FText GetMenuCategory() const override;
 	// End of UK2Node interface
 
 	/** Get the 'valid cast' exec pin */
diff --git a/Engine/Source/Editor/BlueprintGraph/Classes/K2Node_Event.h b/Engine/Source/Editor/BlueprintGraph/Classes/K2Node_Event.h
--- a/Engine/Source/Editor/BlueprintGraph/Classes/K2Node_Event.h
+++ b/Engine/Source/Editor/BlueprintGraph/Classes/K2Node_Event.h
@@ -62,6 +62,7 @@ class UK2Node_Event : public UK2Node_EditablePinBase
 	BLUEPRINTGRAPH_API virtual void ExpandNode(class FKismetCompilerContext& CompilerContext, UEdGraph* SourceGraph) override;
 	BLUEPRINTGRAPH_API virtual FText GetToolTipHeading() const;
 	BLUEPRINTGRAPH_API virtual void GetNodeAttributes( TArray<TKeyValuePair<FString, FString>>& OutNodeAttributes ) const override;
+	virtual FText GetMenuCategory() const override;
 	// End UK2Node interface
 
 	/** Checks whether the parameters for this event node are compatible with the specified function entry node */
diff --git a/Engine/Source/Editor/BlueprintGraph/Classes/K2Node_MacroInstance.h b/Engine/Source/Editor/BlueprintGraph/Classes/K2Node_MacroInstance.h
--- a/Engine/Source/Editor/BlueprintGraph/Classes/K2Node_MacroInstance.h
+++ b/Engine/Source/Editor/BlueprintGraph/Classes/K2Node_MacroInstance.h
@@ -55,6 +55,7 @@ public:
 	virtual FText GetActiveBreakpointToolTipText() const override;
 	virtual bool HasExternalBlueprintDependencies(TArray<class UStruct*>* OptionalOutput) const override;
 	virtual void GetNodeAttributes( TArray<TKeyValuePair<FString, FString>>& OutNodeAttributes ) const override;
+	virtual FText GetMenuCategory() const override;
 	// End UK2Node interface
 
 	void SetMacroGraph(UEdGraph* Graph) { MacroGraphReference.SetGraph(Graph); }
diff --git a/Engine/Source/Editor/BlueprintGraph/Private/K2Node_CallFunction.cpp b/Engine/Source/Editor/BlueprintGraph/Private/K2Node_CallFunction.cpp
--- a/Engine/Source/Editor/BlueprintGraph/Private/K2Node_CallFunction.cpp
+++ b/Engine/Source/Editor/BlueprintGraph/Private/K2Node_CallFunction.cpp
@@ -1502,6 +1502,12 @@ void UK2Node_CallFunction::GetNodeAttributes( TArray<TKeyValuePair<FString, FStr
 	OutNodeAttributes.Add( TKeyValuePair<FString, FString>( TEXT( "Name" ), TargetFunctionName ));
 }
 
+FText UK2Node_CallFunction::GetMenuCategory() const
+{
+	UFunction* TargetFunction = GetTargetFunction();
+	return FText::FromString(GetDefaultCategoryForFunction(TargetFunction, TEXT("")));
+}
+
 bool UK2Node_CallFunction::HasExternalBlueprintDependencies(TArray<class UStruct*>* OptionalOutput) const
 {
 	const UClass* SourceClass = FunctionReference.GetMemberParentClass(this);
diff --git a/Engine/Source/Editor/BlueprintGraph/Private/K2Node_DynamicCast.cpp b/Engine/Source/Editor/BlueprintGraph/Private/K2Node_DynamicCast.cpp
--- a/Engine/Source/Editor/BlueprintGraph/Private/K2Node_DynamicCast.cpp
+++ b/Engine/Source/Editor/BlueprintGraph/Private/K2Node_DynamicCast.cpp
@@ -3,6 +3,9 @@
 
 #include "BlueprintGraphPrivatePCH.h"
 #include "DynamicCastHandler.h"
+#include "EditorCategoryUtils.h"
+
+#define LOCTEXT_NAMESPACE "K2Node_DynamicCast"
 
 UK2Node_DynamicCast::UK2Node_DynamicCast(const class FPostConstructInitializeProperties& PCIP)
 	: Super(PCIP)
@@ -158,3 +161,10 @@ bool UK2Node_DynamicCast::HasExternalBlueprintDependencies(TArray<class UStruct*
 	}
 	return bResult || Super::HasExternalBlueprintDependencies(OptionalOutput);
 }
+
+FText UK2Node_DynamicCast::GetMenuCategory() const
+{
+	return FEditorCategoryUtils::BuildCategoryString(FCommonEditorCategory::Utilities, LOCTEXT("ActionMenuCategory", "Casting"));
+}
+
+#undef LOCTEXT_NAMESPACE
diff --git a/Engine/Source/Editor/BlueprintGraph/Private/K2Node_Event.cpp b/Engine/Source/Editor/BlueprintGraph/Private/K2Node_Event.cpp
--- a/Engine/Source/Editor/BlueprintGraph/Private/K2Node_Event.cpp
+++ b/Engine/Source/Editor/BlueprintGraph/Private/K2Node_Event.cpp
@@ -2,10 +2,10 @@
 
 
 #include "BlueprintGraphPrivatePCH.h"
-
 #include "CompilerResultsLog.h"
 #include "KismetCompiler.h"
 #include "EventEntryHandler.h"
+#include "BlueprintActionMenuBuilder.h"
 
 const FString UK2Node_Event::DelegateOutputName(TEXT("OutputDelegate"));
 
@@ -629,6 +629,17 @@ void UK2Node_Event::GetNodeAttributes( TArray<TKeyValuePair<FString, FString>>&
 	OutNodeAttributes.Add( TKeyValuePair<FString, FString>( TEXT( "Name" ), GetFunctionName().ToString() ));
 }
 
+FText UK2Node_Event::GetMenuCategory() const
+{
+	FText Category;
+	if (UFunction* Function = FindField<UFunction>(EventSignatureClass, EventSignatureName))
+	{
+		Category = FText::FromString(UK2Node_CallFunction::GetDefaultCategoryForFunction(Function, TEXT("")));
+	}
+
+	return FText::Format(LOCTEXT("ConcatenatedCategory", "{0}|{1}"), FBlueprintActionMenuBuilder::AddEventCategory, Category);
+}
+
 bool UK2Node_Event::IsDeprecated() const
 {
 	if (UFunction* Function = FindField<UFunction>(EventSignatureClass, EventSignatureName))
diff --git a/Engine/Source/Editor/BlueprintGraph/Private/K2Node_MacroInstance.cpp b/Engine/Source/Editor/BlueprintGraph/Private/K2Node_MacroInstance.cpp
--- a/Engine/Source/Editor/BlueprintGraph/Private/K2Node_MacroInstance.cpp
+++ b/Engine/Source/Editor/BlueprintGraph/Private/K2Node_MacroInstance.cpp
@@ -1,8 +1,8 @@
 // Copyright 1998-2014 Epic Games, Inc. All Rights Reserved.
 
 #include "BlueprintGraphPrivatePCH.h"
-
 #include "Slate.h"
+#include "EditorCategoryUtils.h"
 
 #define LOCTEXT_NAMESPACE "K2Node_MacroInstance"
 
@@ -427,5 +427,20 @@ void UK2Node_MacroInstance::GetNodeAttributes( TArray<TKeyValuePair<FString, FSt
 	OutNodeAttributes.Add( TKeyValuePair<FString, FString>( TEXT( "Name" ), MacroName ));
 }
 
+FText UK2Node_MacroInstance::GetMenuCategory() const
+{
+	FText MenuCategory = FEditorCategoryUtils::GetCommonCategory(FCommonEditorCategory::Utilities);
+	if (UEdGraph* MacroGraph = GetMacroGraph())
+	{
+		FKismetUserDeclaredFunctionMetadata* MacroGraphMetadata = UK2Node_MacroInstance::GetAssociatedGraphMetadata(MacroGraph);
+		if ((MacroGraphMetadata != nullptr) && !MacroGraphMetadata->Category.IsEmpty())
+		{
+			MenuCategory = FText::FromString(MacroGraphMetadata->Category);
+		}
+	}
+
+	return MenuCategory;
+}
+
 
 #undef LOCTEXT_NAMESPACE
