Optimizing the event node spawner so the new blueprint menu builder doesn't have to spawn temp nodes to pull data from.

diff --git a/Engine/Source/Editor/BlueprintGraph/Classes/BlueprintEventNodeSpawner.h b/Engine/Source/Editor/BlueprintGraph/Classes/BlueprintEventNodeSpawner.h
--- a/Engine/Source/Editor/BlueprintGraph/Classes/BlueprintEventNodeSpawner.h
+++ b/Engine/Source/Editor/BlueprintGraph/Classes/BlueprintEventNodeSpawner.h
@@ -30,8 +30,10 @@ public:
 
 	// UBlueprintNodeSpawner interface
 	virtual UEdGraphNode* Invoke(UEdGraph* ParentGraph) const override;
-	virtual FText    GetDefaultMenuName() const override;
-	virtual FText    GetDefaultMenuCategory() const override;
+	virtual FText GetDefaultMenuName() const override;
+	virtual FText GetDefaultMenuCategory() const override;
+	virtual FText GetDefaultMenuTooltip() const override;
+	virtual FText GetDefaultSearchKeywords() const override;
 	// End UBlueprintNodeSpawner interface
 	
 	/**
diff --git a/Engine/Source/Editor/BlueprintGraph/Private/BlueprintEventNodeSpawner.cpp b/Engine/Source/Editor/BlueprintGraph/Private/BlueprintEventNodeSpawner.cpp
--- a/Engine/Source/Editor/BlueprintGraph/Private/BlueprintEventNodeSpawner.cpp
+++ b/Engine/Source/Editor/BlueprintGraph/Private/BlueprintEventNodeSpawner.cpp
@@ -2,7 +2,8 @@
 
 #include "BlueprintGraphPrivatePCH.h"
 #include "BlueprintEventNodeSpawner.h"
-#include "EdGraphSchema_K2.h"	// for GetFriendlySignitureName()
+#include "EdGraphSchema_K2.h"			// for GetFriendlySignitureName()
+#include "BlueprintActionMenuBuilder.h"	// for AddEventCategory
 
 #define LOCTEXT_NAMESPACE "BlueprintEventNodeSpawner"
 
@@ -75,8 +76,24 @@ FText UBlueprintEventNodeSpawner::GetDefaultMenuName() const
 //------------------------------------------------------------------------------
 FText UBlueprintEventNodeSpawner::GetDefaultMenuCategory() const
 {
-	// @TODO: 
-	return Super::GetDefaultMenuCategory();
+	check(EventFunc != nullptr);
+	
+	FString RootCategory = FBlueprintActionMenuBuilder::AddEventCategory.ToString();
+	return FText::FromString(UK2Node_CallFunction::GetDefaultCategoryForFunction(EventFunc, RootCategory));
+}
+
+//------------------------------------------------------------------------------
+FText UBlueprintEventNodeSpawner::GetDefaultMenuTooltip() const
+{
+	check(EventFunc != nullptr);
+	return FText::FromString(UK2Node_CallFunction::GetDefaultTooltipForFunction(EventFunc));
+}
+
+//------------------------------------------------------------------------------
+FText UBlueprintEventNodeSpawner::GetDefaultSearchKeywords() const
+{
+	check(EventFunc != nullptr);
+	return FText::FromString(UK2Node_CallFunction::GetKeywordsForFunction(EventFunc));
 }
 
 //------------------------------------------------------------------------------
