Fixing broken builder, missed some includes for non-unity builds.

diff --git a/Engine/Source/Editor/BlueprintGraph/Private/K2Node_GetInputAxisKeyValue.cpp b/Engine/Source/Editor/BlueprintGraph/Private/K2Node_GetInputAxisKeyValue.cpp
--- a/Engine/Source/Editor/BlueprintGraph/Private/K2Node_GetInputAxisKeyValue.cpp
+++ b/Engine/Source/Editor/BlueprintGraph/Private/K2Node_GetInputAxisKeyValue.cpp
@@ -3,6 +3,8 @@
 #include "BlueprintGraphPrivatePCH.h"
 #include "K2Node_GetInputAxisKeyValue.h"
 #include "CompilerResultsLog.h"
+#include "BlueprintNodeSpawner.h"
+#include "EditorCategoryUtils.h"
 
 #define LOCTEXT_NAMESPACE "K2Node_GetInputAxisKeyValue"
 
diff --git a/Engine/Source/Editor/BlueprintGraph/Private/K2Node_TemporaryVariable.cpp b/Engine/Source/Editor/BlueprintGraph/Private/K2Node_TemporaryVariable.cpp
--- a/Engine/Source/Editor/BlueprintGraph/Private/K2Node_TemporaryVariable.cpp
+++ b/Engine/Source/Editor/BlueprintGraph/Private/K2Node_TemporaryVariable.cpp
@@ -3,7 +3,7 @@
 
 #include "BlueprintGraphPrivatePCH.h"
 #include "KismetCompiler.h"
-#include "BlueprintActionFilter.h"
+#include "BlueprintNodeSpawner.h"
 #include "EditorCategoryUtils.h"
 
 class FKCHandler_TemporaryVariable : public FNodeHandlingFunctor
