Fixing broken builder, missed some includes for non-unity builds.

diff --git a/Engine/Source/Editor/BlueprintGraph/Private/K2Node_CustomEvent.cpp b/Engine/Source/Editor/BlueprintGraph/Private/K2Node_CustomEvent.cpp
--- a/Engine/Source/Editor/BlueprintGraph/Private/K2Node_CustomEvent.cpp
+++ b/Engine/Source/Editor/BlueprintGraph/Private/K2Node_CustomEvent.cpp
@@ -5,6 +5,7 @@
 
 #include "Kismet2NameValidators.h"
 #include "CompilerResultsLog.h"
+#include "BlueprintNodeSpawner.h"
 
 #define LOCTEXT_NAMESPACE "K2Node_CustomEvent"
 
diff --git a/Engine/Source/Editor/BlueprintGraph/Private/K2Node_GetDataTableRow.cpp b/Engine/Source/Editor/BlueprintGraph/Private/K2Node_GetDataTableRow.cpp
--- a/Engine/Source/Editor/BlueprintGraph/Private/K2Node_GetDataTableRow.cpp
+++ b/Engine/Source/Editor/BlueprintGraph/Private/K2Node_GetDataTableRow.cpp
@@ -4,6 +4,7 @@
 #include "KismetCompiler.h"
 //#include "Runtime/Engine/Classes/Kismet/GameplayStatics.h"
 #include "Kismet/DataTableFunctionLibrary.h"
+#include "BlueprintNodeSpawner.h"
 #include "EditorCategoryUtils.h"
 
 #define LOCTEXT_NAMESPACE "K2Node_GetDataTableRow"
