UE4: Fixing the build. GetMenuCategory is overridden in another module and needs to be exported to link correctly.

diff --git a/Engine/Source/Editor/BlueprintGraph/Classes/K2Node_Event.h b/Engine/Source/Editor/BlueprintGraph/Classes/K2Node_Event.h
--- a/Engine/Source/Editor/BlueprintGraph/Classes/K2Node_Event.h
+++ b/Engine/Source/Editor/BlueprintGraph/Classes/K2Node_Event.h
@@ -62,7 +62,7 @@ class UK2Node_Event : public UK2Node_EditablePinBase
 	BLUEPRINTGRAPH_API virtual void ExpandNode(class FKismetCompilerContext& CompilerContext, UEdGraph* SourceGraph) override;
 	BLUEPRINTGRAPH_API virtual FText GetToolTipHeading() const;
 	BLUEPRINTGRAPH_API virtual void GetNodeAttributes( TArray<TKeyValuePair<FString, FString>>& OutNodeAttributes ) const override;
-	virtual FText GetMenuCategory() const override;
+	BLUEPRINTGRAPH_API virtual FText GetMenuCategory() const override;
 	// End UK2Node interface
 
 	/** Checks whether the parameters for this event node are compatible with the specified function entry node */
