Fix SplitPin issues when reconstructing a Variable node

diff --git a/Engine/Source/Editor/BlueprintGraph/Private/K2Node_Variable.cpp b/Engine/Source/Editor/BlueprintGraph/Private/K2Node_Variable.cpp
--- a/Engine/Source/Editor/BlueprintGraph/Private/K2Node_Variable.cpp
+++ b/Engine/Source/Editor/BlueprintGraph/Private/K2Node_Variable.cpp
@@ -175,6 +175,45 @@ UK2Node::ERedirectType UK2Node_Variable::DoPinsMatchForReconstruction( const UEd
 	{
 		if (K2Schema->ArePinTypesCompatible(NewPin->PinType, OldPin->PinType))
 		{
+			// If these are split pins, we need to do some name checking logic
+			if (NewPin->ParentPin)
+			{
+				// If the OldPin is not split, then these don't match
+				if (OldPin->ParentPin == nullptr)
+				{
+					return ERedirectType_None;
+				}
+
+				// Go through and find the original variable pin.
+				// If the number of steps out to the original variable pin is not the same then these don't match
+				const UEdGraphPin* ParentmostNewPin = NewPin;
+				const UEdGraphPin* ParentmostOldPin = OldPin;
+
+				while (ParentmostNewPin->ParentPin)
+				{
+					if (ParentmostOldPin->ParentPin == nullptr)
+					{
+						return ERedirectType_None;
+					}
+					ParentmostNewPin = ParentmostNewPin->ParentPin;
+					ParentmostOldPin = ParentmostOldPin->ParentPin;
+				}
+
+				if (ParentmostOldPin->ParentPin)
+				{
+					return ERedirectType_None;
+				}
+
+				// Compare whether the names, ignoring the original variable's name in the case of renames, match
+				FString NewPinPropertyName = NewPin->PinName.RightChop(ParentmostNewPin->PinName.Len() + 1);
+				FString OldPinPropertyName = OldPin->PinName.RightChop(ParentmostOldPin->PinName.Len() + 1);
+
+				if (NewPinPropertyName != OldPinPropertyName)
+				{
+					return ERedirectType_None;
+				}
+			}
+
 			return ERedirectType_Name;
 		}
 		else if ((OldPin->PinName == NewPin->PinName) && ((NewPin->PinType.PinCategory == K2Schema->PC_Object) ||
