Setting up BlueprintEventNodeSpawner to work with custom events.

diff --git a/Engine/Source/Editor/BlueprintGraph/Classes/BlueprintEventNodeSpawner.h b/Engine/Source/Editor/BlueprintGraph/Classes/BlueprintEventNodeSpawner.h
--- a/Engine/Source/Editor/BlueprintGraph/Classes/BlueprintEventNodeSpawner.h
+++ b/Engine/Source/Editor/BlueprintGraph/Classes/BlueprintEventNodeSpawner.h
@@ -5,6 +5,8 @@
 #include "BlueprintNodeSpawner.h"
 #include "BlueprintEventNodeSpawner.generated.h"
 
+class UK2Node_Event;
+
 /**
  * Takes care of spawning UK2Node_Event nodes. Acts as the "action" portion of
  * certain FBlueprintActionMenuItems. Will not spawn a new event node if one
@@ -26,19 +28,28 @@ public:
 	 * @param  EventFunc	The function you want assigned to new nodes.
 	 * @return A newly allocated instance of this class.
 	 */
-	static UBlueprintEventNodeSpawner* Create(UFunction const* const EventFunc);
+	BLUEPRINTGRAPH_API static UBlueprintEventNodeSpawner* Create(UFunction const* const EventFunc);
+
+	/**
+	 * Creates a new UBlueprintEventNodeSpawner for custom events. The 
+	 * CustomEventName can be left blank if the node will pick one itself on
+	 * instantiation.
+	 *
+	 * @param  NodeClass		The event node type that you want this to spawn.
+	 * @param  CustomEventName	The name you want assigned to the event.
+	 * @return A newly allocated instance of this class.
+	 */
+	BLUEPRINTGRAPH_API static UBlueprintEventNodeSpawner* Create(TSubclassOf<UK2Node_Event> NodeClass, FName CustomEventName);
 
 	// UBlueprintNodeSpawner interface
 	virtual UEdGraphNode* Invoke(UEdGraph* ParentGraph) const override;
 	virtual FText GetDefaultMenuName() const override;
-	virtual FText GetDefaultMenuCategory() const override;
-	virtual FText GetDefaultMenuTooltip() const override;
 	virtual FText GetDefaultSearchKeywords() const override;
 	// End UBlueprintNodeSpawner interface
 	
 	/**
 	 * Retrieves the function that this assigns to spawned nodes (defines the
-	 * event's signature).
+	 * event's signature). Can be null if this spawner was for a custom event.
 	 *
 	 * @return The event function that this class was initialized with.
 	 */
@@ -48,4 +59,8 @@ private:
 	/** The function to configure new nodes with. */
 	UPROPERTY()
     UFunction const* EventFunc;
+
+	/** The custom name to configure new event nodes with. */
+	UPROPERTY()
+	FName CustomEventName;
 };
diff --git a/Engine/Source/Editor/BlueprintGraph/Private/BlueprintActionFilter.cpp b/Engine/Source/Editor/BlueprintGraph/Private/BlueprintActionFilter.cpp
--- a/Engine/Source/Editor/BlueprintGraph/Private/BlueprintActionFilter.cpp
+++ b/Engine/Source/Editor/BlueprintGraph/Private/BlueprintActionFilter.cpp
@@ -252,9 +252,10 @@ static UClass* BlueprintActionFilterImpl::GetActionClass(UBlueprintNodeSpawner c
 	}
 	else if (UBlueprintEventNodeSpawner const* EventSpawner = Cast<UBlueprintEventNodeSpawner>(BlueprintAction))
 	{
-		UFunction const* EventFunc = EventSpawner->GetEventFunction();
-		check(EventFunc != nullptr);
-		ActionClass = EventFunc->GetOuterUClass();
+		if (UFunction const* EventFunc = EventSpawner->GetEventFunction())
+		{
+			ActionClass = EventFunc->GetOuterUClass();
+		}		
 	}
 	else if (UBlueprintPropertyNodeSpawner const* PropertySpawner = Cast<UBlueprintPropertyNodeSpawner>(BlueprintAction))
 	{
@@ -353,31 +354,31 @@ static bool BlueprintActionFilterImpl::IsEventUnimplementable(FBlueprintActionFi
 	
 	if (UBlueprintEventNodeSpawner const* EventSpawner = Cast<UBlueprintEventNodeSpawner>(BlueprintAction))
 	{
-		UFunction const* EventFunc = EventSpawner->GetEventFunction();
-		check(EventFunc != nullptr);
-		
-		UClass* FuncOwner = EventFunc->GetOuterUClass();
-		for (UBlueprint const* Blueprint : FilterContext.Blueprints)
+		if (UFunction const* EventFunc = EventSpawner->GetEventFunction())
 		{
-			UClass* BpClass = Blueprint->GeneratedClass;
-			check(BpClass != nullptr);
-			
-			// if this function belongs directly to this blueprint, then it is
-			// already implemented here (this action however is valid for sub-
-			// classes, as they can override the event's functionality)
-			if (FuncOwner == BpClass)
-			{
-				bIsFilteredOut = true;
-				break;
-			}
-			
-			// you can only implement events that you inherit; so if this
-			// blueprint is not a subclass of the event's owner, then we're not
-			// allowed to implement it
-			if (!BpClass->IsChildOf(FuncOwner))
+			UClass* FuncOwner = EventFunc->GetOuterUClass();
+			for (UBlueprint const* Blueprint : FilterContext.Blueprints)
 			{
-				bIsFilteredOut = true;
-				break;
+				UClass* BpClass = Blueprint->GeneratedClass;
+				check(BpClass != nullptr);
+
+				// if this function belongs directly to this blueprint, then it is
+				// already implemented here (this action however is valid for sub-
+				// classes, as they can override the event's functionality)
+				if (FuncOwner == BpClass)
+				{
+					bIsFilteredOut = true;
+					break;
+				}
+
+				// you can only implement events that you inherit; so if this
+				// blueprint is not a subclass of the event's owner, then we're not
+				// allowed to implement it
+				if (!BpClass->IsChildOf(FuncOwner))
+				{
+					bIsFilteredOut = true;
+					break;
+				}
 			}
 		}
 	}
diff --git a/Engine/Source/Editor/BlueprintGraph/Private/BlueprintActionMenuBuilder.cpp b/Engine/Source/Editor/BlueprintGraph/Private/BlueprintActionMenuBuilder.cpp
--- a/Engine/Source/Editor/BlueprintGraph/Private/BlueprintActionMenuBuilder.cpp
+++ b/Engine/Source/Editor/BlueprintGraph/Private/BlueprintActionMenuBuilder.cpp
@@ -310,7 +310,6 @@ FBlueprintActionMenuBuilderImpl::FMenuSectionDefinition::FMenuSectionDefinition(
  * FBlueprintActionMenuBuilder
  ******************************************************************************/
 
-FText const FBlueprintActionMenuBuilder::AddEventCategory  = LOCTEXT("AddEventCategory",  "Add Event");
 FText const FBlueprintActionMenuBuilder::VariablesCategory = LOCTEXT("VariablesCategory", "Variables");
 
 //------------------------------------------------------------------------------
diff --git a/Engine/Source/Editor/BlueprintGraph/Private/BlueprintEventNodeSpawner.cpp b/Engine/Source/Editor/BlueprintGraph/Private/BlueprintEventNodeSpawner.cpp
--- a/Engine/Source/Editor/BlueprintGraph/Private/BlueprintEventNodeSpawner.cpp
+++ b/Engine/Source/Editor/BlueprintGraph/Private/BlueprintEventNodeSpawner.cpp
@@ -7,6 +7,42 @@
 
 #define LOCTEXT_NAMESPACE "BlueprintEventNodeSpawner"
 
+/*******************************************************************************
+ * Static UBlueprintEventNodeSpawner Helpers
+ ******************************************************************************/
+
+namespace UBlueprintEventNodeSpawnerImpl
+{
+	/**
+	 * Helper function for scanning a blueprint for a certain, custom named,
+	 * event.
+	 * 
+	 * @param  Blueprint	The blueprint you want to look through
+	 * @param  CustomName	The event name you want to check for.
+	 * @return Null if no event was found, otherwise a pointer to the named event.
+	 */
+	static UK2Node_Event* FindCustomEventNode(UBlueprint* Blueprint, FName const CustomName);
+}
+
+//------------------------------------------------------------------------------
+UK2Node_Event* UBlueprintEventNodeSpawnerImpl::FindCustomEventNode(UBlueprint* Blueprint, FName const CustomName)
+{
+	UK2Node_Event* FoundNode = nullptr;
+
+	TArray<UK2Node_Event*> AllEvents;
+	FBlueprintEditorUtils::GetAllNodesOfClass<UK2Node_Event>(Blueprint, AllEvents);
+
+	for (UK2Node_Event* EventNode : AllEvents)
+	{
+		if (EventNode->CustomFunctionName == CustomName)
+		{
+			FoundNode = EventNode;
+			break;
+		}
+	}
+	return FoundNode;
+}
+
 /*******************************************************************************
  * UBlueprintEventNodeSpawner
  ******************************************************************************/
@@ -23,6 +59,16 @@ UBlueprintEventNodeSpawner* UBlueprintEventNodeSpawner::Create(UFunction const*
 	return NodeSpawner;
 }
 
+//------------------------------------------------------------------------------
+UBlueprintEventNodeSpawner* UBlueprintEventNodeSpawner::Create(TSubclassOf<UK2Node_Event> NodeClass, FName CustomEventName)
+{
+	UBlueprintEventNodeSpawner* NodeSpawner = NewObject<UBlueprintEventNodeSpawner>(GetTransientPackage());
+	NodeSpawner->NodeClass       = NodeClass;
+	NodeSpawner->CustomEventName = CustomEventName;
+
+	return NodeSpawner;
+}
+
 //------------------------------------------------------------------------------
 UBlueprintEventNodeSpawner::UBlueprintEventNodeSpawner(class FPostConstructInitializeProperties const& PCIP)
 	: Super(PCIP)
@@ -36,22 +82,45 @@ UEdGraphNode* UBlueprintEventNodeSpawner::Invoke(UEdGraph* ParentGraph) const
 	check(ParentGraph != nullptr);
 	UBlueprint* Blueprint = FBlueprintEditorUtils::FindBlueprintForGraphChecked(ParentGraph);
 
-	check(EventFunc != nullptr);
-
-	UClass* ClassOwner = CastChecked<UClass>(EventFunc->GetOuter());
+	UK2Node_Event* EventNode = nullptr;
+	
+	FName EventName = CustomEventName;
+	if (EventFunc != nullptr)
+	{
+		check(EventFunc != nullptr);
+		EventName = EventFunc->GetFName();
+	}
+	
+	UClass* ClassOwner = Blueprint->GeneratedClass;
 	// look to see if a node for this event already exists (only one node is
 	// allowed per event, per blueprint)
-	UK2Node_Event* EventNode = FBlueprintEditorUtils::FindOverrideForFunction(Blueprint, ClassOwner, EventFunc->GetFName());
+	if (EventFunc != nullptr)
+	{
+		ClassOwner = CastChecked<UClass>(EventFunc->GetOuter());
+		EventNode  = FBlueprintEditorUtils::FindOverrideForFunction(Blueprint, ClassOwner, EventName);
+	}
+	else
+	{
+		EventNode  = UBlueprintEventNodeSpawnerImpl::FindCustomEventNode(Blueprint, EventName);
+	}
 
 	// if there is no existing node, then we can happily spawn one into the graph
 	if (EventNode == nullptr)
 	{
-		EventNode = NewObject<UK2Node_Event>(ParentGraph);
-		EventNode->EventSignatureName  = EventFunc->GetFName();
-		EventNode->EventSignatureClass = ClassOwner;
-		EventNode->bOverrideFunction   = true;
-
 		bool bIsTemplateNode = ParentGraph->HasAnyFlags(RF_Transient);
+
+		EventNode = CastChecked<UK2Node_Event>(Super::Invoke(ParentGraph));
+		if (EventFunc != nullptr)
+		{
+			EventNode->EventSignatureName  = EventName;
+			EventNode->EventSignatureClass = ClassOwner;
+			EventNode->bOverrideFunction   = true;
+		}
+		else if (!bIsTemplateNode)
+		{
+			EventNode->CustomFunctionName  = CustomEventName;
+		}
+
 		if (!bIsTemplateNode)
 		{
 			EventNode->SetFlags(RF_Transactional);
@@ -67,33 +136,32 @@ UEdGraphNode* UBlueprintEventNodeSpawner::Invoke(UEdGraph* ParentGraph) const
 //------------------------------------------------------------------------------
 FText UBlueprintEventNodeSpawner::GetDefaultMenuName() const
 {
-	check(EventFunc != nullptr);
-	
-	FString const EventName = UEdGraphSchema_K2::GetFriendlySignitureName(EventFunc);
-	return FText::Format(LOCTEXT("EventWithSignatureName", "Event {0}"), FText::FromString(EventName));
-}
-
-//------------------------------------------------------------------------------
-FText UBlueprintEventNodeSpawner::GetDefaultMenuCategory() const
-{
-	check(EventFunc != nullptr);
-	
-	FString RootCategory = FBlueprintActionMenuBuilder::AddEventCategory.ToString();
-	return FText::FromString(UK2Node_CallFunction::GetDefaultCategoryForFunction(EventFunc, RootCategory));
-}
+	FText EventName;
+	if (EventFunc != nullptr)
+	{
+		EventName = FText::FromString(UEdGraphSchema_K2::GetFriendlySignitureName(EventFunc));
+	}
+	else if (!CustomEventName.IsNone())
+	{
+		EventName = FText::FromName(CustomEventName);
+	}
 
-//------------------------------------------------------------------------------
-FText UBlueprintEventNodeSpawner::GetDefaultMenuTooltip() const
-{
-	check(EventFunc != nullptr);
-	return FText::FromString(UK2Node_CallFunction::GetDefaultTooltipForFunction(EventFunc));
+	if (!EventName.IsEmpty())
+	{
+		EventName = FText::Format(LOCTEXT("EventWithSignatureName", "Event {0}"), EventName);
+	}
+	return EventName;
 }
 
 //------------------------------------------------------------------------------
 FText UBlueprintEventNodeSpawner::GetDefaultSearchKeywords() const
 {
-	check(EventFunc != nullptr);
-	return FText::FromString(UK2Node_CallFunction::GetKeywordsForFunction(EventFunc));
+	FText Keywords = FText::GetEmpty();
+	if (EventFunc != nullptr)
+	{
+		Keywords = FText::FromString(UK2Node_CallFunction::GetKeywordsForFunction(EventFunc));
+	}
+	return Keywords;
 }
 
 //------------------------------------------------------------------------------
diff --git a/Engine/Source/Editor/BlueprintGraph/Private/K2Node_CustomEvent.cpp b/Engine/Source/Editor/BlueprintGraph/Private/K2Node_CustomEvent.cpp
--- a/Engine/Source/Editor/BlueprintGraph/Private/K2Node_CustomEvent.cpp
+++ b/Engine/Source/Editor/BlueprintGraph/Private/K2Node_CustomEvent.cpp
@@ -5,7 +5,7 @@
 
 #include "Kismet2NameValidators.h"
 #include "CompilerResultsLog.h"
-#include "BlueprintNodeSpawner.h"
+#include "BlueprintEventNodeSpawner.h"
 
 #define LOCTEXT_NAMESPACE "K2Node_CustomEvent"
 
@@ -230,7 +230,7 @@ void UK2Node_CustomEvent::ValidateNodeDuringCompilation(class FCompilerResultsLo
 
 void UK2Node_CustomEvent::GetMenuActions(TArray<UBlueprintNodeSpawner*>& ActionListOut) const
 {
-	UBlueprintNodeSpawner* NodeSpawner = UBlueprintNodeSpawner::Create(GetClass());
+	UBlueprintNodeSpawner* NodeSpawner = UBlueprintEventNodeSpawner::Create(GetClass(), FName());
 	check(NodeSpawner != nullptr);
 
 	auto SetupCustomEventNodeLambda = [](UEdGraphNode* NewNode, bool bIsTemplateNode)
diff --git a/Engine/Source/Editor/BlueprintGraph/Private/K2Node_Event.cpp b/Engine/Source/Editor/BlueprintGraph/Private/K2Node_Event.cpp
--- a/Engine/Source/Editor/BlueprintGraph/Private/K2Node_Event.cpp
+++ b/Engine/Source/Editor/BlueprintGraph/Private/K2Node_Event.cpp
@@ -5,7 +5,6 @@
 #include "CompilerResultsLog.h"
 #include "KismetCompiler.h"
 #include "EventEntryHandler.h"
-#include "BlueprintActionMenuBuilder.h"
 
 const FString UK2Node_Event::DelegateOutputName(TEXT("OutputDelegate"));
 
@@ -631,13 +630,19 @@ void UK2Node_Event::GetNodeAttributes( TArray<TKeyValuePair<FString, FString>>&
 
 FText UK2Node_Event::GetMenuCategory() const
 {
-	FText Category;
+	FText FuncCategory;
 	if (UFunction* Function = FindField<UFunction>(EventSignatureClass, EventSignatureName))
 	{
-		Category = FText::FromString(UK2Node_CallFunction::GetDefaultCategoryForFunction(Function, TEXT("")));
+		FuncCategory = FText::FromString(UK2Node_CallFunction::GetDefaultCategoryForFunction(Function, TEXT("")));
 	}
 
-	return FText::Format(LOCTEXT("ConcatenatedCategory", "{0}|{1}"), FBlueprintActionMenuBuilder::AddEventCategory, Category);
+	FText Category = LOCTEXT("AddEventCategory", "Add Event");
+	if (!FuncCategory.IsEmpty())
+	{
+		Category = FText::Format(LOCTEXT("ConcatenatedCategory", "{0}|{1}"), Category, FuncCategory);
+	}
+
+	return Category;
 }
 
 bool UK2Node_Event::IsDeprecated() const
diff --git a/Engine/Source/Editor/BlueprintGraph/Public/BlueprintActionMenuBuilder.h b/Engine/Source/Editor/BlueprintGraph/Public/BlueprintActionMenuBuilder.h
--- a/Engine/Source/Editor/BlueprintGraph/Public/BlueprintActionMenuBuilder.h
+++ b/Engine/Source/Editor/BlueprintGraph/Public/BlueprintActionMenuBuilder.h
@@ -26,7 +26,6 @@ public:
 	FBlueprintActionMenuBuilder() {}
 	
 	/** Blueprint centric menu categories (not generic enough to make it into FCommonEditorCategory) */
-	static FText const AddEventCategory;
 	static FText const VariablesCategory;
 
 	/**
diff --git a/Engine/Source/Editor/GameplayAbilitiesEditor/Classes/K2Node_GameplayCueEvent.h b/Engine/Source/Editor/GameplayAbilitiesEditor/Classes/K2Node_GameplayCueEvent.h
--- a/Engine/Source/Editor/GameplayAbilitiesEditor/Classes/K2Node_GameplayCueEvent.h
+++ b/Engine/Source/Editor/GameplayAbilitiesEditor/Classes/K2Node_GameplayCueEvent.h
@@ -14,10 +14,10 @@ class UK2Node_GameplayCueEvent : public UK2Node_Event
 	virtual FString GetTooltip() const override;
 	virtual FText GetNodeTitle(ENodeTitleType::Type TitleType) const override;
 	virtual void GetMenuEntries(struct FGraphContextMenuBuilder& ContextMenuBuilder) const override;
+	virtual bool CanPasteHere(const UEdGraph* TargetGraph, const UEdGraphSchema* Schema) const override;
 	// End UEdGraphNode interface
 	
 	// Begin UK2Node interface
 	virtual void GetMenuActions(TArray<UBlueprintNodeSpawner*>& ActionListOut) const override;
-	virtual FText GetMenuCategory() const override;
 	// End UK2Node interface
 };
\ No newline at end of file
diff --git a/Engine/Source/Editor/GameplayAbilitiesEditor/Private/K2Node_GameplayCueEvent.cpp b/Engine/Source/Editor/GameplayAbilitiesEditor/Private/K2Node_GameplayCueEvent.cpp
--- a/Engine/Source/Editor/GameplayAbilitiesEditor/Private/K2Node_GameplayCueEvent.cpp
+++ b/Engine/Source/Editor/GameplayAbilitiesEditor/Private/K2Node_GameplayCueEvent.cpp
@@ -7,8 +7,7 @@
 #include "K2ActionMenuBuilder.h"
 #include "GameplayTagsModule.h"
 #include "GameplayCueInterface.h"
-#include "BlueprintNodeSpawner.h"
-#include "BlueprintActionMenuBuilder.h" // for AddEventCategory
+#include "BlueprintEventNodeSpawner.h"
 
 #define LOCTEXT_NAMESPACE "K2Node_GameplayCueEvent"
 
@@ -30,12 +29,10 @@ FText UK2Node_GameplayCueEvent::GetNodeTitle(ENodeTitleType::Type TitleType) con
 	//return LOCTEXT("HandleGameplayCueEvent", "HandleGameplaCueEvent");
 }
 
-void UK2Node_GameplayCueEvent::GetMenuEntries(FGraphContextMenuBuilder& Context) const
+bool UK2Node_GameplayCueEvent::CanPasteHere(UEdGraph const* TargetGraph, UEdGraphSchema const* Schema) const
 {
-	Super::GetMenuEntries(Context);
-
 	bool bValidClass = false;
-	UBlueprint* MyBlueprint = Cast<UBlueprint>(Context.CurrentGraph->GetOuter());
+	UBlueprint* MyBlueprint = Cast<UBlueprint>(TargetGraph->GetOuter());
 	if (MyBlueprint && MyBlueprint->GeneratedClass)
 	{
 		if (MyBlueprint->GeneratedClass->ImplementsInterface(UGameplayCueInterface::StaticClass()))
@@ -44,7 +41,14 @@ void UK2Node_GameplayCueEvent::GetMenuEntries(FGraphContextMenuBuilder& Context)
 		}
 	}
 
-	if (!bValidClass) 
+	return bValidClass;
+}
+
+void UK2Node_GameplayCueEvent::GetMenuEntries(FGraphContextMenuBuilder& Context) const
+{
+	Super::GetMenuEntries(Context);
+
+	if (!CanPasteHere(Context.CurrentGraph, GetDefault<UEdGraphSchema>(Context.CurrentGraph->Schema)))
 	{
 		return;
 	}
@@ -98,16 +102,12 @@ void UK2Node_GameplayCueEvent::GetMenuActions(TArray<UBlueprintNodeSpawner*>& Ac
 	{
 		UBlueprintNodeSpawner::FCustomizeNodeDelegate PostSpawnDelegate = UBlueprintNodeSpawner::FCustomizeNodeDelegate::CreateStatic(CustomizeCueNodeLambda, TagIt->GetTagName());
 		
-		UBlueprintNodeSpawner* NodeSpawner = UBlueprintNodeSpawner::Create(GetClass(), PostSpawnDelegate);
+		UBlueprintNodeSpawner* NodeSpawner = UBlueprintEventNodeSpawner::Create(GetClass(), TagIt->GetTagName());
 		check(NodeSpawner != nullptr);
+		NodeSpawner->CustomizeNodeDelegate = PostSpawnDelegate;
 		
 		ActionListOut.Add(NodeSpawner);
 	}
 }
 
-FText UK2Node_GameplayCueEvent::GetMenuCategory() const
-{
-	return FText::Format(LOCTEXT("ActionMenuCategory", "{0}|GameplayCue"), FBlueprintActionMenuBuilder::AddEventCategory);
-}
-
 #undef LOCTEXT_NAMESPACE
