Properly update function parameter pin names when a split struct pin is renamed

diff --git a/Engine/Source/Editor/BlueprintGraph/Private/K2Node.cpp b/Engine/Source/Editor/BlueprintGraph/Private/K2Node.cpp
--- a/Engine/Source/Editor/BlueprintGraph/Private/K2Node.cpp
+++ b/Engine/Source/Editor/BlueprintGraph/Private/K2Node.cpp
@@ -892,11 +892,29 @@ ERenamePinResult UK2Node::RenameUserDefinedPin(const FString& OldName, const FSt
 
 	if(!bTest)
 	{
+		Pin->Modify();
 		Pin->PinName = NewName;
 		if(!Pin->DefaultTextValue.IsEmpty())
 		{
 			Pin->GetSchema()->TrySetDefaultText(*Pin, Pin->DefaultTextValue);
 		}
+
+		if (Pin->SubPins.Num() > 0)
+		{
+			TArray<UEdGraphPin*> PinsToUpdate = Pin->SubPins;
+
+			while (PinsToUpdate.Num() > 0)
+			{
+				UEdGraphPin* PinToRename = PinsToUpdate.Pop();
+				if (PinToRename->SubPins.Num() > 0)
+				{
+					PinsToUpdate.Append(PinToRename->SubPins);
+				}
+				PinToRename->Modify();
+				PinToRename->PinName = NewName + PinToRename->PinName.RightChop(OldName.Len());
+				PinToRename->PinFriendlyName = FText::FromString(NewName + PinToRename->PinFriendlyName.ToString().RightChop(OldName.Len()));
+			}
+		}
 	}
 
 	return ERenamePinResult::ERenamePinResult_Success;
diff --git a/Engine/Source/Editor/UnrealEd/Private/Kismet2/BlueprintEditorUtils.cpp b/Engine/Source/Editor/UnrealEd/Private/Kismet2/BlueprintEditorUtils.cpp
--- a/Engine/Source/Editor/UnrealEd/Private/Kismet2/BlueprintEditorUtils.cpp
+++ b/Engine/Source/Editor/UnrealEd/Private/Kismet2/BlueprintEditorUtils.cpp
@@ -234,11 +234,7 @@ static void RenameVariableReferences(UBlueprint* Blueprint, UClass* VariableClas
 						VariableNode->VariableReference.SetExternalMember(NewVarName, NodeRefClass);
 					}
 
-					if(UEdGraphPin* const Pin = VariableNode->FindPin(OldVarName.ToString()))
-					{
-						Pin->Modify();
-						Pin->PinName = NewVarName.ToString();
-					}
+					VariableNode->RenameUserDefinedPin(OldVarName.ToString(), NewVarName.ToString());
 				}
 				continue;
 			}
