Make type coercion work off of all wildcard pins in call array function nodes
Remove split pins when collapsing back to a wildcard on call array function nodes

diff --git a/Engine/Source/Editor/BlueprintGraph/Classes/K2Node_CallArrayFunction.h b/Engine/Source/Editor/BlueprintGraph/Classes/K2Node_CallArrayFunction.h
--- a/Engine/Source/Editor/BlueprintGraph/Classes/K2Node_CallArrayFunction.h
+++ b/Engine/Source/Editor/BlueprintGraph/Classes/K2Node_CallArrayFunction.h
@@ -51,6 +51,6 @@ class UK2Node_CallArrayFunction : public UK2Node_CallFunction
 	BLUEPRINTGRAPH_API static bool IsWildcardProperty(UFunction* InArrayFunction, const UProperty* InProperty);
 
 	void GetArrayTypeDependentPins(TArray<UEdGraphPin*>& OutPins) const;
-	void PropagateArrayTypeInfo();
+	void PropagateArrayTypeInfo(const UEdGraphPin* SourcePin);
 };
 
diff --git a/Engine/Source/Editor/BlueprintGraph/Private/K2Node_CallArrayFunction.cpp b/Engine/Source/Editor/BlueprintGraph/Private/K2Node_CallArrayFunction.cpp
--- a/Engine/Source/Editor/BlueprintGraph/Private/K2Node_CallArrayFunction.cpp
+++ b/Engine/Source/Editor/BlueprintGraph/Private/K2Node_CallArrayFunction.cpp
@@ -33,7 +33,7 @@ void UK2Node_CallArrayFunction::AllocateDefaultPins()
 		}
 	}
 
-	PropagateArrayTypeInfo();
+	PropagateArrayTypeInfo(TargetArrayPin);
 }
 
 void UK2Node_CallArrayFunction::PostReconstructNode()
@@ -45,25 +45,65 @@ void UK2Node_CallArrayFunction::NotifyPinConnectionListChanged(UEdGraphPin* Pin)
 {
 	Super::NotifyPinConnectionListChanged(Pin);
 
-	if( Pin == GetTargetArrayPin() )
+	TArray<UEdGraphPin*> PinsToCheck;
+	GetArrayTypeDependentPins(PinsToCheck);
+
+	for (int32 Index = 0; Index < PinsToCheck.Num(); ++Index)
+	{
+		UEdGraphPin* PinToCheck = PinsToCheck[Index];
+		if (PinToCheck->SubPins.Num() > 0)
+		{
+			PinsToCheck.Append(PinToCheck->SubPins);
+		}
+	}
+
+	PinsToCheck.Add(GetTargetArrayPin());
+
+	if (PinsToCheck.Contains(Pin))
 	{
+		const UEdGraphSchema_K2* Schema = GetDefault<UEdGraphSchema_K2>();
+		bool bNeedToPropagate = false;
+
 		if( Pin->LinkedTo.Num() > 0 )
 		{
-			UEdGraphPin* LinkedTo = Pin->LinkedTo[0];
-			check(LinkedTo);
-			check(LinkedTo->PinType.bIsArray);
+			if (Pin->PinType.PinCategory == Schema->PC_Wildcard)
+			{
+				UEdGraphPin* LinkedTo = Pin->LinkedTo[0];
+				check(LinkedTo);
+				check(Pin->PinType.bIsArray == LinkedTo->PinType.bIsArray);
+
+				Pin->PinType.PinCategory = LinkedTo->PinType.PinCategory;
+				Pin->PinType.PinSubCategory = LinkedTo->PinType.PinSubCategory;
+				Pin->PinType.PinSubCategoryObject = LinkedTo->PinType.PinSubCategoryObject;
 
-			Pin->PinType = LinkedTo->PinType;
+				bNeedToPropagate = true;
+			}
 		}
 		else
 		{
-			const UEdGraphSchema_K2* Schema = GetDefault<UEdGraphSchema_K2>();
-			Pin->PinType.PinCategory = Schema->PC_Wildcard;
-			Pin->PinType.PinSubCategory = TEXT("");
-			Pin->PinType.PinSubCategoryObject = NULL;
+			bNeedToPropagate = true;
+
+			for (UEdGraphPin* PinToCheck : PinsToCheck)
+			{
+				if (PinToCheck->LinkedTo.Num() > 0)
+				{
+					bNeedToPropagate = false;
+					break;
+				}
+			}
+
+			if (bNeedToPropagate)
+			{
+				Pin->PinType.PinCategory = Schema->PC_Wildcard;
+				Pin->PinType.PinSubCategory = TEXT("");
+				Pin->PinType.PinSubCategoryObject = NULL;
+			}
 		}
 
-		PropagateArrayTypeInfo();
+		if (bNeedToPropagate)
+		{
+			PropagateArrayTypeInfo(Pin);
+		}
 	}
 }
 
@@ -153,32 +193,38 @@ void UK2Node_CallArrayFunction::GetArrayTypeDependentPins(TArray<UEdGraphPin*>&
 	}
 }
 
-void UK2Node_CallArrayFunction::PropagateArrayTypeInfo()
+void UK2Node_CallArrayFunction::PropagateArrayTypeInfo(const UEdGraphPin* SourcePin)
 {
-	const UEdGraphPin* ArrayTargetPin = GetTargetArrayPin();
-
-	if( ArrayTargetPin )
+	if( SourcePin )
 	{
 		const UEdGraphSchema_K2* Schema = GetDefault<UEdGraphSchema_K2>();
 
 		TArray<UEdGraphPin*> DependentPins;
 		GetArrayTypeDependentPins(DependentPins);
+		DependentPins.Add(GetTargetArrayPin());
 	
 		// Propagate pin type info (except for array info!) to pins with dependent types
-		for( TArray<UEdGraphPin*>::TIterator it(DependentPins); it; ++it )
+		for (UEdGraphPin* CurrentPin : DependentPins)
 		{
-			UEdGraphPin* CurrentPin = *it;
-			CurrentPin->PinType.PinCategory = ArrayTargetPin->PinType.PinCategory;
-			CurrentPin->PinType.PinSubCategory = ArrayTargetPin->PinType.PinSubCategory;
-			CurrentPin->PinType.PinSubCategoryObject = ArrayTargetPin->PinType.PinSubCategoryObject;
-
-			// Verify that all previous connections to this pin are still valid with the new type
-			for( TArray<UEdGraphPin*>::TIterator ConnectionIt(CurrentPin->LinkedTo); ConnectionIt; ++ConnectionIt )
+			if (CurrentPin != SourcePin)
 			{
-				UEdGraphPin* ConnectedPin = *ConnectionIt;
-				if( !Schema->ArePinsCompatible(CurrentPin, ConnectedPin) )
+				if (CurrentPin->SubPins.Num() > 0)
+				{
+					Schema->RecombinePin(CurrentPin->SubPins[0]);
+				}
+
+				CurrentPin->PinType.PinCategory = SourcePin->PinType.PinCategory;
+				CurrentPin->PinType.PinSubCategory = SourcePin->PinType.PinSubCategory;
+				CurrentPin->PinType.PinSubCategoryObject = SourcePin->PinType.PinSubCategoryObject;
+
+				// Verify that all previous connections to this pin are still valid with the new type
+				for (TArray<UEdGraphPin*>::TIterator ConnectionIt(CurrentPin->LinkedTo); ConnectionIt; ++ConnectionIt)
 				{
-					CurrentPin->BreakLinkTo(ConnectedPin);
+					UEdGraphPin* ConnectedPin = *ConnectionIt;
+					if (!Schema->ArePinsCompatible(CurrentPin, ConnectedPin))
+					{
+						CurrentPin->BreakLinkTo(ConnectedPin);
+					}
 				}
 			}
 		}
