Don't lose type when reconstructing array function nodes if the target array pin is not connected

diff --git a/Engine/Source/Editor/BlueprintGraph/Private/K2Node_CallArrayFunction.cpp b/Engine/Source/Editor/BlueprintGraph/Private/K2Node_CallArrayFunction.cpp
--- a/Engine/Source/Editor/BlueprintGraph/Private/K2Node_CallArrayFunction.cpp
+++ b/Engine/Source/Editor/BlueprintGraph/Private/K2Node_CallArrayFunction.cpp
@@ -38,7 +38,13 @@ void UK2Node_CallArrayFunction::AllocateDefaultPins()
 
 void UK2Node_CallArrayFunction::PostReconstructNode()
 {
-	PinConnectionListChanged( GetTargetArrayPin() );
+	for (UEdGraphPin* Pin : Pins)
+	{
+		if (Pin->LinkedTo.Num() > 0)
+		{
+			PinConnectionListChanged(Pin);
+		}
+	}
 }
 
 void UK2Node_CallArrayFunction::NotifyPinConnectionListChanged(UEdGraphPin* Pin)
