Never show a default value widget for an array pin (fixes showing default value widget on the array input pin on  For Each macros when type is coerced from the output pin)

diff --git a/Engine/Source/Editor/BlueprintGraph/Private/EdGraphSchema_K2.cpp b/Engine/Source/Editor/BlueprintGraph/Private/EdGraphSchema_K2.cpp
--- a/Engine/Source/Editor/BlueprintGraph/Private/EdGraphSchema_K2.cpp
+++ b/Engine/Source/Editor/BlueprintGraph/Private/EdGraphSchema_K2.cpp
@@ -3038,7 +3038,7 @@ bool UEdGraphSchema_K2::ShouldHidePinDefaultValue(UEdGraphPin* Pin) const
 {
 	check(Pin != NULL);
 
-	if ((Pin->PinName == PN_Self && Pin->LinkedTo.Num() > 0) || (Pin->PinType.PinCategory == PC_Exec) || (Pin->PinType.bIsReference && !IsAutoCreateRefTerm(Pin)) || Pin->bDefaultValueIsIgnored)
+	if (Pin->bDefaultValueIsIgnored || Pin->PinType.bIsArray || (Pin->PinName == PN_Self && Pin->LinkedTo.Num() > 0) || (Pin->PinType.PinCategory == PC_Exec) || (Pin->PinType.bIsReference && !IsAutoCreateRefTerm(Pin)))
 	{
 		return true;
 	}
