More fallout from CL 2221066/2221070 (fixing link errors not caught in xcode).

diff --git a/Engine/Source/Editor/BlueprintGraph/Classes/BlueprintComponentNodeSpawner.h b/Engine/Source/Editor/BlueprintGraph/Classes/BlueprintComponentNodeSpawner.h
--- a/Engine/Source/Editor/BlueprintGraph/Classes/BlueprintComponentNodeSpawner.h
+++ b/Engine/Source/Editor/BlueprintGraph/Classes/BlueprintComponentNodeSpawner.h
@@ -14,7 +14,7 @@ class UActorComponent;
  * FEdGraphSchemaAction_K2AddComponent.
  */
 UCLASS(Transient)
-class UBlueprintComponentNodeSpawner : public UBlueprintNodeSpawner
+class BLUEPRINTGRAPH_API UBlueprintComponentNodeSpawner : public UBlueprintNodeSpawner
 {
 	GENERATED_UCLASS_BODY()
 
diff --git a/Engine/Source/Editor/BlueprintGraph/Classes/BlueprintEventNodeSpawner.h b/Engine/Source/Editor/BlueprintGraph/Classes/BlueprintEventNodeSpawner.h
--- a/Engine/Source/Editor/BlueprintGraph/Classes/BlueprintEventNodeSpawner.h
+++ b/Engine/Source/Editor/BlueprintGraph/Classes/BlueprintEventNodeSpawner.h
@@ -15,7 +15,7 @@ class UK2Node_Event;
  * FEdGraphSchemaAction_K2ViewNode.
  */
 UCLASS(Transient)
-class UBlueprintEventNodeSpawner : public UBlueprintNodeSpawner
+class BLUEPRINTGRAPH_API UBlueprintEventNodeSpawner : public UBlueprintNodeSpawner
 {
 	GENERATED_UCLASS_BODY()
 
@@ -28,7 +28,7 @@ public:
 	 * @param  EventFunc	The function you want assigned to new nodes.
 	 * @return A newly allocated instance of this class.
 	 */
-	BLUEPRINTGRAPH_API static UBlueprintEventNodeSpawner* Create(UFunction const* const EventFunc);
+	static UBlueprintEventNodeSpawner* Create(UFunction const* const EventFunc);
 
 	/**
 	 * Creates a new UBlueprintEventNodeSpawner for custom events. The 
@@ -39,7 +39,7 @@ public:
 	 * @param  CustomEventName	The name you want assigned to the event.
 	 * @return A newly allocated instance of this class.
 	 */
-	BLUEPRINTGRAPH_API static UBlueprintEventNodeSpawner* Create(TSubclassOf<UK2Node_Event> NodeClass, FName CustomEventName);
+	static UBlueprintEventNodeSpawner* Create(TSubclassOf<UK2Node_Event> NodeClass, FName CustomEventName);
 
 	// UBlueprintNodeSpawner interface
 	virtual UEdGraphNode* Invoke(UEdGraph* ParentGraph) const override;
diff --git a/Engine/Source/Editor/BlueprintGraph/Classes/BlueprintFunctionNodeSpawner.h b/Engine/Source/Editor/BlueprintGraph/Classes/BlueprintFunctionNodeSpawner.h
--- a/Engine/Source/Editor/BlueprintGraph/Classes/BlueprintFunctionNodeSpawner.h
+++ b/Engine/Source/Editor/BlueprintGraph/Classes/BlueprintFunctionNodeSpawner.h
@@ -10,7 +10,7 @@
  * "action" portion of certain FBlueprintActionMenuItems. 
  */
 UCLASS(Transient)
-class UBlueprintFunctionNodeSpawner : public UBlueprintNodeSpawner
+class BLUEPRINTGRAPH_API UBlueprintFunctionNodeSpawner : public UBlueprintNodeSpawner
 {
 	GENERATED_UCLASS_BODY()
 
@@ -23,7 +23,7 @@ public:
 	 * @param  Function		The function you want assigned to new nodes.
 	 * @return A newly allocated instance of this class.
 	 */
-	BLUEPRINTGRAPH_API static UBlueprintFunctionNodeSpawner* Create(UFunction const* const Function);
+	static UBlueprintFunctionNodeSpawner* Create(UFunction const* const Function);
 
 	// UBlueprintNodeSpawner interface
 	virtual UEdGraphNode* Invoke(UEdGraph* ParentGraph) const override;
diff --git a/Engine/Source/Editor/BlueprintGraph/Classes/BlueprintPropertyNodeSpawner.h b/Engine/Source/Editor/BlueprintGraph/Classes/BlueprintPropertyNodeSpawner.h
--- a/Engine/Source/Editor/BlueprintGraph/Classes/BlueprintPropertyNodeSpawner.h
+++ b/Engine/Source/Editor/BlueprintGraph/Classes/BlueprintPropertyNodeSpawner.h
@@ -12,7 +12,7 @@
  * FEdGraphSchemaAction_K2Var, FEdGraphSchemaAction_K2Delegate, etc.
  */
 UCLASS(Transient)
-class UBlueprintPropertyNodeSpawner : public UBlueprintNodeSpawner
+class BLUEPRINTGRAPH_API UBlueprintPropertyNodeSpawner : public UBlueprintNodeSpawner
 {
 	GENERATED_UCLASS_BODY()
 
@@ -26,7 +26,7 @@ public:
 	 * @param  NodeClass	The type of node you want spawned (can't distinct a getter from a setter from the property alone).
 	 * @return A newly allocated instance of this class.
 	 */
-	BLUEPRINTGRAPH_API static UBlueprintPropertyNodeSpawner* Create(UProperty const* const Property, UClass* const NodeClass = nullptr);
+	static UBlueprintPropertyNodeSpawner* Create(UProperty const* const Property, UClass* const NodeClass = nullptr);
 
 	// UBlueprintNodeSpawner interface
 	virtual UEdGraphNode* Invoke(UEdGraph* ParentGraph) const override;
@@ -39,14 +39,14 @@ public:
 	 *
 	 * @return The property that this class was initialized with.
 	 */
-	BLUEPRINTGRAPH_API UProperty const* GetProperty() const;
+	UProperty const* GetProperty() const;
 
 	/**
 	 * Checks to see if the associated property is a delegate property.
 	 *
 	 * @return True if the property is a UMulticastDelegateProperty (otherwise false).
 	 */
-	BLUEPRINTGRAPH_API bool IsDelegateProperty() const;
+	bool IsDelegateProperty() const;
 	
 private:
 	/** The property to configure new nodes with. */
diff --git a/Engine/Source/Editor/BlueprintGraph/Public/BlueprintNodeSpawner.h b/Engine/Source/Editor/BlueprintGraph/Public/BlueprintNodeSpawner.h
--- a/Engine/Source/Editor/BlueprintGraph/Public/BlueprintNodeSpawner.h
+++ b/Engine/Source/Editor/BlueprintGraph/Public/BlueprintNodeSpawner.h
@@ -19,7 +19,7 @@ class UEdGraphNode;
  * these "actions".
  */
 UCLASS(Transient)
-class UBlueprintNodeSpawner : public UObject
+class BLUEPRINTGRAPH_API UBlueprintNodeSpawner : public UObject
 {
 	GENERATED_UCLASS_BODY()
 	DECLARE_DELEGATE_TwoParams(FCustomizeNodeDelegate, UEdGraphNode*, bool);
@@ -34,7 +34,7 @@ public:
 	 * @param  CustomizeNodeDelegate    A delegate to perform specialized node setup post-spawn.
 	 * @return A newly allocated instance of this class.
 	 */
-	BLUEPRINTGRAPH_API static UBlueprintNodeSpawner* Create(UClass* const NodeClass, FCustomizeNodeDelegate CustomizeNodeDelegate = FCustomizeNodeDelegate());
+	static UBlueprintNodeSpawner* Create(UClass* const NodeClass, FCustomizeNodeDelegate CustomizeNodeDelegate = FCustomizeNodeDelegate());
 	
 	/**
 	 * Holds the class of node to spawn. May be null for sub-classes that know
diff --git a/Engine/Source/Editor/Kismet/Public/BlueprintActionMenuUtils.h b/Engine/Source/Editor/Kismet/Public/BlueprintActionMenuUtils.h
--- a/Engine/Source/Editor/Kismet/Public/BlueprintActionMenuUtils.h
+++ b/Engine/Source/Editor/Kismet/Public/BlueprintActionMenuUtils.h
@@ -17,7 +17,7 @@ struct FBlueprintActionMenuUtils
 	 * @param  FilterClass	If not null, then this specifies the class whose members we want listed (and nothing else).
 	 * @param  MenuOut		The structure that will be populated with palette menu items.
 	 */
-	static void MakePaletteMenu(FBlueprintActionContext const& Context, UClass* FilterClass, FBlueprintActionMenuBuilder& MenuOut);
+	KISMET_API static void MakePaletteMenu(FBlueprintActionContext const& Context, UClass* FilterClass, FBlueprintActionMenuBuilder& MenuOut);
 	
 	/**
 	 * A centralized utility function for constructing blueprint context menus.
@@ -28,6 +28,6 @@ struct FBlueprintActionMenuUtils
 	 * @param  SelectedProperties	A set of selected properties to offer contextual menu options for.
 	 * @param  MenuOut				The structure that will be populated with context menu items.
 	 */
-	static void MakeContextMenu(FBlueprintActionContext const& Context, TArray<UProperty*> const& SelectedProperties, FBlueprintActionMenuBuilder& MenuOut);
+	KISMET_API static void MakeContextMenu(FBlueprintActionContext const& Context, TArray<UProperty*> const& SelectedProperties, FBlueprintActionMenuBuilder& MenuOut);
 };
 
