Keeping the node-spawners (for the new BP action menu system) from getting garbage collected.

diff --git a/Engine/Source/Editor/BlueprintGraph/Classes/BlueprintComponentNodeSpawner.h b/Engine/Source/Editor/BlueprintGraph/Classes/BlueprintComponentNodeSpawner.h
--- a/Engine/Source/Editor/BlueprintGraph/Classes/BlueprintComponentNodeSpawner.h
+++ b/Engine/Source/Editor/BlueprintGraph/Classes/BlueprintComponentNodeSpawner.h
@@ -25,9 +25,10 @@ public:
 	 * viable as a spawnable component (do that before calling this).
 	 *
 	 * @param  ComponentClass	The component type you want spawned nodes to spawn.
+	 * @param  Outer			Optional outer for the new spawner (if left null, the transient package will be used).
 	 * @return A newly allocated instance of this class.
 	 */
-	static UBlueprintComponentNodeSpawner* Create(TSubclassOf<UActorComponent> const ComponentClass);
+	static UBlueprintComponentNodeSpawner* Create(TSubclassOf<UActorComponent> const ComponentClass, UObject* Outer = nullptr);
 
 	// UBlueprintNodeSpawner interface
 	virtual UEdGraphNode* Invoke(UEdGraph* ParentGraph) const override;
diff --git a/Engine/Source/Editor/BlueprintGraph/Classes/BlueprintEventNodeSpawner.h b/Engine/Source/Editor/BlueprintGraph/Classes/BlueprintEventNodeSpawner.h
--- a/Engine/Source/Editor/BlueprintGraph/Classes/BlueprintEventNodeSpawner.h
+++ b/Engine/Source/Editor/BlueprintGraph/Classes/BlueprintEventNodeSpawner.h
@@ -26,9 +26,10 @@ public:
 	 * viable as a blueprint event (do that before calling this).
 	 *
 	 * @param  EventFunc	The function you want assigned to new nodes.
+	 * @param  Outer		Optional outer for the new spawner (if left null, the transient package will be used).
 	 * @return A newly allocated instance of this class.
 	 */
-	static UBlueprintEventNodeSpawner* Create(UFunction const* const EventFunc);
+	static UBlueprintEventNodeSpawner* Create(UFunction const* const EventFunc, UObject* Outer = nullptr);
 
 	/**
 	 * Creates a new UBlueprintEventNodeSpawner for custom events. The 
@@ -37,9 +38,10 @@ public:
 	 *
 	 * @param  NodeClass		The event node type that you want this to spawn.
 	 * @param  CustomEventName	The name you want assigned to the event.
+	 * @param  Outer			Optional outer for the new spawner (if left null, the transient package will be used).
 	 * @return A newly allocated instance of this class.
 	 */
-	static UBlueprintEventNodeSpawner* Create(TSubclassOf<UK2Node_Event> NodeClass, FName CustomEventName);
+	static UBlueprintEventNodeSpawner* Create(TSubclassOf<UK2Node_Event> NodeClass, FName CustomEventName, UObject* Outer = nullptr);
 
 	// UBlueprintNodeSpawner interface
 	virtual UEdGraphNode* Invoke(UEdGraph* ParentGraph) const override;
diff --git a/Engine/Source/Editor/BlueprintGraph/Classes/BlueprintFunctionNodeSpawner.h b/Engine/Source/Editor/BlueprintGraph/Classes/BlueprintFunctionNodeSpawner.h
--- a/Engine/Source/Editor/BlueprintGraph/Classes/BlueprintFunctionNodeSpawner.h
+++ b/Engine/Source/Editor/BlueprintGraph/Classes/BlueprintFunctionNodeSpawner.h
@@ -21,9 +21,10 @@ public:
 	 * viable as a blueprint function call (do that before calling this).
 	 *
 	 * @param  Function		The function you want assigned to new nodes.
+	 * @param  Outer		Optional outer for the new spawner (if left null, the transient package will be used).
 	 * @return A newly allocated instance of this class.
 	 */
-	static UBlueprintFunctionNodeSpawner* Create(UFunction const* const Function);
+	static UBlueprintFunctionNodeSpawner* Create(UFunction const* const Function, UObject* Outer = nullptr);
 
 	// UBlueprintNodeSpawner interface
 	virtual UEdGraphNode* Invoke(UEdGraph* ParentGraph) const override;
diff --git a/Engine/Source/Editor/BlueprintGraph/Classes/BlueprintPropertyNodeSpawner.h b/Engine/Source/Editor/BlueprintGraph/Classes/BlueprintPropertyNodeSpawner.h
--- a/Engine/Source/Editor/BlueprintGraph/Classes/BlueprintPropertyNodeSpawner.h
+++ b/Engine/Source/Editor/BlueprintGraph/Classes/BlueprintPropertyNodeSpawner.h
@@ -23,10 +23,11 @@ public:
 	 * accessible from blueprints (do that before calling this).
 	 *
 	 * @param  Property		The property you want assigned to spawned nodes.
+	 * @param  Outer		Optional outer for the new spawner (if left null, the transient package will be used).
 	 * @param  NodeClass	The type of node you want spawned (can't distinct a getter from a setter from the property alone).
 	 * @return A newly allocated instance of this class.
 	 */
-	static UBlueprintPropertyNodeSpawner* Create(UProperty const* const Property, UClass* const NodeClass = nullptr);
+	static UBlueprintPropertyNodeSpawner* Create(UProperty const* const Property, UObject* Outer = nullptr, TSubclassOf<UEdGraphNode> const NodeClass = nullptr);
 
 	// UBlueprintNodeSpawner interface
 	virtual UEdGraphNode* Invoke(UEdGraph* ParentGraph) const override;
diff --git a/Engine/Source/Editor/BlueprintGraph/Private/BlueprintActionDatabase.cpp b/Engine/Source/Editor/BlueprintGraph/Private/BlueprintActionDatabase.cpp
--- a/Engine/Source/Editor/BlueprintGraph/Private/BlueprintActionDatabase.cpp
+++ b/Engine/Source/Editor/BlueprintGraph/Private/BlueprintActionDatabase.cpp
@@ -245,7 +245,7 @@ static UBlueprintNodeSpawner* FBlueprintNodeSpawnerFactory::MakeCommentNodeSpawn
 template<class DelegateNodeType>
 static UBlueprintNodeSpawner* FBlueprintNodeSpawnerFactory::MakeDelegateNodeSpawner(UMulticastDelegateProperty* Property)
 {
-	UBlueprintNodeSpawner* NodeSpawner = UBlueprintPropertyNodeSpawner::Create(Property, DelegateNodeType::StaticClass());
+	UBlueprintNodeSpawner* NodeSpawner = UBlueprintPropertyNodeSpawner::Create(Property, GetTransientPackage(), DelegateNodeType::StaticClass());
 	check(NodeSpawner != nullptr);
 	
 	auto CustomizeDelegateNodeLambda = [](UEdGraphNode* NewNode, bool bIsTemplateNode, UMulticastDelegateProperty* Property)
@@ -269,7 +269,38 @@ static UBlueprintNodeSpawner* FBlueprintNodeSpawnerFactory::MakeDelegateNodeSpaw
 
 namespace FBlueprintActionDatabaseImpl
 {
-	typedef FBlueprintActionDatabase::FActionList FActionList;
+	/** 
+	 * Wrapper around FBlueprintActionDatabase's FActionList. Manages the 
+	 * RF_RootSet flag for the UBlueprintNodeSpawners (ensures they don't get 
+	 * GC'd whilst in the database.
+	 */
+	struct FActionList
+	{
+	public:
+		/** Clears the passed action-list and removes the RF_RootSet from all of its actions */
+		FActionList(FBlueprintActionDatabase::FActionList& DatabaseIn)
+			: ClassDatabase(DatabaseIn)
+		{
+			for (UBlueprintNodeSpawner* Action : ClassDatabase)
+			{
+				check(Action->GetOuter() == GetTransientPackage());
+				Action->ClearFlags(RF_RootSet);
+			}
+			ClassDatabase.Empty();
+		}
+
+		/** 
+		 * Passes the wrapped ClassDatabase this node-spawner to add, after 
+		 * adding the RF_RootSet flag (to keep it from getting GC'd).
+		 *
+		 * @param  NodeSpawner	The action you want added to the database.
+		 */
+		void Add(UBlueprintNodeSpawner* NodeSpawner);
+
+	private:
+		/** */
+		FBlueprintActionDatabase::FActionList& ClassDatabase;
+	};
 
 	/**
 	 * Mimics UEdGraphSchema_K2::CanUserKismetAccessVariable(); however, this 
@@ -385,6 +416,18 @@ namespace FBlueprintActionDatabaseImpl
 	static void AddAutonomousNodeActions(UClass* const Class, FActionList& ActionListOut);
 }
 
+//------------------------------------------------------------------------------
+void FBlueprintActionDatabaseImpl::FActionList::Add(UBlueprintNodeSpawner* NodeSpawner)
+{
+	check(NodeSpawner->GetOuter() == GetTransientPackage());
+	// since this spawner's outer is the transient package, we want to mark it
+	// root so that it doesn't get GC'd (we have to be careful and remove this 
+	// flag when we refresh).
+	NodeSpawner->SetFlags(RF_RootSet);
+
+	ClassDatabase.Add(NodeSpawner);
+}
+
 //------------------------------------------------------------------------------
 static bool FBlueprintActionDatabaseImpl::IsPropertyBlueprintVisible(UProperty const* const Property)
 {
@@ -475,9 +518,10 @@ static void FBlueprintActionDatabaseImpl::AddClassPropertyActions(UClass const*
  		}
 		else
 		{
-			UBlueprintPropertyNodeSpawner* GetterSpawner = UBlueprintPropertyNodeSpawner::Create(Property, UK2Node_VariableGet::StaticClass());
+			UPackage* SpawnerOuter = GetTransientPackage();
+			UBlueprintPropertyNodeSpawner* GetterSpawner = UBlueprintPropertyNodeSpawner::Create(Property, SpawnerOuter, UK2Node_VariableGet::StaticClass());
 			ActionListOut.Add(GetterSpawner);
-			UBlueprintPropertyNodeSpawner* SetterSpawner = UBlueprintPropertyNodeSpawner::Create(Property, UK2Node_VariableSet::StaticClass());
+			UBlueprintPropertyNodeSpawner* SetterSpawner = UBlueprintPropertyNodeSpawner::Create(Property, SpawnerOuter, UK2Node_VariableSet::StaticClass());
 			ActionListOut.Add(SetterSpawner);
 		}
 	}
@@ -629,7 +673,14 @@ static void FBlueprintActionDatabaseImpl::AddAutonomousNodeActions(UClass* const
 	{
 		UK2Node const* const NodeCDO = Class->GetDefaultObject<UK2Node>();
 		check(NodeCDO != nullptr);
-		NodeCDO->GetMenuActions(ActionListOut);
+
+		FBlueprintActionDatabase::FActionList NodeActionList;
+		NodeCDO->GetMenuActions(NodeActionList);
+
+		for (UBlueprintNodeSpawner* Spawner : NodeActionList)
+		{
+			ActionListOut.Add(Spawner);
+		}
 	}
 	// unfortunately, UEdGraphNode_Comment is not a UK2Node and therefore cannot
 	// leverage UK2Node's GetMenuActions() function, so here we HACK it in
@@ -684,8 +735,7 @@ void FBlueprintActionDatabase::RefreshClassActions(UClass* const Class)
 	// blueprint compile (if it is a SKEL or REINST class, then don't bother)
 	if (!bIsSkelClass && !bIsReinstClass)
 	{
-		FActionList& ClassActionList = ClassActions.FindOrAdd(Class);
-		ClassActionList.Empty();
+		FBlueprintActionDatabaseImpl::FActionList ClassActionList(ClassActions.FindOrAdd(Class));
 		
 		// class field actions (nodes that represent and perform actions on
 		// specific fields of the class... functions, properties, etc.)
diff --git a/Engine/Source/Editor/BlueprintGraph/Private/BlueprintComponentNodeSpawner.cpp b/Engine/Source/Editor/BlueprintGraph/Private/BlueprintComponentNodeSpawner.cpp
--- a/Engine/Source/Editor/BlueprintGraph/Private/BlueprintComponentNodeSpawner.cpp
+++ b/Engine/Source/Editor/BlueprintGraph/Private/BlueprintComponentNodeSpawner.cpp
@@ -11,11 +11,16 @@
  ******************************************************************************/
 
 //------------------------------------------------------------------------------
-UBlueprintComponentNodeSpawner* UBlueprintComponentNodeSpawner::Create(TSubclassOf<UActorComponent> const ComponentClass)
+UBlueprintComponentNodeSpawner* UBlueprintComponentNodeSpawner::Create(TSubclassOf<UActorComponent> const ComponentClass, UObject* Outer/* = nullptr*/)
 {
 	check(ComponentClass != nullptr);
 
-	UBlueprintComponentNodeSpawner* NodeSpawner = NewObject<UBlueprintComponentNodeSpawner>(GetTransientPackage());
+	if (Outer == nullptr)
+	{
+		Outer = GetTransientPackage();
+	}
+
+	UBlueprintComponentNodeSpawner* NodeSpawner = NewObject<UBlueprintComponentNodeSpawner>(Outer);
 	NodeSpawner->ComponentClass = ComponentClass;
 	NodeSpawner->NodeClass      = UK2Node_AddComponent::StaticClass();
 
diff --git a/Engine/Source/Editor/BlueprintGraph/Private/BlueprintEventNodeSpawner.cpp b/Engine/Source/Editor/BlueprintGraph/Private/BlueprintEventNodeSpawner.cpp
--- a/Engine/Source/Editor/BlueprintGraph/Private/BlueprintEventNodeSpawner.cpp
+++ b/Engine/Source/Editor/BlueprintGraph/Private/BlueprintEventNodeSpawner.cpp
@@ -47,11 +47,16 @@ UK2Node_Event* UBlueprintEventNodeSpawnerImpl::FindCustomEventNode(UBlueprint* B
  ******************************************************************************/
 
 //------------------------------------------------------------------------------
-UBlueprintEventNodeSpawner* UBlueprintEventNodeSpawner::Create(UFunction const* const EventFunc)
+UBlueprintEventNodeSpawner* UBlueprintEventNodeSpawner::Create(UFunction const* const EventFunc, UObject* Outer/* = nullptr*/)
 {
 	check(EventFunc != nullptr);
 
-	UBlueprintEventNodeSpawner* NodeSpawner = NewObject<UBlueprintEventNodeSpawner>(GetTransientPackage());
+	if (Outer == nullptr)
+	{
+		Outer = GetTransientPackage();
+	}
+
+	UBlueprintEventNodeSpawner* NodeSpawner = NewObject<UBlueprintEventNodeSpawner>(Outer);
 	NodeSpawner->EventFunc = EventFunc;
 	NodeSpawner->NodeClass = UK2Node_Event::StaticClass();
 
@@ -59,9 +64,14 @@ UBlueprintEventNodeSpawner* UBlueprintEventNodeSpawner::Create(UFunction const*
 }
 
 //------------------------------------------------------------------------------
-UBlueprintEventNodeSpawner* UBlueprintEventNodeSpawner::Create(TSubclassOf<UK2Node_Event> NodeClass, FName CustomEventName)
+UBlueprintEventNodeSpawner* UBlueprintEventNodeSpawner::Create(TSubclassOf<UK2Node_Event> NodeClass, FName CustomEventName, UObject* Outer/* = nullptr*/)
 {
-	UBlueprintEventNodeSpawner* NodeSpawner = NewObject<UBlueprintEventNodeSpawner>(GetTransientPackage());
+	if (Outer == nullptr)
+	{
+		Outer = GetTransientPackage();
+	}
+
+	UBlueprintEventNodeSpawner* NodeSpawner = NewObject<UBlueprintEventNodeSpawner>(Outer);
 	NodeSpawner->NodeClass       = NodeClass;
 	NodeSpawner->CustomEventName = CustomEventName;
 
diff --git a/Engine/Source/Editor/BlueprintGraph/Private/BlueprintFunctionNodeSpawner.cpp b/Engine/Source/Editor/BlueprintGraph/Private/BlueprintFunctionNodeSpawner.cpp
--- a/Engine/Source/Editor/BlueprintGraph/Private/BlueprintFunctionNodeSpawner.cpp
+++ b/Engine/Source/Editor/BlueprintGraph/Private/BlueprintFunctionNodeSpawner.cpp
@@ -10,7 +10,7 @@
 
 //------------------------------------------------------------------------------
 // Evolved from FK2ActionMenuBuilder::AddSpawnInfoForFunction()
-UBlueprintFunctionNodeSpawner* UBlueprintFunctionNodeSpawner::Create(UFunction const* const Function)
+UBlueprintFunctionNodeSpawner* UBlueprintFunctionNodeSpawner::Create(UFunction const* const Function, UObject* Outer/* = nullptr*/)
 {
 	check(Function != nullptr);
 
@@ -19,7 +19,11 @@ UBlueprintFunctionNodeSpawner* UBlueprintFunctionNodeSpawner::Create(UFunction c
 	bool const bIsCommutativeAssociativeBinaryOp = Function->HasMetaData(FBlueprintMetadata::MD_CommutativeAssociativeBinaryOperator);
 	bool const bIsMaterialParamCollectionFunc = Function->HasMetaData(FBlueprintMetadata::MD_MaterialParameterCollectionFunction);
 
-	UBlueprintFunctionNodeSpawner* NodeSpawner = NewObject<UBlueprintFunctionNodeSpawner>(GetTransientPackage());
+	if (Outer == nullptr)
+	{
+		Outer = GetTransientPackage();
+	}
+	UBlueprintFunctionNodeSpawner* NodeSpawner = NewObject<UBlueprintFunctionNodeSpawner>(Outer);
 	NodeSpawner->Function  = Function;
 
 	if (bIsCommutativeAssociativeBinaryOp && bIsPure)
diff --git a/Engine/Source/Editor/BlueprintGraph/Private/BlueprintNodeSpawner.cpp b/Engine/Source/Editor/BlueprintGraph/Private/BlueprintNodeSpawner.cpp
--- a/Engine/Source/Editor/BlueprintGraph/Private/BlueprintNodeSpawner.cpp
+++ b/Engine/Source/Editor/BlueprintGraph/Private/BlueprintNodeSpawner.cpp
@@ -8,12 +8,17 @@
 *******************************************************************************/
 
 //------------------------------------------------------------------------------
-UBlueprintNodeSpawner* UBlueprintNodeSpawner::Create(UClass* const NodeClass, FCustomizeNodeDelegate CustomizeNodeDelegate/* = FCustomizeNodeDelegate()*/)
+UBlueprintNodeSpawner* UBlueprintNodeSpawner::Create(TSubclassOf<UEdGraphNode> const NodeClass, UObject* Outer/* = nullptr*/, FCustomizeNodeDelegate CustomizeNodeDelegate/* = FCustomizeNodeDelegate()*/)
 {
 	check(NodeClass != nullptr);
 	check(NodeClass->IsChildOf<UEdGraphNode>());
+
+	if (Outer == nullptr)
+	{
+		Outer = GetTransientPackage();
+	}
 	
-	UBlueprintNodeSpawner* NodeSpawner = NewObject<UBlueprintNodeSpawner>(GetTransientPackage());
+	UBlueprintNodeSpawner* NodeSpawner = NewObject<UBlueprintNodeSpawner>(Outer);
 	NodeSpawner->NodeClass = NodeClass;
 	NodeSpawner->CustomizeNodeDelegate = CustomizeNodeDelegate;
 	
diff --git a/Engine/Source/Editor/BlueprintGraph/Private/BlueprintPropertyNodeSpawner.cpp b/Engine/Source/Editor/BlueprintGraph/Private/BlueprintPropertyNodeSpawner.cpp
--- a/Engine/Source/Editor/BlueprintGraph/Private/BlueprintPropertyNodeSpawner.cpp
+++ b/Engine/Source/Editor/BlueprintGraph/Private/BlueprintPropertyNodeSpawner.cpp
@@ -13,11 +13,16 @@
  ******************************************************************************/
 
 //------------------------------------------------------------------------------
-UBlueprintPropertyNodeSpawner* UBlueprintPropertyNodeSpawner::Create(UProperty const* const Property, UClass* const NodeClass)
+UBlueprintPropertyNodeSpawner* UBlueprintPropertyNodeSpawner::Create(UProperty const* const Property, UObject* Outer/* = nullptr*/, TSubclassOf<UEdGraphNode> const NodeClass/* = nullptr*/)
 {
 	check(Property != nullptr);
 
-	UBlueprintPropertyNodeSpawner* NodeSpawner = NewObject<UBlueprintPropertyNodeSpawner>(GetTransientPackage());
+	if (Outer == nullptr)
+	{
+		Outer = GetTransientPackage();
+	}
+
+	UBlueprintPropertyNodeSpawner* NodeSpawner = NewObject<UBlueprintPropertyNodeSpawner>(Outer);
 	NodeSpawner->Property  = Property;
 	NodeSpawner->NodeClass = NodeClass;
 
diff --git a/Engine/Source/Editor/BlueprintGraph/Public/BlueprintNodeSpawner.h b/Engine/Source/Editor/BlueprintGraph/Public/BlueprintNodeSpawner.h
--- a/Engine/Source/Editor/BlueprintGraph/Public/BlueprintNodeSpawner.h
+++ b/Engine/Source/Editor/BlueprintGraph/Public/BlueprintNodeSpawner.h
@@ -31,17 +31,18 @@ public:
 	 * the supplied parameters.
 	 *
 	 * @param  NodeClass                The node type that you want the spawner to spawn.
+	 * @param  Outer					Optional outer for the new spawner (if left null, the transient package will be used).
 	 * @param  CustomizeNodeDelegate    A delegate to perform specialized node setup post-spawn.
 	 * @return A newly allocated instance of this class.
 	 */
-	static UBlueprintNodeSpawner* Create(UClass* const NodeClass, FCustomizeNodeDelegate CustomizeNodeDelegate = FCustomizeNodeDelegate());
+	static UBlueprintNodeSpawner* Create(TSubclassOf<UEdGraphNode> const NodeClass, UObject* Outer = nullptr, FCustomizeNodeDelegate CustomizeNodeDelegate = FCustomizeNodeDelegate());
 	
 	/**
 	 * Holds the class of node to spawn. May be null for sub-classes that know
 	 * specifically what node types to spawn (if null for an instance of this 
 	 * class, then nothing will be spawned).
 	 */
-	UClass* NodeClass;
+	TSubclassOf<UEdGraphNode> NodeClass;
 	
 	/**
 	 * A delegate to perform specialized node setup post-spawn (so we don't have
