#ue4 Fix for enum nodes not preloading their enums

diff --git a/Engine/Source/Editor/BlueprintGraph/Private/K2Node_GetNumEnumEntries.cpp b/Engine/Source/Editor/BlueprintGraph/Private/K2Node_GetNumEnumEntries.cpp
--- a/Engine/Source/Editor/BlueprintGraph/Private/K2Node_GetNumEnumEntries.cpp
+++ b/Engine/Source/Editor/BlueprintGraph/Private/K2Node_GetNumEnumEntries.cpp
@@ -51,6 +51,12 @@ void UK2Node_GetNumEnumEntries::ExpandNode(class FKismetCompilerContext& Compile
 			return;
 		}
 
+		// Force the enum to load its values if it hasn't already
+		if (Enum->HasAnyFlags(RF_NeedLoad))
+		{
+			Enum->GetLinker()->Preload(Enum);
+		}
+
 		const UEdGraphSchema_K2* Schema = CompilerContext.GetSchema();
 		check(NULL != Schema);
 
