Allowing for external variables from the palette (from VariableDragDropAction, for the new BP menu system).

diff --git a/Engine/Source/Editor/BlueprintGraph/Classes/BlueprintPropertyNodeSpawner.h b/Engine/Source/Editor/BlueprintGraph/Classes/BlueprintPropertyNodeSpawner.h
--- a/Engine/Source/Editor/BlueprintGraph/Classes/BlueprintPropertyNodeSpawner.h
+++ b/Engine/Source/Editor/BlueprintGraph/Classes/BlueprintPropertyNodeSpawner.h
@@ -24,10 +24,25 @@ public:
 	 *
 	 * @param  Property		The property you want assigned to spawned nodes.
 	 * @param  Outer		Optional outer for the new spawner (if left null, the transient package will be used).
-	 * @param  NodeClass	The type of node you want spawned (can't distinct a getter from a setter from the property alone).
 	 * @return A newly allocated instance of this class.
 	 */
-	static UBlueprintPropertyNodeSpawner* Create(UProperty const* const Property, UObject* Outer = nullptr, TSubclassOf<UEdGraphNode> const NodeClass = nullptr);
+	static UBlueprintPropertyNodeSpawner* Create(UProperty const* const Property, UObject* Outer = nullptr);
+
+	/**
+	 * Templatized method for creating new UBlueprintPropertyNodeSpawners for 
+	 * a specific node class.
+	 *
+	 * @param  Property		The property you want assigned to spawned nodes.
+	 * @param  Outer		Optional outer for the new spawner (if left null, the transient package will be used).
+	 * @return A newly allocated instance of this class (setup to spawn 'NodeType' nodes).
+	 */
+	template<class NodeType>
+	static UBlueprintPropertyNodeSpawner* Create(UProperty const* const Property, UObject* Outer = nullptr)
+	{
+		UBlueprintPropertyNodeSpawner* NewSpawner = Create(Property, Outer);
+		NewSpawner->NodeClass = NodeType::StaticClass();
+		return NewSpawner;
+	}
 
 	// UBlueprintNodeSpawner interface
 	virtual UEdGraphNode* Invoke(UEdGraph* ParentGraph) const override;
diff --git a/Engine/Source/Editor/BlueprintGraph/Private/BlueprintActionDatabase.cpp b/Engine/Source/Editor/BlueprintGraph/Private/BlueprintActionDatabase.cpp
--- a/Engine/Source/Editor/BlueprintGraph/Private/BlueprintActionDatabase.cpp
+++ b/Engine/Source/Editor/BlueprintGraph/Private/BlueprintActionDatabase.cpp
@@ -96,17 +96,6 @@ namespace FBlueprintNodeSpawnerFactory
 	 * @return A new node-spawner, setup to spawn a UEdGraphNode_Comment.
 	 */
 	static UBlueprintNodeSpawner* MakeCommentNodeSpawner();
-	
-	/**
-	 * A templatized method which constructs a node-spawner for various delegate
-	 * node types (any node with a UMulticastDelegateProperty). Takes the
-	 * specified property and assigns it to the node post-spawn.
-	 *
-	 * @param  Property	 The delegate you want set for the spawned node.
-	 * @return A new node-spawner, setup to spawn some delegate node (defined by the DelegateNodeType template param).
-	 */
-	template<class DelegateNodeType>
-	static UBlueprintNodeSpawner* MakeDelegateNodeSpawner(UMulticastDelegateProperty* Property);
 };
 
 //------------------------------------------------------------------------------
@@ -241,28 +230,6 @@ static UBlueprintNodeSpawner* FBlueprintNodeSpawnerFactory::MakeCommentNodeSpawn
 	return NodeSpawner;
 }
 
-//------------------------------------------------------------------------------
-template<class DelegateNodeType>
-static UBlueprintNodeSpawner* FBlueprintNodeSpawnerFactory::MakeDelegateNodeSpawner(UMulticastDelegateProperty* Property)
-{
-	UBlueprintNodeSpawner* NodeSpawner = UBlueprintPropertyNodeSpawner::Create(Property, GetTransientPackage(), DelegateNodeType::StaticClass());
-	check(NodeSpawner != nullptr);
-	
-	auto CustomizeDelegateNodeLambda = [](UEdGraphNode* NewNode, bool bIsTemplateNode, UMulticastDelegateProperty* Property)
-	{
-		DelegateNodeType* DelegateNode = CastChecked<DelegateNodeType>(NewNode);
-		
-		UBlueprint* Blueprint = DelegateNode->GetBlueprint();
-		check(Blueprint != nullptr);
-		bool const bIsSelfContext = Blueprint->GeneratedClass->IsChildOf(Property->GetOwnerClass());
-		
-		DelegateNode->SetFromProperty(Property, bIsSelfContext);
-	};
-	NodeSpawner->CustomizeNodeDelegate = UBlueprintNodeSpawner::FCustomizeNodeDelegate::CreateStatic(CustomizeDelegateNodeLambda, Property);
-	
-	return NodeSpawner;
-}
-
 /*******************************************************************************
  * Static FBlueprintActionDatabase Helpers
  ******************************************************************************/
@@ -419,6 +386,7 @@ namespace FBlueprintActionDatabaseImpl
 //------------------------------------------------------------------------------
 void FBlueprintActionDatabaseImpl::FActionList::Add(UBlueprintNodeSpawner* NodeSpawner)
 {
+	check(NodeSpawner != nullptr);
 	check(NodeSpawner->GetOuter() == GetTransientPackage());
 	// since this spawner's outer is the transient package, we want to mark it
 	// root so that it doesn't get GC'd (we have to be careful and remove this 
@@ -495,7 +463,7 @@ static void FBlueprintActionDatabaseImpl::AddClassPropertyActions(UClass const*
 			UMulticastDelegateProperty* DelegateProperty = CastChecked<UMulticastDelegateProperty>(Property);
 			if (DelegateProperty->HasAnyPropertyFlags(CPF_BlueprintAssignable))
 			{
-				UBlueprintNodeSpawner* AddSpawner = MakeDelegateNodeSpawner<UK2Node_AddDelegate>(DelegateProperty);
+				UBlueprintNodeSpawner* AddSpawner = UBlueprintPropertyNodeSpawner::Create<UK2Node_AddDelegate>(DelegateProperty);
 				ActionListOut.Add(AddSpawner);
 				
 				// @TODO: account for: GetEventDispatcherNodesForClass() - FEdGraphSchemaAction_K2AssignDelegate
@@ -503,13 +471,13 @@ static void FBlueprintActionDatabaseImpl::AddClassPropertyActions(UClass const*
 			
 			if (DelegateProperty->HasAnyPropertyFlags(CPF_BlueprintCallable))
 			{
-				UBlueprintNodeSpawner* CallSpawner = MakeDelegateNodeSpawner<UK2Node_CallDelegate>(DelegateProperty);
+				UBlueprintNodeSpawner* CallSpawner = UBlueprintPropertyNodeSpawner::Create<UK2Node_CallDelegate>(DelegateProperty);
 				ActionListOut.Add(CallSpawner);
 			}
 			
-			UBlueprintNodeSpawner* RemoveSpawner = MakeDelegateNodeSpawner<UK2Node_RemoveDelegate>(DelegateProperty);
+			UBlueprintNodeSpawner* RemoveSpawner = UBlueprintPropertyNodeSpawner::Create<UK2Node_RemoveDelegate>(DelegateProperty);
 			ActionListOut.Add(RemoveSpawner);
-			UBlueprintNodeSpawner* ClearSpawner  = MakeDelegateNodeSpawner<UK2Node_ClearDelegate>(DelegateProperty);
+			UBlueprintNodeSpawner* ClearSpawner  = UBlueprintPropertyNodeSpawner::Create<UK2Node_ClearDelegate>(DelegateProperty);
 			ActionListOut.Add(ClearSpawner);
 
 			// @TODO: AddBoundEventActionsForClass()
@@ -518,10 +486,10 @@ static void FBlueprintActionDatabaseImpl::AddClassPropertyActions(UClass const*
  		}
 		else
 		{
-			UPackage* SpawnerOuter = GetTransientPackage();
-			UBlueprintPropertyNodeSpawner* GetterSpawner = UBlueprintPropertyNodeSpawner::Create(Property, SpawnerOuter, UK2Node_VariableGet::StaticClass());
+			
+			UBlueprintPropertyNodeSpawner* GetterSpawner = UBlueprintPropertyNodeSpawner::Create<UK2Node_VariableGet>(Property);
 			ActionListOut.Add(GetterSpawner);
-			UBlueprintPropertyNodeSpawner* SetterSpawner = UBlueprintPropertyNodeSpawner::Create(Property, SpawnerOuter, UK2Node_VariableSet::StaticClass());
+			UBlueprintPropertyNodeSpawner* SetterSpawner = UBlueprintPropertyNodeSpawner::Create<UK2Node_VariableSet>(Property);
 			ActionListOut.Add(SetterSpawner);
 		}
 	}
diff --git a/Engine/Source/Editor/BlueprintGraph/Private/BlueprintPropertyNodeSpawner.cpp b/Engine/Source/Editor/BlueprintGraph/Private/BlueprintPropertyNodeSpawner.cpp
--- a/Engine/Source/Editor/BlueprintGraph/Private/BlueprintPropertyNodeSpawner.cpp
+++ b/Engine/Source/Editor/BlueprintGraph/Private/BlueprintPropertyNodeSpawner.cpp
@@ -5,6 +5,7 @@
 #include "K2Node_Variable.h"
 #include "EditorStyleSettings.h"	// for bShowFriendlyNames
 #include "ObjectEditorUtils.h"		// for GetCategory()
+#include "BlueprintEditorUtils.h"	// for FindBlueprintForNodeChecked()
 
 #define LOCTEXT_NAMESPACE "BlueprintPropertyNodeSpawner"
 
@@ -13,7 +14,7 @@
  ******************************************************************************/
 
 //------------------------------------------------------------------------------
-UBlueprintPropertyNodeSpawner* UBlueprintPropertyNodeSpawner::Create(UProperty const* const Property, UObject* Outer/* = nullptr*/, TSubclassOf<UEdGraphNode> const NodeClass/* = nullptr*/)
+UBlueprintPropertyNodeSpawner* UBlueprintPropertyNodeSpawner::Create(UProperty const* const Property, UObject* Outer/* = nullptr*/)
 {
 	check(Property != nullptr);
 
@@ -24,8 +25,6 @@ UBlueprintPropertyNodeSpawner* UBlueprintPropertyNodeSpawner::Create(UProperty c
 
 	UBlueprintPropertyNodeSpawner* NodeSpawner = NewObject<UBlueprintPropertyNodeSpawner>(Outer);
 	NodeSpawner->Property  = Property;
-	NodeSpawner->NodeClass = NodeClass;
-
 	return NodeSpawner;
 }
 
@@ -41,17 +40,22 @@ UEdGraphNode* UBlueprintPropertyNodeSpawner::Invoke(UEdGraph* ParentGraph) const
 {
 	// if a NodeClass was set, then this should spawn something
 	UEdGraphNode* NewNode = Super::Invoke(ParentGraph);
-	check(Property != nullptr);	
-	
-	if (UK2Node_Variable* VarNode = Cast<UK2Node_Variable>(NewNode))
-	{
-		UBlueprint* Blueprint = VarNode->GetBlueprint();
+	check(Property != nullptr);
 
-		UClass* BlueprintClass = Blueprint->GeneratedClass;
-		UClass* PropertyClass = Cast<UClass>(Property->GetOuterUField());
+	if (NewNode != nullptr)
+	{
+		UBlueprint* Blueprint = FBlueprintEditorUtils::FindBlueprintForNodeChecked(NewNode);
+		UClass* OwnerClass = Property->GetOwnerClass();
+		bool const bIsSelfContext = Blueprint->SkeletonGeneratedClass->IsChildOf(OwnerClass);
 
-		bool const bIsSelfContext = BlueprintClass->IsChildOf(PropertyClass);
-		VarNode->SetFromProperty(Property, bIsSelfContext);
+		if (UK2Node_Variable* VarNode = Cast<UK2Node_Variable>(NewNode))
+		{
+			VarNode->SetFromProperty(Property, bIsSelfContext);
+		}
+		else if (UK2Node_BaseMCDelegate* DelegateNode = Cast<UK2Node_BaseMCDelegate>(NewNode))
+		{
+			DelegateNode->SetFromProperty(Property, bIsSelfContext);
+		}
 	}
 	// else, without a NodeClass explicitly specified, we don't know what kind  
 	// of property node to spawn (for example: a getter versus a setter)... 
diff --git a/Engine/Source/Editor/BlueprintGraph/Private/EdGraphSchema_K2.cpp b/Engine/Source/Editor/BlueprintGraph/Private/EdGraphSchema_K2.cpp
--- a/Engine/Source/Editor/BlueprintGraph/Private/EdGraphSchema_K2.cpp
+++ b/Engine/Source/Editor/BlueprintGraph/Private/EdGraphSchema_K2.cpp
@@ -4762,13 +4762,17 @@ void UEdGraphSchema_K2::ConfigureVarNode(UK2Node_Variable* InVarNode, FName InVa
 	{
 		InVarNode->VariableReference.SetSelfMember(InVariableName);
 	}
-	else if (InVariableSource->IsA(UStruct::StaticClass()))
+	else if (InVariableSource->IsA(UClass::StaticClass()))
 	{
-		InVarNode->VariableReference.SetLocalMember(InVariableName, InVariableSource, FBlueprintEditorUtils::FindLocalVariableGuidByName(InTargetBlueprint, InVariableSource, InVariableName));
+		InVarNode->VariableReference.SetExternalMember(InVariableName, CastChecked<UClass>(InVariableSource));
 	}
 	else
 	{
-		InVarNode->VariableReference.SetExternalMember(InVariableName, CastChecked<UClass>(InVariableSource));
+		FGuid LocalVarGuid = FBlueprintEditorUtils::FindLocalVariableGuidByName(InTargetBlueprint, InVariableSource, InVariableName);
+		if (LocalVarGuid.IsValid())
+		{
+			InVarNode->VariableReference.SetLocalMember(InVariableName, InVariableSource, LocalVarGuid);
+		}
 	}
 }
 
