Separate Control and Command keys on Mac; some editor shortcuts remapped to use Ctrl to free Cmd for Mac-specific shortcuts (Cmd+W, Cmd+M and Cmd+H)

diff --git a/Engine/Source/Developer/GameplayDebugger/Private/GameplayDebuggingControllerComponent.cpp b/Engine/Source/Developer/GameplayDebugger/Private/GameplayDebuggingControllerComponent.cpp
--- a/Engine/Source/Developer/GameplayDebugger/Private/GameplayDebuggingControllerComponent.cpp
+++ b/Engine/Source/Developer/GameplayDebugger/Private/GameplayDebuggingControllerComponent.cpp
@@ -127,7 +127,7 @@ void UGameplayDebuggingControllerComponent::BindActivationKeys()
 	if (PlayerOwner.IsValid() && PlayerOwner->InputComponent && PlayerOwner->PlayerInput)
 	{
 		// find current activation key used for 'EnableGDT' binding
-		FInputChord ActivationKey(EKeys::Quote, false, false, false);
+		FInputChord ActivationKey(EKeys::Quote, false, false, false, false);
 		for (uint32 BindIndex = 0; BindIndex < (uint32)PlayerOwner->PlayerInput->DebugExecBindings.Num(); BindIndex++)
 		{
 			if (PlayerOwner->PlayerInput->DebugExecBindings[BindIndex].Command == TEXT("cheat EnableGDT"))
diff --git a/Engine/Source/Developer/Profiler/Private/ProfilerCommands.cpp b/Engine/Source/Developer/Profiler/Private/ProfilerCommands.cpp
--- a/Engine/Source/Developer/Profiler/Private/ProfilerCommands.cpp
+++ b/Engine/Source/Developer/Profiler/Private/ProfilerCommands.cpp
@@ -37,8 +37,13 @@ void FProfilerCommands::RegisterCommands()
 
 	UI_COMMAND( ProfilerManager_Save, "Save", "Saves all collected data to file or files", EUserInterfaceActionType::Button, FInputGesture( EModifierKey::Control, EKeys::S ) );
 	UI_COMMAND( StatsProfiler, "Statistics", "Enables the Stats Profiler", EUserInterfaceActionType::ToggleButton, FInputGesture( EModifierKey::Control, EKeys::P ) );
+#if PLATFORM_MAC
+	UI_COMMAND( MemoryProfiler, "Memory", "Enables the Memory Profiler", EUserInterfaceActionType::ToggleButton, FInputGesture( EModifierKey::Command, EKeys::M ) );
+	UI_COMMAND( FPSChart, "FPS Chart", "Shows the FPS Chart", EUserInterfaceActionType::ToggleButton, FInputGesture( EModifierKey::Command, EKeys::H ) );
+#else
 	UI_COMMAND( MemoryProfiler, "Memory", "Enables the Memory Profiler", EUserInterfaceActionType::ToggleButton, FInputGesture( EModifierKey::Control, EKeys::M ) );
 	UI_COMMAND( FPSChart, "FPS Chart", "Shows the FPS Chart", EUserInterfaceActionType::ToggleButton, FInputGesture( EModifierKey::Control, EKeys::H ) );
+#endif
 
 	UI_COMMAND( OpenSettings, "Settings", "Opens the settings for the profiler", EUserInterfaceActionType::Button, FInputGesture( EModifierKey::Control, EKeys::O ) );
 
diff --git a/Engine/Source/Editor/BlueprintGraph/Classes/K2Node_InputKey.h b/Engine/Source/Editor/BlueprintGraph/Classes/K2Node_InputKey.h
--- a/Engine/Source/Editor/BlueprintGraph/Classes/K2Node_InputKey.h
+++ b/Engine/Source/Editor/BlueprintGraph/Classes/K2Node_InputKey.h
@@ -26,18 +26,22 @@ class UK2Node_InputKey : public UK2Node
 	UPROPERTY(EditAnywhere, Category="Input")
 	uint32 bOverrideParentBinding:1;
 
-	// Does this binding requires the control key to be held
+	// Does this binding require the control key to be held
 	UPROPERTY(EditAnywhere, Category="Modifier")
 	uint32 bControl:1;
 
-	// Does this binding requires the alt key to be held
+	// Does this binding require the alt key to be held
 	UPROPERTY(EditAnywhere, Category="Modifier")
 	uint32 bAlt:1;
 
-	// Does this binding requires the shift key to be held
+	// Does this binding require the shift key to be held
 	UPROPERTY(EditAnywhere, Category="Modifier")
 	uint32 bShift:1;
 
+	// Does this binding require the command key to be held
+	UPROPERTY(EditAnywhere, Category="Modifier")
+	uint32 bCommand:1;
+
 	// Begin UObject interface
 	virtual void PostLoad() override;
 	// End UObject interface
diff --git a/Engine/Source/Editor/BlueprintGraph/Private/K2Node_InputKey.cpp b/Engine/Source/Editor/BlueprintGraph/Private/K2Node_InputKey.cpp
--- a/Engine/Source/Editor/BlueprintGraph/Private/K2Node_InputKey.cpp
+++ b/Engine/Source/Editor/BlueprintGraph/Private/K2Node_InputKey.cpp
@@ -44,34 +44,66 @@ FLinearColor UK2Node_InputKey::GetNodeTitleColor() const
 
 FName UK2Node_InputKey::GetModifierName() const
 {
-	if ( bControl && !bAlt && !bShift )
+	if ( bControl && !bAlt && !bShift && !bCommand )
 	{
 		return FName("Ctrl");
 	}
-	else if ( bControl && bAlt && !bShift )
+	else if ( bControl && bAlt && !bShift && !bCommand )
 	{
 		return FName("Ctrl+Alt");
 	}
-	else if ( bControl && !bAlt && bShift )
+	else if ( bControl && !bAlt && bShift && !bCommand )
 	{
 		return FName("Ctrl+Shift");
 	}
-	else if ( bControl && bAlt && bShift )
+	else if ( bControl && !bAlt && !bShift && bCommand )
+	{
+		return FName("Ctrl+Cmd");
+	}
+	else if ( bControl && bAlt && bShift && !bCommand )
 	{
 		return FName("Ctrl+Alt+Shift");
 	}
-	else if ( !bControl && bAlt && !bShift )
+	else if ( bControl && bAlt && !bShift && bCommand )
+	{
+		return FName("Ctrl+Cmd+Alt");
+	}
+	else if ( bControl && !bAlt && bShift && bCommand )
+	{
+		return FName("Ctrl+Cmd+Shift");
+	}
+	else if ( !bControl && bAlt && bShift && bCommand )
+	{
+		return FName("Cmd+Alt+Shift");
+	}
+	else if ( bControl && bAlt && bShift && bCommand )
+	{
+		return FName("Ctrl+Cmd+Alt+Shift");
+	}
+	else if ( !bControl && bAlt && !bShift && !bCommand )
 	{
 		return FName("Alt");
 	}
-	else if ( !bControl && bAlt && bShift )
+	else if ( !bControl && bAlt && bShift && !bCommand )
 	{
 		return FName("Alt+Shift");
 	}
-	else if ( !bControl && !bAlt && bShift )
+	else if ( !bControl && bAlt && !bShift && bCommand )
+	{
+		return FName("Cmd+Alt");
+	}
+	else if ( !bControl && !bAlt && bShift && bCommand )
+	{
+		return FName("Cmd+Shift");
+	}
+	else if ( !bControl && !bAlt && bShift && !bCommand )
 	{
 		return FName("Shift");
 	}
+	else if ( !bControl && !bAlt && !bShift && bCommand )
+	{
+		return FName("Cmd");
+	}
 
 	return NAME_None;
 }
@@ -79,50 +111,110 @@ FName UK2Node_InputKey::GetModifierName() const
 FText UK2Node_InputKey::GetModifierText() const
 {
 	//@todo This should be unified with other places in the editor [10/11/2013 justin.sargent]
-	if ( bControl && !bAlt && !bShift )
+	if ( bControl && !bAlt && !bShift && !bCommand )
 	{
 #if PLATFORM_MAC
-		return NSLOCTEXT("UK2Node_InputKey", "KeyName_Command", "Command");
+		return NSLOCTEXT("UK2Node_InputKey", "KeyName_Command", "Cmd");
 #else
 		return NSLOCTEXT("UK2Node_InputKey", "KeyName_Control", "Ctrl");
 #endif
 	}
-	else if ( bControl && bAlt && !bShift )
+	else if ( bControl && bAlt && !bShift && !bCommand )
 	{
 #if PLATFORM_MAC
-		return NSLOCTEXT("UK2Node_InputKey", "KeyName_Command + KeyName_Alt", "Command+Alt");
+		return NSLOCTEXT("UK2Node_InputKey", "KeyName_Command + KeyName_Alt", "Cmd+Alt");
 #else
 		return NSLOCTEXT("UK2Node_InputKey", "KeyName_Control + KeyName_Alt", "Ctrl+Alt");
 #endif
 	}
-	else if ( bControl && !bAlt && bShift )
+	else if ( bControl && !bAlt && bShift && !bCommand )
+	{
+#if PLATFORM_MAC
+		return NSLOCTEXT("UK2Node_InputKey", "KeyName_Command + KeyName_Shift", "Cmd+Shift");
+#else
+		return NSLOCTEXT("UK2Node_InputKey", "KeyName_Control + KeyName_Shift", "Ctrl+Shift");
+#endif
+	}
+	else if ( bControl && !bAlt && !bShift && bCommand )
 	{
 #if PLATFORM_MAC
-		return NSLOCTEXT("UK2Node_InputKey", "KeyName_Command + KeyName_Shift", "Command+Shift");
+		return NSLOCTEXT("UK2Node_InputKey", "KeyName_Command + KeyName_Control", "Cmd+Ctrl");
+#else
+		return NSLOCTEXT("UK2Node_InputKey", "KeyName_Control + KeyName_Command", "Ctrl+Cmd");
+#endif
+	}
+	else if ( !bControl && bAlt && !bShift && bCommand )
+	{
+#if PLATFORM_MAC
+		return NSLOCTEXT("UK2Node_InputKey", "KeyName_Control + KeyName_Alt", "Ctrl+Alt");
 #else
+		return NSLOCTEXT("UK2Node_InputKey", "KeyName_Command + KeyName_Alt", "Cmd+Alt");
+#endif
+	}
+	else if ( !bControl && !bAlt && bShift && bCommand )
+	{
+#if PLATFORM_MAC
 		return NSLOCTEXT("UK2Node_InputKey", "KeyName_Control + KeyName_Shift", "Ctrl+Shift");
+#else
+		return NSLOCTEXT("UK2Node_InputKey", "KeyName_Command + KeyName_Shift", "Cmd+Shift");
 #endif
 	}
-	else if ( bControl && bAlt && bShift )
+	else if ( bControl && bAlt && bShift && !bCommand )
 	{
 #if PLATFORM_MAC
-		return NSLOCTEXT("UK2Node_InputKey", "KeyName_Command + KeyName_Alt + KeyName_Shift", "Command+Alt+Shift");
+		return NSLOCTEXT("UK2Node_InputKey", "KeyName_Command + KeyName_Alt + KeyName_Shift", "Cmd+Alt+Shift");
 #else
 		return NSLOCTEXT("UK2Node_InputKey", "KeyName_Control + KeyName_Alt + KeyName_Shift", "Ctrl+Alt+Shift");
 #endif
 	}
-	else if ( !bControl && bAlt && !bShift )
+	else if ( bControl && bAlt && !bShift && bCommand )
+	{
+#if PLATFORM_MAC
+		return NSLOCTEXT("UK2Node_InputKey", "KeyName_Command + KeyName_Control + KeyName_Alt", "Cmd+Ctrl+Alt");
+#else
+		return NSLOCTEXT("UK2Node_InputKey", "KeyName_Control + KeyName_Command + KeyName_Alt", "Ctrl+Cmd+Alt");
+#endif
+	}
+	else if ( bControl && !bAlt && bShift && bCommand )
+	{
+#if PLATFORM_MAC
+		return NSLOCTEXT("UK2Node_InputKey", "KeyName_Command + KeyName_Control + KeyName_Shift", "Cmd+Ctrl+Shift");
+#else
+		return NSLOCTEXT("UK2Node_InputKey", "KeyName_Control + KeyName_Command + KeyName_Shift", "Ctrl+Cmd+Shift");
+#endif
+	}
+	else if ( bControl && bAlt && bShift && bCommand )
+	{
+#if PLATFORM_MAC
+		return NSLOCTEXT("UK2Node_InputKey", "KeyName_Command + KeyName_Control + KeyName_Alt + KeyName_Shift", "Cmd+Ctrl+Alt+Shift");
+#else
+		return NSLOCTEXT("UK2Node_InputKey", "KeyName_Control + KeyName_Command + KeyName_Alt + KeyName_Shift", "Ctrl+Cmd+Alt+Shift");
+#endif
+	}
+	else if ( !bControl && bAlt && !bShift && !bCommand )
 	{
 		return NSLOCTEXT("UK2Node_InputKey", "KeyName_Alt", "Alt");
 	}
-	else if ( !bControl && bAlt && bShift )
+	else if ( !bControl && bAlt && bShift && !bCommand )
 	{
 		return NSLOCTEXT("UK2Node_InputKey", "KeyName_Alt + KeyName_Shift", "Alt+Shift");
 	}
-	else if ( !bControl && !bAlt && bShift )
+	else if ( !bControl && bAlt && !bShift && bCommand )
+	{
+#if PLATFORM_MAC
+		return NSLOCTEXT("UK2Node_InputKey", "KeyName_Control + KeyName_Alt", "Ctrl+Alt");
+#else
+		return NSLOCTEXT("UK2Node_InputKey", "KeyName_Command + KeyName_Alt", "Cmd+Alt");
+#endif
+	}
+	else if ( !bControl && !bAlt && bShift && !bCommand )
 	{
 		return NSLOCTEXT("UK2Node_InputKey", "KeyName_Shift", "Shift");
 	}
+	else if ( !bControl && !bAlt && !bShift && bCommand )
+	{
+		return NSLOCTEXT("UK2Node_InputKey", "KeyName_Command", "Cmd");
+	}
 
 	return FText::GetEmpty();
 }
diff --git a/Engine/Source/Editor/InputBindingEditor/Private/Widgets/SGestureEditor.cpp b/Engine/Source/Editor/InputBindingEditor/Private/Widgets/SGestureEditor.cpp
--- a/Engine/Source/Editor/InputBindingEditor/Private/Widgets/SGestureEditor.cpp
+++ b/Engine/Source/Editor/InputBindingEditor/Private/Widgets/SGestureEditor.cpp
@@ -61,6 +61,7 @@ FReply SGestureEditor::OnKeyDown( const FGeometry& MyGeometry, const FKeyboardEv
 		EditingInputGesture.bCtrl = InKeyboardEvent.IsControlDown();
 		EditingInputGesture.bAlt = InKeyboardEvent.IsAltDown();
 		EditingInputGesture.bShift = InKeyboardEvent.IsShiftDown();
+		EditingInputGesture.bCmd = InKeyboardEvent.IsCommandDown();
 
 		LoadText();
 		//@todo checking the length of localized string is not valid, at the very least in this manner [10/11/2013 justin.sargent]
diff --git a/Engine/Source/Editor/LevelEditor/Private/LevelEditorActions.cpp b/Engine/Source/Editor/LevelEditor/Private/LevelEditorActions.cpp
--- a/Engine/Source/Editor/LevelEditor/Private/LevelEditorActions.cpp
+++ b/Engine/Source/Editor/LevelEditor/Private/LevelEditorActions.cpp
@@ -2700,7 +2700,11 @@ void FLevelEditorCommands::RegisterCommands()
 	UI_COMMAND( MergeActors, "Create Mesh Proxy...", "Harvest geometry from selected actors and merge them into single mesh", EUserInterfaceActionType::Button, FInputGesture() );
 	UI_COMMAND( MergeActorsByMaterials, "Merge Actors...", "Harvest geometry from selected actors and merge grouping them by materials", EUserInterfaceActionType::Button, FInputGesture() );
 
+#if PLATFORM_MAC
+	UI_COMMAND( ShowAll, "Show All Actors", "Shows all actors", EUserInterfaceActionType::Button, FInputGesture( EModifierKey::Command, EKeys::H ) );
+#else
 	UI_COMMAND( ShowAll, "Show All Actors", "Shows all actors", EUserInterfaceActionType::Button, FInputGesture( EModifierKey::Control, EKeys::H ) );
+#endif
 	UI_COMMAND( ShowSelectedOnly, "Show Only Selected", "Shows only the selected actors", EUserInterfaceActionType::Button, FInputGesture() );
 	UI_COMMAND( ShowSelected, "Show Selected", "Shows the selected actors", EUserInterfaceActionType::Button, FInputGesture( EModifierKey::Shift, EKeys::H ) );
 	UI_COMMAND( HideSelected, "Hide Selected", "Hides the selected actors", EUserInterfaceActionType::Button, FInputGesture( EKeys::H ) );
@@ -2763,7 +2767,11 @@ void FLevelEditorCommands::RegisterCommands()
 	UI_COMMAND( KeepSimulationChanges, "Keep Simulation Changes", "Saves the changes made to this actor in Simulate mode to the actor's default state.", EUserInterfaceActionType::Button, FInputGesture( EKeys::K ) );
 
 	UI_COMMAND( MakeActorLevelCurrent, "Make Selected Actor's Level Current", "Makes the selected actors level the current level", EUserInterfaceActionType::Button, FInputGesture( EKeys::M ) );
+#if PLATFORM_MAC
+	UI_COMMAND( MoveSelectedToCurrentLevel, "Move Selection to Current Level", "Moves the selected actors to the current level", EUserInterfaceActionType::Button, FInputGesture( EModifierKey::Command, EKeys::M ) );
+#else
 	UI_COMMAND( MoveSelectedToCurrentLevel, "Move Selection to Current Level", "Moves the selected actors to the current level", EUserInterfaceActionType::Button, FInputGesture( EModifierKey::Control, EKeys::M ) );
+#endif
 	UI_COMMAND( FindLevelsInLevelBrowser, "Find Levels in Level Browser", "Finds the selected actors level in the level browser", EUserInterfaceActionType::Button, FInputGesture() );
 	UI_COMMAND( AddLevelsToSelection, "Add Levels to Selection", "Adds the selected actors levels  to the current level browser selection", EUserInterfaceActionType::Button, FInputGesture() );
 	UI_COMMAND( RemoveLevelsFromSelection, "Remove Levels from Selection", "Removes the selected actors levels from the current level browser selection", EUserInterfaceActionType::Button, FInputGesture() );
diff --git a/Engine/Source/Editor/Matinee/Private/MatineeActions.cpp b/Engine/Source/Editor/Matinee/Private/MatineeActions.cpp
--- a/Engine/Source/Editor/Matinee/Private/MatineeActions.cpp
+++ b/Engine/Source/Editor/Matinee/Private/MatineeActions.cpp
@@ -171,7 +171,11 @@ void FMatineeCommands::RegisterCommands()
 	UI_COMMAND(SplitAnimKey,"Split Anim Key","",EUserInterfaceActionType::Button,FInputGesture(EModifierKey::None, EKeys::R));
 	UI_COMMAND(MoveActiveUp,"Move Active Track/Group Up","",EUserInterfaceActionType::Button,FInputGesture(EModifierKey::None, EKeys::Up));
 	UI_COMMAND(MoveActiveDown,"Move Active Track/Group Down","",EUserInterfaceActionType::Button,FInputGesture(EModifierKey::None, EKeys::Down));
+#if PLATFORM_MAC
+	UI_COMMAND(DuplicateSelectedKeys,"Duplicate Selected Keys","",EUserInterfaceActionType::Button,FInputGesture(EModifierKey::Command, EKeys::W));
+#else
 	UI_COMMAND(DuplicateSelectedKeys,"Duplicate Selected Keys","",EUserInterfaceActionType::Button,FInputGesture(EModifierKey::Control, EKeys::W));
+#endif
 	UI_COMMAND(CropAnimationBeginning,"Crop Animation Start","",EUserInterfaceActionType::Button,FInputGesture(EModifierKey::Control, EKeys::I));
 	UI_COMMAND(CropAnimationEnd,"Crop Animation End","",EUserInterfaceActionType::Button,FInputGesture(EModifierKey::Control, EKeys::O));
 	UI_COMMAND(ChangeKeyInterpModeAuto,"Change Key Interp Mode Auto","",EUserInterfaceActionType::Button,FInputGesture(EModifierKey::None, EKeys::One));
diff --git a/Engine/Source/Editor/Matinee/Private/MatineeViewportClient.cpp b/Engine/Source/Editor/Matinee/Private/MatineeViewportClient.cpp
--- a/Engine/Source/Editor/Matinee/Private/MatineeViewportClient.cpp
+++ b/Engine/Source/Editor/Matinee/Private/MatineeViewportClient.cpp
@@ -164,6 +164,7 @@ bool FMatineeViewportClient::InputKey(FViewport* Viewport, int32 ControllerId, F
 	const bool bCtrlDown = Viewport->KeyState(EKeys::LeftControl) || Viewport->KeyState(EKeys::RightControl);
 	const bool bShiftDown = Viewport->KeyState(EKeys::LeftShift) || Viewport->KeyState(EKeys::RightShift);
 	const bool bAltDown = Viewport->KeyState(EKeys::LeftAlt) || Viewport->KeyState(EKeys::RightAlt);
+	const bool bCmdDown = Viewport->KeyState(EKeys::LeftCommand) || Viewport->KeyState(EKeys::RightCommand);
 
 	const int32 HitX = Viewport->GetMouseX();
 	const int32 HitY = Viewport->GetMouseY();
@@ -569,6 +570,7 @@ bool FMatineeViewportClient::InputKey(FViewport* Viewport, int32 ControllerId, F
 							DragData.bCtrlDown = bCtrlDown;
 							DragData.bAltDown = bAltDown;
 							DragData.bShiftDown = bShiftDown;
+							DragData.bCmdDown = bCmdDown;
 							Proxy->ClickedObject->BeginDrag(DragData);
 						}
 					}
@@ -828,7 +830,7 @@ bool FMatineeViewportClient::InputKey(FViewport* Viewport, int32 ControllerId, F
 			InterpEd->ZoomView( 1.0f / FMatinee::InterpEditor_ZoomIncrement, InterpEd->bZoomToScrubPos );
 		}
 
-		FModifierKeysState ModKeys( bShiftDown, bShiftDown, bCtrlDown, bCtrlDown, bAltDown, bAltDown);
+		FModifierKeysState ModKeys(bShiftDown, bShiftDown, bCtrlDown, bCtrlDown, bAltDown, bAltDown, bCmdDown, bCmdDown);
 		FKeyboardEvent KeyEvent(Key, ModKeys, false, 0);
 		InterpEd->ProcessCommandBindings(KeyEvent);
 	}
diff --git a/Engine/Source/Runtime/Core/Private/Mac/MacApplication.cpp b/Engine/Source/Runtime/Core/Private/Mac/MacApplication.cpp
--- a/Engine/Source/Runtime/Core/Private/Mac/MacApplication.cpp
+++ b/Engine/Source/Runtime/Core/Private/Mac/MacApplication.cpp
@@ -815,14 +815,16 @@ FModifierKeysState FMacApplication::GetModifierKeys() const
 {
 	uint32 CurrentFlags = ModifierKeysFlags;
 
-	const bool bIsLeftShiftDown		= ( CurrentFlags & ( 1 << 0 ) ) != 0;
-	const bool bIsRightShiftDown	= ( CurrentFlags & ( 1 << 1 ) ) != 0;
-	const bool bIsLeftControlDown	= ( CurrentFlags & ( 1 << 6 ) ) != 0 || ( CurrentFlags & ( 1 << 2 ) ) != 0; // Mac pretends the Command key is Ctrl
-	const bool bIsRightControlDown	= ( CurrentFlags & ( 1 << 7 ) ) != 0 || ( CurrentFlags & ( 1 << 3 ) ) != 0; // Mac pretends the Command key is Ctrl
-	const bool bIsLeftAltDown		= ( CurrentFlags & ( 1 << 4 ) ) != 0;
-	const bool bIsRightAltDown		= ( CurrentFlags & ( 1 << 5 ) ) != 0;
-
-	return FModifierKeysState( bIsLeftShiftDown, bIsRightShiftDown, bIsLeftControlDown, bIsRightControlDown, bIsLeftAltDown, bIsRightAltDown );
+	const bool bIsLeftShiftDown			= ( CurrentFlags & ( 1 << 0 ) ) != 0;
+	const bool bIsRightShiftDown		= ( CurrentFlags & ( 1 << 1 ) ) != 0;
+	const bool bIsLeftControlDown		= ( CurrentFlags & ( 1 << 6 ) ) != 0; // Mac pretends the Command key is Control
+	const bool bIsRightControlDown		= ( CurrentFlags & ( 1 << 7 ) ) != 0; // Mac pretends the Command key is Control
+	const bool bIsLeftAltDown			= ( CurrentFlags & ( 1 << 4 ) ) != 0;
+	const bool bIsRightAltDown			= ( CurrentFlags & ( 1 << 5 ) ) != 0;
+	const bool bIsLeftCommandDown		= ( CurrentFlags & ( 1 << 2 ) ) != 0; // Mac pretends the Control key is Command
+	const bool bIsRightCommandDown		= ( CurrentFlags & ( 1 << 3 ) ) != 0; // Mac pretends the Control key is Command
+
+	return FModifierKeysState( bIsLeftShiftDown, bIsRightShiftDown, bIsLeftControlDown, bIsRightControlDown, bIsLeftAltDown, bIsRightAltDown, bIsLeftCommandDown, bIsRightCommandDown );
 }
 
 FPlatformRect FMacApplication::GetWorkArea( const FPlatformRect& CurrentWindow ) const
diff --git a/Engine/Source/Runtime/Core/Private/Mac/MacPlatformMisc.cpp b/Engine/Source/Runtime/Core/Private/Mac/MacPlatformMisc.cpp
--- a/Engine/Source/Runtime/Core/Private/Mac/MacPlatformMisc.cpp
+++ b/Engine/Source/Runtime/Core/Private/Mac/MacPlatformMisc.cpp
@@ -317,7 +317,7 @@ void FMacPlatformMisc::UpdateWindowMenu()
 
 	NSMenuItem* MinimizeItem = [[[NSMenuItem alloc] initWithTitle:@"Minimize" action:@selector(miniaturize:) keyEquivalent:@"m"] autorelease];
 	NSMenuItem* ZoomItem = [[[NSMenuItem alloc] initWithTitle:@"Zoom" action:@selector(performZoom:) keyEquivalent:@""] autorelease];
-	NSMenuItem* CloseItem = [[[NSMenuItem alloc] initWithTitle:@"Close" action:@selector(performClose:) keyEquivalent:@""] autorelease];
+	NSMenuItem* CloseItem = [[[NSMenuItem alloc] initWithTitle:@"Close" action:@selector(performClose:) keyEquivalent:@"w"] autorelease];
 	NSMenuItem* BringAllToFrontItem = [[[NSMenuItem alloc] initWithTitle:@"Bring All to Front" action:@selector(arrangeInFront:) keyEquivalent:@""] autorelease];
 
 	[WindowMenu addItem:MinimizeItem];
@@ -555,16 +555,16 @@ uint32 FMacPlatformMisc::GetKeyMap( uint16* KeyCodes, FString* KeyNames, uint32
 		ADDKEYMAP( kVK_F11, TEXT("F11") );
 		ADDKEYMAP( kVK_F12, TEXT("F12") );
 
-        //Mac pretends the Command key is Ctrl
-        ADDKEYMAP( MMK_RightCommand, TEXT("RightControl") );
-        ADDKEYMAP( MMK_LeftCommand, TEXT("LeftControl") );
-        ADDKEYMAP( MMK_LeftShift, TEXT("LeftShift") );
-        ADDKEYMAP( MMK_CapsLock, TEXT("CapsLock") );
-        ADDKEYMAP( MMK_LeftAlt, TEXT("LeftAlt") );
-        ADDKEYMAP( MMK_LeftControl, TEXT("LeftControl") );
-        ADDKEYMAP( MMK_RightShift, TEXT("RightShift") );
-        ADDKEYMAP( MMK_RightAlt, TEXT("RightAlt") );
-        ADDKEYMAP( MMK_RightControl, TEXT("RightControl") );
+		// Mac pretends the Command key is Ctrl and Ctrl is Command key
+		ADDKEYMAP( MMK_RightCommand, TEXT("RightControl") );
+		ADDKEYMAP( MMK_LeftCommand, TEXT("LeftControl") );
+		ADDKEYMAP( MMK_LeftShift, TEXT("LeftShift") );
+		ADDKEYMAP( MMK_CapsLock, TEXT("CapsLock") );
+		ADDKEYMAP( MMK_LeftAlt, TEXT("LeftAlt") );
+		ADDKEYMAP( MMK_LeftControl, TEXT("LeftCommand") );
+		ADDKEYMAP( MMK_RightShift, TEXT("RightShift") );
+		ADDKEYMAP( MMK_RightAlt, TEXT("RightAlt") );
+		ADDKEYMAP( MMK_RightControl, TEXT("RightCommand") );
 	}
 
 	check(NumMappings < MaxMappings);
diff --git a/Engine/Source/Runtime/Core/Public/GenericPlatform/GenericApplication.h b/Engine/Source/Runtime/Core/Public/GenericPlatform/GenericApplication.h
--- a/Engine/Source/Runtime/Core/Public/GenericPlatform/GenericApplication.h
+++ b/Engine/Source/Runtime/Core/Public/GenericPlatform/GenericApplication.h
@@ -41,13 +41,17 @@ public:
 					    const bool bInIsLeftControlDown,
 					    const bool bInIsRightControlDown,
 					    const bool bInIsLeftAltDown,
-					    const bool bInIsRightAltDown )
+						const bool bInIsRightAltDown,
+						const bool bInIsLeftCommandDown,
+						const bool bInIsRightCommandDown )
 		: bIsLeftShiftDown( bInIsLeftShiftDown ),
 		  bIsRightShiftDown( bInIsRightShiftDown ),
 		  bIsLeftControlDown( bInIsLeftControlDown ),
 		  bIsRightControlDown( bInIsRightControlDown ),
 		  bIsLeftAltDown( bInIsLeftAltDown ),
-		  bIsRightAltDown( bInIsRightAltDown )
+		  bIsRightAltDown( bInIsRightAltDown ),
+		  bIsLeftCommandDown( bInIsLeftCommandDown ),
+		  bIsRightCommandDown( bInIsRightCommandDown )
 	{
 	}
 
@@ -141,6 +145,36 @@ public:
 	{
 		return bIsRightAltDown;
 	}
+	
+	/**
+	 * Returns true if either command key was down when this event occurred
+	 *
+	 * @return  True if command is pressed
+	 */
+	bool IsCommandDown() const
+	{
+		return bIsLeftCommandDown || bIsRightCommandDown;
+	}
+	
+	/**
+	 * Returns true if left command key was down when this event occurred
+	 *
+	 * @return  True if left command is pressed
+	 */
+	bool IsLeftCommandDown() const
+	{
+		return bIsLeftCommandDown;
+	}
+	
+	/**
+	 * Returns true if right command key was down when this event occurred
+	 *
+	 * @return  True if right command is pressed
+	 */
+	bool IsRightCommandDown() const
+	{
+		return bIsRightCommandDown;
+	}
 
 
 private:
@@ -162,6 +196,12 @@ private:
 
 	/** True if the right alt key was down when this event occurred. */
 	bool bIsRightAltDown;
+	
+	/** True if the left command key was down when this event occurred. */
+	bool bIsLeftCommandDown;
+	
+	/** True if the right command key was down when this event occurred. */
+	bool bIsRightCommandDown;
 };
 
 struct FPlatformRect
@@ -267,7 +307,7 @@ public:
 
 	virtual void* GetCapture( void ) const { return NULL; }
 
-	virtual FModifierKeysState GetModifierKeys() const  { return FModifierKeysState( false, false, false, false, false, false ); }
+	virtual FModifierKeysState GetModifierKeys() const  { return FModifierKeysState( false, false, false, false, false, false, false, false ); }
 
 	virtual void SetHighPrecisionMouseMode( const bool Enable, const TSharedPtr< FGenericWindow >& InWindow ) { };
 
diff --git a/Engine/Source/Runtime/Engine/Classes/Components/InputComponent.h b/Engine/Source/Runtime/Engine/Classes/Components/InputComponent.h
--- a/Engine/Source/Runtime/Engine/Classes/Components/InputComponent.h
+++ b/Engine/Source/Runtime/Engine/Classes/Components/InputComponent.h
@@ -86,6 +86,9 @@ struct FInputChord
 	UPROPERTY()
 	uint32 bAlt:1;
 
+	UPROPERTY()
+	uint32 bCmd:1;
+
 	enum RelationshipType
 	{
 		None,
@@ -100,14 +103,16 @@ struct FInputChord
 		: bShift(false)
 		, bCtrl(false)
 		, bAlt(false)
+		, bCmd(false)
 	{
 	}
 
-	FInputChord(const FKey InKey, const bool bInShift, const bool bInCtrl, const bool bInAlt)
+	FInputChord(const FKey InKey, const bool bInShift, const bool bInCtrl, const bool bInAlt, const bool bInCmd)
 		: Key(InKey)
 		, bShift(bInShift)
 		, bCtrl(bInCtrl)
 		, bAlt(bInAlt)
+		, bCmd(bInCmd)
 	{
 	}
 
@@ -116,7 +121,8 @@ struct FInputChord
 		return (   Key == Other.Key
 				&& bShift == Other.bShift
 				&& bCtrl == Other.bCtrl
-				&& bAlt == Other.bAlt);
+				&& bAlt == Other.bAlt
+				&& bCmd == Other.bCmd);
 	}
 };
 
@@ -507,7 +513,7 @@ public:
 	template<class UserClass>
 	FInputKeyBinding& BindKey(const FKey Key, const EInputEvent KeyEvent, UserClass* Object, typename FInputActionHandlerSignature::TUObjectMethodDelegate< UserClass >::FMethodPtr Func)
 	{
-		return BindKey(FInputChord(Key, false, false, false), KeyEvent, Object, Func);
+		return BindKey(FInputChord(Key, false, false, false, false), KeyEvent, Object, Func);
 	}
 
 	template<class UserClass>
diff --git a/Engine/Source/Runtime/Engine/Classes/GameFramework/PlayerInput.h b/Engine/Source/Runtime/Engine/Classes/GameFramework/PlayerInput.h
--- a/Engine/Source/Runtime/Engine/Classes/GameFramework/PlayerInput.h
+++ b/Engine/Source/Runtime/Engine/Classes/GameFramework/PlayerInput.h
@@ -42,6 +42,9 @@ struct FKeyBind
 	UPROPERTY(config)
 	uint32 Alt:1;
 
+	UPROPERTY(config)
+	uint32 Cmd:1;
+
 	/** if true, the bind will not be activated if the corresponding key is held down */
 	UPROPERTY(config)
 	uint32 bIgnoreCtrl:1;
@@ -52,6 +55,9 @@ struct FKeyBind
 	UPROPERTY(config)
 	uint32 bIgnoreAlt:1;
 
+	UPROPERTY(config)
+	uint32 bIgnoreCmd:1;
+
 	FKeyBind() 
 	{
 	}
@@ -127,13 +133,18 @@ struct FInputActionKeyMapping
 	UPROPERTY(EditAnywhere, Category="Input")
 	uint32 bAlt:1;
 
+	/** true if one of the Cmd keys must be down when the KeyEvent is received to be acknowledged */
+	UPROPERTY(EditAnywhere, Category="Input")
+	uint32 bCmd:1;
+
 	bool operator==(const FInputActionKeyMapping& Other) const
 	{
 		return (   ActionName == Other.ActionName
 				&& Key == Other.Key
 				&& bShift == Other.bShift
 				&& bCtrl == Other.bCtrl
-				&& bAlt == Other.bAlt);
+				&& bAlt == Other.bAlt
+				&& bCmd == Other.bCmd);
 	}
 
 	bool operator<(const FInputActionKeyMapping& Other) const
@@ -150,12 +161,13 @@ struct FInputActionKeyMapping
 		return bResult;
 	}
 
-	FInputActionKeyMapping(const FName InActionName = NAME_None, const FKey InKey = EKeys::Invalid, const bool bInShift = false, const bool bInCtrl = false, const bool bInAlt = false)
+	FInputActionKeyMapping(const FName InActionName = NAME_None, const FKey InKey = EKeys::Invalid, const bool bInShift = false, const bool bInCtrl = false, const bool bInAlt = false, const bool bInCmd = false)
 		: ActionName(InActionName)
 		, Key(InKey)
 		, bShift(bInShift)
 		, bCtrl(bInCtrl)
 		, bAlt(bInAlt)
+		, bCmd(bInCmd)
 	{}
 };
 
@@ -436,6 +448,9 @@ public:
 	/** @return true if shift key is pressed */
 	bool IsShiftPressed() const;
 
+	/** @return true if cmd key is pressed */
+	bool IsCmdPressed() const;
+
 #if !UE_BUILD_SHIPPING
 	/**
 	 * Exec handler
diff --git a/Engine/Source/Runtime/Engine/Private/UserInterface/PlayerInput.cpp b/Engine/Source/Runtime/Engine/Private/UserInterface/PlayerInput.cpp
--- a/Engine/Source/Runtime/Engine/Private/UserInterface/PlayerInput.cpp
+++ b/Engine/Source/Runtime/Engine/Private/UserInterface/PlayerInput.cpp
@@ -551,13 +551,14 @@ void UPlayerInput::GetChordsForAction(const FInputActionBinding& ActionBinding,
 				// test modifier conditions and ignore the event if they failed
 				if (	(KeyMapping.bAlt == false || IsAltPressed())
 					&&	(KeyMapping.bCtrl == false || IsCtrlPressed())
-					&&	(KeyMapping.bShift == false || IsShiftPressed())
+					&&	(KeyMapping.bShift == false || IsShiftPressed()
+					&&	(KeyMapping.bCmd == false || IsCmdPressed()))
 					&& 	KeyEventOccurred(KeyMapping.Key, ActionBinding.KeyEvent, EventIndices))
 				{
 					bool bAddDelegate = true;
 
 					// look through the found chords and determine if this is masked (or masks) anything in the array
-					const FInputChord Chord(KeyMapping.Key, KeyMapping.bShift, KeyMapping.bCtrl, KeyMapping.bAlt);
+					const FInputChord Chord(KeyMapping.Key, KeyMapping.bShift, KeyMapping.bCtrl, KeyMapping.bAlt, KeyMapping.bCmd);
 					for (int32 ChordIndex = FoundChords.Num() - 1; ChordIndex >= 0; --ChordIndex)
 					{
 						FInputChord::RelationshipType ChordRelationship = Chord.GetRelationship(FoundChords[ChordIndex].Chord);
@@ -614,6 +615,7 @@ void UPlayerInput::GetChordForKey(const FInputKeyBinding& KeyBinding, const bool
 		if (	(KeyBinding.Chord.bAlt == false || IsAltPressed())
 			&&	(KeyBinding.Chord.bCtrl == false || IsCtrlPressed())
 			&&	(KeyBinding.Chord.bShift == false || IsShiftPressed())
+			&&	(KeyBinding.Chord.bCmd == false || IsCmdPressed())
 			&& 	KeyEventOccurred(KeyBinding.Chord.Key, KeyBinding.KeyEvent, EventIndices))
 		{
 			bool bAddDelegate = true;
@@ -1408,6 +1410,11 @@ bool UPlayerInput::IsShiftPressed() const
 	return IsPressed(EKeys::LeftShift) || IsPressed(EKeys::RightShift);
 }
 
+bool UPlayerInput::IsCmdPressed() const
+{
+	return IsPressed(EKeys::LeftCommand) || IsPressed(EKeys::RightCommand);
+}
+
 void UPlayerInput::ConditionalInitAxisProperties()
 {
 	// Initialize AxisProperties map if needed.
@@ -1432,7 +1439,8 @@ bool UPlayerInput::IsKeyHandledByAction( FKey Key ) const
 		if( Mapping.Key == Key && 
 			(Mapping.bAlt == false || IsAltPressed()) &&
 			(Mapping.bCtrl == false || IsCtrlPressed()) &&
-			(Mapping.bShift == false || IsShiftPressed()) )
+			(Mapping.bShift == false || IsShiftPressed()) &&
+			(Mapping.bCmd == false || IsCmdPressed()) )
 		{
 			return true;
 		}
@@ -1540,6 +1548,7 @@ FString UPlayerInput::GetBind(FKey Key)
 		const bool bControlPressed = IsCtrlPressed();
 		const bool bShiftPressed = IsShiftPressed();
 		const bool bAltPressed = IsAltPressed();
+		const bool bCmdPressed = IsCmdPressed();
 
 		for ( int32 BindIndex = DebugExecBindings.Num() - 1; BindIndex >= 0; BindIndex-- )
 		{
@@ -1548,8 +1557,8 @@ FString UPlayerInput::GetBind(FKey Key)
 			{
 				// if the modifier key pressed [or this key-bind doesn't require that key], and the key-bind isn't
 				// configured to ignore they modifier key, we've found a match.
-				if ((!Bind.Control || bControlPressed) && (!Bind.Shift || bShiftPressed) && (!Bind.Alt || bAltPressed)
-					&&	(!Bind.bIgnoreCtrl || !bControlPressed) && (!Bind.bIgnoreShift || !bShiftPressed) && (!Bind.bIgnoreAlt || !bAltPressed))
+				if ((!Bind.Control || bControlPressed) && (!Bind.Shift || bShiftPressed) && (!Bind.Alt || bAltPressed) && (!Bind.Cmd || bCmdPressed)
+					&&	(!Bind.bIgnoreCtrl || !bControlPressed) && (!Bind.bIgnoreShift || !bShiftPressed) && (!Bind.bIgnoreAlt || !bAltPressed) && (!Bind.bIgnoreCmd || !bCmdPressed))
 				{
 					return DebugExecBindings[BindIndex].Command;
 				}
diff --git a/Engine/Source/Runtime/Engine/Public/Interpolation.h b/Engine/Source/Runtime/Engine/Public/Interpolation.h
--- a/Engine/Source/Runtime/Engine/Public/Interpolation.h
+++ b/Engine/Source/Runtime/Engine/Public/Interpolation.h
@@ -15,6 +15,7 @@ struct FInterpEdInputData
 		  bCtrlDown( false ),
 		  bAltDown( false ),
 		  bShiftDown( false ),
+		  bCmdDown( false ),
 		  MouseStart( 0, 0 ),
 		  MouseCurrent( 0, 0 ),
 		  PixelsPerSec( 0.0f )
@@ -28,6 +29,7 @@ struct FInterpEdInputData
 		  bCtrlDown( false ),
 		  bAltDown( false ),
 		  bShiftDown( false ),
+		  bCmdDown( false ),
 		  MouseStart( 0, 0 ),
 		  MouseCurrent( 0, 0 ),
 		  PixelsPerSec( 0.0f )
@@ -42,6 +44,7 @@ struct FInterpEdInputData
 	bool bCtrlDown;
 	bool bAltDown;
 	bool bShiftDown;
+	bool bCmdDown;
 	FIntPoint MouseStart;
 	FIntPoint MouseCurrent;
 	float PixelsPerSec;
diff --git a/Engine/Source/Runtime/Slate/Private/Framework/Commands/GenericCommands.cpp b/Engine/Source/Runtime/Slate/Private/Framework/Commands/GenericCommands.cpp
--- a/Engine/Source/Runtime/Slate/Private/Framework/Commands/GenericCommands.cpp
+++ b/Engine/Source/Runtime/Slate/Private/Framework/Commands/GenericCommands.cpp
@@ -7,7 +7,11 @@ void FGenericCommands::RegisterCommands()
 	UI_COMMAND( Cut, "Cut", "Cut selection", EUserInterfaceActionType::Button, FInputGesture( EModifierKey::Control, EKeys::X ) )
 	UI_COMMAND( Copy, "Copy", "Copy selection", EUserInterfaceActionType::Button, FInputGesture( EModifierKey::Control, EKeys::C ) )
 	UI_COMMAND( Paste, "Paste", "Paste clipboard contents", EUserInterfaceActionType::Button, FInputGesture( EModifierKey::Control, EKeys::V ) )
+#if PLATFORM_MAC
+	UI_COMMAND( Duplicate, "Duplicate", "Duplicate selection", EUserInterfaceActionType::Button, FInputGesture( EModifierKey::Command, EKeys::W ) )
+#else
 	UI_COMMAND( Duplicate, "Duplicate", "Duplicate selection", EUserInterfaceActionType::Button, FInputGesture( EModifierKey::Control, EKeys::W ) )
+#endif
 	UI_COMMAND( Undo, "Undo", "Undo last action", EUserInterfaceActionType::Button, FInputGesture( EModifierKey::Control, EKeys::Z ) )
 	UI_COMMAND( Redo, "Redo", "Redo last action", EUserInterfaceActionType::Button, FInputGesture( EModifierKey::Control, EKeys::Y ) )
 
diff --git a/Engine/Source/Runtime/Slate/Private/Framework/Commands/InputGesture.cpp b/Engine/Source/Runtime/Slate/Private/Framework/Commands/InputGesture.cpp
--- a/Engine/Source/Runtime/Slate/Private/Framework/Commands/InputGesture.cpp
+++ b/Engine/Source/Runtime/Slate/Private/Framework/Commands/InputGesture.cpp
@@ -21,50 +21,114 @@ FText FInputGesture::GetInputText( ) const
 {
 	FText FormatString = LOCTEXT("ShortcutKeyWithNoAdditionalModifierKeys", "{Key}");
 
-	if ( bCtrl && !bAlt && !bShift )
+	if ( bCtrl && !bAlt && !bShift && !bCmd )
 	{
 #if PLATFORM_MAC
-		FormatString = LOCTEXT("KeyName_Command", "Command+{Key}");
+		FormatString = LOCTEXT("KeyName_Command", "Cmd+{Key}");
 #else
 		FormatString = LOCTEXT("KeyName_Control", "Ctrl+{Key}");
 #endif
 	}
-	else if ( bCtrl && bAlt && !bShift )
+	else if ( bCtrl && bAlt && !bShift && !bCmd )
 	{
 #if PLATFORM_MAC
-		FormatString = LOCTEXT("KeyName_Command + KeyName_Alt", "Command+Alt+{Key}");
+		FormatString = LOCTEXT("KeyName_Command + KeyName_Alt", "Cmd+Alt+{Key}");
 #else
 		FormatString = LOCTEXT("KeyName_Control + KeyName_Alt", "Ctrl+Alt+{Key}");
 #endif
 	}
-	else if ( bCtrl && !bAlt && bShift )
+	else if ( bCtrl && !bAlt && bShift && !bCmd )
 	{
 #if PLATFORM_MAC
-		FormatString = LOCTEXT("KeyName_Command + KeyName_Shift", "Command+Shift+{Key}");
+		FormatString = LOCTEXT("KeyName_Command + KeyName_Shift", "Cmd+Shift+{Key}");
 #else
 		FormatString = LOCTEXT("KeyName_Control + KeyName_Shift", "Ctrl+Shift+{Key}");
 #endif
 	}
-	else if ( bCtrl && bAlt && bShift )
+	else if ( bCtrl && !bAlt && !bShift && bCmd )
 	{
 #if PLATFORM_MAC
-		FormatString = LOCTEXT("KeyName_Command + KeyName_Alt + KeyName_Shift", "Command+Alt+Shift+{Key}");
+		FormatString = LOCTEXT("KeyName_Command + KeyName_Control", "Cmd+Ctrl+{Key}");
 #else
+		FormatString = LOCTEXT("KeyName_Control + KeyName_Command", "Ctrl+Cmd+{Key}");
+#endif
+	}
+	else if ( !bCtrl && bAlt && !bShift && bCmd )
+	{
+#if PLATFORM_MAC
+		FormatString = LOCTEXT("KeyName_Control + KeyName_Alt", "Ctrl+Alt+{Key}");
+#else
+		FormatString = LOCTEXT("KeyName_Command + KeyName_Alt", "Cmd+Alt+{Key}");
+#endif
+	}
+	else if ( !bCtrl && !bAlt && bShift && bCmd )
+	{
+#if PLATFORM_MAC
+		FormatString = LOCTEXT("KeyName_Control + KeyName_Shift", "Ctrl+Shift+{Key}");
+#else
+		FormatString = LOCTEXT("KeyName_Command + KeyName_Shift", "Cmd+Shift+{Key}");
+#endif
+	}
+	else if ( bCtrl && bAlt && bShift && !bCmd )
+	{
+#if PLATFORM_MAC
+		FormatString = LOCTEXT("KeyName_Command + KeyName_Alt + KeyName_Shift", "Cmd+Alt+Shift+{Key}");
+#else
+		FormatString = LOCTEXT("KeyName_Control + KeyName_Alt + KeyName_Shift", "Ctrl+Alt+Shift+{Key}");
+#endif
+	}
+	else if ( !bCtrl && bAlt && bShift && bCmd )
+	{
+#if PLATFORM_MAC
 		FormatString = LOCTEXT("KeyName_Control + KeyName_Alt + KeyName_Shift", "Ctrl+Alt+Shift+{Key}");
+#else
+		FormatString = LOCTEXT("KeyName_Command + KeyName_Alt + KeyName_Shift", "Cmd+Alt+Shift+{Key}");
+#endif
+	}
+	else if ( bCtrl && !bAlt && bShift && bCmd )
+	{
+#if PLATFORM_MAC
+		FormatString = LOCTEXT("KeyName_Command + KeyName_Control + KeyName_Shift", "Cmd+Ctrl+Shift+{Key}");
+#else
+		FormatString = LOCTEXT("KeyName_Control + KeyName_Command + KeyName_Shift", "Ctrl+Cmd+Shift+{Key}");
+#endif
+	}
+	else if ( bCtrl && bAlt && !bShift && bCmd )
+	{
+#if PLATFORM_MAC
+		FormatString = LOCTEXT("KeyName_Command + KeyName_Control + KeyName_Alt", "Cmd+Ctrl+Alt+{Key}");
+#else
+		FormatString = LOCTEXT("KeyName_Control + KeyName_Command + KeyName_Alt", "Ctrl+Cmd+Alt+{Key}");
 #endif
 	}
-	else if ( !bCtrl && bAlt && !bShift )
+	else if ( bCtrl && bAlt && bShift && bCmd )
+	{
+#if PLATFORM_MAC
+		FormatString = LOCTEXT("KeyName_Command + KeyName_Control + KeyName_Alt + KeyName_Shift", "Cmd+Ctrl+Alt+Shift{Key}");
+#else
+		FormatString = LOCTEXT("KeyName_Control + KeyName_Command + KeyName_Alt + KeyName_Shift", "Ctrl+Cmd+Alt+Shift{Key}");
+#endif
+	}
+	else if ( !bCtrl && bAlt && !bShift && !bCmd )
 	{
 		FormatString = LOCTEXT("KeyName_Alt", "Alt+{Key}");
 	}
-	else if ( !bCtrl && bAlt && bShift )
+	else if ( !bCtrl && bAlt && bShift && !bCmd )
 	{
 		FormatString = LOCTEXT("KeyName_Alt + KeyName_Shift", "Alt+Shift+{Key}");
 	}
-	else if ( !bCtrl && !bAlt && bShift )
+	else if ( !bCtrl && !bAlt && bShift && !bCmd )
 	{
 		FormatString = LOCTEXT("KeyName_Shift", "Shift+{Key}");
 	}
+	else if ( !bCtrl && !bAlt && !bShift && bCmd )
+	{
+#if PLATFORM_MAC
+		FormatString = LOCTEXT("KeyName_Control", "Ctrl+{Key}");
+#else
+		FormatString = LOCTEXT("KeyName_Command", "Cmd+{Key}");
+#endif
+	}
 
 	FFormatNamedArguments Args;
 	Args.Add( TEXT("Key"), GetKeyText() );
diff --git a/Engine/Source/Runtime/Slate/Private/Widgets/SViewport.cpp b/Engine/Source/Runtime/Slate/Private/Widgets/SViewport.cpp
--- a/Engine/Source/Runtime/Slate/Private/Widgets/SViewport.cpp
+++ b/Engine/Source/Runtime/Slate/Private/Widgets/SViewport.cpp
@@ -86,7 +86,7 @@ int32 SViewport::OnPaint( const FGeometry& AllottedGeometry, const FSlateRect& M
 				CursorPosScreenSpace,
 				FVector2D::ZeroVector,
 				TSet<FKey>(),
-				FModifierKeysState(false, false, false, false, false, false) )
+				FModifierKeysState(false, false, false, false, false, false, false, false) )
 		 );
 		EMouseCursor::Type CursorType = Reply.GetCursor();
 
diff --git a/Engine/Source/Runtime/Slate/Public/Framework/Commands/InputGesture.h b/Engine/Source/Runtime/Slate/Public/Framework/Commands/InputGesture.h
--- a/Engine/Source/Runtime/Slate/Public/Framework/Commands/InputGesture.h
+++ b/Engine/Source/Runtime/Slate/Public/Framework/Commands/InputGesture.h
@@ -17,7 +17,7 @@ namespace EModifierKey
 	/** No key. */
 	const Type None	= 0;
 
-	/** Ctrl key. */
+	/** Ctrl key (Command key on Mac, Control key on Windows). */
 	const Type Control = 1 << 0;
 
 	/** Alt key. */
@@ -25,6 +25,9 @@ namespace EModifierKey
 
 	/** Shift key. */
 	const Type Shift = 1 << 2;
+
+	/** Cmd key (Control key on Mac, Win key on Windows) */
+	const Type Command = 1 << 3;
 };
 
 
@@ -45,6 +48,9 @@ struct SLATE_API FInputGesture
 	/** True if shift must be pressed */
 	uint32 bShift:1;
 
+	/** True if command must be pressed */
+	uint32 bCmd:1;
+
 public:
 
 	/**
@@ -52,7 +58,7 @@ public:
 	 */
 	FInputGesture( )
 	{
-		bCtrl = bAlt = bShift = 0;
+		bCtrl = bAlt = bShift = bCmd = 0;
 	}
 
 	/**
@@ -67,6 +73,7 @@ public:
 		bCtrl = (InModifierKeys & EModifierKey::Control) != 0;
 		bAlt = (InModifierKeys & EModifierKey::Alt) != 0;
 		bShift = (InModifierKeys & EModifierKey::Shift) != 0;
+		bCmd = (InModifierKeys & EModifierKey::Command) != 0;
 	}
 
 	/**
@@ -77,7 +84,7 @@ public:
 	FInputGesture( const FKey InKey )
 		: Key(InKey)
 	{
-		bCtrl = bAlt = bShift = 0;
+		bCtrl = bAlt = bShift = bCmd = 0;
 	}
 
 	/**
@@ -92,6 +99,7 @@ public:
 		bCtrl = (InModifierKeys & EModifierKey::Control) != 0;
 		bAlt = (InModifierKeys & EModifierKey::Alt) != 0;
 		bShift = (InModifierKeys & EModifierKey::Shift) != 0;
+		bCmd = (InModifierKeys & EModifierKey::Command) != 0;
 	}
 
 	/**
@@ -104,6 +112,7 @@ public:
 		, bCtrl(Other.bCtrl)
 		, bAlt(Other.bAlt)
 		, bShift(Other.bShift)
+		, bCmd(Other.bCmd)
 	{ }
 
 public:
@@ -117,7 +126,7 @@ public:
 	 */
 	bool operator!=( const FInputGesture& Other ) const
 	{
-		return (Key != Other.Key) || (bCtrl != Other.bCtrl) || (bAlt != Other.bAlt) || (bShift != Other.bShift);
+		return (Key != Other.Key) || (bCtrl != Other.bCtrl) || (bAlt != Other.bAlt) || (bShift != Other.bShift) || (bCmd != Other.bCmd);
 	}
 
 	/**
@@ -129,7 +138,7 @@ public:
 	 */
 	bool operator==( const FInputGesture& Other ) const
 	{
-		return (Key == Other.Key) && (bCtrl == Other.bCtrl) && (bAlt == Other.bAlt) && (bShift == Other.bShift);
+		return (Key == Other.Key) && (bCtrl == Other.bCtrl) && (bAlt == Other.bAlt) && (bShift == Other.bShift) && (bCmd == Other.bCmd);
 	}
 
 public:
@@ -155,7 +164,7 @@ public:
 	 */
 	bool HasAnyModifierKeys( ) const
 	{
-		return bCtrl || bAlt || bShift;
+		return bCtrl || bAlt || bShift || bCmd;
 	}
 
 	/**
@@ -179,6 +188,7 @@ public:
 		bCtrl = InTemplate.bCtrl;
 		bAlt = InTemplate.bAlt;
 		bShift = InTemplate.bShift;
+		bCmd = InTemplate.bCmd;
 	}
 
 public:
@@ -192,6 +202,6 @@ public:
 	 */
 	friend uint32 GetTypeHash( const FInputGesture& Gesture )
 	{
-		return GetTypeHash(Gesture.Key) ^ ((Gesture.bCtrl << 2) | (Gesture.bShift << 1) | Gesture.bAlt);
+		return GetTypeHash(Gesture.Key) ^ ((Gesture.bCtrl << 3) | (Gesture.bShift << 2) | (Gesture.bAlt << 1) | Gesture.bCmd);
 	}
 };
diff --git a/Engine/Source/Runtime/SlateCore/Public/Input/Events.h b/Engine/Source/Runtime/SlateCore/Public/Input/Events.h
--- a/Engine/Source/Runtime/SlateCore/Public/Input/Events.h
+++ b/Engine/Source/Runtime/SlateCore/Public/Input/Events.h
@@ -98,7 +98,7 @@ public:
 	 * UStruct Constructor.  Not meant for normal usage.
 	 */
 	FInputEvent()
-		: ModifierKeys(FModifierKeysState(false, false, false, false, false, false))
+		: ModifierKeys(FModifierKeysState(false, false, false, false, false, false, false, false))
 		, bIsRepeat(false)
 		, EventPath(nullptr)
 	{ }
@@ -222,6 +222,36 @@ public:
 		return ModifierKeys.IsRightAltDown();
 	}
 
+	/**
+	 * Returns true if either command key was down when this event occurred
+	 *
+	 * @return  True if command is pressed
+	 */
+	bool IsCommandDown() const
+	{
+		return ModifierKeys.IsCommandDown();
+	}
+
+	/**
+	 * Returns true if left command key was down when this event occurred
+	 *
+	 * @return  True if left command is pressed
+	 */
+	bool IsLeftCommandDown() const
+	{
+		return ModifierKeys.IsLeftCommandDown();
+	}
+
+	/**
+	 * Returns true if right command key was down when this event occurred
+	 *
+	 * @return  True if right command is pressed
+	 */
+	bool IsRightCommandDown() const
+	{
+		return ModifierKeys.IsRightCommandDown();
+	}
+
 	/** The event path provides additional context for handling */
 	FGeometry FindGeometry( const TSharedRef<SWidget>& WidgetToFind ) const
 	{
@@ -267,7 +297,7 @@ public:
 	 * UStruct Constructor.  Not meant for normal usage.
 	 */
 	FKeyboardEvent()
-		: FInputEvent(FModifierKeysState(false, false, false, false, false, false), false)
+		: FInputEvent(FModifierKeysState(false, false, false, false, false, false, false, false), false)
 		, Key(EKeys::SpaceBar)
 		, CharacterCode(0)
 	{
@@ -334,7 +364,7 @@ public:
 	 * UStruct Constructor.  Not meant for normal usage.
 	 */
 	FCharacterEvent()
-		: FInputEvent(FModifierKeysState(false, false, false, false, false, false), false)
+		: FInputEvent(FModifierKeysState(false, false, false, false, false, false, false, false), false)
 		, Character(0)
 	{
 	}
@@ -468,7 +498,7 @@ public:
 		const FVector2D& InScreenSpacePosition,
 		const FVector2D& InLastScreenSpacePosition,
 		bool bPressLeftMouseButton,
-		const FModifierKeysState& InModifierKeys = FModifierKeysState(false, false, false, false, false, false),
+		const FModifierKeysState& InModifierKeys = FModifierKeysState(false, false, false, false, false, false, false, false),
 		uint32 InTouchpadIndex=0
 	)
 		: FInputEvent(InModifierKeys, false)
@@ -597,7 +627,7 @@ public:
 	{ }
 
 	FControllerEvent( FKey InEffectingButton, int32 InUserIndex, float InAnalogValue, bool bIsRepeat )
-		: FInputEvent(FModifierKeysState(false, false, false, false, false, false), bIsRepeat)
+		: FInputEvent(FModifierKeysState(false, false, false, false, false, false, false, false), bIsRepeat)
 		, EffectingButton(InEffectingButton)
 		, UserIndex(InUserIndex)
 		, AnalogValue(InAnalogValue)
@@ -656,7 +686,7 @@ public:
 		const FVector& InGravity, 
 		const FVector& InAcceleration
 	)
-		: FInputEvent(FModifierKeysState(false, false, false, false, false, false), false)
+		: FInputEvent(FModifierKeysState(false, false, false, false, false, false, false, false), false)
 		, UserIndex(InUserIndex)
 		, Tilt(InTilt)
 		, RotationRate(InRotationRate)
