Split pin actions as part of node reconstruction should not be undoable

diff --git a/Engine/Source/Editor/BlueprintGraph/Private/EdGraphSchema_K2.cpp b/Engine/Source/Editor/BlueprintGraph/Private/EdGraphSchema_K2.cpp
--- a/Engine/Source/Editor/BlueprintGraph/Private/EdGraphSchema_K2.cpp
+++ b/Engine/Source/Editor/BlueprintGraph/Private/EdGraphSchema_K2.cpp
@@ -4590,8 +4590,6 @@ void UEdGraphSchema_K2::SplitPin(UEdGraphPin* Pin) const
 	UK2Node* K2Node = Cast<UK2Node>(GraphNode);
 	UEdGraph* Graph = CastChecked<UEdGraph>(GraphNode->GetOuter());
 
-	const FScopedTransaction Transaction( LOCTEXT("SplitStructPin", "Split Struct Pin") );
-
 	GraphNode->Modify();
 	Pin->Modify();
 
@@ -4691,8 +4689,6 @@ void UEdGraphSchema_K2::SplitPin(UEdGraphPin* Pin) const
 
 void UEdGraphSchema_K2::RecombinePin(UEdGraphPin* Pin) const
 {
-	const FScopedTransaction Transaction( LOCTEXT("RecombineStructPin", "Recombine Struct Pin") );
-
 	UEdGraphNode* GraphNode = Pin->GetOwningNode();
 	UEdGraphPin* ParentPin = Pin->ParentPin;
 
diff --git a/Engine/Source/Editor/Kismet/Private/BlueprintEditor.cpp b/Engine/Source/Editor/Kismet/Private/BlueprintEditor.cpp
--- a/Engine/Source/Editor/Kismet/Private/BlueprintEditor.cpp
+++ b/Engine/Source/Editor/Kismet/Private/BlueprintEditor.cpp
@@ -2751,6 +2751,8 @@ void FBlueprintEditor::OnSplitStructPin()
 		check(GetBlueprintObj()->SkeletonGeneratedClass);
 		check(TargetPin);
 
+		const FScopedTransaction Transaction( LOCTEXT("SplitStructPin", "Split Struct Pin") );
+
 		const UEdGraphSchema_K2* K2Schema = GetDefault<UEdGraphSchema_K2>();
 		K2Schema->SplitPin(TargetPin);
 	}
@@ -2784,6 +2786,8 @@ void FBlueprintEditor::OnRecombineStructPin()
 		check(GetBlueprintObj()->SkeletonGeneratedClass);
 		check(TargetPin);
 
+		const FScopedTransaction Transaction( LOCTEXT("RecombineStructPin", "Recombine Struct Pin") );
+
 		const UEdGraphSchema_K2* K2Schema = GetDefault<UEdGraphSchema_K2>();
 		K2Schema->RecombinePin(TargetPin);
 	}
