A quick fix for  TTP #342030 EDITOR: AUTOMATION: CRASH: Assertion failed: Property && Container in BlueprintEditorUtils.cpp

diff --git a/Engine/Source/Editor/BlueprintGraph/Private/K2Node_SpawnActorFromClass.cpp b/Engine/Source/Editor/BlueprintGraph/Private/K2Node_SpawnActorFromClass.cpp
--- a/Engine/Source/Editor/BlueprintGraph/Private/K2Node_SpawnActorFromClass.cpp
+++ b/Engine/Source/Editor/BlueprintGraph/Private/K2Node_SpawnActorFromClass.cpp
@@ -81,6 +81,8 @@ void UK2Node_SpawnActorFromClass::CreatePinsForClass(UClass* InClass)
 
 	const UEdGraphSchema_K2* K2Schema = GetDefault<UEdGraphSchema_K2>();
 
+	const UObject* const ClassDefaultObject = InClass->GetDefaultObject(false);
+
 	for (TFieldIterator<UProperty> PropertyIt(InClass, EFieldIteratorFlags::IncludeSuper); PropertyIt; ++PropertyIt)
 	{
 		UProperty* Property = *PropertyIt;
@@ -99,10 +101,10 @@ void UK2Node_SpawnActorFromClass::CreatePinsForClass(UClass* InClass)
 			UEdGraphPin* Pin = CreatePin(EGPD_Input, TEXT(""), TEXT(""), NULL, false, false, Property->GetName());
 			const bool bPinGood = (Pin != NULL) && K2Schema->ConvertPropertyToPinType(Property, /*out*/ Pin->PinType);
 			
-			if( K2Schema->PinDefaultValueIsEditable(*Pin) )
+			if( ClassDefaultObject && K2Schema->PinDefaultValueIsEditable(*Pin))
 			{
 				FString DefaultValueAsString;
-				const bool bDefaultValueSet = FBlueprintEditorUtils::PropertyValueToString(Property, (uint8*)InClass->ClassDefaultObject, DefaultValueAsString);
+				const bool bDefaultValueSet = FBlueprintEditorUtils::PropertyValueToString(Property, reinterpret_cast<const uint8*>(ClassDefaultObject), DefaultValueAsString);
 				check( bDefaultValueSet );
 				K2Schema->TrySetDefaultValue(*Pin, DefaultValueAsString);
 			}
