ttp341422 BP: Improve signature mismatch errors

diff --git a/Engine/Source/Editor/BlueprintGraph/Private/K2Node_CreateDelegate.cpp b/Engine/Source/Editor/BlueprintGraph/Private/K2Node_CreateDelegate.cpp
--- a/Engine/Source/Editor/BlueprintGraph/Private/K2Node_CreateDelegate.cpp
+++ b/Engine/Source/Editor/BlueprintGraph/Private/K2Node_CreateDelegate.cpp
@@ -38,18 +38,30 @@ bool UK2Node_CreateDelegate::IsValid(FString* OutMsg, bool bDontUseSkeletalClass
 
 	if (GetFunctionName() == NAME_None)
 	{
+		if (OutMsg)
+		{
+			*OutMsg = NSLOCTEXT("K2Node", "No_function_name", "No function name.").ToString();
+		}
 		return false;
 	}
 
 	const UEdGraphPin* DelegatePin = GetDelegateOutPin();
 	if(!DelegatePin)
 	{
+		if (OutMsg)
+		{
+			*OutMsg = NSLOCTEXT("K2Node", "No_delegate_out_pin", "No delegate out pin.").ToString();
+		}
 		return false;
 	}
 
 	const UFunction* Signature = GetDelegateSignature();
 	if(!Signature)
 	{
+		if (OutMsg)
+		{
+			*OutMsg = NSLOCTEXT("K2Node", "Signature_not_found", "Signature not found.").ToString();
+		}
 		return false;
 	}
 
@@ -60,6 +72,10 @@ bool UK2Node_CreateDelegate::IsValid(FString* OutMsg, bool bDontUseSkeletalClass
 			FMemberReference::ResolveSimpleMemberReference<UFunction>(OtherPin->PinType.PinSubCategoryMemberReference) : NULL;
 		if(!OtherSignature || !Signature->IsSignatureCompatibleWith(OtherSignature))
 		{
+			if (OutMsg)
+			{
+				*OutMsg = NSLOCTEXT("K2Node", "No_delegate_out_pin", "No delegate out pin.").ToString();
+			}
 			return false;
 		}
 	}
@@ -67,16 +83,38 @@ bool UK2Node_CreateDelegate::IsValid(FString* OutMsg, bool bDontUseSkeletalClass
 	UClass* ScopeClass = GetScopeClass(bDontUseSkeletalClassForSelf);
 	if(!ScopeClass)
 	{
+		if (OutMsg)
+		{
+			*OutMsg = NSLOCTEXT("K2Node", "Class_not_found", "Class not found.").ToString();
+		}
 		return false;
 	}
 
 	FMemberReference MemeberReference;
 	MemeberReference.SetDirect(SelectedFunctionName, SelectedFunctionGuid, ScopeClass, false);
 	const UFunction* FoundFunction = MemeberReference.ResolveMember<UFunction>((UClass*) NULL);
-	if (!FoundFunction || 
-		!Signature->IsSignatureCompatibleWith(FoundFunction) || 
-		!UEdGraphSchema_K2::FunctionCanBeUsedInDelegate(FoundFunction))
+	if (!FoundFunction)
+	{
+		if (OutMsg)
+		{
+			*OutMsg = NSLOCTEXT("K2Node", "Function_not_found", "Function not found.").ToString();
+		}
+		return false;
+	}
+	if (!Signature->IsSignatureCompatibleWith(FoundFunction))
 	{
+		if (OutMsg)
+		{
+			*OutMsg = NSLOCTEXT("K2Node", "Function_not_compatible", "Function not compatible.").ToString();
+		}
+		return false;
+	}
+	if (!UEdGraphSchema_K2::FunctionCanBeUsedInDelegate(FoundFunction))
+	{
+		if (OutMsg)
+		{
+			*OutMsg = NSLOCTEXT("K2Node", "Function_cannot_be_used_in_delegate", "Function cannot be used in delegate.").ToString();
+		}
 		return false;
 	}
 
@@ -105,7 +143,7 @@ void UK2Node_CreateDelegate::ValidationAfterFunctionsAreCreated(class FCompilerR
 	FString Msg;
 	if(!IsValid(&Msg, bFullCompile))
 	{
-		MessageLog.Error(*FString::Printf( TEXT("%s %s"), *NSLOCTEXT("K2Node", "WrongDelegate", "Events signatures don't match.").ToString(), *Msg));
+		MessageLog.Error(*FString::Printf( TEXT("%s %s @@"), *NSLOCTEXT("K2Node", "WrongDelegate", "Event signature error: ").ToString(), *Msg), this);
 	}
 }
 
