Hiding private/protected functions from blueprints that shouldn't have access to them.
(TTP #332382)

diff --git a/Engine/Source/Editor/BlueprintGraph/Private/EdGraphSchema_K2.cpp b/Engine/Source/Editor/BlueprintGraph/Private/EdGraphSchema_K2.cpp
--- a/Engine/Source/Editor/BlueprintGraph/Private/EdGraphSchema_K2.cpp
+++ b/Engine/Source/Editor/BlueprintGraph/Private/EdGraphSchema_K2.cpp
@@ -313,6 +313,7 @@ bool UEdGraphSchema_K2::CanFunctionBeUsedInClass(const UClass* InClass, UFunctio
 		const bool bIsBlueprintProtected = InFunction->GetBoolMetaData(FBlueprintMetadata::MD_Protected);
 		const bool bIsUnsafeForConstruction = InFunction->GetBoolMetaData(FBlueprintMetadata::MD_UnsafeForConstructionScripts);	
 		const bool bFunctionHidden = FObjectEditorUtils::IsFunctionHiddenFromClass(InFunction, InClass);
+		const bool bIsPrivate = InFunction->GetBoolMetaData(FBlueprintMetadata::MD_Private);
 
 		const bool bFunctionStatic = InFunction->HasAllFunctionFlags(FUNC_Static);
 		const bool bHasReturnParams = (InFunction->GetReturnProperty() != NULL);
@@ -323,6 +324,9 @@ bool UEdGraphSchema_K2::CanFunctionBeUsedInClass(const UClass* InClass, UFunctio
 		const bool bClassIsAnActor = InClass->IsChildOf( AActor::StaticClass() );
 		const bool bClassIsAComponent = InClass->IsChildOf( UActorComponent::StaticClass() );
 
+		const bool bFuncBelongsToSubClass = InClass->IsChildOf(InFunction->GetOuterUClass());
+		const bool bFuncBelongsToClass    = bFuncBelongsToSubClass && (InFunction->GetOuterUClass() == InClass);
+
 		const bool bFunctionHasReturnOrOutParameters = bHasReturnParams || DoesFunctionHaveOutParameters(InFunction);
 
 		// This will evaluate to false if there are multiple actors selected and the function has a return value or out parameters
@@ -330,11 +334,12 @@ bool UEdGraphSchema_K2::CanFunctionBeUsedInClass(const UClass* InClass, UFunctio
 
 		if (((bIsPureFunc && bPureFuncs) || (!bIsPureFunc && bImperativeFuncs) || (bIsConstFunc && bConstFuncs))
 			&& (!bIsLatent || bLatentFuncs)
-			&& (!bIsBlueprintProtected || bProtectedFuncs)
+			&& (!bIsBlueprintProtected || (bProtectedFuncs && bFuncBelongsToSubClass))
 			&& (!bIsUnsafeForConstruction || !bIsConstructionScript)
 			&& !bFunctionHidden
 			&& (bAllowForEachCall || !bInCalledForEach)
-			&& bAllowReturnValuesForNoneOrSingleActors )
+			&& bAllowReturnValuesForNoneOrSingleActors 
+			&& (!bIsPrivate || bFuncBelongsToClass) )
 		{
 			return true;
 		}
