#UE4 Applying fix we made in spawnactorfromclass

diff --git a/Engine/Source/Editor/BlueprintGraph/Private/K2Node_AddComponent.cpp b/Engine/Source/Editor/BlueprintGraph/Private/K2Node_AddComponent.cpp
--- a/Engine/Source/Editor/BlueprintGraph/Private/K2Node_AddComponent.cpp
+++ b/Engine/Source/Editor/BlueprintGraph/Private/K2Node_AddComponent.cpp
@@ -84,8 +84,12 @@ void UK2Node_AddComponent::AllocatePinsForExposedVariables()
 	const UEdGraphSchema_K2* K2Schema = GetDefault<UEdGraphSchema_K2>();
 
 	const UActorComponent* TemplateComponent = GetTemplateFromNode();
-	if(const UClass* ComponentClass = (TemplateComponent ? TemplateComponent->GetClass() : NULL))
+	const UClass* ComponentClass = TemplateComponent ? TemplateComponent->GetClass() : NULL;
+
+	if(ComponentClass)
 	{
+		const UObject* ClassDefaultObject = ComponentClass ? ComponentClass->ClassDefaultObject : NULL;
+
 		for (TFieldIterator<UProperty> PropertyIt(ComponentClass, EFieldIteratorFlags::IncludeSuper); PropertyIt; ++PropertyIt)
 		{
 			UProperty* Property = *PropertyIt;
@@ -104,10 +108,10 @@ void UK2Node_AddComponent::AllocatePinsForExposedVariables()
 					Pin->PinType = PinType;
 					bHasExposedVariable = true;
 
-					if (K2Schema->PinDefaultValueIsEditable(*Pin))
+					if (ClassDefaultObject && K2Schema->PinDefaultValueIsEditable(*Pin))
 					{
 						FString DefaultValueAsString;
-						const bool bDefaultValueSet = FBlueprintEditorUtils::PropertyValueToString(Property, (uint8*)ComponentClass->ClassDefaultObject, DefaultValueAsString);
+						const bool bDefaultValueSet = FBlueprintEditorUtils::PropertyValueToString(Property, reinterpret_cast<const uint8*>(ClassDefaultObject), DefaultValueAsString);
 						check(bDefaultValueSet);
 						K2Schema->TrySetDefaultValue(*Pin, DefaultValueAsString);
 					}
