Remove unused local variable

diff --git a/Engine/Source/Editor/BlueprintGraph/Private/EdGraphSchema_K2.cpp b/Engine/Source/Editor/BlueprintGraph/Private/EdGraphSchema_K2.cpp
--- a/Engine/Source/Editor/BlueprintGraph/Private/EdGraphSchema_K2.cpp
+++ b/Engine/Source/Editor/BlueprintGraph/Private/EdGraphSchema_K2.cpp
@@ -318,7 +318,6 @@ bool UEdGraphSchema_K2::CanFunctionBeUsedInClass(const UClass* InClass, UFunctio
 		const bool bFunctionStatic = InFunction->HasAllFunctionFlags(FUNC_Static);
 		const bool bHasReturnParams = (InFunction->GetReturnProperty() != NULL);
 		const bool bHasArrayPointerParms = InFunction->HasMetaData(TEXT("ArrayParm"));
-		const bool bHasCustomStructureParams = InFunction->HasMetaData(TEXT("CustomStructureParam"));
 		const bool bAllowForEachCall = !bFunctionStatic && !bIsLatent && !bIsPureFunc && !bIsConstFunc && !bHasReturnParams && !bHasArrayPointerParms;
 
 		const bool bClassIsAnActor = InClass->IsChildOf( AActor::StaticClass() );
