BP: Refresh all nodes sort nodes, so ones using Wildcards are refreshed later.

TTP #341914 SwingNinja: MyCharacter_Test blueprint fails to compile.

diff --git a/Engine/Source/Editor/BlueprintGraph/Classes/K2Node.h b/Engine/Source/Editor/BlueprintGraph/Classes/K2Node.h
--- a/Engine/Source/Editor/BlueprintGraph/Classes/K2Node.h
+++ b/Engine/Source/Editor/BlueprintGraph/Classes/K2Node.h
@@ -309,6 +309,14 @@ class UK2Node : public UEdGraphNode
 	 */
 	virtual FText GetMenuCategory() const { return FText::GetEmpty(); }
 
+	BLUEPRINTGRAPH_API enum EBaseNodeRefreshPriority
+	{
+		Low_UsesDependentWildcard = 100,
+		Normal = 200,
+	};
+
+	BLUEPRINTGRAPH_API virtual int32 GetNodeRefreshPriority() const { return EBaseNodeRefreshPriority::Normal; }
+
 protected:
 	/** 
 	 * A mapping from old property and function names to new ones.  Get primed from INI files, and should contain entries for properties, functions, and delegates that get moved, so they can be fixed up
diff --git a/Engine/Source/Editor/BlueprintGraph/Classes/K2Node_AssignmentStatement.h b/Engine/Source/Editor/BlueprintGraph/Classes/K2Node_AssignmentStatement.h
--- a/Engine/Source/Editor/BlueprintGraph/Classes/K2Node_AssignmentStatement.h
+++ b/Engine/Source/Editor/BlueprintGraph/Classes/K2Node_AssignmentStatement.h
@@ -28,6 +28,7 @@ class UK2Node_AssignmentStatement : public UK2Node
 	virtual class FNodeHandlingFunctor* CreateNodeHandler(class FKismetCompilerContext& CompilerContext) const override;
 	virtual void GetMenuActions(TArray<UBlueprintNodeSpawner*>& ActionListOut) const override;
 	virtual FText GetMenuCategory() const override;
+	virtual int32 GetNodeRefreshPriority() const override { return EBaseNodeRefreshPriority::Low_UsesDependentWildcard; }
 	// End UK2Node interface
 
 	/** Get the Then output pin */
diff --git a/Engine/Source/Editor/BlueprintGraph/Classes/K2Node_EnumEquality.h b/Engine/Source/Editor/BlueprintGraph/Classes/K2Node_EnumEquality.h
--- a/Engine/Source/Editor/BlueprintGraph/Classes/K2Node_EnumEquality.h
+++ b/Engine/Source/Editor/BlueprintGraph/Classes/K2Node_EnumEquality.h
@@ -29,6 +29,7 @@ class UK2Node_EnumEquality : public UK2Node
 	virtual void ExpandNode(class FKismetCompilerContext& CompilerContext, UEdGraph* SourceGraph) override;
 	virtual void GetMenuActions(TArray<UBlueprintNodeSpawner*>& ActionListOut) const override;
 	virtual FText GetMenuCategory() const override;
+	virtual int32 GetNodeRefreshPriority() const override { return EBaseNodeRefreshPriority::Low_UsesDependentWildcard; }
 	// End UK2Node interface
 
 	/** Get the return value pin */
diff --git a/Engine/Source/Editor/BlueprintGraph/Classes/K2Node_GetEnumeratorName.h b/Engine/Source/Editor/BlueprintGraph/Classes/K2Node_GetEnumeratorName.h
--- a/Engine/Source/Editor/BlueprintGraph/Classes/K2Node_GetEnumeratorName.h
+++ b/Engine/Source/Editor/BlueprintGraph/Classes/K2Node_GetEnumeratorName.h
@@ -30,6 +30,7 @@ class UK2Node_GetEnumeratorName : public UK2Node
 	virtual void NotifyPinConnectionListChanged(UEdGraphPin* Pin) override;
 	virtual void GetMenuActions(TArray<UBlueprintNodeSpawner*>& ActionListOut) const override;
 	virtual FText GetMenuCategory() const override;
+	virtual int32 GetNodeRefreshPriority() const override { return EBaseNodeRefreshPriority::Low_UsesDependentWildcard; }
 	// End UK2Node interface
 
 	void UpdatePinType();
diff --git a/Engine/Source/Editor/BlueprintGraph/Classes/K2Node_Knot.h b/Engine/Source/Editor/BlueprintGraph/Classes/K2Node_Knot.h
--- a/Engine/Source/Editor/BlueprintGraph/Classes/K2Node_Knot.h
+++ b/Engine/Source/Editor/BlueprintGraph/Classes/K2Node_Knot.h
@@ -27,6 +27,7 @@ public:
 	virtual void ExpandNode(class FKismetCompilerContext& CompilerContext, UEdGraph* SourceGraph) override;
 	virtual void NotifyPinConnectionListChanged(UEdGraphPin* Pin) override;
 	virtual void PostReconstructNode() override;
+	virtual int32 GetNodeRefreshPriority() const override { return EBaseNodeRefreshPriority::Low_UsesDependentWildcard; }
 	// End of UK2Node interface
 
 	UEdGraphPin* GetInputPin() const
diff --git a/Engine/Source/Editor/BlueprintGraph/Classes/K2Node_MacroInstance.h b/Engine/Source/Editor/BlueprintGraph/Classes/K2Node_MacroInstance.h
--- a/Engine/Source/Editor/BlueprintGraph/Classes/K2Node_MacroInstance.h
+++ b/Engine/Source/Editor/BlueprintGraph/Classes/K2Node_MacroInstance.h
@@ -56,6 +56,7 @@ public:
 	virtual bool HasExternalBlueprintDependencies(TArray<class UStruct*>* OptionalOutput) const override;
 	virtual void GetNodeAttributes( TArray<TKeyValuePair<FString, FString>>& OutNodeAttributes ) const override;
 	virtual FText GetMenuCategory() const override;
+	virtual int32 GetNodeRefreshPriority() const override { return EBaseNodeRefreshPriority::Low_UsesDependentWildcard; }
 	// End UK2Node interface
 
 	void SetMacroGraph(UEdGraph* Graph) { MacroGraphReference.SetGraph(Graph); }
diff --git a/Engine/Source/Editor/BlueprintGraph/Classes/K2Node_MakeArray.h b/Engine/Source/Editor/BlueprintGraph/Classes/K2Node_MakeArray.h
--- a/Engine/Source/Editor/BlueprintGraph/Classes/K2Node_MakeArray.h
+++ b/Engine/Source/Editor/BlueprintGraph/Classes/K2Node_MakeArray.h
@@ -38,6 +38,7 @@ public:
 	virtual bool IsConnectionDisallowed(const UEdGraphPin* MyPin, const UEdGraphPin* OtherPin, FString& OutReason) const override;
 	virtual void GetMenuActions(TArray<UBlueprintNodeSpawner*>& ActionListOut) const override;
 	virtual FText GetMenuCategory() const override;
+	virtual int32 GetNodeRefreshPriority() const override { return EBaseNodeRefreshPriority::Low_UsesDependentWildcard; }
 	// End of UK2Node interface
 
 
diff --git a/Engine/Source/Editor/BlueprintGraph/Classes/K2Node_Select.h b/Engine/Source/Editor/BlueprintGraph/Classes/K2Node_Select.h
--- a/Engine/Source/Editor/BlueprintGraph/Classes/K2Node_Select.h
+++ b/Engine/Source/Editor/BlueprintGraph/Classes/K2Node_Select.h
@@ -54,6 +54,7 @@ class UK2Node_Select : public UK2Node
 	virtual bool IsConnectionDisallowed(const UEdGraphPin* MyPin, const UEdGraphPin* OtherPin, FString& OutReason) const override;
 	virtual void GetMenuActions(TArray<UBlueprintNodeSpawner*>& ActionListOut) const override;
 	virtual FText GetMenuCategory() const override;
+	virtual int32 GetNodeRefreshPriority() const override { return EBaseNodeRefreshPriority::Low_UsesDependentWildcard; }
 	// End UK2Node interface
 
 	/** Get the return value pin */
diff --git a/Engine/Source/Editor/BlueprintGraph/Classes/K2Node_VariableSetRef.h b/Engine/Source/Editor/BlueprintGraph/Classes/K2Node_VariableSetRef.h
--- a/Engine/Source/Editor/BlueprintGraph/Classes/K2Node_VariableSetRef.h
+++ b/Engine/Source/Editor/BlueprintGraph/Classes/K2Node_VariableSetRef.h
@@ -21,6 +21,7 @@ class UK2Node_VariableSetRef : public UK2Node
 	virtual class FNodeHandlingFunctor* CreateNodeHandler(class FKismetCompilerContext& CompilerContext) const override;
 	virtual void GetMenuActions(TArray<UBlueprintNodeSpawner*>& ActionListOut) const override;
 	virtual FText GetMenuCategory() const override;
+	virtual int32 GetNodeRefreshPriority() const override { return EBaseNodeRefreshPriority::Low_UsesDependentWildcard; }
 	// End UK2Node interface
 
 	/** 
diff --git a/Engine/Source/Editor/UnrealEd/Private/Kismet2/BlueprintEditorUtils.cpp b/Engine/Source/Editor/UnrealEd/Private/Kismet2/BlueprintEditorUtils.cpp
--- a/Engine/Source/Editor/UnrealEd/Private/Kismet2/BlueprintEditorUtils.cpp
+++ b/Engine/Source/Editor/UnrealEd/Private/Kismet2/BlueprintEditorUtils.cpp
@@ -51,7 +51,12 @@ struct FCompareNodePriority
 		const bool NodeAChangesStructure = A.NodeCausesStructuralBlueprintChange();
 		const bool NodeBChangesStructure = B.NodeCausesStructuralBlueprintChange();
 
-		return ((NodeAChangesStructure > NodeBChangesStructure) ? true : false);
+		if (NodeAChangesStructure != NodeBChangesStructure)
+		{
+			return NodeAChangesStructure;
+		}
+		
+		return A.GetNodeRefreshPriority() > B.GetNodeRefreshPriority();
 	}
 };
 
