Abort Blueprint compilation if an added component has static mobility when its parent does not.
#ttp 341133 - BLUEPRINT: Cant spawn playerstart using childActorComponent
#branch UE4
#proj Editor.BlueprintGraph
#add Partially disallowed by putting a check at BP compile time that if the child actor has static mobility, then the ChildActorComponent (the parent) must also be static. However this is not infallible, and it will still be possible to attach a static component to a non-static one, higher up the chain.  I explored the code a little, but couldn't find a way of checking for this at BP compile time, as I couldn't find a hierarchy of components at this point.

diff --git a/Engine/Source/Editor/BlueprintGraph/Private/K2Node_AddComponent.cpp b/Engine/Source/Editor/BlueprintGraph/Private/K2Node_AddComponent.cpp
--- a/Engine/Source/Editor/BlueprintGraph/Private/K2Node_AddComponent.cpp
+++ b/Engine/Source/Editor/BlueprintGraph/Private/K2Node_AddComponent.cpp
@@ -195,6 +195,23 @@ void UK2Node_AddComponent::ValidateNodeDuringCompilation(FCompilerResultsLog& Me
 					MessageLog.Error(*FText::Format(NSLOCTEXT("KismetCompiler", "AddSelfComponent_Error", "@@ cannot add a '{ChildActorClass}' component in the construction script (could cause infinite recursion)."), Args).ToString(), this);
 				}
 			}
+			else if (ChildActorClass != nullptr)
+			{
+				AActor const* ChildActor = Cast<AActor>(ChildActorClass->ClassDefaultObject);
+				check(ChildActor != nullptr);
+				if (ChildActor->GetRootComponent() == nullptr)
+				{
+					FFormatNamedArguments Args;
+					Args.Add(TEXT("ChildActorClass"), FText::FromString(ChildActorClass->GetName()));
+					MessageLog.Error(*FText::Format(NSLOCTEXT("KismetCompiler", "AddComponentWithoutRoot_Error", "@@ cannot add a '{ChildActorClass}' component as it has no root component."), Args).ToString(), this);
+				}
+				else if (ChildActor->GetRootComponent()->Mobility == EComponentMobility::Static && ChildActorComponent->Mobility != EComponentMobility::Static)
+				{
+					FFormatNamedArguments Args;
+					Args.Add(TEXT("ChildActorClass"), FText::FromString(ChildActorClass->GetName()));
+					MessageLog.Error(*FText::Format(NSLOCTEXT("KismetCompiler", "AddStaticChildActorComponent_Error", "@@ cannot add a '{ChildActorClass}' component as it has static mobility, and the ChildActorComponent does not."), Args).ToString(), this);
+				}
+			}
 		}
 	}
 	else
