#UE4 Fix for TTP#342132.. Event node patching logic could create nodes in non-event graphs..

diff --git a/Engine/Source/Editor/BlueprintGraph/Private/EdGraphSchema_K2.cpp b/Engine/Source/Editor/BlueprintGraph/Private/EdGraphSchema_K2.cpp
--- a/Engine/Source/Editor/BlueprintGraph/Private/EdGraphSchema_K2.cpp
+++ b/Engine/Source/Editor/BlueprintGraph/Private/EdGraphSchema_K2.cpp
@@ -4098,6 +4098,12 @@ UEdGraphNode* UEdGraphSchema_K2::CreateSubstituteNode(UEdGraphNode* Node, const
 			Graph = Node->GetGraph();
 		}
 
+		// Can only place events in ubergraphs
+		if (GetGraphType(Graph) != EGraphType::GT_Ubergraph)
+		{
+			return NULL;
+		}
+
 		// Find the Blueprint that owns the graph
 		UBlueprint* Blueprint = Graph ? FBlueprintEditorUtils::FindBlueprintForGraph(Graph) : NULL;
 		if(Blueprint && Blueprint->SkeletonGeneratedClass)
