BP: self node: 'self' pin should not have SubCategoryObject specifed (it's confusing and error prone). The change is related to cl#2235695

diff --git a/Engine/Source/Editor/BlueprintGraph/Classes/K2Node_Self.h b/Engine/Source/Editor/BlueprintGraph/Classes/K2Node_Self.h
--- a/Engine/Source/Editor/BlueprintGraph/Classes/K2Node_Self.h
+++ b/Engine/Source/Editor/BlueprintGraph/Classes/K2Node_Self.h
@@ -10,10 +10,6 @@ class UK2Node_Self : public UK2Node
 {
 	GENERATED_UCLASS_BODY()
 
-	/** Class that this variable is defined in.  */
-	UPROPERTY(transient)
-	TSubclassOf<class UObject>  SelfClass;
-
 	// Begin UEdGraphNode interface
 	virtual FText GetNodeTitle(ENodeTitleType::Type TitleType) const override;
 	virtual FString GetTooltip() const override;
diff --git a/Engine/Source/Editor/BlueprintGraph/Private/K2ActionMenuBuilder.cpp b/Engine/Source/Editor/BlueprintGraph/Private/K2ActionMenuBuilder.cpp
--- a/Engine/Source/Editor/BlueprintGraph/Private/K2ActionMenuBuilder.cpp
+++ b/Engine/Source/Editor/BlueprintGraph/Private/K2ActionMenuBuilder.cpp
@@ -1833,7 +1833,6 @@ void FK2ActionMenuBuilder::GetVariableGettersSettersForClass(FBlueprintGraphActi
 		TSharedPtr<FEdGraphSchemaAction_K2NewNode> SelfVar = FK2ActionMenuBuilder::AddNewNodeAction(ContextMenuBuilder, K2ActionCategories::VariablesCategory, FText::FromString(SelfString), SelfTooltip);
 
 		UK2Node_Self* SelfNode = ContextMenuBuilder.CreateTemplateNode<UK2Node_Self>();
-		SelfNode->SelfClass = (UClass*)Class;
 		SelfVar->NodeTemplate = SelfNode;
 		SelfVar->Keywords = SelfNode->GetKeywords();
 	}
diff --git a/Engine/Source/Editor/BlueprintGraph/Private/K2Node_Self.cpp b/Engine/Source/Editor/BlueprintGraph/Private/K2Node_Self.cpp
--- a/Engine/Source/Editor/BlueprintGraph/Private/K2Node_Self.cpp
+++ b/Engine/Source/Editor/BlueprintGraph/Private/K2Node_Self.cpp
@@ -38,10 +38,8 @@ UK2Node_Self::UK2Node_Self(const class FPostConstructInitializeProperties& PCIP)
 
 void UK2Node_Self::AllocateDefaultPins()
 {
-	SelfClass = GetBlueprint()->GeneratedClass;
-
 	const UEdGraphSchema_K2* K2Schema = GetDefault<UEdGraphSchema_K2>();
-	CreatePin(EGPD_Output, K2Schema->PC_Object, K2Schema->PSC_Self, SelfClass, false, false, K2Schema->PN_Self);
+	CreatePin(EGPD_Output, K2Schema->PC_Object, K2Schema->PSC_Self, NULL, false, false, K2Schema->PN_Self);
 
 	Super::AllocateDefaultPins();
 }
@@ -75,7 +73,7 @@ void UK2Node_Self::ValidateNodeDuringCompilation(class FCompilerResultsLog& Mess
 {
 	Super::ValidateNodeDuringCompilation(MessageLog);
 	const UBlueprint* Blueprint = GetBlueprint();
-	const UClass* MyClass = SelfClass ? SelfClass : (Blueprint ? Blueprint->GeneratedClass : NULL);
+	const UClass* MyClass = Blueprint ? Blueprint->GeneratedClass : NULL;
 	const bool bValidClass = !MyClass || !MyClass->IsChildOf(UBlueprintFunctionLibrary::StaticClass());
 	if (!bValidClass)
 	{
