#UE4 fix #342132 for async tasks..

diff --git a/Engine/Source/Editor/BlueprintGraph/Classes/K2Node_BaseAsyncTask.h b/Engine/Source/Editor/BlueprintGraph/Classes/K2Node_BaseAsyncTask.h
--- a/Engine/Source/Editor/BlueprintGraph/Classes/K2Node_BaseAsyncTask.h
+++ b/Engine/Source/Editor/BlueprintGraph/Classes/K2Node_BaseAsyncTask.h
@@ -20,6 +20,7 @@ class BLUEPRINTGRAPH_API UK2Node_BaseAsyncTask : public UK2Node
 	virtual void GetMenuEntries(FGraphContextMenuBuilder& ContextMenuBuilder) const override;
 	virtual FString GetTooltip() const override;
 	virtual FText GetNodeTitle(ENodeTitleType::Type TitleType) const override;
+	virtual bool CanPasteHere(const UEdGraph* TargetGraph, const UEdGraphSchema* Schema) const override;
 	// End of UEdGraphNode interface
 
 	// UK2Node interface
diff --git a/Engine/Source/Editor/BlueprintGraph/Private/K2Node_BaseAsyncTask.cpp b/Engine/Source/Editor/BlueprintGraph/Private/K2Node_BaseAsyncTask.cpp
--- a/Engine/Source/Editor/BlueprintGraph/Private/K2Node_BaseAsyncTask.cpp
+++ b/Engine/Source/Editor/BlueprintGraph/Private/K2Node_BaseAsyncTask.cpp
@@ -31,6 +31,17 @@ FText UK2Node_BaseAsyncTask::GetNodeTitle(ENodeTitleType::Type TitleType) const
 	return FText::FromString(FunctionToolTipText);
 }
 
+bool UK2Node_BaseAsyncTask::CanPasteHere(const UEdGraph* TargetGraph, const UEdGraphSchema* Schema) const
+{
+	// Can only place events in ubergraphs, and basicasync task creates an event node:
+	if (Schema->GetGraphType(TargetGraph) != EGraphType::GT_Ubergraph)
+	{
+		UE_LOG(LogBlueprint, Log, TEXT("Cannot paste event node (%s) in this graph because it is not an ubergraph."), *GetFName().ToString());
+		return false;
+	}
+	return UK2Node::CanPasteHere(TargetGraph, Schema);
+}
+
 void UK2Node_BaseAsyncTask::GetMenuEntries(FGraphContextMenuBuilder& ContextMenuBuilder) const
 {
 	const UEdGraphSchema_K2* K2Schema = GetDefault<UEdGraphSchema_K2>();
