No longer preventing child-actors from being added if they don't have a root component.
(TTP #342536)

diff --git a/Engine/Source/Editor/BlueprintGraph/Private/K2Node_AddComponent.cpp b/Engine/Source/Editor/BlueprintGraph/Private/K2Node_AddComponent.cpp
--- a/Engine/Source/Editor/BlueprintGraph/Private/K2Node_AddComponent.cpp
+++ b/Engine/Source/Editor/BlueprintGraph/Private/K2Node_AddComponent.cpp
@@ -199,13 +199,9 @@ void UK2Node_AddComponent::ValidateNodeDuringCompilation(FCompilerResultsLog& Me
 			{
 				AActor const* ChildActor = Cast<AActor>(ChildActorClass->ClassDefaultObject);
 				check(ChildActor != nullptr);
-				if (ChildActor->GetRootComponent() == nullptr)
-				{
-					FFormatNamedArguments Args;
-					Args.Add(TEXT("ChildActorClass"), FText::FromString(ChildActorClass->GetName()));
-					MessageLog.Error(*FText::Format(NSLOCTEXT("KismetCompiler", "AddComponentWithoutRoot_Error", "@@ cannot add a '{ChildActorClass}' component as it has no root component."), Args).ToString(), this);
-				}
-				else if (ChildActor->GetRootComponent()->Mobility == EComponentMobility::Static && ChildActorComponent->Mobility != EComponentMobility::Static)
+				USceneComponent* RootComponent = ChildActor->GetRootComponent();
+
+				if ((RootComponent != nullptr) && (RootComponent->Mobility == EComponentMobility::Static) && (ChildActorComponent->Mobility != EComponentMobility::Static))
 				{
 					FFormatNamedArguments Args;
 					Args.Add(TEXT("ChildActorClass"), FText::FromString(ChildActorClass->GetName()));
