Catching the case where struct ref pins on the new set-struct-members nodes are left empty (now presenting the user a compiler error rather than a crash)
(TTP #342994)

diff --git a/Engine/Source/Editor/BlueprintGraph/Classes/K2Node_SetFieldsInStruct.h b/Engine/Source/Editor/BlueprintGraph/Classes/K2Node_SetFieldsInStruct.h
--- a/Engine/Source/Editor/BlueprintGraph/Classes/K2Node_SetFieldsInStruct.h
+++ b/Engine/Source/Editor/BlueprintGraph/Classes/K2Node_SetFieldsInStruct.h
@@ -14,6 +14,7 @@ class UK2Node_SetFieldsInStruct : public UK2Node_MakeStruct
 	virtual FText GetNodeTitle(ENodeTitleType::Type TitleType) const override;
 	virtual FString GetTooltip() const override;
 	virtual FName GetPaletteIcon(FLinearColor& OutColor) const override;
+	virtual void ValidateNodeDuringCompilation(class FCompilerResultsLog& MessageLog) const override;
 	// End  UEdGraphNode interface
 
 	// Begin K2Node interface
diff --git a/Engine/Source/Editor/BlueprintGraph/Private/K2Node_SetFieldsInStruct.cpp b/Engine/Source/Editor/BlueprintGraph/Private/K2Node_SetFieldsInStruct.cpp
--- a/Engine/Source/Editor/BlueprintGraph/Private/K2Node_SetFieldsInStruct.cpp
+++ b/Engine/Source/Editor/BlueprintGraph/Private/K2Node_SetFieldsInStruct.cpp
@@ -3,6 +3,7 @@
 #include "BlueprintGraphPrivatePCH.h"
 #include "K2Node_SetFieldsInStruct.h"
 #include "MakeStructHandler.h"
+#include "CompilerResultsLog.h"
 
 #define LOCTEXT_NAMESPACE "K2Node_MakeStruct"
 
@@ -72,6 +73,20 @@ FName UK2Node_SetFieldsInStruct::GetPaletteIcon(FLinearColor& OutColor) const
 	return UK2Node_Variable::GetPaletteIcon(OutColor);
 }
 
+void UK2Node_SetFieldsInStruct::ValidateNodeDuringCompilation(FCompilerResultsLog& MessageLog) const
+{
+	Super::ValidateNodeDuringCompilation(MessageLog);
+
+	if (UEdGraphPin* FoundPin = FindPinChecked(SetFieldsInStructHelper::StructRefPinName()))
+	{
+		if (FoundPin->LinkedTo.Num() <= 0)
+		{
+			FText ErrorMessage = LOCTEXT("SetStructFields_NoStructRefError", "The @@ pin must be connected to the struct that you wish to set.");
+			MessageLog.Error(*ErrorMessage.ToString(), FoundPin);
+		}
+	}
+}
+
 FNodeHandlingFunctor* UK2Node_SetFieldsInStruct::CreateNodeHandler(class FKismetCompilerContext& CompilerContext) const
 {
 	return new FKCHandler_SetFieldsInStruct(CompilerContext);
diff --git a/Engine/Source/Editor/BlueprintGraph/Private/MakeStructHandler.cpp b/Engine/Source/Editor/BlueprintGraph/Private/MakeStructHandler.cpp
--- a/Engine/Source/Editor/BlueprintGraph/Private/MakeStructHandler.cpp
+++ b/Engine/Source/Editor/BlueprintGraph/Private/MakeStructHandler.cpp
@@ -53,14 +53,21 @@ void FKCHandler_MakeStruct::RegisterNets(FKismetFunctionContext& Context, UEdGra
 	check(NULL != Net);
 	FBPTerminal** FoundTerm = Context.NetMap.Find(Net);
 	FBPTerminal* Term = FoundTerm ? *FoundTerm : NULL;
-	check(NULL != Term);
 
-	UStruct* StructInTerm = Cast<UStruct>(Term->Type.PinSubCategoryObject.Get());
-	if (NULL == StructInTerm || !StructInTerm->IsChildOf(Node->StructType))
+	if (Term == nullptr)
 	{
-		CompilerContext.MessageLog.Error(*LOCTEXT("MakeStruct_NoMatch_Error", "Structures don't match for @@").ToString(), Node);
+		CompilerContext.MessageLog.Error(*LOCTEXT("MakeStruct_NoTerm_Error", "Failed to generate a term for the @@ pin; was it a struct reference that was left unset?").ToString(), OutputPin);
+	}
+	else
+	{
+		UStruct* StructInTerm = Cast<UStruct>(Term->Type.PinSubCategoryObject.Get());
+		if (NULL == StructInTerm || !StructInTerm->IsChildOf(Node->StructType))
+		{
+			CompilerContext.MessageLog.Error(*LOCTEXT("MakeStruct_NoMatch_Error", "Structures don't match for @@").ToString(), Node);
+		}
 	}
 }
+	
 
 void FKCHandler_MakeStruct::RegisterNet(FKismetFunctionContext& Context, UEdGraphPin* Net)
 {
