(WIP) Making NodeSpawners responsible for the placement of new nodes (passing the desired location along to them, to do with as they please). This fixes comment nodes so they properly position themselves in the new system.

diff --git a/Engine/Source/Editor/BlueprintGraph/Classes/BlueprintComponentNodeSpawner.h b/Engine/Source/Editor/BlueprintGraph/Classes/BlueprintComponentNodeSpawner.h
--- a/Engine/Source/Editor/BlueprintGraph/Classes/BlueprintComponentNodeSpawner.h
+++ b/Engine/Source/Editor/BlueprintGraph/Classes/BlueprintComponentNodeSpawner.h
@@ -31,9 +31,9 @@ public:
 	static UBlueprintComponentNodeSpawner* Create(TSubclassOf<UActorComponent> const ComponentClass, UObject* Outer = nullptr);
 
 	// UBlueprintNodeSpawner interface
-	virtual UEdGraphNode* Invoke(UEdGraph* ParentGraph) const override;
-	virtual FText    GetDefaultMenuName() const override;
-	virtual FText    GetDefaultMenuCategory() const override;
+	virtual UEdGraphNode* Invoke(UEdGraph* ParentGraph, FVector2D const Location) const override;
+	virtual FText GetDefaultMenuName() const override;
+	virtual FText GetDefaultMenuCategory() const override;
 	// End UBlueprintNodeSpawner interface
 	
 	/**
diff --git a/Engine/Source/Editor/BlueprintGraph/Classes/BlueprintEventNodeSpawner.h b/Engine/Source/Editor/BlueprintGraph/Classes/BlueprintEventNodeSpawner.h
--- a/Engine/Source/Editor/BlueprintGraph/Classes/BlueprintEventNodeSpawner.h
+++ b/Engine/Source/Editor/BlueprintGraph/Classes/BlueprintEventNodeSpawner.h
@@ -44,7 +44,7 @@ public:
 	static UBlueprintEventNodeSpawner* Create(TSubclassOf<UK2Node_Event> NodeClass, FName CustomEventName, UObject* Outer = nullptr);
 
 	// UBlueprintNodeSpawner interface
-	virtual UEdGraphNode* Invoke(UEdGraph* ParentGraph) const override;
+	virtual UEdGraphNode* Invoke(UEdGraph* ParentGraph, FVector2D const Location) const override;
 	virtual FText GetDefaultMenuName() const override;
 	virtual FText GetDefaultSearchKeywords() const override;
 	// End UBlueprintNodeSpawner interface
diff --git a/Engine/Source/Editor/BlueprintGraph/Classes/BlueprintFunctionNodeSpawner.h b/Engine/Source/Editor/BlueprintGraph/Classes/BlueprintFunctionNodeSpawner.h
--- a/Engine/Source/Editor/BlueprintGraph/Classes/BlueprintFunctionNodeSpawner.h
+++ b/Engine/Source/Editor/BlueprintGraph/Classes/BlueprintFunctionNodeSpawner.h
@@ -27,7 +27,7 @@ public:
 	static UBlueprintFunctionNodeSpawner* Create(UFunction const* const Function, UObject* Outer = nullptr);
 
 	// UBlueprintNodeSpawner interface
-	virtual UEdGraphNode* Invoke(UEdGraph* ParentGraph) const override;
+	virtual UEdGraphNode* Invoke(UEdGraph* ParentGraph, FVector2D const Location) const override;
 	virtual FText GetDefaultMenuName() const override;
 	virtual FText GetDefaultMenuCategory() const override;
 	virtual FText GetDefaultMenuTooltip() const override;
diff --git a/Engine/Source/Editor/BlueprintGraph/Classes/BlueprintPropertyNodeSpawner.h b/Engine/Source/Editor/BlueprintGraph/Classes/BlueprintPropertyNodeSpawner.h
--- a/Engine/Source/Editor/BlueprintGraph/Classes/BlueprintPropertyNodeSpawner.h
+++ b/Engine/Source/Editor/BlueprintGraph/Classes/BlueprintPropertyNodeSpawner.h
@@ -45,7 +45,7 @@ public:
 	}
 
 	// UBlueprintNodeSpawner interface
-	virtual UEdGraphNode* Invoke(UEdGraph* ParentGraph) const override;
+	virtual UEdGraphNode* Invoke(UEdGraph* ParentGraph, FVector2D const Location) const override;
 	virtual FText GetDefaultMenuName() const override;
 	virtual FText GetDefaultMenuCategory() const override;
 	virtual FName GetDefaultMenuIcon(FLinearColor& ColorOut) override;
diff --git a/Engine/Source/Editor/BlueprintGraph/Private/BlueprintActionDatabase.cpp b/Engine/Source/Editor/BlueprintGraph/Private/BlueprintActionDatabase.cpp
--- a/Engine/Source/Editor/BlueprintGraph/Private/BlueprintActionDatabase.cpp
+++ b/Engine/Source/Editor/BlueprintGraph/Private/BlueprintActionDatabase.cpp
@@ -220,8 +220,12 @@ static UBlueprintNodeSpawner* FBlueprintNodeSpawnerFactory::MakeCommentNodeSpawn
 		UBlueprint* Blueprint = FBlueprintEditorUtils::FindBlueprintForGraph(OuterGraph);
 		check(Blueprint != nullptr);
 
-		FSlateRect Bounds;
-		FKismetEditorUtilities::GetBoundsForSelectedNodes(Blueprint, Bounds, 50.0f);
+		float const OldNodePosX = NewNode->NodePosX;
+		float const OldNodePosY = NewNode->NodePosY;
+
+		static const float CommentNodePadding = 50.0f;
+		FSlateRect Bounds(OldNodePosX - CommentNodePadding, OldNodePosY - CommentNodePadding, OldNodePosX + CommentNodePadding, OldNodePosY + CommentNodePadding);
+		FKismetEditorUtilities::GetBoundsForSelectedNodes(Blueprint, Bounds, CommentNodePadding);
 		CommentNode->SetBounds(Bounds);
 	};
 
diff --git a/Engine/Source/Editor/BlueprintGraph/Private/BlueprintComponentNodeSpawner.cpp b/Engine/Source/Editor/BlueprintGraph/Private/BlueprintComponentNodeSpawner.cpp
--- a/Engine/Source/Editor/BlueprintGraph/Private/BlueprintComponentNodeSpawner.cpp
+++ b/Engine/Source/Editor/BlueprintGraph/Private/BlueprintComponentNodeSpawner.cpp
@@ -36,30 +36,31 @@ UBlueprintComponentNodeSpawner::UBlueprintComponentNodeSpawner(class FPostConstr
 //------------------------------------------------------------------------------
 // Evolved from a combination of FK2ActionMenuBuilder::CreateAddComponentAction()
 // and FEdGraphSchemaAction_K2AddComponent::PerformAction().
-UEdGraphNode* UBlueprintComponentNodeSpawner::Invoke(UEdGraph* ParentGraph) const
+UEdGraphNode* UBlueprintComponentNodeSpawner::Invoke(UEdGraph* ParentGraph, FVector2D const Location) const
 {
-	UK2Node_AddComponent* NewNode = NewObject<UK2Node_AddComponent>(ParentGraph);
-	UBlueprint* Blueprint = NewNode->GetBlueprint();
-
-	UFunction* AddComponentFunc = FindFieldChecked<UFunction>(AActor::StaticClass(), UK2Node_AddComponent::GetAddComponentFunctionName());
-	NewNode->FunctionReference.SetFromField<UFunction>(AddComponentFunc, FBlueprintEditorUtils::IsActorBased(Blueprint));
-
-	bool bIsTemplateNode = ParentGraph->HasAnyFlags(RF_Transient);
-	if (CustomizeNodeDelegate.IsBound())
+	check(ComponentClass != nullptr);
+	
+	auto PostSpawnLambda = [](UEdGraphNode* NewNode, bool bIsTemplateNode, FCustomizeNodeDelegate UserDelegate)
 	{
-		CustomizeNodeDelegate.Execute(NewNode, bIsTemplateNode);
-	}
+		UserDelegate.ExecuteIfBound(NewNode, bIsTemplateNode);
+		
+		UK2Node_AddComponent* AddCompNode = CastChecked<UK2Node_AddComponent>(NewNode);
+		UBlueprint* Blueprint = AddCompNode->GetBlueprint();
+		
+		UFunction* AddComponentFunc = FindFieldChecked<UFunction>(AActor::StaticClass(), UK2Node_AddComponent::GetAddComponentFunctionName());
+		AddCompNode->FunctionReference.SetFromField<UFunction>(AddComponentFunc, FBlueprintEditorUtils::IsActorBased(Blueprint));
+	};
+	FCustomizeNodeDelegate PostSpawnDelegate = FCustomizeNodeDelegate::CreateStatic(PostSpawnLambda, CustomizeNodeDelegate);
+	
+	UK2Node_AddComponent* NewNode = CastChecked<UK2Node_AddComponent>(Super::Invoke(ParentGraph, Location, PostSpawnDelegate));
 
+	bool bIsTemplateNode = ParentGraph->HasAnyFlags(RF_Transient);
 	if (!bIsTemplateNode)
 	{
-		NewNode->SetFlags(RF_Transactional);
-		NewNode->AllocateDefaultPins();
-
-		check(ComponentClass != nullptr);
+		UBlueprint* Blueprint = NewNode->GetBlueprint();
 		UActorComponent* ComponentTemplate = ConstructObject<UActorComponent>(ComponentClass, Blueprint->GeneratedClass);
 		ComponentTemplate->SetFlags(RF_ArchetypeObject);
 
-		// @TODO: what is this is a template?
 		Blueprint->ComponentTemplates.Add(ComponentTemplate);
 
 		// set the name of the template as the default for the TemplateName param
diff --git a/Engine/Source/Editor/BlueprintGraph/Private/BlueprintEventNodeSpawner.cpp b/Engine/Source/Editor/BlueprintGraph/Private/BlueprintEventNodeSpawner.cpp
--- a/Engine/Source/Editor/BlueprintGraph/Private/BlueprintEventNodeSpawner.cpp
+++ b/Engine/Source/Editor/BlueprintGraph/Private/BlueprintEventNodeSpawner.cpp
@@ -86,7 +86,7 @@ UBlueprintEventNodeSpawner::UBlueprintEventNodeSpawner(class FPostConstructIniti
 }
 
 //------------------------------------------------------------------------------
-UEdGraphNode* UBlueprintEventNodeSpawner::Invoke(UEdGraph* ParentGraph) const
+UEdGraphNode* UBlueprintEventNodeSpawner::Invoke(UEdGraph* ParentGraph, FVector2D const Location) const
 {
 	check(ParentGraph != nullptr);
 	UBlueprint* Blueprint = FBlueprintEditorUtils::FindBlueprintForGraphChecked(ParentGraph);
@@ -116,25 +116,25 @@ UEdGraphNode* UBlueprintEventNodeSpawner::Invoke(UEdGraph* ParentGraph) const
 	// if there is no existing node, then we can happily spawn one into the graph
 	if (EventNode == nullptr)
 	{
-		bool bIsTemplateNode = ParentGraph->HasAnyFlags(RF_Transient);
-
-		EventNode = CastChecked<UK2Node_Event>(Super::Invoke(ParentGraph));
-		if (EventFunc != nullptr)
-		{
-			EventNode->EventSignatureName  = EventName;
-			EventNode->EventSignatureClass = ClassOwner;
-			EventNode->bOverrideFunction   = true;
-		}
-		else if (!bIsTemplateNode)
+		auto PostSpawnLambda = [](UEdGraphNode* NewNode, bool bIsTemplateNode, UFunction const* EventFunc, FName EventName, FCustomizeNodeDelegate UserDelegate)
 		{
-			EventNode->CustomFunctionName  = CustomEventName;
-		}
-
-		if (!bIsTemplateNode)
-		{
-			EventNode->SetFlags(RF_Transactional);
-			EventNode->AllocateDefaultPins();
-		}
+			UK2Node_Event* EventNode = CastChecked<UK2Node_Event>(NewNode);
+			if (EventFunc != nullptr)
+			{
+				EventNode->EventSignatureName  = EventName;
+				EventNode->EventSignatureClass = EventFunc->GetOuterUClass();
+				EventNode->bOverrideFunction   = true;
+			}
+			else if (!bIsTemplateNode)
+			{
+				EventNode->CustomFunctionName = EventName;
+			}
+
+			UserDelegate.ExecuteIfBound(NewNode, bIsTemplateNode);
+		};
+
+		FCustomizeNodeDelegate PostSpawnDelegate = FCustomizeNodeDelegate::CreateStatic(PostSpawnLambda, EventFunc, EventName, CustomizeNodeDelegate);
+		EventNode = Cast<UK2Node_Event>(Super::Invoke(ParentGraph, Location, PostSpawnDelegate));
 	}
 	// else, a node for this event already exists, and we should return that 
 	// (the FBlueprintActionMenuItem should detect this and focus in on it).
diff --git a/Engine/Source/Editor/BlueprintGraph/Private/BlueprintFunctionNodeSpawner.cpp b/Engine/Source/Editor/BlueprintGraph/Private/BlueprintFunctionNodeSpawner.cpp
--- a/Engine/Source/Editor/BlueprintGraph/Private/BlueprintFunctionNodeSpawner.cpp
+++ b/Engine/Source/Editor/BlueprintGraph/Private/BlueprintFunctionNodeSpawner.cpp
@@ -57,28 +57,22 @@ UBlueprintFunctionNodeSpawner::UBlueprintFunctionNodeSpawner(class FPostConstruc
 }
 
 //------------------------------------------------------------------------------
-UEdGraphNode* UBlueprintFunctionNodeSpawner::Invoke(UEdGraph* ParentGraph) const
+UEdGraphNode* UBlueprintFunctionNodeSpawner::Invoke(UEdGraph* ParentGraph, FVector2D const Location) const
 {
-	UEdGraphNode* NewNode = Super::Invoke(ParentGraph);
-	// user could have changed the node class (to something like
-	// UK2Node_BaseAsyncTask, which also wraps a function)
-	if (UK2Node_CallFunction* FuncNode = Cast<UK2Node_CallFunction>(NewNode))
+	auto PostSpawnSetupLambda = [](UEdGraphNode* NewNode, bool bIsTemplateNode, UFunction const* Function, FCustomizeNodeDelegate UserDelegate)
 	{
-		FuncNode->SetFromFunction(Function);
-	}
-	
-	bool bIsTemplateNode = ParentGraph->HasAnyFlags(RF_Transient);
-	if (CustomizeNodeDelegate.IsBound())
-	{
-		CustomizeNodeDelegate.Execute(NewNode, bIsTemplateNode);
-	}
-	if (!bIsTemplateNode)
-	{
-		NewNode->SetFlags(RF_Transactional);
-		NewNode->AllocateDefaultPins();
-	}
- 	
-	return NewNode;
+		// user could have changed the node class (to something like
+		// UK2Node_BaseAsyncTask, which also wraps a function)
+		if (UK2Node_CallFunction* FuncNode = Cast<UK2Node_CallFunction>(NewNode))
+		{
+			FuncNode->SetFromFunction(Function);
+		}
+
+		UserDelegate.ExecuteIfBound(NewNode, bIsTemplateNode);
+	};
+
+	FCustomizeNodeDelegate PostSpawnSetupDelegate = FCustomizeNodeDelegate::CreateStatic(PostSpawnSetupLambda, Function, CustomizeNodeDelegate);
+	return Super::Invoke(ParentGraph, Location, PostSpawnSetupDelegate);
 }
 
 //------------------------------------------------------------------------------
diff --git a/Engine/Source/Editor/BlueprintGraph/Private/BlueprintNodeSpawner.cpp b/Engine/Source/Editor/BlueprintGraph/Private/BlueprintNodeSpawner.cpp
--- a/Engine/Source/Editor/BlueprintGraph/Private/BlueprintNodeSpawner.cpp
+++ b/Engine/Source/Editor/BlueprintGraph/Private/BlueprintNodeSpawner.cpp
@@ -33,28 +33,9 @@ UBlueprintNodeSpawner::UBlueprintNodeSpawner(class FPostConstructInitializePrope
 }
 
 //------------------------------------------------------------------------------
-UEdGraphNode* UBlueprintNodeSpawner::Invoke(UEdGraph* ParentGraph) const
+UEdGraphNode* UBlueprintNodeSpawner::Invoke(UEdGraph* ParentGraph, FVector2D const Location) const
 {
-	UEdGraphNode* NewNode = nullptr;	
-	if (NodeClass != nullptr)
-	{
-		NewNode = NewObject<UEdGraphNode>(ParentGraph, NodeClass);
-		check(NewNode != nullptr);
-
-		bool bIsTemplateNode = ParentGraph->HasAnyFlags(RF_Transient);		
-		if (CustomizeNodeDelegate.IsBound())
-		{
-			CustomizeNodeDelegate.Execute(NewNode, bIsTemplateNode);
-		}
-
-		if (!bIsTemplateNode)
-		{
-			NewNode->SetFlags(RF_Transactional);
-			NewNode->AllocateDefaultPins();
-		}
-	}
-
-	return NewNode;
+	return Invoke(ParentGraph, Location, CustomizeNodeDelegate);
 }
 
 //------------------------------------------------------------------------------
@@ -116,7 +97,7 @@ UEdGraphNode* UBlueprintNodeSpawner::MakeTemplateNode(UEdGraph* Outer) const
 	//        graph's class type (and graph type)
 	if (CachedNodeTemplate == nullptr)// || (CachedNodeTemplate->GetOuter() != Outer))
 	{
-		CachedNodeTemplate = Invoke(Outer);
+		CachedNodeTemplate = Invoke(Outer, FVector2D::ZeroVector);
 		if (CachedNodeTemplate != nullptr)
 		{
 			CachedNodeTemplate->SetFlags(RF_Transient | RF_ArchetypeObject);
@@ -125,3 +106,30 @@ UEdGraphNode* UBlueprintNodeSpawner::MakeTemplateNode(UEdGraph* Outer) const
 	
 	return CachedNodeTemplate;
 }
+
+//------------------------------------------------------------------------------
+UEdGraphNode* UBlueprintNodeSpawner::Invoke(UEdGraph* ParentGraph, FVector2D const Location, FCustomizeNodeDelegate PostSpawnDelegate) const
+{
+	UEdGraphNode* NewNode = nullptr;
+	if (NodeClass != nullptr)
+	{
+		NewNode = NewObject<UEdGraphNode>(ParentGraph, NodeClass);
+		check(NewNode != nullptr);
+
+		// position the node before invoking PostSpawnDelegate (in case it 
+		// wishes to modify this positioning)
+		NewNode->NodePosX = Location.X;
+		NewNode->NodePosY = Location.Y;
+
+		bool const bIsTemplateNode = ParentGraph->HasAnyFlags(RF_Transient);
+		PostSpawnDelegate.ExecuteIfBound(NewNode, bIsTemplateNode);
+
+		if (!bIsTemplateNode)
+		{
+			NewNode->SetFlags(RF_Transactional);
+			NewNode->AllocateDefaultPins();
+		}
+	}
+
+	return NewNode;
+}
diff --git a/Engine/Source/Editor/BlueprintGraph/Private/BlueprintPropertyNodeSpawner.cpp b/Engine/Source/Editor/BlueprintGraph/Private/BlueprintPropertyNodeSpawner.cpp
--- a/Engine/Source/Editor/BlueprintGraph/Private/BlueprintPropertyNodeSpawner.cpp
+++ b/Engine/Source/Editor/BlueprintGraph/Private/BlueprintPropertyNodeSpawner.cpp
@@ -36,10 +36,10 @@ UBlueprintPropertyNodeSpawner::UBlueprintPropertyNodeSpawner(class FPostConstruc
 }
 
 //------------------------------------------------------------------------------
-UEdGraphNode* UBlueprintPropertyNodeSpawner::Invoke(UEdGraph* ParentGraph) const
+UEdGraphNode* UBlueprintPropertyNodeSpawner::Invoke(UEdGraph* ParentGraph, FVector2D const Location) const
 {
 	// if a NodeClass was set, then this should spawn something
-	UEdGraphNode* NewNode = Super::Invoke(ParentGraph);
+	UEdGraphNode* NewNode = Super::Invoke(ParentGraph, Location);
 	check(Property != nullptr);
 
 	if (NewNode != nullptr)
diff --git a/Engine/Source/Editor/BlueprintGraph/Public/BlueprintNodeSpawner.h b/Engine/Source/Editor/BlueprintGraph/Public/BlueprintNodeSpawner.h
--- a/Engine/Source/Editor/BlueprintGraph/Public/BlueprintNodeSpawner.h
+++ b/Engine/Source/Editor/BlueprintGraph/Public/BlueprintNodeSpawner.h
@@ -8,6 +8,10 @@
 class UEdGraph;
 class UEdGraphNode;
 
+/*******************************************************************************
+* UBlueprintNodeSpawner Declaration
+******************************************************************************/
+
 /**
  * Intended to be wrapped and used by FBlueprintActionMenuItem. Rather than 
  * sub-classing the menu item, we choose to subclass this instead (for 
@@ -30,18 +34,30 @@ public:
 	 * the allocated spawner's NodeClass and CustomizeNodeDelegate fields from
 	 * the supplied parameters.
 	 *
-	 * @param  NodeClass                The node type that you want the spawner to spawn.
-	 * @param  Outer					Optional outer for the new spawner (if left null, the transient package will be used).
-	 * @param  CustomizeNodeDelegate    A delegate to perform specialized node setup post-spawn.
+	 * @param  NodeClass			The node type that you want the spawner to spawn.
+	 * @param  Outer				Optional outer for the new spawner (if left null, the transient package will be used).
+	 * @param  PostSpawnDelegate    A delegate to perform specialized node setup post-spawn.
 	 * @return A newly allocated instance of this class.
 	 */
-	static UBlueprintNodeSpawner* Create(TSubclassOf<UEdGraphNode> const NodeClass, UObject* Outer = nullptr, FCustomizeNodeDelegate CustomizeNodeDelegate = FCustomizeNodeDelegate());
+	static UBlueprintNodeSpawner* Create(TSubclassOf<UEdGraphNode> const NodeClass, UObject* Outer = nullptr, FCustomizeNodeDelegate PostSpawnDelegate = FCustomizeNodeDelegate());
+	
+	/**
+	 * Templatized version of the above Create() method (where we specify the 
+	 * spawner's node class through the template argument).
+	 *
+	 * @param  Outer				Optional outer for the new spawner (if left null, the transient package will be used).
+	 * @param  PostSpawnDelegate    A delegate to perform specialized node setup post-spawn.
+	 * @return A newly allocated instance of this class.
+	 */
+	template<class NodeType>
+	static UBlueprintNodeSpawner* Create(UObject* Outer = nullptr, FCustomizeNodeDelegate PostSpawnDelegate = FCustomizeNodeDelegate());
 	
 	/**
 	 * Holds the class of node to spawn. May be null for sub-classes that know
 	 * specifically what node types to spawn (if null for an instance of this 
 	 * class, then nothing will be spawned).
 	 */
+	UPROPERTY()
 	TSubclassOf<UEdGraphNode> NodeClass;
 	
 	/**
@@ -62,9 +78,10 @@ public:
 	 * this case upon use.
 	 * 
 	 * @param  ParentGraph	The graph you want the node spawned into.
-	 * @return Null if it failed to spawn a node, otherwise a newly spawned node or posibly one that already existed.
+	 * @param  Location     Where you want the new node positioned in the graph.
+	 * @return Null if it failed to spawn a node, otherwise a newly spawned node or possibly one that already existed.
 	 */
-	virtual UEdGraphNode* Invoke(UEdGraph* ParentGraph) const;
+	virtual UEdGraphNode* Invoke(UEdGraph* ParentGraph, FVector2D const Location) const;
 	
 	// @TODO: for favorites, generated for the spawner type + data (function, property, node, etc.)
 	//virtual FGuid GetActionHash() const;
@@ -146,6 +163,18 @@ public:
 	 */
 	UEdGraphNode* MakeTemplateNode(UEdGraph* Outer) const;
 
+protected:
+	/**
+	 * Protected Invoke() that let's sub-classes specify their own post-spawn
+	 * delegate. Creates a new node based off of the set NodeClass.
+	 * 
+	 * @param  ParentGraph			The graph you want the node spawned into.
+	 * @param  Location				Where you want the new node positioned in the graph.
+	 * @param  PostSpawnDelegate	A delegate to run after spawning the node, but prior to allocating the node's pins.
+	 * @return Null if it failed to spawn a node (if NodeClass is null), otherwise a newly spawned node.
+	 */
+	UEdGraphNode* Invoke(UEdGraph* ParentGraph, FVector2D const Location, FCustomizeNodeDelegate PostSpawnDelegate) const;
+
 private:
 	/**
 	 * We try to avoid the spawning of node-templates (save on perf/mem where we
@@ -158,3 +187,14 @@ private:
 	UPROPERTY(Transient)
 	mutable UEdGraphNode* CachedNodeTemplate;
 };
+
+/*******************************************************************************
+ * Templatized UBlueprintNodeSpawner Implementation
+ ******************************************************************************/
+
+//------------------------------------------------------------------------------
+template<class NodeType>
+UBlueprintNodeSpawner* UBlueprintNodeSpawner::Create(UObject* Outer, FCustomizeNodeDelegate PostSpawnDelegate)
+{
+	Create(NodeType::StaticClass(), Outer, PostSpawnDelegate);
+}
\ No newline at end of file
diff --git a/Engine/Source/Editor/Kismet/Private/BlueprintActionMenuItem.cpp b/Engine/Source/Editor/Kismet/Private/BlueprintActionMenuItem.cpp
--- a/Engine/Source/Editor/Kismet/Private/BlueprintActionMenuItem.cpp
+++ b/Engine/Source/Editor/Kismet/Private/BlueprintActionMenuItem.cpp
@@ -10,6 +10,7 @@
 #include "BPVariableDragDropAction.h"
 #include "BPDelegateDragDropAction.h"
 #include "BlueprintEditor.h"		// for GetVarIconAndColor()
+#include "SNodePanel.h"				// for GetSnapGridSize()
 
 #define LOCTEXT_NAMESPACE "BlueprintActionMenuItem"
 
@@ -71,8 +72,27 @@ UEdGraphNode* FBlueprintActionMenuItem::PerformAction(UEdGraph* ParentGraph, UEd
 	using namespace FBlueprintMenuActionItemImpl;
 	FScopedTransaction Transaction(LOCTEXT("AddNodeTransaction", "Add Node"));
 	
+	FVector2D ModifiedLocation = Location;
+	if (FromPin != nullptr)
+	{
+		// for input pins, a new node will generally overlap the node being
+		// dragged from... work out if we want add in some spacing from the connecting node
+		if (FromPin->Direction == EGPD_Input)
+		{
+			UEdGraphNode* FromNode = FromPin->GetOwningNode();
+			check(FromNode != nullptr);
+			float const FromNodeX = FromNode->NodePosX;
+
+			static const float MinNodeDistance = 60.f; // min distance between spawned nodes (to keep them from overlapping)
+			if (MinNodeDistance > FMath::Abs(FromNodeX - Location.X))
+			{
+				ModifiedLocation.X = FromNodeX - MinNodeDistance;
+			}
+		}
+	}
+
 	// this could return an existing node
-	UEdGraphNode* SpawnedNode = Action->Invoke(ParentGraph);
+	UEdGraphNode* SpawnedNode = Action->Invoke(ParentGraph, Location);
 	
 	// if a returned node hasn't been added to the graph yet (it must have been freshly spawned)
 	if (ParentGraph->Nodes.Find(SpawnedNode) == INDEX_NONE)
@@ -83,37 +103,17 @@ UEdGraphNode* FBlueprintActionMenuItem::PerformAction(UEdGraph* ParentGraph, UEd
 		ParentGraph->AddNode(SpawnedNode, /*bool bFromUI =*/true, bSelectNewNode);
 		
 		SpawnedNode->PostPlacedNewNode();
-		SpawnedNode->NodePosX = Location.X;
-		SpawnedNode->NodePosY = Location.Y;
+		SpawnedNode->SnapToGrid(SNodePanel::GetSnapGridSize());
 		
 		if (FromPin != nullptr)
 		{
-			// for input pins, a new node will generally overlap the node being
-			// dragged from... work out if we want add in some spacing from the connecting node
-			if (FromPin->Direction == EGPD_Input)
-			{
-				UEdGraphNode* FromNode = FromPin->GetOwningNode();
-				check(FromNode != nullptr);
-				float const FromNodeX = FromNode->NodePosX;
-				
-				static const float MinNodeDistance = 60.f; // min distance between spawned nodes (to keep them from overlapping)
-				if (MinNodeDistance > FMath::Abs(FromNodeX - Location.X))
-				{
-					SpawnedNode->NodePosX = FromNodeX - MinNodeDistance;
-				}
-			}
-			
 			// modify before the call to AutowireNewNode() below
 			FromPin->Modify();
+			// make sure to auto-wire after we position the new node (in case
+			// the auto-wire creates a conversion node to put between them)
+			SpawnedNode->AutowireNewNode(FromPin);
 		}
 		
-		static const float GridSnapSize = 16.f; // @TODO: ensure this is the same as SNodePanel::GetSnapGridSize()
-		SpawnedNode->SnapToGrid(GridSnapSize);
-		
-		// make sure to auto-wire after we position the new node (in case
-		// the auto-wire creates a conversion node to put between them)
-		SpawnedNode->AutowireNewNode(FromPin);
-		
 		FBlueprintEditorUtils::AnalyticsTrackNewNode(SpawnedNode);
 		DirtyBlueprintFromNewNode(SpawnedNode);
 	}
