Adding a new method to graph nodes that determines if that node is compatible with certain graphs (split out functionality from CanPasteHere() so we can leverage it in the new blueprint menu system).

diff --git a/Engine/Source/Editor/BlueprintGraph/Classes/K2Node_AssignmentStatement.h b/Engine/Source/Editor/BlueprintGraph/Classes/K2Node_AssignmentStatement.h
--- a/Engine/Source/Editor/BlueprintGraph/Classes/K2Node_AssignmentStatement.h
+++ b/Engine/Source/Editor/BlueprintGraph/Classes/K2Node_AssignmentStatement.h
@@ -19,7 +19,7 @@ class UK2Node_AssignmentStatement : public UK2Node
 	virtual void AllocateDefaultPins() override;
 	virtual FString GetTooltip() const override;
 	virtual FText GetNodeTitle(ENodeTitleType::Type TitleType) const override;
-	virtual bool CanPasteHere(UEdGraph const* TargetGraph, UEdGraphSchema const* Schema) const override;
+	virtual bool IsCompatibleWithGraph(UEdGraph const* TargetGraph) const override;
 	// End UEdGraphNode interface
 
 	// Begin UK2Node interface
diff --git a/Engine/Source/Editor/BlueprintGraph/Classes/K2Node_BaseAsyncTask.h b/Engine/Source/Editor/BlueprintGraph/Classes/K2Node_BaseAsyncTask.h
--- a/Engine/Source/Editor/BlueprintGraph/Classes/K2Node_BaseAsyncTask.h
+++ b/Engine/Source/Editor/BlueprintGraph/Classes/K2Node_BaseAsyncTask.h
@@ -20,7 +20,7 @@ class BLUEPRINTGRAPH_API UK2Node_BaseAsyncTask : public UK2Node
 	virtual void GetMenuEntries(FGraphContextMenuBuilder& ContextMenuBuilder) const override;
 	virtual FString GetTooltip() const override;
 	virtual FText GetNodeTitle(ENodeTitleType::Type TitleType) const override;
-	virtual bool CanPasteHere(const UEdGraph* TargetGraph, const UEdGraphSchema* Schema) const override;
+	virtual bool IsCompatibleWithGraph(const UEdGraph* TargetGraph) const override;
 	// End of UEdGraphNode interface
 
 	// UK2Node interface
diff --git a/Engine/Source/Editor/BlueprintGraph/Classes/K2Node_Event.h b/Engine/Source/Editor/BlueprintGraph/Classes/K2Node_Event.h
--- a/Engine/Source/Editor/BlueprintGraph/Classes/K2Node_Event.h
+++ b/Engine/Source/Editor/BlueprintGraph/Classes/K2Node_Event.h
@@ -41,7 +41,8 @@ class UK2Node_Event : public UK2Node_EditablePinBase
 	BLUEPRINTGRAPH_API virtual FString GetKeywords() const override;	
 	BLUEPRINTGRAPH_API virtual FText GetNodeTitle(ENodeTitleType::Type TitleType) const override;
 	BLUEPRINTGRAPH_API virtual FLinearColor GetNodeTitleColor() const override;
-	BLUEPRINTGRAPH_API virtual bool CanPasteHere(const UEdGraph* TargetGraph, const UEdGraphSchema* Schema) const override;
+	BLUEPRINTGRAPH_API virtual bool CanPasteHere(const UEdGraph* TargetGraph) const override;
+	BLUEPRINTGRAPH_API virtual bool IsCompatibleWithGraph(const UEdGraph* TargetGraph) const override;
 	BLUEPRINTGRAPH_API virtual FName GetCornerIcon() const override;
 	BLUEPRINTGRAPH_API virtual bool IsDeprecated() const override;
 	BLUEPRINTGRAPH_API virtual FString GetDeprecationMessage() const override;
diff --git a/Engine/Source/Editor/BlueprintGraph/Classes/K2Node_InputAxisEvent.h b/Engine/Source/Editor/BlueprintGraph/Classes/K2Node_InputAxisEvent.h
--- a/Engine/Source/Editor/BlueprintGraph/Classes/K2Node_InputAxisEvent.h
+++ b/Engine/Source/Editor/BlueprintGraph/Classes/K2Node_InputAxisEvent.h
@@ -31,7 +31,7 @@ class UK2Node_InputAxisEvent : public UK2Node_Event
 	// Begin EdGraphNode interface
 	virtual FText GetNodeTitle(ENodeTitleType::Type TitleType) const override;
 	virtual FString GetTooltip() const override;
-	virtual bool CanPasteHere(const UEdGraph* TargetGraph, const UEdGraphSchema* Schema) const override;
+	virtual bool IsCompatibleWithGraph(const UEdGraph* TargetGraph) const override;
 	// End EdGraphNode interface
 
 	// Begin UK2Node interface
diff --git a/Engine/Source/Editor/BlueprintGraph/Classes/K2Node_InputAxisKeyEvent.h b/Engine/Source/Editor/BlueprintGraph/Classes/K2Node_InputAxisKeyEvent.h
--- a/Engine/Source/Editor/BlueprintGraph/Classes/K2Node_InputAxisKeyEvent.h
+++ b/Engine/Source/Editor/BlueprintGraph/Classes/K2Node_InputAxisKeyEvent.h
@@ -28,7 +28,7 @@ class UK2Node_InputAxisKeyEvent : public UK2Node_Event
 	virtual FText GetNodeTitle(ENodeTitleType::Type TitleType) const override;
 	virtual FString GetTooltip() const override;
 	virtual FName GetPaletteIcon(FLinearColor& OutColor) const override;
-	virtual bool CanPasteHere(const UEdGraph* TargetGraph, const UEdGraphSchema* Schema) const override;
+	virtual bool IsCompatibleWithGraph(const UEdGraph* TargetGraph) const override;
 	// End EdGraphNode interface
 
 	// Begin UK2Node interface
diff --git a/Engine/Source/Editor/BlueprintGraph/Classes/K2Node_LocalVariable.h b/Engine/Source/Editor/BlueprintGraph/Classes/K2Node_LocalVariable.h
--- a/Engine/Source/Editor/BlueprintGraph/Classes/K2Node_LocalVariable.h
+++ b/Engine/Source/Editor/BlueprintGraph/Classes/K2Node_LocalVariable.h
@@ -24,7 +24,7 @@ class UDEPRECATED_K2Node_LocalVariable : public UK2Node_TemporaryVariable
 	virtual TSharedPtr<class INameValidatorInterface> MakeNameValidator() const override;
 	virtual void PostPlacedNewNode() override;
 	virtual void PostPasteNode() override;
-	virtual bool CanPasteHere(const UEdGraph* TargetGraph, const UEdGraphSchema* Schema) const override;
+	virtual bool IsCompatibleWithGraph(const UEdGraph* TargetGraph) const override;
 	virtual void ReconstructNode() override;
 	// End UEdGraphNode interface.
 
diff --git a/Engine/Source/Editor/BlueprintGraph/Classes/K2Node_SpawnActor.h b/Engine/Source/Editor/BlueprintGraph/Classes/K2Node_SpawnActor.h
--- a/Engine/Source/Editor/BlueprintGraph/Classes/K2Node_SpawnActor.h
+++ b/Engine/Source/Editor/BlueprintGraph/Classes/K2Node_SpawnActor.h
@@ -20,12 +20,12 @@ class UK2Node_SpawnActor : public UK2Node
 	virtual bool ShouldWarnOnDeprecation() const override;
 	virtual FString GetDeprecationMessage() const override;
 	virtual FName GetPaletteIcon(FLinearColor& OutColor) const override{ return TEXT("GraphEditor.SpawnActor_16x"); }
+	virtual bool IsCompatibleWithGraph(const UEdGraph* TargetGraph) const override;
 	// End UEdGraphNode interface.
 
 	// Begin UK2Node interface
 	virtual bool IsNodeSafeToIgnore() const override { return true; }
 	virtual void ReallocatePinsDuringReconstruction(TArray<UEdGraphPin*>& OldPins) override;
-	virtual bool CanPasteHere(const UEdGraph* TargetGraph, const UEdGraphSchema* Schema) const override;
 	virtual class FNodeHandlingFunctor* CreateNodeHandler(class FKismetCompilerContext& CompilerContext) const override;
 	virtual void ExpandNode(class FKismetCompilerContext& CompilerContext, UEdGraph* SourceGraph) override;
 	virtual bool HasExternalBlueprintDependencies(TArray<class UStruct*>* OptionalOutput) const override;
diff --git a/Engine/Source/Editor/BlueprintGraph/Classes/K2Node_SpawnActorFromClass.h b/Engine/Source/Editor/BlueprintGraph/Classes/K2Node_SpawnActorFromClass.h
--- a/Engine/Source/Editor/BlueprintGraph/Classes/K2Node_SpawnActorFromClass.h
+++ b/Engine/Source/Editor/BlueprintGraph/Classes/K2Node_SpawnActorFromClass.h
@@ -19,12 +19,12 @@ class UK2Node_SpawnActorFromClass : public UK2Node
 	virtual void ExpandNode(class FKismetCompilerContext& CompilerContext, UEdGraph* SourceGraph) override;
 	virtual bool HasExternalBlueprintDependencies(TArray<class UStruct*>* OptionalOutput) const override;
 	virtual FName GetPaletteIcon(FLinearColor& OutColor) const override{ return TEXT("GraphEditor.SpawnActor_16x"); }
+	virtual bool IsCompatibleWithGraph(const UEdGraph* TargetGraph) const override;
 	// End UEdGraphNode interface.
 
 	// Begin UK2Node interface
 	virtual bool IsNodeSafeToIgnore() const override { return true; }
 	virtual void ReallocatePinsDuringReconstruction(TArray<UEdGraphPin*>& OldPins) override;
-	virtual bool CanPasteHere(const UEdGraph* TargetGraph, const UEdGraphSchema* Schema) const override;
 	virtual void GetNodeAttributes( TArray<TKeyValuePair<FString, FString>>& OutNodeAttributes ) const override;
 	virtual void GetMenuActions(TArray<UBlueprintNodeSpawner*>& ActionListOut) const override;
 	virtual FText GetMenuCategory() const override;
diff --git a/Engine/Source/Editor/BlueprintGraph/Classes/K2Node_TemporaryVariable.h b/Engine/Source/Editor/BlueprintGraph/Classes/K2Node_TemporaryVariable.h
--- a/Engine/Source/Editor/BlueprintGraph/Classes/K2Node_TemporaryVariable.h
+++ b/Engine/Source/Editor/BlueprintGraph/Classes/K2Node_TemporaryVariable.h
@@ -24,7 +24,7 @@ class UK2Node_TemporaryVariable : public UK2Node
 	virtual FString GetTooltip() const override;
 	virtual FText GetNodeTitle(ENodeTitleType::Type TitleType) const override;
 	virtual FString GetDescriptiveCompiledName() const override;
-	virtual bool CanPasteHere(UEdGraph const* TargetGraph, UEdGraphSchema const* Schema) const override;
+	virtual bool IsCompatibleWithGraph(UEdGraph const* TargetGraph) const override;
 	// End UEdGraphNode interface.
 
 	// Begin UK2Node interface.
diff --git a/Engine/Source/Editor/BlueprintGraph/Classes/K2Node_Timeline.h b/Engine/Source/Editor/BlueprintGraph/Classes/K2Node_Timeline.h
--- a/Engine/Source/Editor/BlueprintGraph/Classes/K2Node_Timeline.h
+++ b/Engine/Source/Editor/BlueprintGraph/Classes/K2Node_Timeline.h
@@ -37,7 +37,7 @@ class UK2Node_Timeline : public UK2Node
 	virtual void PrepareForCopying() override;
 	virtual FLinearColor GetNodeTitleColor() const override;
 	virtual FText GetNodeTitle(ENodeTitleType::Type TitleType) const override;
-	virtual bool CanPasteHere(const UEdGraph* TargetGraph, const UEdGraphSchema* Schema) const override;
+	virtual bool IsCompatibleWithGraph(const UEdGraph* TargetGraph) const override;
 	virtual void FindDiffs(class UEdGraphNode* OtherNode, struct FDiffResults& Results )  override;
 	virtual void OnRenameNode(const FString& NewName) override;
 	virtual TSharedPtr<class INameValidatorInterface> MakeNameValidator() const override;
diff --git a/Engine/Source/Editor/BlueprintGraph/Private/BlueprintActionFilter.cpp b/Engine/Source/Editor/BlueprintGraph/Private/BlueprintActionFilter.cpp
--- a/Engine/Source/Editor/BlueprintGraph/Private/BlueprintActionFilter.cpp
+++ b/Engine/Source/Editor/BlueprintGraph/Private/BlueprintActionFilter.cpp
@@ -143,46 +143,6 @@ namespace BlueprintActionFilterImpl
 	 * @return True if the action would spawn a node (and the graph wouldn't allow it).
 	 */
 	static bool IsIncompatibleWithGraphType(FBlueprintActionFilter const& Filter, UBlueprintNodeSpawner const* BlueprintAction);
-
-	/**
-	 * Rejection test that checks to see if the supplied node-spawner would 
-	 * produce a node incompatible with an ubergraph.
-	 * 
-	 * @param  Context			Holds the blueprint/graph/pin context for this test.
-	 * @param  BlueprintAction	The action you wish to query.
-	 * @return True if the action would spawn a node not allowed in ubergraphs.
-	 */
-	static bool IsIncompatibleWithUbergraph(FBlueprintActionContext const& Context, UBlueprintNodeSpawner const* BlueprintAction);
-
-	/**
-	 * Rejection test that checks to see if the supplied node-spawner would 
-	 * produce a node incompatible with an animation graph.
-	 * 
-	 * @param  Context			Holds the blueprint/graph/pin context for this test.
-	 * @param  BlueprintAction	The action you wish to query.
-	 * @return True if the action would spawn a node not allowed in animation graphs.
-	 */
-	static bool IsIncompatibleWithAnimGraph(FBlueprintActionContext const& Context, UBlueprintNodeSpawner const* BlueprintAction);
-
-	/**
-	 * Rejection test that checks to see if the supplied node-spawner would 
-	 * produce a node incompatible with a macro graph.
-	 * 
-	 * @param  Context			Holds the blueprint/graph/pin context for this test.
-	 * @param  BlueprintAction	The action you wish to query.
-	 * @return True if the action would spawn a node not allowed in macro graphs.
-	 */
-	static bool IsIncompatibleWithMacroGraph(FBlueprintActionContext const& Context, UBlueprintNodeSpawner const* BlueprintAction);
-
-	/**
-	 * Rejection test that checks to see if the supplied node-spawner would 
-	 * produce an node incompatible with a function graph.
-	 * 
-	 * @param  Context			Holds the blueprint/graph/pin context for this test.
-	 * @param  BlueprintAction	The action you wish to query.
-	 * @return True if the action would spawn a node not allowed in function graphs.
-	 */
-	static bool IsIncompatibleWithFuncGraph(FBlueprintActionContext const& Context, UBlueprintNodeSpawner const* BlueprintAction);
 	
 	/**
 	 * Rejection test that checks to see if the node-spawner has any associated
@@ -216,17 +176,6 @@ namespace BlueprintActionFilterImpl
 	 */
 	static bool IsFieldCategoryHidden(FBlueprintActionFilter const& Filter, UBlueprintNodeSpawner const* BlueprintAction);
 
-	/**
-	 * Rejection test that checks to see if the node-spawner would produce a 
-	 * node that couldn't be pasted into the specified graph(s) (check against 
-	 * the node's CDO to save on perf, to keep from spawning a temp node).
-	 * 
-	 * @param  Context			Holds the graph context for this test.
-	 * @param  BlueprintAction	The action you wish to query.
-	 * @return True if the action would produce a node that cannot be pasted in the graph.
-	 */
-	static bool IsNodeUnpastable(FBlueprintActionFilter const& Filter, UBlueprintNodeSpawner const* BlueprintAction);
-
 	/**
 	 * Rejection test that checks to see if the node-spawner would produce a 
 	 * node type that isn't white-listed. 
@@ -607,123 +556,26 @@ static bool BlueprintActionFilterImpl::IsIncompatibleImpureNode(FBlueprintAction
 	return bIsFilteredOut;
 }
 
-//------------------------------------------------------------------------------
-static bool BlueprintActionFilterImpl::IsIncompatibleWithUbergraph(FBlueprintActionContext const& /*Context*/, UBlueprintNodeSpawner const* BlueprintAction)
-{
-	bool bIsFilteredOut = false;
-	return bIsFilteredOut;
-}
-
-//------------------------------------------------------------------------------
-static bool BlueprintActionFilterImpl::IsIncompatibleWithAnimGraph(FBlueprintActionContext const& /*Context*/, UBlueprintNodeSpawner const* BlueprintAction)
-{
-	bool bIsFilteredOut = IsLatent(BlueprintAction); // latent actions are only allowed in ubergraphs
-	return bIsFilteredOut;
-}
-
-//------------------------------------------------------------------------------
-static bool BlueprintActionFilterImpl::IsIncompatibleWithMacroGraph(FBlueprintActionContext const& /*Context*/, UBlueprintNodeSpawner const* BlueprintAction)
-{
-	bool bIsFilteredOut = IsLatent(BlueprintAction); // latent actions are only allowed in ubergraphs
-	return bIsFilteredOut;
-}
-
-//------------------------------------------------------------------------------
-static bool BlueprintActionFilterImpl::IsIncompatibleWithFuncGraph(FBlueprintActionContext const& /*Context*/, UBlueprintNodeSpawner const* BlueprintAction)
-{
-	bool bIsFilteredOut = IsLatent(BlueprintAction); // latent actions are only allowed in ubergraphs
-	return bIsFilteredOut;
-}
-
 //------------------------------------------------------------------------------
 static bool BlueprintActionFilterImpl::IsIncompatibleWithGraphType(FBlueprintActionFilter const& Filter, UBlueprintNodeSpawner const* BlueprintAction)
 {
 	bool bIsFilteredOut = false;
 	FBlueprintActionContext const& FilterContext = Filter.Context;
 	
-	uint32 GraphTypesTestedMask = 0x00;
-	for (UEdGraph* Graph : FilterContext.Graphs)
+	if (UClass* NodeClass = BlueprintAction->NodeClass)
 	{
-		EGraphType GraphType = GT_MAX;
-		if (UEdGraphSchema_K2* K2Schema = Graph->Schema->GetDefaultObject<UEdGraphSchema_K2>())
-		{
-			GraphType = K2Schema->GetGraphType(Graph);
-		}
-		
-		uint32 GraphTypeBit = (1 << GraphType);
-		if (GraphTypesTestedMask & GraphTypeBit)
-		{
-			continue;
-		}
-		
-		switch (GraphType)
-		{
-		case GT_Ubergraph:
-			{
-				bIsFilteredOut = IsIncompatibleWithUbergraph(FilterContext, BlueprintAction);
-				break;
-			}
-		case GT_Animation:
-			{
-				bIsFilteredOut = IsIncompatibleWithAnimGraph(FilterContext, BlueprintAction);
-				break;
-			}
-		case GT_Macro:
-			{
-				bIsFilteredOut = IsIncompatibleWithMacroGraph(FilterContext, BlueprintAction);
-				break;
-			}
-		case GT_Function:
-			{
-				bIsFilteredOut = IsIncompatibleWithFuncGraph(FilterContext, BlueprintAction);
-				break;
-			}
-		}
-		
-		if (bIsFilteredOut)
-		{
-			break;
-		}
-		GraphTypesTestedMask |= GraphTypeBit;
-	}
-
-	return bIsFilteredOut;
-}
-
-//------------------------------------------------------------------------------
-static bool BlueprintActionFilterImpl::IsNodeUnpastable(FBlueprintActionFilter const& Filter, UBlueprintNodeSpawner const* BlueprintAction)
-{
-	bool bIsFilteredOut = false;
-	FBlueprintActionContext const& FilterContext = Filter.Context;
-	
-	for (UEdGraph* Graph : FilterContext.Graphs)
-	{		
-		UClass* NodeClass = BlueprintAction->NodeClass;
-		if (NodeClass == nullptr)
-		{
-			continue;
-		}
-		else if (NodeClass->IsChildOf<UK2Node_Event>())// @TODO: make all event nodes, this: BlueprintAction->IsA<UBlueprintEventNodeSpawner>())
-		{
-			// events need a valid function to check CanPasteHere() (we can't use the CDO)
-			// @TODO: modify the CDO's function reference, and then restore it
-			continue;
-		}
-
-		// wouldn't want to spawn a temp node here (too costly to be doing while
-		// filtering)... if a node that is unpasteble slips through, then the 
-		// user should be notified through ui feedback way (we purposely do this
-		// with the palette, so that it can have an all encompassing list)
 		if (UEdGraphNode* NodeCDO = NodeClass->GetDefaultObject<UEdGraphNode>())
 		{
-			if (!NodeCDO->CanPasteHere(Graph, Graph->Schema->GetDefaultObject<UEdGraphSchema>()))
+			for (UEdGraph* Graph : FilterContext.Graphs)
 			{
-				bIsFilteredOut = true;
-				break;
+				if (!NodeCDO->IsCompatibleWithGraph(Graph))
+				{
+					bIsFilteredOut = true;
+					break;
+				}
 			}
 		}
 	}
-	
 	return bIsFilteredOut;
 }
 
@@ -789,8 +641,7 @@ FBlueprintActionFilter::FBlueprintActionFilter(uint32 Flags/*= 0x00*/)
 	AddIsFilteredTest(FIsFilteredDelegate::CreateStatic(IsExternalField));
 	AddIsFilteredTest(FIsFilteredDelegate::CreateStatic(IsIncompatibleImpureNode));
 	AddIsFilteredTest(FIsFilteredDelegate::CreateStatic(IsPermissionNotGranted));
-	//AddIsFilteredTest(FIsFilteredDelegate::CreateStatic(IsIncompatibleWithGraphType));
-	AddIsFilteredTest(FIsFilteredDelegate::CreateStatic(IsNodeUnpastable));
+	AddIsFilteredTest(FIsFilteredDelegate::CreateStatic(IsIncompatibleWithGraphType));
 	AddIsFilteredTest(FIsFilteredDelegate::CreateStatic(IsFieldCategoryHidden));
 
 	// @TODO: account for all K2ActionMenuBuilder checks...
diff --git a/Engine/Source/Editor/BlueprintGraph/Private/K2Node_AssignmentStatement.cpp b/Engine/Source/Editor/BlueprintGraph/Private/K2Node_AssignmentStatement.cpp
--- a/Engine/Source/Editor/BlueprintGraph/Private/K2Node_AssignmentStatement.cpp
+++ b/Engine/Source/Editor/BlueprintGraph/Private/K2Node_AssignmentStatement.cpp
@@ -84,17 +84,16 @@ FText UK2Node_AssignmentStatement::GetNodeTitle(ENodeTitleType::Type TitleType)
 	return LOCTEXT("Assign", "Assign");
 }
 
-bool UK2Node_AssignmentStatement::CanPasteHere(UEdGraph const* TargetGraph, UEdGraphSchema const* Schema) const
+bool UK2Node_AssignmentStatement::IsCompatibleWithGraph(UEdGraph const* TargetGraph) const
 {
-	bool bIsPastable = Super::CanPasteHere(TargetGraph, Schema);
-	if (bIsPastable)
+	bool bIsCompatible = Super::IsCompatibleWithGraph(TargetGraph);
+	if (bIsCompatible)
 	{
-		EGraphType const GraphType = Schema->GetGraphType(TargetGraph);
-
-		bIsPastable = (GraphType != GT_Ubergraph) && (GraphType != GT_Animation);
+		EGraphType const GraphType = TargetGraph->GetSchema()->GetGraphType(TargetGraph);
+		bIsCompatible = (GraphType != GT_Ubergraph) && (GraphType != GT_Animation);
 	}
 
-	return bIsPastable;
+	return bIsCompatible;
 }
 
 void UK2Node_AssignmentStatement::NotifyPinConnectionListChanged(UEdGraphPin* Pin)
diff --git a/Engine/Source/Editor/BlueprintGraph/Private/K2Node_BaseAsyncTask.cpp b/Engine/Source/Editor/BlueprintGraph/Private/K2Node_BaseAsyncTask.cpp
--- a/Engine/Source/Editor/BlueprintGraph/Private/K2Node_BaseAsyncTask.cpp
+++ b/Engine/Source/Editor/BlueprintGraph/Private/K2Node_BaseAsyncTask.cpp
@@ -31,15 +31,15 @@ FText UK2Node_BaseAsyncTask::GetNodeTitle(ENodeTitleType::Type TitleType) const
 	return FText::FromString(FunctionToolTipText);
 }
 
-bool UK2Node_BaseAsyncTask::CanPasteHere(const UEdGraph* TargetGraph, const UEdGraphSchema* Schema) const
+bool UK2Node_BaseAsyncTask::IsCompatibleWithGraph(const UEdGraph* TargetGraph) const
 {
+	bool bIsCompatible = true;
 	// Can only place events in ubergraphs, and basicasync task creates an event node:
-	if (Schema->GetGraphType(TargetGraph) != EGraphType::GT_Ubergraph)
+	if (TargetGraph->GetSchema()->GetGraphType(TargetGraph) != EGraphType::GT_Ubergraph)
 	{
-		UE_LOG(LogBlueprint, Log, TEXT("Cannot paste event node (%s) in this graph because it is not an ubergraph."), *GetFName().ToString());
-		return false;
+		bIsCompatible = false;
 	}
-	return UK2Node::CanPasteHere(TargetGraph, Schema);
+	return bIsCompatible && Super::IsCompatibleWithGraph(TargetGraph);
 }
 
 void UK2Node_BaseAsyncTask::GetMenuEntries(FGraphContextMenuBuilder& ContextMenuBuilder) const
diff --git a/Engine/Source/Editor/BlueprintGraph/Private/K2Node_Event.cpp b/Engine/Source/Editor/BlueprintGraph/Private/K2Node_Event.cpp
--- a/Engine/Source/Editor/BlueprintGraph/Private/K2Node_Event.cpp
+++ b/Engine/Source/Editor/BlueprintGraph/Private/K2Node_Event.cpp
@@ -343,123 +343,125 @@ bool UK2Node_Event::IsFunctionEntryCompatible(const UK2Node_FunctionEntry* Entry
 	return (EventPins.Num() == 0) && (EntryPins.Num() == 0);
 }
 
-bool UK2Node_Event::CanPasteHere(const UEdGraph* TargetGraph, const UEdGraphSchema* Schema) const
+bool UK2Node_Event::IsCompatibleWithGraph(const UEdGraph* TargetGraph) const
 {
-	// By default, to be safe, we don't allow events to be pasted, except under special circumstances (see below)
-	bool bDisallowPaste = true;
+	bool bIsCompatible = Super::IsCompatibleWithGraph(TargetGraph);
+	if (bIsCompatible)
+	{
+		EGraphType const GraphType = TargetGraph->GetSchema()->GetGraphType(TargetGraph);
+		bIsCompatible = (GraphType == EGraphType::GT_Ubergraph);
+	}
+	
+	return bIsCompatible;
+}
 
-	// Ensure that we can be instanced under the specified schema
-	if(CanCreateUnderSpecifiedSchema(Schema))
+bool UK2Node_Event::CanPasteHere(const UEdGraph* TargetGraph) const
+{
+	// By default, to be safe, we don't allow events to be pasted, except under special circumstances (see below)
+	bool bDisallowPaste = !Super::CanPasteHere(TargetGraph);
+	if(!bDisallowPaste)
 	{
-		// Can only place events in ubergraphs
-		if(Schema->GetGraphType(TargetGraph) == EGraphType::GT_Ubergraph)
+		// Find the Blueprint that owns the target graph
+		UBlueprint* Blueprint = FBlueprintEditorUtils::FindBlueprintForGraph(TargetGraph);
+		if(Blueprint && Blueprint->SkeletonGeneratedClass)
 		{
-			// Find the Blueprint that owns the target graph
-			UBlueprint* Blueprint = FBlueprintEditorUtils::FindBlueprintForGraph(TargetGraph);
-			if(Blueprint && Blueprint->SkeletonGeneratedClass)
-			{
-				TArray<FName> ExistingNamesInUse;
-				TArray<FString> ExcludedEventNames;
-				TArray<UK2Node_Event*> ExistingEventNodes;
-				TArray<UClass*> ImplementedInterfaceClasses;
+			TArray<FName> ExistingNamesInUse;
+			TArray<FString> ExcludedEventNames;
+			TArray<UK2Node_Event*> ExistingEventNodes;
+			TArray<UClass*> ImplementedInterfaceClasses;
 
-				// Gather all names in use by the Blueprint class
-				FBlueprintEditorUtils::GetFunctionNameList(Blueprint, ExistingNamesInUse);
-				FBlueprintEditorUtils::GetClassVariableList(Blueprint, ExistingNamesInUse);
+			// Gather all names in use by the Blueprint class
+			FBlueprintEditorUtils::GetFunctionNameList(Blueprint, ExistingNamesInUse);
+			FBlueprintEditorUtils::GetClassVariableList(Blueprint, ExistingNamesInUse);
 
-				// Gather all existing event nodes
-				FBlueprintEditorUtils::GetAllNodesOfClass<UK2Node_Event>(Blueprint, ExistingEventNodes);
+			// Gather all existing event nodes
+			FBlueprintEditorUtils::GetAllNodesOfClass<UK2Node_Event>(Blueprint, ExistingEventNodes);
 
-				// Gather any event names excluded by the Blueprint class
-				const FString ExclusionListKeyName = TEXT("KismetHideOverrides");
-				if(Blueprint->ParentClass->HasMetaData(*ExclusionListKeyName))
-				{
-					const FString ExcludedEventNameString = Blueprint->ParentClass->GetMetaData(*ExclusionListKeyName);
-					ExcludedEventNameString.ParseIntoArray(&ExcludedEventNames, TEXT(","), true);
-				}
+			// Gather any event names excluded by the Blueprint class
+			const FString ExclusionListKeyName = TEXT("KismetHideOverrides");
+			if(Blueprint->ParentClass->HasMetaData(*ExclusionListKeyName))
+			{
+				const FString ExcludedEventNameString = Blueprint->ParentClass->GetMetaData(*ExclusionListKeyName);
+				ExcludedEventNameString.ParseIntoArray(&ExcludedEventNames, TEXT(","), true);
+			}
 
-				// Gather all interfaces implemented by the Blueprint class
-				FBlueprintEditorUtils::FindImplementedInterfaces(Blueprint, true, ImplementedInterfaceClasses);
+			// Gather all interfaces implemented by the Blueprint class
+			FBlueprintEditorUtils::FindImplementedInterfaces(Blueprint, true, ImplementedInterfaceClasses);
 
-				// If this is an internal event, don't paste this event
-				if(!bInternalEvent)
+			// If this is an internal event, don't paste this event
+			if(!bInternalEvent)
+			{
+				// If this is a function override
+				if(bOverrideFunction)
 				{
-					// If this is a function override
-					if(bOverrideFunction)
+					// If the function name is hidden by the parent class, don't paste this event
+					bDisallowPaste = EventSignatureClass == Blueprint->ParentClass
+						&& ExcludedEventNames.Contains(EventSignatureName.ToString());
+					if(!bDisallowPaste)
 					{
-						// If the function name is hidden by the parent class, don't paste this event
-						bDisallowPaste = EventSignatureClass == Blueprint->ParentClass
-							&& ExcludedEventNames.Contains(EventSignatureName.ToString());
-						if(!bDisallowPaste)
+						// If the event function is already handled in this Blueprint, don't paste this event
+						for(int32 i = 0; i < ExistingEventNodes.Num() && !bDisallowPaste; ++i)
 						{
-							// If the event function is already handled in this Blueprint, don't paste this event
-							for(int32 i = 0; i < ExistingEventNodes.Num() && !bDisallowPaste; ++i)
-							{
-								bDisallowPaste = ExistingEventNodes[i]->bOverrideFunction
-									&& ExistingEventNodes[i]->EventSignatureName == EventSignatureName
-									&& ExistingEventNodes[i]->EventSignatureClass == EventSignatureClass;
-							}
+							bDisallowPaste = ExistingEventNodes[i]->bOverrideFunction
+								&& ExistingEventNodes[i]->EventSignatureName == EventSignatureName
+								&& ExistingEventNodes[i]->EventSignatureClass == EventSignatureClass;
+						}
 
-							if(!bDisallowPaste)
-							{
-								// If the signature class is not implemented by the Blueprint parent class or an interface, don't paste this event
-								bDisallowPaste = !Blueprint->ParentClass->IsChildOf(EventSignatureClass)
-									&& !ImplementedInterfaceClasses.Contains(EventSignatureClass);
-								if(bDisallowPaste)
-								{
-									UE_LOG(LogBlueprint, Log, TEXT("Cannot paste event node (%s) directly because the event signature class (%s) is incompatible with this Blueprint."), *GetFName().ToString(), EventSignatureClass ? *EventSignatureClass->GetFName().ToString() : TEXT("NONE"));
-								}
-							}
-							else
+						if(!bDisallowPaste)
+						{
+							// If the signature class is not implemented by the Blueprint parent class or an interface, don't paste this event
+							bDisallowPaste = !Blueprint->ParentClass->IsChildOf(EventSignatureClass)
+								&& !ImplementedInterfaceClasses.Contains(EventSignatureClass);
+							if(bDisallowPaste)
 							{
-								UE_LOG(LogBlueprint, Log, TEXT("Cannot paste event node (%s) directly because the event function (%s) is already handled."), *GetFName().ToString(), *EventSignatureName.ToString());
+								UE_LOG(LogBlueprint, Log, TEXT("Cannot paste event node (%s) directly because the event signature class (%s) is incompatible with this Blueprint."), *GetFName().ToString(), EventSignatureClass ? *EventSignatureClass->GetFName().ToString() : TEXT("NONE"));
 							}
 						}
 						else
 						{
-							UE_LOG(LogBlueprint, Log, TEXT("Cannot paste event node (%s) directly because the event function (%s) is hidden by the Blueprint parent class (%s)."), *GetFName().ToString(), *EventSignatureName.ToString(), EventSignatureClass ? *EventSignatureClass->GetFName().ToString() : TEXT("NONE"));
+							UE_LOG(LogBlueprint, Log, TEXT("Cannot paste event node (%s) directly because the event function (%s) is already handled."), *GetFName().ToString(), *EventSignatureName.ToString());
 						}
 					}
-					else if(CustomFunctionName != NAME_None)
+					else
 					{
-						// If this name is already in use, we can't paste this event
-						bDisallowPaste = ExistingNamesInUse.Contains(CustomFunctionName);
+						UE_LOG(LogBlueprint, Log, TEXT("Cannot paste event node (%s) directly because the event function (%s) is hidden by the Blueprint parent class (%s)."), *GetFName().ToString(), *EventSignatureName.ToString(), EventSignatureClass ? *EventSignatureClass->GetFName().ToString() : TEXT("NONE"));
+					}
+				}
+				else if(CustomFunctionName != NAME_None)
+				{
+					// If this name is already in use, we can't paste this event
+					bDisallowPaste = ExistingNamesInUse.Contains(CustomFunctionName);
 
-						if(!bDisallowPaste)
+					if(!bDisallowPaste)
+					{
+						// Handle events that have a custom function name with an actual signature name/class that is not an override (e.g. AnimNotify events)
+						if(EventSignatureName != NAME_None)
 						{
-							// Handle events that have a custom function name with an actual signature name/class that is not an override (e.g. AnimNotify events)
-							if(EventSignatureName != NAME_None)
+							// If the signature class is not implemented by the Blueprint parent class or an interface, don't paste this event
+							bDisallowPaste = !Blueprint->ParentClass->IsChildOf(EventSignatureClass)
+								&& !ImplementedInterfaceClasses.Contains(EventSignatureClass);
+							if(bDisallowPaste)
 							{
-								// If the signature class is not implemented by the Blueprint parent class or an interface, don't paste this event
-								bDisallowPaste = !Blueprint->ParentClass->IsChildOf(EventSignatureClass)
-									&& !ImplementedInterfaceClasses.Contains(EventSignatureClass);
-								if(bDisallowPaste)
-								{
-									UE_LOG(LogBlueprint, Log, TEXT("Cannot paste event node (%s) directly because the custom event function (%s) with event signature name (%s) has an event signature class (%s) that is incompatible with this Blueprint."), *GetFName().ToString(), *CustomFunctionName.ToString(), *EventSignatureName.ToString(), EventSignatureClass ? *EventSignatureClass->GetFName().ToString() : TEXT("NONE"));
-								}
+								UE_LOG(LogBlueprint, Log, TEXT("Cannot paste event node (%s) directly because the custom event function (%s) with event signature name (%s) has an event signature class (%s) that is incompatible with this Blueprint."), *GetFName().ToString(), *CustomFunctionName.ToString(), *EventSignatureName.ToString(), EventSignatureClass ? *EventSignatureClass->GetFName().ToString() : TEXT("NONE"));
 							}
 						}
-						else
-						{
-							UE_LOG(LogBlueprint, Log, TEXT("Cannot paste event node (%s) directly because the custom event function (%s) is already handled."), *GetFName().ToString(), *CustomFunctionName.ToString());
-						}
 					}
 					else
 					{
-						UE_LOG(LogBlueprint, Log, TEXT("Cannot paste event node (%s) directly because the event configuration is not specifically handled (EventSignatureName=%s, EventSignatureClass=%s)."), *GetFName().ToString(), *EventSignatureName.ToString(), EventSignatureClass ? *EventSignatureClass->GetFName().ToString() : TEXT("NONE"));
+						UE_LOG(LogBlueprint, Log, TEXT("Cannot paste event node (%s) directly because the custom event function (%s) is already handled."), *GetFName().ToString(), *CustomFunctionName.ToString());
 					}
 				}
 				else
 				{
-					UE_LOG(LogBlueprint, Log, TEXT("Cannot paste event node (%s) directly because it is flagged as an internal event."), *GetFName().ToString());
+					UE_LOG(LogBlueprint, Log, TEXT("Cannot paste event node (%s) directly because the event configuration is not specifically handled (EventSignatureName=%s, EventSignatureClass=%s)."), *GetFName().ToString(), *EventSignatureName.ToString(), EventSignatureClass ? *EventSignatureClass->GetFName().ToString() : TEXT("NONE"));
 				}
 			}
+			else
+			{
+				UE_LOG(LogBlueprint, Log, TEXT("Cannot paste event node (%s) directly because it is flagged as an internal event."), *GetFName().ToString());
+			}
 		}
 	}
-	else
-	{
-		UE_LOG(LogBlueprint, Log, TEXT("Cannot paste event node (%s) directly because it cannot be created under the specified schema."), *GetFName().ToString());
-	}
 
 	return !bDisallowPaste;
 }
diff --git a/Engine/Source/Editor/BlueprintGraph/Private/K2Node_InputAxisEvent.cpp b/Engine/Source/Editor/BlueprintGraph/Private/K2Node_InputAxisEvent.cpp
--- a/Engine/Source/Editor/BlueprintGraph/Private/K2Node_InputAxisEvent.cpp
+++ b/Engine/Source/Editor/BlueprintGraph/Private/K2Node_InputAxisEvent.cpp
@@ -85,35 +85,19 @@ void UK2Node_InputAxisEvent::RegisterDynamicBinding(UDynamicBlueprintBinding* Bi
 	InputAxisBindingObject->InputAxisDelegateBindings.Add(Binding);
 }
 
-bool UK2Node_InputAxisEvent::CanPasteHere(const UEdGraph* TargetGraph, const UEdGraphSchema* Schema) const
+bool UK2Node_InputAxisEvent::IsCompatibleWithGraph(const UEdGraph* TargetGraph) const
 {
 	// By default, to be safe, we don't allow events to be pasted, except under special circumstances (see below)
-	bool bAllowPaste = false;
+	bool bIsCompatible = false;
 
-	// Ensure that we can be instanced under the specified schema
-	if (CanCreateUnderSpecifiedSchema(Schema))
+	// Find the Blueprint that owns the target graph
+	UBlueprint* Blueprint = FBlueprintEditorUtils::FindBlueprintForGraph(TargetGraph);
+	if (Blueprint && Blueprint->SkeletonGeneratedClass)
 	{
-		// Can only place events in ubergraphs
-		if (Schema->GetGraphType(TargetGraph) == EGraphType::GT_Ubergraph)
-		{
-			// Find the Blueprint that owns the target graph
-			UBlueprint* Blueprint = FBlueprintEditorUtils::FindBlueprintForGraph(TargetGraph);
-			if (Blueprint && Blueprint->SkeletonGeneratedClass)
-			{
-				bAllowPaste = Blueprint->ParentClass->IsChildOf(AActor::StaticClass());
-				if (!bAllowPaste)
-				{
-					UE_LOG(LogBlueprint, Log, TEXT("Cannot paste event node (%s) directly because the graph does not belong to an Actor."), *GetFName().ToString());
-				}
-			}
-		}
-	}
-	else
-	{
-		UE_LOG(LogBlueprint, Log, TEXT("Cannot paste event node (%s) directly because it cannot be created under the specified schema."), *GetFName().ToString());
+		bIsCompatible = Blueprint->ParentClass->IsChildOf(AActor::StaticClass());
 	}
 
-	return bAllowPaste;
+	return bIsCompatible && Super::IsCompatibleWithGraph(TargetGraph);
 }
 
 void UK2Node_InputAxisEvent::GetMenuActions(TArray<UBlueprintNodeSpawner*>& ActionListOut) const
diff --git a/Engine/Source/Editor/BlueprintGraph/Private/K2Node_InputAxisKeyEvent.cpp b/Engine/Source/Editor/BlueprintGraph/Private/K2Node_InputAxisKeyEvent.cpp
--- a/Engine/Source/Editor/BlueprintGraph/Private/K2Node_InputAxisKeyEvent.cpp
+++ b/Engine/Source/Editor/BlueprintGraph/Private/K2Node_InputAxisKeyEvent.cpp
@@ -88,35 +88,19 @@ void UK2Node_InputAxisKeyEvent::RegisterDynamicBinding(UDynamicBlueprintBinding*
 	InputAxisKeyBindingObject->InputAxisKeyDelegateBindings.Add(Binding);
 }
 
-bool UK2Node_InputAxisKeyEvent::CanPasteHere(const UEdGraph* TargetGraph, const UEdGraphSchema* Schema) const
+bool UK2Node_InputAxisKeyEvent::IsCompatibleWithGraph(const UEdGraph* TargetGraph) const
 {
 	// By default, to be safe, we don't allow events to be pasted, except under special circumstances (see below)
-	bool bAllowPaste = false;
-
-	// Ensure that we can be instanced under the specified schema
-	if (CanCreateUnderSpecifiedSchema(Schema))
-	{
-		// Can only place events in ubergraphs
-		if (Schema->GetGraphType(TargetGraph) == EGraphType::GT_Ubergraph)
-		{
-			// Find the Blueprint that owns the target graph
-			UBlueprint* Blueprint = FBlueprintEditorUtils::FindBlueprintForGraph(TargetGraph);
-			if (Blueprint && Blueprint->SkeletonGeneratedClass)
-			{
-				bAllowPaste = Blueprint->ParentClass->IsChildOf(AActor::StaticClass());
-				if (!bAllowPaste)
-				{
-					UE_LOG(LogBlueprint, Log, TEXT("Cannot paste event node (%s) directly because the graph does not belong to an Actor."), *GetFName().ToString());
-				}
-			}
-		}
-	}
-	else
+	bool bIsCompatible = false;
+	
+	// Find the Blueprint that owns the target graph
+	UBlueprint* Blueprint = FBlueprintEditorUtils::FindBlueprintForGraph(TargetGraph);
+	if (Blueprint && Blueprint->SkeletonGeneratedClass)
 	{
-		UE_LOG(LogBlueprint, Log, TEXT("Cannot paste event node (%s) directly because it cannot be created under the specified schema."), *GetFName().ToString());
+		bIsCompatible = Blueprint->ParentClass->IsChildOf(AActor::StaticClass());
 	}
 
-	return bAllowPaste;
+	return bIsCompatible && Super::IsCompatibleWithGraph(TargetGraph);
 }
 
 void UK2Node_InputAxisKeyEvent::GetMenuActions(TArray<UBlueprintNodeSpawner*>& ActionListOut) const
diff --git a/Engine/Source/Editor/BlueprintGraph/Private/K2Node_LocalVariable.cpp b/Engine/Source/Editor/BlueprintGraph/Private/K2Node_LocalVariable.cpp
--- a/Engine/Source/Editor/BlueprintGraph/Private/K2Node_LocalVariable.cpp
+++ b/Engine/Source/Editor/BlueprintGraph/Private/K2Node_LocalVariable.cpp
@@ -87,21 +87,10 @@ void UDEPRECATED_K2Node_LocalVariable::PostPasteNode()
 	CustomVariableName = FBlueprintEditorUtils::FindUniqueKismetName(GetBlueprint(), CustomVariableName.GetPlainNameString());
 }
 
-bool UDEPRECATED_K2Node_LocalVariable::CanPasteHere(const UEdGraph* TargetGraph, const UEdGraphSchema* Schema) const
+bool UDEPRECATED_K2Node_LocalVariable::IsCompatibleWithGraph(const UEdGraph* TargetGraph) const
 {
-	// Local variables can only be pasted into function graphs
-	if(Super::CanPasteHere(TargetGraph, Schema))
-	{
-		UBlueprint* Blueprint = FBlueprintEditorUtils::FindBlueprintForGraph(TargetGraph);
-		if(Blueprint)
-		{
-			const UEdGraphSchema_K2* K2Schema = Cast<UEdGraphSchema_K2>(Schema);
-			check(K2Schema);
-			return K2Schema->GetGraphType(TargetGraph) == GT_Function;
-		}
-	}
-
-	return false;
+	bool const bIsCompatible = (TargetGraph->GetSchema()->GetGraphType(TargetGraph) == GT_Function);
+	return bIsCompatible && Super::IsCompatibleWithGraph(TargetGraph);
 }
 
 void UDEPRECATED_K2Node_LocalVariable::ReconstructNode()
diff --git a/Engine/Source/Editor/BlueprintGraph/Private/K2Node_SpawnActor.cpp b/Engine/Source/Editor/BlueprintGraph/Private/K2Node_SpawnActor.cpp
--- a/Engine/Source/Editor/BlueprintGraph/Private/K2Node_SpawnActor.cpp
+++ b/Engine/Source/Editor/BlueprintGraph/Private/K2Node_SpawnActor.cpp
@@ -248,10 +248,10 @@ FText UK2Node_SpawnActor::GetNodeTitle(ENodeTitleType::Type TitleType) const
 	return FText::Format(NSLOCTEXT("K2Node", "SpawnActor", "SpawnActor {ActorName}"), Args);
 }
 
-bool UK2Node_SpawnActor::CanPasteHere( const UEdGraph* TargetGraph, const UEdGraphSchema* Schema ) const 
+bool UK2Node_SpawnActor::IsCompatibleWithGraph(const UEdGraph* TargetGraph) const 
 {
-	UBlueprint* Blueprint = Cast<UBlueprint>(TargetGraph->GetOuter());
-	return CanCreateUnderSpecifiedSchema(Schema) && (!Blueprint || FBlueprintEditorUtils::FindUserConstructionScript(Blueprint) != TargetGraph);
+	UBlueprint* Blueprint = FBlueprintEditorUtils::FindBlueprintForGraph(TargetGraph);
+	return Super::IsCompatibleWithGraph(TargetGraph) && (!Blueprint || FBlueprintEditorUtils::FindUserConstructionScript(Blueprint) != TargetGraph);
 }
 
 FNodeHandlingFunctor* UK2Node_SpawnActor::CreateNodeHandler(FKismetCompilerContext& CompilerContext) const
diff --git a/Engine/Source/Editor/BlueprintGraph/Private/K2Node_SpawnActorFromClass.cpp b/Engine/Source/Editor/BlueprintGraph/Private/K2Node_SpawnActorFromClass.cpp
--- a/Engine/Source/Editor/BlueprintGraph/Private/K2Node_SpawnActorFromClass.cpp
+++ b/Engine/Source/Editor/BlueprintGraph/Private/K2Node_SpawnActorFromClass.cpp
@@ -298,10 +298,10 @@ FText UK2Node_SpawnActorFromClass::GetNodeTitle(ENodeTitleType::Type TitleType)
 	return NodeTitle;
 }
 
-bool UK2Node_SpawnActorFromClass::CanPasteHere( const UEdGraph* TargetGraph, const UEdGraphSchema* Schema ) const 
+bool UK2Node_SpawnActorFromClass::IsCompatibleWithGraph(const UEdGraph* TargetGraph) const 
 {
-	UBlueprint* Blueprint = Cast<UBlueprint>(TargetGraph->GetOuter());
-	return CanCreateUnderSpecifiedSchema(Schema) && (!Blueprint || FBlueprintEditorUtils::FindUserConstructionScript(Blueprint) != TargetGraph);
+	UBlueprint* Blueprint = FBlueprintEditorUtils::FindBlueprintForGraph(TargetGraph);
+	return Super::IsCompatibleWithGraph(TargetGraph) && (!Blueprint || FBlueprintEditorUtils::FindUserConstructionScript(Blueprint) != TargetGraph);
 }
 
 void UK2Node_SpawnActorFromClass::GetNodeAttributes( TArray<TKeyValuePair<FString, FString>>& OutNodeAttributes ) const
diff --git a/Engine/Source/Editor/BlueprintGraph/Private/K2Node_TemporaryVariable.cpp b/Engine/Source/Editor/BlueprintGraph/Private/K2Node_TemporaryVariable.cpp
--- a/Engine/Source/Editor/BlueprintGraph/Private/K2Node_TemporaryVariable.cpp
+++ b/Engine/Source/Editor/BlueprintGraph/Private/K2Node_TemporaryVariable.cpp
@@ -88,21 +88,21 @@ FString UK2Node_TemporaryVariable::GetDescriptiveCompiledName() const
 	return Result;
 }
 
-bool UK2Node_TemporaryVariable::CanPasteHere(UEdGraph const* TargetGraph, UEdGraphSchema const* Schema) const
+bool UK2Node_TemporaryVariable::IsCompatibleWithGraph(UEdGraph const* TargetGraph) const
 {
-	bool bIsPastable = Super::CanPasteHere(TargetGraph, Schema);
-	if (bIsPastable)
+	bool bIsCompatible = Super::IsCompatibleWithGraph(TargetGraph);
+	if (bIsCompatible)
 	{
-		EGraphType const GraphType = Schema->GetGraphType(TargetGraph);
-		bIsPastable = (GraphType != GT_Ubergraph) && (GraphType != GT_Animation);
+		EGraphType const GraphType = TargetGraph->GetSchema()->GetGraphType(TargetGraph);
+		bIsCompatible = (GraphType != GT_Ubergraph) && (GraphType != GT_Animation);
 
-		if (bIsPastable && (GraphType != GT_Macro))
+		if (bIsCompatible && (GraphType != GT_Macro))
 		{
-			bIsPastable = !bIsPersistent;
+			bIsCompatible = !bIsPersistent;
 		}
 	}
 
-	return bIsPastable;
+	return bIsCompatible;
 }
 
 // get variable pin
diff --git a/Engine/Source/Editor/BlueprintGraph/Private/K2Node_Timeline.cpp b/Engine/Source/Editor/BlueprintGraph/Private/K2Node_Timeline.cpp
--- a/Engine/Source/Editor/BlueprintGraph/Private/K2Node_Timeline.cpp
+++ b/Engine/Source/Editor/BlueprintGraph/Private/K2Node_Timeline.cpp
@@ -205,19 +205,19 @@ void UK2Node_Timeline::PostPasteNode()
 	}
 }
 
-bool UK2Node_Timeline::CanPasteHere(const UEdGraph* TargetGraph, const UEdGraphSchema* Schema) const
+bool UK2Node_Timeline::IsCompatibleWithGraph(const UEdGraph* TargetGraph) const
 {
-	if(Super::CanPasteHere(TargetGraph, Schema))
+	if(Super::IsCompatibleWithGraph(TargetGraph))
 	{
 		UBlueprint* Blueprint = FBlueprintEditorUtils::FindBlueprintForGraph(TargetGraph);
 		if(Blueprint)
 		{
-			const UEdGraphSchema_K2* K2Schema = Cast<UEdGraphSchema_K2>(Schema);
+			const UEdGraphSchema_K2* K2Schema = Cast<UEdGraphSchema_K2>(TargetGraph->GetSchema());
 			check(K2Schema);
 
 			const bool bSupportsEventGraphs = FBlueprintEditorUtils::DoesSupportEventGraphs(Blueprint);
-			const bool bAllowEvents = K2Schema->GetGraphType(TargetGraph) == GT_Ubergraph && bSupportsEventGraphs &&
-				Blueprint->BlueprintType != BPTYPE_MacroLibrary;
+			const bool bAllowEvents = (K2Schema->GetGraphType(TargetGraph) == GT_Ubergraph) && bSupportsEventGraphs &&
+				(Blueprint->BlueprintType != BPTYPE_MacroLibrary);
 
 			if(bAllowEvents)
 			{
@@ -234,7 +234,7 @@ bool UK2Node_Timeline::CanPasteHere(const UEdGraph* TargetGraph, const UEdGraphS
 					{
 						if (UK2Node_Composite* Composite = Cast<UK2Node_Composite>(TargetGraph->GetOuter()))
 						{
-							TargetGraph =  Cast<UEdGraph>(Composite->GetOuter());
+							TargetGraph = Cast<UEdGraph>(Composite->GetOuter());
 						}
 						else if (K2Schema->GetGraphType(TargetGraph) == GT_Ubergraph)
 						{
diff --git a/Engine/Source/Editor/GameplayAbilitiesEditor/Classes/K2Node_GameplayCueEvent.h b/Engine/Source/Editor/GameplayAbilitiesEditor/Classes/K2Node_GameplayCueEvent.h
--- a/Engine/Source/Editor/GameplayAbilitiesEditor/Classes/K2Node_GameplayCueEvent.h
+++ b/Engine/Source/Editor/GameplayAbilitiesEditor/Classes/K2Node_GameplayCueEvent.h
@@ -14,7 +14,7 @@ class UK2Node_GameplayCueEvent : public UK2Node_Event
 	virtual FString GetTooltip() const override;
 	virtual FText GetNodeTitle(ENodeTitleType::Type TitleType) const override;
 	virtual void GetMenuEntries(struct FGraphContextMenuBuilder& ContextMenuBuilder) const override;
-	virtual bool CanPasteHere(const UEdGraph* TargetGraph, const UEdGraphSchema* Schema) const override;
+	virtual bool IsCompatibleWithGraph(const UEdGraph* TargetGraph) const override;
 	// End UEdGraphNode interface
 	
 	// Begin UK2Node interface
diff --git a/Engine/Source/Editor/GameplayAbilitiesEditor/Classes/K2Node_LatentAbilityCall.h b/Engine/Source/Editor/GameplayAbilitiesEditor/Classes/K2Node_LatentAbilityCall.h
--- a/Engine/Source/Editor/GameplayAbilitiesEditor/Classes/K2Node_LatentAbilityCall.h
+++ b/Engine/Source/Editor/GameplayAbilitiesEditor/Classes/K2Node_LatentAbilityCall.h
@@ -17,6 +17,7 @@ class UK2Node_LatentAbilityCall : public UK2Node_BaseAsyncTask
 	virtual bool CanCreateUnderSpecifiedSchema(const UEdGraphSchema* DesiredSchema) const;
 	virtual void GetMenuEntries(FGraphContextMenuBuilder& ContextMenuBuilder) const override;
 	virtual void GetMenuActions(TArray<UBlueprintNodeSpawner*>& ActionListOut) const override;
+	virtual bool IsCompatibleWithGraph(UEdGraph const* TargetGraph) const override;
 	// End of UEdGraphNode interface
 
 	virtual void AllocateDefaultPins() override;
diff --git a/Engine/Source/Editor/GameplayAbilitiesEditor/Private/K2Node_GameplayCueEvent.cpp b/Engine/Source/Editor/GameplayAbilitiesEditor/Private/K2Node_GameplayCueEvent.cpp
--- a/Engine/Source/Editor/GameplayAbilitiesEditor/Private/K2Node_GameplayCueEvent.cpp
+++ b/Engine/Source/Editor/GameplayAbilitiesEditor/Private/K2Node_GameplayCueEvent.cpp
@@ -29,26 +29,26 @@ FText UK2Node_GameplayCueEvent::GetNodeTitle(ENodeTitleType::Type TitleType) con
 	//return LOCTEXT("HandleGameplayCueEvent", "HandleGameplaCueEvent");
 }
 
-bool UK2Node_GameplayCueEvent::CanPasteHere(UEdGraph const* TargetGraph, UEdGraphSchema const* Schema) const
+bool UK2Node_GameplayCueEvent::IsCompatibleWithGraph(UEdGraph const* TargetGraph) const
 {
-	bool bValidClass = false;
-	UBlueprint* MyBlueprint = Cast<UBlueprint>(TargetGraph->GetOuter());
-	if (MyBlueprint && MyBlueprint->GeneratedClass)
+	bool bIsCompatible = false;
+	if (bIsCompatible)
 	{
-		if (MyBlueprint->GeneratedClass->ImplementsInterface(UGameplayCueInterface::StaticClass()))
+		if (UBlueprint* Blueprint = FBlueprintEditorUtils::FindBlueprintForGraph(TargetGraph))
 		{
-			bValidClass = true;
+			check(Blueprint->GeneratedClass != nullptr);
+			bIsCompatible = Blueprint->GeneratedClass->ImplementsInterface(UGameplayCueInterface::StaticClass());
 		}
 	}
-
-	return bValidClass;
+	
+	return bIsCompatible && Super::IsCompatibleWithGraph(TargetGraph);
 }
 
 void UK2Node_GameplayCueEvent::GetMenuEntries(FGraphContextMenuBuilder& Context) const
 {
 	Super::GetMenuEntries(Context);
 
-	if (!CanPasteHere(Context.CurrentGraph, GetDefault<UEdGraphSchema>(Context.CurrentGraph->Schema)))
+	if (!IsCompatibleWithGraph(Context.CurrentGraph))
 	{
 		return;
 	}
diff --git a/Engine/Source/Editor/GameplayAbilitiesEditor/Private/K2Node_LatentAbilityCall.cpp b/Engine/Source/Editor/GameplayAbilitiesEditor/Private/K2Node_LatentAbilityCall.cpp
--- a/Engine/Source/Editor/GameplayAbilitiesEditor/Private/K2Node_LatentAbilityCall.cpp
+++ b/Engine/Source/Editor/GameplayAbilitiesEditor/Private/K2Node_LatentAbilityCall.cpp
@@ -23,52 +23,57 @@ bool UK2Node_LatentAbilityCall::CanCreateUnderSpecifiedSchema(const UEdGraphSche
 	return Super::CanCreateUnderSpecifiedSchema(DesiredSchema);
 }
 
-void UK2Node_LatentAbilityCall::GetMenuEntries(FGraphContextMenuBuilder& ContextMenuBuilder) const
+bool UK2Node_LatentAbilityCall::IsCompatibleWithGraph(UEdGraph const* TargetGraph) const
 {
-	bool bValidClass = false;
-	UBlueprint* MyBlueprint = Cast<UBlueprint>(ContextMenuBuilder.CurrentGraph->GetOuter());
-	if (MyBlueprint && MyBlueprint->GeneratedClass)
+	bool bIsCompatible = false;
+	
+	EGraphType GraphType = TargetGraph->GetSchema()->GetGraphType(TargetGraph);
+	bool const bAllowLatentFuncs = (GraphType == GT_Ubergraph || GraphType == GT_Macro);
+	
+	if (bAllowLatentFuncs)
 	{
-		if (MyBlueprint->GeneratedClass->IsChildOf(UGameplayAbility::StaticClass()))
+		UBlueprint* MyBlueprint = FBlueprintEditorUtils::FindBlueprintForGraph(TargetGraph);
+		if (MyBlueprint && MyBlueprint->GeneratedClass)
 		{
-			bValidClass = true;
+			if (MyBlueprint->GeneratedClass->IsChildOf(UGameplayAbility::StaticClass()))
+			{
+				bIsCompatible = true;
+			}
 		}
 	}
+	
+	return bIsCompatible;
+}
 
-	if (!bValidClass)
+void UK2Node_LatentAbilityCall::GetMenuEntries(FGraphContextMenuBuilder& ContextMenuBuilder) const
+{
+	if (!IsCompatibleWithGraph(ContextMenuBuilder.CurrentGraph))
 	{
 		return;
 	}
 
-	const UEdGraphSchema_K2* K2Schema = GetDefault<UEdGraphSchema_K2>();
-	EGraphType GraphType = K2Schema->GetGraphType(ContextMenuBuilder.CurrentGraph);
-	const bool bAllowLatentFuncs = (GraphType == GT_Ubergraph || GraphType == GT_Macro);
-
-	if (bAllowLatentFuncs)
+	for (TObjectIterator<UClass> ClassIt; ClassIt; ++ClassIt)
 	{
-		for (TObjectIterator<UClass> ClassIt; ClassIt; ++ClassIt)
+		UClass* TestClass = *ClassIt;
+		if (TestClass->IsChildOf(UAbilityTask::StaticClass()) && !TestClass->HasAnyClassFlags(CLASS_Abstract | CLASS_Deprecated))
 		{
-			UClass* TestClass = *ClassIt;
-			if (TestClass->IsChildOf(UAbilityTask::StaticClass()) && !TestClass->HasAnyClassFlags(CLASS_Abstract | CLASS_Deprecated))
+			for (TFieldIterator<UFunction> FuncIt(TestClass, EFieldIteratorFlags::ExcludeSuper); FuncIt; ++FuncIt)
 			{
-				for (TFieldIterator<UFunction> FuncIt(TestClass, EFieldIteratorFlags::ExcludeSuper); FuncIt; ++FuncIt)
-				{
-					// See if the function is a static factory method for online proxies
-					UFunction* CandidateFunction = *FuncIt;
+				// See if the function is a static factory method for online proxies
+				UFunction* CandidateFunction = *FuncIt;
 
-					UObjectProperty* ReturnProperty = Cast<UObjectProperty>(CandidateFunction->GetReturnProperty());
-					const bool bValidReturnType = (ReturnProperty != nullptr) && (ReturnProperty->PropertyClass != nullptr) && (ReturnProperty->PropertyClass->IsChildOf(UAbilityTask::StaticClass()));
+				UObjectProperty* ReturnProperty = Cast<UObjectProperty>(CandidateFunction->GetReturnProperty());
+				const bool bValidReturnType = (ReturnProperty != nullptr) && (ReturnProperty->PropertyClass != nullptr) && (ReturnProperty->PropertyClass->IsChildOf(UAbilityTask::StaticClass()));
 
-					if (CandidateFunction->HasAllFunctionFlags(FUNC_Static) && bValidReturnType)
-					{
-						// Create a node template for this factory method
-						UK2Node_LatentAbilityCall* NodeTemplate = NewObject<UK2Node_LatentAbilityCall>(ContextMenuBuilder.OwnerOfTemporaries);
-						NodeTemplate->ProxyFactoryFunctionName = CandidateFunction->GetFName();
-						NodeTemplate->ProxyFactoryClass = TestClass;
-						NodeTemplate->ProxyClass = ReturnProperty->PropertyClass;
+				if (CandidateFunction->HasAllFunctionFlags(FUNC_Static) && bValidReturnType)
+				{
+					// Create a node template for this factory method
+					UK2Node_LatentAbilityCall* NodeTemplate = NewObject<UK2Node_LatentAbilityCall>(ContextMenuBuilder.OwnerOfTemporaries);
+					NodeTemplate->ProxyFactoryFunctionName = CandidateFunction->GetFName();
+					NodeTemplate->ProxyFactoryClass = TestClass;
+					NodeTemplate->ProxyClass = ReturnProperty->PropertyClass;
 
-						CreateDefaultMenuEntry(NodeTemplate, ContextMenuBuilder);
-					}
+					CreateDefaultMenuEntry(NodeTemplate, ContextMenuBuilder);
 				}
 			}
 		}
diff --git a/Engine/Source/Editor/Kismet/Private/BlueprintEditor.cpp b/Engine/Source/Editor/Kismet/Private/BlueprintEditor.cpp
--- a/Engine/Source/Editor/Kismet/Private/BlueprintEditor.cpp
+++ b/Engine/Source/Editor/Kismet/Private/BlueprintEditor.cpp
@@ -3572,7 +3572,7 @@ bool FBlueprintEditor::CanCollapseSelectionToFunction(TSet<class UEdGraphNode*>&
 	{
 		UEdGraphNode* Node = *NodeIt;
 
-		if(!Node->CanPasteHere(FunctionGraph, K2Schema))
+		if(!Node->CanPasteHere(FunctionGraph))
 		{
 			if (UK2Node_CustomEvent* const CustomEvent = Cast<UK2Node_CustomEvent>(Node))
 			{
@@ -3719,7 +3719,7 @@ bool FBlueprintEditor::CanCollapseSelectionToMacro(TSet<class UEdGraphNode*>& In
 	{
 		UEdGraphNode* Node = Cast<UEdGraphNode>(*NodeIt);
 
-		if(!Node->CanPasteHere(MacroGraph, Schema))
+		if(!Node->CanPasteHere(MacroGraph))
 		{
 			LogResults.Error(*LOCTEXT("CannotPasteNodeMacro_Error", "@@ cannot be placed in macro graph").ToString(), Node);
 			bCollapseAllowed = false;
@@ -5368,7 +5368,7 @@ void FBlueprintEditor::ExpandNode(UEdGraphNode* InNodeToExpand, UEdGraph* InSour
 
 		// We do not check CanPasteHere when determining CanCollapseNodes, unlike CanCollapseSelectionToFunction/Macro,
 		// so when expanding a collapsed graph we don't want to check the CanPasteHere function:
-		if (!bIsCollapsedGraph && !Node->CanPasteHere(DestinationGraph, Node->GetSchema()))
+		if (!bIsCollapsedGraph && !Node->CanPasteHere(DestinationGraph))
 		{
 			continue;
 		}
diff --git a/Engine/Source/Editor/Kismet/Private/SBlueprintSubPalette.cpp b/Engine/Source/Editor/Kismet/Private/SBlueprintSubPalette.cpp
--- a/Engine/Source/Editor/Kismet/Private/SBlueprintSubPalette.cpp
+++ b/Engine/Source/Editor/Kismet/Private/SBlueprintSubPalette.cpp
@@ -126,7 +126,7 @@ static bool CanPaletteItemBePlaced(TSharedPtr<FEdGraphSchemaAction> DropActionIn
 			ImpededReasonOut = FText::GetEmpty();
 		}
 		// as a general catch-all, if a node cannot be pasted, it probably can't be created there
-		else if (bCanBePlaced && !NodeToBePlaced->CanPasteHere(HoveredGraphIn, GraphSchema) && !bWillFocusOnExistingNode)
+		else if (bCanBePlaced && !NodeToBePlaced->CanPasteHere(HoveredGraphIn) && !bWillFocusOnExistingNode)
 		{
 			bCanBePlaced = false;
 			ImpededReasonOut = LOCTEXT("CannotPaste", "Cannot place this node in this type of graph");
diff --git a/Engine/Source/Editor/UnrealEd/Classes/MaterialGraph/MaterialGraphNode.h b/Engine/Source/Editor/UnrealEd/Classes/MaterialGraph/MaterialGraphNode.h
--- a/Engine/Source/Editor/UnrealEd/Classes/MaterialGraph/MaterialGraphNode.h
+++ b/Engine/Source/Editor/UnrealEd/Classes/MaterialGraph/MaterialGraphNode.h
@@ -50,7 +50,7 @@ public:
 	// End of UObject interface
 
 	// Begin UEdGraphNode interface.
-	virtual bool CanPasteHere(const UEdGraph* TargetGraph, const UEdGraphSchema* Schema) const override;
+	virtual bool CanPasteHere(const UEdGraph* TargetGraph) const override;
 	virtual FText GetNodeTitle(ENodeTitleType::Type TitleType) const override;
 	virtual FLinearColor GetNodeTitleColor() const override;
 	virtual FString GetTooltip() const override;
diff --git a/Engine/Source/Editor/UnrealEd/Private/EdGraphUtilities.cpp b/Engine/Source/Editor/UnrealEd/Private/EdGraphUtilities.cpp
--- a/Engine/Source/Editor/UnrealEd/Private/EdGraphUtilities.cpp
+++ b/Engine/Source/Editor/UnrealEd/Private/EdGraphUtilities.cpp
@@ -16,17 +16,11 @@ struct FGraphObjectTextFactory : public FCustomizableTextObjectFactory
 public:
 	TSet<UEdGraphNode*> SpawnedNodes;
 	const UEdGraph* DestinationGraph;
-	const UEdGraphSchema* DestinationSchema;
 public:
 	FGraphObjectTextFactory(const UEdGraph* InDestinationGraph)
 		: FCustomizableTextObjectFactory(GWarn)
 		, DestinationGraph(InDestinationGraph)
-		, DestinationSchema(NULL)
 	{
-		if (DestinationGraph != NULL)
-		{
-			DestinationSchema = DestinationGraph->Schema->GetDefaultObject<UEdGraphSchema>();
-		}
 	}
 
 protected:
@@ -36,9 +30,9 @@ protected:
 		{
 			if (DefaultNode->CanDuplicateNode())
 			{
-				if (DestinationSchema != NULL)
+				if (DestinationGraph != NULL)
 				{
-					if (DefaultNode->CanCreateUnderSpecifiedSchema(DestinationSchema))
+					if (DefaultNode->CanCreateUnderSpecifiedSchema(DestinationGraph->GetSchema()))
 					{
 						return true;
 					}
@@ -57,10 +51,10 @@ protected:
 	{
 		if (UEdGraphNode* Node = Cast<UEdGraphNode>(CreatedObject))
 		{
-			if(!Node->CanPasteHere(DestinationGraph, DestinationSchema))
+			if(!Node->CanPasteHere(DestinationGraph))
 			{
 				// Attempt to create a substitute node if it cannot be pasted (note: the return value can be NULL, indicating that the node cannot be pasted into the graph)
-				Node = DestinationSchema->CreateSubstituteNode(Node, DestinationGraph, &InstanceGraph);
+				Node = DestinationGraph->GetSchema()->CreateSubstituteNode(Node, DestinationGraph, &InstanceGraph);
 			}
 
 			if(Node != CreatedObject)
diff --git a/Engine/Source/Editor/UnrealEd/Private/MaterialGraphNode.cpp b/Engine/Source/Editor/UnrealEd/Private/MaterialGraphNode.cpp
--- a/Engine/Source/Editor/UnrealEd/Private/MaterialGraphNode.cpp
+++ b/Engine/Source/Editor/UnrealEd/Private/MaterialGraphNode.cpp
@@ -116,9 +116,9 @@ void UMaterialGraphNode::PostDuplicate(bool bDuplicateForPIE)
 	}
 }
 
-bool UMaterialGraphNode::CanPasteHere(const UEdGraph* TargetGraph, const UEdGraphSchema* Schema) const
+bool UMaterialGraphNode::CanPasteHere(const UEdGraph* TargetGraph) const
 {
-	if (Super::CanPasteHere(TargetGraph, Schema))
+	if (Super::CanPasteHere(TargetGraph))
 	{
 		const UMaterialGraph* MaterialGraph = Cast<const UMaterialGraph>(TargetGraph);
 		if (MaterialGraph)
diff --git a/Engine/Source/Runtime/Engine/Classes/EdGraph/EdGraphNode.h b/Engine/Source/Runtime/Engine/Classes/EdGraph/EdGraphNode.h
--- a/Engine/Source/Runtime/Engine/Classes/EdGraph/EdGraphNode.h
+++ b/Engine/Source/Runtime/Engine/Classes/EdGraph/EdGraphNode.h
@@ -238,12 +238,20 @@ public:
 	/**
 	 * Determine if this node can live in the specified graph
 	 */
-	virtual bool CanPasteHere(const UEdGraph* TargetGraph, const UEdGraphSchema* Schema) const { return CanCreateUnderSpecifiedSchema(Schema); }
+	virtual bool CanPasteHere(const UEdGraph* TargetGraph) const { return IsCompatibleWithGraph(TargetGraph); }
+
+	DEPRECATED(4.5, "The UEdGraphNode::CanPasteHere() that takes a UEdGraphSchema parameter is deprecated, instead use the CanPasteHere() that only takes a single UEdGraph param.")
+	virtual bool CanPasteHere(const UEdGraph* TargetGraph, const UEdGraphSchema* Schema) const { return CanPasteHere(TargetGraph); }
 
 	/**
 	 * Determine if this node can be created under the specified schema
      */
 	virtual bool CanCreateUnderSpecifiedSchema(const UEdGraphSchema* Schema) const { return true; }
+	
+	/**
+	 * Determine if a node of this type can be created for the specified graph.
+     */
+	virtual bool IsCompatibleWithGraph(UEdGraph const* Graph) const;
 
 	/**
 	 * Perform any fixups (deep copies of associated data, etc...) necessary after a node has been pasted in the editor
diff --git a/Engine/Source/Runtime/Engine/Private/EdGraph/EdGraphNode.cpp b/Engine/Source/Runtime/Engine/Private/EdGraph/EdGraphNode.cpp
--- a/Engine/Source/Runtime/Engine/Private/EdGraph/EdGraphNode.cpp
+++ b/Engine/Source/Runtime/Engine/Private/EdGraph/EdGraphNode.cpp
@@ -162,6 +162,11 @@ const class UEdGraphSchema* UEdGraphNode::GetSchema() const
 	return GetGraph()->GetSchema();
 }
 
+bool UEdGraphNode::IsCompatibleWithGraph(UEdGraph const* Graph) const
+{
+	return CanCreateUnderSpecifiedSchema(Graph->GetSchema());
+}
+
 FLinearColor UEdGraphNode::GetNodeTitleColor() const
 {
 	return FLinearColor(0.4f, 0.62f, 1.0f);
