(WIP) Adding new bound blueprint node spawners for the new context menu system.

diff --git a/Engine/Source/Editor/BlueprintGraph/Classes/BlueprintBoundNodeSpawner.h b/Engine/Source/Editor/BlueprintGraph/Classes/BlueprintBoundNodeSpawner.h
new file mode 100644
--- /dev/null
+++ b/Engine/Source/Editor/BlueprintGraph/Classes/BlueprintBoundNodeSpawner.h
@@ -0,0 +1,54 @@
+// Copyright 1998-2014 Epic Games, Inc. All Rights Reserved.
+
+#pragma once
+
+#include "BlueprintNodeSpawner.h"
+#include "BlueprintBoundNodeSpawner.generated.h"
+
+// Forward declarations
+class UActorComponent;
+
+/**
+ *
+ */
+UCLASS(Transient)
+class BLUEPRINTGRAPH_API UBlueprintBoundNodeSpawner : public UBlueprintNodeSpawner
+{
+	GENERATED_UCLASS_BODY()
+
+public:
+	/**
+	 * 
+	 */
+	static UBlueprintBoundNodeSpawner* Create(TSubclassOf<UEdGraphNode> NodeClass, UObject* BoundObject, UObject* Outer = nullptr);
+
+	/**
+	 * 
+	 */
+	static UBlueprintBoundNodeSpawner* Create(UBlueprintNodeSpawner* SubSpawner, UObject* BoundObject, UObject* Outer = nullptr);
+
+	// UBlueprintNodeSpawner interface
+	virtual UEdGraphNode* Invoke(UEdGraph* ParentGraph, FVector2D const Location) const override;
+	virtual FText GetDefaultMenuName() const override;
+	virtual FText GetDefaultMenuCategory() const override;
+	// End UBlueprintNodeSpawner interface
+	
+	/**
+	 * 
+	 */
+	UObject const* GetBoundObject() const;
+	
+	/**
+	 *
+	 */
+	UBlueprintNodeSpawner const* GetSubSpawner() const;
+
+private:
+	/** */
+	UPROPERTY()
+	TWeakObjectPtr<UObject> BoundObjPtr;
+
+	/** */
+	UPROPERTY()
+	UBlueprintNodeSpawner* SubSpawner;
+};
diff --git a/Engine/Source/Editor/BlueprintGraph/Classes/EdGraphSchema_K2_Actions.h b/Engine/Source/Editor/BlueprintGraph/Classes/EdGraphSchema_K2_Actions.h
--- a/Engine/Source/Editor/BlueprintGraph/Classes/EdGraphSchema_K2_Actions.h
+++ b/Engine/Source/Editor/BlueprintGraph/Classes/EdGraphSchema_K2_Actions.h
@@ -416,7 +416,7 @@ struct BLUEPRINTGRAPH_API FEdGraphSchemaAction_K2TargetNode : public FEdGraphSch
 
 /** Action to paste at this location on graph*/
 USTRUCT()
-struct FEdGraphSchemaAction_K2PasteHere : public FEdGraphSchemaAction_K2NewNode
+struct BLUEPRINTGRAPH_API FEdGraphSchemaAction_K2PasteHere : public FEdGraphSchemaAction_K2NewNode
 {
 	GENERATED_USTRUCT_BODY()
 
diff --git a/Engine/Source/Editor/BlueprintGraph/Private/BlueprintActionDatabase.cpp b/Engine/Source/Editor/BlueprintGraph/Private/BlueprintActionDatabase.cpp
--- a/Engine/Source/Editor/BlueprintGraph/Private/BlueprintActionDatabase.cpp
+++ b/Engine/Source/Editor/BlueprintGraph/Private/BlueprintActionDatabase.cpp
@@ -9,6 +9,7 @@
 #include "BlueprintPropertyNodeSpawner.h"
 #include "BlueprintEventNodeSpawner.h"
 #include "BlueprintComponentNodeSpawner.h"
+#include "BlueprintBoundNodeSpawner.h"
 
 // used below in FBlueprintNodeSpawnerFactory::MakeMacroNodeSpawner()
 #include "K2Node_MacroInstance.h"
@@ -309,6 +310,16 @@ namespace FBlueprintActionDatabaseImpl
 	 */
 	static void AddClassPropertyActions(UClass const* const Class, FActionList& ActionListOut);
 
+	/**
+	 * Loops over all of the property's class properties and creates bound 
+	 * actions for any that are viable for blueprint use. Evolved from certain
+	 * parts of FK2ActionMenuBuilder::AddBoundEventActionsForClass().
+	 *
+	 * @param  Property			The component property that you want node-spawners for.
+	 * @param  ActionListOut	The list you want populated with the new spawners.
+	 */
+	static void AddComponentBoundActions(UObjectProperty* Property, FActionList& ActionListOut);
+
 	/**
 	 * Evolved from FClassDynamicCastHelper::GetClassDynamicCastNodes(). If the
 	 * specified class is a viable blueprint variable type, then two cast nodes
@@ -485,12 +496,15 @@ static void FBlueprintActionDatabaseImpl::AddClassPropertyActions(UClass const*
 			ActionListOut.Add(ClearSpawner);
 
 			// @TODO: AddBoundEventActionsForClass()
-			//   UK2Node_ComponentBoundEvent 
 			//   UK2Node_ActorBoundEvent
  		}
 		else
 		{
-			
+			if (UObjectProperty* ObjProperty = Cast<UObjectProperty>(Property))
+			{
+				AddComponentBoundActions(ObjProperty, ActionListOut);
+			}
+
 			UBlueprintPropertyNodeSpawner* GetterSpawner = UBlueprintPropertyNodeSpawner::Create<UK2Node_VariableGet>(Property);
 			ActionListOut.Add(GetterSpawner);
 			UBlueprintPropertyNodeSpawner* SetterSpawner = UBlueprintPropertyNodeSpawner::Create<UK2Node_VariableSet>(Property);
@@ -501,6 +515,65 @@ static void FBlueprintActionDatabaseImpl::AddClassPropertyActions(UClass const*
 	// @TODO: if blueprint class, loop over function graphs and get local variables
 }
 
+//------------------------------------------------------------------------------
+static void FBlueprintActionDatabaseImpl::AddComponentBoundActions(UObjectProperty* ComponentProperty, FActionList& ActionListOut)
+{
+	check(ComponentProperty != nullptr);
+
+	bool const bIsComponentProperty = ComponentProperty->PropertyClass->IsChildOf<UActorComponent>();
+	if (bIsComponentProperty)
+	{
+		// evolved from FK2ActionMenuBuilder::GetBoundEventsFromComponentSelection()
+		for (TFieldIterator<UMulticastDelegateProperty> PropertyIt(ComponentProperty->PropertyClass, EFieldIteratorFlags::IncludeSuper); PropertyIt; ++PropertyIt)
+		{
+			UMulticastDelegateProperty* Property = *PropertyIt;
+			if (!IsPropertyBlueprintVisible(Property))
+			{
+				continue;
+			}
+
+			UBlueprintPropertyNodeSpawner* PropertySpawner = UBlueprintPropertyNodeSpawner::Create<UK2Node_ComponentBoundEvent>(Property);
+			UBlueprintBoundNodeSpawner* BoundEventSpawner = UBlueprintBoundNodeSpawner::Create(PropertySpawner, ComponentProperty);
+
+			auto BindComponentLambda = [](UEdGraphNode* NewNode, bool /*bIsTemplateNode*/, TWeakObjectPtr<UObjectProperty> ComponentProp, TWeakObjectPtr<UMulticastDelegateProperty> DelegateProp)
+			{
+				UK2Node_ComponentBoundEvent* BoundNode = CastChecked<UK2Node_ComponentBoundEvent>(NewNode);
+				if (ComponentProp.IsValid() && DelegateProp.IsValid())
+				{
+					BoundNode->InitializeComponentBoundEventParams(ComponentProp.Get(), DelegateProp.Get());
+				}
+			};
+
+			TWeakObjectPtr<UObjectProperty> ComponentPropPtr = ComponentProperty;
+			TWeakObjectPtr<UMulticastDelegateProperty> DelegatePropPtr = Property;
+			BoundEventSpawner->CustomizeNodeDelegate = UBlueprintNodeSpawner::FCustomizeNodeDelegate::CreateStatic(BindComponentLambda, ComponentPropPtr, DelegatePropPtr);
+			
+			ActionListOut.Add(BoundEventSpawner);
+		}
+		
+		// evolved from FK2ActionMenuBuilder::GetFunctionCallsOnSelectedComponents()
+		for (TFieldIterator<UFunction> FuncIt(ComponentProperty->PropertyClass, EFieldIteratorFlags::IncludeSuper); FuncIt; ++FuncIt)
+		{
+			UFunction* ComponentFunc = *FuncIt;
+			if (!UEdGraphSchema_K2::CanUserKismetCallFunction(ComponentFunc))
+			{
+				continue;
+			}
+			
+			bool const bIsPure  = (ComponentFunc->HasAnyFunctionFlags(FUNC_BlueprintPure) != false);
+			bool const bIsConst = (ComponentFunc->HasAnyFunctionFlags(FUNC_Const) != false);
+			
+			// don't know why we only allow impure ("imperitive") and const functions,
+			// but this mimics legacy functionality in GetFunctionCallsOnSelectedComponents()
+			if (!bIsPure || bIsConst)
+			{
+				UBlueprintFunctionNodeSpawner* FuncSpawner = UBlueprintFunctionNodeSpawner::Create(ComponentFunc);
+				ActionListOut.Add(UBlueprintBoundNodeSpawner::Create(FuncSpawner, ComponentProperty));
+			}
+		}
+	}
+}
+
 //------------------------------------------------------------------------------
 static void FBlueprintActionDatabaseImpl::AddClassCastActions(UClass* const Class, FActionList& ActionListOut)
 {
@@ -732,7 +805,7 @@ void FBlueprintActionDatabase::RefreshClassActions(UClass* const Class)
 		AddAutonomousNodeActions(Class, ClassActionList);
 	}
 
-	// @TODO: account for all K2ActionMenuBuilder mathods...
+	// @TODO: account for all K2ActionMenuBuilder methods...
 	//   GetLiteralsFromActorSelection() - UK2Node_Literal
 	//   GetAnimNotifyMenuItems()
 	//   GetMatineeControllers() - UK2Node_MatineeController
diff --git a/Engine/Source/Editor/BlueprintGraph/Private/BlueprintActionFilter.cpp b/Engine/Source/Editor/BlueprintGraph/Private/BlueprintActionFilter.cpp
--- a/Engine/Source/Editor/BlueprintGraph/Private/BlueprintActionFilter.cpp
+++ b/Engine/Source/Editor/BlueprintGraph/Private/BlueprintActionFilter.cpp
@@ -7,6 +7,7 @@
 #include "BlueprintEventNodeSpawner.h"
 #include "BlueprintComponentNodeSpawner.h"
 #include "BlueprintPropertyNodeSpawner.h"
+#include "BlueprintBoundNodeSpawner.h"
 #include "EdGraphSchema_K2.h"		// for FBlueprintMetadata
 #include "BlueprintEditorUtils.h"	// for FindBlueprintForGraph()
 #include "ObjectEditorUtils.h"		// for IsFunctionHiddenFromClass()/IsVariableCategoryHiddenFromClass()
@@ -71,6 +72,15 @@ namespace BlueprintActionFilterImpl
 	 */
 	static bool IsLatent(UBlueprintNodeSpawner const* NodeSpawner);
 	
+	/**
+	 * Utility method to check and see if the specified node-spawner would
+	 * produce a node associated with a global field.
+	 *
+	 * @param  NodeSpawner	The action you want to query.
+	 * @return True if the action is associated with some global field, false if not (or unknown).
+	 */
+	static bool IsGlobal(UField const* Field);
+	
 	/**
 	 * Rejection test that checks to see if the supplied node-spawner would
 	 * produce an event that does NOT belong to the specified blueprint.
@@ -153,7 +163,7 @@ namespace BlueprintActionFilterImpl
 	 * @param  BlueprintAction	The action you wish to query.
 	 * @return True if the action would spawn a node not allowed in function graphs.
 	 */
-	static bool IsGlobalField(FBlueprintActionFilter const& Filter, UBlueprintNodeSpawner const* BlueprintAction);
+	static bool IsGlobal(FBlueprintActionFilter const& Filter, UBlueprintNodeSpawner const* BlueprintAction);
 
 	/**
 	 * Rejection test that checks to see if the node-spawner is associated with 
@@ -186,6 +196,17 @@ namespace BlueprintActionFilterImpl
 	 * @return True if the action would produce a non-whitelisted node.
 	 */
 	static bool IsFilteredNodeType(FBlueprintActionFilter const& Filter, UBlueprintNodeSpawner const* BlueprintAction, bool const bMatchExplicitly = false);
+
+	/**
+	 * Rejection test that checks to see if the node-spawner is tied to a 
+	 * specific object that is not currently selected.
+	 * 
+	 * 
+	 * @param  Context			Holds the selected object context for this test.
+	 * @param  BlueprintAction	The action you wish to query.
+	 * @return True if the action would produce a bound node, tied to an object that isn't selected.
+	 */
+	static bool IsBoundToUnselectedObject(FBlueprintActionFilter const& Filter, UBlueprintNodeSpawner const* BlueprintAction);
 };
 
 //------------------------------------------------------------------------------
@@ -210,12 +231,16 @@ static UClass* BlueprintActionFilterImpl::GetActionClass(UBlueprintNodeSpawner c
 	{
 		UProperty const* Property = PropertySpawner->GetProperty();
 		check(Property != nullptr);
-		ActionClass = Cast<UClass>(Property->GetOuterUField());
+		ActionClass = Property->GetOwnerClass();
 	}
 	else if (UBlueprintComponentNodeSpawner const* ComponentSpawner = Cast<UBlueprintComponentNodeSpawner>(BlueprintAction))
 	{
 		ActionClass = ComponentSpawner->GetComponentClass();
 	}
+	else if (UBlueprintBoundNodeSpawner const* BoundSpawner = Cast<UBlueprintBoundNodeSpawner>(BlueprintAction))
+	{
+		ActionClass = GetActionClass(BoundSpawner->GetSubSpawner());
+	}
 	
 	return ActionClass;
 }
@@ -233,6 +258,10 @@ static UField const* BlueprintActionFilterImpl::GetAssociatedField(UBlueprintNod
 	{
 		ClassField = PropertySpawner->GetProperty();
 	}
+	else if (UBlueprintBoundNodeSpawner const* BoundSpawner = Cast<UBlueprintBoundNodeSpawner>(BlueprintAction))
+	{
+		ClassField = GetAssociatedField(BoundSpawner->GetSubSpawner());
+	}
 
 	return ClassField;
 }
@@ -250,6 +279,10 @@ static UFunction const* BlueprintActionFilterImpl::GetAssociatedFunction(UBluepr
 	{
 		Function = EventSpawner->GetEventFunction();
 	}
+	else if (UBlueprintBoundNodeSpawner const* BoundSpawner = Cast<UBlueprintBoundNodeSpawner>(BlueprintAction))
+	{
+		Function = GetAssociatedFunction(BoundSpawner->GetSubSpawner());
+	}
 
 	return Function;
 }
@@ -295,6 +328,28 @@ static bool BlueprintActionFilterImpl::IsLatent(UBlueprintNodeSpawner const* Nod
 	return bIsLatent;
 }
 
+//------------------------------------------------------------------------------
+static bool BlueprintActionFilterImpl::IsGlobal(UField const* Field)
+{
+	bool bIsGlobal = false;
+	
+	UClass* ClassOuter = Cast<UClass>(Field->GetOuter());
+	// outer is probably a UPackage (for like global enums, structs, etc.)
+	if (ClassOuter == nullptr)
+	{
+		bIsGlobal = true;
+	}
+	else if (UFunction const* Function = Cast<UFunction>(Field))
+	{
+		bool const bIsPublic = !Function->HasMetaData(FBlueprintMetadata::MD_Protected) &&
+			!Function->HasMetaData(FBlueprintMetadata::MD_Private);
+		
+		bIsGlobal = Function->HasAnyFunctionFlags(FUNC_Static) && bIsPublic;
+	}
+	
+	return bIsGlobal;
+}
+
 //------------------------------------------------------------------------------
 static bool BlueprintActionFilterImpl::IsEventUnimplementable(FBlueprintActionFilter const& Filter, UBlueprintNodeSpawner const* BlueprintAction)
 {
@@ -450,13 +505,26 @@ static bool BlueprintActionFilterImpl::IsDeprecated(FBlueprintActionFilter const
 }
 
 //------------------------------------------------------------------------------
-static bool BlueprintActionFilterImpl::IsGlobalField(FBlueprintActionFilter const& Filter, UBlueprintNodeSpawner const* BlueprintAction)
+static bool BlueprintActionFilterImpl::IsGlobal(FBlueprintActionFilter const& Filter, UBlueprintNodeSpawner const* BlueprintAction)
 {
 	bool bIsFilteredOut = false;
 	if (UField const* ClassField = GetAssociatedField(BlueprintAction))
 	{
-		UClass* FieldClass = Cast<UClass>(ClassField->GetOuter());
-		bIsFilteredOut = (FieldClass == nullptr) || FieldClass->IsChildOf<UBlueprintFunctionLibrary>();
+		bIsFilteredOut = IsGlobal(ClassField);
+		
+		UClass* FieldClass = ClassField->GetOwnerClass();
+		if (bIsFilteredOut && (FieldClass != nullptr))
+		{
+			for (UClass const* Class : Filter.OwnerClasses)
+			{
+				bool const bIsInternalFunc = Class->IsChildOf(FieldClass);
+				if (bIsInternalFunc)
+				{
+					bIsFilteredOut = false;
+					break;
+				}
+			}
+		}
 	}
 
 	return bIsFilteredOut;
@@ -467,12 +535,11 @@ static bool BlueprintActionFilterImpl::IsExternalField(FBlueprintActionFilter co
 {
 	bool bIsFilteredOut = false;
 	if (UField const* ClassField = GetAssociatedField(BlueprintAction))
-	{		
+	{
 		UClass* FieldClass = Cast<UClass>(ClassField->GetOuter());
-		bool const bIsGlobal = (FieldClass == nullptr) || FieldClass->IsChildOf<UBlueprintFunctionLibrary>();
-
-		// global (and static library) fields can stay (unless explicitly excluded... save that for a separate test)
-		if (!bIsGlobal)
+		// global (and static library) fields can stay (unless explicitly
+		// excluded... save that for a separate test)
+		if ((FieldClass != nullptr) && !IsGlobal(ClassField))
 		{
 			for (UClass const* Class : Filter.OwnerClasses)
 			{
@@ -589,8 +656,8 @@ static bool BlueprintActionFilterImpl::IsFilteredNodeType(FBlueprintActionFilter
 	{
 		for (TSubclassOf<UEdGraphNode> AllowedClass : Filter.NodeTypes)
 		{
-			if ((bExcludeChildClasses && (AllowedClass != NodeClass)) ||
-				!NodeClass->IsChildOf(AllowedClass))
+			if ((NodeClass->IsChildOf(AllowedClass) && !bExcludeChildClasses) ||
+				(AllowedClass == NodeClass))
 			{
 				bIsFilteredOut = false;
 				break;
@@ -611,6 +678,29 @@ static bool BlueprintActionFilterImpl::IsFilteredNodeType(FBlueprintActionFilter
 	return bIsFilteredOut;
 }
 
+//------------------------------------------------------------------------------
+static bool BlueprintActionFilterImpl::IsBoundToUnselectedObject(FBlueprintActionFilter const& Filter, UBlueprintNodeSpawner const* BlueprintAction)
+{
+	bool bIsFilteredOut = false;
+	if (UBlueprintBoundNodeSpawner const* BoundSpawner = Cast<UBlueprintBoundNodeSpawner>(BlueprintAction))
+	{
+		bIsFilteredOut = true;
+
+		if (UObject const* BoundObject = BoundSpawner->GetBoundObject())
+		{
+			for (UObject* SelectedObj : Filter.Context.SelectedObjects)
+			{
+				if (SelectedObj == BoundObject)
+				{
+					bIsFilteredOut = false;
+					break;
+				}
+			}
+		}
+	}
+	return bIsFilteredOut;
+}
+
 /*******************************************************************************
  * FBlueprintActionFilter
  ******************************************************************************/
@@ -631,12 +721,13 @@ FBlueprintActionFilter::FBlueprintActionFilter(uint32 Flags/*= 0x00*/)
 		AddIsFilteredTest(FIsFilteredDelegate::CreateStatic(IsDeprecated));
 	}
 	AddIsFilteredTest(FIsFilteredDelegate::CreateStatic(IsFilteredNodeType, (Flags & BPFILTER_ExcludeChildNodeTypes) != 0));
+	AddIsFilteredTest(FIsFilteredDelegate::CreateStatic(IsBoundToUnselectedObject));
 	AddIsFilteredTest(FIsFilteredDelegate::CreateStatic(IsEventUnimplementable));
 	AddIsFilteredTest(FIsFilteredDelegate::CreateStatic(IsFieldInaccessible));
 	AddIsFilteredTest(FIsFilteredDelegate::CreateStatic(IsRestrictedClassMember));
 	if (Flags & BPFILTER_ExcludeGlobalFields)
 	{
-		AddIsFilteredTest(FIsFilteredDelegate::CreateStatic(IsGlobalField));
+		AddIsFilteredTest(FIsFilteredDelegate::CreateStatic(IsGlobal));
 	}
 	AddIsFilteredTest(FIsFilteredDelegate::CreateStatic(IsExternalField));
 	AddIsFilteredTest(FIsFilteredDelegate::CreateStatic(IsIncompatibleImpureNode));
diff --git a/Engine/Source/Editor/BlueprintGraph/Private/BlueprintBoundNodeSpawner.cpp b/Engine/Source/Editor/BlueprintGraph/Private/BlueprintBoundNodeSpawner.cpp
new file mode 100644
--- /dev/null
+++ b/Engine/Source/Editor/BlueprintGraph/Private/BlueprintBoundNodeSpawner.cpp
@@ -0,0 +1,190 @@
+// Copyright 1998-2014 Epic Games, Inc. All Rights Reserved.
+
+#include "BlueprintGraphPrivatePCH.h"
+#include "BlueprintComponentNodeSpawner.h"
+#include "K2Node_ComponentBoundEvent.h"
+#include "BlueprintPropertyNodeSpawner.h"
+#include "K2Node_AddComponent.h"
+#include "K2Node_CallFunction.h"
+#include "ComponentAssetBroker.h"
+
+#define LOCTEXT_NAMESPACE "BlueprintBoundNodeSpawner"
+
+/*******************************************************************************
+ * Static UBlueprintBoundNodeSpawner Helpers
+ ******************************************************************************/
+
+namespace BlueprintBoundNodeSpawnerImpl
+{
+	static void BindAddComponentNodeWithAsset(UEdGraphNode* NewNode, UObject* AssetObj);
+	static void BindFunctionNodeWithProperty(UEdGraphNode* NewNode, UObject* PropertyObj);
+};
+
+//------------------------------------------------------------------------------
+static void BlueprintBoundNodeSpawnerImpl::BindAddComponentNodeWithAsset(UEdGraphNode* NewNode, UObject* AssetObj)
+{
+	// @TODO: assert the UObject is an asset
+	UK2Node_AddComponent* AddCompNode = CastChecked<UK2Node_AddComponent>(NewNode);
+	
+	if (UActorComponent* ComponentTemplate = AddCompNode->GetTemplateFromNode())
+	{
+		FComponentAssetBrokerage::AssignAssetToComponent(ComponentTemplate, AssetObj);
+	}
+}
+
+//------------------------------------------------------------------------------
+static void BlueprintBoundNodeSpawnerImpl::BindFunctionNodeWithProperty(UEdGraphNode* NewNode, UObject* PropertyObj)
+{
+	UK2Node_CallFunction* FuncNode = CastChecked<UK2Node_CallFunction>(NewNode);
+	if (UProperty const* BindingProperty = Cast<UProperty>(PropertyObj))
+	{
+		// @TODO: unnecessary allocation, could do this all by hand, or with a local
+		UBlueprintNodeSpawner* TempNodeSpawner = UBlueprintPropertyNodeSpawner::Create<UK2Node_VariableGet>(BindingProperty);
+		
+		float const EstimatedVarNodeWidth = 224.0f;
+		FVector2D VarNodePos;
+		VarNodePos.X = FuncNode->NodePosX - EstimatedVarNodeWidth;
+		VarNodePos.Y = FuncNode->NodePosY;
+		
+		float const EstimatedVarNodeHeight  = 48.0f;
+		float const EstimatedFuncNodeHeight = UEdGraphSchema_K2::EstimateNodeHeight(FuncNode);
+		float const FuncNodeMidYCoordinate  = FuncNode->NodePosY + (EstimatedFuncNodeHeight / 2.0f);
+		VarNodePos.Y = FuncNodeMidYCoordinate - (EstimatedVarNodeWidth / 2.0f);
+		
+		UEdGraph* ParentGraph = FuncNode->GetGraph();
+		UK2Node_VariableGet* GetVarNode = CastChecked<UK2Node_VariableGet>(TempNodeSpawner->Invoke(ParentGraph, VarNodePos));
+		
+		UEdGraphPin* LiteralOutput = GetVarNode->GetValuePin();
+		UEdGraphPin* CallSelfInput = FuncNode->FindPin(GetDefault<UEdGraphSchema_K2>()->PN_Self);
+		// connect the new "get-var" node with the spawned function node
+		if ((LiteralOutput != nullptr) && (CallSelfInput != nullptr))
+		{
+			LiteralOutput->MakeLinkTo(CallSelfInput);
+		}
+	}
+}
+
+/*******************************************************************************
+ * UBlueprintBoundNodeSpawner
+ ******************************************************************************/
+
+//------------------------------------------------------------------------------
+UBlueprintBoundNodeSpawner* UBlueprintBoundNodeSpawner::Create(TSubclassOf<UEdGraphNode> NodeClass, UObject* BoundObject, UObject* Outer/*= nullptr*/)
+{
+	UBlueprintNodeSpawner* SubSpawner = UBlueprintNodeSpawner::Create(NodeClass);
+	return Create(SubSpawner, BoundObject, Outer);
+}
+
+//------------------------------------------------------------------------------
+UBlueprintBoundNodeSpawner* UBlueprintBoundNodeSpawner::Create(UBlueprintNodeSpawner* SubSpawner, UObject* BoundObject, UObject* Outer/*= nullptr*/)
+{
+	if (Outer == nullptr)
+	{
+		Outer = GetTransientPackage();
+	}
+	check(SubSpawner != nullptr);
+
+	UBlueprintBoundNodeSpawner* NodeSpawner = NewObject<UBlueprintBoundNodeSpawner>(Outer);
+	NodeSpawner->NodeClass   = SubSpawner->NodeClass;
+	NodeSpawner->SubSpawner  = DuplicateObject(SubSpawner, NodeSpawner);
+	NodeSpawner->SubSpawner->CustomizeNodeDelegate = SubSpawner->CustomizeNodeDelegate;
+	NodeSpawner->BoundObjPtr = BoundObject;
+
+	return NodeSpawner;
+}
+
+//------------------------------------------------------------------------------
+UBlueprintBoundNodeSpawner::UBlueprintBoundNodeSpawner(class FPostConstructInitializeProperties const& PCIP)
+	: Super(PCIP)
+	, BoundObjPtr(nullptr)
+{
+}
+
+//------------------------------------------------------------------------------
+UEdGraphNode* UBlueprintBoundNodeSpawner::Invoke(UEdGraph* ParentGraph, FVector2D const Location) const
+{
+	using namespace BlueprintBoundNodeSpawnerImpl;
+	check(SubSpawner != nullptr);
+
+	UEdGraphNode* NewNode = SubSpawner->Invoke(ParentGraph, Location);
+	
+	UObject* BoundObj = nullptr;
+	if (BoundObjPtr.IsValid())
+	{
+		BoundObj = BoundObjPtr.Get();
+	}
+	
+	bool const bIsTemplateNode = ParentGraph->HasAnyFlags(RF_Transient);
+	if (!bIsTemplateNode && NewNode->IsA<UK2Node_AddComponent>())
+	{
+		BindAddComponentNodeWithAsset(NewNode, BoundObj);
+	}
+	else if (NewNode->IsA<UK2Node_CallFunction>())
+	{
+		BindFunctionNodeWithProperty(NewNode, BoundObj);
+	}
+	
+	return NewNode;
+}
+
+//------------------------------------------------------------------------------
+FText UBlueprintBoundNodeSpawner::GetDefaultMenuName() const
+{
+	check(SubSpawner != nullptr);
+	
+	FText MenuName;
+	if (NodeClass == UK2Node_ComponentBoundEvent::StaticClass())
+	{
+		if (UBlueprintPropertyNodeSpawner* PropSpawner = Cast<UBlueprintPropertyNodeSpawner>(SubSpawner))
+		{
+			FText const PropertyName = FText::FromName(PropSpawner->GetProperty()->GetFName());
+			MenuName = FText::Format(LOCTEXT("ComponentEventName", "Add {0}"), PropertyName);
+		}
+	}
+	else
+	{
+		MenuName = SubSpawner->GetDefaultMenuName();
+	}
+		
+	return MenuName;
+}
+
+//------------------------------------------------------------------------------
+FText UBlueprintBoundNodeSpawner::GetDefaultMenuCategory() const
+{
+	check(BoundObjPtr != nullptr);
+	check(SubSpawner != nullptr);
+	
+	FText DefaultCategory = SubSpawner->GetDefaultMenuCategory();
+	if (NodeClass == UK2Node_ComponentBoundEvent::StaticClass())
+	{
+		FText const PropertyName = FText::FromName(BoundObjPtr->GetFName());
+		DefaultCategory = FText::Format(LOCTEXT("ComponentEventCategory", "Add Event for {0}|{1}"), PropertyName, DefaultCategory);
+	}
+	else if (NodeClass == UK2Node_CallFunction::StaticClass())
+	{
+		FText const PropertyName = FText::FromName(BoundObjPtr->GetFName());
+		DefaultCategory = FText::Format(LOCTEXT("ComponentEventCategory", "Call Function on {0}|{1}"), PropertyName, DefaultCategory);
+	}
+	
+	return DefaultCategory;
+}
+
+//------------------------------------------------------------------------------
+UObject const* UBlueprintBoundNodeSpawner::GetBoundObject() const
+{
+	UObject* RawBoundObj = nullptr;
+	if (BoundObjPtr.IsValid())
+	{
+		RawBoundObj = BoundObjPtr.Get();
+	}
+	return RawBoundObj;
+}
+
+//------------------------------------------------------------------------------
+UBlueprintNodeSpawner const* UBlueprintBoundNodeSpawner::GetSubSpawner() const
+{
+	return SubSpawner;
+}
+
+#undef LOCTEXT_NAMESPACE
diff --git a/Engine/Source/Editor/BlueprintGraph/Public/BlueprintActionFilter.h b/Engine/Source/Editor/BlueprintGraph/Public/BlueprintActionFilter.h
--- a/Engine/Source/Editor/BlueprintGraph/Public/BlueprintActionFilter.h
+++ b/Engine/Source/Editor/BlueprintGraph/Public/BlueprintActionFilter.h
@@ -32,6 +32,13 @@ struct FBlueprintActionContext
 	 * be viable for every pin to pass the filter).
 	 */
 	TArray<UEdGraphPin*> Pins;
+
+	/** 
+	 * A list of objects the user currently has selected (things like blueprint 
+	 * properties, level actors, content-browser assets, etc.). Objects listed 
+	 * here will result in certain bound actions passing the filter.
+	 */
+	TArray<UObject*> SelectedObjects;
 };
 
 
diff --git a/Engine/Source/Editor/Kismet/Private/BlueprintActionMenuUtils.cpp b/Engine/Source/Editor/Kismet/Private/BlueprintActionMenuUtils.cpp
--- a/Engine/Source/Editor/Kismet/Private/BlueprintActionMenuUtils.cpp
+++ b/Engine/Source/Editor/Kismet/Private/BlueprintActionMenuUtils.cpp
@@ -4,11 +4,37 @@
 #include "BlueprintActionMenuUtils.h"
 #include "BlueprintActionMenuBuilder.h"
 #include "BlueprintActionFilter.h"
+#include "BlueprintBoundNodeSpawner.h"
+#include "KismetEditorUtilities.h" // for CanPasteNodes()
 #include "K2Node_CallFunction.h"
 #include "K2Node_Event.h"
 
 #define LOCTEXT_NAMESPACE "BlueprintActionMenuUtils"
 
+/*******************************************************************************
+* Static FBlueprintActionMenuUtils Helpers
+******************************************************************************/
+
+namespace BlueprintActionMenuUtilsImpl
+{
+	static bool IsUnBoundSpawner(FBlueprintActionFilter const& Filter, UBlueprintNodeSpawner const* BlueprintAction);
+}
+
+//------------------------------------------------------------------------------
+static bool BlueprintActionMenuUtilsImpl::IsUnBoundSpawner(FBlueprintActionFilter const& /*Filter*/, UBlueprintNodeSpawner const* BlueprintAction)
+{
+	bool bIsFilteredOut = true;
+	if (UBlueprintBoundNodeSpawner const* BoundSpawner = Cast<UBlueprintBoundNodeSpawner>(BlueprintAction))
+	{
+		bIsFilteredOut = (BoundSpawner->GetBoundObject() == nullptr);
+	}
+	return bIsFilteredOut;
+}
+
+/*******************************************************************************
+* FBlueprintActionMenuUtils
+******************************************************************************/
+
 //------------------------------------------------------------------------------
 void FBlueprintActionMenuUtils::MakePaletteMenu(FBlueprintActionContext const& Context, UClass* FilterClass, FBlueprintActionMenuBuilder& MenuOut)
 {
@@ -27,49 +53,75 @@ void FBlueprintActionMenuUtils::MakePaletteMenu(FBlueprintActionContext const& C
 	{
 		MenuFilter.OwnerClasses.Add(FilterClass);
 	}
-	
+
 	MenuOut.AddMenuSection(MenuFilter, LOCTEXT("PaletteRoot", "Library"), /*MenuOrder =*/0, FBlueprintActionMenuBuilder::ConsolidatePropertyActions);
 	MenuOut.RebuildActionList();
 }
 
 //------------------------------------------------------------------------------
-void FBlueprintActionMenuUtils::MakeContextMenu(FBlueprintActionContext const& Context, TArray<UProperty*> const& SelectedProperties, FBlueprintActionMenuBuilder& MenuOut)
+void FBlueprintActionMenuUtils::MakeContextMenu(FBlueprintActionContext const& Context, bool bIsContextSensitive, FBlueprintActionMenuBuilder& MenuOut)
 {
-	MenuOut.Empty();
+	using namespace BlueprintActionMenuUtilsImpl;
 	
-	FBlueprintActionFilter MenuFilter;
-	MenuFilter.Context = Context;
+	static int32 const MainMenuSectionGroup   = 000;
+	static int32 const ComponentsSectionGroup = 100;
+
+	FBlueprintActionFilter ComponentsFilter;
+	ComponentsFilter.Context = Context;
+	ComponentsFilter.NodeTypes.Add(UK2Node_ComponentBoundEvent::StaticClass());
+	ComponentsFilter.NodeTypes.Add(UK2Node_CallFunction::StaticClass());
+	ComponentsFilter.AddIsFilteredTest(FBlueprintActionFilter::FIsFilteredDelegate::CreateStatic(IsUnBoundSpawner));
 	
-	for (UProperty* SelectedProperty : SelectedProperties)
+	bool bHasComponentsSection = false;
+	for (UObject* Selection : Context.SelectedObjects)
 	{
-		static int32 const ComponentsSectionGroup = 100;
-		
-		if (UObjectProperty* ObjProperty = Cast<UObjectProperty>(SelectedProperty))
+		if (UObjectProperty* ObjProperty = Cast<UObjectProperty>(Selection))
 		{
-			MenuFilter.OwnerClasses.Add(ObjProperty->PropertyClass);
-			FText const PropertyName = FText::FromName(ObjProperty->GetFName());
-			
-			MenuFilter.NodeTypes.Add(UK2Node_CallFunction::StaticClass());
-			FText const FuncSectionName = FText::Format(LOCTEXT("ComponentFuncSection", "Call Function on {0}"), PropertyName);
-			MenuOut.AddMenuSection(MenuFilter, FuncSectionName, ComponentsSectionGroup);
-			
-			MenuFilter.NodeTypes[0] = UK2Node_Event::StaticClass();
-			FText const EventSectionName = FText::Format(LOCTEXT("ComponentEventSection", "Add Event for {0}"), PropertyName);
-			MenuOut.AddMenuSection(MenuFilter, EventSectionName, ComponentsSectionGroup);
-			
-			MenuFilter.NodeTypes.Empty();
+			ComponentsFilter.OwnerClasses.Add(ObjProperty->PropertyClass);
+			bHasComponentsSection = true;
 		}
 	}
-	
-	for (UBlueprint* Blueprint : Context.Blueprints)
+
+	MenuOut.Empty();
+
+	FBlueprintActionFilter MenuFilter;
+	MenuFilter.Context = Context;
+	MenuFilter.Context.SelectedObjects.Empty();
+
+	if (bIsContextSensitive)
 	{
-		MenuFilter.OwnerClasses.Empty();
-		MenuFilter.OwnerClasses.Add(Blueprint->GeneratedClass);
+		if (bHasComponentsSection)
+		{
+			MenuOut.AddMenuSection(ComponentsFilter, FText::GetEmpty(), ComponentsSectionGroup);
+		}
 		
-		MenuOut.AddMenuSection(MenuFilter);
+		for (UBlueprint* Blueprint : Context.Blueprints)
+		{
+			MenuFilter.OwnerClasses.Add(Blueprint->GeneratedClass);
+		}
 	}
-	
+
+	MenuOut.AddMenuSection(MenuFilter);
 	MenuOut.RebuildActionList();
+
+	for (UEdGraph const* Graph : Context.Graphs)
+	{
+		if (FKismetEditorUtilities::CanPasteNodes(Graph))
+		{
+			// @TODO: Grey out menu option with tooltip if one of the nodes cannot paste into this graph
+			TSharedPtr<FEdGraphSchemaAction> PasteHereAction(new FEdGraphSchemaAction_K2PasteHere(TEXT(""), LOCTEXT("PasteHereMenuName", "Paste here"), TEXT(""), MainMenuSectionGroup));
+			MenuOut.AddAction(PasteHereAction);
+			break;
+		}
+	}
+
+	if (bIsContextSensitive && (Context.SelectedObjects.Num() == 0))
+	{
+		FText SelectComponentMsg     = LOCTEXT("SelectComponentForEvents",        "Select a Component to see available Events & Functions");
+		FText SelectComponentToolTip = LOCTEXT("SelectComponentForEventsTooltip", "Select a Component in the MyBlueprint tab to see available Events and Functions in this menu.");
+		TSharedPtr<FEdGraphSchemaAction> MsgAction = TSharedPtr<FEdGraphSchemaAction>(new FEdGraphSchemaAction_Dummy(TEXT(""), SelectComponentMsg, SelectComponentToolTip.ToString(), ComponentsSectionGroup));
+		MenuOut.AddAction(MsgAction);
+	}
 }
 
 #undef LOCTEXT_NAMESPACE
diff --git a/Engine/Source/Editor/Kismet/Private/SBlueprintActionMenu.cpp b/Engine/Source/Editor/Kismet/Private/SBlueprintActionMenu.cpp
--- a/Engine/Source/Editor/Kismet/Private/SBlueprintActionMenu.cpp
+++ b/Engine/Source/Editor/Kismet/Private/SBlueprintActionMenu.cpp
@@ -401,7 +401,6 @@ void SBlueprintActionMenu::CollectAllActions(FGraphActionListBuilderBase& OutAll
 	// with a valid node, then fix it up in filtering
 	FilterContext.Graphs.Add(GraphObj);
 	
-	TArray<UProperty*> SelectedProperties;
 	if (bIsContextSensitive)
 	{
 		FilterContext.Pins = DraggedFromPins;
@@ -409,7 +408,7 @@ void SBlueprintActionMenu::CollectAllActions(FGraphActionListBuilderBase& OutAll
 		FEdGraphSchemaAction_K2Var* SelectedVar = BlueprintEditor->GetMyBlueprintWidget()->SelectionAsVar();
 		if ((SelectedVar != nullptr) && (SelectedVar->GetProperty() != nullptr))
 		{
-			SelectedProperties.Add(SelectedVar->GetProperty());
+			FilterContext.SelectedObjects.Add(SelectedVar->GetProperty());
 		}
 	}
 	
@@ -417,7 +416,7 @@ void SBlueprintActionMenu::CollectAllActions(FGraphActionListBuilderBase& OutAll
 	// NOTE: cannot call GetGraphContextActions() during serialization and GC due to its use of FindObject()
 	if(!GIsSavingPackage && !GIsGarbageCollecting)
 	{
-		FBlueprintActionMenuUtils::MakeContextMenu(FilterContext, SelectedProperties, MenuBuilder);
+		FBlueprintActionMenuUtils::MakeContextMenu(FilterContext, bIsContextSensitive, MenuBuilder);
 	}
 	// copy the added options back to the main list
 	OutAllActions.Append(MenuBuilder); // @TODO: Avoid this copy
diff --git a/Engine/Source/Editor/Kismet/Public/BlueprintActionMenuUtils.h b/Engine/Source/Editor/Kismet/Public/BlueprintActionMenuUtils.h
--- a/Engine/Source/Editor/Kismet/Public/BlueprintActionMenuUtils.h
+++ b/Engine/Source/Editor/Kismet/Public/BlueprintActionMenuUtils.h
@@ -28,6 +28,6 @@ struct FBlueprintActionMenuUtils
 	 * @param  SelectedProperties	A set of selected properties to offer contextual menu options for.
 	 * @param  MenuOut				The structure that will be populated with context menu items.
 	 */
-	KISMET_API static void MakeContextMenu(FBlueprintActionContext const& Context, TArray<UProperty*> const& SelectedProperties, FBlueprintActionMenuBuilder& MenuOut);
+	KISMET_API static void MakeContextMenu(FBlueprintActionContext const& Context, bool bIsContextSensitive, FBlueprintActionMenuBuilder& MenuOut);
 };
 
diff --git a/Engine/Source/Editor/UnrealEd/Private/Commandlets/DumpBlueprintsInfoCommandlet.cpp b/Engine/Source/Editor/UnrealEd/Private/Commandlets/DumpBlueprintsInfoCommandlet.cpp
--- a/Engine/Source/Editor/UnrealEd/Private/Commandlets/DumpBlueprintsInfoCommandlet.cpp
+++ b/Engine/Source/Editor/UnrealEd/Private/Commandlets/DumpBlueprintsInfoCommandlet.cpp
@@ -745,7 +745,9 @@ static UBlueprint* DumpBlueprintInfoUtils::MakeTempBlueprint(UClass* ParentClass
 			}
 		}
 
-		FBlueprintEditorUtils::MarkBlueprintAsStructurallyModified(MadeBlueprint);		
+		FBlueprintEditorUtils::MarkBlueprintAsStructurallyModified(MadeBlueprint);
+		MadeBlueprint->SetFlags(RF_RootSet); // to keep the BP from being garbage collected
+		FKismetEditorUtilities::CompileBlueprint(MadeBlueprint);
 		ClassBlueprints.Add(ParentClass, MadeBlueprint);
 	}
 
@@ -914,9 +916,10 @@ static FString DumpBlueprintInfoUtils::GetActionKey(FGraphActionListBuilderBase:
 //------------------------------------------------------------------------------
 static void DumpBlueprintInfoUtils::GetComponentProperties(UBlueprint* Blueprint, TArray<UObjectProperty*>& PropertiesOut)
 {
-	if (AActor* BlueprintCDO = Cast<AActor>(Blueprint->SkeletonGeneratedClass->GetDefaultObject()))
+	UClass* BpClass = Blueprint->GeneratedClass;
+	if (BpClass->IsChildOf<AActor>())
 	{
-		for (TFieldIterator<UObjectProperty> PropertyIt(Blueprint->SkeletonGeneratedClass, EFieldIteratorFlags::IncludeSuper); PropertyIt; ++PropertyIt)
+		for (TFieldIterator<UObjectProperty> PropertyIt(BpClass, EFieldIteratorFlags::IncludeSuper); PropertyIt; ++PropertyIt)
 		{
 			// SMyBlueprint filters out component variables in SMyBlueprint::CollectAllActions() using CPF_BlueprintVisible/CPF_Parm flags
 			if (PropertyIt->PropertyClass->IsChildOf(UActorComponent::StaticClass()) &&
@@ -1102,6 +1105,14 @@ static void DumpBlueprintInfoUtils::DumpActionList(uint32 Indent, FGraphActionLi
 	{
 		bool operator()(FGraphActionListBuilderBase::ActionGroup const& LHS, FGraphActionListBuilderBase::ActionGroup const& RHS) const
 		{
+			TSharedPtr<FEdGraphSchemaAction> LHSAction = LHS.Actions[0];
+			TSharedPtr<FEdGraphSchemaAction> RHSAction = RHS.Actions[0];
+			
+			if (LHSAction->Grouping != RHSAction->Grouping)
+			{
+				return LHSAction->Grouping > RHSAction->Grouping;
+			}
+			
 			FString LhKey = GetActionKey(LHS);
 			FString RhKey = GetActionKey(RHS);
 			return (LhKey.Compare(RhKey) < 0);
@@ -1534,7 +1545,7 @@ static double DumpBlueprintInfoUtils::GetContextMenuActions(FBlueprintGraphActio
 		{
 			if (UProperty* SelectedProperty = Cast<UObjectProperty>(SelectedObj))
 			{
-				SelectedProperties.Add(SelectedProperty);
+				FilterContext.SelectedObjects.Add(SelectedProperty);
 			}
 		}
 		
@@ -1544,7 +1555,7 @@ static double DumpBlueprintInfoUtils::GetContextMenuActions(FBlueprintGraphActio
 			FBlueprintActionDatabase::Prime();
 			
 			FScopedDurationTimer DurationTimer(MenuBuildDuration);
-			FBlueprintActionMenuUtils::MakeContextMenu(FilterContext, SelectedProperties, MenuBuilder);
+			FBlueprintActionMenuUtils::MakeContextMenu(FilterContext, /*bIsContextSensitive =*/true, MenuBuilder);
 		}
 		ActionBuilder.Append(MenuBuilder);
 	}
