(WIP) Fixing up bound event spawners so that they're updated when a blueprint is recompiled/altered (so that they catch new properties, etc.).

diff --git a/Engine/Source/Editor/BlueprintGraph/Classes/BlueprintBoundNodeSpawner.h b/Engine/Source/Editor/BlueprintGraph/Classes/BlueprintBoundNodeSpawner.h
--- a/Engine/Source/Editor/BlueprintGraph/Classes/BlueprintBoundNodeSpawner.h
+++ b/Engine/Source/Editor/BlueprintGraph/Classes/BlueprintBoundNodeSpawner.h
@@ -9,7 +9,12 @@
 class UActorComponent;
 
 /**
- *
+ * Takes care of spawning arbitrary nodes that are bound to specific UObject 
+ * instances (like: functions being called on a specific object, or events
+ * triggering for specific level actors, etc.). UBlueprintNodeSpawner classes
+ * act as the "action" portion of FBlueprintActionMenuItems. Evolved from 
+ * UK2Node_CallFunctionOnMember and various other pieces of code from 
+ * FK2ActionMenuBuilder.
  */
 UCLASS(Transient)
 class BLUEPRINTGRAPH_API UBlueprintBoundNodeSpawner : public UBlueprintNodeSpawner
@@ -18,12 +23,28 @@ class BLUEPRINTGRAPH_API UBlueprintBoundNodeSpawner : public UBlueprintNodeSpawn
 
 public:
 	/**
-	 * 
+	 * Creates a new UBlueprintBoundNodeSpawner that is charged with spawning 
+	 * nodes (of the specified type) and bind those nodes to the supplied
+	 * object.
+	 *
+	 * @param  NodeClass	The event node type that you want this to spawn.
+	 * @param  BoundObject	The object you want spawned nodes bound to.
+	 * @param  Outer		Optional outer for the new spawner (if left null, the transient package will be used).
+	 * @return A newly allocated instance of this class.
 	 */
 	static UBlueprintBoundNodeSpawner* Create(TSubclassOf<UEdGraphNode> NodeClass, UObject* BoundObject, UObject* Outer = nullptr);
 
 	/**
-	 * 
+	 * Creates a new UBlueprintBoundNodeSpawner that is charged with using the 
+	 * supplied SubSpawner to spawn nodes and bind those nodes to the 
+	 * specified object. This allows us to leverage other UBlueprintNodeSpawner's
+	 * knowledge on how to spawn certain node types, while this class is solely 
+	 * in charge of binding.
+	 *
+	 * @param  SubSpawner	The spawner that will instantiate the node.
+	 * @param  BoundObject	The object you want spawned nodes bound to.
+	 * @param  Outer		Optional outer for the new spawner (if left null, the transient package will be used).
+	 * @return A newly allocated instance of this class.
 	 */
 	static UBlueprintBoundNodeSpawner* Create(UBlueprintNodeSpawner* SubSpawner, UObject* BoundObject, UObject* Outer = nullptr);
 
@@ -34,21 +55,30 @@ public:
 	// End UBlueprintNodeSpawner interface
 	
 	/**
-	 * 
+ 	 * Retrieves the object that this binds to spawned nodes.
+	 * @return The object that this class was initialized with.
 	 */
 	UObject const* GetBoundObject() const;
 	
 	/**
+	 * Retrieves the sub-spawner that this wraps (which is actually in-charge of
+	 * allocating the nodes that this produces). 
 	 *
+	 * @return The sub-spawner that this class was initialized with.
 	 */
 	UBlueprintNodeSpawner const* GetSubSpawner() const;
 
 private:
-	/** */
+	/** This object that this spawner attempts to bind new nodes with. */
 	UPROPERTY()
 	TWeakObjectPtr<UObject> BoundObjPtr;
 
-	/** */
+	/** 
+	 * Instead of allocating nodes itself, this class relies on this sub-spawner  
+	 * to do the spawning, while this class takes care of binding the nodes. 
+	 * This is to keep function/event/etc. node spawning logic contained and 
+	 * un-duplicated.
+	 */
 	UPROPERTY()
 	UBlueprintNodeSpawner* SubSpawner;
 };
diff --git a/Engine/Source/Editor/BlueprintGraph/Private/BlueprintActionDatabase.cpp b/Engine/Source/Editor/BlueprintGraph/Private/BlueprintActionDatabase.cpp
--- a/Engine/Source/Editor/BlueprintGraph/Private/BlueprintActionDatabase.cpp
+++ b/Engine/Source/Editor/BlueprintGraph/Private/BlueprintActionDatabase.cpp
@@ -241,6 +241,19 @@ static UBlueprintNodeSpawner* FBlueprintNodeSpawnerFactory::MakeCommentNodeSpawn
 
 namespace FBlueprintActionDatabaseImpl
 {
+	/** 
+	 * Since the action-database is a singleton, we can track its state in a
+	 * static enum value here (hidden as an implementation detail to keep the 
+	 * API clean).
+	 */
+	enum EDatabaseState
+	{
+		DB_Uninitialized,
+		DB_Initializing,
+		DB_Primed,
+	};
+	static EDatabaseState DatabaseState = DB_Uninitialized;
+
 	/** 
 	 * Wrapper around FBlueprintActionDatabase's FActionList. Manages the 
 	 * RF_RootSet flag for the UBlueprintNodeSpawners (ensures they don't get 
@@ -396,6 +409,8 @@ namespace FBlueprintActionDatabaseImpl
 	 * @param  ActionListOut	The list you want populated with new spawners.
 	 */
 	static void AddAutonomousNodeActions(UClass* const Class, FActionList& ActionListOut);
+
+	static void OnBlueprintChanged(UBlueprint* InBlueprint);
 }
 
 //------------------------------------------------------------------------------
@@ -683,7 +698,7 @@ static void FBlueprintActionDatabaseImpl::AddBlueprintGraphActions(UClass const*
 {
 	using namespace FBlueprintNodeSpawnerFactory; // for MakeMacroNodeSpawner()
 
-	if (UBlueprint const* const Blueprint = Cast<UBlueprint const>(Class->ClassGeneratedBy))
+	if (UBlueprint* const Blueprint = Cast<UBlueprint>(Class->ClassGeneratedBy))
 	{
 		for (auto GraphIt = Blueprint->MacroGraphs.CreateConstIterator(); GraphIt; ++GraphIt)
 		{
@@ -739,18 +754,49 @@ static void FBlueprintActionDatabaseImpl::AddAutonomousNodeActions(UClass* const
 	}
 }
 
+//------------------------------------------------------------------------------
+static void FBlueprintActionDatabaseImpl::OnBlueprintChanged(UBlueprint* Blueprint)
+{
+	FBlueprintActionDatabase& ActionDatabase = FBlueprintActionDatabase::Get();
+	ActionDatabase.RefreshClassActions(Blueprint->SkeletonGeneratedClass);
+	ActionDatabase.RefreshClassActions(Blueprint->GeneratedClass);
+}
+
 /*******************************************************************************
  * FBlueprintActionDatabase
  ******************************************************************************/
 
+//------------------------------------------------------------------------------
+void FBlueprintActionDatabase::Prime()
+{
+	using namespace FBlueprintActionDatabaseImpl; // for DatabaseState
+	if (DatabaseState == DB_Uninitialized)
+	{
+		// set DB_Initializing before Get(), else we end up in an infinite loop
+		DatabaseState = DB_Initializing;
+
+		FBlueprintActionDatabase& DatabaseInst = Get();
+		DatabaseInst.RefreshAll();
+
+		for (TObjectIterator<UBlueprint> BpIt; BpIt; ++BpIt)
+		{
+			BpIt->OnChanged().AddStatic(&FBlueprintActionDatabaseImpl::OnBlueprintChanged);
+		}
+
+		DatabaseState = DB_Primed;
+	}
+}
+
 //------------------------------------------------------------------------------
 FBlueprintActionDatabase& FBlueprintActionDatabase::Get()
 {
+	using namespace FBlueprintActionDatabaseImpl; // for DatabaseState/DB_Uninitialized
+
 	static FBlueprintActionDatabase DatabaseInst;
-	if (DatabaseInst.ClassActions.Num() == 0)
+	if (DatabaseState == DB_Uninitialized)
 	{
 		// prime the database the first time we access it
-		DatabaseInst.RefreshAll();
+		Prime();
 	}
 
 	return DatabaseInst;
@@ -773,12 +819,12 @@ void FBlueprintActionDatabase::RefreshClassActions(UClass* const Class)
 	using namespace FBlueprintActionDatabaseImpl;
 
 	check(Class != nullptr);
-	bool const bIsSkelClass   = FKismetEditorUtilities::IsClassABlueprintSkeleton(Class);
-	bool const bIsReinstClass = Class->HasAnyClassFlags(CLASS_NewerVersionExists);
+	bool const bIsBlueprintClass = Class->HasAnyClassFlags(CLASS_CompiledFromBlueprint);
+	bool const bIsSkelClass      = bIsBlueprintClass && FKismetEditorUtilities::IsClassABlueprintSkeleton(Class);
+	bool const bIsReinstClass    = Class->HasAnyClassFlags(CLASS_NewerVersionExists);
+	bool const bIsIgnorableClass = bIsReinstClass || (bIsBlueprintClass && !bIsSkelClass);
 	
-	// make sure this class is not an intermediate that was part of some
-	// blueprint compile (if it is a SKEL or REINST class, then don't bother)
-	if (!bIsSkelClass && !bIsReinstClass)
+	if (!bIsIgnorableClass)
 	{
 		FBlueprintActionDatabaseImpl::FActionList ClassActionList(ClassActions.FindOrAdd(Class));
 		
diff --git a/Engine/Source/Editor/BlueprintGraph/Private/BlueprintBoundNodeSpawner.cpp b/Engine/Source/Editor/BlueprintGraph/Private/BlueprintBoundNodeSpawner.cpp
--- a/Engine/Source/Editor/BlueprintGraph/Private/BlueprintBoundNodeSpawner.cpp
+++ b/Engine/Source/Editor/BlueprintGraph/Private/BlueprintBoundNodeSpawner.cpp
@@ -16,7 +16,22 @@
 
 namespace BlueprintBoundNodeSpawnerImpl
 {
+	/**
+	 * For UK2Node_AddComponent nodes. Spawns the component with a certain 
+	 * asset (mesh, child-actor, etc.).
+	 * 
+	 * @param  NewNode	The newly spawned node that needs to be bound.
+	 * @param  AssetObj	The asset object you want bound to the node.
+	 */
 	static void BindAddComponentNodeWithAsset(UEdGraphNode* NewNode, UObject* AssetObj);
+
+	/**
+	 * Binds function nodes to specific properties. Results in a call to a 
+	 * function on the supplied property.
+	 * 
+	 * @param  NewNode		The newly spawned node that needs to be bound.
+	 * @param  PropertyObj	The property object you want bound to the node.
+	 */
 	static void BindFunctionNodeWithProperty(UEdGraphNode* NewNode, UObject* PropertyObj);
 };
 
@@ -53,6 +68,9 @@ static void BlueprintBoundNodeSpawnerImpl::BindFunctionNodeWithProperty(UEdGraph
 		
 		UEdGraph* ParentGraph = FuncNode->GetGraph();
 		UK2Node_VariableGet* GetVarNode = CastChecked<UK2Node_VariableGet>(TempNodeSpawner->Invoke(ParentGraph, VarNodePos));
+
+		ParentGraph->Modify();
+		ParentGraph->AddNode(GetVarNode, /*bFromUI =*/false, /*bSelectNewNode =*/false); 
 		
 		UEdGraphPin* LiteralOutput = GetVarNode->GetValuePin();
 		UEdGraphPin* CallSelfInput = FuncNode->FindPin(GetDefault<UEdGraphSchema_K2>()->PN_Self);
diff --git a/Engine/Source/Editor/BlueprintGraph/Private/BlueprintNodeSpawner.cpp b/Engine/Source/Editor/BlueprintGraph/Private/BlueprintNodeSpawner.cpp
--- a/Engine/Source/Editor/BlueprintGraph/Private/BlueprintNodeSpawner.cpp
+++ b/Engine/Source/Editor/BlueprintGraph/Private/BlueprintNodeSpawner.cpp
@@ -128,6 +128,7 @@ UEdGraphNode* UBlueprintNodeSpawner::Invoke(UEdGraph* ParentGraph, FVector2D con
 		{
 			NewNode->SetFlags(RF_Transactional);
 			NewNode->AllocateDefaultPins();
+			NewNode->PostPlacedNewNode();
 		}
 	}
 
diff --git a/Engine/Source/Editor/BlueprintGraph/Private/BlueprintPropertyNodeSpawner.cpp b/Engine/Source/Editor/BlueprintGraph/Private/BlueprintPropertyNodeSpawner.cpp
--- a/Engine/Source/Editor/BlueprintGraph/Private/BlueprintPropertyNodeSpawner.cpp
+++ b/Engine/Source/Editor/BlueprintGraph/Private/BlueprintPropertyNodeSpawner.cpp
@@ -38,11 +38,7 @@ UBlueprintPropertyNodeSpawner::UBlueprintPropertyNodeSpawner(class FPostConstruc
 //------------------------------------------------------------------------------
 UEdGraphNode* UBlueprintPropertyNodeSpawner::Invoke(UEdGraph* ParentGraph, FVector2D const Location) const
 {
-	// if a NodeClass was set, then this should spawn something
-	UEdGraphNode* NewNode = Super::Invoke(ParentGraph, Location);
-	check(Property != nullptr);
-
-	if (NewNode != nullptr)
+	auto PostSpawnSetupLambda = [](UEdGraphNode* NewNode, bool bIsTemplateNode, UProperty const* Property, FCustomizeNodeDelegate UserDelegate)
 	{
 		UBlueprint* Blueprint = FBlueprintEditorUtils::FindBlueprintForNodeChecked(NewNode);
 		UClass* OwnerClass = Property->GetOwnerClass();
@@ -56,15 +52,12 @@ UEdGraphNode* UBlueprintPropertyNodeSpawner::Invoke(UEdGraph* ParentGraph, FVect
 		{
 			DelegateNode->SetFromProperty(Property, bIsSelfContext);
 		}
-	}
-	// else, without a NodeClass explicitly specified, we don't know what kind  
-	// of property node to spawn (for example: a getter versus a setter)... 
-	// we can use a null NodeClass to represent the property itself (an 
-	// accumulation of all the nodes that it COULD spawn), and the ui can then 
-	// provide the user with a sub-menu to pick from (but it is up to some 
-	// external system to interpret/use it that way)
-
-	return NewNode;
+
+		UserDelegate.ExecuteIfBound(NewNode, bIsTemplateNode);
+	};
+
+	FCustomizeNodeDelegate PostSpawnSetupDelegate = FCustomizeNodeDelegate::CreateStatic(PostSpawnSetupLambda, Property, CustomizeNodeDelegate);
+	return Super::Invoke(ParentGraph, Location, PostSpawnSetupDelegate);
 }
 
 //------------------------------------------------------------------------------
diff --git a/Engine/Source/Editor/BlueprintGraph/Public/BlueprintActionDatabase.h b/Engine/Source/Editor/BlueprintGraph/Public/BlueprintActionDatabase.h
--- a/Engine/Source/Editor/BlueprintGraph/Public/BlueprintActionDatabase.h
+++ b/Engine/Source/Editor/BlueprintGraph/Public/BlueprintActionDatabase.h
@@ -25,7 +25,7 @@ public:
 	/**
 	 * Ensures that the database singleton is initialized and populated.
 	 */
-	static void Prime() { Get(); }
+	static void Prime();
 
 	/**
 	 * Getter to access the database singleton. Will populate the database first 
diff --git a/Engine/Source/Editor/Kismet/Private/BlueprintActionMenuItem.cpp b/Engine/Source/Editor/Kismet/Private/BlueprintActionMenuItem.cpp
--- a/Engine/Source/Editor/Kismet/Private/BlueprintActionMenuItem.cpp
+++ b/Engine/Source/Editor/Kismet/Private/BlueprintActionMenuItem.cpp
@@ -98,11 +98,12 @@ UEdGraphNode* FBlueprintActionMenuItem::PerformAction(UEdGraph* ParentGraph, UEd
 	if (ParentGraph->Nodes.Find(SpawnedNode) == INDEX_NONE)
 	{
 		check(SpawnedNode != nullptr);
-		
+
+		// @TODO: Move Modify()/AddNode() into UBlueprintNodeSpawner and pull selection functionality out here
 		ParentGraph->Modify();
-		ParentGraph->AddNode(SpawnedNode, /*bool bFromUI =*/true, bSelectNewNode);
-		
-		SpawnedNode->PostPlacedNewNode();
+		ParentGraph->AddNode(SpawnedNode, /*bFromUI =*/true, bSelectNewNode);
+		// @TODO: if this spawned multiple nodes, then we should be selecting all of them
+
 		SpawnedNode->SnapToGrid(SNodePanel::GetSnapGridSize());
 		
 		if (FromPin != nullptr)
