Warnings and erros about using static arrays with blueprints.
Warnings about using properties without BlueprintReatWrite or BlueprintReadOnly tags.

ttp343139 BP:  PUBLIC CRASH:  ensure(1 == StructProperty->ArrayDim) in FScriptBuilderBase::EmitTermExpr
ttp343016 ensure(1 == StructProperty->ArrayDim) in FScriptBuilderBase::EmitTermExpr

diff --git a/Engine/Source/Editor/BlueprintGraph/Classes/K2Node_BreakStruct.h b/Engine/Source/Editor/BlueprintGraph/Classes/K2Node_BreakStruct.h
--- a/Engine/Source/Editor/BlueprintGraph/Classes/K2Node_BreakStruct.h
+++ b/Engine/Source/Editor/BlueprintGraph/Classes/K2Node_BreakStruct.h
@@ -10,6 +10,7 @@ class UK2Node_BreakStruct : public UK2Node_StructMemberGet
 	GENERATED_UCLASS_BODY()
 
 	static bool CanBeBroken(const UScriptStruct* Struct);
+	static bool CanCreatePinForProperty(const UProperty* Property);
 
 	// Begin UEdGraphNode interface
 	virtual void AllocateDefaultPins() override;
diff --git a/Engine/Source/Editor/BlueprintGraph/Private/K2Node_BreakStruct.cpp b/Engine/Source/Editor/BlueprintGraph/Private/K2Node_BreakStruct.cpp
--- a/Engine/Source/Editor/BlueprintGraph/Private/K2Node_BreakStruct.cpp
+++ b/Engine/Source/Editor/BlueprintGraph/Private/K2Node_BreakStruct.cpp
@@ -124,6 +124,17 @@ UK2Node_BreakStruct::UK2Node_BreakStruct(const class FPostConstructInitializePro
 {
 }
 
+bool UK2Node_BreakStruct::CanCreatePinForProperty(const UProperty* Property)
+{
+	const UEdGraphSchema_K2* Schema = GetDefault<UEdGraphSchema_K2>();
+	FEdGraphPinType DumbGraphPinType;
+	const bool bConvertable = Schema->ConvertPropertyToPinType(Property, /*out*/ DumbGraphPinType);
+
+	//TODO: remove CPF_Edit in a next release.
+	const bool bVisible = Property ? Property->HasAnyPropertyFlags(CPF_Edit | CPF_BlueprintVisible) : false;
+	return bVisible && bConvertable;
+}
+
 bool UK2Node_BreakStruct::CanBeBroken(const UScriptStruct* Struct)
 {
 	const UEdGraphSchema_K2* Schema = GetDefault<UEdGraphSchema_K2>();
@@ -131,15 +142,9 @@ bool UK2Node_BreakStruct::CanBeBroken(const UScriptStruct* Struct)
 	{
 		for (TFieldIterator<UProperty> It(Struct); It; ++It)
 		{
-			if(const UProperty* Property = *It)
+			if (CanCreatePinForProperty(*It))
 			{
-				FEdGraphPinType DumbGraphPinType;
-				const bool bConvertable = Schema->ConvertPropertyToPinType(Property, /*out*/ DumbGraphPinType);
-				const bool bVisible = Property->HasAnyPropertyFlags(CPF_BlueprintVisible|CPF_Edit);
-				if(bVisible && bConvertable)
-				{
-					return true;
-				}
+				return true;
 			}
 		}
 	}
@@ -153,7 +158,19 @@ void UK2Node_BreakStruct::AllocateDefaultPins()
 	{
 		CreatePin(EGPD_Input, Schema->PC_Struct, TEXT(""), StructType, false, true, StructType->GetName(), true);
 		
-		UK2Node_StructMemberGet::AllocateDefaultPins();
+		struct FBreakStructPinManager : public FStructOperationOptionalPinManager
+		{
+			virtual bool CanTreatPropertyAsOptional(UProperty* TestProperty) const override
+			{
+				return UK2Node_BreakStruct::CanCreatePinForProperty(TestProperty);
+			}
+		};
+
+		{
+			FBreakStructPinManager OptionalPinManager;
+			OptionalPinManager.RebuildPropertyList(ShowPinForProperties, StructType);
+			OptionalPinManager.CreateVisiblePins(ShowPinForProperties, StructType, EGPD_Output, this);
+		}
 
 		// When struct has a lot of fields, mark their pins as advanced
 		if(Pins.Num() > 5)
@@ -195,6 +212,37 @@ void UK2Node_BreakStruct::ValidateNodeDuringCompilation(class FCompilerResultsLo
 	{
 		MessageLog.Error(*LOCTEXT("NoStruct_Error", "No Struct in @@").ToString(), this);
 	}
+	else
+	{
+		bool bHasAnyBlueprintVisibleProperty = false;
+		for (TFieldIterator<UProperty> It(StructType); It; ++It)
+		{
+			const UProperty* Property = *It;
+			if (CanCreatePinForProperty(Property))
+			{
+				const bool bIsBlueprintVisible = Property->HasAnyPropertyFlags(CPF_BlueprintVisible);
+				bHasAnyBlueprintVisibleProperty |= bIsBlueprintVisible;
+
+				const UEdGraphPin* Pin = Property ? FindPin(Property->GetName()) : NULL;
+				const bool bIsLinked = Pin && Pin->LinkedTo.Num();
+
+				if (!bIsBlueprintVisible && bIsLinked)
+				{
+					MessageLog.Warning(*LOCTEXT("PropertyIsNotBPVisible_Warning", "@@ - the native property is not tagged as BlueprintReadWrite or BlueprintReadOnly, the pin will be removed in a future release.").ToString(), Pin);
+				}
+
+				if ((Property->ArrayDim > 1) && bIsLinked)
+				{
+					MessageLog.Warning(*LOCTEXT("StaticArray_Warning", "@@ - the native property is a static array, which is not supported by blueprints").ToString(), Pin);
+				}
+			}
+		}
+
+		if (!bHasAnyBlueprintVisibleProperty)
+		{
+			MessageLog.Warning(*LOCTEXT("StructHasNoBPVisibleProperties_Warning", "@@ has no property tagged as BlueprintReadWrite or BlueprintReadOnly. The node will be removed in a future release.").ToString(), this);
+		}
+	}
 }
 
 FLinearColor UK2Node_BreakStruct::GetNodeTitleColor() const 
diff --git a/Engine/Source/Editor/BlueprintGraph/Private/K2Node_MakeStruct.cpp b/Engine/Source/Editor/BlueprintGraph/Private/K2Node_MakeStruct.cpp
--- a/Engine/Source/Editor/BlueprintGraph/Private/K2Node_MakeStruct.cpp
+++ b/Engine/Source/Editor/BlueprintGraph/Private/K2Node_MakeStruct.cpp
@@ -116,6 +116,36 @@ void UK2Node_MakeStruct::ValidateNodeDuringCompilation(class FCompilerResultsLog
 	{
 		MessageLog.Error(*LOCTEXT("NoStruct_Error", "No Struct in @@").ToString(), this);
 	}
+	else
+	{
+		bool bHasAnyBlueprintVisibleProperty = false;
+		for (TFieldIterator<UProperty> It(StructType); It; ++It)
+		{
+			const UProperty* Property = *It;
+			if (CanBeExposed(Property))
+			{
+				const bool bIsBlueprintVisible = Property->HasAnyPropertyFlags(CPF_BlueprintVisible);
+				bHasAnyBlueprintVisibleProperty |= bIsBlueprintVisible;
+
+				const UEdGraphPin* Pin = Property ? FindPin(Property->GetName()) : NULL;
+
+				if (!bIsBlueprintVisible)
+				{
+					MessageLog.Warning(*LOCTEXT("PropertyIsNotBPVisible_Warning", "@@ - the native property is not tagged as BlueprintReadWrite, the pin will be removed in a future release.").ToString(), Pin);
+				}
+
+				if (Property->ArrayDim > 1)
+				{
+					MessageLog.Warning(*LOCTEXT("StaticArray_Warning", "@@ - the native property is a static array, which is not supported by blueprints").ToString(), Pin);
+				}
+			}
+		}
+
+		if (!bHasAnyBlueprintVisibleProperty)
+		{
+			MessageLog.Warning(*LOCTEXT("StructHasNoBPVisibleProperties_Warning", "@@ has no property tagged as BlueprintReadWrite. The node will be removed in a future release.").ToString(), this);
+		}
+	}
 }
 
 
@@ -155,6 +185,8 @@ bool UK2Node_MakeStruct::CanBeExposed(const UProperty* Property)
 
 		FEdGraphPinType DumbGraphPinType;
 		const bool bConvertable = Schema->ConvertPropertyToPinType(Property, /*out*/ DumbGraphPinType);
+
+		//TODO: remove CPF_Edit in a future release
 		const bool bVisible = Property->HasAnyPropertyFlags(CPF_BlueprintVisible|CPF_Edit);
 		const bool bBlueprintReadOnly = Property->HasAllPropertyFlags(CPF_BlueprintReadOnly);
 		if(bVisible && bConvertable && !bBlueprintReadOnly)
diff --git a/Engine/Source/Editor/BlueprintGraph/Private/K2Node_Variable.cpp b/Engine/Source/Editor/BlueprintGraph/Private/K2Node_Variable.cpp
--- a/Engine/Source/Editor/BlueprintGraph/Private/K2Node_Variable.cpp
+++ b/Engine/Source/Editor/BlueprintGraph/Private/K2Node_Variable.cpp
@@ -302,6 +302,11 @@ void UK2Node_Variable::ValidateNodeDuringCompilation(class FCompilerResultsLog&
 			MessageLog.Warning(*FString::Printf(*LOCTEXT("VariableDeprecated", "Variable '%s' for @@ was deprecated.  Please update it.").ToString(), *VariableReference.GetMemberName().ToString()), this);
 		}
 	}
+
+	if (VariableProperty && (VariableProperty->ArrayDim > 1))
+	{
+		MessageLog.Warning(*LOCTEXT("StaticArray_Warning", "@@ - the native property is a static array, which is not supported by blueprints").ToString(), this);
+	}
 }
 
 FName UK2Node_Variable::GetPaletteIcon(FLinearColor& ColorOut) const
diff --git a/Engine/Source/Programs/UnrealHeaderTool/Private/HeaderParser.cpp b/Engine/Source/Programs/UnrealHeaderTool/Private/HeaderParser.cpp
--- a/Engine/Source/Programs/UnrealHeaderTool/Private/HeaderParser.cpp
+++ b/Engine/Source/Programs/UnrealHeaderTool/Private/HeaderParser.cpp
@@ -5548,6 +5548,18 @@ Found:
 		FError::Throwf(TEXT("Function '%s': Base implementation of RPCs cannot be in a state. Add a stub outside state scope."), *TopFunction->GetName());
 	}
 
+	if (TopFunction->FunctionFlags & (FUNC_BlueprintCallable | FUNC_BlueprintEvent))
+	{
+		for (TFieldIterator<UProperty> It(TopFunction); It; ++It)
+		{
+			UProperty const* const Param = *It;
+			if (Param->ArrayDim > 1)
+			{
+				FError::Throwf(TEXT("Static array cannot be exposed to blueprint. Function: %s Parameter %s\n"), *TopFunction->GetName(), *Param->GetName());
+			}
+		}
+	}
+
 	// Just declaring a function, so end the nesting.
 	PopNest( AllClasses, NEST_FunctionDeclaration, TypeOfFunction );
 
@@ -5889,6 +5901,12 @@ void FHeaderParser::CompileVariableDeclaration(FClasses& AllClasses, UStruct* St
 				StructBeingBuilt->StructFlags = EStructFlags(StructBeingBuilt->StructFlags | STRUCT_HasInstancedReference);
 			}
 		}
+
+		if (NewProperty->HasAnyPropertyFlags(CPF_BlueprintVisible) && (NewProperty->ArrayDim > 1))
+		{
+			UE_LOG(LogCompile, Warning, TEXT("Static array cannot be exposed to blueprint %s.%s\n"), *Struct->GetName(), *NewProperty->GetName());
+		}
+
 	} while( MatchSymbol(TEXT(",")) );
 
 	// Expect a semicolon.
diff --git a/Engine/Source/Runtime/Engine/Classes/Engine/Scene.h b/Engine/Source/Runtime/Engine/Classes/Engine/Scene.h
--- a/Engine/Source/Runtime/Engine/Classes/Engine/Scene.h
+++ b/Engine/Source/Runtime/Engine/Classes/Engine/Scene.h
@@ -590,7 +590,7 @@ struct FPostProcessSettings
 	class UTexture* LensFlareBokehShape;
 
 	/** RGB defines the lens flare color, A it's position. This is a temporary solution. */
-	UPROPERTY(EditAnywhere, BlueprintReadWrite, Category=LensFlares, meta=(editcondition = "bOverride_LensFlareTints", DisplayName = "Tints"))
+	UPROPERTY(EditAnywhere, Category=LensFlares, meta=(editcondition = "bOverride_LensFlareTints", DisplayName = "Tints"))
 	FLinearColor LensFlareTints[8];
 
 	/** 0..1 0=off/no vignette .. 1=strong vignette */
