Fix numerous places using PinName instead of PinFriendlyName for UI purposes

diff --git a/Engine/Source/Editor/BehaviorTreeEditor/Private/EdGraphSchema_BehaviorTree.cpp b/Engine/Source/Editor/BehaviorTreeEditor/Private/EdGraphSchema_BehaviorTree.cpp
--- a/Engine/Source/Editor/BehaviorTreeEditor/Private/EdGraphSchema_BehaviorTree.cpp
+++ b/Engine/Source/Editor/BehaviorTreeEditor/Private/EdGraphSchema_BehaviorTree.cpp
@@ -458,7 +458,7 @@ void UEdGraphSchema_BehaviorTree::GetBreakLinkToSubMenuActions( class FMenuBuild
 			// Add name of connection if possible
 			FFormatNamedArguments Args;
 			Args.Add( TEXT("NodeTitle"), Title );
-			Args.Add( TEXT("PinName"), FText::FromString( Pin->PinName ) );
+			Args.Add( TEXT("PinName"), Pin->GetDisplayName() );
 			Title = FText::Format( LOCTEXT("BreakDescPin", "{NodeTitle} ({PinName})"), Args );
 		}
 
diff --git a/Engine/Source/Editor/BehaviorTreeEditor/Private/EdGraphSchema_BehaviorTreeDecorator.cpp b/Engine/Source/Editor/BehaviorTreeEditor/Private/EdGraphSchema_BehaviorTreeDecorator.cpp
--- a/Engine/Source/Editor/BehaviorTreeEditor/Private/EdGraphSchema_BehaviorTreeDecorator.cpp
+++ b/Engine/Source/Editor/BehaviorTreeEditor/Private/EdGraphSchema_BehaviorTreeDecorator.cpp
@@ -213,7 +213,7 @@ void UEdGraphSchema_BehaviorTreeDecorator::GetBreakLinkToSubMenuActions( class F
 			// Add name of connection if possible
 			FFormatNamedArguments Args;
 			Args.Add( TEXT("NodeTitle"), Title );
-			Args.Add( TEXT("PinName"), FText::FromString( Pin->PinName ) );
+			Args.Add( TEXT("PinName"), Pin->GetDisplayName() );
 			Title = FText::Format( LOCTEXT("BreakDescPin", "{NodeTitle} ({PinName})"), Args );
 		}
 
diff --git a/Engine/Source/Editor/BlueprintGraph/Private/EdGraphSchema_K2.cpp b/Engine/Source/Editor/BlueprintGraph/Private/EdGraphSchema_K2.cpp
--- a/Engine/Source/Editor/BlueprintGraph/Private/EdGraphSchema_K2.cpp
+++ b/Engine/Source/Editor/BlueprintGraph/Private/EdGraphSchema_K2.cpp
@@ -1143,12 +1143,12 @@ void UEdGraphSchema_K2::GetBreakLinkToSubMenuActions( class FMenuBuilder& MenuBu
 		FString TitleString = Title.ToString();
 		if ( Pin->PinName != TEXT("") )
 		{
-			TitleString = FString::Printf(TEXT("%s (%s)"), *TitleString, *Pin->PinName);
+			TitleString = FString::Printf(TEXT("%s (%s)"), *TitleString, *Pin->GetDisplayName().ToString());
 
 			// Add name of connection if possible
 			FFormatNamedArguments Args;
 			Args.Add( TEXT("NodeTitle"), Title );
-			Args.Add( TEXT("PinName"), FText::FromString( Pin->PinName ) );
+			Args.Add( TEXT("PinName"), Pin->GetDisplayName() );
 			Title = FText::Format( LOCTEXT("BreakDescPin", "{NodeTitle} ({PinName})"), Args );
 		}
 
@@ -1186,12 +1186,12 @@ void UEdGraphSchema_K2::GetJumpToConnectionSubMenuActions( class FMenuBuilder& M
 		FString TitleString = Title.ToString();
 		if ( PinLink->PinName != TEXT("") )
 		{
-			TitleString = FString::Printf(TEXT("%s (%s)"), *TitleString, *PinLink->PinName);
+			TitleString = FString::Printf(TEXT("%s (%s)"), *TitleString, *PinLink->GetDisplayName().ToString());
 
 			// Add name of connection if possible
 			FFormatNamedArguments Args;
 			Args.Add( TEXT("NodeTitle"), Title );
-			Args.Add( TEXT("PinName"), FText::FromString( PinLink->PinName ) );
+			Args.Add( TEXT("PinName"), PinLink->GetDisplayName() );
 			Title = FText::Format( LOCTEXT("JumpToDescPin", "{NodeTitle} ({PinName})"), Args );
 		}
 
@@ -4668,20 +4668,16 @@ void UEdGraphSchema_K2::SplitPin(UEdGraphPin* Pin) const
 			{
 				if (Pin->ParentPin)
 				{
-					SubPin->PinFriendlyName = FText::FromString(FString::Printf(TEXT("%s %s")
-																, *Pin->PinFriendlyName.ToString()
-																, (ProtoPin->PinFriendlyName.IsEmpty() ? *ProtoPin->PinName : *ProtoPin->PinFriendlyName.ToString())));
+					SubPin->PinFriendlyName = FText::FromString(FString::Printf(TEXT("%s %s"), *Pin->GetDisplayName().ToString(), *ProtoPin->GetDisplayName().ToString()));
 				}
 				else
 				{
-					SubPin->PinFriendlyName = FText::FromString(ProtoPin->PinFriendlyName.IsEmpty() ? ProtoPin->PinName : ProtoPin->PinFriendlyName.ToString());
+					SubPin->PinFriendlyName = ProtoPin->GetDisplayName();
 				}
 			}
 			else
 			{
-				SubPin->PinFriendlyName = FText::FromString(FString::Printf(TEXT("%s %s")
-                                                            , (Pin->PinFriendlyName.IsEmpty() ? *Pin->PinName : *Pin->PinFriendlyName.ToString())
-                                                            , (ProtoPin->PinFriendlyName.IsEmpty() ? *ProtoPin->PinName : *ProtoPin->PinFriendlyName.ToString())));
+				SubPin->PinFriendlyName = FText::FromString(FString::Printf(TEXT("%s %s"), *Pin->GetDisplayName().ToString(), *ProtoPin->GetDisplayName().ToString()));
 			}
 
 			SubPin->DefaultValue = ProtoPin->DefaultValue;
diff --git a/Engine/Source/Editor/EnvironmentQueryEditor/Private/EdGraphSchema_EnvironmentQuery.cpp b/Engine/Source/Editor/EnvironmentQueryEditor/Private/EdGraphSchema_EnvironmentQuery.cpp
--- a/Engine/Source/Editor/EnvironmentQueryEditor/Private/EdGraphSchema_EnvironmentQuery.cpp
+++ b/Engine/Source/Editor/EnvironmentQueryEditor/Private/EdGraphSchema_EnvironmentQuery.cpp
@@ -310,7 +310,7 @@ void UEdGraphSchema_EnvironmentQuery::GetBreakLinkToSubMenuActions( class FMenuB
 			// Add name of connection if possible
 			FFormatNamedArguments Args;
 			Args.Add( TEXT("NodeTitle"), Title );
-			Args.Add( TEXT("PinName"), FText::FromString( Pin->PinName ) );
+			Args.Add( TEXT("PinName"), Pin->GetDisplayName() );
 			Title = FText::Format( LOCTEXT("BreakDescPin", "{NodeTitle} ({PinName})"), Args );
 		}
 
diff --git a/Engine/Source/Editor/KismetWidgets/Private/SKismetLinearExpression.cpp b/Engine/Source/Editor/KismetWidgets/Private/SKismetLinearExpression.cpp
--- a/Engine/Source/Editor/KismetWidgets/Private/SKismetLinearExpression.cpp
+++ b/Engine/Source/Editor/KismetWidgets/Private/SKismetLinearExpression.cpp
@@ -70,7 +70,7 @@ TSharedRef<SWidget> SKismetLinearExpression::MakeNodeWidget(const UEdGraphNode*
 		// Instead show it as a special sort of variable get
 		FFormatNamedArguments Args;
 		Args.Add(TEXT("NodeTitle"), Node->GetNodeTitle(ENodeTitleType::ListView));
-		Args.Add(TEXT("PinName"), FromPin->PinFriendlyName.IsEmpty() ? FText::FromString(FromPin->PinName) : FromPin->PinFriendlyName);
+		Args.Add(TEXT("PinName"), FromPin->GetDisplayName());
 		const FText EffectiveVariableName = FText::Format(LOCTEXT("NodeTitleWithPinName", "{NodeTitle}_{PinName}"), Args );
 
 		return SNew(SOverlay)
diff --git a/Engine/Source/Editor/UnrealEd/Private/Kismet2/CompilerResultsLog.cpp b/Engine/Source/Editor/UnrealEd/Private/Kismet2/CompilerResultsLog.cpp
--- a/Engine/Source/Editor/UnrealEd/Private/Kismet2/CompilerResultsLog.cpp
+++ b/Engine/Source/Editor/UnrealEd/Private/Kismet2/CompilerResultsLog.cpp
@@ -160,7 +160,7 @@ void FCompilerResultsLog::InternalLogMessage(const EMessageSeverity::Type& Sever
 					}
 					else if (Pin != NULL)
 					{
-						ObjText = Pin->PinFriendlyName.IsEmpty() ? FText::FromString(Pin->PinName) : Pin->PinFriendlyName;
+						ObjText = Pin->GetDisplayName();
 					}
 					else
 					{
diff --git a/Engine/Source/Editor/UnrealEd/Private/MaterialGraphSchema.cpp b/Engine/Source/Editor/UnrealEd/Private/MaterialGraphSchema.cpp
--- a/Engine/Source/Editor/UnrealEd/Private/MaterialGraphSchema.cpp
+++ b/Engine/Source/Editor/UnrealEd/Private/MaterialGraphSchema.cpp
@@ -180,7 +180,7 @@ void UMaterialGraphSchema::GetBreakLinkToSubMenuActions( class FMenuBuilder& Men
 			// Add name of connection if possible
 			FFormatNamedArguments Args;
 			Args.Add( TEXT("NodeTitle"), Title );
-			Args.Add( TEXT("PinName"), FText::FromString( Pin->PinName ) );
+			Args.Add( TEXT("PinName"), Pin->GetDisplayName() );
 			Title = FText::Format( LOCTEXT("BreakDescPin", "{NodeTitle} ({PinName})"), Args );
 		}
 
diff --git a/Engine/Source/Editor/UnrealEd/Private/SoundClassGraphSchema.cpp b/Engine/Source/Editor/UnrealEd/Private/SoundClassGraphSchema.cpp
--- a/Engine/Source/Editor/UnrealEd/Private/SoundClassGraphSchema.cpp
+++ b/Engine/Source/Editor/UnrealEd/Private/SoundClassGraphSchema.cpp
@@ -49,7 +49,7 @@ void USoundClassGraphSchema::GetBreakLinkToSubMenuActions( class FMenuBuilder& M
 			// Add name of connection if possible
 			FFormatNamedArguments Args;
 			Args.Add( TEXT("NodeTitle"), Title );
-			Args.Add( TEXT("PinName"), FText::FromString( Pin->PinName ) );
+			Args.Add( TEXT("PinName"), Pin->GetDisplayName() );
 			Title = FText::Format( LOCTEXT("BreakDescPin", "{NodeTitle} ({PinName})"), Args );
 		}
 
diff --git a/Engine/Source/Runtime/Engine/Classes/EdGraph/EdGraphPin.h b/Engine/Source/Runtime/Engine/Classes/EdGraph/EdGraphPin.h
--- a/Engine/Source/Runtime/Engine/Classes/EdGraph/EdGraphPin.h
+++ b/Engine/Source/Runtime/Engine/Classes/EdGraph/EdGraphPin.h
@@ -281,6 +281,9 @@ public:
 	/** Get the current DefaultObject path name, or DefaultValue if its null */
 	ENGINE_API FString GetDefaultAsString() const;
 
+	/** Returns how the name of the pin should be displayed in the UI */
+	ENGINE_API FText GetDisplayName() const;
+
 	void ResetToDefaults()
 	{
 		check(LinkedTo.Num() == 0);
diff --git a/Engine/Source/Runtime/Engine/Private/EdGraph/EdGraphPin.cpp b/Engine/Source/Runtime/Engine/Private/EdGraph/EdGraphPin.cpp
--- a/Engine/Source/Runtime/Engine/Private/EdGraph/EdGraphPin.cpp
+++ b/Engine/Source/Runtime/Engine/Private/EdGraph/EdGraphPin.cpp
@@ -202,6 +202,11 @@ FString UEdGraphPin::GetDefaultAsString() const
 	}
 }
 
+FText UEdGraphPin::GetDisplayName() const
+{
+	return FText::FromString(GetSchema()->GetPinDisplayName(this));
+}
+
 const FString UEdGraphPin::GetLinkInfoString( const FString& InFunctionName, const FString& InInfoData, const UEdGraphPin* InToPin ) const
 {
 #if WITH_EDITOR
