Fixing issue where stale blueprint menu actions were being primed (after their corresponding had been trashed).

diff --git a/Engine/Source/Editor/BlueprintGraph/Private/BlueprintActionDatabase.cpp b/Engine/Source/Editor/BlueprintGraph/Private/BlueprintActionDatabase.cpp
--- a/Engine/Source/Editor/BlueprintGraph/Private/BlueprintActionDatabase.cpp
+++ b/Engine/Source/Editor/BlueprintGraph/Private/BlueprintActionDatabase.cpp
@@ -748,37 +748,44 @@ void FBlueprintActionDatabase::Tick(float DeltaTime)
 {
 	const double DurationThreshold = FMath::Min(0.003, DeltaTime * 0.01);
 	
-	double CurrentDuration = 0.0;
-	while (!NewActions.IsEmpty())
+	double TotalFrameDuration = 0.0;
+	while (!ActionPrimingQueue.IsEmpty() && (TotalFrameDuration < DurationThreshold))
 	{
+		FActionIndex ActionIndex;
+		ActionPrimingQueue.Dequeue(ActionIndex);
+			
+		TWeakObjectPtr<UClass> ActionClass = ActionIndex.Key;
+		if (ActionClass.IsValid())
 		{
-			FScopedDurationTimer ScopedTimer(CurrentDuration);
-
-			TWeakObjectPtr<UBlueprintNodeSpawner> NewSpawner;
-			NewActions.Dequeue(NewSpawner);
-
-			// @TODO: The weak pointer could still be valid, but the class could  
-			//        have since been trashed (awaiting GC)... need to block on
-			//        spawners that have since been removed from the database
-			if (!NewSpawner.IsStale())
+			// make sure this class is still listed in the database
+			if (FActionList* ClassActionList = ClassActions.Find(ActionClass.Get()))
 			{
-				if (UEdGraphNode* TemplateNode = NewSpawner->GetTemplateNode(nullptr))
+				int32 ActionListIndex = ActionIndex.Value;
+				for (; (ActionListIndex < ClassActionList->Num()) && (TotalFrameDuration < DurationThreshold); ++ActionListIndex)
 				{
+					FScopedDurationTimer ScopedTimer(TotalFrameDuration);
+
+					UEdGraphNode* TemplateNode = (*ClassActionList)[ActionListIndex]->GetTemplateNode();
 					// since we're doing this incrementally, someone could have 
 					// already requested this template, and allocated its pins 
-					if (TemplateNode->Pins.Num() == 0)
+					// (don't need to do redundant work)
+					if ((TemplateNode != nullptr) && (TemplateNode->Pins.Num() == 0))
 					{
 						// in certain scenarios we need pin information from the 
 						// spawner (to help filter by pin context)
 						TemplateNode->AllocateDefaultPins();
 					}
 				}
+
+				// if we still have actions to prime, pop this guy back on the 
+				// queue (with an updated index)
+				if (ActionListIndex < ClassActionList->Num())
+				{
+					ActionIndex.Value = ActionListIndex;
+					ActionPrimingQueue.Enqueue(ActionIndex);
+				}
 			}
 		}
-		if (CurrentDuration > DurationThreshold)
-		{
-			break;
-		}
 	}
 }
 
@@ -838,10 +845,13 @@ void FBlueprintActionDatabase::RefreshClassActions(UClass* const Class)
 		// normally when sifting through fields on all known classes)
 		AddAutonomousNodeActions(Class, ClassActionList);
 
-		// queue all the new actions for priming 
-		for (UBlueprintNodeSpawner* Action : ClassActionList)
+		// queue the newly added actions for priming
+		if (ClassActionList.Num() > 0)
 		{
-			NewActions.Enqueue(Action);
+			FActionIndex ClassQueueEntry;
+			ClassQueueEntry.Key   = Class;
+			ClassQueueEntry.Value = 0;
+			ActionPrimingQueue.Enqueue(ClassQueueEntry);
 		}
 	}
 
diff --git a/Engine/Source/Editor/BlueprintGraph/Public/BlueprintActionDatabase.h b/Engine/Source/Editor/BlueprintGraph/Public/BlueprintActionDatabase.h
--- a/Engine/Source/Editor/BlueprintGraph/Public/BlueprintActionDatabase.h
+++ b/Engine/Source/Editor/BlueprintGraph/Public/BlueprintActionDatabase.h
@@ -88,9 +88,10 @@ private:
 	FClassActionMap ClassActions;
 
 	/** 
-	 * Holds newly allocated actions that need to be "primed". Priming is 
-	 * something we do on Tick() aimed at speeding up performance (like 
-	 * pre-caching the spawner's template-node, etc.).
+	 * References newly allocated actions that need to be "primed". Priming is 
+	 * something we do on Tick() aimed at speeding up performance (like pre-
+	 * caching each spawner's template-node, etc.).
 	 */
-	TQueue<TWeakObjectPtr<UBlueprintNodeSpawner>> NewActions;
+	typedef TPair<TWeakObjectPtr<UClass>, int32> FActionIndex;
+	TQueue<FActionIndex> ActionPrimingQueue;
 };
