Fixing issue with new blueprint menu system, where custom events couldn't be spawned (instead, it was focusing on any existing event).

diff --git a/Engine/Source/Editor/BlueprintGraph/Private/BlueprintEventNodeSpawner.cpp b/Engine/Source/Editor/BlueprintGraph/Private/BlueprintEventNodeSpawner.cpp
--- a/Engine/Source/Editor/BlueprintGraph/Private/BlueprintEventNodeSpawner.cpp
+++ b/Engine/Source/Editor/BlueprintGraph/Private/BlueprintEventNodeSpawner.cpp
@@ -28,17 +28,20 @@ UK2Node_Event* UBlueprintEventNodeSpawnerImpl::FindCustomEventNode(UBlueprint* B
 {
 	UK2Node_Event* FoundNode = nullptr;
 
-	TArray<UK2Node_Event*> AllEvents;
-	FBlueprintEditorUtils::GetAllNodesOfClass<UK2Node_Event>(Blueprint, AllEvents);
-
-	for (UK2Node_Event* EventNode : AllEvents)
+	if (CustomName != NAME_None)
 	{
-		if (EventNode->CustomFunctionName == CustomName)
+		TArray<UK2Node_Event*> AllEvents;
+		FBlueprintEditorUtils::GetAllNodesOfClass<UK2Node_Event>(Blueprint, AllEvents);
+
+		for (UK2Node_Event* EventNode : AllEvents)
 		{
-			FoundNode = EventNode;
-			break;
+			if (EventNode->CustomFunctionName == CustomName)
+			{
+				FoundNode = EventNode;
+				break;
+			}
 		}
-	}
+	}	
 	return FoundNode;
 }
 
