Ensuring that we generate a node-guid for the template-nodes (since custom-event nodes assert on this).

diff --git a/Engine/Source/Editor/BlueprintGraph/Private/BlueprintNodeSpawner.cpp b/Engine/Source/Editor/BlueprintGraph/Private/BlueprintNodeSpawner.cpp
--- a/Engine/Source/Editor/BlueprintGraph/Private/BlueprintNodeSpawner.cpp
+++ b/Engine/Source/Editor/BlueprintGraph/Private/BlueprintNodeSpawner.cpp
@@ -135,6 +135,7 @@ UEdGraphNode* UBlueprintNodeSpawner::Invoke(UEdGraph* ParentGraph, FVector2D con
 	{
 		NewNode = NewObject<UEdGraphNode>(ParentGraph, NodeClass);
 		check(NewNode != nullptr);
+		NewNode->CreateNewGuid();
 
 		// position the node before invoking PostSpawnDelegate (in case it 
 		// wishes to modify this positioning)
@@ -147,8 +148,8 @@ UEdGraphNode* UBlueprintNodeSpawner::Invoke(UEdGraph* ParentGraph, FVector2D con
 		if (!bIsTemplateNode)
 		{
 			NewNode->SetFlags(RF_Transactional);
-			NewNode->AllocateDefaultPins();
 			NewNode->PostPlacedNewNode();
+			NewNode->AllocateDefaultPins();
 		}
 	}
 
