Some missed includes for non-unity builds.

diff --git a/Engine/Source/Editor/BlueprintGraph/Private/BlueprintNodeTemplateCache.cpp b/Engine/Source/Editor/BlueprintGraph/Private/BlueprintNodeTemplateCache.cpp
--- a/Engine/Source/Editor/BlueprintGraph/Private/BlueprintNodeTemplateCache.cpp
+++ b/Engine/Source/Editor/BlueprintGraph/Private/BlueprintNodeTemplateCache.cpp
@@ -2,8 +2,10 @@
 
 #include "BlueprintGraphPrivatePCH.h"
 #include "BlueprintNodeTemplateCache.h"
-#include "BlueprintEditorUtils.h" // for FindFirstNativeClass(), FindBlueprintForGraph()
+#include "BlueprintEditorUtils.h"	// for FindFirstNativeClass(), FindBlueprintForGraph()
 #include "EdGraph/EdGraph.h"
+#include "KismetEditorUtilities.h"	// for CreateBlueprint()
+#include "BlueprintNodeSpawner.h"   // for NodeClass/Invoke()
 
 /*******************************************************************************
  * Static FBlueprintNodeTemplateCache Helpers
