Adding in missing menu icons for the new blueprint action menu systems.

diff --git a/Engine/Source/Editor/BlueprintGraph/Classes/BlueprintComponentNodeSpawner.h b/Engine/Source/Editor/BlueprintGraph/Classes/BlueprintComponentNodeSpawner.h
--- a/Engine/Source/Editor/BlueprintGraph/Classes/BlueprintComponentNodeSpawner.h
+++ b/Engine/Source/Editor/BlueprintGraph/Classes/BlueprintComponentNodeSpawner.h
@@ -34,6 +34,7 @@ public:
 	virtual UEdGraphNode* Invoke(UEdGraph* ParentGraph, FVector2D const Location) const override;
 	virtual FText GetDefaultMenuName() const override;
 	virtual FText GetDefaultMenuCategory() const override;
+	virtual FName GetDefaultMenuIcon(FLinearColor& ColorOut) const override;
 	// End UBlueprintNodeSpawner interface
 	
 	/**
diff --git a/Engine/Source/Editor/BlueprintGraph/Classes/BlueprintPropertyNodeSpawner.h b/Engine/Source/Editor/BlueprintGraph/Classes/BlueprintPropertyNodeSpawner.h
--- a/Engine/Source/Editor/BlueprintGraph/Classes/BlueprintPropertyNodeSpawner.h
+++ b/Engine/Source/Editor/BlueprintGraph/Classes/BlueprintPropertyNodeSpawner.h
@@ -48,7 +48,7 @@ public:
 	virtual UEdGraphNode* Invoke(UEdGraph* ParentGraph, FVector2D const Location) const override;
 	virtual FText GetDefaultMenuName() const override;
 	virtual FText GetDefaultMenuCategory() const override;
-	virtual FName GetDefaultMenuIcon(FLinearColor& ColorOut) override;
+	virtual FName GetDefaultMenuIcon(FLinearColor& ColorOut) const override;
 	// End UBlueprintNodeSpawner interface
 
 	/**
diff --git a/Engine/Source/Editor/BlueprintGraph/Private/BlueprintComponentNodeSpawner.cpp b/Engine/Source/Editor/BlueprintGraph/Private/BlueprintComponentNodeSpawner.cpp
--- a/Engine/Source/Editor/BlueprintGraph/Private/BlueprintComponentNodeSpawner.cpp
+++ b/Engine/Source/Editor/BlueprintGraph/Private/BlueprintComponentNodeSpawner.cpp
@@ -3,6 +3,7 @@
 #include "BlueprintGraphPrivatePCH.h"
 #include "BlueprintComponentNodeSpawner.h"
 #include "K2Node_AddComponent.h"
+#include "ClassIconFinder.h" // for FindIconNameForClass()
 
 #define LOCTEXT_NAMESPACE "BlueprintComponenetNodeSpawner"
 
@@ -119,6 +120,13 @@ FText UBlueprintComponentNodeSpawner::GetDefaultMenuCategory() const
 	return FText::Format(LOCTEXT("ComponentCategory", "Add Component|{0}"), ClassGroup);
 }
 
+//------------------------------------------------------------------------------
+FName UBlueprintComponentNodeSpawner::GetDefaultMenuIcon(FLinearColor& ColorOut) const
+{
+	check(ComponentClass != nullptr);
+	return FClassIconFinder::FindIconNameForClass(ComponentClass);
+}
+
 //------------------------------------------------------------------------------
 TSubclassOf<UActorComponent> UBlueprintComponentNodeSpawner::GetComponentClass() const
 {
diff --git a/Engine/Source/Editor/BlueprintGraph/Private/BlueprintNodeSpawner.cpp b/Engine/Source/Editor/BlueprintGraph/Private/BlueprintNodeSpawner.cpp
--- a/Engine/Source/Editor/BlueprintGraph/Private/BlueprintNodeSpawner.cpp
+++ b/Engine/Source/Editor/BlueprintGraph/Private/BlueprintNodeSpawner.cpp
@@ -111,7 +111,7 @@ FText UBlueprintNodeSpawner::GetDefaultSearchKeywords() const
 }
 
 //------------------------------------------------------------------------------
-FName UBlueprintNodeSpawner::GetDefaultMenuIcon(FLinearColor& ColorOut)
+FName UBlueprintNodeSpawner::GetDefaultMenuIcon(FLinearColor& ColorOut) const
 {
 	ColorOut = FLinearColor::White;
 	// DO NOT make a template node and query it here (the separate ui building
diff --git a/Engine/Source/Editor/BlueprintGraph/Private/BlueprintPropertyNodeSpawner.cpp b/Engine/Source/Editor/BlueprintGraph/Private/BlueprintPropertyNodeSpawner.cpp
--- a/Engine/Source/Editor/BlueprintGraph/Private/BlueprintPropertyNodeSpawner.cpp
+++ b/Engine/Source/Editor/BlueprintGraph/Private/BlueprintPropertyNodeSpawner.cpp
@@ -85,6 +85,12 @@ FText UBlueprintPropertyNodeSpawner::GetDefaultMenuName() const
 		{
 			MenuName = FText::Format(LOCTEXT("SetterMenuName", "Set {0}"), MenuName);
 		}
+		else if (IsDelegateProperty())
+		{
+			// there are bunch of different delegate node type, so let that 
+			// determine the menu entry text
+			MenuName = FText::GetEmpty();
+		}
 	}
 
 	return MenuName;
@@ -108,7 +114,7 @@ FText UBlueprintPropertyNodeSpawner::GetDefaultMenuCategory() const
 }
 
 //------------------------------------------------------------------------------
-FName UBlueprintPropertyNodeSpawner::GetDefaultMenuIcon(FLinearColor& ColorOut)
+FName UBlueprintPropertyNodeSpawner::GetDefaultMenuIcon(FLinearColor& ColorOut) const
 {
 	FName BrushName = Super::GetDefaultMenuIcon(ColorOut);
 	if (NodeClass == nullptr)
diff --git a/Engine/Source/Editor/BlueprintGraph/Public/BlueprintNodeSpawner.h b/Engine/Source/Editor/BlueprintGraph/Public/BlueprintNodeSpawner.h
--- a/Engine/Source/Editor/BlueprintGraph/Public/BlueprintNodeSpawner.h
+++ b/Engine/Source/Editor/BlueprintGraph/Public/BlueprintNodeSpawner.h
@@ -144,7 +144,7 @@ public:
 	 * @param  ColorOut		The color to tint the icon with.
 	 * @return Name of the brush to use (use FEditorStyle::GetBrush() to resolve).
 	 */
-	virtual FName GetDefaultMenuIcon(FLinearColor& ColorOut);
+	virtual FName GetDefaultMenuIcon(FLinearColor& ColorOut) const;
 
 	/**
 	 * Retrieves a cached template for the node that this is set to spawn. Will
diff --git a/Engine/Source/Editor/Kismet/Private/BlueprintActionMenuItem.cpp b/Engine/Source/Editor/Kismet/Private/BlueprintActionMenuItem.cpp
--- a/Engine/Source/Editor/Kismet/Private/BlueprintActionMenuItem.cpp
+++ b/Engine/Source/Editor/Kismet/Private/BlueprintActionMenuItem.cpp
@@ -161,7 +161,7 @@ void FBlueprintActionMenuItem::AddReferencedObjects(FReferenceCollector& Collect
 FSlateBrush const* FBlueprintActionMenuItem::GetMenuIcon(FSlateColor& ColorOut)
 {
 	// if this brush is invalid
-	if ((IconBrush == nullptr) || !IconBrush->HasUObject())
+	if (IconBrush == nullptr)
 	{
 		if (UBlueprintPropertyNodeSpawner const* PropertySpawner = Cast<UBlueprintPropertyNodeSpawner>(Action))
 		{
