#ue4 Made RotationFromXVector and GetRotationXVector default conversions for vector / rotator, by popular request

#github Adapted from GHPR 257 (Rama)

diff --git a/Engine/Source/Editor/BlueprintGraph/Private/EdGraphSchema_K2.cpp b/Engine/Source/Editor/BlueprintGraph/Private/EdGraphSchema_K2.cpp
--- a/Engine/Source/Editor/BlueprintGraph/Private/EdGraphSchema_K2.cpp
+++ b/Engine/Source/Editor/BlueprintGraph/Private/EdGraphSchema_K2.cpp
@@ -1474,10 +1474,14 @@ bool UEdGraphSchema_K2::SearchForAutocastFunction(const UEdGraphPin* OutputPin,
 			{
 				TargetFunction = TEXT("Conv_VectorToTransform");
 			}
-			if ((InputPin->PinType.PinCategory == PC_Struct) && (InputStructType == LinearColorStruct))
+			else if ((InputPin->PinType.PinCategory == PC_Struct) && (InputStructType == LinearColorStruct))
 			{
 				TargetFunction = TEXT("Conv_VectorToLinearColor");
 			}
+			else if ((InputPin->PinType.PinCategory == PC_Struct) && (InputStructType == RotatorStruct))
+			{
+				TargetFunction = TEXT("Conv_VectorToRotator");
+			}
 			else if (InputPin->PinType.PinCategory == PC_String)
 			{
 				TargetFunction = TEXT("Conv_VectorToString");
@@ -1490,6 +1494,10 @@ bool UEdGraphSchema_K2::SearchForAutocastFunction(const UEdGraphPin* OutputPin,
 			{
 				TargetFunction = TEXT("MakeTransform");
 			}
+			else if ((InputPin->PinType.PinCategory == PC_Struct) && (InputStructType == VectorStruct))
+			{
+				TargetFunction = TEXT("Conv_RotatorToVector");
+			}
 			else if (InputPin->PinType.PinCategory == PC_String)
 			{
 				TargetFunction = TEXT("Conv_RotatorToString");
