Change RTTI system in FEdGraphSchemaAction to use FNames instead of FString

diff --git a/Engine/Source/Editor/BehaviorTreeEditor/Private/SBehaviorTreeBlackboardView.cpp b/Engine/Source/Editor/BehaviorTreeEditor/Private/SBehaviorTreeBlackboardView.cpp
--- a/Engine/Source/Editor/BehaviorTreeEditor/Private/SBehaviorTreeBlackboardView.cpp
+++ b/Engine/Source/Editor/BehaviorTreeEditor/Private/SBehaviorTreeBlackboardView.cpp
@@ -28,8 +28,8 @@ namespace EBlackboardSectionTitles
 class FEdGraphSchemaAction_BlackboardEntry : public FEdGraphSchemaAction_Dummy
 {
 public:
-	static FString StaticGetTypeId() { static FString Type = TEXT("FEdGraphSchemaAction_BlackboardEntry"); return Type; }
-	virtual FString GetTypeId() const override { return StaticGetTypeId(); }
+	static FName StaticGetTypeId() { static FName Type("FEdGraphSchemaAction_BlackboardEntry"); return Type; }
+	virtual FName GetTypeId() const override { return StaticGetTypeId(); }
 
 	FEdGraphSchemaAction_BlackboardEntry( UBlackboardData* InBlackboardData, FBlackboardEntry& InKey, bool bInIsInherited )
 		: FEdGraphSchemaAction_Dummy()
diff --git a/Engine/Source/Editor/BlueprintGraph/Classes/EdGraphSchema_K2.h b/Engine/Source/Editor/BlueprintGraph/Classes/EdGraphSchema_K2.h
--- a/Engine/Source/Editor/BlueprintGraph/Classes/EdGraphSchema_K2.h
+++ b/Engine/Source/Editor/BlueprintGraph/Classes/EdGraphSchema_K2.h
@@ -11,8 +11,8 @@ struct BLUEPRINTGRAPH_API FEdGraphSchemaAction_K2Struct : public FEdGraphSchemaA
 	GENERATED_USTRUCT_BODY()
 
 	// Simple type info
-	static FString StaticGetTypeId() {static FString Type = TEXT("FEdGraphSchemaAction_K2Struct"); return Type;}
-	virtual FString GetTypeId() const override { return StaticGetTypeId(); } 
+	static FName StaticGetTypeId() {static FName Type("FEdGraphSchemaAction_K2Struct"); return Type;}
+	virtual FName GetTypeId() const override { return StaticGetTypeId(); } 
 
 	UStruct* Struct;
 
diff --git a/Engine/Source/Editor/BlueprintGraph/Classes/EdGraphSchema_K2_Actions.h b/Engine/Source/Editor/BlueprintGraph/Classes/EdGraphSchema_K2_Actions.h
--- a/Engine/Source/Editor/BlueprintGraph/Classes/EdGraphSchema_K2_Actions.h
+++ b/Engine/Source/Editor/BlueprintGraph/Classes/EdGraphSchema_K2_Actions.h
@@ -32,8 +32,8 @@ struct BLUEPRINTGRAPH_API FEdGraphSchemaAction_K2NewNode : public FEdGraphSchema
 	GENERATED_USTRUCT_BODY()
 
 	// Simple type info
-	static FString StaticGetTypeId() {static FString Type = TEXT("FEdGraphSchemaAction_K2NewNode"); return Type;}
-	virtual FString GetTypeId() const override { return StaticGetTypeId(); } 
+	static FName StaticGetTypeId() {static FName Type("FEdGraphSchemaAction_K2NewNode"); return Type;}
+	virtual FName GetTypeId() const override { return StaticGetTypeId(); } 
 
 	/** Template of node we want to create */
 	UPROPERTY()
@@ -79,8 +79,8 @@ struct BLUEPRINTGRAPH_API FEdGraphSchemaAction_K2ViewNode : public FEdGraphSchem
 	GENERATED_USTRUCT_BODY()
 
 	// Simple type info
-	static FString StaticGetTypeId() {static FString Type = TEXT("FEdGraphSchemaAction_K2ViewNode"); return Type;}
-	virtual FString GetTypeId() const override { return StaticGetTypeId(); } 
+	static FName StaticGetTypeId() {static FName Type("FEdGraphSchemaAction_K2ViewNode"); return Type;}
+	virtual FName GetTypeId() const override { return StaticGetTypeId(); } 
 
 	/** node we want to view */
 	UPROPERTY()
@@ -113,8 +113,8 @@ struct BLUEPRINTGRAPH_API FEdGraphSchemaAction_K2AssignDelegate : public FEdGrap
 {
 	GENERATED_USTRUCT_BODY()
 
-	static FString StaticGetTypeId() {static FString Type = TEXT("FEdGraphSchemaAction_K2AssignDelegate"); return Type;}
-	virtual FString GetTypeId() const override { return StaticGetTypeId(); } 
+	static FName StaticGetTypeId() {static FName Type("FEdGraphSchemaAction_K2AssignDelegate"); return Type;}
+	virtual FName GetTypeId() const override { return StaticGetTypeId(); } 
 
 	FEdGraphSchemaAction_K2AssignDelegate() 
 		:FEdGraphSchemaAction_K2NewNode()
@@ -141,8 +141,8 @@ struct BLUEPRINTGRAPH_API FEdGraphSchemaAction_EventFromFunction : public FEdGra
 {
 	GENERATED_USTRUCT_BODY()
 
-	static FString StaticGetTypeId() {static FString Type = TEXT("FEdGraphSchemaAction_EventFromFunction"); return Type;}
-	virtual FString GetTypeId() const override { return StaticGetTypeId(); } 
+	static FName StaticGetTypeId() {static FName Type("FEdGraphSchemaAction_EventFromFunction"); return Type;}
+	virtual FName GetTypeId() const override { return StaticGetTypeId(); } 
 
 	UPROPERTY()
 	class UFunction* SignatureFunction;
@@ -175,8 +175,8 @@ struct FEdGraphSchemaAction_K2AddComponent : public FEdGraphSchemaAction_K2NewNo
 	GENERATED_USTRUCT_BODY()
 
 	// Simple type info
-	static FString StaticGetTypeId() {static FString Type = TEXT("FEdGraphSchemaAction_K2AddComponent"); return Type;}
-	virtual FString GetTypeId() const override { return StaticGetTypeId(); } 
+	static FName StaticGetTypeId() {static FName Type("FEdGraphSchemaAction_K2AddComponent"); return Type;}
+	virtual FName GetTypeId() const override { return StaticGetTypeId(); } 
 
 	/** Class of component we want to add */
 	UPROPERTY()
@@ -215,8 +215,8 @@ struct FEdGraphSchemaAction_K2AddTimeline : public FEdGraphSchemaAction_K2NewNod
 	GENERATED_USTRUCT_BODY()
 
 	// Simple type info
-	static FString StaticGetTypeId() {static FString Type = TEXT("FEdGraphSchemaAction_K2AddTimeline"); return Type;}
-	virtual FString GetTypeId() const override { return StaticGetTypeId(); } 
+	static FName StaticGetTypeId() {static FName Type("FEdGraphSchemaAction_K2AddTimeline"); return Type;}
+	virtual FName GetTypeId() const override { return StaticGetTypeId(); } 
 
 	FEdGraphSchemaAction_K2AddTimeline()
 		: FEdGraphSchemaAction_K2NewNode()
@@ -242,8 +242,8 @@ struct FEdGraphSchemaAction_K2AddEvent : public FEdGraphSchemaAction_K2NewNode
 	GENERATED_USTRUCT_BODY()
 
 	// Simple type info
-	static FString StaticGetTypeId() {static FString Type = TEXT("FEdGraphSchemaAction_K2AddEvent"); return Type;}
-	virtual FString GetTypeId() const override { return StaticGetTypeId(); } 
+	static FName StaticGetTypeId() {static FName Type("FEdGraphSchemaAction_K2AddEvent"); return Type;}
+	virtual FName GetTypeId() const override { return StaticGetTypeId(); } 
 
 	FEdGraphSchemaAction_K2AddEvent()
 		: FEdGraphSchemaAction_K2NewNode()
@@ -272,8 +272,8 @@ struct FEdGraphSchemaAction_K2AddCustomEvent : public FEdGraphSchemaAction_K2New
 	GENERATED_USTRUCT_BODY()
 
 	// Simple type info
-	static FString StaticGetTypeId() {static FString Type = TEXT("FEdGraphSchemaAction_K2AddCustomEvent"); return Type;}
-	virtual FString GetTypeId() const override { return StaticGetTypeId(); } 
+	static FName StaticGetTypeId() {static FName Type("FEdGraphSchemaAction_K2AddCustomEvent"); return Type;}
+	virtual FName GetTypeId() const override { return StaticGetTypeId(); } 
 
 	FEdGraphSchemaAction_K2AddCustomEvent()
 		: FEdGraphSchemaAction_K2NewNode()
@@ -299,8 +299,8 @@ struct FEdGraphSchemaAction_K2AddCallOnActor : public FEdGraphSchemaAction_K2New
 	GENERATED_USTRUCT_BODY()
 
 	// Simple type info
-	static FString StaticGetTypeId() {static FString Type = TEXT("FEdGraphSchemaAction_K2AddCallOnActor"); return Type;}
-	virtual FString GetTypeId() const override { return StaticGetTypeId(); } 
+	static FName StaticGetTypeId() {static FName Type("FEdGraphSchemaAction_K2AddCallOnActor"); return Type;}
+	virtual FName GetTypeId() const override { return StaticGetTypeId(); } 
 
 	/** Pointer to actors in level we want to call function on */
 	UPROPERTY()
@@ -358,8 +358,8 @@ struct FEdGraphSchemaAction_K2AddComment : public FEdGraphSchemaAction
 	GENERATED_USTRUCT_BODY()
 
 	// Simple type info
-	static FString StaticGetTypeId() {static FString Type = TEXT("FEdGraphSchemaAction_K2AddComment"); return Type;}
-	virtual FString GetTypeId() const override { return StaticGetTypeId(); } 
+	static FName StaticGetTypeId() {static FName Type("FEdGraphSchemaAction_K2AddComment"); return Type;}
+	virtual FName GetTypeId() const override { return StaticGetTypeId(); } 
 
 	FEdGraphSchemaAction_K2AddComment()
 		: FEdGraphSchemaAction()
@@ -393,8 +393,8 @@ struct FEdGraphSchemaAction_K2AddDocumentation : public FEdGraphSchemaAction
 	GENERATED_USTRUCT_BODY()
 
 	// Simple type info
-	static FString StaticGetTypeId() {static FString Type = TEXT("FEdGraphSchemaAction_K2AddDocumentation"); return Type;}
-	virtual FString GetTypeId() const override { return StaticGetTypeId(); } 
+	static FName StaticGetTypeId() {static FName Type("FEdGraphSchemaAction_K2AddDocumentation"); return Type;}
+	virtual FName GetTypeId() const override { return StaticGetTypeId(); } 
 
 	FEdGraphSchemaAction_K2AddDocumentation()
 		: FEdGraphSchemaAction()
@@ -429,8 +429,8 @@ struct BLUEPRINTGRAPH_API FEdGraphSchemaAction_K2TargetNode : public FEdGraphSch
 
 
 	// Simple type info
-	static FString StaticGetTypeId() {static FString Type = TEXT("FEdGraphSchemaAction_K2TargetNode"); return Type;}
-	virtual FString GetTypeId() const override { return StaticGetTypeId(); } 
+	static FName StaticGetTypeId() {static FName Type("FEdGraphSchemaAction_K2TargetNode"); return Type;}
+	virtual FName GetTypeId() const override { return StaticGetTypeId(); } 
 
 	FEdGraphSchemaAction_K2TargetNode()
 		: FEdGraphSchemaAction_K2NewNode()
@@ -456,8 +456,8 @@ struct BLUEPRINTGRAPH_API FEdGraphSchemaAction_K2PasteHere : public FEdGraphSche
 	GENERATED_USTRUCT_BODY()
 
 	// Simple type info
-	static FString StaticGetTypeId() {static FString Type = TEXT("FEdGraphSchemaAction_K2PasteHere"); return Type;}
-	virtual FString GetTypeId() const override { return StaticGetTypeId(); } 
+	static FName StaticGetTypeId() {static FName Type("FEdGraphSchemaAction_K2PasteHere"); return Type;}
+	virtual FName GetTypeId() const override { return StaticGetTypeId(); } 
 
 	FEdGraphSchemaAction_K2PasteHere ()
 		: FEdGraphSchemaAction_K2NewNode()
@@ -483,8 +483,8 @@ struct BLUEPRINTGRAPH_API FEdGraphSchemaAction_K2Enum : public FEdGraphSchemaAct
 	GENERATED_USTRUCT_BODY()
 
 	// Simple type info
-	static FString StaticGetTypeId() {static FString Type = TEXT("FEdGraphSchemaAction_K2Enum"); return Type;}
-	virtual FString GetTypeId() const override { return StaticGetTypeId(); } 
+	static FName StaticGetTypeId() {static FName Type("FEdGraphSchemaAction_K2Enum"); return Type;}
+	virtual FName GetTypeId() const override { return StaticGetTypeId(); } 
 
 	UEnum* Enum;
 
@@ -537,8 +537,8 @@ public:
 	}
 
 	// Simple type info
-	static FString StaticGetTypeId() {static FString Type = TEXT("FEdGraphSchemaAction_K2Var"); return Type;}
-	virtual FString GetTypeId() const override { return StaticGetTypeId(); } 
+	static FName StaticGetTypeId() {static FName Type("FEdGraphSchemaAction_K2Var"); return Type;}
+	virtual FName GetTypeId() const override { return StaticGetTypeId(); } 
 
 	FEdGraphSchemaAction_K2Var() 
 		: FEdGraphSchemaAction()
@@ -595,8 +595,8 @@ public:
 	}
 
 	// Simple type info
-	static FString StaticGetTypeId() {static FString Type = TEXT("FEdGraphSchemaAction_K2LocalVar"); return Type;}
-	virtual FString GetTypeId() const override { return StaticGetTypeId(); } 
+	static FName StaticGetTypeId() {static FName Type("FEdGraphSchemaAction_K2LocalVar"); return Type;}
+	virtual FName GetTypeId() const override { return StaticGetTypeId(); } 
 
 	FEdGraphSchemaAction_K2LocalVar() 
 		: FEdGraphSchemaAction()
@@ -652,8 +652,8 @@ struct BLUEPRINTGRAPH_API FEdGraphSchemaAction_K2Graph : public FEdGraphSchemaAc
 	GENERATED_USTRUCT_BODY()
 
 	// Simple type info
-	static FString StaticGetTypeId() {static FString Type = TEXT("FEdGraphSchemaAction_K2Graph"); return Type;}
-	virtual FString GetTypeId() const override { return StaticGetTypeId(); } 
+	static FName StaticGetTypeId() {static FName Type("FEdGraphSchemaAction_K2Graph"); return Type;}
+	virtual FName GetTypeId() const override { return StaticGetTypeId(); } 
 
 	/** Name of function or class */
 	FName FuncName;
@@ -715,14 +715,14 @@ struct BLUEPRINTGRAPH_API FEdGraphSchemaAction_K2Event : public FEdGraphSchemaAc
 	 * 
 	 * @return A static identifier for actions of this type.
 	 */
-	static FString StaticGetTypeId() 
+	static FName StaticGetTypeId() 
 	{
-		static FString Type = TEXT("FEdGraphSchemaAction_K2Event"); 
+		static FName Type("FEdGraphSchemaAction_K2Event"); 
 		return Type;
 	}
 
 	// FEdGraphSchemaAction interface
-	virtual FString GetTypeId() const override { return StaticGetTypeId(); }
+	virtual FName GetTypeId() const override { return StaticGetTypeId(); }
 	virtual bool IsParentable() const override { return true; }
 	// End of FEdGraphSchemaAction interface
 };
@@ -754,8 +754,8 @@ public:
 	}
 
 	// Simple type info
-	static FString StaticGetTypeId() {static FString Type = TEXT("FEdGraphSchemaAction_K2Delegate"); return Type;}
-	virtual FString GetTypeId() const override { return StaticGetTypeId(); } 
+	static FName StaticGetTypeId() {static FName Type("FEdGraphSchemaAction_K2Delegate"); return Type;}
+	virtual FName GetTypeId() const override { return StaticGetTypeId(); } 
 
 	/** The associated editor graph for this schema */
 	UEdGraph* EdGraph;
diff --git a/Engine/Source/Editor/BlueprintGraph/Private/EdGraphSchema_K2_Actions.cpp b/Engine/Source/Editor/BlueprintGraph/Private/EdGraphSchema_K2_Actions.cpp
--- a/Engine/Source/Editor/BlueprintGraph/Private/EdGraphSchema_K2_Actions.cpp
+++ b/Engine/Source/Editor/BlueprintGraph/Private/EdGraphSchema_K2_Actions.cpp
@@ -36,7 +36,7 @@ UK2Node const* FK2SchemaActionUtils::ExtractNodeTemplateFromAction(TSharedPtr<FE
 	UK2Node const* TemplateNode = NULL;
 	if (PaletteAction.IsValid())
 	{
-		FString const ActionId = PaletteAction->GetTypeId();
+		FName const ActionId = PaletteAction->GetTypeId();
 		// if this action inherits from FEdGraphSchemaAction_K2NewNode
 		if (ActionId == FEdGraphSchemaAction_K2NewNode::StaticGetTypeId() ||
 			ActionId == FEdGraphSchemaAction_K2AssignDelegate::StaticGetTypeId() ||
diff --git a/Engine/Source/Editor/Kismet/Private/BlueprintPaletteFavorites.cpp b/Engine/Source/Editor/Kismet/Private/BlueprintPaletteFavorites.cpp
--- a/Engine/Source/Editor/Kismet/Private/BlueprintPaletteFavorites.cpp
+++ b/Engine/Source/Editor/Kismet/Private/BlueprintPaletteFavorites.cpp
@@ -89,7 +89,7 @@ FFavoritedBlueprintPaletteItem::FFavoritedBlueprintPaletteItem(TSharedPtr<FEdGra
 {
 	if (PaletteActionIn.IsValid())
 	{
-		FString const ActionId = PaletteActionIn->GetTypeId();
+		FName const ActionId = PaletteActionIn->GetTypeId();
 		if (ActionId == FEdGraphSchemaAction_K2AddComponent::StaticGetTypeId())
 		{
 			FEdGraphSchemaAction_K2AddComponent* AddComponentAction = (FEdGraphSchemaAction_K2AddComponent*)PaletteActionIn.Get();
diff --git a/Engine/Source/Editor/Kismet/Private/SBlueprintPalette.cpp b/Engine/Source/Editor/Kismet/Private/SBlueprintPalette.cpp
--- a/Engine/Source/Editor/Kismet/Private/SBlueprintPalette.cpp
+++ b/Engine/Source/Editor/Kismet/Private/SBlueprintPalette.cpp
@@ -916,7 +916,7 @@ END_SLATE_FUNCTION_BUILD_OPTIMIZATION
 //------------------------------------------------------------------------------
 TSharedRef<SWidget> SBlueprintPaletteItem::CreateTextSlotWidget(const FSlateFontInfo& NameFont, FCreateWidgetForActionData* const InCreateData, bool const bIsReadOnlyIn)
 {
-	FString const ActionTypeId = InCreateData->Action->GetTypeId();
+	FName const ActionTypeId = InCreateData->Action->GetTypeId();
 
 	FOnVerifyTextChanged OnVerifyTextChanged;
 	FOnTextCommitted     OnTextCommitted;
diff --git a/Engine/Source/Editor/Kismet/Public/BlueprintActionMenuItem.h b/Engine/Source/Editor/Kismet/Public/BlueprintActionMenuItem.h
--- a/Engine/Source/Editor/Kismet/Public/BlueprintActionMenuItem.h
+++ b/Engine/Source/Editor/Kismet/Public/BlueprintActionMenuItem.h
@@ -24,14 +24,14 @@ struct FBlueprintActionMenuItem : public FEdGraphSchemaAction
 	GENERATED_USTRUCT_BODY();
 	
 public:
-	static FString StaticGetTypeId() { static FString const TypeId = TEXT("FBlueprintActionMenuItem"); return TypeId; }
+	static FName StaticGetTypeId() { static FName const TypeId("FBlueprintActionMenuItem"); return TypeId; }
 	
 	/** Constructors */
 	FBlueprintActionMenuItem() : Action(nullptr), IconBrush(nullptr), IconTint(FLinearColor::White) {}
 	FBlueprintActionMenuItem(UBlueprintNodeSpawner* NodeSpawner, FSlateBrush const* MenuIcon, FSlateColor const& IconTint, int32 MenuGrouping = 0);
 	
 	// FEdGraphSchemaAction interface
-	virtual FString       GetTypeId() const final { return StaticGetTypeId(); }
+	virtual FName         GetTypeId() const final { return StaticGetTypeId(); }
 	virtual UEdGraphNode* PerformAction(UEdGraph* ParentGraph, UEdGraphPin* FromPin, FVector2D const Location, bool bSelectNewNode = true) final;
 	virtual UEdGraphNode* PerformAction(UEdGraph* ParentGraph, TArray<UEdGraphPin*>& FromPins, FVector2D const Location, bool bSelectNewNode = true) final;
 	virtual void          AddReferencedObjects(FReferenceCollector& Collector) final;
diff --git a/Engine/Source/Editor/UnrealEd/Classes/MaterialGraph/MaterialGraphSchema.h b/Engine/Source/Editor/UnrealEd/Classes/MaterialGraph/MaterialGraphSchema.h
--- a/Engine/Source/Editor/UnrealEd/Classes/MaterialGraph/MaterialGraphSchema.h
+++ b/Engine/Source/Editor/UnrealEd/Classes/MaterialGraph/MaterialGraphSchema.h
@@ -15,8 +15,8 @@ struct UNREALED_API FMaterialGraphSchemaAction_NewNode : public FEdGraphSchemaAc
 	class UClass* MaterialExpressionClass;
 
 	// Simple type info
-	static FString StaticGetTypeId() {static FString Type = TEXT("FMaterialGraphSchemaAction_NewNode"); return Type;}
-	virtual FString GetTypeId() const override { return StaticGetTypeId(); } 
+	static FName StaticGetTypeId() {static FName Type("FMaterialGraphSchemaAction_NewNode"); return Type;}
+	virtual FName GetTypeId() const override { return StaticGetTypeId(); } 
 
 	FMaterialGraphSchemaAction_NewNode() 
 		: FEdGraphSchemaAction()
@@ -52,8 +52,8 @@ struct UNREALED_API FMaterialGraphSchemaAction_NewFunctionCall : public FEdGraph
 	FString FunctionPath;
 
 	// Simple type info
-	static FString StaticGetTypeId() {static FString Type = TEXT("FMaterialGraphSchemaAction_NewFunctionCall"); return Type;}
-	virtual FString GetTypeId() const override { return StaticGetTypeId(); } 
+	static FName StaticGetTypeId() {static FName Type("FMaterialGraphSchemaAction_NewFunctionCall"); return Type;}
+	virtual FName GetTypeId() const override { return StaticGetTypeId(); } 
 
 	FMaterialGraphSchemaAction_NewFunctionCall() 
 		: FEdGraphSchemaAction()
@@ -75,8 +75,8 @@ struct UNREALED_API FMaterialGraphSchemaAction_NewComment : public FEdGraphSchem
 	GENERATED_USTRUCT_BODY();
 
 	// Simple type info
-	static FString StaticGetTypeId() {static FString Type = TEXT("FMaterialGraphSchemaAction_NewComment"); return Type;}
-	virtual FString GetTypeId() const override { return StaticGetTypeId(); } 
+	static FName StaticGetTypeId() {static FName Type("FMaterialGraphSchemaAction_NewComment"); return Type;}
+	virtual FName GetTypeId() const override { return StaticGetTypeId(); } 
 
 	FMaterialGraphSchemaAction_NewComment() 
 		: FEdGraphSchemaAction()
@@ -98,8 +98,8 @@ struct UNREALED_API FMaterialGraphSchemaAction_Paste : public FEdGraphSchemaActi
 	GENERATED_USTRUCT_BODY();
 
 	// Simple type info
-	static FString StaticGetTypeId() {static FString Type = TEXT("FMaterialGraphSchemaAction_Paste"); return Type;}
-	virtual FString GetTypeId() const override { return StaticGetTypeId(); } 
+	static FName StaticGetTypeId() {static FName Type("FMaterialGraphSchemaAction_Paste"); return Type;}
+	virtual FName GetTypeId() const override { return StaticGetTypeId(); } 
 
 	FMaterialGraphSchemaAction_Paste() 
 		: FEdGraphSchemaAction()
diff --git a/Engine/Source/Editor/UnrealEd/Classes/SoundClassGraph/SoundClassGraphSchema.h b/Engine/Source/Editor/UnrealEd/Classes/SoundClassGraph/SoundClassGraphSchema.h
--- a/Engine/Source/Editor/UnrealEd/Classes/SoundClassGraph/SoundClassGraphSchema.h
+++ b/Engine/Source/Editor/UnrealEd/Classes/SoundClassGraph/SoundClassGraphSchema.h
@@ -10,7 +10,7 @@ struct UNREALED_API FSoundClassGraphSchemaAction_NewNode : public FEdGraphSchema
 	GENERATED_USTRUCT_BODY();
 
 	// Simple type info
-	static FString StaticGetTypeId() {static FString Type = TEXT("FSoundClassGraphSchemaAction_NewNode"); return Type;}
+	static FName StaticGetTypeId() {static FName Type("FSoundClassGraphSchemaAction_NewNode"); return Type;}
 
 	FSoundClassGraphSchemaAction_NewNode() 
 		: FEdGraphSchemaAction()
@@ -23,7 +23,7 @@ struct UNREALED_API FSoundClassGraphSchemaAction_NewNode : public FEdGraphSchema
 	{}
 
 	// FEdGraphSchemaAction interface
-	virtual FString GetTypeId() const override { return StaticGetTypeId(); } 
+	virtual FName GetTypeId() const override { return StaticGetTypeId(); } 
 	virtual UEdGraphNode* PerformAction(class UEdGraph* ParentGraph, UEdGraphPin* FromPin, const FVector2D Location, bool bSelectNewNode = true) override;
 	// End of FEdGraphSchemaAction interface
 
diff --git a/Engine/Source/Editor/UnrealEd/Private/Commandlets/DumpBlueprintsInfoCommandlet.cpp b/Engine/Source/Editor/UnrealEd/Private/Commandlets/DumpBlueprintsInfoCommandlet.cpp
--- a/Engine/Source/Editor/UnrealEd/Private/Commandlets/DumpBlueprintsInfoCommandlet.cpp
+++ b/Engine/Source/Editor/UnrealEd/Private/Commandlets/DumpBlueprintsInfoCommandlet.cpp
@@ -1176,7 +1176,7 @@ static void DumpBlueprintInfoUtils::DumpActionMenuItem(uint32 Indent, FGraphActi
 		FString const IndentedNewline = "\n" + BuildIndentString(Indent);
 
 		ActionEntry += " : {";
-		ActionEntry += IndentedNewline + "\"ActionType\"  : \"" + PrimeAction->GetTypeId() + "\",";
+		ActionEntry += IndentedNewline + "\"ActionType\"  : \"" + PrimeAction->GetTypeId().ToString() + "\",";
 		ActionEntry += IndentedNewline + "\"Name\"        : \"" + ActionName + "\",";
 		ActionEntry += IndentedNewline + "\"Category\"    : \"";
 		if (bHasCategory)
diff --git a/Engine/Source/Runtime/Engine/Classes/EdGraph/EdGraphSchema.h b/Engine/Source/Runtime/Engine/Classes/EdGraph/EdGraphSchema.h
--- a/Engine/Source/Runtime/Engine/Classes/EdGraph/EdGraphSchema.h
+++ b/Engine/Source/Runtime/Engine/Classes/EdGraph/EdGraphSchema.h
@@ -53,8 +53,8 @@ struct ENGINE_API FEdGraphSchemaAction
 	GENERATED_USTRUCT_BODY()
 
 	// Simple type info
-	static FString StaticGetTypeId() {static FString Type = TEXT("FEdGraphSchemaAction"); return Type;}
-	virtual FString GetTypeId() const { return StaticGetTypeId(); }
+	static FName StaticGetTypeId() {static FName Type("FEdGraphSchemaAction"); return Type;}
+	virtual FName GetTypeId() const { return StaticGetTypeId(); }
 
 	/** The menu text that should be displayed for this node in the creation menu */
 	UPROPERTY()
@@ -149,8 +149,8 @@ struct ENGINE_API FEdGraphSchemaAction_NewNode : public FEdGraphSchemaAction
 	GENERATED_USTRUCT_BODY()
 
 	// Simple type info
-	static FString StaticGetTypeId() {static FString Type = TEXT("FEdGraphSchemaAction_NewNode"); return Type;}
-	virtual FString GetTypeId() const override { return StaticGetTypeId(); } 
+	static FName StaticGetTypeId() {static FName Type("FEdGraphSchemaAction_NewNode"); return Type;}
+	virtual FName GetTypeId() const override { return StaticGetTypeId(); } 
 
 	/** Template of node we want to create */
 	UPROPERTY()
@@ -188,8 +188,8 @@ struct ENGINE_API FEdGraphSchemaAction_NewNode : public FEdGraphSchemaAction
 /** Dummy action, useful for putting messages in the menu */
 struct FEdGraphSchemaAction_Dummy : public FEdGraphSchemaAction
 {
-	static FString StaticGetTypeId() { static FString Type = TEXT("FEdGraphSchemaAction_Dummy"); return Type; }
-	virtual FString GetTypeId() const override{ return StaticGetTypeId(); }
+	static FName StaticGetTypeId() { static FName Type("FEdGraphSchemaAction_Dummy"); return Type; }
+	virtual FName GetTypeId() const override { return StaticGetTypeId(); }
 
 	FEdGraphSchemaAction_Dummy()
 	: FEdGraphSchemaAction()
