Making it so functions/variables from commandlet Blueprints (those in the transient package) get added to the BP action database.

diff --git a/Engine/Source/Editor/BlueprintGraph/Private/BlueprintActionDatabase.cpp b/Engine/Source/Editor/BlueprintGraph/Private/BlueprintActionDatabase.cpp
--- a/Engine/Source/Editor/BlueprintGraph/Private/BlueprintActionDatabase.cpp
+++ b/Engine/Source/Editor/BlueprintGraph/Private/BlueprintActionDatabase.cpp
@@ -812,12 +812,11 @@ void FBlueprintActionDatabase::RefreshClassActions(UClass* const Class)
 	using namespace BlueprintActionDatabaseImpl;
 
 	check(Class != nullptr);
-	bool const bIsBlueprintClass = Class->HasAnyClassFlags(CLASS_CompiledFromBlueprint);
+	bool const bIsBlueprintClass = (Cast<UBlueprintGeneratedClass>(Class) != nullptr);
 	bool const bIsSkelClass      = bIsBlueprintClass && FKismetEditorUtilities::IsClassABlueprintSkeleton(Class);
 	bool const bIsReinstClass    = Class->HasAnyClassFlags(CLASS_NewerVersionExists);
-	bool const bIsTransientClass = (Class->GetOutermost() == GetTransientPackage());// Class->HasAnyFlags(RF_Transient);
-	bool const bIsIgnorableClass = bIsReinstClass || bIsTransientClass || (bIsBlueprintClass && !bIsSkelClass);
-	
+	bool const bIsIgnorableClass = bIsReinstClass || (bIsBlueprintClass && !bIsSkelClass);
+
 	if (!bIsIgnorableClass)
 	{
 		FActionList& ClassActionList = ClassActions.FindOrAdd(Class);
