Property pin checking was backwards for the new BP menu system (also preventing array pins from being marked as compatible with "Target" self pins).

diff --git a/Engine/Source/Editor/BlueprintGraph/Private/BlueprintActionFilter.cpp b/Engine/Source/Editor/BlueprintGraph/Private/BlueprintActionFilter.cpp
--- a/Engine/Source/Editor/BlueprintGraph/Private/BlueprintActionFilter.cpp
+++ b/Engine/Source/Editor/BlueprintGraph/Private/BlueprintActionFilter.cpp
@@ -799,7 +799,7 @@ static bool BlueprintActionFilterImpl::IsPinCompatibleWithTargetSelf(UEdGraphPin
 	if ((Pin->Direction == EGPD_Output) && (OwnerClass != nullptr))
 	{
 		FEdGraphPinType const& PinType = Pin->PinType;
-		if (PinType.PinSubCategoryObject.IsValid()) 
+		if (!PinType.bIsArray && PinType.PinSubCategoryObject.IsValid())
 		{
 			// if PC_Object/PC_Interface
 			if (UClass const* const PinObjClass = Cast<UClass>(PinType.PinSubCategoryObject.Get()))
@@ -896,17 +896,17 @@ static bool BlueprintActionFilterImpl::IsMissmatchedPropertyType(FBlueprintActio
 
 					EEdGraphPinDirection const PinDir = ContextPin->Direction;
 					if ((PinDir == EGPD_Input) && bIsGetter)
-					{
-						FEdGraphPinType InputPinType;
-						K2Schema->ConvertPropertyToPinType(Property, InputPinType);
-						bIsFilteredOut = !K2Schema->ArePinTypesCompatible(ContextPinType, InputPinType);
-					}
-					else if ((PinDir == EGPD_Output) && bIsSetter)
 					{
 						FEdGraphPinType OutputPinType;
 						K2Schema->ConvertPropertyToPinType(Property, OutputPinType);
 						bIsFilteredOut = !K2Schema->ArePinTypesCompatible(OutputPinType, ContextPinType);
-					}					
+					}
+					else if ((PinDir == EGPD_Output) && bIsSetter)
+					{
+						FEdGraphPinType InputPinType;
+						K2Schema->ConvertPropertyToPinType(Property, InputPinType);
+						bIsFilteredOut = !K2Schema->ArePinTypesCompatible(ContextPinType, InputPinType);
+					}	
 				}
 				else
 				{
