(WIP) Setting it up so new assets callback into the new blueprint database to add associated actions.

diff --git a/Engine/Source/Editor/BlueprintGraph/Private/BlueprintActionDatabase.cpp b/Engine/Source/Editor/BlueprintGraph/Private/BlueprintActionDatabase.cpp
--- a/Engine/Source/Editor/BlueprintGraph/Private/BlueprintActionDatabase.cpp
+++ b/Engine/Source/Editor/BlueprintGraph/Private/BlueprintActionDatabase.cpp
@@ -12,6 +12,9 @@
 #include "BlueprintBoundNodeSpawner.h"
 #include "Stats.h"					// for RETURN_QUICK_DECLARE_CYCLE_STAT()
 #include "ScopedTimers.h"			// for FScopedDurationTimer
+#include "CallbackDevice.h"			// for FCoreDelegates::OnAssetLoaded
+#include "ModuleManager.h"
+#include "AssetRegistryModule.h"	// for OnAssetAdded()/OnAssetRemoved()
 
 // used below in FBlueprintNodeSpawnerFactory::MakeMacroNodeSpawner()
 #include "K2Node_MacroInstance.h"
@@ -37,8 +40,9 @@
 // used below in BlueprintActionDatabaseImpl::AddClassCastActions()
 #include "K2Node_DynamicCast.h"
 #include "K2Node_ClassDynamicCast.h"
-// used below in FBlueprintNodeSpawnerFactory::MakeCommentNodeSpawner()
+// used below in BlueprintActionDatabaseImpl::AddAutonomousNodeActions()
 #include "EdGraph/EdGraphNode_Comment.h"
+#include "EdGraph/EdGraphNode_Documentation.h"
 
 /*******************************************************************************
  * FBlueprintNodeSpawnerFactory
@@ -93,12 +97,12 @@ namespace FBlueprintNodeSpawnerFactory
 	 * not a UK2Node then we can't have it create a spawner for itself (using 
 	 * UK2Node's GetMenuActions() method).
 	 *
-	 * @TODO:  Fix it so comment nodes spawned this way will properly position 
-	 *         themselves (FBlueprintActionMenuItem overrides positioning).
+	 * @param  DocNodeType	The node class type that you want the spawner to be responsible for.
 	 *
 	 * @return A new node-spawner, setup to spawn a UEdGraphNode_Comment.
 	 */
-	static UBlueprintNodeSpawner* MakeCommentNodeSpawner();
+	template <class DocNodeType>
+	static UBlueprintNodeSpawner* MakeDocumentationNodeSpawner();
 };
 
 //------------------------------------------------------------------------------
@@ -209,16 +213,17 @@ static UBlueprintNodeSpawner* FBlueprintNodeSpawnerFactory::MakeMessageNodeSpawn
 }
 
 //------------------------------------------------------------------------------
-static UBlueprintNodeSpawner* FBlueprintNodeSpawnerFactory::MakeCommentNodeSpawner()
+template <class DocNodeType>
+static UBlueprintNodeSpawner* FBlueprintNodeSpawnerFactory::MakeDocumentationNodeSpawner()
 {
-	UBlueprintNodeSpawner* NodeSpawner = UBlueprintNodeSpawner::Create(UEdGraphNode_Comment::StaticClass());
+	UBlueprintNodeSpawner* NodeSpawner = UBlueprintNodeSpawner::Create(DocNodeType::StaticClass());
 	check(NodeSpawner != nullptr);
 
 	auto CustomizeMessageNodeLambda = [](UEdGraphNode* NewNode, bool bIsTemplateNode)
 	{
-		UEdGraphNode_Comment* CommentNode = CastChecked<UEdGraphNode_Comment>(NewNode);
+		DocNodeType* DocNode = CastChecked<DocNodeType>(NewNode);
 
-		UEdGraph* OuterGraph = CommentNode->GetGraph();
+		UEdGraph* OuterGraph = NewNode->GetGraph();
 		check(OuterGraph != nullptr);
 		UBlueprint* Blueprint = FBlueprintEditorUtils::FindBlueprintForGraph(OuterGraph);
 		check(Blueprint != nullptr);
@@ -226,10 +231,10 @@ static UBlueprintNodeSpawner* FBlueprintNodeSpawnerFactory::MakeCommentNodeSpawn
 		float const OldNodePosX = NewNode->NodePosX;
 		float const OldNodePosY = NewNode->NodePosY;
 
-		static const float CommentNodePadding = 50.0f;
-		FSlateRect Bounds(OldNodePosX - CommentNodePadding, OldNodePosY - CommentNodePadding, OldNodePosX + CommentNodePadding, OldNodePosY + CommentNodePadding);
-		FKismetEditorUtilities::GetBoundsForSelectedNodes(Blueprint, Bounds, CommentNodePadding);
-		CommentNode->SetBounds(Bounds);
+		static const float DocNodePadding = 50.0f;
+		FSlateRect Bounds(OldNodePosX - DocNodePadding, OldNodePosY - DocNodePadding, OldNodePosX + DocNodePadding, OldNodePosY + DocNodePadding);
+		FKismetEditorUtilities::GetBoundsForSelectedNodes(Blueprint, Bounds, DocNodePadding);
+		DocNode->SetBounds(Bounds);
 	};
 
 	NodeSpawner->CustomizeNodeDelegate = UBlueprintNodeSpawner::FCustomizeNodeDelegate::CreateStatic(CustomizeMessageNodeLambda);
@@ -376,6 +381,30 @@ namespace BlueprintActionDatabaseImpl
 	 * @param  InBlueprint	The blueprint that has been altered.
 	 */
 	static void OnBlueprintChanged(UBlueprint* InBlueprint);
+
+	/**
+	 * Callback to refresh the database when a new object has just been loaded 
+	 * (to catch blueprint classes that weren't in the initial set, etc.) 
+	 * 
+	 * @param  NewObject	The object that was just loaded.
+	 */
+	static void OnAssetLoaded(UObject* NewObject);
+
+	/**
+	 * Callback to refresh the database when a new object has just been created 
+	 * (to catch new blueprint classes that weren't in the initial set, etc.) 
+	 * 
+	 * @param  NewAssetInfo	Data regarding the newly added asset.
+	 */
+	static void OnAssetAdded(FAssetData const& NewAssetInfo);
+
+	/**
+	 * Callback to refresh the database when an object has been delete (to clear 
+	 * any related classes that were stored in the database) 
+	 * 
+	 * @param  AssetInfo	Data regarding the freshly removed asset.
+	 */
+	static void OnAssetRemoved(FAssetData const& AssetInfo);
 }
 
 //------------------------------------------------------------------------------
@@ -679,7 +708,7 @@ static void BlueprintActionDatabaseImpl::AddBlueprintGraphActions(UClass const*
 //------------------------------------------------------------------------------
 static void BlueprintActionDatabaseImpl::AddAutonomousNodeActions(UClass* const Class, FActionList& ActionListOut)
 {
-	using namespace FBlueprintNodeSpawnerFactory; // for MakeCommentNodeSpawner()
+	using namespace FBlueprintNodeSpawnerFactory; // for MakeDocumentationNodeSpawner()
 
 	if (Class->IsChildOf(UK2Node::StaticClass()) && !Class->HasAnyClassFlags(CLASS_Abstract))
 	{
@@ -696,7 +725,12 @@ static void BlueprintActionDatabaseImpl::AddAutonomousNodeActions(UClass* const
 	//        with a better (more generalized) solution.
 	else if (Class == UEdGraphNode_Comment::StaticClass())
 	{
-		ActionListOut.Add(MakeCommentNodeSpawner());
+		ActionListOut.Add(MakeDocumentationNodeSpawner<UEdGraphNode_Comment>());
+	}
+	else if (Class == UEdGraphNode_Documentation::StaticClass())
+	{
+		// @TODO: BOO! (see comment above)
+		ActionListOut.Add(MakeDocumentationNodeSpawner<UEdGraphNode_Documentation>());
 	}
 }
 
@@ -708,6 +742,48 @@ static void BlueprintActionDatabaseImpl::OnBlueprintChanged(UBlueprint* Blueprin
 	ActionDatabase.RefreshClassActions(Blueprint->GeneratedClass);
 }
 
+//------------------------------------------------------------------------------
+static void BlueprintActionDatabaseImpl::OnAssetLoaded(UObject* NewObject)
+{
+	if (UBlueprint* NewBlueprint = Cast<UBlueprint>(NewObject))
+	{
+		OnBlueprintChanged(NewBlueprint);
+		NewBlueprint->OnChanged().AddStatic(&BlueprintActionDatabaseImpl::OnBlueprintChanged);
+	}
+}
+
+//------------------------------------------------------------------------------
+static void BlueprintActionDatabaseImpl::OnAssetAdded(FAssetData const& NewAssetInfo)
+{
+	if (NewAssetInfo.IsAssetLoaded())
+	{
+		if (NewAssetInfo.GetClass()->IsChildOf<UBlueprint>())
+		{
+			UBlueprint* NewBlueprint = CastChecked<UBlueprint>(NewAssetInfo.GetAsset());
+			OnBlueprintChanged(NewBlueprint);
+			NewBlueprint->OnChanged().AddStatic(&BlueprintActionDatabaseImpl::OnBlueprintChanged);
+		}
+	}
+}
+
+//------------------------------------------------------------------------------
+static void BlueprintActionDatabaseImpl::OnAssetRemoved(FAssetData const& AssetInfo)
+{
+	FBlueprintActionDatabase& ActionDatabase = FBlueprintActionDatabase::Get();
+
+	UClass* AssetClass = AssetInfo.GetClass();
+	ActionDatabase.ClearClassActions(AssetClass);
+
+	if (AssetInfo.IsAssetLoaded())
+	{
+		if (AssetClass->IsChildOf<UBlueprint>())
+		{
+			UBlueprint* BlueprintAsset = CastChecked<UBlueprint>(AssetInfo.GetAsset());
+			BlueprintAsset->OnChanged().RemoveStatic(&BlueprintActionDatabaseImpl::OnBlueprintChanged);
+		}
+	}
+}
+
 /*******************************************************************************
  * FBlueprintActionDatabase
  ******************************************************************************/
@@ -732,6 +808,11 @@ FBlueprintActionDatabase::FBlueprintActionDatabase()
 	{
 		BpIt->OnChanged().AddStatic(&BlueprintActionDatabaseImpl::OnBlueprintChanged);
 	}
+	FCoreDelegates::OnAssetLoaded.AddStatic(&BlueprintActionDatabaseImpl::OnAssetLoaded);
+
+	IAssetRegistry& AssetRegistry = FModuleManager::LoadModuleChecked<FAssetRegistryModule>(TEXT("AssetRegistry")).Get();
+	AssetRegistry.OnAssetAdded().AddStatic(&BlueprintActionDatabaseImpl::OnAssetAdded);
+	AssetRegistry.OnAssetRemoved().AddStatic(&BlueprintActionDatabaseImpl::OnAssetRemoved);
 }
 
 //------------------------------------------------------------------------------
@@ -747,20 +828,19 @@ void FBlueprintActionDatabase::AddReferencedObjects(FReferenceCollector& Collect
 void FBlueprintActionDatabase::Tick(float DeltaTime)
 {
 	const double DurationThreshold = FMath::Min(0.003, DeltaTime * 0.01);
-	
+
 	double TotalFrameDuration = 0.0;
-	while (!ActionPrimingQueue.IsEmpty() && (TotalFrameDuration < DurationThreshold))
+	while ((ActionPrimingQueue.Num() > 0) && (TotalFrameDuration < DurationThreshold))
 	{
-		FActionIndex ActionIndex;
-		ActionPrimingQueue.Dequeue(ActionIndex);
+		auto ActionIndex = ActionPrimingQueue.CreateIterator();	
 			
-		TWeakObjectPtr<UClass> ActionClass = ActionIndex.Key;
+		TWeakObjectPtr<UClass> ActionClass = ActionIndex.Key();
 		if (ActionClass.IsValid())
 		{
 			// make sure this class is still listed in the database
 			if (FActionList* ClassActionList = ClassActions.Find(ActionClass.Get()))
 			{
-				int32 ActionListIndex = ActionIndex.Value;
+				int32& ActionListIndex = ActionIndex.Value();
 				for (; (ActionListIndex < ClassActionList->Num()) && (TotalFrameDuration < DurationThreshold); ++ActionListIndex)
 				{
 					FScopedDurationTimer ScopedTimer(TotalFrameDuration);
@@ -777,14 +857,19 @@ void FBlueprintActionDatabase::Tick(float DeltaTime)
 					}
 				}
 
-				// if we still have actions to prime, pop this guy back on the 
-				// queue (with an updated index)
-				if (ActionListIndex < ClassActionList->Num())
+				if (ActionListIndex >= ClassActionList->Num())
 				{
-					ActionIndex.Value = ActionListIndex;
-					ActionPrimingQueue.Enqueue(ActionIndex);
+					ActionPrimingQueue.Remove(ActionClass);
 				}
 			}
+			else
+			{
+				ActionPrimingQueue.Remove(ActionClass);
+			}
+		}
+		else
+		{
+			ActionPrimingQueue.Remove(ActionClass);
 		}
 	}
 }
@@ -847,10 +932,7 @@ void FBlueprintActionDatabase::RefreshClassActions(UClass* const Class)
 		// queue the newly added actions for priming
 		if (ClassActionList.Num() > 0)
 		{
-			FActionIndex ClassQueueEntry;
-			ClassQueueEntry.Key   = Class;
-			ClassQueueEntry.Value = 0;
-			ActionPrimingQueue.Enqueue(ClassQueueEntry);
+			ActionPrimingQueue.Add(Class, 0);
 		}
 	}
 
@@ -869,6 +951,12 @@ void FBlueprintActionDatabase::RefreshClassActions(UClass* const Class)
 	//   GetVariableGettersSettersForClass()
 }
 
+//------------------------------------------------------------------------------
+void FBlueprintActionDatabase::ClearClassActions(UClass* const Class)
+{
+	ClassActions.Remove(Class);
+}
+
 //------------------------------------------------------------------------------
 FBlueprintActionDatabase::FClassActionMap const& FBlueprintActionDatabase::GetAllActions()
 {
diff --git a/Engine/Source/Editor/BlueprintGraph/Public/BlueprintActionDatabase.h b/Engine/Source/Editor/BlueprintGraph/Public/BlueprintActionDatabase.h
--- a/Engine/Source/Editor/BlueprintGraph/Public/BlueprintActionDatabase.h
+++ b/Engine/Source/Editor/BlueprintGraph/Public/BlueprintActionDatabase.h
@@ -58,6 +58,14 @@ public:
 	 * @param  Class	The class entry you want rebuilt.
 	 */
 	void RefreshClassActions(UClass* const Class);
+
+	/**
+	 * Finds the database entry for the specified class and wipes it. The entry 
+	 * won't be rebuilt, unless RefreshClassActions() is explicitly called after.
+	 * 
+	 * @param  Class	The class entry you want emptied.
+	 */
+	void ClearClassActions(UClass* const Class);
 	
 	/**
 	 * Will populate the database first if it hasn't been created yet, and then 
@@ -92,6 +100,5 @@ private:
 	 * something we do on Tick() aimed at speeding up performance (like pre-
 	 * caching each spawner's template-node, etc.).
 	 */
-	typedef TPair<TWeakObjectPtr<UClass>, int32> FActionIndex;
-	TQueue<FActionIndex> ActionPrimingQueue;
+	TMap<TWeakObjectPtr<UClass>, int32> ActionPrimingQueue;
 };
diff --git a/Engine/Source/Editor/UnrealEd/Private/Commandlets/DumpBlueprintsInfoCommandlet.cpp b/Engine/Source/Editor/UnrealEd/Private/Commandlets/DumpBlueprintsInfoCommandlet.cpp
--- a/Engine/Source/Editor/UnrealEd/Private/Commandlets/DumpBlueprintsInfoCommandlet.cpp
+++ b/Engine/Source/Editor/UnrealEd/Private/Commandlets/DumpBlueprintsInfoCommandlet.cpp
@@ -21,6 +21,8 @@
 #include "BlueprintActionDatabase.h"
 #include "BlueprintActionMenuUtils.h"
 #include "K2Node_Composite.h"
+#include "ModuleManager.h"
+#include "AssetToolsModule.h"
 
 DEFINE_LOG_CATEGORY_STATIC(LogBlueprintInfoDump, Log, All);
 
@@ -596,11 +598,13 @@ static AActor* DumpBlueprintInfoUtils::SpawnLevelActor(UClass* ActorClass, bool
 	else if (FKismetEditorUtilities::CanCreateBlueprintOfClass(ActorClass))
 	{
 		UActorFactory* NewFactory = ConstructObject<UActorFactory>(UActorFactoryBlueprint::StaticClass());
+		NewFactory->AddToRoot();
 
 		UBlueprint* ActorTemplate = MakeTempBlueprint(ActorClass);
 		SpawnedActor = FActorFactoryAssetProxy::AddActorForAsset(ActorTemplate,
 			/*SelectActor =*/bSelect,
 			RF_Transient, NewFactory, NAME_None);
+		NewFactory->RemoveFromRoot();
 	}
 	// @TODO: What about non-blueprintable actors (brushes, etc.)... the code below crashes
 // 	else
@@ -646,11 +650,17 @@ static UBlueprint* DumpBlueprintInfoUtils::MakeTempBlueprint(UClass* ParentClass
 		UClass* BlueprintClass = UBlueprint::StaticClass();
 		UClass* GeneratedClass = UBlueprintGeneratedClass::StaticClass();
 		EBlueprintType BlueprintType = BPTYPE_Normal;
+		UFactory* AssetFactory = nullptr;
 
 		if (bIsAnimBlueprint)
 		{
 			BlueprintClass = UAnimBlueprint::StaticClass();
 			GeneratedClass = UAnimBlueprintGeneratedClass::StaticClass();
+
+			UAnimBlueprintFactory* BlueprintFactory = ConstructObject<UAnimBlueprintFactory>(UAnimBlueprintFactory::StaticClass());
+			BlueprintFactory->ParentClass   = ParentClass;
+			BlueprintFactory->BlueprintType = BlueprintType;
+			AssetFactory = BlueprintFactory;
 		}
 		else if (bIsLevelBlueprint)
 		{
@@ -666,6 +676,12 @@ static UBlueprint* DumpBlueprintInfoUtils::MakeTempBlueprint(UClass* ParentClass
 				BlueprintOuter = World->GetCurrentLevel();
 			}
 		}
+		else 
+		{
+			UBlueprintFactory* BlueprintFactory = ConstructObject<UBlueprintFactory>(UBlueprintFactory::StaticClass());
+			BlueprintFactory->ParentClass = ParentClass;
+			AssetFactory = BlueprintFactory;
+		}
 		// @TODO: UEditorUtilityBlueprint
 
 		FString const ClassName = ParentClass->GetName();
@@ -673,7 +689,16 @@ static UBlueprint* DumpBlueprintInfoUtils::MakeTempBlueprint(UClass* ParentClass
 		FName   const TempBpName = MakeUniqueObjectName(BlueprintOuter, BlueprintClass, FName(*DesiredName));
 
 		check(FKismetEditorUtilities::CanCreateBlueprintOfClass(ParentClass));
-		MadeBlueprint = FKismetEditorUtilities::CreateBlueprint(ParentClass, BlueprintOuter, TempBpName, BlueprintType, BlueprintClass, GeneratedClass);
+		if (AssetFactory != nullptr)
+		{
+			IAssetTools& AssetTools = FModuleManager::LoadModuleChecked<FAssetToolsModule>("AssetTools").Get();
+			UObject* NewAsset = AssetTools.CreateAsset(TempBpName.ToString(), BlueprintOuter->GetPathName(), BlueprintClass, AssetFactory, FName("DumpBlueprintsInfoCommandlet"));
+			MadeBlueprint = CastChecked<UBlueprint>(NewAsset);
+		}
+		else
+		{
+			MadeBlueprint = FKismetEditorUtilities::CreateBlueprint(ParentClass, BlueprintOuter, TempBpName, BlueprintType, BlueprintClass, GeneratedClass);
+		}
 		
 		// if this is an animation blueprint, then we want anim specific graphs to test as well (if it has an anim graph)...
 		if (bIsAnimBlueprint && (MadeBlueprint->FunctionGraphs.Num() > 0))
@@ -747,7 +772,7 @@ static UBlueprint* DumpBlueprintInfoUtils::MakeTempBlueprint(UClass* ParentClass
 		}
 
 		FBlueprintEditorUtils::MarkBlueprintAsStructurallyModified(MadeBlueprint);
-		MadeBlueprint->SetFlags(RF_RootSet); // to keep the BP from being garbage collected
+		MadeBlueprint->AddToRoot(); // to keep the BP from being garbage collected
 		FKismetEditorUtilities::CompileBlueprint(MadeBlueprint);
 		ClassBlueprints.Add(ParentClass, MadeBlueprint);
 	}
@@ -917,7 +942,7 @@ static FString DumpBlueprintInfoUtils::GetActionKey(FGraphActionListBuilderBase:
 //------------------------------------------------------------------------------
 static void DumpBlueprintInfoUtils::GetComponentProperties(UBlueprint* Blueprint, TArray<UObjectProperty*>& PropertiesOut)
 {
-	UClass* BpClass = Blueprint->GeneratedClass;
+	UClass* BpClass = (Blueprint->SkeletonGeneratedClass != nullptr) ? Blueprint->SkeletonGeneratedClass : Blueprint->ParentClass;
 	if (BpClass->IsChildOf<AActor>())
 	{
 		for (TFieldIterator<UObjectProperty> PropertyIt(BpClass, EFieldIteratorFlags::IncludeSuper); PropertyIt; ++PropertyIt)
@@ -988,10 +1013,7 @@ static double DumpBlueprintInfoUtils::GetPaletteMenuActions(FBlueprintPaletteLis
 		FilterContext.Blueprints.Add(const_cast<UBlueprint*>(PaletteBuilder.Blueprint));
 		
 		FBlueprintActionMenuBuilder MenuBuilder(nullptr);
-		{
-			// prime the database so it's not recorded in our timing capture
-			FBlueprintActionDatabase::Get();
-			
+		{			
 			FScopedDurationTimer DurationTimer(MenuBuildDuration);
 			FBlueprintActionMenuUtils::MakePaletteMenu(FilterContext, PaletteFilter, MenuBuilder);
 		}
@@ -1557,9 +1579,6 @@ static double DumpBlueprintInfoUtils::GetContextMenuActions(FBlueprintGraphActio
 		
 		FBlueprintActionMenuBuilder MenuBuilder(nullptr);
 		{
-			// prime the database so it's not recorded in our timing capture
-			FBlueprintActionDatabase::Get();
-
 			FScopedDurationTimer DurationTimer(MenuBuildDuration);
 			FBlueprintActionMenuUtils::MakeContextMenu(FilterContext, /*bIsContextSensitive =*/true, MenuBuilder);
 		}
@@ -1741,6 +1760,15 @@ UDumpBlueprintsInfoCommandlet::UDumpBlueprintsInfoCommandlet(class FPostConstruc
 //------------------------------------------------------------------------------
 int32 UDumpBlueprintsInfoCommandlet::Main(FString const& Params)
 {
+	bool const bCachedIsRunning     = GIsRunning;
+	bool const bCachedExitRequested = GIsRequestingExit;
+	// priming the FBlueprintActionDatabase requires GIsRequestingExit to be 
+	// true; so that it registers its database entries with the GC system, via 
+	// AddReferencedObjects() (without GIsRequestingExit, the FGCObject 
+	// constructor doesn't register itself).
+	GIsRunning		  = true;
+	GIsRequestingExit = false;
+
 	TArray<FString> Tokens, Switches;
 	ParseCommandLine(*Params, Tokens, Switches);
 
@@ -1752,6 +1780,13 @@ int32 UDumpBlueprintsInfoCommandlet::Main(FString const& Params)
 	bool const bSplitFilesByClass = (CommandOptions.DumpFlags & DumpBlueprintInfoUtils::BPDUMP_FilePerBlueprint) != 0;
 	bool const bDiffGeneratedFile = !CommandOptions.DiffPath.IsEmpty() || !CommandOptions.DiffCommand.IsEmpty();
 
+	bool const bUseLegacyMenu = (CommandOptions.DumpFlags & DumpBlueprintInfoUtils::BPDUMP_UseLegacyMenuBuilder) != 0;
+	if (!bUseLegacyMenu)
+	{
+		// prime the database so it's not recorded in our timing capture
+		FBlueprintActionDatabase::Get();
+	}
+
 	UClass* const SelectedObjType = DumpBlueprintInfoUtils::CommandOptions.SelectedObjectType;
 	// if the user specified that they want a level actor selected during the 
 	// dump, then spawn one and select it (extra blueprint context actions 
@@ -1874,6 +1909,11 @@ int32 UDumpBlueprintsInfoCommandlet::Main(FString const& Params)
 	}
 
 	CloseFileStream(&FileOut);
+
+	// restore the globals that we forcefully overrode
+	GIsRequestingExit = bCachedExitRequested;
+	GIsRunning        = bCachedIsRunning;
+
 	return 0;
 }
 
