UMG - Adding additional options to new node actions to automatically take you to the new node, by default this is disabled, because it isn't nessesary most of the time, but for UMG where nodes can be created outside of the graph editing context it's useful.

diff --git a/Engine/Source/Editor/BlueprintGraph/Classes/EdGraphSchema_K2_Actions.h b/Engine/Source/Editor/BlueprintGraph/Classes/EdGraphSchema_K2_Actions.h
--- a/Engine/Source/Editor/BlueprintGraph/Classes/EdGraphSchema_K2_Actions.h
+++ b/Engine/Source/Editor/BlueprintGraph/Classes/EdGraphSchema_K2_Actions.h
@@ -39,15 +39,19 @@ struct BLUEPRINTGRAPH_API FEdGraphSchemaAction_K2NewNode : public FEdGraphSchema
 	UPROPERTY()
 	class UK2Node* NodeTemplate;
 
+	UPROPERTY()
+	bool bGotoNode;
 
 	FEdGraphSchemaAction_K2NewNode() 
 		: FEdGraphSchemaAction()
 		, NodeTemplate(NULL)
+		, bGotoNode(false)
 	{}
 
 	FEdGraphSchemaAction_K2NewNode(const FString& InNodeCategory, const FText& InMenuDesc, const FString& InToolTip, const int32 InGrouping)
 		: FEdGraphSchemaAction(InNodeCategory, InMenuDesc, InToolTip, InGrouping) 
 		, NodeTemplate(NULL)
+		, bGotoNode(false)
 	{}
 
 	// FEdGraphSchemaAction interface
diff --git a/Engine/Source/Editor/BlueprintGraph/Private/EdGraphSchema_K2_Actions.cpp b/Engine/Source/Editor/BlueprintGraph/Private/EdGraphSchema_K2_Actions.cpp
--- a/Engine/Source/Editor/BlueprintGraph/Private/EdGraphSchema_K2_Actions.cpp
+++ b/Engine/Source/Editor/BlueprintGraph/Private/EdGraphSchema_K2_Actions.cpp
@@ -143,6 +143,12 @@ UEdGraphNode* FEdGraphSchemaAction_K2NewNode::PerformAction(class UEdGraph* Pare
 		// Clear any error messages resulting from placing a node.  They'll be flagged on the next compile
 		K2Node->ErrorMsg.Empty();
 		K2Node->bHasCompilerMessage = false;
+
+		if ( bGotoNode && ResultNode )
+		{
+			// Select existing node
+			FKismetEditorUtilities::BringKismetToFocusAttentionOnObject(ResultNode);
+		}
 	}
 
 	return ResultNode;
