Crash fixed, when dragging context menu, from output pin of interface type.

diff --git a/Engine/Source/Editor/BlueprintGraph/Private/EdGraphSchema_K2.cpp b/Engine/Source/Editor/BlueprintGraph/Private/EdGraphSchema_K2.cpp
--- a/Engine/Source/Editor/BlueprintGraph/Private/EdGraphSchema_K2.cpp
+++ b/Engine/Source/Editor/BlueprintGraph/Private/EdGraphSchema_K2.cpp
@@ -2827,9 +2827,8 @@ bool UEdGraphSchema_K2::ArePinTypesCompatible(const FEdGraphPinType& Output, con
 		{
 			InputClass = CallingContext;
 		}
-		check(InputClass != nullptr);
 
-		return InputClass->ImplementsInterface(InterfaceClass) || InterfaceClass->IsChildOf(InputClass);
+		return InputClass && (InputClass->ImplementsInterface(InterfaceClass) || InterfaceClass->IsChildOf(InputClass));
 	}
 	else if ((Output.PinCategory == PC_Object) && (Input.PinCategory == PC_Interface))
 	{
