(WIP) Adding in delegate assignment actions for the new blueprint menu system.

diff --git a/Engine/Source/Editor/BlueprintGraph/Classes/K2Node_AssignDelegate.h b/Engine/Source/Editor/BlueprintGraph/Classes/K2Node_AssignDelegate.h
new file mode 100644
--- /dev/null
+++ b/Engine/Source/Editor/BlueprintGraph/Classes/K2Node_AssignDelegate.h
@@ -0,0 +1,24 @@
+// Copyright 1998-2014 Epic Games, Inc. All Rights Reserved.
+#pragma once
+
+#include "K2Node_AddDelegate.h"
+#include "K2Node_AssignDelegate.generated.h"
+
+/**
+ * Modeled after FEdGraphSchemaAction_K2AssignDelegate for the newer Blueprint
+ * menu system. Acts simply as a UK2Node_AddDelegate with an attached custom-
+ * event node (spawned in PostPlacedNewNode).
+ */
+UCLASS(MinimalAPI)
+class UK2Node_AssignDelegate : public UK2Node_AddDelegate
+{
+	GENERATED_UCLASS_BODY()
+
+public:
+	// UEdGraphNode interface
+	virtual FText GetNodeTitle(ENodeTitleType::Type TitleType) const override;
+	virtual FText GetMenuCategory() const override;
+	virtual bool IsCompatibleWithGraph(const UEdGraph* TargetGraph) const override;
+	virtual void PostPlacedNewNode() override;
+	// End of UEdGraphNode interface
+};
diff --git a/Engine/Source/Editor/BlueprintGraph/Private/BlueprintActionDatabase.cpp b/Engine/Source/Editor/BlueprintGraph/Private/BlueprintActionDatabase.cpp
--- a/Engine/Source/Editor/BlueprintGraph/Private/BlueprintActionDatabase.cpp
+++ b/Engine/Source/Editor/BlueprintGraph/Private/BlueprintActionDatabase.cpp
@@ -15,6 +15,7 @@
 #include "CallbackDevice.h"			// for FCoreDelegates::OnAssetLoaded
 #include "ModuleManager.h"
 #include "AssetRegistryModule.h"	// for OnAssetAdded()/OnAssetRemoved()
+#include "BlueprintEditorUtils.h"	// for FindBlueprintForGraph()
 
 // used below in FBlueprintNodeSpawnerFactory::MakeMacroNodeSpawner()
 #include "K2Node_MacroInstance.h"
@@ -37,6 +38,7 @@
 #include "K2Node_CallDelegate.h"
 #include "K2Node_RemoveDelegate.h"
 #include "K2Node_ClearDelegate.h"
+#include "K2Node_AssignDelegate.h"
 // used below in BlueprintActionDatabaseImpl::AddClassCastActions()
 #include "K2Node_DynamicCast.h"
 #include "K2Node_ClassDynamicCast.h"
@@ -44,6 +46,8 @@
 #include "EdGraph/EdGraphNode_Comment.h"
 #include "EdGraph/EdGraphNode_Documentation.h"
 
+#define LOCTEXT_NAMESPACE "BlueprintActionDatabase"
+
 /*******************************************************************************
  * FBlueprintNodeSpawnerFactory
  ******************************************************************************/
@@ -103,6 +107,15 @@ namespace FBlueprintNodeSpawnerFactory
 	 */
 	template <class DocNodeType>
 	static UBlueprintNodeSpawner* MakeDocumentationNodeSpawner();
+
+	/**
+	 * Constructs a delegate binding node along with a connected event that is  
+	 * triggered from the specified delegate.
+	 * 
+	 * @param  DelegateProperty	The delegate the spawner will bind to.
+	 * @return A new node-spawner, setup to spawn a UK2Node_AssignDelegate.
+	 */
+	static UBlueprintNodeSpawner* MakeAssignDelegateNodeSpawner(UMulticastDelegateProperty* DelegateProperty);
 };
 
 //------------------------------------------------------------------------------
@@ -242,6 +255,42 @@ static UBlueprintNodeSpawner* FBlueprintNodeSpawnerFactory::MakeDocumentationNod
 	return NodeSpawner;
 }
 
+//------------------------------------------------------------------------------
+static UBlueprintNodeSpawner* FBlueprintNodeSpawnerFactory::MakeAssignDelegateNodeSpawner(UMulticastDelegateProperty* DelegateProperty)
+{
+	// @TODO: it'd be awesome to have both nodes spawned by this available for 
+	//        context pin matching (the delegate inputs and the event outputs)
+	return UBlueprintPropertyNodeSpawner::Create<UK2Node_AssignDelegate>(DelegateProperty);
+// 	// using the event spawner as
+// 	UBlueprintNodeSpawner* EventNodeSpawner = UBlueprintEventNodeSpawner::Create(UK2Node_CustomEvent::StaticClass(), FName());
+// 	check(EventNodeSpawner != nullptr);
+// 
+// 	auto SetupDelegateEventNodeLambda = [](UEdGraphNode* NewNode, bool bIsTemplateNode, TWeakObjectPtr<UMulticastDelegateProperty> DelegatePropPtr)
+// 	{
+// 		UK2Node_CustomEvent* EventNode = CastChecked<UK2Node_CustomEvent>(NewNode);
+// 		if (DelegatePropPtr.IsValid())
+// 		{
+// 			UMulticastDelegateProperty* DelegateProp = DelegatePropPtr.Get();		
+// 			// don't needlessly spend unique blueprint names
+// 			if (!bIsTemplateNode)
+// 			{
+// 				UBlueprint* Blueprint = EventNode->GetBlueprint();
+// 
+// 				FText DelegateName = FText::FromName(DelegateProp->GetFName());
+// 				FText DesiredEventName = FText::Format(LOCTEXT("BindedEventName", "{0}_Event"), DelegateName);
+// 				EventNode->CustomFunctionName = FBlueprintEditorUtils::FindUniqueKismetName(Blueprint, DesiredEventName.ToString());
+// 			}
+// 		}
+// 	};
+// 
+// 	TWeakObjectPtr<UMulticastDelegateProperty> DelegatePropPtr = DelegateProperty;
+// 	EventNodeSpawner->CustomizeNodeDelegate = UBlueprintNodeSpawner::FCustomizeNodeDelegate::CreateStatic(SetupDelegateEventNodeLambda, DelegatePropPtr);
+// 
+// 	// the bound spawner should conform the event node to match the delegate's 
+// 	// signature
+// 	return UBlueprintBoundNodeSpawner::Create(EventNodeSpawner, DelegateProperty);
+}
+
 /*******************************************************************************
  * Static FBlueprintActionDatabase Helpers
  ******************************************************************************/
@@ -477,7 +526,8 @@ static void BlueprintActionDatabaseImpl::AddClassPropertyActions(UClass const* c
 				UBlueprintNodeSpawner* AddSpawner = UBlueprintPropertyNodeSpawner::Create<UK2Node_AddDelegate>(DelegateProperty);
 				ActionListOut.Add(AddSpawner);
 				
-				// @TODO: account for: GetEventDispatcherNodesForClass() - FEdGraphSchemaAction_K2AssignDelegate
+				UBlueprintNodeSpawner* AssignSpawner = MakeAssignDelegateNodeSpawner(DelegateProperty);
+				ActionListOut.Add(AssignSpawner);
 			}
 			
 			if (DelegateProperty->HasAnyPropertyFlags(CPF_BlueprintCallable))
@@ -507,8 +557,6 @@ static void BlueprintActionDatabaseImpl::AddClassPropertyActions(UClass const* c
 			ActionListOut.Add(SetterSpawner);
 		}
 	}
-
-	// @TODO: if blueprint class, loop over function graphs and get local variables
 }
 
 //------------------------------------------------------------------------------
@@ -968,3 +1016,5 @@ FBlueprintActionDatabase::FClassActionMap const& FBlueprintActionDatabase::GetAl
 	}
 	return ClassActions;
 }
+
+#undef LOCTEXT_NAMESPACE
diff --git a/Engine/Source/Editor/BlueprintGraph/Private/BlueprintActionFilter.cpp b/Engine/Source/Editor/BlueprintGraph/Private/BlueprintActionFilter.cpp
--- a/Engine/Source/Editor/BlueprintGraph/Private/BlueprintActionFilter.cpp
+++ b/Engine/Source/Editor/BlueprintGraph/Private/BlueprintActionFilter.cpp
@@ -223,9 +223,10 @@ namespace BlueprintActionFilterImpl
 	 * 
 	 * @param  Pin	
 	 * @param  MemberField	
+	 * @param  MemberNodeSpawner
 	 * @return 
 	 */
-	static bool IsPinCompatibleWithTargetSelf(UEdGraphPin const* Pin, UField const* MemberField);
+	static bool IsPinCompatibleWithTargetSelf(UEdGraphPin const* Pin, UField const* MemberField, UBlueprintNodeSpawner const* MemberNodeSpawner);
 
 	/**
 	 * 
@@ -791,20 +792,43 @@ static bool BlueprintActionFilterImpl::HasMatchingPin(UBlueprintNodeSpawner cons
 }
 
 //------------------------------------------------------------------------------
-static bool BlueprintActionFilterImpl::IsPinCompatibleWithTargetSelf(UEdGraphPin const* Pin, UField const* MemberField)
+static bool BlueprintActionFilterImpl::IsPinCompatibleWithTargetSelf(UEdGraphPin const* Pin, UField const* MemberField, UBlueprintNodeSpawner const* MemberNodeSpawner)
 {
 	bool bIsCompatible = false;
+	checkSlow(GetAssociatedField(MemberNodeSpawner) == MemberField);
 
 	UClass* OwnerClass = Cast<UClass>(MemberField->GetOuter());
 	if ((Pin->Direction == EGPD_Output) && (OwnerClass != nullptr))
 	{
 		FEdGraphPinType const& PinType = Pin->PinType;
-		if (!PinType.bIsArray && PinType.PinSubCategoryObject.IsValid())
+		if (PinType.PinSubCategoryObject.IsValid())
 		{
 			// if PC_Object/PC_Interface
-			if (UClass const* const PinObjClass = Cast<UClass>(PinType.PinSubCategoryObject.Get()))
+			UClass const* const PinObjClass = Cast<UClass>(PinType.PinSubCategoryObject.Get());
+			if ((PinObjClass != nullptr) && PinObjClass->IsChildOf(OwnerClass))
 			{
-				bIsCompatible = PinObjClass->IsChildOf(OwnerClass);
+				bIsCompatible = true;
+				if (PinType.bIsArray)
+				{
+					if (UFunction const* Function = Cast<UFunction>(MemberField))
+					{
+						// @TODO: duplicated from UK2Node_CallFunction::AllowMultipleSelfs()... 
+						//        move into a shared resource (don't want to call 
+						//        AllowMultipleSelfs()  to save on perf)
+						bool const bHasReturnParam = (Function->GetReturnProperty() != nullptr);
+						bool const bIsImpure = (Function->HasAnyFunctionFlags(FUNC_BlueprintPure) == false);
+						bool const bIsLatent = (Function->HasMetaData(FBlueprintMetadata::MD_Latent) != false);
+						bIsCompatible = !bHasReturnParam && bIsImpure && !bIsLatent;
+					}
+					else
+					{
+						UEdGraph* OuterGraph = Pin->GetOwningNode()->GetGraph();
+						if (UK2Node* TemplateNode = Cast<UK2Node>(MemberNodeSpawner->GetTemplateNode(OuterGraph)))
+						{
+							bIsCompatible = TemplateNode->AllowMultipleSelfs(/*bInputAsArray =*/true);
+						}	
+					}
+				}
 			}
 		}
 	}
@@ -828,24 +852,29 @@ static bool BlueprintActionFilterImpl::IsFunctionMissingPinParam(FBlueprintActio
 			FEdGraphPinType const& PinType = ContextPin->PinType;
 			UK2Node const* const K2Node = CastChecked<UK2Node const>(ContextPin->GetOwningNode());
 			EEdGraphPinDirection const PinDir = ContextPin->Direction;
-			
+
 			if (K2Schema->IsExecPin(*ContextPin))
 			{
 				bIsFilteredOut = (bIsEventSpawner && (PinDir == EGPD_Output)) || !IsImpure(BlueprintAction);
 			}
 			else
 			{
-				bIsFilteredOut = true;
-				if (K2Schema->FunctionHasParamOfType(AssociatedFunc, K2Node->GetBlueprint(), PinType, (PinDir == EGPD_Input)))
+				// event nodes have their parameters as outputs (even though
+				// the function signature would have them as inputs), so we 
+				// want to flip the connotation here
+				bool const bWantsOutputConnection = (PinDir == EGPD_Input) ^ bIsEventSpawner;
+
+				if (K2Schema->FunctionHasParamOfType(AssociatedFunc, K2Node->GetBlueprint(), PinType, bWantsOutputConnection))
 				{
 					bIsFilteredOut = false;
 				}
-				else if (!bIsEventSpawner)
+				else
 				{
 					// need to take "Target" self pins into consideration for objects
-					bIsFilteredOut = !IsPinCompatibleWithTargetSelf(ContextPin, AssociatedFunc);
+					bIsFilteredOut = bIsEventSpawner || !IsPinCompatibleWithTargetSelf(ContextPin, AssociatedFunc, BlueprintAction);
 				}
 			}
+			
 		}
 	}
 
@@ -875,7 +904,7 @@ static bool BlueprintActionFilterImpl::IsMissmatchedPropertyType(FBlueprintActio
 				UEdGraphSchema_K2 const* K2Schema = CastChecked<UEdGraphSchema_K2 const>(ContextPin->GetSchema());
 
 				// have to account for "self" context pin
-				if (IsPinCompatibleWithTargetSelf(ContextPin, Property))
+				if (IsPinCompatibleWithTargetSelf(ContextPin, Property, BlueprintAction))
 				{
 					continue;
 				}
diff --git a/Engine/Source/Editor/BlueprintGraph/Private/BlueprintBoundNodeSpawner.cpp b/Engine/Source/Editor/BlueprintGraph/Private/BlueprintBoundNodeSpawner.cpp
--- a/Engine/Source/Editor/BlueprintGraph/Private/BlueprintBoundNodeSpawner.cpp
+++ b/Engine/Source/Editor/BlueprintGraph/Private/BlueprintBoundNodeSpawner.cpp
@@ -7,6 +7,8 @@
 #include "K2Node_AddComponent.h"
 #include "K2Node_CallFunction.h"
 #include "ComponentAssetBroker.h"
+#include "Editor/EditorEngine.h"	// for GetFriendlyName()
+#include "ObjectEditorUtils.h"		// for GetCategory()
 
 #define LOCTEXT_NAMESPACE "BlueprintBoundNodeSpawner"
 
@@ -33,6 +35,15 @@ namespace BlueprintBoundNodeSpawnerImpl
 	 * @param  PropertyObj	The property object you want bound to the node.
 	 */
 	static void BindFunctionNodeWithProperty(UEdGraphNode* NewNode, UObject* PropertyObj);
+
+	/**
+	 * 
+	 * 
+	 * @param  NewNode	
+	 * @param  DelegateProperty	
+	 * @return 
+	 */
+	static void BindEventNodeToDelegate(UEdGraphNode* NewNode, UObject* DelegateProperty);
 };
 
 //------------------------------------------------------------------------------
@@ -82,6 +93,42 @@ static void BlueprintBoundNodeSpawnerImpl::BindFunctionNodeWithProperty(UEdGraph
 	}
 }
 
+//------------------------------------------------------------------------------
+static void BlueprintBoundNodeSpawnerImpl::BindEventNodeToDelegate(UEdGraphNode* NewNode, UObject* InDelegateProperty)
+{
+	UMulticastDelegateProperty* DelegateProperty = CastChecked<UMulticastDelegateProperty>(InDelegateProperty);
+	if (UK2Node_Event* EventNode = Cast<UK2Node_Event>(NewNode))
+	{
+		FVector2D const SpawnLocation(EventNode->NodePosX, EventNode->NodePosY);
+		FVector2D const EventNodeOffset(-150.f, -150.f);
+		EventNode->NodePosX += EventNodeOffset.X;
+		EventNode->NodePosY += EventNodeOffset.Y;
+
+		// @TODO: unnecessary allocation, could do this all by hand, or with a local
+		UBlueprintNodeSpawner* DelegateNodeSpawner = UBlueprintPropertyNodeSpawner::Create<UK2Node_AddDelegate>(DelegateProperty);
+
+		UEdGraph* ParentGraph = EventNode->GetGraph();
+		UK2Node_AddDelegate* DelegateNode = CastChecked<UK2Node_AddDelegate>(DelegateNodeSpawner->Invoke(ParentGraph, SpawnLocation));
+
+		ParentGraph->Modify();
+		ParentGraph->AddNode(DelegateNode, /*bFromUI =*/false, /*bSelectNewNode =*/false);
+
+		UEdGraphPin* DelegateInputPin = DelegateNode->GetDelegatePin();
+		UEdGraphPin* EventDelegatePin = EventNode->FindPinChecked(UK2Node_Event::DelegateOutputName);
+
+		if ((DelegateInputPin != nullptr) && (EventDelegatePin != nullptr))
+		{
+			UEdGraphSchema_K2 const* K2Schema = GetDefault<UEdGraphSchema_K2>();
+			if (K2Schema->TryCreateConnection(EventDelegatePin, DelegateInputPin))
+			{
+				// refresh the node once connected to the delegate (it should
+				// conform it's pins to match the delegate signature)
+				EventNode->ReconstructNode();
+			}
+		}
+	}
+}
+
 /*******************************************************************************
  * UBlueprintBoundNodeSpawner
  ******************************************************************************/
@@ -141,6 +188,13 @@ UEdGraphNode* UBlueprintBoundNodeSpawner::Invoke(UEdGraph* ParentGraph, FVector2
 	{
 		BindFunctionNodeWithProperty(NewNode, BoundObj);
 	}
+	else if (BoundObj != nullptr)
+	{
+		if (!bIsTemplateNode && BoundObj->IsA<UMulticastDelegateProperty>())
+		{
+			BindEventNodeToDelegate(NewNode, BoundObj);
+		}
+	}
 	
 	return NewNode;
 }
@@ -149,6 +203,8 @@ UEdGraphNode* UBlueprintBoundNodeSpawner::Invoke(UEdGraph* ParentGraph, FVector2
 FText UBlueprintBoundNodeSpawner::GetDefaultMenuName() const
 {
 	check(SubSpawner != nullptr);
+
+	bool const bShowFriendlyNames = GetDefault<UEditorStyleSettings>()->bShowFriendlyNames;
 	
 	FText MenuName;
 	if (NodeClass == UK2Node_ComponentBoundEvent::StaticClass())
@@ -159,6 +215,13 @@ FText UBlueprintBoundNodeSpawner::GetDefaultMenuName() const
 			MenuName = FText::Format(LOCTEXT("ComponentEventName", "Add {0}"), PropertyName);
 		}
 	}
+	else if (BoundObjPtr.IsValid() && BoundObjPtr->IsA<UMulticastDelegateProperty>())
+	{
+		UMulticastDelegateProperty* DelegateProperty = CastChecked<UMulticastDelegateProperty>(BoundObjPtr.Get());
+		FText const PropertyName = FText::FromString(bShowFriendlyNames ? UEditorEngine::GetFriendlyName(DelegateProperty) : DelegateProperty->GetName());
+
+		MenuName = FText::Format(LOCTEXT("ComponentEventName", "Assign {0}"), PropertyName);
+	}
 	else
 	{
 		MenuName = SubSpawner->GetDefaultMenuName();
@@ -184,7 +247,14 @@ FText UBlueprintBoundNodeSpawner::GetDefaultMenuCategory() const
 		FText const PropertyName = FText::FromName(BoundObjPtr->GetFName());
 		DefaultCategory = FText::Format(LOCTEXT("ComponentEventCategory", "Call Function on {0}|{1}"), PropertyName, DefaultCategory);
 	}
-	
+	else if (UMulticastDelegateProperty* DelegateProperty = Cast<UMulticastDelegateProperty>(BoundObjPtr.Get()))
+	{
+		DefaultCategory = FText::FromString(FObjectEditorUtils::GetCategory(DelegateProperty));
+		if (DefaultCategory.IsEmpty())
+		{
+			DefaultCategory = LOCTEXT("BoundDelegateCategory", "Event Dispatchers");
+		}
+	}	
 	return DefaultCategory;
 }
 
diff --git a/Engine/Source/Editor/BlueprintGraph/Private/BlueprintPropertyNodeSpawner.cpp b/Engine/Source/Editor/BlueprintGraph/Private/BlueprintPropertyNodeSpawner.cpp
--- a/Engine/Source/Editor/BlueprintGraph/Private/BlueprintPropertyNodeSpawner.cpp
+++ b/Engine/Source/Editor/BlueprintGraph/Private/BlueprintPropertyNodeSpawner.cpp
@@ -6,6 +6,7 @@
 #include "EditorStyleSettings.h"	// for bShowFriendlyNames
 #include "ObjectEditorUtils.h"		// for GetCategory()
 #include "BlueprintEditorUtils.h"	// for FindBlueprintForNodeChecked()
+#include "EditorCategoryUtils.h"
 
 #define LOCTEXT_NAMESPACE "BlueprintPropertyNodeSpawner"
 
@@ -108,7 +109,7 @@ FText UBlueprintPropertyNodeSpawner::GetDefaultMenuCategory() const
 	}
 	else if (PropertyCategory.IsEmpty())
 	{
-		PropertyCategory = LOCTEXT("DefaultDelegateCategory", "Event Dispatchers");
+		PropertyCategory = FEditorCategoryUtils::GetCommonCategory(FCommonEditorCategory::Delegates);
 	}
 	return PropertyCategory;
 }
diff --git a/Engine/Source/Editor/BlueprintGraph/Private/K2Node_AssignDelegate.cpp b/Engine/Source/Editor/BlueprintGraph/Private/K2Node_AssignDelegate.cpp
new file mode 100644
--- /dev/null
+++ b/Engine/Source/Editor/BlueprintGraph/Private/K2Node_AssignDelegate.cpp
@@ -0,0 +1,101 @@
+// Copyright 1998-2014 Epic Games, Inc. All Rights Reserved.
+
+#include "BlueprintGraphPrivatePCH.h"
+#include "K2Node_AssignDelegate.h"
+#include "K2Node_CustomEvent.h"		// for CreateFromFunction()
+#include "BlueprintEditorUtils.h"	// for DoesSupportEventGraphs(), FindUniqueKismetName(), etc.
+#include "EditorStyleSettings.h"	// for bShowFriendlyNames
+#include "ObjectEditorUtils.h"		// for GetCategory()
+#include "EditorCategoryUtils.h"
+#include "Editor/EditorEngine.h"	// for GetFriendlyName()
+
+#define LOCTEXT_NAMESPACE "K2Node_AssignDelegate"
+
+//------------------------------------------------------------------------------
+UK2Node_AssignDelegate::UK2Node_AssignDelegate(const class FPostConstructInitializeProperties& PCIP)
+	: Super(PCIP)
+{
+}
+
+//------------------------------------------------------------------------------
+FText UK2Node_AssignDelegate::GetNodeTitle(ENodeTitleType::Type TitleType) const
+{
+	FText NodeTitle;
+	if (TitleType == ENodeTitleType::ListView)
+	{
+		FText PropertyName;
+		if (UProperty* Property = GetProperty())
+		{
+			bool const bShowFriendlyNames = GetDefault<UEditorStyleSettings>()->bShowFriendlyNames;
+			PropertyName = FText::FromString(bShowFriendlyNames ? UEditorEngine::GetFriendlyName(Property) : Property->GetName());
+		}
+		else
+		{
+			PropertyName = LOCTEXT("InvalidPropertyTitle", "<invalid delegate>");
+		}
+		NodeTitle = FText::Format(LOCTEXT("AssignDelegateTitle", "Assign {0}"), PropertyName);
+	}
+	else
+	{
+		NodeTitle = Super::GetNodeTitle(TitleType);
+	}
+
+	return NodeTitle;
+}
+
+//------------------------------------------------------------------------------
+FText UK2Node_AssignDelegate::GetMenuCategory() const
+{
+	FText MenuCategory = FEditorCategoryUtils::GetCommonCategory(FCommonEditorCategory::Delegates);
+	if (UProperty* Property = GetProperty())
+	{
+		MenuCategory = FText::FromString(FObjectEditorUtils::GetCategory(Property));
+	}
+	return MenuCategory;
+}
+
+//------------------------------------------------------------------------------
+bool UK2Node_AssignDelegate::IsCompatibleWithGraph(const UEdGraph* TargetGraph) const
+{
+	bool bIsCompatible = Super::IsCompatibleWithGraph(TargetGraph);
+
+	UBlueprint* Blueprint = FBlueprintEditorUtils::FindBlueprintForGraph(TargetGraph);
+	check(Blueprint != nullptr);
+	UEdGraphSchema const* Schema = TargetGraph->GetSchema();
+
+	EGraphType GraphType = Schema->GetGraphType(TargetGraph);
+	bool const bAllowEvents = (GraphType == GT_Ubergraph) && FBlueprintEditorUtils::DoesSupportEventGraphs(Blueprint);
+	
+	return bAllowEvents && Super::IsCompatibleWithGraph(TargetGraph);
+}
+
+//------------------------------------------------------------------------------
+void UK2Node_AssignDelegate::PostPlacedNewNode()
+{
+	UEdGraphPin* InDelegatePin = GetDelegatePin();
+	if (InDelegatePin->LinkedTo.Num() == 0)
+	{
+		if (UMulticastDelegateProperty* DelegateProp = Cast<UMulticastDelegateProperty>(GetProperty()))
+		{
+			FVector2D const LocationOffset(-150.f, +150.f);
+
+			FText DelegateName = FText::FromName(DelegateProp->GetFName());
+			FText DesiredEventName = FText::Format(LOCTEXT("BindedEventName", "{0}_Event"), DelegateName);
+			FName EventName = FBlueprintEditorUtils::FindUniqueKismetName(GetBlueprint(), DesiredEventName.ToString());
+
+			UK2Node_CustomEvent* EventNode = UK2Node_CustomEvent::CreateFromFunction(
+				FVector2D(NodePosX, NodePosY) + LocationOffset,
+				GetGraph(), EventName.ToString(), DelegateProp->SignatureFunction, /*bSelectNewNode =*/false);
+
+			if (EventNode != nullptr)
+			{
+				UEdGraphSchema_K2 const* K2Schema = GetDefault<UEdGraphSchema_K2>();
+				UEdGraphPin* OutDelegatePin = EventNode->FindPinChecked(UK2Node_CustomEvent::DelegateOutputName);
+
+				K2Schema->TryCreateConnection(OutDelegatePin, InDelegatePin);
+			}
+		}
+	}
+}
+
+#undef LOCTEXT_NAMESPACE
diff --git a/Engine/Source/Editor/BlueprintGraph/Private/K2Node_CustomEvent.cpp b/Engine/Source/Editor/BlueprintGraph/Private/K2Node_CustomEvent.cpp
--- a/Engine/Source/Editor/BlueprintGraph/Private/K2Node_CustomEvent.cpp
+++ b/Engine/Source/Editor/BlueprintGraph/Private/K2Node_CustomEvent.cpp
@@ -298,6 +298,7 @@ UK2Node_CustomEvent* UK2Node_CustomEvent::CreateFromFunction(FVector2D GraphPosi
 		CustomEventNode = NewObject<UK2Node_CustomEvent>(ParentGraph);
 		CustomEventNode->CustomFunctionName = FName(*Name);
 		CustomEventNode->SetFlags(RF_Transactional);
+		ParentGraph->Modify();
 		ParentGraph->AddNode(CustomEventNode, true, bSelectNewNode);
 		CustomEventNode->CreateNewGuid();
 		CustomEventNode->PostPlacedNewNode();
diff --git a/Engine/Source/Editor/UnrealEd/Private/EditorCategoryUtils.cpp b/Engine/Source/Editor/UnrealEd/Private/EditorCategoryUtils.cpp
--- a/Engine/Source/Editor/UnrealEd/Private/EditorCategoryUtils.cpp
+++ b/Engine/Source/Editor/UnrealEd/Private/EditorCategoryUtils.cpp
@@ -64,6 +64,7 @@ static FFormatNamedArguments& FEditorCategoryUtilsImpl::GetCategoryTable()
 		REGISTER_ROOT_CATEGORY(Pawn);
 		REGISTER_ROOT_CATEGORY(Rendering);
 		REGISTER_ROOT_CATEGORY(Utilities);
+		ASSIGN_ROOT_CATEGORY(Delegates, "Event Dispatchers");
 		REGISTER_ROOT_CATEGORY(Variables);
 
 #undef REGISTER_ROOT_CATEGORY
@@ -131,6 +132,7 @@ FText const& FEditorCategoryUtils::GetCommonCategory(FCommonEditorCategory::EVal
 		REGISTER_COMMON_CATEGORY(Pawn);
 		REGISTER_COMMON_CATEGORY(Rendering);
 		REGISTER_COMMON_CATEGORY(Utilities);
+		REGISTER_COMMON_CATEGORY(Delegates);
 		REGISTER_COMMON_CATEGORY(Variables);
 
 		REGISTER_COMMON_CATEGORY(FlowControl);
diff --git a/Engine/Source/Editor/UnrealEd/Public/EditorCategoryUtils.h b/Engine/Source/Editor/UnrealEd/Public/EditorCategoryUtils.h
--- a/Engine/Source/Editor/UnrealEd/Public/EditorCategoryUtils.h
+++ b/Engine/Source/Editor/UnrealEd/Public/EditorCategoryUtils.h
@@ -41,6 +41,7 @@ struct FCommonEditorCategory
 		Enum,
 		Struct,
 		Macro,
+		Delegates,
 
 		Variables,
 	};
