BP: UK2Node_EnumLiteral fix. It caused a crash in CheckAddress in Map::Add.

diff --git a/Engine/Source/Editor/BlueprintGraph/Private/K2Node_EnumLiteral.cpp b/Engine/Source/Editor/BlueprintGraph/Private/K2Node_EnumLiteral.cpp
--- a/Engine/Source/Editor/BlueprintGraph/Private/K2Node_EnumLiteral.cpp
+++ b/Engine/Source/Editor/BlueprintGraph/Private/K2Node_EnumLiteral.cpp
@@ -79,7 +79,8 @@ public:
 		UEdGraphPin* OutPin = Node->FindPinChecked(Context.Schema->PN_ReturnValue);
 		if (ensure(Context.NetMap.Find(OutPin) == NULL))
 		{
-			Context.NetMap.Add(OutPin, *ValueSource);
+			FBPTerminal* TerminalPtr = *ValueSource; //necessary because of CheckAddress in Map::Add
+			Context.NetMap.Add(OutPin, TerminalPtr);
 		}
 	}
 };
