(WIP) Adding support for local variable in the new blueprint menu system.

diff --git a/Engine/Source/Editor/BlueprintGraph/Classes/BlueprintVariableNodeSpawner.h b/Engine/Source/Editor/BlueprintGraph/Classes/BlueprintVariableNodeSpawner.h
new file mode 100644
--- /dev/null
+++ b/Engine/Source/Editor/BlueprintGraph/Classes/BlueprintVariableNodeSpawner.h
@@ -0,0 +1,109 @@
+// Copyright 1998-2014 Epic Games, Inc. All Rights Reserved.
+
+#pragma once
+
+#include "BlueprintNodeSpawner.h"
+#include "BlueprintVariableNodeSpawner.generated.h"
+
+// Forward declarations
+class UEdGraph;
+class UK2Node_Variable;
+
+/**
+ * Takes care of spawning variable getter/setter nodes. Serves as the "action"
+ * portion for certain FBlueprintActionMenuItems. Evolved from 
+ * FEdGraphSchemaAction_K2Var, and can spawn nodes for both member-variables and 
+ * local function variables.
+ */
+UCLASS(Transient)
+class BLUEPRINTGRAPH_API UBlueprintVariableNodeSpawner : public UBlueprintNodeSpawner
+{
+	GENERATED_UCLASS_BODY()
+
+public:
+	/**
+	 * Creates a new UBlueprintVariableNodeSpawner, charged with spawning 
+	 * a member-variable node (for a variable that has an associated UProperty) 
+	 * 
+	 * @param  NodeClass	The node type that you want the spawner to spawn.
+	 * @param  VarProperty	The property that represents the member-variable you want nodes spawned for.
+	 * @param  Outer		Optional outer for the new spawner (if left null, the transient package will be used).
+	 * @return A newly allocated instance of this class.
+	 */
+	static UBlueprintVariableNodeSpawner* Create(TSubclassOf<UK2Node_Variable> NodeClass, UProperty* VarProperty, UObject* Outer = nullptr);
+
+	/**
+	 * Creates a new UBlueprintVariableNodeSpawner, charged with spawning
+	 * a local-variable node (for a variable that belongs to a specific graph).
+	 * 
+	 * @param  NodeClass	The node type that you want the spawner to spawn.
+	 * @param  VarContext	The graph that the local variable belongs to.
+	 * @param  VarDesc		Details the local variable (name, type, etc.)
+	 * @param  Outer		Optional outer for the new spawner (if left null, the transient package will be used).
+	 * @return A newly allocated instance of this class.
+	 */
+	static UBlueprintVariableNodeSpawner* Create(TSubclassOf<UK2Node_Variable> NodeClass, UEdGraph* VarContext, FBPVariableDescription const& VarDesc, UObject* Outer = nullptr);
+
+	// UBlueprintNodeSpawner interface
+	virtual UEdGraphNode* Invoke(UEdGraph* ParentGraph, FVector2D const Location) const override;
+	virtual FText GetDefaultMenuName() const override;
+	virtual FText GetDefaultMenuCategory() const override;
+	virtual FName GetDefaultMenuIcon(FLinearColor& ColorOut) const override;
+	// End UBlueprintNodeSpawner interface
+	
+	/**
+	 * Since this spawner can wrap both local and member variables, we use this
+	 * method to discern between the two.
+	 * 
+	 * @return True if this spawner wraps a local variable, false if not.
+	 */
+	bool IsLocalVariable() const;
+
+	/**
+	 * If this is a local variable, then this will return the UEdGraph that it
+	 * belongs to, otherwise it pulls the outer from the wrapped member-variable 
+	 * property.
+	 * 
+	 * @return Null if the variable is no longer valid, otherwise an object pointer to the variable's owner.
+	 */
+	UObject const* GetVarOuter() const;
+
+	/**
+	 * Accessor to the variable's property. Will be null if this is for a local 
+	 * variable (as they don't have UProperties associated with them).
+	 * 
+	 * @return Null if this wraps a local variable (or if the variable property is stale), otherwise the property this was initialized with. 
+	 */
+	UProperty const* GetVarProperty() const;
+
+	/**
+	 * Utility function for easily accessing the variable's type (needs to pull
+	 * the information differently if it is a local variable as opposed to a
+	 * member variable with a UProperty).
+	 * 
+	 * @return A struct detailing the wrapped variable's type.
+	 */
+	FEdGraphPinType GetVarType() const;
+
+private:
+	/**
+	 * Utility function for easily accessing the variable's name (needs to pull
+	 * the information differently if it is a local variable as opposed to a
+	 * member variable with a UProperty).
+	 * 
+	 * @return A friendly, user presentable, name for the variable that this wraps. 
+	 */
+	FText GetVariableName() const;
+
+	/** */
+	UPROPERTY()
+	TWeakObjectPtr<UProperty> MemberVariable;
+
+	/** The graph that the local variable belongs to (if this is a local variable spawner). */
+	UPROPERTY()
+	TWeakObjectPtr<UEdGraph>  LocalVarOuter;
+
+	/** */
+	UPROPERTY()
+	FBPVariableDescription    LocalVarDesc;
+};
\ No newline at end of file
diff --git a/Engine/Source/Editor/BlueprintGraph/Private/BlueprintActionDatabase.cpp b/Engine/Source/Editor/BlueprintGraph/Private/BlueprintActionDatabase.cpp
--- a/Engine/Source/Editor/BlueprintGraph/Private/BlueprintActionDatabase.cpp
+++ b/Engine/Source/Editor/BlueprintGraph/Private/BlueprintActionDatabase.cpp
@@ -10,6 +10,7 @@
 #include "BlueprintEventNodeSpawner.h"
 #include "BlueprintComponentNodeSpawner.h"
 #include "BlueprintBoundNodeSpawner.h"
+#include "BlueprintVariableNodeSpawner.h"
 #include "Stats.h"					// for RETURN_QUICK_DECLARE_CYCLE_STAT()
 #include "ScopedTimers.h"			// for FScopedDurationTimer
 #include "CallbackDevice.h"			// for FCoreDelegates::OnAssetLoaded
@@ -261,34 +262,6 @@ static UBlueprintNodeSpawner* FBlueprintNodeSpawnerFactory::MakeAssignDelegateNo
 	// @TODO: it'd be awesome to have both nodes spawned by this available for 
 	//        context pin matching (the delegate inputs and the event outputs)
 	return UBlueprintPropertyNodeSpawner::Create<UK2Node_AssignDelegate>(DelegateProperty);
-// 	// using the event spawner as
-// 	UBlueprintNodeSpawner* EventNodeSpawner = UBlueprintEventNodeSpawner::Create(UK2Node_CustomEvent::StaticClass(), FName());
-// 	check(EventNodeSpawner != nullptr);
-// 
-// 	auto SetupDelegateEventNodeLambda = [](UEdGraphNode* NewNode, bool bIsTemplateNode, TWeakObjectPtr<UMulticastDelegateProperty> DelegatePropPtr)
-// 	{
-// 		UK2Node_CustomEvent* EventNode = CastChecked<UK2Node_CustomEvent>(NewNode);
-// 		if (DelegatePropPtr.IsValid())
-// 		{
-// 			UMulticastDelegateProperty* DelegateProp = DelegatePropPtr.Get();		
-// 			// don't needlessly spend unique blueprint names
-// 			if (!bIsTemplateNode)
-// 			{
-// 				UBlueprint* Blueprint = EventNode->GetBlueprint();
-// 
-// 				FText DelegateName = FText::FromName(DelegateProp->GetFName());
-// 				FText DesiredEventName = FText::Format(LOCTEXT("BindedEventName", "{0}_Event"), DelegateName);
-// 				EventNode->CustomFunctionName = FBlueprintEditorUtils::FindUniqueKismetName(Blueprint, DesiredEventName.ToString());
-// 			}
-// 		}
-// 	};
-// 
-// 	TWeakObjectPtr<UMulticastDelegateProperty> DelegatePropPtr = DelegateProperty;
-// 	EventNodeSpawner->CustomizeNodeDelegate = UBlueprintNodeSpawner::FCustomizeNodeDelegate::CreateStatic(SetupDelegateEventNodeLambda, DelegatePropPtr);
-// 
-// 	// the bound spawner should conform the event node to match the delegate's 
-// 	// signature
-// 	return UBlueprintBoundNodeSpawner::Create(EventNodeSpawner, DelegateProperty);
 }
 
 /*******************************************************************************
@@ -734,22 +707,25 @@ static void BlueprintActionDatabaseImpl::AddBlueprintGraphActions(UClass const*
 			ActionListOut.Add(MakeMacroNodeSpawner(*GraphIt));
 		}
 
-		// @TODO: local variables
-// 		for (auto GraphIt = Blueprint->FunctionGraphs.CreateConstIterator(); GraphIt; ++GraphIt)
-// 		{
-// 			UEdGraph* FunctionGraph = (*GraphIt);
-// 
-// 			TArray<UK2Node_FunctionEntry*> GraphEntryNodes;
-// 			FunctionGraph->GetNodesOfClass<UK2Node_FunctionEntry>(GraphEntryNodes);
-// 
-// 			for (UK2Node_FunctionEntry* FunctionEntry : GraphEntryNodes)
-// 			{
-// 				for (FBPVariableDescription const& LocalVar : FunctionEntry->LocalVariables)
-// 				{
-// 					ActionListOut.Add(MakeLocalVarNodeSpawner(LocalVar, VT_Getter));
-// 				}
-// 			}
-// 		}
+		// local variables
+		for (auto GraphIt = Blueprint->FunctionGraphs.CreateConstIterator(); GraphIt; ++GraphIt)
+		{
+			UEdGraph* FunctionGraph = (*GraphIt);
+
+			TArray<UK2Node_FunctionEntry*> GraphEntryNodes;
+			FunctionGraph->GetNodesOfClass<UK2Node_FunctionEntry>(GraphEntryNodes);
+
+			for (UK2Node_FunctionEntry* FunctionEntry : GraphEntryNodes)
+			{
+				for (FBPVariableDescription const& LocalVar : FunctionEntry->LocalVariables)
+				{
+					UBlueprintNodeSpawner* GetVarSpawner = UBlueprintVariableNodeSpawner::Create(UK2Node_VariableGet::StaticClass(), FunctionGraph, LocalVar);
+					ActionListOut.Add(GetVarSpawner);
+					UBlueprintNodeSpawner* SetVarSpawner = UBlueprintVariableNodeSpawner::Create(UK2Node_VariableSet::StaticClass(), FunctionGraph, LocalVar);
+					ActionListOut.Add(SetVarSpawner);
+				}
+			}
+		}
 	}	
 }
 
diff --git a/Engine/Source/Editor/BlueprintGraph/Private/BlueprintActionFilter.cpp b/Engine/Source/Editor/BlueprintGraph/Private/BlueprintActionFilter.cpp
--- a/Engine/Source/Editor/BlueprintGraph/Private/BlueprintActionFilter.cpp
+++ b/Engine/Source/Editor/BlueprintGraph/Private/BlueprintActionFilter.cpp
@@ -8,6 +8,7 @@
 #include "BlueprintComponentNodeSpawner.h"
 #include "BlueprintPropertyNodeSpawner.h"
 #include "BlueprintBoundNodeSpawner.h"
+#include "BlueprintVariableNodeSpawner.h"
 #include "EdGraphSchema_K2.h"		// for FBlueprintMetadata
 #include "BlueprintEditorUtils.h"	// for FindBlueprintForGraph()
 #include "ObjectEditorUtils.h"		// for IsFunctionHiddenFromClass()/IsVariableCategoryHiddenFromClass()
@@ -25,7 +26,15 @@
  ******************************************************************************/
 
 namespace BlueprintActionFilterImpl
-{	
+{
+	/**
+	 * 
+	 * 
+	 * @param  Blueprint	
+	 * @return 
+	 */
+	static UClass* GetBlueprintClass(UBlueprint const* const Blueprint);
+
 	/**
 	 * All actions are associated with some class, whether it's the class of the
 	 * node it will spawn, or the outer of the associated field (property,
@@ -54,6 +63,14 @@ namespace BlueprintActionFilterImpl
 	 * @return The action's associated function (null if it doesn't have one).
 	 */
 	static UFunction const* GetAssociatedFunction(UBlueprintNodeSpawner const* BlueprintAction);
+
+	/**
+	 * 
+	 * 
+	 * @param  BlueprintAction	
+	 * @return 
+	 */
+	static UProperty const* GetAssociatedProperty(UBlueprintNodeSpawner const* BlueprintAction);
 	
 	/**
 	 * Utility method to check and see if the specified node-spawner would 
@@ -209,6 +226,15 @@ namespace BlueprintActionFilterImpl
 	 */
 	static bool IsBoundToUnselectedObject(FBlueprintActionFilter const& Filter, UBlueprintNodeSpawner const* BlueprintAction);
 
+	/**
+	 * 
+	 * 
+	 * @param  Filter	
+	 * @param  BlueprintAction	
+	 * @return 
+	 */
+	static bool IsOutOfScopeLocalVariable(FBlueprintActionFilter const& Filter, UBlueprintNodeSpawner const* BlueprintAction);
+
 	/**
 	 * 
 	 * 
@@ -256,29 +282,21 @@ namespace BlueprintActionFilterImpl
 	static bool IsMissingMatchingPinParam(FBlueprintActionFilter const& Filter, UBlueprintNodeSpawner const* BlueprintAction);
 };
 
+
+//------------------------------------------------------------------------------
+UClass* BlueprintActionFilterImpl::GetBlueprintClass(UBlueprint const* const Blueprint)
+{
+	return (Blueprint->SkeletonGeneratedClass != nullptr) ? Blueprint->SkeletonGeneratedClass : Blueprint->ParentClass;
+}
+
 //------------------------------------------------------------------------------
 static UClass* BlueprintActionFilterImpl::GetActionClass(UBlueprintNodeSpawner const* BlueprintAction)
 {
 	UClass* ActionClass = BlueprintAction->NodeClass;
 	
-	if (UBlueprintFunctionNodeSpawner const* FuncNodeSpawner = Cast<UBlueprintFunctionNodeSpawner>(BlueprintAction))
-	{
-		UFunction const* Function = FuncNodeSpawner->GetFunction();
-		check(Function != nullptr);
-		ActionClass = Function->GetOuterUClass();
-	}
-	else if (UBlueprintEventNodeSpawner const* EventSpawner = Cast<UBlueprintEventNodeSpawner>(BlueprintAction))
-	{
-		if (UFunction const* EventFunc = EventSpawner->GetEventFunction())
-		{
-			ActionClass = EventFunc->GetOuterUClass();
-		}		
-	}
-	else if (UBlueprintPropertyNodeSpawner const* PropertySpawner = Cast<UBlueprintPropertyNodeSpawner>(BlueprintAction))
+	if (UField const* AssociatedField = GetAssociatedField(BlueprintAction))
 	{
-		UProperty const* Property = PropertySpawner->GetProperty();
-		check(Property != nullptr);
-		ActionClass = Property->GetOwnerClass();
+		ActionClass = AssociatedField->GetOwnerClass();
 	}
 	else if (UBlueprintComponentNodeSpawner const* ComponentSpawner = Cast<UBlueprintComponentNodeSpawner>(BlueprintAction))
 	{
@@ -288,6 +306,16 @@ static UClass* BlueprintActionFilterImpl::GetActionClass(UBlueprintNodeSpawner c
 	{
 		ActionClass = GetActionClass(BoundSpawner->GetSubSpawner());
 	}
+	else if (UBlueprintVariableNodeSpawner const* VarSpawner = Cast<UBlueprintVariableNodeSpawner>(BlueprintAction))
+	{
+		UObject const* VarOuter = VarSpawner->GetVarOuter();
+		if (VarSpawner->IsLocalVariable())
+		{
+			UBlueprint* Blueprint = FBlueprintEditorUtils::FindBlueprintForGraphChecked(CastChecked<UEdGraph>(VarOuter));
+			ActionClass = GetBlueprintClass(Blueprint);
+		}
+		// else, should have been caught in GetAssociatedField()
+	}
 	
 	return ActionClass;
 }
@@ -301,15 +329,14 @@ static UField const* BlueprintActionFilterImpl::GetAssociatedField(UBlueprintNod
 	{
 		ClassField = Function;
 	}
-	else if (UBlueprintPropertyNodeSpawner const* PropertySpawner = Cast<UBlueprintPropertyNodeSpawner>(BlueprintAction))
+	else if (UProperty const* Property = GetAssociatedProperty(BlueprintAction))
 	{
-		ClassField = PropertySpawner->GetProperty();
+		ClassField = Property;
 	}
 	else if (UBlueprintBoundNodeSpawner const* BoundSpawner = Cast<UBlueprintBoundNodeSpawner>(BlueprintAction))
 	{
 		ClassField = GetAssociatedField(BoundSpawner->GetSubSpawner());
 	}
-
 	return ClassField;
 }
 
@@ -334,6 +361,22 @@ static UFunction const* BlueprintActionFilterImpl::GetAssociatedFunction(UBluepr
 	return Function;
 }
 
+//------------------------------------------------------------------------------
+static UProperty const* BlueprintActionFilterImpl::GetAssociatedProperty(UBlueprintNodeSpawner const* BlueprintAction)
+{
+	UProperty const* Property = nullptr;
+
+	if (UBlueprintPropertyNodeSpawner const* PropertySpawner = Cast<UBlueprintPropertyNodeSpawner>(BlueprintAction))
+	{
+		Property = PropertySpawner->GetProperty();
+	}
+	else if (UBlueprintVariableNodeSpawner const* VarSpawner = Cast<UBlueprintVariableNodeSpawner>(BlueprintAction))
+	{
+		Property = VarSpawner->GetVarProperty();
+	}
+	return Property;
+}
+
 //------------------------------------------------------------------------------
 static bool BlueprintActionFilterImpl::IsImpure(UBlueprintNodeSpawner const* NodeSpawner)
 {
@@ -488,7 +531,7 @@ static bool BlueprintActionFilterImpl::IsRestrictedClassMember(FBlueprintActionF
 			{
 				bool bIsClassListed = false;
 				
-				UClass* BpClass = Blueprint->SkeletonGeneratedClass;
+				UClass* BpClass = GetBlueprintClass(Blueprint);
 				// walk the class inheritance chain to see if this class is one
 				// of the allowed
 				while (!bIsClassListed && (BpClass != nullptr))
@@ -519,12 +562,9 @@ static bool BlueprintActionFilterImpl::IsPermissionNotGranted(FBlueprintActionFi
 	bool bIsFilteredOut = false;
 	FBlueprintActionContext const& FilterContext = Filter.Context;
 
-	if (UBlueprintPropertyNodeSpawner const* PropertySpawner = Cast<UBlueprintPropertyNodeSpawner>(BlueprintAction))
+	if (UProperty const* const Property = GetAssociatedProperty(BlueprintAction))
 	{
-		UProperty const* const Property = PropertySpawner->GetProperty();
-		check(Property != nullptr);
-		UClass const* const NodeClass = PropertySpawner->NodeClass;
-
+		UClass const* const NodeClass = BlueprintAction->NodeClass;
 		for (UBlueprint const* Blueprint : FilterContext.Blueprints)
 		{
 			bool const bIsReadOnly = FBlueprintEditorUtils::IsPropertyReadOnlyInCurrentBlueprint(Blueprint, Property);
@@ -750,6 +790,30 @@ static bool BlueprintActionFilterImpl::IsBoundToUnselectedObject(FBlueprintActio
 	return bIsFilteredOut;
 }
 
+//------------------------------------------------------------------------------
+static bool BlueprintActionFilterImpl::IsOutOfScopeLocalVariable(FBlueprintActionFilter const& Filter, UBlueprintNodeSpawner const* BlueprintAction)
+{
+	bool bIsFilteredOut = false;
+	if (UBlueprintVariableNodeSpawner const* VarSpawner = Cast<UBlueprintVariableNodeSpawner>(BlueprintAction))
+	{
+		if (VarSpawner->IsLocalVariable())
+		{
+			bIsFilteredOut = (Filter.Context.Graphs.Num() == 0);
+
+			UEdGraph const* VarOuter = Cast<UEdGraph const>(VarSpawner->GetVarOuter());
+			for (UEdGraph* Graph : Filter.Context.Graphs)
+			{
+				if (Graph != VarOuter)
+				{
+					bIsFilteredOut = true;
+					break;
+				}
+			}
+		}
+	}
+	return bIsFilteredOut;
+}
+
 //------------------------------------------------------------------------------
 static bool BlueprintActionFilterImpl::HasMatchingPin(UBlueprintNodeSpawner const* BlueprintAction, UEdGraphPin const* Pin)
 {
@@ -767,7 +831,7 @@ static bool BlueprintActionFilterImpl::HasMatchingPin(UBlueprintNodeSpawner cons
 		check(Blueprint != nullptr);
 		UEdGraphSchema_K2 const* Schema = CastChecked<UEdGraphSchema_K2>(OuterGraph->GetSchema());
 
-		UClass* CallingContext = (Blueprint->SkeletonGeneratedClass != nullptr) ? Blueprint->SkeletonGeneratedClass : Blueprint->ParentClass;
+		UClass* CallingContext = GetBlueprintClass(Blueprint);
 		UK2Node* K2TemplateNode = Cast<UK2Node>(TemplateNode);
 		
 		for (int32 PinIndex = 0; !bHasCompatiblePin && (PinIndex < TemplateNode->Pins.Num()); ++PinIndex)
@@ -885,9 +949,7 @@ static bool BlueprintActionFilterImpl::IsFunctionMissingPinParam(FBlueprintActio
 static bool BlueprintActionFilterImpl::IsMissmatchedPropertyType(FBlueprintActionFilter const& Filter, UBlueprintNodeSpawner const* BlueprintAction)
 {
 	bool bIsFilteredOut = false;
-
-	UField const* AssociatedField = GetAssociatedField(BlueprintAction);
-	if (UProperty const* Property = Cast<UProperty const>(AssociatedField))
+	if (UProperty const* Property = GetAssociatedProperty(BlueprintAction))
 	{
 		TArray<UEdGraphPin*> const& ContextPins = Filter.Context.Pins;
 		if (ContextPins.Num() > 0)
@@ -1017,6 +1079,7 @@ FBlueprintActionFilter::FBlueprintActionFilter(uint32 Flags/*= 0x00*/)
 	AddIsFilteredTest(FIsFilteredDelegate::CreateStatic(IsFilteredNodeType, (Flags & BPFILTER_ExcludeChildNodeTypes) != 0));
 	AddIsFilteredTest(FIsFilteredDelegate::CreateStatic(IsExternalField));
 	AddIsFilteredTest(FIsFilteredDelegate::CreateStatic(IsBoundToUnselectedObject));
+	AddIsFilteredTest(FIsFilteredDelegate::CreateStatic(IsOutOfScopeLocalVariable));
 
 	// @TODO: account for all K2ActionMenuBuilder checks...
 	//    NodeCDO->CanCreateUnderSpecifiedSchema(this)
diff --git a/Engine/Source/Editor/BlueprintGraph/Private/BlueprintVariableNodeSpawner.cpp b/Engine/Source/Editor/BlueprintGraph/Private/BlueprintVariableNodeSpawner.cpp
new file mode 100644
--- /dev/null
+++ b/Engine/Source/Editor/BlueprintGraph/Private/BlueprintVariableNodeSpawner.cpp
@@ -0,0 +1,201 @@
+// Copyright 1998-2014 Epic Games, Inc. All Rights Reserved.
+
+#include "BlueprintGraphPrivatePCH.h"
+#include "BlueprintVariableNodeSpawner.h"
+#include "K2Node_Variable.h"
+#include "EditorStyleSettings.h"	// for bShowFriendlyNames
+#include "Editor/EditorEngine.h"	// for GetFriendlyName()
+#include "ObjectEditorUtils.h"		// for GetCategory()
+#include "EdGraphSchema_K2.h"		// for ConvertPropertyToPinType()
+
+#define LOCTEXT_NAMESPACE "BlueprintVariableNodeSpawner"
+
+/*******************************************************************************
+ * UBlueprintVariableNodeSpawner
+ ******************************************************************************/
+
+//------------------------------------------------------------------------------
+UBlueprintVariableNodeSpawner* UBlueprintVariableNodeSpawner::Create(TSubclassOf<UK2Node_Variable> NodeClass, UProperty* VarProperty, UObject* Outer/*= nullptr*/)
+{
+	check(VarProperty != nullptr);
+	if (Outer == nullptr)
+	{
+		Outer = GetTransientPackage();
+	}
+
+	UBlueprintVariableNodeSpawner* NodeSpawner = NewObject<UBlueprintVariableNodeSpawner>(Outer);
+	NodeSpawner->NodeClass      = NodeClass;
+	NodeSpawner->MemberVariable = VarProperty;
+	return NodeSpawner;
+}
+
+//------------------------------------------------------------------------------
+UBlueprintVariableNodeSpawner* UBlueprintVariableNodeSpawner::Create(TSubclassOf<UK2Node_Variable> NodeClass, UEdGraph* VarContext, FBPVariableDescription const& VarDesc, UObject* Outer/*= nullptr*/)
+{
+	check(VarContext != nullptr);
+	if (Outer == nullptr)
+	{
+		Outer = GetTransientPackage();
+	}
+
+	UBlueprintVariableNodeSpawner* NodeSpawner = NewObject<UBlueprintVariableNodeSpawner>(Outer);
+	NodeSpawner->NodeClass     = NodeClass;
+	NodeSpawner->LocalVarOuter = VarContext;
+	NodeSpawner->LocalVarDesc  = VarDesc;
+	return NodeSpawner;
+}
+
+//------------------------------------------------------------------------------
+UBlueprintVariableNodeSpawner::UBlueprintVariableNodeSpawner(class FPostConstructInitializeProperties const& PCIP)
+	: Super(PCIP)
+{
+}
+
+//------------------------------------------------------------------------------
+UEdGraphNode* UBlueprintVariableNodeSpawner::Invoke(UEdGraph* ParentGraph, FVector2D const Location) const
+{
+	FCustomizeNodeDelegate PostSpawnDelegate = CustomizeNodeDelegate;
+
+	if (IsLocalVariable())
+	{
+		auto LocalVarSetupLambda = [](UEdGraphNode* NewNode, bool bIsTemplateNode, FName VarName, UObject const* VarOuter, FGuid VarGuid, FCustomizeNodeDelegate UserDelegate)
+		{
+			UK2Node_Variable* VarNode = CastChecked<UK2Node_Variable>(NewNode);
+			VarNode->VariableReference.SetLocalMember(VarName, VarOuter->GetName(), VarGuid);
+			UserDelegate.ExecuteIfBound(NewNode, bIsTemplateNode);
+		};
+
+		if (UObject const* LocalVarOuter = GetVarOuter())
+		{
+			PostSpawnDelegate = FCustomizeNodeDelegate::CreateStatic(LocalVarSetupLambda, LocalVarDesc.VarName, LocalVarOuter, LocalVarDesc.VarGuid, CustomizeNodeDelegate);
+		}
+	}
+	else if (MemberVariable.IsValid())
+	{
+		auto MemberVarSetupLambda = [](UEdGraphNode* NewNode, bool bIsTemplateNode, UProperty const* Property, FCustomizeNodeDelegate UserDelegate)
+		{
+			UBlueprint* Blueprint = FBlueprintEditorUtils::FindBlueprintForNodeChecked(NewNode);
+			UClass* OwnerClass = Property->GetOwnerClass();
+			bool const bIsSelfContext = Blueprint->SkeletonGeneratedClass->IsChildOf(OwnerClass);
+
+			UK2Node_Variable* VarNode = CastChecked<UK2Node_Variable>(NewNode);
+			VarNode->SetFromProperty(Property, bIsSelfContext);
+			UserDelegate.ExecuteIfBound(NewNode, bIsTemplateNode);
+		};
+
+		UProperty const* Property = MemberVariable.Get();
+		PostSpawnDelegate = FCustomizeNodeDelegate::CreateStatic(MemberVarSetupLambda, Property, CustomizeNodeDelegate);
+	}
+
+	return Super::Invoke(ParentGraph, Location, PostSpawnDelegate);
+}
+
+//------------------------------------------------------------------------------
+FText UBlueprintVariableNodeSpawner::GetDefaultMenuName() const
+{	
+	FText MenuName = GetVariableName();
+	if (NodeClass != nullptr)
+	{
+		if (NodeClass->IsChildOf<UK2Node_VariableGet>())
+		{
+			MenuName = FText::Format(LOCTEXT("GetterMenuName", "Get {0}"), MenuName);
+		}
+		else if (NodeClass->IsChildOf<UK2Node_VariableSet>())
+		{
+			MenuName = FText::Format(LOCTEXT("SetterMenuName", "Set {0}"), MenuName);
+		}
+	}
+
+	return MenuName;
+}
+
+//------------------------------------------------------------------------------
+FText UBlueprintVariableNodeSpawner::GetDefaultMenuCategory() const
+{
+	FText VarSubCategory;
+	if (IsLocalVariable())
+	{
+		VarSubCategory = FText::FromName(LocalVarDesc.Category);
+	}
+	else if (MemberVariable.IsValid())
+	{
+		VarSubCategory = FText::FromString(FObjectEditorUtils::GetCategory(MemberVariable.Get()));
+	}
+	return FText::Format(LOCTEXT("VariableCategory", "Variables|{0}"), VarSubCategory);
+}
+
+//------------------------------------------------------------------------------
+FName UBlueprintVariableNodeSpawner::GetDefaultMenuIcon(FLinearColor& ColorOut) const
+{
+	return UK2Node_Variable::GetVarIconFromPinType(GetVarType(), ColorOut);
+}
+
+//------------------------------------------------------------------------------
+bool UBlueprintVariableNodeSpawner::IsLocalVariable() const
+{
+	return (LocalVarDesc.VarName != NAME_None);
+}
+
+//------------------------------------------------------------------------------
+UObject const* UBlueprintVariableNodeSpawner::GetVarOuter() const
+{
+	UObject const* VarOuter = nullptr;
+	if (IsLocalVariable())
+	{
+		if (LocalVarOuter.IsValid())
+		{
+			VarOuter = LocalVarOuter.Get();
+		}
+	}
+	else if (MemberVariable.IsValid())
+	{
+		VarOuter = MemberVariable->GetOuter();
+	}	
+	return VarOuter;
+}
+
+//------------------------------------------------------------------------------
+UProperty const* UBlueprintVariableNodeSpawner::GetVarProperty() const
+{
+	UProperty const* VarProperty = nullptr;
+	if (MemberVariable.IsValid())
+	{
+		VarProperty = MemberVariable.Get();
+	}
+	return VarProperty;
+}
+
+//------------------------------------------------------------------------------
+FEdGraphPinType UBlueprintVariableNodeSpawner::GetVarType() const
+{
+	FEdGraphPinType VarType;
+	if (IsLocalVariable())
+	{
+		VarType = LocalVarDesc.VarType;
+	}
+	else if (UProperty const* VarProperty = GetVarProperty())
+	{
+		UEdGraphSchema_K2 const* K2Schema = GetDefault<UEdGraphSchema_K2>();
+		K2Schema->ConvertPropertyToPinType(VarProperty, VarType);
+	}
+	return VarType;
+}
+
+//------------------------------------------------------------------------------
+FText UBlueprintVariableNodeSpawner::GetVariableName() const
+{
+	FText VarName;
+
+	bool bShowFriendlyNames = GetDefault<UEditorStyleSettings>()->bShowFriendlyNames;
+	if (IsLocalVariable())
+	{
+		VarName = bShowFriendlyNames ? FText::FromString(LocalVarDesc.FriendlyName) : FText::FromName(LocalVarDesc.VarName);
+	}
+	else if (MemberVariable.IsValid())
+	{
+		VarName = bShowFriendlyNames ? FText::FromString(UEditorEngine::GetFriendlyName(MemberVariable.Get())) : FText::FromName(MemberVariable->GetFName());
+	}
+	return VarName;
+}
+
+#undef LOCTEXT_NAMESPACE
diff --git a/Engine/Source/Editor/BlueprintGraph/Public/BlueprintNodeSpawner.h b/Engine/Source/Editor/BlueprintGraph/Public/BlueprintNodeSpawner.h
--- a/Engine/Source/Editor/BlueprintGraph/Public/BlueprintNodeSpawner.h
+++ b/Engine/Source/Editor/BlueprintGraph/Public/BlueprintNodeSpawner.h
@@ -178,5 +178,5 @@ protected:
 template<class NodeType>
 UBlueprintNodeSpawner* UBlueprintNodeSpawner::Create(UObject* Outer, FCustomizeNodeDelegate PostSpawnDelegate)
 {
-	Create(NodeType::StaticClass(), Outer, PostSpawnDelegate);
+	return Create(NodeType::StaticClass(), Outer, PostSpawnDelegate);
 }
\ No newline at end of file
