NUX - Adding the byte->text bool->text auto conversion logic to the schema.

diff --git a/Engine/Source/Editor/BlueprintGraph/Private/EdGraphSchema_K2.cpp b/Engine/Source/Editor/BlueprintGraph/Private/EdGraphSchema_K2.cpp
--- a/Engine/Source/Editor/BlueprintGraph/Private/EdGraphSchema_K2.cpp
+++ b/Engine/Source/Editor/BlueprintGraph/Private/EdGraphSchema_K2.cpp
@@ -886,7 +886,7 @@ void UEdGraphSchema_K2::GetContextMenuActions(const UEdGraph* CurrentGraph, cons
 				{
 					MenuBuilder->AddMenuEntry( FGraphEditorCommands::Get().RecombineStructPin );
 				}
-    
+	
 				// Conditionally add the execution path pin removal if this is an execution branching node
 				if( InGraphPin->Direction == EGPD_Output && InGraphPin->GetOwningNode())
 				{
@@ -1545,6 +1545,10 @@ bool UEdGraphSchema_K2::SearchForAutocastFunction(const UEdGraphPin* OutputPin,
 		{
 			TargetFunction = TEXT("Conv_BoolToByte");
 		}
+		else if ( InputPin->PinType.PinCategory == PC_Text )
+		{
+			TargetFunction = TEXT("Conv_BoolToText");
+		}
 	}
 	else if (OutputPin->PinType.PinCategory == PC_Byte &&
 			 (OutputPin->PinType.PinSubCategoryObject == NULL || !OutputPin->PinType.PinSubCategoryObject->IsA(UEnum::StaticClass())))
@@ -1561,6 +1565,10 @@ bool UEdGraphSchema_K2::SearchForAutocastFunction(const UEdGraphPin* OutputPin,
 		{
 			TargetFunction = TEXT("Conv_ByteToFloat");
 		}
+		else if ( InputPin->PinType.PinCategory == PC_Text )
+		{
+			TargetFunction = TEXT("Conv_ByteToText");
+		}
 	}
 	else if (OutputPin->PinType.PinCategory == PC_Name)
 	{
