(WIP) Solution for blueprint database holding onto references to blueprints that users want to delete.

diff --git a/Engine/Source/Editor/BlueprintGraph/Private/BlueprintActionDatabase.cpp b/Engine/Source/Editor/BlueprintGraph/Private/BlueprintActionDatabase.cpp
--- a/Engine/Source/Editor/BlueprintGraph/Private/BlueprintActionDatabase.cpp
+++ b/Engine/Source/Editor/BlueprintGraph/Private/BlueprintActionDatabase.cpp
@@ -420,6 +420,15 @@ namespace BlueprintActionDatabaseImpl
 	 */
 	static void OnAssetAdded(FAssetData const& NewAssetInfo);
 
+	/**
+	 * Callback to clear out object references so that a object can be deleted 
+	 * without resistance from the actions cached here. Objects passed here
+	 * won't necessarily be deleted (the user could still choose to cancel).
+	 * 
+	 * @param  ObjectsForDelete	A list of objects that MIGHT be deleted.
+	 */
+	static void OnAssetsPendingDelete(TArray<UObject*> const& ObjectsForDelete);
+
 	/**
 	 * Callback to refresh the database when an object has been delete (to clear 
 	 * any related classes that were stored in the database) 
@@ -427,6 +436,13 @@ namespace BlueprintActionDatabaseImpl
 	 * @param  AssetInfo	Data regarding the freshly removed asset.
 	 */
 	static void OnAssetRemoved(FAssetData const& AssetInfo);
+
+	/** 
+	 * Classes that we cleared from the database (to remove references, and make 
+	 * way for a delete), but in-case the class wasn't deleted we need them 
+	 * tracked here so we can add them back in.
+	 */
+	TSet<UClass*> PendingDelete;
 }
 
 //------------------------------------------------------------------------------
@@ -792,19 +808,42 @@ static void BlueprintActionDatabaseImpl::OnAssetAdded(FAssetData const& NewAsset
 }
 
 //------------------------------------------------------------------------------
-static void BlueprintActionDatabaseImpl::OnAssetRemoved(FAssetData const& AssetInfo)
+static void BlueprintActionDatabaseImpl::OnAssetsPendingDelete(TArray<UObject*> const& ObjectsForDelete)
 {
 	FBlueprintActionDatabase& ActionDatabase = FBlueprintActionDatabase::Get();
+	for (UObject* DeletingObject : ObjectsForDelete)
+	{
+		if (UBlueprint* NewBlueprint = Cast<UBlueprint>(DeletingObject))
+		{
+			// have to temporarily remove references (so that this delete isn't 
+			// blocked by dangling references)
+			ActionDatabase.ClearClassActions(NewBlueprint->GeneratedClass);
+			ActionDatabase.ClearClassActions(NewBlueprint->SkeletonGeneratedClass);
+
+			// in case they choose not to delete the object, we need to add 
+			// these back in to the database, soi we track them here
+			PendingDelete.Add(NewBlueprint->GeneratedClass);
+			PendingDelete.Add(NewBlueprint->SkeletonGeneratedClass);
+		}
+	}
+}
 
-	UClass* AssetClass = AssetInfo.GetClass();
-	ActionDatabase.ClearClassActions(AssetClass);
+//------------------------------------------------------------------------------
+static void BlueprintActionDatabaseImpl::OnAssetRemoved(FAssetData const& AssetInfo)
+{
+	FBlueprintActionDatabase& ActionDatabase = FBlueprintActionDatabase::Get();
 
 	if (AssetInfo.IsAssetLoaded())
 	{
+		UClass* AssetClass = AssetInfo.GetClass();
 		if (AssetClass->IsChildOf<UBlueprint>())
 		{
 			UBlueprint* BlueprintAsset = CastChecked<UBlueprint>(AssetInfo.GetAsset());
 			BlueprintAsset->OnChanged().RemoveStatic(&BlueprintActionDatabaseImpl::OnBlueprintChanged);
+
+			// the delete went through, so we don't need to track these for re-add
+			PendingDelete.Remove(BlueprintAsset->GeneratedClass);
+			PendingDelete.Remove(BlueprintAsset->SkeletonGeneratedClass);
 		}
 	}
 }
@@ -838,6 +877,8 @@ FBlueprintActionDatabase::FBlueprintActionDatabase()
 	IAssetRegistry& AssetRegistry = FModuleManager::LoadModuleChecked<FAssetRegistryModule>(TEXT("AssetRegistry")).Get();
 	AssetRegistry.OnAssetAdded().AddStatic(&BlueprintActionDatabaseImpl::OnAssetAdded);
 	AssetRegistry.OnAssetRemoved().AddStatic(&BlueprintActionDatabaseImpl::OnAssetRemoved);
+
+	FEditorDelegates::OnAssetsPreDelete.AddStatic(&BlueprintActionDatabaseImpl::OnAssetsPendingDelete);
 }
 
 //------------------------------------------------------------------------------
@@ -854,6 +895,14 @@ void FBlueprintActionDatabase::Tick(float DeltaTime)
 {
 	const double DurationThreshold = FMath::Min(0.003, DeltaTime * 0.01);
 
+	// entries that were removed from the database, in preparation for a delete
+	// (but the user ended up not deleting the object)
+	for (UClass* Class : BlueprintActionDatabaseImpl::PendingDelete)
+	{
+		RefreshClassActions(Class);
+	}
+	BlueprintActionDatabaseImpl::PendingDelete.Empty();
+
 	double TotalFrameDuration = 0.0;
 	while ((ActionPrimingQueue.Num() > 0) && (TotalFrameDuration < DurationThreshold))
 	{
diff --git a/Engine/Source/Editor/BlueprintGraph/Private/BlueprintNodeTemplateCache.cpp b/Engine/Source/Editor/BlueprintGraph/Private/BlueprintNodeTemplateCache.cpp
--- a/Engine/Source/Editor/BlueprintGraph/Private/BlueprintNodeTemplateCache.cpp
+++ b/Engine/Source/Editor/BlueprintGraph/Private/BlueprintNodeTemplateCache.cpp
@@ -120,7 +120,6 @@ static UEdGraph* BlueprintNodeTemplateCacheImpl::AddGraph(UBlueprint* BlueprintO
 
 //------------------------------------------------------------------------------
 FBlueprintNodeTemplateCache::FBlueprintNodeTemplateCache()
-	: FGCObject()
 {
 	using namespace BlueprintNodeTemplateCacheImpl; // for MakeCompatibleBlueprint()
 
diff --git a/Engine/Source/Editor/UnrealEd/Private/Editor.cpp b/Engine/Source/Editor/UnrealEd/Private/Editor.cpp
--- a/Engine/Source/Editor/UnrealEd/Private/Editor.cpp
+++ b/Engine/Source/Editor/UnrealEd/Private/Editor.cpp
@@ -135,6 +135,7 @@ FEditorDelegates::FOnEditorCameraMoved					FEditorDelegates::OnEditorCameraMoved
 FEditorDelegates::FOnDollyPerspectiveCamera				FEditorDelegates::OnDollyPerspectiveCamera;
 FEditorDelegates::FOnBlueprintContextMenuCreated		FEditorDelegates::OnBlueprintContextMenuCreated;
 FSimpleMulticastDelegate								FEditorDelegates::OnShutdownPostPackagesSaved;
+FEditorDelegates::FOnAssetsPreDelete					FEditorDelegates::OnAssetsPreDelete;
 
 /*-----------------------------------------------------------------------------
 	Globals.
diff --git a/Engine/Source/Editor/UnrealEd/Private/ObjectTools.cpp b/Engine/Source/Editor/UnrealEd/Private/ObjectTools.cpp
--- a/Engine/Source/Editor/UnrealEd/Private/ObjectTools.cpp
+++ b/Engine/Source/Editor/UnrealEd/Private/ObjectTools.cpp
@@ -1563,6 +1563,10 @@ namespace ObjectTools
 			return false;
 		}
 
+		// let systems clean up any unnecessary references that they may have 
+		// (so that they're not flagged in the dialog)
+		FEditorDelegates::OnAssetsPreDelete.Broadcast(ObjectsToDelete);
+
 		TSharedRef<FAssetDeleteModel> DeleteModel = MakeShareable(new FAssetDeleteModel(ObjectsToDelete));
 
 		if ( bShowConfirmation )
diff --git a/Engine/Source/Editor/UnrealEd/Public/Editor.h b/Engine/Source/Editor/UnrealEd/Public/Editor.h
--- a/Engine/Source/Editor/UnrealEd/Public/Editor.h
+++ b/Engine/Source/Editor/UnrealEd/Public/Editor.h
@@ -88,6 +88,9 @@ struct UNREALED_API FEditorDelegates
 	DECLARE_MULTICAST_DELEGATE_OneParam(FOnBlueprintContextMenuCreated, FBlueprintGraphActionListBuilder &/*ContextMenuBuilder*/);
 	/** Delegate used for entering or exiting an editor mode */
 	DECLARE_MULTICAST_DELEGATE_OneParam(FOnEditorModeTransitioned, FEdMode* /*Mode*/);
+	/** delegate type for when a user requests to delete certain assets... DOES NOT mean the asset(s) will be deleted (the user could cancel) */
+	DECLARE_MULTICAST_DELEGATE_OneParam(FOnAssetsPreDelete, const TArray<UObject*>&);
+
 
 	/** Called when the CurrentLevel is switched to a new level.  Note that this event won't be fired for temporary
 		changes to the current level, such as when copying/pasting actors. */
@@ -177,6 +180,8 @@ struct UNREALED_API FEditorDelegates
 	static FOnBlueprintContextMenuCreated OnBlueprintContextMenuCreated;
 	/** Called on editor shutdown after packages have been successfully saved */
 	static FSimpleMulticastDelegate OnShutdownPostPackagesSaved;
+	/** Called when the user requests certain assets be deleted (DOES NOT imply that the asset will be deleted... the user could cancel) */
+	static FOnAssetsPreDelete OnAssetsPreDelete;
 };
 
 /**
