(WIP) Fixing up edge cases in the new blueprint menu system involving interface functions.

diff --git a/Engine/Source/Editor/BlueprintGraph/Classes/BlueprintFunctionNodeSpawner.h b/Engine/Source/Editor/BlueprintGraph/Classes/BlueprintFunctionNodeSpawner.h
--- a/Engine/Source/Editor/BlueprintGraph/Classes/BlueprintFunctionNodeSpawner.h
+++ b/Engine/Source/Editor/BlueprintGraph/Classes/BlueprintFunctionNodeSpawner.h
@@ -5,6 +5,9 @@
 #include "BlueprintNodeSpawner.h"
 #include "BlueprintFunctionNodeSpawner.generated.h"
 
+// Forward declarations
+class UK2Node_CallFunction;
+
 /**
  * Takes care of spawning various UK2Node_CallFunction nodes. Acts as the 
  * "action" portion of certain FBlueprintActionMenuItems. 
@@ -26,6 +29,18 @@ public:
 	 */
 	static UBlueprintFunctionNodeSpawner* Create(UFunction const* const Function, UObject* Outer = nullptr);
 
+	/**
+	 * Creates a new UBlueprintFunctionNodeSpawner for the specified function.
+	 * Does not do any compatibility checking to ensure that the function is
+	 * viable as a blueprint function call (do that before calling this).
+	 * 
+	 * @param  NodeClass	The type of node you want the spawner to create.
+	 * @param  Function		The function you want assigned to new nodes.
+	 * @param  Outer		Optional outer for the new spawner (if left null, the transient package will be used).
+	 * @return A newly allocated instance of this class.
+	 */
+	static UBlueprintFunctionNodeSpawner* Create(TSubclassOf<UK2Node_CallFunction> NodeClass, UFunction const* const Function, UObject* Outer = nullptr);
+
 	// UBlueprintNodeSpawner interface
 	virtual UEdGraphNode* Invoke(UEdGraph* ParentGraph, FVector2D const Location) const override;
 	virtual FText GetDefaultMenuName() const override;
diff --git a/Engine/Source/Editor/BlueprintGraph/Private/BlueprintActionDatabase.cpp b/Engine/Source/Editor/BlueprintGraph/Private/BlueprintActionDatabase.cpp
--- a/Engine/Source/Editor/BlueprintGraph/Private/BlueprintActionDatabase.cpp
+++ b/Engine/Source/Editor/BlueprintGraph/Private/BlueprintActionDatabase.cpp
@@ -478,6 +478,10 @@ static void BlueprintActionDatabaseImpl::AddClassFunctionActions(UClass const* c
 		
 		if (UEdGraphSchema_K2::CanUserKismetCallFunction(Function))
 		{
+			// @TODO: if this is a Blueprint, and this function is from a 
+			//        Blueprint "implemented interface", then we don't need to 
+			//        include it (the function is accounted for in from the 
+			//        interface class).
 			if (UBlueprintFunctionNodeSpawner* NodeSpawner = UBlueprintFunctionNodeSpawner::Create(Function))
 			{
 				ActionListOut.Add(NodeSpawner);
diff --git a/Engine/Source/Editor/BlueprintGraph/Private/BlueprintActionFilter.cpp b/Engine/Source/Editor/BlueprintGraph/Private/BlueprintActionFilter.cpp
--- a/Engine/Source/Editor/BlueprintGraph/Private/BlueprintActionFilter.cpp
+++ b/Engine/Source/Editor/BlueprintGraph/Private/BlueprintActionFilter.cpp
@@ -194,6 +194,15 @@ namespace BlueprintActionFilterImpl
 	 */
 	static bool IsExternalField(FBlueprintActionFilter const& Filter, UBlueprintNodeSpawner const* BlueprintAction);
 
+	/**
+	 * 
+	 * 
+	 * @param  Filter	
+	 * @param  BlueprintAction	
+	 * @return 
+	 */
+	static bool IsUneededMessageNode(FBlueprintActionFilter const& Filter, UBlueprintNodeSpawner const* BlueprintAction);
+
 	/**
 	 * Rejection test that checks to see if the node-spawner is associated with 
 	 * a field that is hidden from the specified blueprint (via metadata).
@@ -626,13 +635,15 @@ static bool BlueprintActionFilterImpl::IsExternalField(FBlueprintActionFilter co
 	if (UField const* ClassField = GetAssociatedField(BlueprintAction))
 	{
 		UClass* FieldClass = Cast<UClass>(ClassField->GetOuter());
+		bool const bIsInterfaceMethod = FieldClass->IsChildOf<UInterface>();
+
 		// global (and static library) fields can stay (unless explicitly
 		// excluded... save that for a separate test)
 		if ((FieldClass != nullptr) && !IsGlobal(ClassField))
 		{
 			for (UClass const* Class : Filter.OwnerClasses)
 			{
-				bool const bIsInternalFunc = Class->IsChildOf(FieldClass);	
+				bool const bIsInternalFunc = Class->IsChildOf(FieldClass) || (bIsInterfaceMethod && Class->ImplementsInterface(FieldClass));
 				if (!bIsInternalFunc)
 				{
 					bIsFilteredOut = true;
@@ -645,6 +656,30 @@ static bool BlueprintActionFilterImpl::IsExternalField(FBlueprintActionFilter co
 	return bIsFilteredOut;
 }
 
+//------------------------------------------------------------------------------
+static bool BlueprintActionFilterImpl::IsUneededMessageNode(FBlueprintActionFilter const& Filter, UBlueprintNodeSpawner const* BlueprintAction)
+{
+	bool bIsFilteredOut = false;
+	if (BlueprintAction->NodeClass->IsChildOf<UK2Node_Message>())
+	{
+		UFunction const* const InterfaceFunc = GetAssociatedFunction(BlueprintAction);
+		checkSlow(InterfaceFunc != nullptr);
+		checkSlow(InterfaceFunc->GetOuterUClass()->IsChildOf<UInterface>());
+
+		bIsFilteredOut = (Filter.Context.Pins.Num() > 0);
+		for (UEdGraphPin const* ContextPin : Filter.Context.Pins)
+		{
+			if (!IsPinCompatibleWithTargetSelf(ContextPin, InterfaceFunc, BlueprintAction))
+			{
+				bIsFilteredOut = false;
+				break;
+			}
+		}
+
+	}
+	return bIsFilteredOut;
+}
+
 //------------------------------------------------------------------------------
 static bool BlueprintActionFilterImpl::IsFieldCategoryHidden(FBlueprintActionFilter const& Filter, UBlueprintNodeSpawner const* BlueprintAction)
 {
@@ -865,11 +900,23 @@ static bool BlueprintActionFilterImpl::IsPinCompatibleWithTargetSelf(UEdGraphPin
 	if ((Pin->Direction == EGPD_Output) && (OwnerClass != nullptr))
 	{
 		FEdGraphPinType const& PinType = Pin->PinType;
-		if (PinType.PinSubCategoryObject.IsValid())
+		UEdGraphSchema const* const PinSchema = Pin->GetSchema();
+		checkSlow(PinSchema != nullptr);
+
+		UClass const* PinObjClass = nullptr;
+		if (PinSchema->IsSelfPin(*Pin))
+		{
+			UBlueprint const* Blueprint = FBlueprintEditorUtils::FindBlueprintForNodeChecked(Pin->GetOwningNode());
+			PinObjClass = GetBlueprintClass(Blueprint);
+		}
+		else if (PinType.PinSubCategoryObject.IsValid()) // if PC_Object/PC_Interface
+		{
+			PinObjClass = Cast<UClass>(PinType.PinSubCategoryObject.Get());
+		}
+		
+		if (PinObjClass != nullptr)
 		{
-			// if PC_Object/PC_Interface
-			UClass const* const PinObjClass = Cast<UClass>(PinType.PinSubCategoryObject.Get());
-			if ((PinObjClass != nullptr) && PinObjClass->IsChildOf(OwnerClass))
+			if (PinObjClass->IsChildOf(OwnerClass) || PinObjClass->ImplementsInterface(OwnerClass))
 			{
 				bIsCompatible = true;
 				if (PinType.bIsArray)
@@ -878,7 +925,8 @@ static bool BlueprintActionFilterImpl::IsPinCompatibleWithTargetSelf(UEdGraphPin
 					{
 						// @TODO: duplicated from UK2Node_CallFunction::AllowMultipleSelfs()... 
 						//        move into a shared resource (don't want to call 
-						//        AllowMultipleSelfs()  to save on perf)
+						//        GetTemplateNode() to get at AllowMultipleSelfs() 
+						//        for perf reasons)
 						bool const bHasReturnParam = (Function->GetReturnProperty() != nullptr);
 						bool const bIsImpure = (Function->HasAnyFunctionFlags(FUNC_BlueprintPure) == false);
 						bool const bIsLatent = (Function->HasMetaData(FBlueprintMetadata::MD_Latent) != false);
@@ -1056,7 +1104,6 @@ FBlueprintActionFilter::FBlueprintActionFilter(uint32 Flags/*= 0x00*/)
 	AddIsFilteredTest(FIsFilteredDelegate::CreateStatic(IsMissingMatchingPinParam));
 	AddIsFilteredTest(FIsFilteredDelegate::CreateStatic(IsMissmatchedPropertyType));
 	AddIsFilteredTest(FIsFilteredDelegate::CreateStatic(IsFunctionMissingPinParam));
-
 	AddIsFilteredTest(FIsFilteredDelegate::CreateStatic(IsIncompatibleImpureNode));
 	
 	AddIsFilteredTest(FIsFilteredDelegate::CreateStatic(IsFieldCategoryHidden));	
@@ -1172,5 +1219,6 @@ bool FBlueprintActionFilter::IsFilteredByThis(UBlueprintNodeSpawner const* Bluep
 			break;
 		}
 	}
+
 	return bIsFiltered;
 }
diff --git a/Engine/Source/Editor/BlueprintGraph/Private/BlueprintFunctionNodeSpawner.cpp b/Engine/Source/Editor/BlueprintGraph/Private/BlueprintFunctionNodeSpawner.cpp
--- a/Engine/Source/Editor/BlueprintGraph/Private/BlueprintFunctionNodeSpawner.cpp
+++ b/Engine/Source/Editor/BlueprintGraph/Private/BlueprintFunctionNodeSpawner.cpp
@@ -4,9 +4,7 @@
 #include "BlueprintFunctionNodeSpawner.h"
 #include "EdGraphSchema_K2.h"	// for FBlueprintMetadata
 
-/*******************************************************************************
- * UBlueprintFunctionNodeSpawner
- ******************************************************************************/
+#define LOCTEXT_NAMESPACE "BlueprintFunctionNodeSpawner"
 
 //------------------------------------------------------------------------------
 // Evolved from FK2ActionMenuBuilder::AddSpawnInfoForFunction()
@@ -19,32 +17,48 @@ UBlueprintFunctionNodeSpawner* UBlueprintFunctionNodeSpawner::Create(UFunction c
 	bool const bIsCommutativeAssociativeBinaryOp = Function->HasMetaData(FBlueprintMetadata::MD_CommutativeAssociativeBinaryOperator);
 	bool const bIsMaterialParamCollectionFunc = Function->HasMetaData(FBlueprintMetadata::MD_MaterialParameterCollectionFunction);
 
-	if (Outer == nullptr)
-	{
-		Outer = GetTransientPackage();
-	}
-	UBlueprintFunctionNodeSpawner* NodeSpawner = NewObject<UBlueprintFunctionNodeSpawner>(Outer);
-	NodeSpawner->Function  = Function;
-
+	TSubclassOf<UK2Node_CallFunction> NodeClass;
 	if (bIsCommutativeAssociativeBinaryOp && bIsPure)
 	{
-		NodeSpawner->NodeClass = UK2Node_CommutativeAssociativeBinaryOperator::StaticClass();
+		NodeClass = UK2Node_CommutativeAssociativeBinaryOperator::StaticClass();
 	}
 	else if (bIsMaterialParamCollectionFunc)
 	{
-		NodeSpawner->NodeClass = UK2Node_CallMaterialParameterCollectionFunction::StaticClass();
+		NodeClass = UK2Node_CallMaterialParameterCollectionFunction::StaticClass();
 	}
 	// @TODO:
 	//   else if CallOnMember => UK2Node_CallFunctionOnMember
 	//   else if bIsParentContext => UK2Node_CallParentFunction
 	else if (bHasArrayPointerParms)
 	{
-		NodeSpawner->NodeClass = UK2Node_CallArrayFunction::StaticClass();
+		NodeClass = UK2Node_CallArrayFunction::StaticClass();
 	}
 	else
+	{
+		NodeClass = UK2Node_CallFunction::StaticClass();
+	}
+
+	return Create(NodeClass, Function, Outer);
+}
+
+//------------------------------------------------------------------------------
+UBlueprintFunctionNodeSpawner* UBlueprintFunctionNodeSpawner::Create(TSubclassOf<UK2Node_CallFunction> NodeClass, UFunction const* const Function, UObject* Outer/* = nullptr*/)
+{
+	if (Outer == nullptr)
+	{
+		Outer = GetTransientPackage();
+	}
+	UBlueprintFunctionNodeSpawner* NodeSpawner = NewObject<UBlueprintFunctionNodeSpawner>(Outer);
+	NodeSpawner->Function = Function;
+
+	if (NodeClass == nullptr)
 	{
 		NodeSpawner->NodeClass = UK2Node_CallFunction::StaticClass();
 	}
+	else
+	{
+		NodeSpawner->NodeClass = NodeClass;
+	}
 
 	return NodeSpawner;
 }
@@ -108,3 +122,5 @@ UFunction const* UBlueprintFunctionNodeSpawner::GetFunction() const
 {
 	return Function;
 }
+
+#undef LOCTEXT_NAMESPACE
diff --git a/Engine/Source/Editor/BlueprintGraph/Private/EdGraphSchema_K2.cpp b/Engine/Source/Editor/BlueprintGraph/Private/EdGraphSchema_K2.cpp
--- a/Engine/Source/Editor/BlueprintGraph/Private/EdGraphSchema_K2.cpp
+++ b/Engine/Source/Editor/BlueprintGraph/Private/EdGraphSchema_K2.cpp
@@ -2847,9 +2847,8 @@ bool UEdGraphSchema_K2::ArePinTypesCompatible(const FEdGraphPinType& Output, con
 		{
 			OutputClass = CallingContext;
 		}
-		check(OutputClass != nullptr);
 
-		return OutputClass->ImplementsInterface(InterfaceClass) || OutputClass->IsChildOf(InterfaceClass);
+		return OutputClass && OutputClass->ImplementsInterface(InterfaceClass) || OutputClass->IsChildOf(InterfaceClass);
 	}
 
 	// Pins representing BLueprint objects and subclass of UObject can match when EditoronlyBP.bAllowClassAndBlueprintPinMatching=true (BaseEngine.ini)
diff --git a/Engine/Source/Editor/Kismet/Private/BlueprintActionMenuItem.cpp b/Engine/Source/Editor/Kismet/Private/BlueprintActionMenuItem.cpp
--- a/Engine/Source/Editor/Kismet/Private/BlueprintActionMenuItem.cpp
+++ b/Engine/Source/Editor/Kismet/Private/BlueprintActionMenuItem.cpp
@@ -119,6 +119,10 @@ UEdGraphNode* FBlueprintActionMenuItem::PerformAction(UEdGraph* ParentGraph, UEd
 	else if (SpawnedNode != nullptr)
 	{
 		FKismetEditorUtilities::BringKismetToFocusAttentionOnObject(SpawnedNode);
+		if (FromPin != nullptr)
+		{
+			SpawnedNode->AutowireNewNode(FromPin);
+		}
 	}
 	
 	return SpawnedNode;
