Making sure that knot/reroute nodes are accounted for in the new blueprint menu system.

diff --git a/Engine/Source/Editor/BlueprintGraph/Classes/K2Node_Knot.h b/Engine/Source/Editor/BlueprintGraph/Classes/K2Node_Knot.h
--- a/Engine/Source/Editor/BlueprintGraph/Classes/K2Node_Knot.h
+++ b/Engine/Source/Editor/BlueprintGraph/Classes/K2Node_Knot.h
@@ -28,6 +28,7 @@ public:
 	virtual void NotifyPinConnectionListChanged(UEdGraphPin* Pin) override;
 	virtual void PostReconstructNode() override;
 	virtual int32 GetNodeRefreshPriority() const override { return EBaseNodeRefreshPriority::Low_UsesDependentWildcard; }
+	virtual void GetMenuActions(TArray<UBlueprintNodeSpawner*>& ActionListOut) const override;
 	// End of UK2Node interface
 
 	UEdGraphPin* GetInputPin() const
diff --git a/Engine/Source/Editor/BlueprintGraph/Private/K2Node_Knot.cpp b/Engine/Source/Editor/BlueprintGraph/Private/K2Node_Knot.cpp
--- a/Engine/Source/Editor/BlueprintGraph/Private/K2Node_Knot.cpp
+++ b/Engine/Source/Editor/BlueprintGraph/Private/K2Node_Knot.cpp
@@ -3,6 +3,7 @@
 #include "BlueprintGraphPrivatePCH.h"
 #include "K2ActionMenuBuilder.h"
 #include "Kismet2NameValidators.h"
+#include "BlueprintNodeSpawner.h"
 
 #define LOCTEXT_NAMESPACE "K2Node_Knot"
 
@@ -40,6 +41,10 @@ FText UK2Node_Knot::GetNodeTitle(ENodeTitleType::Type TitleType) const
 	{
 		return FText::FromString(NodeComment);
 	}
+	else if (TitleType == ENodeTitleType::ListView)
+	{
+		return LOCTEXT("KnotListTitle", "Add Reroute Node...");
+	}
 	else
 	{
 		return LOCTEXT("KnotTitle", "Reroute Node");
@@ -140,6 +145,14 @@ void UK2Node_Knot::PostReconstructNode()
 	}
 }
 
+void UK2Node_Knot::GetMenuActions(TArray<UBlueprintNodeSpawner*>& ActionListOut) const
+{
+	UBlueprintNodeSpawner* NodeSpawner = UBlueprintNodeSpawner::Create(GetClass());
+	check(NodeSpawner != nullptr);
+
+	ActionListOut.Add(NodeSpawner);
+}
+
 bool UK2Node_Knot::ShouldOverridePinNames() const
 {
 	return true;
