GetDataTableRow Blueprint node: fix for TTP 339211 EDITOR: Get Data Table Row, Out Row not updating as expected

Deleting some code that was copied over from the SpawnActor node, but not actually useful/needed.

diff --git a/Engine/Source/Editor/BlueprintGraph/Classes/K2Node_GetDataTableRow.h b/Engine/Source/Editor/BlueprintGraph/Classes/K2Node_GetDataTableRow.h
--- a/Engine/Source/Editor/BlueprintGraph/Classes/K2Node_GetDataTableRow.h
+++ b/Engine/Source/Editor/BlueprintGraph/Classes/K2Node_GetDataTableRow.h
@@ -34,10 +34,10 @@ class UK2Node_GetDataTableRow : public UK2Node
 	// End UK2Node interface
 
 
-	/** Create new pins to show properties on archetype */
+	/** Set the return type of our struct */
 	void SetReturnTypeForStruct(UScriptStruct* InClass);
-	/** See if this is a spawn variable pin, or a 'default' pin */
-	bool IsDataTablePin(UEdGraphPin* Pin);
+	/** Get the return type of our struct */
+	UScriptStruct* GetReturnTypeForStruct();
 
 	/** Get the then output pin */
 	UEdGraphPin* GetThenPin() const;
diff --git a/Engine/Source/Editor/BlueprintGraph/Private/K2Node_GetDataTableRow.cpp b/Engine/Source/Editor/BlueprintGraph/Private/K2Node_GetDataTableRow.cpp
--- a/Engine/Source/Editor/BlueprintGraph/Private/K2Node_GetDataTableRow.cpp
+++ b/Engine/Source/Editor/BlueprintGraph/Private/K2Node_GetDataTableRow.cpp
@@ -69,13 +69,23 @@ void UK2Node_GetDataTableRow::SetPinToolTip(UEdGraphPin& MutatablePin, const FTe
 
 void UK2Node_GetDataTableRow::SetReturnTypeForStruct(UScriptStruct* RowStruct)
 {
-	check(RowStruct != NULL);
+	if (RowStruct == NULL)
+	{
+		RowStruct = FTableRowBase::StaticStruct();
+	}
 
     // Change class of output pin
 	UEdGraphPin* ResultPin = GetResultPin();
 	ResultPin->PinType.PinSubCategoryObject = RowStruct;
 }
 
+UScriptStruct* UK2Node_GetDataTableRow::GetReturnTypeForStruct()
+{
+	UScriptStruct* ReturnStructType = (UScriptStruct*)(GetResultPin()->PinType.PinSubCategoryObject.Get());
+
+	return ReturnStructType;
+}
+
 UScriptStruct* UK2Node_GetDataTableRow::GetDataTableRowStructType(const TArray<UEdGraphPin*>* InPinsToSearch /*=NULL*/) const
 {
 	UScriptStruct* RowStructType = NULL;
@@ -100,7 +110,7 @@ void UK2Node_GetDataTableRow::ReallocatePinsDuringReconstruction(TArray<UEdGraph
 	Super::ReallocatePinsDuringReconstruction(OldPins);
 		
 	UScriptStruct* RowStruct = GetDataTableRowStructType(&OldPins);
-	if( RowStruct != NULL )
+	if( RowStruct != NULL)
 	{
 		SetReturnTypeForStruct(RowStruct);
 	}
@@ -119,47 +129,26 @@ FText UK2Node_GetDataTableRow::GetMenuCategory() const
 	return FEditorCategoryUtils::GetCommonCategory(FCommonEditorCategory::Utilities);
 }
 
-bool UK2Node_GetDataTableRow::IsDataTablePin(UEdGraphPin* Pin)
-{
-	const UEdGraphSchema_K2* K2Schema = GetDefault<UEdGraphSchema_K2>();
-
-	return(	Pin->PinName != K2Schema->PN_Execute &&
-			Pin->PinName != K2Schema->PN_Then &&
-			Pin->PinName != K2Schema->PN_ReturnValue &&
-			Pin->PinName != UK2Node_GetDataTableRowHelper::DataTablePinName &&
-			Pin->PinName != UK2Node_GetDataTableRowHelper::RowNotFoundPinName &&
-			Pin->PinName != UK2Node_GetDataTableRowHelper::RowNamePinName );
-}
-
-
 void UK2Node_GetDataTableRow::PinDefaultValueChanged(UEdGraphPin* ChangedPin) 
 {
 	if (ChangedPin->PinName == UK2Node_GetDataTableRowHelper::DataTablePinName)
 	{
 		const UEdGraphSchema_K2* K2Schema = GetDefault<UEdGraphSchema_K2>();
 
-		// Because the archetype has changed, we break the output link as the output pin type will change
 		UEdGraphPin* ResultPin = GetResultPin();
-		ResultPin->BreakAllPinLinks();
-
-		// Remove all pins related to archetype variables
-		TArray<UEdGraphPin*> OldPins = Pins;
-		for (int32 i = 0; i < OldPins.Num(); i++)
-		{
-			UEdGraphPin* OldPin = OldPins[i];
-			if (IsDataTablePin(OldPin))
-			{
-				OldPin->BreakAllPinLinks();
-				Pins.Remove(OldPin);
-			}
-		}
 
+		UScriptStruct* PrevRowStruct = GetReturnTypeForStruct();
 		UScriptStruct* RowStruct = GetDataTableRowStructType();
-		if (RowStruct != NULL)
+
+		if (RowStruct != PrevRowStruct || RowStruct == NULL)
 		{
-			SetReturnTypeForStruct(RowStruct);
+			// Because the data table type has changed, we break the output link as the output pin type will change
+			ResultPin->BreakAllPinLinks();
+
 		}
 
+		SetReturnTypeForStruct(RowStruct);
+
 		// Refresh the UI for the graph so the pin changes show up
 		UEdGraph* Graph = GetGraph();
 		Graph->NotifyGraphChanged();
