Accidently left an incomplete file in my last changelist, fixing up to save build breaks.

diff --git a/Engine/Source/Editor/BlueprintGraph/Private/BlueprintActionFilter.cpp b/Engine/Source/Editor/BlueprintGraph/Private/BlueprintActionFilter.cpp
--- a/Engine/Source/Editor/BlueprintGraph/Private/BlueprintActionFilter.cpp
+++ b/Engine/Source/Editor/BlueprintGraph/Private/BlueprintActionFilter.cpp
@@ -196,6 +196,15 @@ namespace BlueprintActionFilterImpl
 	 */
 	static bool IsBoundToUnselectedObject(FBlueprintActionFilter const& Filter, UBlueprintNodeSpawner const* BlueprintAction);
 
+	/**
+	 * 
+	 * 
+	 * @param  Filter	
+	 * @param  BlueprintAction	
+	 * @return 
+	 */
+	static bool IsBoundFunctionImpure(FBlueprintActionFilter const& Filter, UBlueprintNodeSpawner const* BlueprintAction);
+
 	/**
 	 * 
 	 * 
@@ -706,10 +715,10 @@ static bool BlueprintActionFilterImpl::IsBoundFunctionImpure(FBlueprintActionFil
 	bool bIsFilteredOut = false;
 	if (UBlueprintBoundNodeSpawner const* BoundSpawner = Cast<UBlueprintBoundNodeSpawner>(BlueprintAction))
 	{
-		if (UFunction* BoundFunction = FBlueprintEditorUtils::GetAssociatedFunction(BoundSpawner))
+		if (UFunction const* BoundFunction = FBlueprintNodeSpawnerUtils::GetAssociatedFunction(BoundSpawner))
 		{
 			bIsFilteredOut = !BoundFunction->HasAnyFunctionFlags(FUNC_BlueprintPure) && 
-				!Function->HasAnyFunctionFlags(FUNC_Const);
+				!BoundFunction->HasAnyFunctionFlags(FUNC_Const);
 		}
 	}
 	return bIsFilteredOut;
