#UE4 Fix for rare crash

diff --git a/Engine/Source/Editor/BlueprintGraph/Private/BlueprintActionDatabase.cpp b/Engine/Source/Editor/BlueprintGraph/Private/BlueprintActionDatabase.cpp
--- a/Engine/Source/Editor/BlueprintGraph/Private/BlueprintActionDatabase.cpp
+++ b/Engine/Source/Editor/BlueprintGraph/Private/BlueprintActionDatabase.cpp
@@ -750,9 +750,13 @@ static void BlueprintActionDatabaseImpl::AddAutonomousNodeActions(UClass* const
 //------------------------------------------------------------------------------
 static void BlueprintActionDatabaseImpl::OnBlueprintChanged(UBlueprint* Blueprint)
 {
-	FBlueprintActionDatabase& ActionDatabase = FBlueprintActionDatabase::Get();
-	ActionDatabase.RefreshClassActions(Blueprint->SkeletonGeneratedClass);
-	ActionDatabase.RefreshClassActions(Blueprint->GeneratedClass);
+	if( Blueprint->SkeletonGeneratedClass &&
+		Blueprint->GeneratedClass )
+	{
+		FBlueprintActionDatabase& ActionDatabase = FBlueprintActionDatabase::Get();
+		ActionDatabase.RefreshClassActions(Blueprint->SkeletonGeneratedClass);
+		ActionDatabase.RefreshClassActions(Blueprint->GeneratedClass);
+	}
 }
 
 //------------------------------------------------------------------------------
