Fix Mac CIS errors

diff --git a/Engine/Source/Editor/BlueprintGraph/Private/BlueprintBoundNodeSpawner.cpp b/Engine/Source/Editor/BlueprintGraph/Private/BlueprintBoundNodeSpawner.cpp
--- a/Engine/Source/Editor/BlueprintGraph/Private/BlueprintBoundNodeSpawner.cpp
+++ b/Engine/Source/Editor/BlueprintGraph/Private/BlueprintBoundNodeSpawner.cpp
@@ -159,8 +159,8 @@ UBlueprintBoundNodeSpawner* UBlueprintBoundNodeSpawner::Create(UBlueprintNodeSpa
 //------------------------------------------------------------------------------
 UBlueprintBoundNodeSpawner::UBlueprintBoundNodeSpawner(class FPostConstructInitializeProperties const& PCIP)
 	: Super(PCIP)
-	, BoundObjPtr(nullptr)
 	, BoundObjClass(UObject::StaticClass())
+	, BoundObjPtr(nullptr)
 {
 }
 
@@ -287,7 +287,7 @@ bool UBlueprintBoundNodeSpawner::SetBoundObject(UObject const* Object)
 	}
 	else
 	{
-		ensureMsgf(false, TEXT("Cannot bind object (%s) with this node-spawner, it is not a '%s'"), *Object->GetName(), BoundObjClass->GetName());
+		ensureMsgf(false, TEXT("Cannot bind object (%s) with this node-spawner, it is not a '%s'"), *Object->GetName(), *BoundObjClass->GetName());
 	}
 	return bSuccess;
 }
