Show (as intended) a grayed out split pin option on connected pins that are of a splittable struct type

diff --git a/Engine/Source/Editor/BlueprintGraph/Private/EdGraphSchema_K2.cpp b/Engine/Source/Editor/BlueprintGraph/Private/EdGraphSchema_K2.cpp
--- a/Engine/Source/Editor/BlueprintGraph/Private/EdGraphSchema_K2.cpp
+++ b/Engine/Source/Editor/BlueprintGraph/Private/EdGraphSchema_K2.cpp
@@ -772,7 +772,7 @@ static bool IsUsingNonExistantVariable(const UEdGraphNode* InGraphNode, UBluepri
 bool UEdGraphSchema_K2::PinHasSplittableStructType(const UEdGraphPin* InGraphPin) const
 {
 	const FEdGraphPinType& PinType = InGraphPin->PinType;
-	bool bCanSplit = (!PinType.bIsArray && InGraphPin->LinkedTo.Num() == 0 && PinType.PinCategory == PC_Struct);
+	bool bCanSplit = (!PinType.bIsArray && PinType.PinCategory == PC_Struct);
 
 	if (bCanSplit)
 	{
