BP:  an unnecessary warning removed "the native property is not tagged as BlueprintReadWrite, the pin will be removed in a future release."

diff --git a/Engine/Source/Editor/BlueprintGraph/Private/K2Node_MakeStruct.cpp b/Engine/Source/Editor/BlueprintGraph/Private/K2Node_MakeStruct.cpp
--- a/Engine/Source/Editor/BlueprintGraph/Private/K2Node_MakeStruct.cpp
+++ b/Engine/Source/Editor/BlueprintGraph/Private/K2Node_MakeStruct.cpp
@@ -129,7 +129,7 @@ void UK2Node_MakeStruct::ValidateNodeDuringCompilation(class FCompilerResultsLog
 
 				const UEdGraphPin* Pin = Property ? FindPin(Property->GetName()) : NULL;
 
-				if (!bIsBlueprintVisible)
+				if (Pin && !bIsBlueprintVisible)
 				{
 					MessageLog.Warning(*LOCTEXT("PropertyIsNotBPVisible_Warning", "@@ - the native property is not tagged as BlueprintReadWrite, the pin will be removed in a future release.").ToString(), Pin);
 				}
