Fix split pin issues in Set Members node

diff --git a/Engine/Source/Editor/BlueprintGraph/Private/K2Node_SetFieldsInStruct.cpp b/Engine/Source/Editor/BlueprintGraph/Private/K2Node_SetFieldsInStruct.cpp
--- a/Engine/Source/Editor/BlueprintGraph/Private/K2Node_SetFieldsInStruct.cpp
+++ b/Engine/Source/Editor/BlueprintGraph/Private/K2Node_SetFieldsInStruct.cpp
@@ -108,6 +108,12 @@ void UK2Node_SetFieldsInStruct::RemoveFieldPins(const UEdGraphPin* Pin, EPinsToR
 {
 	if (ShowCustomPinActions(Pin, false) && (Pin->GetOwningNodeUnchecked() == this))
 	{
+		// Pretend that the action was done on the hidden parent pin if the pin is split
+		while (Pin->ParentPin != nullptr)
+		{
+			Pin = Pin->ParentPin;
+		}
+
 		bool bWasChanged = false;
 		for (FOptionalPinFromProperty& OptionalProperty : ShowPinForProperties)
 		{
