Make sure the palette builder knows that it is related to the K2 schema

diff --git a/Engine/Source/Editor/BlueprintGraph/Private/K2ActionMenuBuilder.cpp b/Engine/Source/Editor/BlueprintGraph/Private/K2ActionMenuBuilder.cpp
--- a/Engine/Source/Editor/BlueprintGraph/Private/K2ActionMenuBuilder.cpp
+++ b/Engine/Source/Editor/BlueprintGraph/Private/K2ActionMenuBuilder.cpp
@@ -479,6 +479,7 @@ FBlueprintPaletteListBuilder::FBlueprintPaletteListBuilder(UBlueprint const* Blu
 	if (Blueprint != NULL)
 	{
 		OwnerOfTemporaries = NewObject<UEdGraph>((UObject*)Blueprint);
+		OwnerOfTemporaries->Schema = UEdGraphSchema_K2::StaticClass();
 		OwnerOfTemporaries->SetFlags(RF_Transient);
 	}
 }
