UMG - Adding a new abstract node that others can use for making 'Spawn Actor' like nodes where you need most of the spawn actors capability, just not for actors.  The Create Widget function is dead, there's a new Create Widget node that can do things like expose 'expose on spawn' pins.

diff --git a/Engine/Source/Editor/BlueprintGraph/Classes/K2Node_ConstructObjectFromClass.h b/Engine/Source/Editor/BlueprintGraph/Classes/K2Node_ConstructObjectFromClass.h
new file mode 100644
--- /dev/null
+++ b/Engine/Source/Editor/BlueprintGraph/Classes/K2Node_ConstructObjectFromClass.h
@@ -0,0 +1,70 @@
+// Copyright 1998-2014 Epic Games, Inc. All Rights Reserved.
+
+#pragma once
+
+#include "K2Node.h"
+#include "K2Node_ConstructObjectFromClass.generated.h"
+
+UCLASS(abstract)
+class BLUEPRINTGRAPH_API UK2Node_ConstructObjectFromClass : public UK2Node
+{
+	GENERATED_UCLASS_BODY()
+
+	// Begin UEdGraphNode interface.
+	virtual void AllocateDefaultPins() override;
+	virtual FLinearColor GetNodeTitleColor() const override;
+	virtual FText GetNodeTitle(ENodeTitleType::Type TitleType) const override;
+	virtual void PinDefaultValueChanged(UEdGraphPin* Pin) override;
+	virtual FString GetTooltip() const override;
+	virtual void ExpandNode(class FKismetCompilerContext& CompilerContext, UEdGraph* SourceGraph) override;
+	virtual bool HasExternalBlueprintDependencies(TArray<class UStruct*>* OptionalOutput) const override;
+	virtual bool IsCompatibleWithGraph(const UEdGraph* TargetGraph) const override;
+	// End UEdGraphNode interface.
+
+	// Begin UK2Node interface
+	virtual bool IsNodeSafeToIgnore() const override { return true; }
+	virtual void ReallocatePinsDuringReconstruction(TArray<UEdGraphPin*>& OldPins) override;
+	virtual void GetNodeAttributes( TArray<TKeyValuePair<FString, FString>>& OutNodeAttributes ) const override;
+	virtual void GetMenuActions(TArray<UBlueprintNodeSpawner*>& ActionListOut) const override;
+	virtual FText GetMenuCategory() const override;
+	virtual class FNodeHandlingFunctor* CreateNodeHandler(class FKismetCompilerContext& CompilerContext) const override;
+	// End UK2Node interface
+
+	/** Create new pins to show properties on archetype */
+	void CreatePinsForClass(UClass* InClass);
+
+	/** See if this is a spawn variable pin, or a 'default' pin */
+	virtual bool IsSpawnVarPin(UEdGraphPin* Pin);
+
+	/** Get the then output pin */
+	UEdGraphPin* GetThenPin() const;
+	/** Get the blueprint input pin */	
+	UEdGraphPin* GetClassPin(const TArray<UEdGraphPin*>* InPinsToSearch=NULL) const;
+	/** Get the world context input pin, can return NULL */
+	UEdGraphPin* GetWorldContextPin() const;
+	/** Get the result output pin */
+	UEdGraphPin* GetResultPin() const;
+
+	/** Get the class that we are going to spawn, if it's defined as default value */
+	UClass* GetClassToSpawn(const TArray<UEdGraphPin*>* InPinsToSearch=NULL) const;
+
+protected:
+	/** Gets the default node title when no class is selected */
+	virtual FText GetBaseNodeTitle() const;
+	/** Gets the node title when a class has been selected. */
+	virtual FText GetNodeTitleFormat() const;
+	/** Gets base class to use for the 'class' pin.  UObject by default. */
+	virtual UClass* GetClassPinBaseClass() const;
+
+	/**
+	 * Takes the specified "MutatablePin" and sets its 'PinToolTip' field (according
+	 * to the specified description)
+	 * 
+	 * @param   MutatablePin	The pin you want to set tool-tip text on
+	 * @param   PinDescription	A string describing the pin's purpose
+	 */
+	void SetPinToolTip(UEdGraphPin& MutatablePin, const FText& PinDescription) const;
+
+	/** Tooltip text for this node. */
+	FString NodeTooltip;
+};
diff --git a/Engine/Source/Editor/BlueprintGraph/Private/K2Node_ConstructObjectFromClass.cpp b/Engine/Source/Editor/BlueprintGraph/Private/K2Node_ConstructObjectFromClass.cpp
new file mode 100644
--- /dev/null
+++ b/Engine/Source/Editor/BlueprintGraph/Private/K2Node_ConstructObjectFromClass.cpp
@@ -0,0 +1,325 @@
+// Copyright 1998-2014 Epic Games, Inc. All Rights Reserved.
+
+#include "BlueprintGraphPrivatePCH.h"
+#include "KismetCompiler.h"
+#include "Runtime/Engine/Classes/Kismet/GameplayStatics.h"
+#include "BlueprintNodeSpawner.h"
+#include "EditorCategoryUtils.h"
+
+struct FK2Node_ConstructObjectFromClassHelper
+{
+	static FString WorldContextPinName;
+	static FString ClassPinName;
+};
+
+FString FK2Node_ConstructObjectFromClassHelper::WorldContextPinName(TEXT("WorldContextObject"));
+FString FK2Node_ConstructObjectFromClassHelper::ClassPinName(TEXT("Class"));
+
+#define LOCTEXT_NAMESPACE "K2Node_ConstructObjectFromClass"
+
+UK2Node_ConstructObjectFromClass::UK2Node_ConstructObjectFromClass(const class FPostConstructInitializeProperties& PCIP)
+	: Super(PCIP)
+{
+	NodeTooltip = LOCTEXT("NodeTooltip", "Attempts to spawn a new object").ToString();
+}
+
+UClass* UK2Node_ConstructObjectFromClass::GetClassPinBaseClass() const
+{
+	return UObject::StaticClass();
+}
+
+void UK2Node_ConstructObjectFromClass::AllocateDefaultPins()
+{
+	const UEdGraphSchema_K2* K2Schema = GetDefault<UEdGraphSchema_K2>();
+
+	// Add execution pins
+	CreatePin(EGPD_Input, K2Schema->PC_Exec, TEXT(""), NULL, false, false, K2Schema->PN_Execute);
+	CreatePin(EGPD_Output, K2Schema->PC_Exec, TEXT(""), NULL, false, false, K2Schema->PN_Then);
+
+	// If required add the world context pin
+	if ( GetBlueprint()->ParentClass->HasMetaData(FBlueprintMetadata::MD_ShowHiddenSelfPins) )
+	{
+		CreatePin(EGPD_Input, K2Schema->PC_Object, TEXT(""), UObject::StaticClass(), false, false, FK2Node_ConstructObjectFromClassHelper::WorldContextPinName);
+	}
+
+	// Add blueprint pin
+	UEdGraphPin* ClassPin = CreatePin(EGPD_Input, K2Schema->PC_Class, TEXT(""), GetClassPinBaseClass(), false, false, FK2Node_ConstructObjectFromClassHelper::ClassPinName);
+	SetPinToolTip(*ClassPin, LOCTEXT("ClassPinDescription", "The object class you want to spawn"));
+
+	// Result pin
+	UEdGraphPin* ResultPin = CreatePin(EGPD_Output, K2Schema->PC_Object, TEXT(""), GetClassPinBaseClass(), false, false, K2Schema->PN_ReturnValue);
+	SetPinToolTip(*ResultPin, LOCTEXT("ResultPinDescription", "The spawned object"));
+
+	Super::AllocateDefaultPins();
+}
+
+void UK2Node_ConstructObjectFromClass::SetPinToolTip(UEdGraphPin& MutatablePin, const FText& PinDescription) const
+{
+	MutatablePin.PinToolTip = UEdGraphSchema_K2::TypeToString(MutatablePin.PinType);
+
+	UEdGraphSchema_K2 const* const K2Schema = Cast<const UEdGraphSchema_K2>(GetSchema());
+	if (K2Schema != nullptr)
+	{
+		MutatablePin.PinToolTip += TEXT(" ");
+		MutatablePin.PinToolTip += K2Schema->GetPinDisplayName(&MutatablePin);
+	}
+
+	MutatablePin.PinToolTip += FString(TEXT("\n")) + PinDescription.ToString();
+}
+
+void UK2Node_ConstructObjectFromClass::CreatePinsForClass(UClass* InClass)
+{
+	check(InClass != NULL);
+
+	const UEdGraphSchema_K2* K2Schema = GetDefault<UEdGraphSchema_K2>();
+
+	const UObject* const ClassDefaultObject = InClass->GetDefaultObject(false);
+
+	for (TFieldIterator<UProperty> PropertyIt(InClass, EFieldIteratorFlags::IncludeSuper); PropertyIt; ++PropertyIt)
+	{
+		UProperty* Property = *PropertyIt;
+		UClass* PropertyClass = CastChecked<UClass>(Property->GetOuter());
+		const bool bIsDelegate = Property->IsA(UMulticastDelegateProperty::StaticClass());
+		const bool bIsExposedToSpawn = UEdGraphSchema_K2::IsPropertyExposedOnSpawn(Property);
+		const bool bIsSettableExternally = !Property->HasAnyPropertyFlags(CPF_DisableEditOnInstance);
+
+		if(	bIsExposedToSpawn &&
+			!Property->HasAnyPropertyFlags(CPF_Parm) && 
+			bIsSettableExternally &&
+			Property->HasAllPropertyFlags(CPF_BlueprintVisible) &&
+			!bIsDelegate &&
+			(NULL == FindPin(Property->GetName()) ) )
+		{
+			UEdGraphPin* Pin = CreatePin(EGPD_Input, TEXT(""), TEXT(""), NULL, false, false, Property->GetName());
+			const bool bPinGood = (Pin != NULL) && K2Schema->ConvertPropertyToPinType(Property, /*out*/ Pin->PinType);
+			
+			if( ClassDefaultObject && K2Schema->PinDefaultValueIsEditable(*Pin))
+			{
+				FString DefaultValueAsString;
+				const bool bDefaultValueSet = FBlueprintEditorUtils::PropertyValueToString(Property, reinterpret_cast<const uint8*>(ClassDefaultObject), DefaultValueAsString);
+				check( bDefaultValueSet );
+				K2Schema->TrySetDefaultValue(*Pin, DefaultValueAsString);
+			}
+		}
+	}
+
+	// Change class of output pin
+	UEdGraphPin* ResultPin = GetResultPin();
+	ResultPin->PinType.PinSubCategoryObject = InClass;
+}
+
+UClass* UK2Node_ConstructObjectFromClass::GetClassToSpawn(const TArray<UEdGraphPin*>* InPinsToSearch /*=NULL*/) const
+{
+	UClass* UseSpawnClass = NULL;
+	const TArray<UEdGraphPin*>* PinsToSearch = InPinsToSearch ? InPinsToSearch : &Pins;
+
+	UEdGraphPin* ClassPin = GetClassPin(PinsToSearch);
+	if(ClassPin && ClassPin->DefaultObject != NULL && ClassPin->LinkedTo.Num() == 0)
+	{
+		UseSpawnClass = CastChecked<UClass>(ClassPin->DefaultObject);
+	}
+
+	return UseSpawnClass;
+}
+
+void UK2Node_ConstructObjectFromClass::ReallocatePinsDuringReconstruction(TArray<UEdGraphPin*>& OldPins) 
+{
+	AllocateDefaultPins();
+	UClass* UseSpawnClass = GetClassToSpawn(&OldPins);
+	if( UseSpawnClass != NULL )
+	{
+		CreatePinsForClass(UseSpawnClass);
+	}
+}
+
+bool UK2Node_ConstructObjectFromClass::IsSpawnVarPin(UEdGraphPin* Pin)
+{
+	const UEdGraphSchema_K2* K2Schema = GetDefault<UEdGraphSchema_K2>();
+
+	return(	Pin->PinName != K2Schema->PN_Execute &&
+			Pin->PinName != K2Schema->PN_Then &&
+			Pin->PinName != K2Schema->PN_ReturnValue &&
+			Pin->PinName != FK2Node_ConstructObjectFromClassHelper::ClassPinName &&
+			Pin->PinName != FK2Node_ConstructObjectFromClassHelper::WorldContextPinName );
+}
+
+void UK2Node_ConstructObjectFromClass::PinDefaultValueChanged(UEdGraphPin* ChangedPin) 
+{
+	if (ChangedPin->PinName == FK2Node_ConstructObjectFromClassHelper::ClassPinName)
+	{
+		const UEdGraphSchema_K2* K2Schema = GetDefault<UEdGraphSchema_K2>();
+
+		// Because the archetype has changed, we break the output link as the output pin type will change
+		UEdGraphPin* ResultPin = GetResultPin();
+		ResultPin->BreakAllPinLinks();
+
+		// Remove all pins related to archetype variables
+		TArray<UEdGraphPin*> OldPins = Pins;
+		for (int32 i = 0; i < OldPins.Num(); i++)
+		{
+			UEdGraphPin* OldPin = OldPins[i];
+			if (IsSpawnVarPin(OldPin))
+			{
+				OldPin->BreakAllPinLinks();
+				Pins.Remove(OldPin);
+			}
+		}
+
+		UClass* UseSpawnClass = GetClassToSpawn();
+		if (UseSpawnClass != NULL)
+		{
+			CreatePinsForClass(UseSpawnClass);
+		}
+
+		// Refresh the UI for the graph so the pin changes show up
+		UEdGraph* Graph = GetGraph();
+		Graph->NotifyGraphChanged();
+
+		// Mark dirty
+		FBlueprintEditorUtils::MarkBlueprintAsModified(GetBlueprint());
+	}
+}
+
+FString UK2Node_ConstructObjectFromClass::GetTooltip() const
+{
+	return NodeTooltip;
+}
+
+UEdGraphPin* UK2Node_ConstructObjectFromClass::GetThenPin()const
+{
+	const UEdGraphSchema_K2* K2Schema = GetDefault<UEdGraphSchema_K2>();
+
+	UEdGraphPin* Pin = FindPinChecked(K2Schema->PN_Then);
+	check(Pin->Direction == EGPD_Output);
+	return Pin;
+}
+
+UEdGraphPin* UK2Node_ConstructObjectFromClass::GetClassPin(const TArray<UEdGraphPin*>* InPinsToSearch /*= NULL*/) const
+{
+	const TArray<UEdGraphPin*>* PinsToSearch = InPinsToSearch ? InPinsToSearch : &Pins;
+
+	UEdGraphPin* Pin = NULL;
+	for( auto PinIt = PinsToSearch->CreateConstIterator(); PinIt; ++PinIt )
+	{
+		UEdGraphPin* TestPin = *PinIt;
+		if( TestPin && TestPin->PinName == FK2Node_ConstructObjectFromClassHelper::ClassPinName )
+		{
+			Pin = TestPin;
+			break;
+		}
+	}
+	check(Pin == NULL || Pin->Direction == EGPD_Input);
+	return Pin;
+}
+
+UEdGraphPin* UK2Node_ConstructObjectFromClass::GetWorldContextPin() const
+{
+	UEdGraphPin* Pin = FindPin(FK2Node_ConstructObjectFromClassHelper::WorldContextPinName);
+	check(Pin == NULL || Pin->Direction == EGPD_Input);
+	return Pin;
+}
+
+UEdGraphPin* UK2Node_ConstructObjectFromClass::GetResultPin() const
+{
+	const UEdGraphSchema_K2* K2Schema = GetDefault<UEdGraphSchema_K2>();
+
+	UEdGraphPin* Pin = FindPinChecked(K2Schema->PN_ReturnValue);
+	check(Pin->Direction == EGPD_Output);
+	return Pin;
+}
+
+FLinearColor UK2Node_ConstructObjectFromClass::GetNodeTitleColor() const
+{
+	return Super::GetNodeTitleColor();
+}
+
+FText UK2Node_ConstructObjectFromClass::GetBaseNodeTitle() const
+{
+	return NSLOCTEXT("K2Node", "ConstructObject_BaseTitle", "Construct Object from Class");
+}
+
+FText UK2Node_ConstructObjectFromClass::GetNodeTitleFormat() const
+{
+	return NSLOCTEXT("K2Node", "Construct", "Construct {ClassName}");
+}
+
+FText UK2Node_ConstructObjectFromClass::GetNodeTitle(ENodeTitleType::Type TitleType) const
+{
+	FText NodeTitle = GetBaseNodeTitle();
+	if (TitleType != ENodeTitleType::ListView)
+	{
+		FText SpawnString = NSLOCTEXT("K2Node", "None", "NONE");
+		if (UEdGraphPin* ClassPin = GetClassPin())
+		{
+			if(ClassPin->LinkedTo.Num() > 0)
+			{
+				// Blueprint will be determined dynamically, so we don't have the name in this case
+				SpawnString = FText::GetEmpty();
+			}
+			else if(ClassPin->DefaultObject != NULL)
+			{
+				SpawnString = FText::FromString(ClassPin->DefaultObject->GetName());
+			}
+		}
+		
+		FFormatNamedArguments Args;
+		Args.Add(TEXT("ClassName"), SpawnString);
+		
+		NodeTitle = FText::Format(GetNodeTitleFormat(), Args);
+	}
+
+	return NodeTitle;
+}
+
+bool UK2Node_ConstructObjectFromClass::IsCompatibleWithGraph(const UEdGraph* TargetGraph) const 
+{
+	UBlueprint* Blueprint = FBlueprintEditorUtils::FindBlueprintForGraph(TargetGraph);
+	return Super::IsCompatibleWithGraph(TargetGraph) && (!Blueprint || FBlueprintEditorUtils::FindUserConstructionScript(Blueprint) != TargetGraph);
+}
+
+void UK2Node_ConstructObjectFromClass::GetNodeAttributes( TArray<TKeyValuePair<FString, FString>>& OutNodeAttributes ) const
+{
+	UClass* ClassToSpawn = GetClassToSpawn();
+	const FString ClassToSpawnStr = ClassToSpawn ? ClassToSpawn->GetName() : TEXT( "InvalidClass" );
+	OutNodeAttributes.Add( TKeyValuePair<FString, FString>( TEXT( "Type" ), TEXT( "ConstructObjectFromClass" ) ));
+	OutNodeAttributes.Add( TKeyValuePair<FString, FString>( TEXT( "Class" ), GetClass()->GetName() ));
+	OutNodeAttributes.Add( TKeyValuePair<FString, FString>( TEXT( "Name" ), GetName() ));
+	OutNodeAttributes.Add( TKeyValuePair<FString, FString>( TEXT( "ObjectClass" ), ClassToSpawnStr ));
+}
+
+void UK2Node_ConstructObjectFromClass::GetMenuActions(TArray<UBlueprintNodeSpawner*>& ActionListOut) const
+{
+	UBlueprintNodeSpawner* NodeSpawner = UBlueprintNodeSpawner::Create(GetClass());
+	check(NodeSpawner != nullptr);
+
+	ActionListOut.Add(NodeSpawner);
+}
+
+FText UK2Node_ConstructObjectFromClass::GetMenuCategory() const
+{
+	return FEditorCategoryUtils::GetCommonCategory(FCommonEditorCategory::Gameplay);
+}
+
+FNodeHandlingFunctor* UK2Node_ConstructObjectFromClass::CreateNodeHandler(FKismetCompilerContext& CompilerContext) const
+{
+	return new FNodeHandlingFunctor(CompilerContext);
+}
+
+void UK2Node_ConstructObjectFromClass::ExpandNode(class FKismetCompilerContext& CompilerContext, UEdGraph* SourceGraph)
+{
+	Super::ExpandNode(CompilerContext, SourceGraph);	
+}
+
+bool UK2Node_ConstructObjectFromClass::HasExternalBlueprintDependencies(TArray<class UStruct*>* OptionalOutput) const
+{
+	UClass* SourceClass = GetClassToSpawn();
+	const UBlueprint* SourceBlueprint = GetBlueprint();
+	const bool bResult = (SourceClass != NULL) && (SourceClass->ClassGeneratedBy != NULL) && (SourceClass->ClassGeneratedBy != SourceBlueprint);
+	if (bResult && OptionalOutput)
+	{
+		OptionalOutput->Add(SourceClass);
+	}
+	return bResult || Super::HasExternalBlueprintDependencies(OptionalOutput);
+}
+
+#undef LOCTEXT_NAMESPACE
\ No newline at end of file
diff --git a/Engine/Source/Editor/BlueprintGraph/Public/BlueprintGraphDefinitions.h b/Engine/Source/Editor/BlueprintGraph/Public/BlueprintGraphDefinitions.h
--- a/Engine/Source/Editor/BlueprintGraph/Public/BlueprintGraphDefinitions.h
+++ b/Engine/Source/Editor/BlueprintGraph/Public/BlueprintGraphDefinitions.h
@@ -18,6 +18,7 @@
 #include "K2Node_CallParentFunction.h"
 #include "K2Node_ClearDelegate.h"
 #include "K2Node_CommutativeAssociativeBinaryOperator.h"
+#include "K2Node_ConstructObjectFromClass.h"
 #include "K2Node_DoOnceMultiInput.h"
 #include "K2Node_ComponentBoundEvent.h"
 #include "K2Node_Composite.h"
diff --git a/Engine/Source/Editor/UMGEditor/Private/Nodes/K2Node_CreateWidget.cpp b/Engine/Source/Editor/UMGEditor/Private/Nodes/K2Node_CreateWidget.cpp
new file mode 100644
--- /dev/null
+++ b/Engine/Source/Editor/UMGEditor/Private/Nodes/K2Node_CreateWidget.cpp
@@ -0,0 +1,164 @@
+// Copyright 1998-2014 Epic Games, Inc. All Rights Reserved.
+
+#include "UMGEditorPrivatePCH.h"
+
+#include "KismetCompiler.h"
+#include "BlueprintNodeSpawner.h"
+#include "EditorCategoryUtils.h"
+#include "K2ActionMenuBuilder.h"
+
+#include "K2Node_CreateWidget.h"
+
+#define LOCTEXT_NAMESPACE "UMG"
+
+struct FK2Node_CreateWidgetHelper
+{
+	static FString OwningPlayerPinName;
+};
+
+FString FK2Node_CreateWidgetHelper::OwningPlayerPinName(TEXT("OwningPlayer"));
+
+UK2Node_CreateWidget::UK2Node_CreateWidget(const class FPostConstructInitializeProperties& PCIP)
+	: Super(PCIP)
+{
+	NodeTooltip = LOCTEXT("NodeTooltip", "Creates a new widget").ToString();
+}
+
+void UK2Node_CreateWidget::AllocateDefaultPins()
+{
+	Super::AllocateDefaultPins();
+
+	const UEdGraphSchema_K2* K2Schema = GetDefault<UEdGraphSchema_K2>();
+
+	// OwningPlayer pin
+	UEdGraphPin* OwningPlayerPin = CreatePin(EGPD_Input, K2Schema->PC_Object, TEXT(""), APlayerController::StaticClass(), false, false, FK2Node_CreateWidgetHelper::OwningPlayerPinName);
+	SetPinToolTip(*OwningPlayerPin, LOCTEXT("OwningPlayerPinDescription", "The player that 'owns' the widget."));
+}
+
+FLinearColor UK2Node_CreateWidget::GetNodeTitleColor() const
+{
+	return Super::GetNodeTitleColor();
+}
+
+FText UK2Node_CreateWidget::GetBaseNodeTitle() const
+{
+	return LOCTEXT("CreateWidget_BaseTitle", "Create Widget");
+}
+
+FText UK2Node_CreateWidget::GetNodeTitleFormat() const
+{
+	return LOCTEXT("CreateWidget", "Create {ClassName} Widget");
+}
+
+UClass* UK2Node_CreateWidget::GetClassPinBaseClass() const
+{
+	return UUserWidget::StaticClass();
+}
+
+void UK2Node_CreateWidget::GetMenuEntries(FGraphContextMenuBuilder& ContextMenuBuilder) const
+{
+	UK2Node_CreateWidget* TemplateNode = NewObject<UK2Node_CreateWidget>(GetTransientPackage(), GetClass());
+
+	const FString Category = TEXT("");
+	const FText   MenuDesc = LOCTEXT("CreateWidgetMenuOption", "Create Widget...");
+	const FString Tooltip  = TEXT("Create a new UI widget");
+
+	TSharedPtr<FEdGraphSchemaAction_K2NewNode> NodeAction = FK2ActionMenuBuilder::AddNewNodeAction(ContextMenuBuilder, Category, MenuDesc, Tooltip);
+	NodeAction->NodeTemplate = TemplateNode;
+}
+
+UEdGraphPin* UK2Node_CreateWidget::GetOwningPlayerPin() const
+{
+	UEdGraphPin* Pin = FindPin(FK2Node_CreateWidgetHelper::OwningPlayerPinName);
+	check(Pin == NULL || Pin->Direction == EGPD_Input);
+	return Pin;
+}
+
+bool UK2Node_CreateWidget::IsSpawnVarPin(UEdGraphPin* Pin)
+{
+	return( Super::IsSpawnVarPin(Pin) &&
+		Pin->PinName != FK2Node_CreateWidgetHelper::OwningPlayerPinName );
+}
+
+void UK2Node_CreateWidget::ExpandNode(class FKismetCompilerContext& CompilerContext, UEdGraph* SourceGraph)
+{
+	Super::ExpandNode(CompilerContext, SourceGraph);
+
+	if ( CompilerContext.bIsFullCompile )
+	{
+		static FName Create_FunctionName = GET_FUNCTION_NAME_CHECKED(UWidgetBlueprintLibrary, Create);
+		static FString WorldContextObject_ParamName = FString(TEXT("WorldContextObject"));
+		static FString WidgetType_ParamName = FString(TEXT("WidgetType"));
+		static FString OwningPlayer_ParamName = FString(TEXT("OwningPlayer"));
+
+		UK2Node_CreateWidget* CreateWidgetNode = this;
+		UEdGraphPin* SpawnNodeExec = CreateWidgetNode->GetExecPin();
+		UEdGraphPin* SpawnWorldContextPin = CreateWidgetNode->GetWorldContextPin();
+		UEdGraphPin* SpawnOwningPlayerPin = CreateWidgetNode->GetOwningPlayerPin();
+		UEdGraphPin* SpawnClassPin = CreateWidgetNode->GetClassPin();
+		UEdGraphPin* SpawnNodeThen = CreateWidgetNode->GetThenPin();
+		UEdGraphPin* SpawnNodeResult = CreateWidgetNode->GetResultPin();
+
+		UClass* SpawnClass = ( SpawnClassPin != NULL ) ? Cast<UClass>(SpawnClassPin->DefaultObject) : NULL;
+		if ( ( 0 == SpawnClassPin->LinkedTo.Num() ) && ( NULL == SpawnClass ) )
+		{
+			CompilerContext.MessageLog.Error(*LOCTEXT("CreateWidgetNodeMissingClass_Error", "Spawn node @@ must have a class specified.").ToString(), CreateWidgetNode);
+			// we break exec links so this is the only error we get, don't want the CreateWidget node being considered and giving 'unexpected node' type warnings
+			CreateWidgetNode->BreakAllNodeLinks();
+			return;
+		}
+
+		//////////////////////////////////////////////////////////////////////////
+		// create 'UWidgetBlueprintLibrary::Create' call node
+		UK2Node_CallFunction* CallCreateNode = CompilerContext.SpawnIntermediateNode<UK2Node_CallFunction>(CreateWidgetNode, SourceGraph);
+		CallCreateNode->FunctionReference.SetExternalMember(Create_FunctionName, UWidgetBlueprintLibrary::StaticClass());
+		CallCreateNode->AllocateDefaultPins();
+
+		UEdGraphPin* CallCreateExec = CallCreateNode->GetExecPin();
+		UEdGraphPin* CallCreateWorldContextPin = CallCreateNode->FindPinChecked(WorldContextObject_ParamName);
+		UEdGraphPin* CallCreateWidgetTypePin = CallCreateNode->FindPinChecked(WidgetType_ParamName);
+		UEdGraphPin* CallCreateOwningPlayerPin = CallCreateNode->FindPinChecked(OwningPlayer_ParamName);
+		UEdGraphPin* CallCreateResult = CallCreateNode->GetReturnValuePin();
+
+		// Move 'exec' connection from create widget node to 'UWidgetBlueprintLibrary::Create'
+		CompilerContext.MovePinLinksToIntermediate(*SpawnNodeExec, *CallCreateExec);
+
+		if ( SpawnClassPin->LinkedTo.Num() > 0 )
+		{
+			// Copy the 'blueprint' connection from the spawn node to 'UWidgetBlueprintLibrary::Create'
+			CompilerContext.MovePinLinksToIntermediate(*SpawnClassPin, *CallCreateWidgetTypePin);
+		}
+		else
+		{
+			// Copy blueprint literal onto 'UWidgetBlueprintLibrary::Create' call 
+			CallCreateWidgetTypePin->DefaultObject = SpawnClass;
+		}
+
+		// Copy the world context connection from the spawn node to 'UWidgetBlueprintLibrary::Create' if necessary
+		if ( SpawnWorldContextPin )
+		{
+			CompilerContext.MovePinLinksToIntermediate(*SpawnWorldContextPin, *CallCreateWorldContextPin);
+		}
+
+		// Copy the 'Owning Player' connection from the spawn node to 'UWidgetBlueprintLibrary::Create'
+		CompilerContext.MovePinLinksToIntermediate(*SpawnOwningPlayerPin, *CallCreateOwningPlayerPin);
+
+		// Move result connection from spawn node to 'finish spawn'
+		CallCreateResult->PinType = SpawnNodeResult->PinType; // Copy type so it uses the right actor subclass
+		CompilerContext.MovePinLinksToIntermediate(*SpawnNodeResult, *CallCreateResult);
+
+		//////////////////////////////////////////////////////////////////////////
+		// create 'set var' nodes
+
+		// Get 'result' pin from 'begin spawn', this is the actual actor we want to set properties on
+		UEdGraphPin* LastThen = FKismetCompilerUtilities::GenerateAssignmentNodes(CompilerContext, SourceGraph, CallCreateNode, CreateWidgetNode, CallCreateResult, GetClassToSpawn());
+
+		// Move 'then' connection from create widget node to the last 'then'
+		CompilerContext.MovePinLinksToIntermediate(*SpawnNodeThen, *LastThen);
+
+		// Break any links to the expanded node
+		CreateWidgetNode->BreakAllNodeLinks();
+	}
+}
+
+#undef LOCTEXT_NAMESPACE
diff --git a/Engine/Source/Editor/UMGEditor/Private/Nodes/K2Node_CreateWidget.h b/Engine/Source/Editor/UMGEditor/Private/Nodes/K2Node_CreateWidget.h
new file mode 100644
--- /dev/null
+++ b/Engine/Source/Editor/UMGEditor/Private/Nodes/K2Node_CreateWidget.h
@@ -0,0 +1,35 @@
+// Copyright 1998-2014 Epic Games, Inc. All Rights Reserved.
+
+#pragma once
+
+#include "K2Node.h"
+#include "K2Node_CreateWidget.generated.h"
+
+UCLASS()
+class UMGEDITOR_API UK2Node_CreateWidget : public UK2Node_ConstructObjectFromClass
+{
+	GENERATED_UCLASS_BODY()
+
+	// Begin UEdGraphNode interface.
+	virtual void AllocateDefaultPins() override;
+	virtual FLinearColor GetNodeTitleColor() const override;
+	virtual void ExpandNode(class FKismetCompilerContext& CompilerContext, UEdGraph* SourceGraph) override;
+	// End UEdGraphNode interface.
+
+	// Begin UK2Node interface
+	void GetMenuEntries(FGraphContextMenuBuilder& ContextMenuBuilder) const;
+	// End UK2Node interface.
+
+	/** Get the owning player pin */
+	UEdGraphPin* GetOwningPlayerPin() const;
+
+protected:
+	/** Gets the default node title when no class is selected */
+	virtual FText GetBaseNodeTitle() const;
+	/** Gets the node title when a class has been selected. */
+	virtual FText GetNodeTitleFormat() const;
+	/** Gets base class to use for the 'class' pin.  UObject by default. */
+	virtual UClass* GetClassPinBaseClass() const;
+	/**  */
+	virtual bool IsSpawnVarPin(UEdGraphPin* Pin) override;
+};
diff --git a/Engine/Source/Editor/UMGEditor/Private/UMGEditorPrivatePCH.h b/Engine/Source/Editor/UMGEditor/Private/UMGEditorPrivatePCH.h
--- a/Engine/Source/Editor/UMGEditor/Private/UMGEditorPrivatePCH.h
+++ b/Engine/Source/Editor/UMGEditor/Private/UMGEditorPrivatePCH.h
@@ -48,3 +48,6 @@
 #include "WidgetTemplateDragDropOp.h"
 
 #include "WidgetBlueprintEditorUtils.h"
+
+#include "WidgetBlueprintLibrary.h"
+
diff --git a/Engine/Source/Runtime/UMG/Private/UserWidget.cpp b/Engine/Source/Runtime/UMG/Private/UserWidget.cpp
--- a/Engine/Source/Runtime/UMG/Private/UserWidget.cpp
+++ b/Engine/Source/Runtime/UMG/Private/UserWidget.cpp
@@ -488,13 +488,13 @@ const FLocalPlayerContext& UUserWidget::GetPlayerContext() const
 	return PlayerContext;
 }
 
-ULocalPlayer* UUserWidget::GetLocalPlayer() const
+ULocalPlayer* UUserWidget::GetOwningLocalPlayer() const
 {
 	APlayerController* PC = PlayerContext.IsValid() ? PlayerContext.GetPlayerController() : NULL;
 	return PC ? Cast<ULocalPlayer>(PC->Player) : NULL;
 }
 
-APlayerController* UUserWidget::GetPlayerController() const
+APlayerController* UUserWidget::GetOwningPlayer() const
 {
 	return PlayerContext.IsValid() ? PlayerContext.GetPlayerController() : NULL;
 }
diff --git a/Engine/Source/Runtime/UMG/Private/WidgetBlueprintLibrary.cpp b/Engine/Source/Runtime/UMG/Private/WidgetBlueprintLibrary.cpp
--- a/Engine/Source/Runtime/UMG/Private/WidgetBlueprintLibrary.cpp
+++ b/Engine/Source/Runtime/UMG/Private/WidgetBlueprintLibrary.cpp
@@ -158,7 +158,7 @@ FEventReply UWidgetBlueprintLibrary::DetectDrag(UPARAM(ref) FEventReply& Reply,
 	return Reply;
 }
 
-FEventReply UWidgetBlueprintLibrary::IfKeyPressedDetectDrag(const FPointerEvent& PointerEvent, UWidget* WidgetDetectingDrag, FKey DragKey)
+FEventReply UWidgetBlueprintLibrary::DetectDragIfPressed(const FPointerEvent& PointerEvent, UWidget* WidgetDetectingDrag, FKey DragKey)
 {
 	if ( PointerEvent.GetEffectingButton() == DragKey )
 	{
diff --git a/Engine/Source/Runtime/UMG/Public/Blueprint/UserWidget.h b/Engine/Source/Runtime/UMG/Public/Blueprint/UserWidget.h
--- a/Engine/Source/Runtime/UMG/Public/Blueprint/UserWidget.h
+++ b/Engine/Source/Runtime/UMG/Public/Blueprint/UserWidget.h
@@ -139,11 +139,11 @@ public:
 
 	/** Gets the local player associated with this UI. */
 	UFUNCTION(BlueprintCallable, Category="Player")
-	class ULocalPlayer* GetLocalPlayer() const;
+	class ULocalPlayer* GetOwningLocalPlayer() const;
 
 	/** Gets the player controller associated with this UI. */
 	UFUNCTION(BlueprintCallable, Category="Player")
-	class APlayerController* GetPlayerController() const;
+	class APlayerController* GetOwningPlayer() const;
 
 	/** Called when the widget is constructed */
 	UFUNCTION(BlueprintImplementableEvent, Category="User Interface")
diff --git a/Engine/Source/Runtime/UMG/Public/Blueprint/WidgetBlueprintLibrary.h b/Engine/Source/Runtime/UMG/Public/Blueprint/WidgetBlueprintLibrary.h
--- a/Engine/Source/Runtime/UMG/Public/Blueprint/WidgetBlueprintLibrary.h
+++ b/Engine/Source/Runtime/UMG/Public/Blueprint/WidgetBlueprintLibrary.h
@@ -1,6 +1,7 @@
 // Copyright 1998-2014 Epic Games, Inc. All Rights Reserved.
 
 #pragma once
+
 #include "WidgetBlueprintLibrary.generated.h"
 
 UCLASS(MinimalAPI)
@@ -9,7 +10,7 @@ class UWidgetBlueprintLibrary : public UBlueprintFunctionLibrary
 	GENERATED_UCLASS_BODY()
 
 	/** Creates a widget */
-	UFUNCTION(BlueprintCallable, meta=( HidePin="WorldContextObject", DefaultToSelf="WorldContextObject", FriendlyName = "Create Widget" ), Category="User Interface|Widget")
+	UFUNCTION(BlueprintCallable, meta=( HidePin="WorldContextObject", DefaultToSelf="WorldContextObject", FriendlyName = "Create Widget", BlueprintInternalUseOnly = "true" ), Category="User Interface|Widget")
 	static class UUserWidget* Create(UObject* WorldContextObject, TSubclassOf<class UUserWidget> WidgetType, APlayerController* OwningPlayer);
 
 	UFUNCTION(BlueprintCallable, Category="Focus")
@@ -71,7 +72,7 @@ class UWidgetBlueprintLibrary : public UBlueprintFunctionLibrary
 	static FEventReply DetectDrag(UPARAM(ref) FEventReply& Reply, UWidget* WidgetDetectingDrag, FKey DragKey);
 
 	UFUNCTION(BlueprintCallable, meta=( HidePin="WidgetDetectingDrag", DefaultToSelf="WidgetDetectingDrag" ), Category="Widget|Event Reply")
-	static FEventReply IfKeyPressedDetectDrag(const FPointerEvent& PointerEvent, UWidget* WidgetDetectingDrag, FKey DragKey);
+	static FEventReply DetectDragIfPressed(const FPointerEvent& PointerEvent, UWidget* WidgetDetectingDrag, FKey DragKey);
 
 	/**
 	 * An event should return FReply::Handled().EndDragDrop() to request that the current drag/drop operation be terminated.
