DataTable/CurveTable Blueprint nodes: Drop down box for selecting row name

Remove requirement for blueprint node name pins to be non-numeric values (this requirement seems to no longer be present in other areas of the engine)

Modified SGraphPinNameList to use new SNameComboBox class

diff --git a/Engine/Source/Editor/BlueprintGraph/Classes/EdGraphSchema_K2.h b/Engine/Source/Editor/BlueprintGraph/Classes/EdGraphSchema_K2.h
--- a/Engine/Source/Editor/BlueprintGraph/Classes/EdGraphSchema_K2.h
+++ b/Engine/Source/Editor/BlueprintGraph/Classes/EdGraphSchema_K2.h
@@ -136,6 +136,8 @@ public:
 	/** Metadata string that indicates the specified event can be triggered in editor */
 	static const FName MD_CallInEditor;
 
+	/** Metadata to identify an DataTable Pin. Depending on which DataTable is selected, we display different RowName options */
+	static const FName MD_DataTablePin;
 private:
 	// This class should never be instantiated
 	FBlueprintMetadata() {}
diff --git a/Engine/Source/Editor/BlueprintGraph/Classes/K2Node_CallDataTableFunction.h b/Engine/Source/Editor/BlueprintGraph/Classes/K2Node_CallDataTableFunction.h
new file mode 100644
--- /dev/null
+++ b/Engine/Source/Editor/BlueprintGraph/Classes/K2Node_CallDataTableFunction.h
@@ -0,0 +1,16 @@
+// Copyright 1998-2014 Epic Games, Inc. All Rights Reserved.
+
+#pragma once
+#include "K2Node_CallFunction.h"
+#include "K2Node_CallDataTableFunction.generated.h"
+
+UCLASS(MinimalAPI)
+class UK2Node_CallDataTableFunction : public UK2Node_CallFunction
+{
+	GENERATED_UCLASS_BODY()
+
+	// Begin EdGraphNode interface
+	virtual void PinDefaultValueChanged(UEdGraphPin* Pin) override;
+	virtual void NotifyPinConnectionListChanged(UEdGraphPin* Pin) override;
+	// End EdGraphNode interface
+};
\ No newline at end of file
diff --git a/Engine/Source/Editor/BlueprintGraph/Classes/K2Node_GetDataTableRow.h b/Engine/Source/Editor/BlueprintGraph/Classes/K2Node_GetDataTableRow.h
--- a/Engine/Source/Editor/BlueprintGraph/Classes/K2Node_GetDataTableRow.h
+++ b/Engine/Source/Editor/BlueprintGraph/Classes/K2Node_GetDataTableRow.h
@@ -17,6 +17,7 @@ class UK2Node_GetDataTableRow : public UK2Node
 	virtual void AllocateDefaultPins() override;
 	virtual FText GetNodeTitle(ENodeTitleType::Type TitleType) const override;
 	virtual void PinDefaultValueChanged(UEdGraphPin* Pin) override;
+	virtual void NotifyPinConnectionListChanged(UEdGraphPin* Pin) override;
 	virtual FString GetTooltip() const override;
 	virtual void ExpandNode(class FKismetCompilerContext& CompilerContext, UEdGraph* SourceGraph) override;
 	virtual FName GetPaletteIcon(FLinearColor& OutColor) const override
@@ -53,6 +54,9 @@ class UK2Node_GetDataTableRow : public UK2Node
 	/** Get the type of the TableRow to return */
 	UScriptStruct* GetDataTableRowStructType(const TArray<UEdGraphPin*>* InPinsToSearch = NULL) const;
 
+	/** Get the list of possible RowNames */
+	void GetDataTableRowNameList(TArray<TSharedPtr<FName>>& OutList, const TArray<UEdGraphPin*>* InPinsToSearch = NULL) const;
+
 private:
 	/**
 	 * Takes the specified "MutatablePin" and sets its 'PinToolTip' field (according
diff --git a/Engine/Source/Editor/BlueprintGraph/Private/BlueprintFunctionNodeSpawner.cpp b/Engine/Source/Editor/BlueprintGraph/Private/BlueprintFunctionNodeSpawner.cpp
--- a/Engine/Source/Editor/BlueprintGraph/Private/BlueprintFunctionNodeSpawner.cpp
+++ b/Engine/Source/Editor/BlueprintGraph/Private/BlueprintFunctionNodeSpawner.cpp
@@ -16,6 +16,7 @@ UBlueprintFunctionNodeSpawner* UBlueprintFunctionNodeSpawner::Create(UFunction c
 	bool const bHasArrayPointerParms = Function->HasMetaData(TEXT("ArrayParm"));
 	bool const bIsCommutativeAssociativeBinaryOp = Function->HasMetaData(FBlueprintMetadata::MD_CommutativeAssociativeBinaryOperator);
 	bool const bIsMaterialParamCollectionFunc = Function->HasMetaData(FBlueprintMetadata::MD_MaterialParameterCollectionFunction);
+	bool const bIsDataTableFunc = Function->HasMetaData(FBlueprintMetadata::MD_DataTablePin);
 
 	TSubclassOf<UK2Node_CallFunction> NodeClass;
 	if (bIsCommutativeAssociativeBinaryOp && bIsPure)
@@ -26,6 +27,10 @@ UBlueprintFunctionNodeSpawner* UBlueprintFunctionNodeSpawner::Create(UFunction c
 	{
 		NodeClass = UK2Node_CallMaterialParameterCollectionFunction::StaticClass();
 	}
+	else if (bIsDataTableFunc)
+	{
+		NodeClass = UK2Node_CallDataTableFunction::StaticClass();
+	}
 	// @TODO:
 	//   else if CallOnMember => UK2Node_CallFunctionOnMember
 	//   else if bIsParentContext => UK2Node_CallParentFunction
diff --git a/Engine/Source/Editor/BlueprintGraph/Private/BlueprintGraphPanelPinFactory.cpp b/Engine/Source/Editor/BlueprintGraph/Private/BlueprintGraphPanelPinFactory.cpp
new file mode 100644
--- /dev/null
+++ b/Engine/Source/Editor/BlueprintGraph/Private/BlueprintGraphPanelPinFactory.cpp
@@ -0,0 +1,76 @@
+// Copyright 1998-2014 Epic Games, Inc. All Rights Reserved.
+
+#pragma once
+
+#include "BlueprintGraphPrivatePCH.h"
+#include "BlueprintGraphPanelPinFactory.h"
+#include "K2Node_GetDataTableRow.h"
+#include "SGraphPinNameList.h"
+#include "K2Node_CallFunction.h"
+	
+TSharedPtr<class SGraphPin> FBlueprintGraphPanelPinFactory::CreatePin(class UEdGraphPin* InPin) const
+{
+	const UEdGraphSchema_K2* K2Schema = GetDefault<UEdGraphSchema_K2>();
+	if (InPin->PinType.PinCategory == K2Schema->PC_Name)
+	{
+		UObject* Outer = InPin->GetOuter();
+
+		// Create drop down combo boxes for DataTable and CurveTable RowName pins
+		const UEdGraphPin* DataTablePin = nullptr;
+		if (Outer->IsA(UK2Node_CallFunction::StaticClass()))
+		{
+			const UK2Node_CallFunction* CallFunctionNode = CastChecked<UK2Node_CallFunction>(Outer);
+			if (CallFunctionNode)
+			{
+				const UFunction* FunctionToCall = CallFunctionNode->GetTargetFunction();
+				FString DataTablePinName = FunctionToCall->GetMetaData(FBlueprintMetadata::MD_DataTablePin);
+				DataTablePin = CallFunctionNode->FindPinChecked(DataTablePinName);
+			}
+		}
+		else if (Outer->IsA(UK2Node_GetDataTableRow::StaticClass()))
+		{
+			const UK2Node_GetDataTableRow* GetDataTableRowNode = CastChecked<UK2Node_GetDataTableRow>(Outer);
+			DataTablePin = GetDataTableRowNode->GetDataTablePin();
+		}
+		if (DataTablePin)
+		{
+			if (DataTablePin && DataTablePin->DefaultObject != nullptr && DataTablePin->LinkedTo.Num() == 0)
+			{
+				if (DataTablePin->DefaultObject->IsA(UDataTable::StaticClass()))
+				{
+					UDataTable* DataTable = (UDataTable*)DataTablePin->DefaultObject;
+					if (DataTable)
+					{
+						TArray<TSharedPtr<FName>> RowNames;
+						/** Extract all the row names from the RowMap */
+						for (TMap<FName, uint8*>::TConstIterator Iterator(DataTable->RowMap); Iterator; ++Iterator)
+						{
+							/** Create a simple array of the row names */
+							TSharedPtr<FName> RowNameItem = MakeShareable(new FName(Iterator.Key()));
+							RowNames.Add(RowNameItem);
+						}
+						return SNew(SGraphPinNameList, InPin, RowNames);
+					}
+				}
+				if (DataTablePin->DefaultObject->IsA(UCurveTable::StaticClass()))
+				{
+					UCurveTable* CurveTable = (UCurveTable*)DataTablePin->DefaultObject;
+					if (CurveTable)
+					{
+						TArray<TSharedPtr<FName>> RowNames;
+						/** Extract all the row names from the RowMap */
+						for (TMap<FName, FRichCurve*>::TConstIterator Iterator(CurveTable->RowMap); Iterator; ++Iterator)
+						{
+							/** Create a simple array of the row names */
+							TSharedPtr<FName> RowNameItem = MakeShareable(new FName(Iterator.Key()));
+							RowNames.Add(RowNameItem);
+						}
+						return SNew(SGraphPinNameList, InPin, RowNames);
+					}
+				}
+			}
+		}
+	}
+
+	return nullptr;
+}
diff --git a/Engine/Source/Editor/BlueprintGraph/Private/EdGraphSchema_K2.cpp b/Engine/Source/Editor/BlueprintGraph/Private/EdGraphSchema_K2.cpp
--- a/Engine/Source/Editor/BlueprintGraph/Private/EdGraphSchema_K2.cpp
+++ b/Engine/Source/Editor/BlueprintGraph/Private/EdGraphSchema_K2.cpp
@@ -83,6 +83,8 @@ const FName FBlueprintMetadata::MD_Tooltip(TEXT("Tooltip"));
 
 const FName FBlueprintMetadata::MD_CallInEditor(TEXT("CallInEditor"));
 
+const FName FBlueprintMetadata::MD_DataTablePin(TEXT("DataTablePin"));
+
 //////////////////////////////////////////////////////////////////////////
 
 #define LOCTEXT_NAMESPACE "KismetSchema"
@@ -1942,10 +1944,7 @@ bool UEdGraphSchema_K2::DefaultValueSimpleValidation(const FEdGraphPinType& PinT
 	}
 	else if (PinCategory == PC_Name)
 	{
-		if( NewDefaultValue.IsNumeric() )
-		{
-			DVSV_RETURN_MSG( FString::Printf(TEXT("Invalid default name for pin %s"), *(PinName)) ); 
-		}
+		// Anything is allowed
 	}
 	else if ((PinCategory == PC_Object) || (PinCategory == PC_Interface))
 	{
diff --git a/Engine/Source/Editor/BlueprintGraph/Private/K2ActionMenuBuilder.cpp b/Engine/Source/Editor/BlueprintGraph/Private/K2ActionMenuBuilder.cpp
--- a/Engine/Source/Editor/BlueprintGraph/Private/K2ActionMenuBuilder.cpp
+++ b/Engine/Source/Editor/BlueprintGraph/Private/K2ActionMenuBuilder.cpp
@@ -652,6 +652,7 @@ void FK2ActionMenuBuilder::AddSpawnInfoForFunction(const UFunction* Function, bo
 	const bool bIsPure = Function->HasAllFunctionFlags(FUNC_BlueprintPure);
 	const bool bCommutativeAssociativeBinaryOperator = Function->HasMetaData(FBlueprintMetadata::MD_CommutativeAssociativeBinaryOperator);
 	const bool bMaterialParameterCollectionFunction = Function->HasMetaData(FBlueprintMetadata::MD_MaterialParameterCollectionFunction);
+	bool const bIsDataTableFunc = Function->HasMetaData(FBlueprintMetadata::MD_DataTablePin);
 
 	UK2Node_CallFunction* CallTemplate = NULL;
 
@@ -667,6 +668,10 @@ void FK2ActionMenuBuilder::AddSpawnInfoForFunction(const UFunction* Function, bo
 	{
 		CallTemplate = ListBuilder.CreateTemplateNode<UK2Node_CallMaterialParameterCollectionFunction>();
 	}
+	else if (bIsDataTableFunc)
+	{
+		CallTemplate = ListBuilder.CreateTemplateNode<UK2Node_CallDataTableFunction>();
+	}
 	else if (CallOnMember.GetMemberName() != NAME_None)
 	{
 		CallTemplate = ListBuilder.CreateTemplateNode<UK2Node_CallFunctionOnMember>();
diff --git a/Engine/Source/Editor/BlueprintGraph/Private/K2Node_CallDataTableFunction.cpp b/Engine/Source/Editor/BlueprintGraph/Private/K2Node_CallDataTableFunction.cpp
new file mode 100644
--- /dev/null
+++ b/Engine/Source/Editor/BlueprintGraph/Private/K2Node_CallDataTableFunction.cpp
@@ -0,0 +1,35 @@
+// Copyright 1998-2014 Epic Games, Inc. All Rights Reserved.
+
+#include "BlueprintGraphPrivatePCH.h"
+
+#include "CompilerResultsLog.h"
+
+UK2Node_CallDataTableFunction::UK2Node_CallDataTableFunction(const class FPostConstructInitializeProperties& PCIP)
+	: Super(PCIP)
+{
+}
+
+void UK2Node_CallDataTableFunction::PinDefaultValueChanged(UEdGraphPin* Pin)
+{
+	Super::PinDefaultValueChanged(Pin);
+
+	FString DataTablePinName = GetTargetFunction()->GetMetaData(FBlueprintMetadata::MD_DataTablePin);
+	if (Pin->PinName == DataTablePinName)
+	{
+		// When the DataTable pin gets a new value assigned, we need to update the Slate UI so that the RowName drop down gets updated
+		GetGraph()->NotifyGraphChanged();
+	}
+}
+
+void UK2Node_CallDataTableFunction::NotifyPinConnectionListChanged(UEdGraphPin* Pin)
+{
+	Super::NotifyPinConnectionListChanged(Pin);
+
+	FString DataTablePinName = GetTargetFunction()->GetMetaData(FBlueprintMetadata::MD_DataTablePin);
+	if (Pin->PinName == DataTablePinName)
+	{
+		// When the DataTable pin gets a new connection assigned, we need to update the Slate UI so that the RowName drop down gets updated
+		GetGraph()->NotifyGraphChanged();
+	}
+}
+
diff --git a/Engine/Source/Editor/BlueprintGraph/Private/K2Node_GetDataTableRow.cpp b/Engine/Source/Editor/BlueprintGraph/Private/K2Node_GetDataTableRow.cpp
--- a/Engine/Source/Editor/BlueprintGraph/Private/K2Node_GetDataTableRow.cpp
+++ b/Engine/Source/Editor/BlueprintGraph/Private/K2Node_GetDataTableRow.cpp
@@ -41,8 +41,8 @@ void UK2Node_GetDataTableRow::AllocateDefaultPins()
 	SetPinToolTip(*DataTablePin, LOCTEXT("DataTablePinDescription", "The DataTable you want to retreive a row from"));
 
 	// Row Name pin
-	UEdGraphPin* TransformPin = CreatePin(EGPD_Input, K2Schema->PC_Name, TEXT(""), NULL, false, false, UK2Node_GetDataTableRowHelper::RowNamePinName);
-	SetPinToolTip(*TransformPin, LOCTEXT("RowNamePinDescription", "The name of the row to retrieve from the DataTable"));
+	UEdGraphPin* RowNamePin = CreatePin(EGPD_Input, K2Schema->PC_Name, TEXT(""), NULL, false, false, UK2Node_GetDataTableRowHelper::RowNamePinName);
+	SetPinToolTip(*RowNamePin, LOCTEXT("RowNamePinDescription", "The name of the row to retrieve from the DataTable"));
 
 	// Result pin
 	UEdGraphPin* ResultPin = CreatePin(EGPD_Output, K2Schema->PC_Struct, TEXT(""), FTableRowBase::StaticStruct(), false, false, K2Schema->PN_ReturnValue);
@@ -74,9 +74,33 @@ void UK2Node_GetDataTableRow::SetReturnTypeForStruct(UScriptStruct* RowStruct)
 		RowStruct = FTableRowBase::StaticStruct();
 	}
 
-    // Change class of output pin
 	UEdGraphPin* ResultPin = GetResultPin();
-	ResultPin->PinType.PinSubCategoryObject = RowStruct;
+
+	UScriptStruct* NewRowStruct = RowStruct;
+	UScriptStruct* OldRowStruct = GetReturnTypeForStruct();
+
+	// If new Data Table uses a different struct type for it's rows
+	if (NewRowStruct != OldRowStruct)
+	{
+		// Doing this just to force the row name drop down to refresh
+		ReconstructNode();
+
+		// Because the Return Value struct type has changed, we break the output link
+		ResultPin->BreakAllPinLinks();
+
+		// Change class of output pin
+		UEdGraphPin* ResultPin = GetResultPin();
+		ResultPin->PinType.PinSubCategoryObject = RowStruct;
+
+		// When the DataTable pin gets a new value assigned, we need to update the Slate UI so that SGraphNodeCallParameterCollectionFunction will update the ParameterName drop down
+		UEdGraph* Graph = GetGraph();
+		Graph->NotifyGraphChanged();
+
+		UEdGraphPin* OldRowNamePin = GetRowNamePin();
+
+		// Mark dirty
+		FBlueprintEditorUtils::MarkBlueprintAsModified(GetBlueprint());
+	}
 }
 
 UScriptStruct* UK2Node_GetDataTableRow::GetReturnTypeForStruct()
@@ -97,20 +121,47 @@ UScriptStruct* UK2Node_GetDataTableRow::GetDataTableRowStructType(const TArray<U
 		if (DataTablePin->DefaultObject->IsA(UDataTable::StaticClass()))
 		{
 			UDataTable* DataTable = (UDataTable*)DataTablePin->DefaultObject;
-			RowStructType = DataTable->RowStruct;
+			if (DataTable)
+			{
+				RowStructType = DataTable->RowStruct;
+			}
 		}
 	}
 
 	return RowStructType;
 }
 
+void UK2Node_GetDataTableRow::GetDataTableRowNameList(TArray<TSharedPtr<FName>>& RowNames, const TArray<UEdGraphPin*>* InPinsToSearch /*= NULL*/) const
+{
+	const TArray<UEdGraphPin*>* PinsToSearch = InPinsToSearch ? InPinsToSearch : &Pins;
+
+	UEdGraphPin* DataTablePin = GetDataTablePin(PinsToSearch);
+	if (DataTablePin && DataTablePin->DefaultObject != NULL && DataTablePin->LinkedTo.Num() == 0)
+	{
+		if (DataTablePin->DefaultObject->IsA(UDataTable::StaticClass()))
+		{
+			UDataTable* DataTable = (UDataTable*)DataTablePin->DefaultObject;
+			if (DataTable)
+			{
+				/** Extract all the row names from the RowMap */
+				for (TMap<FName, uint8*>::TConstIterator Iterator(DataTable->RowMap); Iterator; ++Iterator)
+				{
+					/** Create a simple array of the row names */
+					TSharedPtr<FName> RowNameItem = MakeShareable(new FName(Iterator.Key()));
+					RowNames.Add(RowNameItem);
+				}
+			}
+		}
+	}
+}
+
 
 void UK2Node_GetDataTableRow::ReallocatePinsDuringReconstruction(TArray<UEdGraphPin*>& OldPins) 
 {
 	Super::ReallocatePinsDuringReconstruction(OldPins);
 		
 	UScriptStruct* RowStruct = GetDataTableRowStructType(&OldPins);
-	if( RowStruct != NULL)
+	if( RowStruct != NULL )
 	{
 		SetReturnTypeForStruct(RowStruct);
 	}
@@ -129,32 +180,31 @@ FText UK2Node_GetDataTableRow::GetMenuCategory() const
 	return FEditorCategoryUtils::GetCommonCategory(FCommonEditorCategory::Utilities);
 }
 
-void UK2Node_GetDataTableRow::PinDefaultValueChanged(UEdGraphPin* ChangedPin) 
+void UK2Node_GetDataTableRow::NotifyPinConnectionListChanged(UEdGraphPin* Pin)
 {
-	if (ChangedPin->PinName == UK2Node_GetDataTableRowHelper::DataTablePinName)
+	if (Pin && Pin->PinName == UK2Node_GetDataTableRowHelper::DataTablePinName)
 	{
-		const UEdGraphSchema_K2* K2Schema = GetDefault<UEdGraphSchema_K2>();
-
+		// Because the archetype has changed, we break the output link as the output pin type will change
 		UEdGraphPin* ResultPin = GetResultPin();
+		ResultPin->BreakAllPinLinks();
 
-		UScriptStruct* PrevRowStruct = GetReturnTypeForStruct();
-		UScriptStruct* RowStruct = GetDataTableRowStructType();
-
-		if (RowStruct != PrevRowStruct || RowStruct == NULL)
+		// If the pin has been connected, clear the default values so that we don't hold on to references
+		if (Pin->LinkedTo.Num() == 1)
 		{
-			// Because the data table type has changed, we break the output link as the output pin type will change
-			ResultPin->BreakAllPinLinks();
+			UScriptStruct* RowStruct = GetDataTableRowStructType(&(Pin->LinkedTo));
 
+			SetReturnTypeForStruct(RowStruct);
 		}
+	}
+}
 
-		SetReturnTypeForStruct(RowStruct);
-
-		// Refresh the UI for the graph so the pin changes show up
-		UEdGraph* Graph = GetGraph();
-		Graph->NotifyGraphChanged();
+void UK2Node_GetDataTableRow::PinDefaultValueChanged(UEdGraphPin* ChangedPin) 
+{
+	if (ChangedPin && ChangedPin->PinName == UK2Node_GetDataTableRowHelper::DataTablePinName)
+	{
+		UScriptStruct* RowStruct = GetDataTableRowStructType();
 
-		// Mark dirty
-		FBlueprintEditorUtils::MarkBlueprintAsModified(GetBlueprint());
+		SetReturnTypeForStruct(RowStruct);
 	}
 }
 
diff --git a/Engine/Source/Editor/BlueprintGraph/Public/BlueprintGraphDefinitions.h b/Engine/Source/Editor/BlueprintGraph/Public/BlueprintGraphDefinitions.h
--- a/Engine/Source/Editor/BlueprintGraph/Public/BlueprintGraphDefinitions.h
+++ b/Engine/Source/Editor/BlueprintGraph/Public/BlueprintGraphDefinitions.h
@@ -11,6 +11,7 @@
 #include "K2Node_BaseMCDelegate.h"
 #include "K2Node_BreakStruct.h"
 #include "K2Node_CallArrayFunction.h"
+#include "K2Node_CallDataTableFunction.h"
 #include "K2Node_CallDelegate.h"
 #include "K2Node_CallFunction.h"
 #include "K2Node_CallFunctionOnMember.h"
diff --git a/Engine/Source/Editor/BlueprintGraph/Public/BlueprintGraphPanelPinFactory.h b/Engine/Source/Editor/BlueprintGraph/Public/BlueprintGraphPanelPinFactory.h
new file mode 100644
--- /dev/null
+++ b/Engine/Source/Editor/BlueprintGraph/Public/BlueprintGraphPanelPinFactory.h
@@ -0,0 +1,10 @@
+// Copyright 1998-2014 Epic Games, Inc. All Rights Reserved.
+
+#pragma once
+
+#include "EdGraphUtilities.h"
+
+class BLUEPRINTGRAPH_API FBlueprintGraphPanelPinFactory: public FGraphPanelPinFactory
+{
+	virtual TSharedPtr<class SGraphPin> CreatePin(class UEdGraphPin* InPin) const override;
+};
diff --git a/Engine/Source/Editor/GraphEditor/Private/KismetNodes/SGraphNodeCallParameterCollectionFunction.cpp b/Engine/Source/Editor/GraphEditor/Private/KismetNodes/SGraphNodeCallParameterCollectionFunction.cpp
--- a/Engine/Source/Editor/GraphEditor/Private/KismetNodes/SGraphNodeCallParameterCollectionFunction.cpp
+++ b/Engine/Source/Editor/GraphEditor/Private/KismetNodes/SGraphNodeCallParameterCollectionFunction.cpp
@@ -7,7 +7,7 @@
 #include "SGraphNodeK2Base.h"
 #include "SGraphNodeK2Default.h"
 #include "SGraphNodeCallParameterCollectionFunction.h"
-#include "KismetPins/SGraphPinNameList.h"
+#include "SGraphPinNameList.h"
 
 //////////////////////////////////////////////////////////////////////////
 // SGraphNodeCallParameterCollectionFunction
@@ -37,7 +37,14 @@ TSharedPtr<SGraphPin> SGraphNodeCallParameterCollectionFunction::CreatePinWidget
 			}
 		}
 
-		TSharedPtr<SGraphPin> NewPin = SNew(SGraphPinNameList, Pin, NameList);
+		TArray<TSharedPtr<FName>> NamePtrList;
+
+		for (FName NameItem : NameList)
+		{
+			NamePtrList.Add(MakeShareable( new FName(NameItem)));
+		}
+
+		TSharedPtr<SGraphPin> NewPin = SNew(SGraphPinNameList, Pin, NamePtrList);
 		return NewPin;
 	}
 	else
diff --git a/Engine/Source/Editor/GraphEditor/Private/KismetPins/SGraphPinNameList.cpp b/Engine/Source/Editor/GraphEditor/Private/KismetPins/SGraphPinNameList.cpp
--- a/Engine/Source/Editor/GraphEditor/Private/KismetPins/SGraphPinNameList.cpp
+++ b/Engine/Source/Editor/GraphEditor/Private/KismetPins/SGraphPinNameList.cpp
@@ -3,10 +3,9 @@
 
 #include "GraphEditorCommon.h"
 #include "SGraphPinComboBox.h"
-#include "SGraphPinEnum.h"
 #include "SGraphPinNameList.h"
 
-void SGraphPinNameList::Construct(const FArguments& InArgs, UEdGraphPin* InGraphPinObj, const TArray<FName>& InNameList)
+void SGraphPinNameList::Construct(const FArguments& InArgs, UEdGraphPin* InGraphPinObj, const TArray<TSharedPtr<FName>>& InNameList)
 {
 	NameList = InNameList;
 	SGraphPin::Construct(SGraphPin::FArguments(), InGraphPinObj);
@@ -14,29 +13,48 @@ void SGraphPinNameList::Construct(const FArguments& InArgs, UEdGraphPin* InGraph
 
 TSharedRef<SWidget>	SGraphPinNameList::GetDefaultValueWidget()
 {
-	//Get list of enum indexes
+	//Get list of name indexes
 	TArray< TSharedPtr<int32> > ComboItems;
 	GenerateComboBoxIndexes( ComboItems );
 
-	//Create widget
-	return SAssignNew(ComboBox, SPinComboBox)
-		.ComboItemList( ComboItems )
-		.VisibleText( this, &SGraphPinNameList::OnGetText )
-		.OnSelectionChanged( this, &SGraphPinNameList::ComboBoxSelectionChanged )
-		.Visibility( this, &SGraphPin::GetDefaultValueVisibility )
-		.OnGetDisplayName(this, &SGraphPinNameList::OnGetFriendlyName);
+	TSharedPtr<FName> CurrentlySelectedName;
+	if (NameList.Num() > 0)
+	{
+		// Preserve previous selection, or set to first in list
+		FName PreviousSelection = FName(*GraphPinObj->GetDefaultAsString());
+		CurrentlySelectedName = NameList[0];
+		for (TSharedPtr<FName> ListNamePtr : NameList)
+		{
+			if (PreviousSelection == *ListNamePtr.Get())
+			{
+				CurrentlySelectedName = ListNamePtr;
+				break;
+			}
+		}
+		check(CurrentlySelectedName.IsValid());
+		GraphPinObj->GetSchema()->TrySetDefaultValue(*GraphPinObj, *CurrentlySelectedName->ToString());
+	}
+
+	// Create widget
+	return SAssignNew(ComboBox, SNameComboBox)
+		.ContentPadding(FMargin(6.0f, 2.0f))
+		.OptionsSource(&NameList)
+		.InitiallySelectedItem(CurrentlySelectedName)
+		.OnSelectionChanged(this, &SGraphPinNameList::ComboBoxSelectionChanged)
+		.Visibility(this, &SGraphPin::GetDefaultValueVisibility)
+		;
 }
 
-FString SGraphPinNameList::OnGetFriendlyName(int32 EnumIndex)
+FString SGraphPinNameList::OnGetFriendlyName(int32 NameIndex)
 {
-	check(EnumIndex < NameList.Num());
-	return NameList[EnumIndex].ToString();
+	check(NameIndex < NameList.Num());
+	return NameList[NameIndex]->ToString();
 }
 
-void SGraphPinNameList::ComboBoxSelectionChanged( TSharedPtr<int32> NewSelection, ESelectInfo::Type /*SelectInfo*/ )
+void SGraphPinNameList::ComboBoxSelectionChanged(TSharedPtr<FName> NameItem, ESelectInfo::Type SelectInfo )
 {
-	check(*NewSelection < NameList.Num());
-	GraphPinObj->GetSchema()->TrySetDefaultValue(*GraphPinObj, NameList[*NewSelection].ToString());
+	check(NameItem.IsValid());
+	GraphPinObj->GetSchema()->TrySetDefaultValue(*GraphPinObj, *NameItem->ToString());
 }
 
 FString SGraphPinNameList::OnGetText() const 
@@ -49,7 +67,7 @@ void SGraphPinNameList::GenerateComboBoxIndexes( TArray< TSharedPtr<int32> >& Ou
 {
 	for (int32 NameIndex = 0; NameIndex < NameList.Num(); NameIndex++)
 	{
-		TSharedPtr<int32> EnumIdxPtr(new int32(NameIndex));
-		OutComboBoxIndexes.Add(EnumIdxPtr);
+		TSharedPtr<int32> NameIdxPtr(new int32(NameIndex));
+		OutComboBoxIndexes.Add(NameIdxPtr);
 	}
-}
\ No newline at end of file
+}
diff --git a/Engine/Source/Editor/GraphEditor/Private/KismetPins/SGraphPinNameList.h b/Engine/Source/Editor/GraphEditor/Private/KismetPins/SGraphPinNameList.h
deleted file mode 100644
diff --git a/Engine/Source/Editor/GraphEditor/Private/KismetPins/SNameComboBox.cpp b/Engine/Source/Editor/GraphEditor/Private/KismetPins/SNameComboBox.cpp
new file mode 100644
--- /dev/null
+++ b/Engine/Source/Editor/GraphEditor/Private/KismetPins/SNameComboBox.cpp
@@ -0,0 +1,79 @@
+// Copyright 1998-2014 Epic Games, Inc. All Rights Reserved.
+
+#include "GraphEditorCommon.h"
+#include "SNameComboBox.h"
+
+
+void SNameComboBox::Construct( const FArguments& InArgs )
+{
+	SelectionChanged = InArgs._OnSelectionChanged;
+	GetTextLabelForItem = InArgs._OnGetNameLabelForItem;
+
+	// Then make widget
+	this->ChildSlot
+	[
+		SAssignNew(NameCombo, SComboBox< TSharedPtr<FName> > )
+		.OptionsSource(InArgs._OptionsSource)
+		.OnGenerateWidget(this, &SNameComboBox::MakeItemWidget)
+		.OnSelectionChanged(this, &SNameComboBox::OnSelectionChanged)
+		.OnComboBoxOpening(InArgs._OnComboBoxOpening)
+		.InitiallySelectedItem(InArgs._InitiallySelectedItem)
+		.ContentPadding(InArgs._ContentPadding)
+		[
+			SNew(STextBlock)
+				.ColorAndOpacity(InArgs._ColorAndOpacity)
+				.Text(this, &SNameComboBox::GetSelectedNameLabel)
+		]
+	];
+	SelectedItem = NameCombo->GetSelectedItem();
+}
+
+FString SNameComboBox::GetItemNameLabel(TSharedPtr<FName> NameItem) const
+{
+	if (!NameItem.IsValid())
+	{
+		return FString();
+	}
+
+	return (GetTextLabelForItem.IsBound())
+		? GetTextLabelForItem.Execute(NameItem)
+		: *NameItem->ToString();
+}
+
+FString SNameComboBox::GetSelectedNameLabel() const
+{
+	TSharedPtr<FName> StringItem = NameCombo->GetSelectedItem();
+	return GetItemNameLabel(StringItem);
+}
+
+TSharedRef<SWidget> SNameComboBox::MakeItemWidget( TSharedPtr<FName> NameItem ) 
+{
+	check( NameItem.IsValid() );
+
+	return SNew(STextBlock)
+		.Text(this, &SNameComboBox::GetItemNameLabel, NameItem);
+}
+
+void SNameComboBox::OnSelectionChanged (TSharedPtr<FName> Selection, ESelectInfo::Type SelectInfo)
+{
+	if (Selection.IsValid())
+	{
+		SelectedItem = Selection;
+	}
+	SelectionChanged.ExecuteIfBound(Selection, SelectInfo);
+}
+
+void SNameComboBox::SetSelectedItem(TSharedPtr<FName> NewSelection)
+{
+	NameCombo->SetSelectedItem(NewSelection);
+}
+
+void SNameComboBox::RefreshOptions()
+{
+	NameCombo->RefreshOptions();
+}
+
+void SNameComboBox::ClearSelection( )
+{
+	NameCombo->ClearSelection();
+}
\ No newline at end of file
diff --git a/Engine/Source/Editor/GraphEditor/Private/NodeFactory.cpp b/Engine/Source/Editor/GraphEditor/Private/NodeFactory.cpp
--- a/Engine/Source/Editor/GraphEditor/Private/NodeFactory.cpp
+++ b/Engine/Source/Editor/GraphEditor/Private/NodeFactory.cpp
@@ -58,7 +58,6 @@
 #include "KismetPins/SGraphPinColor.h"
 #include "KismetPins/SGraphPinEnum.h"
 #include "KismetPins/SGraphPinKey.h"
-#include "KismetPins/SGraphPinNameList.h"
 #include "KismetPins/SGraphPinVector.h"
 #include "KismetPins/SGraphPinVector2D.h"
 #include "NiagaraPins/SGraphPinVector4.h"
diff --git a/Engine/Source/Editor/GraphEditor/Public/SGraphPinNameList.h b/Engine/Source/Editor/GraphEditor/Public/SGraphPinNameList.h
new file mode 100644
--- /dev/null
+++ b/Engine/Source/Editor/GraphEditor/Public/SGraphPinNameList.h
@@ -0,0 +1,59 @@
+// Copyright 1998-2014 Epic Games, Inc. All Rights Reserved.
+
+#pragma once
+#include "SGraphPin.h"
+#include "SNameComboBox.h"
+
+class GRAPHEDITOR_API SGraphPinNameList : public SGraphPin
+{
+public:
+	SLATE_BEGIN_ARGS(SGraphPinNameList) {}
+	SLATE_END_ARGS()
+
+	void Construct(const FArguments& InArgs, UEdGraphPin* InGraphPinObj, const TArray<TSharedPtr<FName>>& InNameList);
+
+protected:
+
+	/**
+	 *	Function to create class specific widget.
+	 *
+	 *	@return Reference to the newly created widget object
+	 */
+	virtual TSharedRef<SWidget>	GetDefaultValueWidget() override;
+
+private:
+
+	/**
+	 *	Function to get current string associated with the combo box selection
+	 *
+	 *	@return currently selected string
+	 */
+	FString OnGetText() const;
+
+	/**
+	 *	Function to generate the list of indexes from the enum object
+	 *
+	 *	@param OutComboBoxIndexes - Int array reference to store the list of indexes
+	 */
+	void GenerateComboBoxIndexes( TArray< TSharedPtr<int32> >& OutComboBoxIndexes );
+
+	/**
+	 *	Function to set the newly selected index
+	 *
+	 * @param NewSelection The newly selected item in the combo box
+	 * @param SelectInfo Provides context on how the selection changed
+	 */
+	void ComboBoxSelectionChanged(TSharedPtr<FName> StringItem, ESelectInfo::Type SelectInfo);
+
+	/**
+	 * Returns the friendly name of the enum at index EnumIndex
+	 *
+	 * @param EnumIndex	- The index of the enum to return the friendly name for
+	 */
+	FString OnGetFriendlyName(int32 EnumIndex);
+
+	TSharedPtr<class SNameComboBox>	ComboBox;
+
+	/** The actual list of FName values to choose from */
+	TArray<TSharedPtr<FName>> NameList;
+};
diff --git a/Engine/Source/Editor/GraphEditor/Public/SNameComboBox.h b/Engine/Source/Editor/GraphEditor/Public/SNameComboBox.h
new file mode 100644
--- /dev/null
+++ b/Engine/Source/Editor/GraphEditor/Public/SNameComboBox.h
@@ -0,0 +1,95 @@
+// Copyright 1998-2014 Epic Games, Inc. All Rights Reserved.
+ 
+#pragma once
+#include "NameTypes.h"
+#include "SWidget.h"
+#include "SharedPointer.h"
+#include "SlateEnums.h"
+#include "SCompoundWidget.h"
+//#include "UnrealString.h"
+
+/**
+ * A combo box that shows FName content.
+ */
+class SNameComboBox : public SCompoundWidget
+{
+public:
+
+	DECLARE_DELEGATE_RetVal_OneParam(FString, FGetNameComboLabel, TSharedPtr<FName>);
+	typedef TSlateDelegates< TSharedPtr<FName> >::FOnSelectionChanged FOnNameSelectionChanged;
+
+	SLATE_BEGIN_ARGS( SNameComboBox ) 
+		: _ColorAndOpacity( FSlateColor::UseForeground() )
+		, _ContentPadding(FMargin(4.0, 2.0))
+		, _OnGetNameLabelForItem()
+		{}
+
+		/** Selection of FNames to pick from */
+		SLATE_ARGUMENT( TArray< TSharedPtr<FName> >*, OptionsSource )
+
+		/** Text color and opacity */
+		SLATE_ATTRIBUTE( FSlateColor, ColorAndOpacity )
+
+		/** Visual padding of the button content for the combobox */
+		SLATE_ATTRIBUTE( FMargin, ContentPadding )
+
+		/** Called when the FName is chosen. */
+		SLATE_EVENT( FOnNameSelectionChanged, OnSelectionChanged)
+
+		/** Called when the combo box is opened */
+		SLATE_EVENT( FOnComboBoxOpening, OnComboBoxOpening )
+
+		/** Called when combo box needs to establish selected item */
+		SLATE_ARGUMENT( TSharedPtr<FName>, InitiallySelectedItem )
+
+		/** [Optional] Called to get the label for the currently selected item */
+		SLATE_EVENT( FGetNameComboLabel, OnGetNameLabelForItem ) 
+	SLATE_END_ARGS()
+
+	void Construct( const FArguments& InArgs );
+
+	/** Called to create a widget for each FName */
+	TSharedRef<SWidget> MakeItemWidget( TSharedPtr<FName> StringItem );
+
+	void SetSelectedItem (TSharedPtr<FName> NewSelection);
+
+	/** Returns the currently selected FName */
+	TSharedPtr<FName> GetSelectedItem()
+	{
+		return SelectedItem;
+	}
+
+	/** Request to reload the name options in the combobox from the OptionsSource attribute */
+	void RefreshOptions();
+
+	/** Clears the selected item in the name combo */
+	void ClearSelection();
+
+private:
+	TSharedPtr<FName> OnGetSelection() const {return SelectedItem;}
+
+	/** Called when selection changes in the combo pop-up */
+	void OnSelectionChanged(TSharedPtr<FName> Selection, ESelectInfo::Type SelectInfo);
+
+	/** Helper method to get the text for a given item in the combo box */
+	FString GetSelectedNameLabel() const;
+
+	FString GetItemNameLabel(TSharedPtr<FName> StringItem) const;
+
+private:
+
+	/** Called to get the text label for an item */
+	FGetNameComboLabel GetTextLabelForItem;
+
+	/** The FName item selected */
+	TSharedPtr<FName> SelectedItem;
+
+	/** Array of shared pointers to FNames so combo widget can work on them */
+	TArray< TSharedPtr<FName> >		Names;
+
+	/** The combo box */
+	TSharedPtr< SComboBox< TSharedPtr<FName> > >	NameCombo;
+
+	/** Forwarding Delegate */
+	FOnNameSelectionChanged SelectionChanged;
+};
diff --git a/Engine/Source/Editor/Kismet/Private/BlueprintEditorModule.cpp b/Engine/Source/Editor/Kismet/Private/BlueprintEditorModule.cpp
--- a/Engine/Source/Editor/Kismet/Private/BlueprintEditorModule.cpp
+++ b/Engine/Source/Editor/Kismet/Private/BlueprintEditorModule.cpp
@@ -15,6 +15,7 @@
 #include "UObjectToken.h"
 #include "InstancedStaticMeshSCSEditorCustomization.h"
 #include "UserDefinedStructureEditor.h"
+#include "BlueprintGraphPanelPinFactory.h"
 
 #define LOCTEXT_NAMESPACE "BlueprintEditor"
 
@@ -100,6 +101,9 @@ void FBlueprintEditorModule::StartupModule()
 
 	// Register internal SCS editor customizations
 	RegisterSCSEditorCustomization("InstancedStaticMeshComponent", FSCSEditorCustomizationBuilder::CreateStatic(&FInstancedStaticMeshSCSEditorCustomization::MakeInstance));
+
+	TSharedPtr<FBlueprintGraphPanelPinFactory> BlueprintGraphPanelPinFactory = MakeShareable(new FBlueprintGraphPanelPinFactory());
+	FEdGraphUtilities::RegisterVisualPinFactory(BlueprintGraphPanelPinFactory);
 }
 
 
diff --git a/Engine/Source/Runtime/Engine/Classes/Kismet/DataTableFunctionLibrary.h b/Engine/Source/Runtime/Engine/Classes/Kismet/DataTableFunctionLibrary.h
--- a/Engine/Source/Runtime/Engine/Classes/Kismet/DataTableFunctionLibrary.h
+++ b/Engine/Source/Runtime/Engine/Classes/Kismet/DataTableFunctionLibrary.h
@@ -23,7 +23,7 @@ class UDataTableFunctionLibrary : public UBlueprintFunctionLibrary
 {
 	GENERATED_UCLASS_BODY()
 
-	UFUNCTION(BlueprintCallable, Category = "DataTable", meta=(ExpandEnumAsExecs="OutResult"))
+	UFUNCTION(BlueprintCallable, Category = "DataTable", meta = (ExpandEnumAsExecs="OutResult", DataTablePin="CurveTable"))
 	static void EvaluateCurveTableRow(UCurveTable* CurveTable, FName RowName, float InXY, TEnumAsByte<EEvaluateCurveTableResult::Type>& OutResult, float& OutXY);
     
     /** Get a Row from a DataTable given a RowName */
