Fix iOS: "#pragma once" snuck into my .cpp file

diff --git a/Engine/Source/Editor/BlueprintGraph/Private/BlueprintGraphPanelPinFactory.cpp b/Engine/Source/Editor/BlueprintGraph/Private/BlueprintGraphPanelPinFactory.cpp
--- a/Engine/Source/Editor/BlueprintGraph/Private/BlueprintGraphPanelPinFactory.cpp
+++ b/Engine/Source/Editor/BlueprintGraph/Private/BlueprintGraphPanelPinFactory.cpp
@@ -1,7 +1,4 @@
 // Copyright 1998-2014 Epic Games, Inc. All Rights Reserved.
-
-#pragma once
-
 #include "BlueprintGraphPrivatePCH.h"
 #include "BlueprintGraphPanelPinFactory.h"
 #include "K2Node_GetDataTableRow.h"
