BP - Fixing a crash with name pins that don't have any data table metadata.

diff --git a/Engine/Source/Editor/BlueprintGraph/Private/BlueprintGraphPanelPinFactory.cpp b/Engine/Source/Editor/BlueprintGraph/Private/BlueprintGraphPanelPinFactory.cpp
--- a/Engine/Source/Editor/BlueprintGraph/Private/BlueprintGraphPanelPinFactory.cpp
+++ b/Engine/Source/Editor/BlueprintGraph/Private/BlueprintGraphPanelPinFactory.cpp
@@ -21,7 +21,7 @@ TSharedPtr<class SGraphPin> FBlueprintGraphPanelPinFactory::CreatePin(class UEdG
 			{
 				const UFunction* FunctionToCall = CallFunctionNode->GetTargetFunction();
 				FString DataTablePinName = FunctionToCall->GetMetaData(FBlueprintMetadata::MD_DataTablePin);
-				DataTablePin = CallFunctionNode->FindPinChecked(DataTablePinName);
+				DataTablePin = CallFunctionNode->FindPin(DataTablePinName);
 			}
 		}
 		else if (Outer->IsA(UK2Node_GetDataTableRow::StaticClass()))
@@ -29,6 +29,7 @@ TSharedPtr<class SGraphPin> FBlueprintGraphPanelPinFactory::CreatePin(class UEdG
 			const UK2Node_GetDataTableRow* GetDataTableRowNode = CastChecked<UK2Node_GetDataTableRow>(Outer);
 			DataTablePin = GetDataTableRowNode->GetDataTablePin();
 		}
+
 		if (DataTablePin)
 		{
 			if (DataTablePin && DataTablePin->DefaultObject != nullptr && DataTablePin->LinkedTo.Num() == 0)
