Handling enum/struct assets (so that the can be added into the blueprint menu database as they're created/updated).

Lots o' files!(changing UK2Node's GetMenuActions() function signature).

diff --git a/Engine/Source/Editor/AnimGraph/Classes/AnimGraphNode_Base.h b/Engine/Source/Editor/AnimGraph/Classes/AnimGraphNode_Base.h
--- a/Engine/Source/Editor/AnimGraph/Classes/AnimGraphNode_Base.h
+++ b/Engine/Source/Editor/AnimGraph/Classes/AnimGraphNode_Base.h
@@ -113,7 +113,7 @@ class UAnimGraphNode_Base : public UK2Node
 	ANIMGRAPH_API virtual bool CanCreateUnderSpecifiedSchema(const UEdGraphSchema* DesiredSchema) const override;
 	ANIMGRAPH_API virtual void GetMenuEntries(FGraphContextMenuBuilder& ContextMenuBuilder) const override;
 	ANIMGRAPH_API virtual void GetNodeAttributes(TArray<TKeyValuePair<FString, FString>>& OutNodeAttributes) const override;
-	ANIMGRAPH_API virtual void GetMenuActions(TArray<UBlueprintNodeSpawner*>& ActionListOut) const override;
+	ANIMGRAPH_API virtual void GetMenuActions(FBlueprintActionDatabaseRegistrar& ActionRegistrar) const override;
 	ANIMGRAPH_API virtual FText GetMenuCategory() const override;
 
 	// By default return any animation assets we have
diff --git a/Engine/Source/Editor/AnimGraph/Classes/AnimGraphNode_BlendSpaceEvaluator.h b/Engine/Source/Editor/AnimGraph/Classes/AnimGraphNode_BlendSpaceEvaluator.h
--- a/Engine/Source/Editor/AnimGraph/Classes/AnimGraphNode_BlendSpaceEvaluator.h
+++ b/Engine/Source/Editor/AnimGraph/Classes/AnimGraphNode_BlendSpaceEvaluator.h
@@ -19,7 +19,7 @@ class UAnimGraphNode_BlendSpaceEvaluator : public UAnimGraphNode_BlendSpaceBase
 	// UEdGraphNode interface
 	virtual FString GetTooltip() const override;
 	virtual FText GetNodeTitle(ENodeTitleType::Type TitleType) const override;
-	virtual void GetMenuActions(TArray<UBlueprintNodeSpawner*>& ActionListOut) const override;
+	virtual void GetMenuActions(FBlueprintActionDatabaseRegistrar& ActionRegistrar) const override;
 	// End of UEdGraphNode interface
 
 	// UAnimGraphNode_Base interface
diff --git a/Engine/Source/Editor/AnimGraph/Classes/AnimGraphNode_BlendSpacePlayer.h b/Engine/Source/Editor/AnimGraph/Classes/AnimGraphNode_BlendSpacePlayer.h
--- a/Engine/Source/Editor/AnimGraph/Classes/AnimGraphNode_BlendSpacePlayer.h
+++ b/Engine/Source/Editor/AnimGraph/Classes/AnimGraphNode_BlendSpacePlayer.h
@@ -20,7 +20,7 @@ class UAnimGraphNode_BlendSpacePlayer : public UAnimGraphNode_BlendSpaceBase
 	// UEdGraphNode interface
 	virtual FString GetTooltip() const override;
 	virtual FText GetNodeTitle(ENodeTitleType::Type TitleType) const override;
-	virtual void GetMenuActions(TArray<UBlueprintNodeSpawner*>& ActionListOut) const override;
+	virtual void GetMenuActions(FBlueprintActionDatabaseRegistrar& ActionRegistrar) const override;
 	// End of UEdGraphNode interface
 
 	// UAnimGraphNode_Base interface
diff --git a/Engine/Source/Editor/AnimGraph/Classes/AnimGraphNode_Root.h b/Engine/Source/Editor/AnimGraph/Classes/AnimGraphNode_Root.h
--- a/Engine/Source/Editor/AnimGraph/Classes/AnimGraphNode_Root.h
+++ b/Engine/Source/Editor/AnimGraph/Classes/AnimGraphNode_Root.h
@@ -20,7 +20,7 @@ class UAnimGraphNode_Root : public UAnimGraphNode_Base
 	virtual bool CanUserDeleteNode() const override { return false; }
 	virtual bool CanDuplicateNode() const override { return false; }
 	virtual void GetMenuEntries(FGraphContextMenuBuilder& ContextMenuBuilder) const override;
-	virtual void GetMenuActions(TArray<UBlueprintNodeSpawner*>& ActionListOut) const override;
+	virtual void GetMenuActions(FBlueprintActionDatabaseRegistrar& ActionRegistrar) const override;
 	// End UEdGraphNode interface.
 
 	// UAnimGraphNode_Base interface
diff --git a/Engine/Source/Editor/AnimGraph/Classes/AnimGraphNode_RotationOffsetBlendSpace.h b/Engine/Source/Editor/AnimGraph/Classes/AnimGraphNode_RotationOffsetBlendSpace.h
--- a/Engine/Source/Editor/AnimGraph/Classes/AnimGraphNode_RotationOffsetBlendSpace.h
+++ b/Engine/Source/Editor/AnimGraph/Classes/AnimGraphNode_RotationOffsetBlendSpace.h
@@ -26,6 +26,6 @@ class UAnimGraphNode_RotationOffsetBlendSpace: public UAnimGraphNode_BlendSpaceB
 
 	// UK2Node interface
 	virtual void GetContextMenuActions(const FGraphNodeContextMenuBuilder& Context) const override;
-	virtual void GetMenuActions(TArray<UBlueprintNodeSpawner*>& ActionListOut) const override;
+	virtual void GetMenuActions(FBlueprintActionDatabaseRegistrar& ActionRegistrar) const override;
 	// End of UK2Node interface
 };
diff --git a/Engine/Source/Editor/AnimGraph/Classes/AnimGraphNode_SaveCachedPose.h b/Engine/Source/Editor/AnimGraph/Classes/AnimGraphNode_SaveCachedPose.h
--- a/Engine/Source/Editor/AnimGraph/Classes/AnimGraphNode_SaveCachedPose.h
+++ b/Engine/Source/Editor/AnimGraph/Classes/AnimGraphNode_SaveCachedPose.h
@@ -25,7 +25,7 @@ class UAnimGraphNode_SaveCachedPose : public UAnimGraphNode_Base
 	// End of UEdGraphNode interface
 
 	// UK2Node interface.
-	virtual void GetMenuActions(TArray<UBlueprintNodeSpawner*>& ActionListOut) const override;
+	virtual void GetMenuActions(FBlueprintActionDatabaseRegistrar& ActionRegistrar) const override;
 	// End of UK2Node interface
 
 	// UAnimGraphNode_Base interface
diff --git a/Engine/Source/Editor/AnimGraph/Classes/AnimGraphNode_SequenceEvaluator.h b/Engine/Source/Editor/AnimGraph/Classes/AnimGraphNode_SequenceEvaluator.h
--- a/Engine/Source/Editor/AnimGraph/Classes/AnimGraphNode_SequenceEvaluator.h
+++ b/Engine/Source/Editor/AnimGraph/Classes/AnimGraphNode_SequenceEvaluator.h
@@ -33,6 +33,6 @@ class UAnimGraphNode_SequenceEvaluator : public UAnimGraphNode_Base
 
 	// UK2Node interface
 	virtual void GetContextMenuActions(const FGraphNodeContextMenuBuilder& Context) const override;
-	virtual void GetMenuActions(TArray<UBlueprintNodeSpawner*>& ActionListOut) const override;
+	virtual void GetMenuActions(FBlueprintActionDatabaseRegistrar& ActionRegistrar) const override;
 	// End of UK2Node interface
 };
diff --git a/Engine/Source/Editor/AnimGraph/Classes/AnimGraphNode_TransitionPoseEvaluator.h b/Engine/Source/Editor/AnimGraph/Classes/AnimGraphNode_TransitionPoseEvaluator.h
--- a/Engine/Source/Editor/AnimGraph/Classes/AnimGraphNode_TransitionPoseEvaluator.h
+++ b/Engine/Source/Editor/AnimGraph/Classes/AnimGraphNode_TransitionPoseEvaluator.h
@@ -28,6 +28,6 @@ class UAnimGraphNode_TransitionPoseEvaluator : public UAnimGraphNode_Base
 	// End of UAnimGraphNode_Base interface
 
 	// UK2Node interface.
-	virtual void GetMenuActions(TArray<UBlueprintNodeSpawner*>& ActionListOut) const override;
+	virtual void GetMenuActions(FBlueprintActionDatabaseRegistrar& ActionRegistrar) const override;
 	// End of UK2Node interface
 };
\ No newline at end of file
diff --git a/Engine/Source/Editor/AnimGraph/Classes/AnimGraphNode_TransitionResult.h b/Engine/Source/Editor/AnimGraph/Classes/AnimGraphNode_TransitionResult.h
--- a/Engine/Source/Editor/AnimGraph/Classes/AnimGraphNode_TransitionResult.h
+++ b/Engine/Source/Editor/AnimGraph/Classes/AnimGraphNode_TransitionResult.h
@@ -23,7 +23,7 @@ class UAnimGraphNode_TransitionResult : public UAnimGraphNode_Base
 	// End of UEdGraphNode interface
 
 	// UK2Node interface.
-	virtual void GetMenuActions(TArray<UBlueprintNodeSpawner*>& ActionListOut) const override;
+	virtual void GetMenuActions(FBlueprintActionDatabaseRegistrar& ActionRegistrar) const override;
 	// End of UK2Node interface
 
 	// UAnimGraphNode_Base interface
diff --git a/Engine/Source/Editor/AnimGraph/Private/AnimGraphNode_Base.cpp b/Engine/Source/Editor/AnimGraph/Private/AnimGraphNode_Base.cpp
--- a/Engine/Source/Editor/AnimGraph/Private/AnimGraphNode_Base.cpp
+++ b/Engine/Source/Editor/AnimGraph/Private/AnimGraphNode_Base.cpp
@@ -6,6 +6,7 @@
 #include "K2ActionMenuBuilder.h" // for FK2ActionMenuBuilder::AddNewNodeAction()
 #include "AnimationGraphSchema.h"
 #include "BlueprintNodeSpawner.h"
+#include "BlueprintActionDatabaseRegistrar.h"
 
 /////////////////////////////////////////////////////
 // FA3NodeOptionalPinManager
@@ -227,11 +228,11 @@ void UAnimGraphNode_Base::GetNodeAttributes( TArray<TKeyValuePair<FString, FStri
 	OutNodeAttributes.Add( TKeyValuePair<FString, FString>( TEXT( "Name" ), GetName() ));
 }
 
-void UAnimGraphNode_Base::GetMenuActions(TArray<UBlueprintNodeSpawner*>& ActionListOut) const
+void UAnimGraphNode_Base::GetMenuActions(FBlueprintActionDatabaseRegistrar& ActionRegistrar) const
 {
 	UBlueprintNodeSpawner* NodeSpawner = UBlueprintNodeSpawner::Create(GetClass());
 	check(NodeSpawner != nullptr);
-	ActionListOut.Add(NodeSpawner);
+	ActionRegistrar.AddBlueprintAction(NodeSpawner);
 }
 
 FText UAnimGraphNode_Base::GetMenuCategory() const
diff --git a/Engine/Source/Editor/AnimGraph/Private/AnimGraphNode_BlendSpaceEvaluator.cpp b/Engine/Source/Editor/AnimGraph/Private/AnimGraphNode_BlendSpaceEvaluator.cpp
--- a/Engine/Source/Editor/AnimGraph/Private/AnimGraphNode_BlendSpaceEvaluator.cpp
+++ b/Engine/Source/Editor/AnimGraph/Private/AnimGraphNode_BlendSpaceEvaluator.cpp
@@ -54,7 +54,7 @@ void UAnimGraphNode_BlendSpaceEvaluator::GetMenuEntries(FGraphContextMenuBuilder
 	// You can convert a regular blend space player to an evaluator via the right click context menu
 }
 
-void UAnimGraphNode_BlendSpaceEvaluator::GetMenuActions(TArray<UBlueprintNodeSpawner*>& ActionListOut) const
+void UAnimGraphNode_BlendSpaceEvaluator::GetMenuActions(FBlueprintActionDatabaseRegistrar& ActionRegistrar) const
 {
 	// Intentionally empty so that we don't get duplicate blend space entries.
 	// You can convert a regular blend space player to an evaluator via the right click context menu
diff --git a/Engine/Source/Editor/AnimGraph/Private/AnimGraphNode_BlendSpacePlayer.cpp b/Engine/Source/Editor/AnimGraph/Private/AnimGraphNode_BlendSpacePlayer.cpp
--- a/Engine/Source/Editor/AnimGraph/Private/AnimGraphNode_BlendSpacePlayer.cpp
+++ b/Engine/Source/Editor/AnimGraph/Private/AnimGraphNode_BlendSpacePlayer.cpp
@@ -5,6 +5,7 @@
 #include "CompilerResultsLog.h"
 #include "AnimGraphNode_BlendSpacePlayer.h"
 #include "BlueprintNodeSpawner.h"
+#include "BlueprintActionDatabaseRegistrar.h"
 
 /////////////////////////////////////////////////////
 // UAnimGraphNode_BlendSpacePlayer
@@ -55,7 +56,7 @@ void UAnimGraphNode_BlendSpacePlayer::GetMenuEntries(FGraphContextMenuBuilder& C
 	GetBlendSpaceEntries(bWantAimOffsets, ContextMenuBuilder);
 }
 
-void UAnimGraphNode_BlendSpacePlayer::GetMenuActions(TArray<UBlueprintNodeSpawner*>& ActionListOut) const
+void UAnimGraphNode_BlendSpacePlayer::GetMenuActions(FBlueprintActionDatabaseRegistrar& ActionRegistrar) const
 {
 	auto PostSpawnSetupLambda = [](UEdGraphNode* NewNode, bool /*bIsTemplateNode*/, TWeakObjectPtr<UBlendSpaceBase> BlendSpace)
 	{
@@ -74,7 +75,7 @@ void UAnimGraphNode_BlendSpacePlayer::GetMenuActions(TArray<UBlueprintNodeSpawne
 			UBlueprintNodeSpawner* NodeSpawner = UBlueprintNodeSpawner::Create(GetClass());
 			check(NodeSpawner != nullptr);
 			// @TODO: Need to file this one under the BlendSpace asset (so it can be cleared/updated when the asset is)
-			ActionListOut.Add(NodeSpawner);
+			ActionRegistrar.AddBlueprintAction(NodeSpawner);
 
 			TWeakObjectPtr<UBlendSpaceBase> BlendSpacePtr = BlendSpace;
 			NodeSpawner->CustomizeNodeDelegate = UBlueprintNodeSpawner::FCustomizeNodeDelegate::CreateStatic(PostSpawnSetupLambda, BlendSpacePtr);
diff --git a/Engine/Source/Editor/AnimGraph/Private/AnimGraphNode_Root.cpp b/Engine/Source/Editor/AnimGraph/Private/AnimGraphNode_Root.cpp
--- a/Engine/Source/Editor/AnimGraph/Private/AnimGraphNode_Root.cpp
+++ b/Engine/Source/Editor/AnimGraph/Private/AnimGraphNode_Root.cpp
@@ -68,7 +68,7 @@ void UAnimGraphNode_Root::GetMenuEntries(FGraphContextMenuBuilder& ContextMenuBu
 	// Intentionally empty.  This node is autogenerated when a new graph is created.
 }
 
-void UAnimGraphNode_Root::GetMenuActions(TArray<UBlueprintNodeSpawner*>& ActionListOut) const
+void UAnimGraphNode_Root::GetMenuActions(FBlueprintActionDatabaseRegistrar& ActionRegistrar) const
 {
 	// Intentionally empty. This node is auto-generated when a new graph is created.
 }
diff --git a/Engine/Source/Editor/AnimGraph/Private/AnimGraphNode_RotationOffsetBlendSpace.cpp b/Engine/Source/Editor/AnimGraph/Private/AnimGraphNode_RotationOffsetBlendSpace.cpp
--- a/Engine/Source/Editor/AnimGraph/Private/AnimGraphNode_RotationOffsetBlendSpace.cpp
+++ b/Engine/Source/Editor/AnimGraph/Private/AnimGraphNode_RotationOffsetBlendSpace.cpp
@@ -5,6 +5,7 @@
 #include "GraphEditorActions.h"
 #include "CompilerResultsLog.h"
 #include "BlueprintNodeSpawner.h"
+#include "BlueprintActionDatabaseRegistrar.h"
 
 /////////////////////////////////////////////////////
 // UAnimGraphNode_RotationOffsetBlendSpace
@@ -44,7 +45,7 @@ void UAnimGraphNode_RotationOffsetBlendSpace::GetMenuEntries(FGraphContextMenuBu
 	GetBlendSpaceEntries(bWantAimOffsets, ContextMenuBuilder);
 }
 
-void UAnimGraphNode_RotationOffsetBlendSpace::GetMenuActions(TArray<UBlueprintNodeSpawner*>& ActionListOut) const
+void UAnimGraphNode_RotationOffsetBlendSpace::GetMenuActions(FBlueprintActionDatabaseRegistrar& ActionRegistrar) const
 {
 	auto PostSpawnSetupLambda = [](UEdGraphNode* NewNode, bool /*bIsTemplateNode*/, TWeakObjectPtr<UBlendSpaceBase> BlendSpace)
 	{
@@ -63,7 +64,7 @@ void UAnimGraphNode_RotationOffsetBlendSpace::GetMenuActions(TArray<UBlueprintNo
 			UBlueprintNodeSpawner* NodeSpawner = UBlueprintNodeSpawner::Create(GetClass());
 			check(NodeSpawner != nullptr);
 			// @TODO: Need to file this one under the BlendSpace asset (so it can be cleared/updated when the asset is)
-			ActionListOut.Add(NodeSpawner);
+			ActionRegistrar.AddBlueprintAction(NodeSpawner);
 
 			TWeakObjectPtr<UBlendSpaceBase> BlendSpacePtr = BlendSpace;
 			NodeSpawner->CustomizeNodeDelegate = UBlueprintNodeSpawner::FCustomizeNodeDelegate::CreateStatic(PostSpawnSetupLambda, BlendSpacePtr);
diff --git a/Engine/Source/Editor/AnimGraph/Private/AnimGraphNode_SaveCachedPose.cpp b/Engine/Source/Editor/AnimGraph/Private/AnimGraphNode_SaveCachedPose.cpp
--- a/Engine/Source/Editor/AnimGraph/Private/AnimGraphNode_SaveCachedPose.cpp
+++ b/Engine/Source/Editor/AnimGraph/Private/AnimGraphNode_SaveCachedPose.cpp
@@ -8,6 +8,7 @@
 #include "K2ActionMenuBuilder.h" // for FK2ActionMenuBuilder::AddNewNodeAction()
 #include "AnimGraphNode_SaveCachedPose.h"
 #include "BlueprintNodeSpawner.h"
+#include "BlueprintActionDatabaseRegistrar.h"
 
 /////////////////////////////////////////////////////
 // FCachedPoseNameValidator
@@ -87,7 +88,7 @@ void UAnimGraphNode_SaveCachedPose::GetMenuEntries(FGraphContextMenuBuilder& Con
 	}
 }
 
-void UAnimGraphNode_SaveCachedPose::GetMenuActions(TArray<UBlueprintNodeSpawner*>& ActionListOut) const
+void UAnimGraphNode_SaveCachedPose::GetMenuActions(FBlueprintActionDatabaseRegistrar& ActionRegistrar) const
 {
 	auto PostSpawnSetupLambda = [](UEdGraphNode* NewNode, bool bIsTemplateNode)
 	{
@@ -103,7 +104,7 @@ void UAnimGraphNode_SaveCachedPose::GetMenuActions(TArray<UBlueprintNodeSpawner*
 	UBlueprintNodeSpawner* NodeSpawner = UBlueprintNodeSpawner::Create(GetClass());
 	NodeSpawner->CustomizeNodeDelegate = UBlueprintNodeSpawner::FCustomizeNodeDelegate::CreateStatic(PostSpawnSetupLambda);
 
-	ActionListOut.Add(NodeSpawner);
+	ActionRegistrar.AddBlueprintAction(NodeSpawner);
 }
 
 bool UAnimGraphNode_SaveCachedPose::IsCompatibleWithGraph(const UEdGraph* TargetGraph) const
diff --git a/Engine/Source/Editor/AnimGraph/Private/AnimGraphNode_SequenceEvaluator.cpp b/Engine/Source/Editor/AnimGraph/Private/AnimGraphNode_SequenceEvaluator.cpp
--- a/Engine/Source/Editor/AnimGraph/Private/AnimGraphNode_SequenceEvaluator.cpp
+++ b/Engine/Source/Editor/AnimGraph/Private/AnimGraphNode_SequenceEvaluator.cpp
@@ -70,7 +70,7 @@ void UAnimGraphNode_SequenceEvaluator::GetMenuEntries(FGraphContextMenuBuilder&
 	// Intentionally empty; you can drop down a regular sequence player and convert into a sequence evaluator in the right-click menu.
 }
 
-void UAnimGraphNode_SequenceEvaluator::GetMenuActions(TArray<UBlueprintNodeSpawner*>& ActionListOut) const
+void UAnimGraphNode_SequenceEvaluator::GetMenuActions(FBlueprintActionDatabaseRegistrar& ActionRegistrar) const
 {
 	// Intentionally empty; you can drop down a regular sequence player and convert into a sequence evaluator in the right-click menu.
 }
diff --git a/Engine/Source/Editor/AnimGraph/Private/AnimGraphNode_TransitionPoseEvaluator.cpp b/Engine/Source/Editor/AnimGraph/Private/AnimGraphNode_TransitionPoseEvaluator.cpp
--- a/Engine/Source/Editor/AnimGraph/Private/AnimGraphNode_TransitionPoseEvaluator.cpp
+++ b/Engine/Source/Editor/AnimGraph/Private/AnimGraphNode_TransitionPoseEvaluator.cpp
@@ -63,7 +63,7 @@ void UAnimGraphNode_TransitionPoseEvaluator::GetMenuEntries(FGraphContextMenuBui
 	// Intentionally empty: Don't allow an option to create them, as they're auto-created when custom blend graphs are made
 }
 
-void UAnimGraphNode_TransitionPoseEvaluator::GetMenuActions(TArray<UBlueprintNodeSpawner*>& ActionListOut) const
+void UAnimGraphNode_TransitionPoseEvaluator::GetMenuActions(FBlueprintActionDatabaseRegistrar& ActionRegistrar) const
 {
 	// Intentionally empty: Don't allow an option to create them, as they're auto-created when custom blend graphs are made
 }
diff --git a/Engine/Source/Editor/AnimGraph/Private/AnimGraphNode_TransitionResult.cpp b/Engine/Source/Editor/AnimGraph/Private/AnimGraphNode_TransitionResult.cpp
--- a/Engine/Source/Editor/AnimGraph/Private/AnimGraphNode_TransitionResult.cpp
+++ b/Engine/Source/Editor/AnimGraph/Private/AnimGraphNode_TransitionResult.cpp
@@ -38,7 +38,7 @@ void UAnimGraphNode_TransitionResult::GetMenuEntries(FGraphContextMenuBuilder& C
 }
 
 
-void UAnimGraphNode_TransitionResult::GetMenuActions(TArray<UBlueprintNodeSpawner*>& ActionListOut) const
+void UAnimGraphNode_TransitionResult::GetMenuActions(FBlueprintActionDatabaseRegistrar& ActionRegistrar) const
 {
 	// Intentionally empty. This node is auto-generated when a transition graph is created.
 }
diff --git a/Engine/Source/Editor/BlueprintGraph/Classes/K2Node.h b/Engine/Source/Editor/BlueprintGraph/Classes/K2Node.h
--- a/Engine/Source/Editor/BlueprintGraph/Classes/K2Node.h
+++ b/Engine/Source/Editor/BlueprintGraph/Classes/K2Node.h
@@ -5,6 +5,7 @@
 #include "K2Node.generated.h"
 
 class UBlueprintNodeSpawner;
+class FBlueprintActionDatabaseRegistrar;
 
 /** Helper struct to allow us to redirect properties and functions through renames and additionally between classes if necessary */
 struct FFieldRemapInfo
@@ -297,7 +298,7 @@ class UK2Node : public UEdGraphNode
 	 *
 	 * @param  ActionListOut	The list to be populated with new spawners.
 	 */
-	virtual void GetMenuActions(TArray<UBlueprintNodeSpawner*>& ActionListOut) const {}
+	virtual void GetMenuActions(FBlueprintActionDatabaseRegistrar& ActionRegisrar) const {}
 
 	/**
 	 * Override to provide a default category for specific node types to be 
diff --git a/Engine/Source/Editor/BlueprintGraph/Classes/K2Node_AssignmentStatement.h b/Engine/Source/Editor/BlueprintGraph/Classes/K2Node_AssignmentStatement.h
--- a/Engine/Source/Editor/BlueprintGraph/Classes/K2Node_AssignmentStatement.h
+++ b/Engine/Source/Editor/BlueprintGraph/Classes/K2Node_AssignmentStatement.h
@@ -26,7 +26,7 @@ class UK2Node_AssignmentStatement : public UK2Node
 	virtual void PostReconstructNode() override;
 	virtual void NotifyPinConnectionListChanged(UEdGraphPin* Pin) override;
 	virtual class FNodeHandlingFunctor* CreateNodeHandler(class FKismetCompilerContext& CompilerContext) const override;
-	virtual void GetMenuActions(TArray<UBlueprintNodeSpawner*>& ActionListOut) const override;
+	virtual void GetMenuActions(FBlueprintActionDatabaseRegistrar& ActionRegistrar) const override;
 	virtual FText GetMenuCategory() const override;
 	virtual int32 GetNodeRefreshPriority() const override { return EBaseNodeRefreshPriority::Low_UsesDependentWildcard; }
 	// End UK2Node interface
diff --git a/Engine/Source/Editor/BlueprintGraph/Classes/K2Node_BaseAsyncTask.h b/Engine/Source/Editor/BlueprintGraph/Classes/K2Node_BaseAsyncTask.h
--- a/Engine/Source/Editor/BlueprintGraph/Classes/K2Node_BaseAsyncTask.h
+++ b/Engine/Source/Editor/BlueprintGraph/Classes/K2Node_BaseAsyncTask.h
@@ -27,7 +27,7 @@ class BLUEPRINTGRAPH_API UK2Node_BaseAsyncTask : public UK2Node
 	virtual void ExpandNode(class FKismetCompilerContext& CompilerContext, UEdGraph* SourceGraph) override;
 	virtual bool HasExternalBlueprintDependencies(TArray<class UStruct*>* OptionalOutput) const override;
 	virtual FName GetCornerIcon() const override;
-	virtual void GetMenuActions(TArray<UBlueprintNodeSpawner*>& ActionListOut) const override;
+	virtual void GetMenuActions(FBlueprintActionDatabaseRegistrar& ActionRegistrar) const override;
 	virtual FText GetMenuCategory() const override;
 	// End of UK2Node interface
 
diff --git a/Engine/Source/Editor/BlueprintGraph/Classes/K2Node_BreakStruct.h b/Engine/Source/Editor/BlueprintGraph/Classes/K2Node_BreakStruct.h
--- a/Engine/Source/Editor/BlueprintGraph/Classes/K2Node_BreakStruct.h
+++ b/Engine/Source/Editor/BlueprintGraph/Classes/K2Node_BreakStruct.h
@@ -26,7 +26,7 @@ class UK2Node_BreakStruct : public UK2Node_StructMemberGet
 	virtual bool DrawNodeAsVariable() const override { return false; }
 	virtual ERedirectType DoPinsMatchForReconstruction(const UEdGraphPin* NewPin, int32 NewPinIndex, const UEdGraphPin* OldPin, int32 OldPinIndex) const override;
 	virtual class FNodeHandlingFunctor* CreateNodeHandler(class FKismetCompilerContext& CompilerContext) const override;
-	virtual void GetMenuActions(TArray<UBlueprintNodeSpawner*>& ActionListOut) const override;
+	virtual void GetMenuActions(FBlueprintActionDatabaseRegistrar& ActionRegistrar) const override;
 	virtual FText GetMenuCategory() const override;
 	// End K2Node interface
 };
diff --git a/Engine/Source/Editor/BlueprintGraph/Classes/K2Node_CastByteToEnum.h b/Engine/Source/Editor/BlueprintGraph/Classes/K2Node_CastByteToEnum.h
--- a/Engine/Source/Editor/BlueprintGraph/Classes/K2Node_CastByteToEnum.h
+++ b/Engine/Source/Editor/BlueprintGraph/Classes/K2Node_CastByteToEnum.h
@@ -31,7 +31,7 @@ class UK2Node_CastByteToEnum : public UK2Node
 	virtual bool IsNodePure() const override { return true; }
 	FNodeHandlingFunctor* CreateNodeHandler(FKismetCompilerContext& CompilerContext) const override;
 	virtual void ExpandNode(class FKismetCompilerContext& CompilerContext, UEdGraph* SourceGraph) override;
-	virtual void GetMenuActions(TArray<UBlueprintNodeSpawner*>& ActionListOut) const override;
+	virtual void GetMenuActions(FBlueprintActionDatabaseRegistrar& ActionRegistrar) const override;
 	virtual FText GetMenuCategory() const override;
 	// End UK2Node interface
 
diff --git a/Engine/Source/Editor/BlueprintGraph/Classes/K2Node_ConstructObjectFromClass.h b/Engine/Source/Editor/BlueprintGraph/Classes/K2Node_ConstructObjectFromClass.h
--- a/Engine/Source/Editor/BlueprintGraph/Classes/K2Node_ConstructObjectFromClass.h
+++ b/Engine/Source/Editor/BlueprintGraph/Classes/K2Node_ConstructObjectFromClass.h
@@ -25,7 +25,7 @@ class BLUEPRINTGRAPH_API UK2Node_ConstructObjectFromClass : public UK2Node
 	virtual bool IsNodeSafeToIgnore() const override { return true; }
 	virtual void ReallocatePinsDuringReconstruction(TArray<UEdGraphPin*>& OldPins) override;
 	virtual void GetNodeAttributes( TArray<TKeyValuePair<FString, FString>>& OutNodeAttributes ) const override;
-	virtual void GetMenuActions(TArray<UBlueprintNodeSpawner*>& ActionListOut) const override;
+	virtual void GetMenuActions(FBlueprintActionDatabaseRegistrar& ActionRegisrar) const override;
 	virtual FText GetMenuCategory() const override;
 	virtual class FNodeHandlingFunctor* CreateNodeHandler(class FKismetCompilerContext& CompilerContext) const override;
 	// End UK2Node interface
diff --git a/Engine/Source/Editor/BlueprintGraph/Classes/K2Node_CustomEvent.h b/Engine/Source/Editor/BlueprintGraph/Classes/K2Node_CustomEvent.h
--- a/Engine/Source/Editor/BlueprintGraph/Classes/K2Node_CustomEvent.h
+++ b/Engine/Source/Editor/BlueprintGraph/Classes/K2Node_CustomEvent.h
@@ -29,7 +29,7 @@ class UK2Node_CustomEvent : public UK2Node_Event
 
 	// Begin UK2Node interface
 	BLUEPRINTGRAPH_API virtual void ValidateNodeDuringCompilation(class FCompilerResultsLog& MessageLog) const override;
-	virtual void GetMenuActions(TArray<UBlueprintNodeSpawner*>& ActionListOut) const override;
+	virtual void GetMenuActions(FBlueprintActionDatabaseRegistrar& ActionRegistrar) const override;
 	// End UK2Node interface
 
 	// Begin UK2Node_EditablePinBase interface
diff --git a/Engine/Source/Editor/BlueprintGraph/Classes/K2Node_DoOnceMultiInput.h b/Engine/Source/Editor/BlueprintGraph/Classes/K2Node_DoOnceMultiInput.h
--- a/Engine/Source/Editor/BlueprintGraph/Classes/K2Node_DoOnceMultiInput.h
+++ b/Engine/Source/Editor/BlueprintGraph/Classes/K2Node_DoOnceMultiInput.h
@@ -51,7 +51,7 @@ public:
 	// UK2Node interface
 	virtual void GetContextMenuActions(const FGraphNodeContextMenuBuilder& Context) const override;
 	virtual void ExpandNode(class FKismetCompilerContext& CompilerContext, UEdGraph* SourceGraph) override;
-	virtual void GetMenuActions(TArray<UBlueprintNodeSpawner*>& ActionListOut) const override;
+	virtual void GetMenuActions(FBlueprintActionDatabaseRegistrar& ActionRegistrar) const override;
 	virtual FText GetMenuCategory() const override;
 	// End of UK2Node interface
 };
\ No newline at end of file
diff --git a/Engine/Source/Editor/BlueprintGraph/Classes/K2Node_EnumEquality.h b/Engine/Source/Editor/BlueprintGraph/Classes/K2Node_EnumEquality.h
--- a/Engine/Source/Editor/BlueprintGraph/Classes/K2Node_EnumEquality.h
+++ b/Engine/Source/Editor/BlueprintGraph/Classes/K2Node_EnumEquality.h
@@ -27,7 +27,7 @@ class UK2Node_EnumEquality : public UK2Node
 	virtual void NotifyPinConnectionListChanged(UEdGraphPin* Pin) override;
 	virtual class FNodeHandlingFunctor* CreateNodeHandler(class FKismetCompilerContext& CompilerContext) const override;
 	virtual void ExpandNode(class FKismetCompilerContext& CompilerContext, UEdGraph* SourceGraph) override;
-	virtual void GetMenuActions(TArray<UBlueprintNodeSpawner*>& ActionListOut) const override;
+	virtual void GetMenuActions(FBlueprintActionDatabaseRegistrar& ActionRegistrar) const override;
 	virtual FText GetMenuCategory() const override;
 	virtual int32 GetNodeRefreshPriority() const override { return EBaseNodeRefreshPriority::Low_UsesDependentWildcard; }
 	virtual bool IsConnectionDisallowed(const UEdGraphPin* MyPin, const UEdGraphPin* OtherPin, FString& OutReason) const override;
diff --git a/Engine/Source/Editor/BlueprintGraph/Classes/K2Node_EnumLiteral.h b/Engine/Source/Editor/BlueprintGraph/Classes/K2Node_EnumLiteral.h
--- a/Engine/Source/Editor/BlueprintGraph/Classes/K2Node_EnumLiteral.h
+++ b/Engine/Source/Editor/BlueprintGraph/Classes/K2Node_EnumLiteral.h
@@ -25,7 +25,7 @@ class UK2Node_EnumLiteral : public UK2Node, public INodeDependingOnEnumInterface
 	// Begin UK2Node interface
 	virtual bool IsNodePure() const override { return true; }
 	virtual FNodeHandlingFunctor* CreateNodeHandler(FKismetCompilerContext& CompilerContext) const override;
-	virtual void GetMenuActions(TArray<UBlueprintNodeSpawner*>& ActionListOut) const override;
+	virtual void GetMenuActions(FBlueprintActionDatabaseRegistrar& ActionRegistrar) const override;
 	virtual FText GetMenuCategory() const override;
 	// End UK2Node interface
 
diff --git a/Engine/Source/Editor/BlueprintGraph/Classes/K2Node_ExecutionSequence.h b/Engine/Source/Editor/BlueprintGraph/Classes/K2Node_ExecutionSequence.h
--- a/Engine/Source/Editor/BlueprintGraph/Classes/K2Node_ExecutionSequence.h
+++ b/Engine/Source/Editor/BlueprintGraph/Classes/K2Node_ExecutionSequence.h
@@ -21,7 +21,7 @@ class UK2Node_ExecutionSequence : public UK2Node
 	// Begin UK2Node interface
 	void ReallocatePinsDuringReconstruction(TArray<UEdGraphPin*>& OldPins) override;
 	virtual class FNodeHandlingFunctor* CreateNodeHandler(class FKismetCompilerContext& CompilerContext) const override;
-	virtual void GetMenuActions(TArray<UBlueprintNodeSpawner*>& ActionListOut) const override;
+	virtual void GetMenuActions(FBlueprintActionDatabaseRegistrar& ActionRegistrar) const override;
 	virtual FText GetMenuCategory() const override;
 	// End UK2Node interface
 
diff --git a/Engine/Source/Editor/BlueprintGraph/Classes/K2Node_ForEachElementInEnum.h b/Engine/Source/Editor/BlueprintGraph/Classes/K2Node_ForEachElementInEnum.h
--- a/Engine/Source/Editor/BlueprintGraph/Classes/K2Node_ForEachElementInEnum.h
+++ b/Engine/Source/Editor/BlueprintGraph/Classes/K2Node_ForEachElementInEnum.h
@@ -26,7 +26,7 @@ class UK2Node_ForEachElementInEnum : public UK2Node
 	virtual bool IsNodePure() const override { return false; }
 	virtual void ExpandNode(class FKismetCompilerContext& CompilerContext, UEdGraph* SourceGraph) override;
 	virtual void ValidateNodeDuringCompilation(class FCompilerResultsLog& MessageLog) const override;
-	virtual void GetMenuActions(TArray<UBlueprintNodeSpawner*>& ActionListOut) const override;
+	virtual void GetMenuActions(FBlueprintActionDatabaseRegistrar& ActionRegistrar) const override;
 	virtual FText GetMenuCategory() const override;
 	// End UK2Node interface
 };
diff --git a/Engine/Source/Editor/BlueprintGraph/Classes/K2Node_FormatText.h b/Engine/Source/Editor/BlueprintGraph/Classes/K2Node_FormatText.h
--- a/Engine/Source/Editor/BlueprintGraph/Classes/K2Node_FormatText.h
+++ b/Engine/Source/Editor/BlueprintGraph/Classes/K2Node_FormatText.h
@@ -28,7 +28,7 @@ class UK2Node_FormatText : public UK2Node
 	virtual bool NodeCausesStructuralBlueprintChange() const override { return true; }
 	virtual void ExpandNode(class FKismetCompilerContext& CompilerContext, UEdGraph* SourceGraph) override;
 	virtual ERedirectType DoPinsMatchForReconstruction(const UEdGraphPin* NewPin, int32 NewPinIndex, const UEdGraphPin* OldPin, int32 OldPinIndex) const override;
-	virtual void GetMenuActions(TArray<UBlueprintNodeSpawner*>& ActionListOut) const override;
+	virtual void GetMenuActions(FBlueprintActionDatabaseRegistrar& ActionRegistrar) const override;
 	virtual FText GetMenuCategory() const override;
 	// End UK2Node interface.
 
diff --git a/Engine/Source/Editor/BlueprintGraph/Classes/K2Node_GetEnumeratorName.h b/Engine/Source/Editor/BlueprintGraph/Classes/K2Node_GetEnumeratorName.h
--- a/Engine/Source/Editor/BlueprintGraph/Classes/K2Node_GetEnumeratorName.h
+++ b/Engine/Source/Editor/BlueprintGraph/Classes/K2Node_GetEnumeratorName.h
@@ -28,7 +28,7 @@ class UK2Node_GetEnumeratorName : public UK2Node
 	virtual void EarlyValidation(class FCompilerResultsLog& MessageLog) const override;
 	virtual void PostReconstructNode() override;
 	virtual void NotifyPinConnectionListChanged(UEdGraphPin* Pin) override;
-	virtual void GetMenuActions(TArray<UBlueprintNodeSpawner*>& ActionListOut) const override;
+	virtual void GetMenuActions(FBlueprintActionDatabaseRegistrar& ActionRegistrar) const override;
 	virtual FText GetMenuCategory() const override;
 	virtual int32 GetNodeRefreshPriority() const override { return EBaseNodeRefreshPriority::Low_UsesDependentWildcard; }
 	// End UK2Node interface
diff --git a/Engine/Source/Editor/BlueprintGraph/Classes/K2Node_GetEnumeratorNameAsString.h b/Engine/Source/Editor/BlueprintGraph/Classes/K2Node_GetEnumeratorNameAsString.h
--- a/Engine/Source/Editor/BlueprintGraph/Classes/K2Node_GetEnumeratorNameAsString.h
+++ b/Engine/Source/Editor/BlueprintGraph/Classes/K2Node_GetEnumeratorNameAsString.h
@@ -17,7 +17,7 @@ class UK2Node_GetEnumeratorNameAsString : public UK2Node_GetEnumeratorName
 	// End UEdGraphNode interface
 
 	// UK2Node interface
-	virtual void GetMenuActions(TArray<UBlueprintNodeSpawner*>& ActionListOut) const override;
+	virtual void GetMenuActions(FBlueprintActionDatabaseRegistrar& ActionRegistrar) const override;
 	virtual FText GetMenuCategory() const override;
 	// End UK2Node interface
 
diff --git a/Engine/Source/Editor/BlueprintGraph/Classes/K2Node_GetInputAxisKeyValue.h b/Engine/Source/Editor/BlueprintGraph/Classes/K2Node_GetInputAxisKeyValue.h
--- a/Engine/Source/Editor/BlueprintGraph/Classes/K2Node_GetInputAxisKeyValue.h
+++ b/Engine/Source/Editor/BlueprintGraph/Classes/K2Node_GetInputAxisKeyValue.h
@@ -33,7 +33,7 @@ class UK2Node_GetInputAxisKeyValue : public UK2Node_CallFunction
 	virtual bool ShouldShowNodeProperties() const override { return true; }
 	virtual UClass* GetDynamicBindingClass() const override;
 	virtual void RegisterDynamicBinding(UDynamicBlueprintBinding* BindingObject) const override;
-	virtual void GetMenuActions(TArray<UBlueprintNodeSpawner*>& ActionListOut) const override;
+	virtual void GetMenuActions(FBlueprintActionDatabaseRegistrar& ActionRegistrar) const override;
 	virtual FText GetMenuCategory() const override;
 	// End UK2Node interface
 	
diff --git a/Engine/Source/Editor/BlueprintGraph/Classes/K2Node_GetInputAxisValue.h b/Engine/Source/Editor/BlueprintGraph/Classes/K2Node_GetInputAxisValue.h
--- a/Engine/Source/Editor/BlueprintGraph/Classes/K2Node_GetInputAxisValue.h
+++ b/Engine/Source/Editor/BlueprintGraph/Classes/K2Node_GetInputAxisValue.h
@@ -32,7 +32,7 @@ class UK2Node_GetInputAxisValue : public UK2Node_CallFunction
 	virtual bool ShouldShowNodeProperties() const override { return true; }
 	virtual UClass* GetDynamicBindingClass() const override;
 	virtual void RegisterDynamicBinding(UDynamicBlueprintBinding* BindingObject) const override;
-	virtual void GetMenuActions(TArray<UBlueprintNodeSpawner*>& ActionListOut) const override;
+	virtual void GetMenuActions(FBlueprintActionDatabaseRegistrar& ActionRegistrar) const override;
 	virtual FText GetMenuCategory() const override;
 	// End UK2Node interface
 	
diff --git a/Engine/Source/Editor/BlueprintGraph/Classes/K2Node_GetInputVectorAxisValue.h b/Engine/Source/Editor/BlueprintGraph/Classes/K2Node_GetInputVectorAxisValue.h
--- a/Engine/Source/Editor/BlueprintGraph/Classes/K2Node_GetInputVectorAxisValue.h
+++ b/Engine/Source/Editor/BlueprintGraph/Classes/K2Node_GetInputVectorAxisValue.h
@@ -18,7 +18,7 @@ class UK2Node_GetInputVectorAxisValue : public UK2Node_GetInputAxisKeyValue
 	virtual bool ShouldShowNodeProperties() const override { return true; }
 	virtual UClass* GetDynamicBindingClass() const override;
 	virtual void RegisterDynamicBinding(UDynamicBlueprintBinding* BindingObject) const override;
-	virtual void GetMenuActions(TArray<UBlueprintNodeSpawner*>& ActionListOut) const override;
+	virtual void GetMenuActions(FBlueprintActionDatabaseRegistrar& ActionRegistrar) const override;
 	// End UK2Node interface
 	
 	void Initialize(const FKey AxisKey);
diff --git a/Engine/Source/Editor/BlueprintGraph/Classes/K2Node_GetNumEnumEntries.h b/Engine/Source/Editor/BlueprintGraph/Classes/K2Node_GetNumEnumEntries.h
--- a/Engine/Source/Editor/BlueprintGraph/Classes/K2Node_GetNumEnumEntries.h
+++ b/Engine/Source/Editor/BlueprintGraph/Classes/K2Node_GetNumEnumEntries.h
@@ -22,7 +22,7 @@ class UK2Node_GetNumEnumEntries : public UK2Node, public INodeDependingOnEnumInt
 	// Begin UK2Node interface
 	virtual bool IsNodePure() const override { return true; }
 	virtual void ExpandNode(class FKismetCompilerContext& CompilerContext, UEdGraph* SourceGraph) override;
-	virtual void GetMenuActions(TArray<UBlueprintNodeSpawner*>& ActionListOut) const override;
+	virtual void GetMenuActions(FBlueprintActionDatabaseRegistrar& ActionRegistrar) const override;
 	virtual FText GetMenuCategory() const override;
 	// End UK2Node interface
 
diff --git a/Engine/Source/Editor/BlueprintGraph/Classes/K2Node_IfThenElse.h b/Engine/Source/Editor/BlueprintGraph/Classes/K2Node_IfThenElse.h
--- a/Engine/Source/Editor/BlueprintGraph/Classes/K2Node_IfThenElse.h
+++ b/Engine/Source/Editor/BlueprintGraph/Classes/K2Node_IfThenElse.h
@@ -21,7 +21,7 @@ class UK2Node_IfThenElse : public UK2Node
 
 	// Begin K2Node interface.
 	virtual class FNodeHandlingFunctor* CreateNodeHandler(class FKismetCompilerContext& CompilerContext) const override;
-	virtual void GetMenuActions(TArray<UBlueprintNodeSpawner*>& ActionListOut) const override;
+	virtual void GetMenuActions(FBlueprintActionDatabaseRegistrar& ActionRegistrar) const override;
 	virtual FText GetMenuCategory() const override;
 	// End K2Node interface.
 
diff --git a/Engine/Source/Editor/BlueprintGraph/Classes/K2Node_InputAction.h b/Engine/Source/Editor/BlueprintGraph/Classes/K2Node_InputAction.h
--- a/Engine/Source/Editor/BlueprintGraph/Classes/K2Node_InputAction.h
+++ b/Engine/Source/Editor/BlueprintGraph/Classes/K2Node_InputAction.h
@@ -42,7 +42,7 @@ class UK2Node_InputAction : public UK2Node
 	virtual void ValidateNodeDuringCompilation(class FCompilerResultsLog& MessageLog) const override;
 	virtual bool ShouldShowNodeProperties() const override { return true; }
 	virtual void ExpandNode(FKismetCompilerContext& CompilerContext, UEdGraph* SourceGraph) override;
-	virtual void GetMenuActions(TArray<UBlueprintNodeSpawner*>& ActionListOut) const override;
+	virtual void GetMenuActions(FBlueprintActionDatabaseRegistrar& ActionRegistrar) const override;
 	virtual FText GetMenuCategory() const override;
 	// End UK2Node interface
 
diff --git a/Engine/Source/Editor/BlueprintGraph/Classes/K2Node_InputAxisEvent.h b/Engine/Source/Editor/BlueprintGraph/Classes/K2Node_InputAxisEvent.h
--- a/Engine/Source/Editor/BlueprintGraph/Classes/K2Node_InputAxisEvent.h
+++ b/Engine/Source/Editor/BlueprintGraph/Classes/K2Node_InputAxisEvent.h
@@ -39,7 +39,7 @@ class UK2Node_InputAxisEvent : public UK2Node_Event
 	virtual bool ShouldShowNodeProperties() const override { return true; }
 	virtual UClass* GetDynamicBindingClass() const override;
 	virtual void RegisterDynamicBinding(UDynamicBlueprintBinding* BindingObject) const override;
-	virtual void GetMenuActions(TArray<UBlueprintNodeSpawner*>& ActionListOut) const override;
+	virtual void GetMenuActions(FBlueprintActionDatabaseRegistrar& ActionRegistrar) const override;
 	virtual FText GetMenuCategory() const override;
 	// End UK2Node interface
 
diff --git a/Engine/Source/Editor/BlueprintGraph/Classes/K2Node_InputAxisKeyEvent.h b/Engine/Source/Editor/BlueprintGraph/Classes/K2Node_InputAxisKeyEvent.h
--- a/Engine/Source/Editor/BlueprintGraph/Classes/K2Node_InputAxisKeyEvent.h
+++ b/Engine/Source/Editor/BlueprintGraph/Classes/K2Node_InputAxisKeyEvent.h
@@ -36,7 +36,7 @@ class UK2Node_InputAxisKeyEvent : public UK2Node_Event
 	virtual bool ShouldShowNodeProperties() const override{ return true; }
 	virtual UClass* GetDynamicBindingClass() const override;
 	virtual void RegisterDynamicBinding(UDynamicBlueprintBinding* BindingObject) const override;
-	virtual void GetMenuActions(TArray<UBlueprintNodeSpawner*>& ActionListOut) const override;
+	virtual void GetMenuActions(FBlueprintActionDatabaseRegistrar& ActionRegistrar) const override;
 	virtual FText GetMenuCategory() const override;
 	// End UK2Node interface
 
diff --git a/Engine/Source/Editor/BlueprintGraph/Classes/K2Node_InputKey.h b/Engine/Source/Editor/BlueprintGraph/Classes/K2Node_InputKey.h
--- a/Engine/Source/Editor/BlueprintGraph/Classes/K2Node_InputKey.h
+++ b/Engine/Source/Editor/BlueprintGraph/Classes/K2Node_InputKey.h
@@ -50,7 +50,7 @@ class UK2Node_InputKey : public UK2Node
 	virtual bool ShouldShowNodeProperties() const override { return true; }
 	virtual void ValidateNodeDuringCompilation(class FCompilerResultsLog& MessageLog) const override;
 	virtual void ExpandNode(FKismetCompilerContext& CompilerContext, UEdGraph* SourceGraph) override;
-	virtual void GetMenuActions(TArray<UBlueprintNodeSpawner*>& ActionListOut) const override;
+	virtual void GetMenuActions(FBlueprintActionDatabaseRegistrar& ActionRegistrar) const override;
 	virtual FText GetMenuCategory() const override;
 	// End UK2Node interface
 
diff --git a/Engine/Source/Editor/BlueprintGraph/Classes/K2Node_InputTouch.h b/Engine/Source/Editor/BlueprintGraph/Classes/K2Node_InputTouch.h
--- a/Engine/Source/Editor/BlueprintGraph/Classes/K2Node_InputTouch.h
+++ b/Engine/Source/Editor/BlueprintGraph/Classes/K2Node_InputTouch.h
@@ -36,7 +36,7 @@ class UK2Node_InputTouch : public UK2Node
 	virtual FLinearColor GetNodeTitleColor() const override;
 	virtual FText GetNodeTitle(ENodeTitleType::Type TitleType) const override;
 	virtual FName GetPaletteIcon(FLinearColor& OutColor) const override{ return TEXT("GraphEditor.TouchEvent_16x"); }
-	virtual void GetMenuActions(TArray<UBlueprintNodeSpawner*>& ActionListOut) const override;
+	virtual void GetMenuActions(FBlueprintActionDatabaseRegistrar& ActionRegistrar) const override;
 	virtual FText GetMenuCategory() const override;
 	virtual bool IsCompatibleWithGraph(const UEdGraph* TargetGraph) const override;
 	// End UEdGraphNode interface.
diff --git a/Engine/Source/Editor/BlueprintGraph/Classes/K2Node_InputVectorAxisEvent.h b/Engine/Source/Editor/BlueprintGraph/Classes/K2Node_InputVectorAxisEvent.h
--- a/Engine/Source/Editor/BlueprintGraph/Classes/K2Node_InputVectorAxisEvent.h
+++ b/Engine/Source/Editor/BlueprintGraph/Classes/K2Node_InputVectorAxisEvent.h
@@ -13,6 +13,6 @@ class UK2Node_InputVectorAxisEvent : public UK2Node_InputAxisKeyEvent
 	virtual void ValidateNodeDuringCompilation(class FCompilerResultsLog& MessageLog) const override;
 	virtual UClass* GetDynamicBindingClass() const override;
 	virtual void RegisterDynamicBinding(UDynamicBlueprintBinding* BindingObject) const override;
-	virtual void GetMenuActions(TArray<UBlueprintNodeSpawner*>& ActionListOut) const override;
+	virtual void GetMenuActions(FBlueprintActionDatabaseRegistrar& ActionRegistrar) const override;
 	// End UK2Node interface
 };
\ No newline at end of file
diff --git a/Engine/Source/Editor/BlueprintGraph/Classes/K2Node_Knot.h b/Engine/Source/Editor/BlueprintGraph/Classes/K2Node_Knot.h
--- a/Engine/Source/Editor/BlueprintGraph/Classes/K2Node_Knot.h
+++ b/Engine/Source/Editor/BlueprintGraph/Classes/K2Node_Knot.h
@@ -28,7 +28,7 @@ public:
 	virtual void NotifyPinConnectionListChanged(UEdGraphPin* Pin) override;
 	virtual void PostReconstructNode() override;
 	virtual int32 GetNodeRefreshPriority() const override { return EBaseNodeRefreshPriority::Low_UsesDependentWildcard; }
-	virtual void GetMenuActions(TArray<UBlueprintNodeSpawner*>& ActionListOut) const override;
+	virtual void GetMenuActions(FBlueprintActionDatabaseRegistrar& ActionRegistrar) const override;
 	// End of UK2Node interface
 
 	UEdGraphPin* GetInputPin() const
diff --git a/Engine/Source/Editor/BlueprintGraph/Classes/K2Node_LocalVariable.h b/Engine/Source/Editor/BlueprintGraph/Classes/K2Node_LocalVariable.h
--- a/Engine/Source/Editor/BlueprintGraph/Classes/K2Node_LocalVariable.h
+++ b/Engine/Source/Editor/BlueprintGraph/Classes/K2Node_LocalVariable.h
@@ -31,6 +31,7 @@ class UDEPRECATED_K2Node_LocalVariable : public UK2Node_TemporaryVariable
 	// Begin UK2Node interface.
 	virtual bool NodeCausesStructuralBlueprintChange() const override { return true; }
 	virtual bool ShouldShowNodeProperties() const override { return true; }
+	virtual void GetMenuActions(FBlueprintActionDatabaseRegistrar& ActionRegistrar) const override {}
 	// End UK2Node interface.
 
 	/**
diff --git a/Engine/Source/Editor/BlueprintGraph/Classes/K2Node_MakeArray.h b/Engine/Source/Editor/BlueprintGraph/Classes/K2Node_MakeArray.h
--- a/Engine/Source/Editor/BlueprintGraph/Classes/K2Node_MakeArray.h
+++ b/Engine/Source/Editor/BlueprintGraph/Classes/K2Node_MakeArray.h
@@ -36,7 +36,7 @@ public:
 	void GetContextMenuActions(const FGraphNodeContextMenuBuilder& Context) const override;
 	virtual class FNodeHandlingFunctor* CreateNodeHandler(class FKismetCompilerContext& CompilerContext) const override;
 	virtual bool IsConnectionDisallowed(const UEdGraphPin* MyPin, const UEdGraphPin* OtherPin, FString& OutReason) const override;
-	virtual void GetMenuActions(TArray<UBlueprintNodeSpawner*>& ActionListOut) const override;
+	virtual void GetMenuActions(FBlueprintActionDatabaseRegistrar& ActionRegistrar) const override;
 	virtual FText GetMenuCategory() const override;
 	virtual int32 GetNodeRefreshPriority() const override { return EBaseNodeRefreshPriority::Low_UsesDependentWildcard; }
 	// End of UK2Node interface
diff --git a/Engine/Source/Editor/BlueprintGraph/Classes/K2Node_MakeStruct.h b/Engine/Source/Editor/BlueprintGraph/Classes/K2Node_MakeStruct.h
--- a/Engine/Source/Editor/BlueprintGraph/Classes/K2Node_MakeStruct.h
+++ b/Engine/Source/Editor/BlueprintGraph/Classes/K2Node_MakeStruct.h
@@ -27,7 +27,7 @@ class UK2Node_MakeStruct : public UK2Node_StructMemberSet
 	virtual bool DrawNodeAsVariable() const override { return false; }
 	virtual class FNodeHandlingFunctor* CreateNodeHandler(class FKismetCompilerContext& CompilerContext) const override;
 	virtual ERedirectType DoPinsMatchForReconstruction(const UEdGraphPin* NewPin, int32 NewPinIndex, const UEdGraphPin* OldPin, int32 OldPinIndex)  const override;
-	virtual void GetMenuActions(TArray<UBlueprintNodeSpawner*>& ActionListOut) const override;
+	virtual void GetMenuActions(FBlueprintActionDatabaseRegistrar& ActionRegistrar) const override;
 	virtual FText GetMenuCategory() const override;
 	// End K2Node interface
 
diff --git a/Engine/Source/Editor/BlueprintGraph/Classes/K2Node_MultiGate.h b/Engine/Source/Editor/BlueprintGraph/Classes/K2Node_MultiGate.h
--- a/Engine/Source/Editor/BlueprintGraph/Classes/K2Node_MultiGate.h
+++ b/Engine/Source/Editor/BlueprintGraph/Classes/K2Node_MultiGate.h
@@ -24,7 +24,7 @@ class UK2Node_MultiGate : public UK2Node_ExecutionSequence
 	virtual void ReallocatePinsDuringReconstruction(TArray<UEdGraphPin*>& OldPins) override;
 	virtual class FNodeHandlingFunctor* CreateNodeHandler(class FKismetCompilerContext& CompilerContext) const override;
 	virtual void ExpandNode(class FKismetCompilerContext& CompilerContext, UEdGraph* SourceGraph) override;
-	virtual void GetMenuActions(TArray<UBlueprintNodeSpawner*>& ActionListOut) const override;
+	virtual void GetMenuActions(FBlueprintActionDatabaseRegistrar& ActionRegistrar) const override;
 	// End UK2Node interface
 
 	/** Getting pin access */
diff --git a/Engine/Source/Editor/BlueprintGraph/Classes/K2Node_Select.h b/Engine/Source/Editor/BlueprintGraph/Classes/K2Node_Select.h
--- a/Engine/Source/Editor/BlueprintGraph/Classes/K2Node_Select.h
+++ b/Engine/Source/Editor/BlueprintGraph/Classes/K2Node_Select.h
@@ -52,7 +52,7 @@ class UK2Node_Select : public UK2Node
 	virtual void NotifyPinConnectionListChanged(UEdGraphPin* Pin) override;
 	virtual class FNodeHandlingFunctor* CreateNodeHandler(class FKismetCompilerContext& CompilerContext) const override;
 	virtual bool IsConnectionDisallowed(const UEdGraphPin* MyPin, const UEdGraphPin* OtherPin, FString& OutReason) const override;
-	virtual void GetMenuActions(TArray<UBlueprintNodeSpawner*>& ActionListOut) const override;
+	virtual void GetMenuActions(FBlueprintActionDatabaseRegistrar& ActionRegistrar) const override;
 	virtual FText GetMenuCategory() const override;
 	virtual int32 GetNodeRefreshPriority() const override { return EBaseNodeRefreshPriority::Low_UsesDependentWildcard; }
 	// End UK2Node interface
diff --git a/Engine/Source/Editor/BlueprintGraph/Classes/K2Node_SelectEnum.h b/Engine/Source/Editor/BlueprintGraph/Classes/K2Node_SelectEnum.h
--- a/Engine/Source/Editor/BlueprintGraph/Classes/K2Node_SelectEnum.h
+++ b/Engine/Source/Editor/BlueprintGraph/Classes/K2Node_SelectEnum.h
@@ -18,7 +18,7 @@ class UDEPRECATED_K2Node_SelectEnum : public UK2Node_Select
 
 	// Begin UK2Node interface
 	virtual void NotifyPinConnectionListChanged(UEdGraphPin* Pin) override;
-	virtual void GetMenuActions(TArray<UBlueprintNodeSpawner*>& ActionListOut) const override;
+	virtual void GetMenuActions(FBlueprintActionDatabaseRegistrar& ActionRegistrar) const override;
 	// End UK2Node interface
 
 	// Begin UK2Node_Select interface
diff --git a/Engine/Source/Editor/BlueprintGraph/Classes/K2Node_Self.h b/Engine/Source/Editor/BlueprintGraph/Classes/K2Node_Self.h
--- a/Engine/Source/Editor/BlueprintGraph/Classes/K2Node_Self.h
+++ b/Engine/Source/Editor/BlueprintGraph/Classes/K2Node_Self.h
@@ -22,7 +22,7 @@ class UK2Node_Self : public UK2Node
 	virtual bool IsNodePure() const override { return true; }
 	virtual bool DrawNodeAsVariable() const override { return true; }
 	virtual class FNodeHandlingFunctor* CreateNodeHandler(class FKismetCompilerContext& CompilerContext) const override;
-	virtual void GetMenuActions(TArray<UBlueprintNodeSpawner*>& ActionListOut) const override;
+	virtual void GetMenuActions(FBlueprintActionDatabaseRegistrar& ActionRegistrar) const override;
 	virtual FText GetMenuCategory() const override;
 	// End K2Node interface
 };
diff --git a/Engine/Source/Editor/BlueprintGraph/Classes/K2Node_SpawnActorFromClass.h b/Engine/Source/Editor/BlueprintGraph/Classes/K2Node_SpawnActorFromClass.h
--- a/Engine/Source/Editor/BlueprintGraph/Classes/K2Node_SpawnActorFromClass.h
+++ b/Engine/Source/Editor/BlueprintGraph/Classes/K2Node_SpawnActorFromClass.h
@@ -26,7 +26,7 @@ class UK2Node_SpawnActorFromClass : public UK2Node
 	virtual bool IsNodeSafeToIgnore() const override { return true; }
 	virtual void ReallocatePinsDuringReconstruction(TArray<UEdGraphPin*>& OldPins) override;
 	virtual void GetNodeAttributes( TArray<TKeyValuePair<FString, FString>>& OutNodeAttributes ) const override;
-	virtual void GetMenuActions(TArray<UBlueprintNodeSpawner*>& ActionListOut) const override;
+	virtual void GetMenuActions(FBlueprintActionDatabaseRegistrar& ActionRegistrar) const override;
 	virtual FText GetMenuCategory() const override;
 	virtual class FNodeHandlingFunctor* CreateNodeHandler(class FKismetCompilerContext& CompilerContext) const override;
 	// End UK2Node interface
diff --git a/Engine/Source/Editor/BlueprintGraph/Classes/K2Node_SwitchEnum.h b/Engine/Source/Editor/BlueprintGraph/Classes/K2Node_SwitchEnum.h
--- a/Engine/Source/Editor/BlueprintGraph/Classes/K2Node_SwitchEnum.h
+++ b/Engine/Source/Editor/BlueprintGraph/Classes/K2Node_SwitchEnum.h
@@ -35,7 +35,7 @@ class UK2Node_SwitchEnum : public UK2Node_Switch, public INodeDependingOnEnumInt
 	// End of UEdGraphNode interface
 
 	// UK2Node interface
-	virtual void GetMenuActions(TArray<UBlueprintNodeSpawner*>& ActionListOut) const override;
+	virtual void GetMenuActions(FBlueprintActionDatabaseRegistrar& ActionRegistrar) const override;
 	// End of UK2Node interface
 
 	// UK2Node_Switch Interface
diff --git a/Engine/Source/Editor/BlueprintGraph/Classes/K2Node_SwitchInteger.h b/Engine/Source/Editor/BlueprintGraph/Classes/K2Node_SwitchInteger.h
--- a/Engine/Source/Editor/BlueprintGraph/Classes/K2Node_SwitchInteger.h
+++ b/Engine/Source/Editor/BlueprintGraph/Classes/K2Node_SwitchInteger.h
@@ -26,7 +26,7 @@ class UK2Node_SwitchInteger : public UK2Node_Switch
 	// End of UEdGraphNode interface
 
 	// UK2Node interface
-	virtual void GetMenuActions(TArray<UBlueprintNodeSpawner*>& ActionListOut) const override;
+	virtual void GetMenuActions(FBlueprintActionDatabaseRegistrar& ActionRegistrar) const override;
 	// End of UK2Node interface
 
 	// UK2Node_Switch Interface
diff --git a/Engine/Source/Editor/BlueprintGraph/Classes/K2Node_SwitchName.h b/Engine/Source/Editor/BlueprintGraph/Classes/K2Node_SwitchName.h
--- a/Engine/Source/Editor/BlueprintGraph/Classes/K2Node_SwitchName.h
+++ b/Engine/Source/Editor/BlueprintGraph/Classes/K2Node_SwitchName.h
@@ -24,7 +24,7 @@ class UK2Node_SwitchName : public UK2Node_Switch
 	// End of UEdGraphNode interface
 
 	// UK2Node interface
-	virtual void GetMenuActions(TArray<UBlueprintNodeSpawner*>& ActionListOut) const override;
+	virtual void GetMenuActions(FBlueprintActionDatabaseRegistrar& ActionRegistrar) const override;
 	// End of UK2Node interface
 
 	// UK2Node_Switch Interface
diff --git a/Engine/Source/Editor/BlueprintGraph/Classes/K2Node_SwitchString.h b/Engine/Source/Editor/BlueprintGraph/Classes/K2Node_SwitchString.h
--- a/Engine/Source/Editor/BlueprintGraph/Classes/K2Node_SwitchString.h
+++ b/Engine/Source/Editor/BlueprintGraph/Classes/K2Node_SwitchString.h
@@ -27,7 +27,7 @@ class UK2Node_SwitchString : public UK2Node_Switch
 	// End of UEdGraphNode interface
 
 	// UK2Node interface
-	virtual void GetMenuActions(TArray<UBlueprintNodeSpawner*>& ActionListOut) const override;
+	virtual void GetMenuActions(FBlueprintActionDatabaseRegistrar& ActionRegistrar) const override;
 	// End of UK2Node interface
 
 	// UK2Node_Switch Interface
diff --git a/Engine/Source/Editor/BlueprintGraph/Classes/K2Node_TemporaryVariable.h b/Engine/Source/Editor/BlueprintGraph/Classes/K2Node_TemporaryVariable.h
--- a/Engine/Source/Editor/BlueprintGraph/Classes/K2Node_TemporaryVariable.h
+++ b/Engine/Source/Editor/BlueprintGraph/Classes/K2Node_TemporaryVariable.h
@@ -30,7 +30,7 @@ class UK2Node_TemporaryVariable : public UK2Node
 	// Begin UK2Node interface.
 	virtual bool IsNodePure() const override;
 	virtual class FNodeHandlingFunctor* CreateNodeHandler(class FKismetCompilerContext& CompilerContext) const override;
-	virtual void GetMenuActions(TArray<UBlueprintNodeSpawner*>& ActionListOut) const override;
+	virtual void GetMenuActions(FBlueprintActionDatabaseRegistrar& ActionRegistrar) const override;
 	virtual FText GetMenuCategory() const override;
 	// End UK2Node interface.
 };
diff --git a/Engine/Source/Editor/BlueprintGraph/Classes/K2Node_Timeline.h b/Engine/Source/Editor/BlueprintGraph/Classes/K2Node_Timeline.h
--- a/Engine/Source/Editor/BlueprintGraph/Classes/K2Node_Timeline.h
+++ b/Engine/Source/Editor/BlueprintGraph/Classes/K2Node_Timeline.h
@@ -52,7 +52,7 @@ class UK2Node_Timeline : public UK2Node
 	virtual class FNodeHandlingFunctor* CreateNodeHandler(class FKismetCompilerContext& CompilerContext) const override;
 	virtual void ExpandNode(class FKismetCompilerContext& CompilerContext, UEdGraph* SourceGraph) override;
 	virtual void GetNodeAttributes( TArray<TKeyValuePair<FString, FString>>& OutNodeAttributes ) const override;
-	virtual void GetMenuActions(TArray<UBlueprintNodeSpawner*>& ActionListOut) const override;
+	virtual void GetMenuActions(FBlueprintActionDatabaseRegistrar& ActionRegistrar) const override;
 	// End UK2Node interface.
 
 	/** Get the 'play' input pin */
diff --git a/Engine/Source/Editor/BlueprintGraph/Classes/K2Node_VariableSetRef.h b/Engine/Source/Editor/BlueprintGraph/Classes/K2Node_VariableSetRef.h
--- a/Engine/Source/Editor/BlueprintGraph/Classes/K2Node_VariableSetRef.h
+++ b/Engine/Source/Editor/BlueprintGraph/Classes/K2Node_VariableSetRef.h
@@ -19,7 +19,7 @@ class UK2Node_VariableSetRef : public UK2Node
 	// Begin UK2Node interface
 	virtual void NotifyPinConnectionListChanged(UEdGraphPin* Pin) override;
 	virtual class FNodeHandlingFunctor* CreateNodeHandler(class FKismetCompilerContext& CompilerContext) const override;
-	virtual void GetMenuActions(TArray<UBlueprintNodeSpawner*>& ActionListOut) const override;
+	virtual void GetMenuActions(FBlueprintActionDatabaseRegistrar& ActionRegistrar) const override;
 	virtual FText GetMenuCategory() const override;
 	virtual int32 GetNodeRefreshPriority() const override { return EBaseNodeRefreshPriority::Low_UsesDependentWildcard; }
 	// End UK2Node interface
diff --git a/Engine/Source/Editor/BlueprintGraph/Private/BlueprintActionDatabase.cpp b/Engine/Source/Editor/BlueprintGraph/Private/BlueprintActionDatabase.cpp
--- a/Engine/Source/Editor/BlueprintGraph/Private/BlueprintActionDatabase.cpp
+++ b/Engine/Source/Editor/BlueprintGraph/Private/BlueprintActionDatabase.cpp
@@ -17,19 +17,10 @@
 #include "ModuleManager.h"
 #include "AssetRegistryModule.h"	// for OnAssetAdded()/OnAssetRemoved()
 #include "BlueprintEditorUtils.h"	// for FindBlueprintForGraph()
+#include "BlueprintActionDatabaseRegistrar.h"
 
 // used below in FBlueprintNodeSpawnerFactory::MakeMacroNodeSpawner()
 #include "K2Node_MacroInstance.h"
-// used below in BlueprintActionDatabaseImpl::AddClassEnumActions()
-#include "K2Node_GetNumEnumEntries.h"
-#include "K2Node_ForEachElementInEnum.h"
-#include "K2Node_EnumLiteral.h"
-#include "K2Node_CastByteToEnum.h"
-#include "K2Node_SwitchEnum.h"
-// used below in BlueprintActionDatabaseImpl::AddClassStructActions()
-#include "K2Node_BreakStruct.h"
-#include "K2Node_MakeStruct.h"
-#include "K2Node_SetFieldsInStruct.h"
 // used below in FBlueprintNodeSpawnerFactory::MakeMessageNodeSpawner()
 #include "K2Node_Message.h"
 // used below in BlueprintActionDatabaseImpl::AddClassPropertyActions()
@@ -43,7 +34,7 @@
 // used below in BlueprintActionDatabaseImpl::AddClassCastActions()
 #include "K2Node_DynamicCast.h"
 #include "K2Node_ClassDynamicCast.h"
-// used below in BlueprintActionDatabaseImpl::AddAutonomousNodeActions()
+// used below in BlueprintActionDatabaseImpl::GetNodeSpectificActions()
 #include "EdGraph/EdGraphNode_Comment.h"
 #include "EdGraph/EdGraphNode_Documentation.h"
 
@@ -65,29 +56,6 @@ namespace FBlueprintNodeSpawnerFactory
 	 */
 	static UBlueprintNodeSpawner* MakeMacroNodeSpawner(UEdGraph* MacroGraph);
 
-	/**
-	 * A templatized method which constructs a node-spawner for various enum 
-	 * node types (any node with a public UEnum* Enum field). Takes the 
-	 * specified enum and applies it to the node post-spawn.
-	 *
-	 * @param  Enum					The enum you want set for the spawned node.
-	 * @param  ExtraSetupCallback	A callback for any further post-spawn customization (other than setting the node's Enum field).
-	 * @return A new node-spawner, setup to spawn some enum node (defined by the EnumNodeType template param).
-	 */
-	template<class EnumNodeType>
-	static UBlueprintNodeSpawner* MakeEnumNodeSpawner(UEnum* Enum, void(*ExtraSetupCallback)(EnumNodeType*) = nullptr);
-
-	/**
-	 * A templatized method which constructs a node-spawner for various struct 
-	 * node types (any node with a public UScriptStruct* StructType field). 
-	 * Takes the specified struct and applies it to the node post-spawn.
-	 *
-	 * @param  Struct	The struct you want set for the spawned node.
-	 * @return A new node-spawner, setup to spawn some struct node (defined by the StructNodeType template param).
-	 */
-	template<class StructNodeType>
-	static UBlueprintNodeSpawner* MakeStructNodeSpawner(UScriptStruct* Struct);
-
 	/**
 	 * Constructs a UK2Node_Message spawner. Sets up the spawner to set 
 	 * spawned nodes with the supplied function.
@@ -151,64 +119,6 @@ static UBlueprintNodeSpawner* FBlueprintNodeSpawnerFactory::MakeMacroNodeSpawner
 	return NodeSpawner;
 }
 
-//------------------------------------------------------------------------------
-template<class EnumNodeType>
-static UBlueprintNodeSpawner* FBlueprintNodeSpawnerFactory::MakeEnumNodeSpawner(UEnum* Enum, void(*ExtraSetupCallback)(EnumNodeType*))
-{
-	UBlueprintNodeSpawner* NodeSpawner = UBlueprintNodeSpawner::Create(EnumNodeType::StaticClass());
-	check(NodeSpawner != nullptr);
-
-	DECLARE_DELEGATE_OneParam(FFurtherCustomizeNodeDelegate, EnumNodeType*);
-
-	auto CustomizeEnumNodeLambda = [](UEdGraphNode* NewNode, bool bIsTemplateNode, TWeakObjectPtr<UEnum> EnumPtr, FFurtherCustomizeNodeDelegate FurtherCustomizeDelegate)
-	{
-		EnumNodeType* EnumNode = CastChecked<EnumNodeType>(NewNode);
-
-		if (EnumPtr.IsValid())
-		{
-			EnumNode->Enum = EnumPtr.Get();
-		}
-
-		if (FurtherCustomizeDelegate.IsBound())
-		{
-			FurtherCustomizeDelegate.Execute(EnumNode);
-		}
-	};
-
-	FFurtherCustomizeNodeDelegate FurtherCustomizeDelegate;
-	if (ExtraSetupCallback != nullptr)
-	{
-		FurtherCustomizeDelegate = FFurtherCustomizeNodeDelegate::CreateStatic(ExtraSetupCallback);
-	}
-
-	TWeakObjectPtr<UEnum> EnumPtr = Enum;
-	NodeSpawner->CustomizeNodeDelegate = UBlueprintNodeSpawner::FCustomizeNodeDelegate::CreateStatic(CustomizeEnumNodeLambda, EnumPtr, FurtherCustomizeDelegate);
-
-	return NodeSpawner;
-}
-
-//------------------------------------------------------------------------------
-template<class StructNodeType>
-static UBlueprintNodeSpawner* FBlueprintNodeSpawnerFactory::MakeStructNodeSpawner(UScriptStruct* Struct)
-{
-	UBlueprintNodeSpawner* NodeSpawner = UBlueprintNodeSpawner::Create(StructNodeType::StaticClass());
-	check(NodeSpawner != nullptr);
-
-	auto CustomizeStructNodeLambda = [](UEdGraphNode* NewNode, bool bIsTemplateNode, TWeakObjectPtr<UScriptStruct> StructPtr)
-	{
-		StructNodeType* StructNode = CastChecked<StructNodeType>(NewNode);
-		if (StructPtr.IsValid())
-		{
-			StructNode->StructType = StructPtr.Get();
-		}
-	};
-
-	TWeakObjectPtr<UScriptStruct> StructPtr = Struct;
-	NodeSpawner->CustomizeNodeDelegate = UBlueprintNodeSpawner::FCustomizeNodeDelegate::CreateStatic(CustomizeStructNodeLambda, StructPtr);
-
-	return NodeSpawner;
-}
-
 //------------------------------------------------------------------------------
 static UBlueprintNodeSpawner* FBlueprintNodeSpawnerFactory::MakeMessageNodeSpawner(UFunction* InterfaceFunction)
 {
@@ -304,6 +214,14 @@ namespace BlueprintActionDatabaseImpl
 	 */
 	static bool IsPropertyBlueprintVisible(UProperty const* const Property);
 
+	/**
+	 * 
+	 * 
+	 * @param  Class	
+	 * @param  ActionListOut	
+	 */
+	static void GetClassMemberActions(UClass* const Class, FActionList& ActionListOut);
+
 	/**
 	 * Loops over all of the class's functions and creates a node-spawners for
 	 * any that are viable for blueprint use. Evolved from 
@@ -345,67 +263,31 @@ namespace BlueprintActionDatabaseImpl
 	 * specified class is a component type (and can be spawned), then a 
 	 * UK2Node_AddComponent spawner is created and added to ActionListOut.
 	 *
-	 * @param  Class			The class who you want a spawner for (the component class).
+	 * @param  ComponentClass	The class who you want a spawner for (the component class).
 	 * @param  ActionListOut	The list you want populated with the new spawner.
 	 */
-	static void AddComponentClassActions(UClass* const Class, FActionList& ActionListOut);
-
-	/**
-	 * Loops over the class's enums and creates node-spawners for any that are 
-	 * viable for blueprint use. Evolved from K2ActionMenuBuilder's 
-	 * GetEnumUtilitiesNodes() (as well as snippets from GetSwitchMenuItems).
-	 *
-	 * NOTE: This only accounts for enums that belong to this specific class, 
-	 *       not autonomous globally-scoped enums (ones with a UPackage outer). 
-	 *       Those enums should instead be accounted for in the appropriate 
-	 *       node's GetMenuActions(). For example, UK2Node_EnumLiteral should
-	 *       add UK2Node_EnumLiteral spawners for any global enums.
-	 *
-	 * @param  Class			The class whose enums you want node-spawners for.
-	 * @param  ActionListOut	The list you want populated with new spawners.
-	 */
-	static void AddClassEnumActions(UClass const* const Class, FActionList& ActionListOut);
-
-	/**
-	 * Loops over the class's structs and creates node-spawners for any that are 
-	 * viable for blueprint use. Evolved from 
-	 * FK2ActionMenuBuilder::GetStructActions().
-	 *
-	 * NOTE: This only captures structs that belong to this specific class, 
-	 *       not autonomous globally-scoped ones (those with a UPackage outer). 
-	 *       Those struct should instead be accounted for in the appropriate 
-	 *       node's GetMenuActions(). For example, UK2Node_MakeStruct should
-	 *       add UK2Node_MakeStruct spawners for any global structs.
-	 *
-	 * @param  Class			The class whose struct you want node-spawners for.
-	 * @param  ActionListOut	The list you want populated with new spawners.
-	 */
-	static void AddClassStructActions(UClass const* const Class, FActionList& ActionListOut);
+	static void AddComponentClassActions(TSubclassOf<UActorComponent> const ComponentClass, FActionList& ActionListOut);
 
 	/**
 	 * If the associated class is a blueprint generated class, then this will
 	 * loop over the blueprint's graphs and create any node-spawners associated
 	 * with those graphs (like UK2Node_MacroInstance spawners for macro graphs).
 	 *
-	 * @param  Class			The class which you want graph associated node-spawners for.
+	 * @param  Blueprint		The blueprint which you want graph associated node-spawners for.
 	 * @param  ActionListOut	The list you want populated with new spawners.
 	 */
-	static void AddBlueprintGraphActions(UClass const* const Class, FActionList& ActionListOut);
+	static void AddBlueprintGraphActions(UBlueprint const* const Blueprint, FActionList& ActionListOut);
 
 	/**
 	 * Emulates UEdGraphSchema::GetGraphContextActions(). If the supplied class  
 	 * is a node type, then it will query the node's CDO for any actions it 
 	 * wishes to add. This helps us keep the code in this file paired down, and 
-	 * makes it easily extensible for new node types. At the high level, this is 
-	 * for node types that aren't associated with an other class other, 
-	 * "autonomous" nodes that don't have any other way of being listed (this is
-	 * where enum/struct nodes missed by AddClassEnumActions()/AddClassStructActions()
-	 * would be added).
+	 * makes it easily extensible for new node types.
 	 * 
-	 * @param  Class			The class which you want node-spawners for.
+	 * @param  NodeClass		The class which you want node-spawners for.
 	 * @param  ActionListOut	The list you want populated with new spawners.
 	 */
-	static void AddAutonomousNodeActions(UClass* const Class, FActionList& ActionListOut);
+	static void GetNodeSpecificActions(TSubclassOf<UEdGraphNode const> const NodeClass, FBlueprintActionDatabaseRegistrar& Registrar);
 
 	/**
 	 * Callback to refresh the database when a blueprint has been altered 
@@ -450,11 +332,14 @@ namespace BlueprintActionDatabaseImpl
 	static void OnAssetRemoved(FAssetData const& AssetInfo);
 
 	/** 
-	 * Classes that we cleared from the database (to remove references, and make 
+	 * Assets that we cleared from the database (to remove references, and make 
 	 * way for a delete), but in-case the class wasn't deleted we need them 
 	 * tracked here so we can add them back in.
 	 */
-	TSet<UClass*> PendingDelete;
+	TSet<UObject*> PendingDelete;
+
+	/** */
+	bool bIsInitializing = false;
 }
 
 //------------------------------------------------------------------------------
@@ -468,6 +353,26 @@ static bool BlueprintActionDatabaseImpl::IsPropertyBlueprintVisible(UProperty co
 	return !Property->HasAnyPropertyFlags(CPF_Parm) && (bIsAccessible || (bIsDelegate && bIsAssignableOrCallable));
 }
 
+//------------------------------------------------------------------------------
+static void BlueprintActionDatabaseImpl::GetClassMemberActions(UClass* const Class, FActionList& ActionListOut)
+{
+	// class field actions (nodes that represent and perform actions on
+	// specific fields of the class... functions, properties, etc.)
+	{
+		AddClassFunctionActions(Class, ActionListOut);
+		AddClassPropertyActions(Class, ActionListOut);
+		// class UEnum actions are added by individual nodes via GetNodeSpecificActions()
+		// class UScriptStruct actions are added by individual nodes via GetNodeSpecificActions()
+	}
+
+	AddClassCastActions(Class, ActionListOut);
+
+	if (Class->IsChildOf<UActorComponent>())
+	{
+		AddComponentClassActions(Class, ActionListOut);
+	}
+}
+
 //------------------------------------------------------------------------------
 static void BlueprintActionDatabaseImpl::AddClassFunctionActions(UClass const* const Class, FActionList& ActionListOut)
 {
@@ -605,13 +510,11 @@ static void BlueprintActionDatabaseImpl::AddClassCastActions(UClass* const Class
 }
 
 //------------------------------------------------------------------------------
-static void BlueprintActionDatabaseImpl::AddComponentClassActions(UClass* const Class, FActionList& ActionListOut)
+static void BlueprintActionDatabaseImpl::AddComponentClassActions(TSubclassOf<UActorComponent> const ComponentClass, FActionList& ActionListOut)
 {
-	bool bIsSpawnable = !Class->HasAnyClassFlags(CLASS_Abstract) && Class->HasMetaData(FBlueprintMetadata::MD_BlueprintSpawnableComponent);
-
-	if (bIsSpawnable && Class->IsChildOf(UActorComponent::StaticClass()))
+	if ((ComponentClass != nullptr) && !ComponentClass->HasAnyClassFlags(CLASS_Abstract) && ComponentClass->HasMetaData(FBlueprintMetadata::MD_BlueprintSpawnableComponent))
 	{
-		if (UBlueprintComponentNodeSpawner* NodeSpawner = UBlueprintComponentNodeSpawner::Create(Class))
+		if (UBlueprintComponentNodeSpawner* NodeSpawner = UBlueprintComponentNodeSpawner::Create(ComponentClass))
 		{
 			ActionListOut.Add(NodeSpawner);
 		}
@@ -619,116 +522,46 @@ static void BlueprintActionDatabaseImpl::AddComponentClassActions(UClass* const
 }
 
 //------------------------------------------------------------------------------
-static void BlueprintActionDatabaseImpl::AddClassEnumActions(UClass const* const Class, FActionList& ActionListOut)
+static void BlueprintActionDatabaseImpl::AddBlueprintGraphActions(UBlueprint const* const Blueprint, FActionList& ActionListOut)
 {
-	using namespace FBlueprintNodeSpawnerFactory; // for MakeEnumNodeSpawner()
+	using namespace FBlueprintNodeSpawnerFactory; // for MakeMacroNodeSpawner()
 
-	for (TFieldIterator<UEnum> EnumIt(Class, EFieldIteratorFlags::ExcludeSuper); EnumIt; ++EnumIt)
+	for (auto GraphIt = Blueprint->MacroGraphs.CreateConstIterator(); GraphIt; ++GraphIt)
 	{
-		UEnum* Enum = *EnumIt;
-		if (!UEdGraphSchema_K2::IsAllowableBlueprintVariableType(Enum))
-		{
-			continue;
-		}
-
-		ActionListOut.Add(MakeEnumNodeSpawner<UK2Node_GetNumEnumEntries>(Enum));
-		ActionListOut.Add(MakeEnumNodeSpawner<UK2Node_ForEachElementInEnum>(Enum));
-		ActionListOut.Add(MakeEnumNodeSpawner<UK2Node_EnumLiteral>(Enum));
-
-		auto SetupEnumByteCastLambda = [](UK2Node_CastByteToEnum* NewNode)
-		{
-			NewNode->bSafe = true;
-		};
-		ActionListOut.Add(MakeEnumNodeSpawner<UK2Node_CastByteToEnum>(Enum, SetupEnumByteCastLambda));
-
-
-		auto SetupEnumSwitchLambda = [](UK2Node_SwitchEnum* NewNode)
-		{
-			// the node's 'Enum' field should have already been setup generically by MakeEnumNodeSpawner()
-			check(NewNode->Enum != nullptr);
-			NewNode->SetEnum(NewNode->Enum); // SetEnum() does more setup work than just setting the field
-		};
-		ActionListOut.Add(MakeEnumNodeSpawner<UK2Node_SwitchEnum>(Enum, SetupEnumSwitchLambda));
+		ActionListOut.Add(MakeMacroNodeSpawner(*GraphIt));
 	}
 
-	// @TODO: what about enum assets? how do we detect newly added ones and refresh those?
-	//    FEditorDelegates::LoadSelectedAssetsIfNeeded.Broadcast();
-}
-
-//------------------------------------------------------------------------------
-static void BlueprintActionDatabaseImpl::AddClassStructActions(UClass const* const Class, FActionList& ActionListOut)
-{
-	using namespace FBlueprintNodeSpawnerFactory; // for MakeStructNodeSpawner()
-
-	for (TFieldIterator<UScriptStruct> StructIt(Class, EFieldIteratorFlags::ExcludeSuper); StructIt; ++StructIt)
+	// local variables
+	for (auto GraphIt = Blueprint->FunctionGraphs.CreateConstIterator(); GraphIt; ++GraphIt)
 	{
-		UScriptStruct* Struct = (*StructIt);
-		if (!UEdGraphSchema_K2::IsAllowableBlueprintVariableType(Struct))
-		{
-			continue;
-		}
-
-		if (UK2Node_BreakStruct::CanBeBroken(Struct))
-		{
-			ActionListOut.Add(MakeStructNodeSpawner<UK2Node_BreakStruct>(Struct));
-		}
-
-		if (UK2Node_MakeStruct::CanBeMade(Struct))
-		{
-			ActionListOut.Add(MakeStructNodeSpawner<UK2Node_MakeStruct>(Struct));
-			ActionListOut.Add(MakeStructNodeSpawner<UK2Node_SetFieldsInStruct>(Struct));
-		}
-	}
-
-	// @TODO: what about struct assets? how do we detect newly added ones and refresh those?
-	//    FEditorDelegates::LoadSelectedAssetsIfNeeded.Broadcast();
-}
+		UEdGraph* FunctionGraph = (*GraphIt);
 
-//------------------------------------------------------------------------------
-static void BlueprintActionDatabaseImpl::AddBlueprintGraphActions(UClass const* const Class, FActionList& ActionListOut)
-{
-	using namespace FBlueprintNodeSpawnerFactory; // for MakeMacroNodeSpawner()
+		TArray<UK2Node_FunctionEntry*> GraphEntryNodes;
+		FunctionGraph->GetNodesOfClass<UK2Node_FunctionEntry>(GraphEntryNodes);
 
-	if (UBlueprint* const Blueprint = Cast<UBlueprint>(Class->ClassGeneratedBy))
-	{
-		for (auto GraphIt = Blueprint->MacroGraphs.CreateConstIterator(); GraphIt; ++GraphIt)
+		for (UK2Node_FunctionEntry* FunctionEntry : GraphEntryNodes)
 		{
-			ActionListOut.Add(MakeMacroNodeSpawner(*GraphIt));
-		}
-
-		// local variables
-		for (auto GraphIt = Blueprint->FunctionGraphs.CreateConstIterator(); GraphIt; ++GraphIt)
-		{
-			UEdGraph* FunctionGraph = (*GraphIt);
-
-			TArray<UK2Node_FunctionEntry*> GraphEntryNodes;
-			FunctionGraph->GetNodesOfClass<UK2Node_FunctionEntry>(GraphEntryNodes);
-
-			for (UK2Node_FunctionEntry* FunctionEntry : GraphEntryNodes)
+			for (FBPVariableDescription const& LocalVar : FunctionEntry->LocalVariables)
 			{
-				for (FBPVariableDescription const& LocalVar : FunctionEntry->LocalVariables)
-				{
-					UBlueprintNodeSpawner* GetVarSpawner = UBlueprintVariableNodeSpawner::Create(UK2Node_VariableGet::StaticClass(), FunctionGraph, LocalVar);
-					ActionListOut.Add(GetVarSpawner);
-					UBlueprintNodeSpawner* SetVarSpawner = UBlueprintVariableNodeSpawner::Create(UK2Node_VariableSet::StaticClass(), FunctionGraph, LocalVar);
-					ActionListOut.Add(SetVarSpawner);
-				}
+				UBlueprintNodeSpawner* GetVarSpawner = UBlueprintVariableNodeSpawner::Create(UK2Node_VariableGet::StaticClass(), FunctionGraph, LocalVar);
+				ActionListOut.Add(GetVarSpawner);
+				UBlueprintNodeSpawner* SetVarSpawner = UBlueprintVariableNodeSpawner::Create(UK2Node_VariableSet::StaticClass(), FunctionGraph, LocalVar);
+				ActionListOut.Add(SetVarSpawner);
 			}
 		}
 	}	
 }
 
 //------------------------------------------------------------------------------
-static void BlueprintActionDatabaseImpl::AddAutonomousNodeActions(UClass* const Class, FActionList& ActionListOut)
+static void BlueprintActionDatabaseImpl::GetNodeSpecificActions(TSubclassOf<UEdGraphNode const> const NodeClass, FBlueprintActionDatabaseRegistrar& Registrar)
 {
 	using namespace FBlueprintNodeSpawnerFactory; // for MakeDocumentationNodeSpawner()
 
-	if (Class->IsChildOf(UK2Node::StaticClass()) && !Class->HasAnyClassFlags(CLASS_Abstract))
+	if (NodeClass->IsChildOf<UK2Node>() && !NodeClass->HasAnyClassFlags(CLASS_Abstract))
 	{
-		UK2Node const* const NodeCDO = Class->GetDefaultObject<UK2Node>();
+		UK2Node const* const NodeCDO = NodeClass->GetDefaultObject<UK2Node>();
 		check(NodeCDO != nullptr);
-
-		NodeCDO->GetMenuActions(ActionListOut);
+		NodeCDO->GetMenuActions(Registrar);
 	}
 	// unfortunately, UEdGraphNode_Comment is not a UK2Node and therefore cannot
 	// leverage UK2Node's GetMenuActions() function, so here we HACK it in
@@ -736,26 +569,24 @@ static void BlueprintActionDatabaseImpl::AddAutonomousNodeActions(UClass* const
 	// @TODO: DO NOT follow this example! Do as I say, not as I do! If we need
 	//        to support other nodes in a similar way, then we should come up
 	//        with a better (more generalized) solution.
-	else if (Class == UEdGraphNode_Comment::StaticClass())
+	else if (NodeClass == UEdGraphNode_Comment::StaticClass())
 	{
-		ActionListOut.Add(MakeDocumentationNodeSpawner<UEdGraphNode_Comment>());
+		Registrar.AddBlueprintAction(MakeDocumentationNodeSpawner<UEdGraphNode_Comment>());
 	}
-	else if (Class == UEdGraphNode_Documentation::StaticClass())
+	else if (NodeClass == UEdGraphNode_Documentation::StaticClass())
 	{
 		// @TODO: BOO! (see comment above)
-		ActionListOut.Add(MakeDocumentationNodeSpawner<UEdGraphNode_Documentation>());
+		Registrar.AddBlueprintAction(MakeDocumentationNodeSpawner<UEdGraphNode_Documentation>());
 	}
 }
 
 //------------------------------------------------------------------------------
 static void BlueprintActionDatabaseImpl::OnBlueprintChanged(UBlueprint* Blueprint)
 {
-	if( Blueprint->SkeletonGeneratedClass &&
-		Blueprint->GeneratedClass )
+	if (Blueprint->IsAsset())
 	{
 		FBlueprintActionDatabase& ActionDatabase = FBlueprintActionDatabase::Get();
-		ActionDatabase.RefreshClassActions(Blueprint->SkeletonGeneratedClass);
-		ActionDatabase.RefreshClassActions(Blueprint->GeneratedClass);
+		ActionDatabase.RefreshAssetActions(Blueprint);
 	}
 }
 
@@ -767,6 +598,11 @@ static void BlueprintActionDatabaseImpl::OnAssetLoaded(UObject* NewObject)
 		OnBlueprintChanged(NewBlueprint);
 		NewBlueprint->OnChanged().AddStatic(&BlueprintActionDatabaseImpl::OnBlueprintChanged);
 	}
+	else
+	{
+		FBlueprintActionDatabase& ActionDatabase = FBlueprintActionDatabase::Get();
+		ActionDatabase.RefreshAssetActions(NewObject);
+	}
 }
 
 //------------------------------------------------------------------------------
@@ -774,12 +610,17 @@ static void BlueprintActionDatabaseImpl::OnAssetAdded(FAssetData const& NewAsset
 {
 	if (NewAssetInfo.IsAssetLoaded())
 	{
-		if (NewAssetInfo.GetClass()->IsChildOf<UBlueprint>())
+		UObject* AssetObject = NewAssetInfo.GetAsset();
+		if (UBlueprint* NewBlueprint = Cast<UBlueprint>(AssetObject))
 		{
-			UBlueprint* NewBlueprint = CastChecked<UBlueprint>(NewAssetInfo.GetAsset());
 			OnBlueprintChanged(NewBlueprint);
 			NewBlueprint->OnChanged().AddStatic(&BlueprintActionDatabaseImpl::OnBlueprintChanged);
 		}
+		else 
+		{
+			FBlueprintActionDatabase& ActionDatabase = FBlueprintActionDatabase::Get();
+			ActionDatabase.RefreshAssetActions(AssetObject);
+		}
 	}
 }
 
@@ -789,18 +630,13 @@ static void BlueprintActionDatabaseImpl::OnAssetsPendingDelete(TArray<UObject*>
 	FBlueprintActionDatabase& ActionDatabase = FBlueprintActionDatabase::Get();
 	for (UObject* DeletingObject : ObjectsForDelete)
 	{
-		if (UBlueprint* NewBlueprint = Cast<UBlueprint>(DeletingObject))
-		{
-			// have to temporarily remove references (so that this delete isn't 
-			// blocked by dangling references)
-			ActionDatabase.ClearClassActions(NewBlueprint->GeneratedClass);
-			ActionDatabase.ClearClassActions(NewBlueprint->SkeletonGeneratedClass);
-
-			// in case they choose not to delete the object, we need to add 
-			// these back in to the database, soi we track them here
-			PendingDelete.Add(NewBlueprint->GeneratedClass);
-			PendingDelete.Add(NewBlueprint->SkeletonGeneratedClass);
-		}
+		// have to temporarily remove references (so that this delete isn't 
+		// blocked by dangling references)
+		ActionDatabase.ClearAssetActions(DeletingObject);
+
+		// in case they choose not to delete the object, we need to add 
+		// these back in to the database, so we track them here
+		PendingDelete.Add(DeletingObject);
 	}
 }
 
@@ -811,16 +647,13 @@ static void BlueprintActionDatabaseImpl::OnAssetRemoved(FAssetData const& AssetI
 
 	if (AssetInfo.IsAssetLoaded())
 	{
-		UClass* AssetClass = AssetInfo.GetClass();
-		if (AssetClass->IsChildOf<UBlueprint>())
+		UObject* AssetObject = AssetInfo.GetAsset();
+		if (UBlueprint* BlueprintAsset = Cast<UBlueprint>(AssetObject))
 		{
-			UBlueprint* BlueprintAsset = CastChecked<UBlueprint>(AssetInfo.GetAsset());
 			BlueprintAsset->OnChanged().RemoveStatic(&BlueprintActionDatabaseImpl::OnBlueprintChanged);
-
-			// the delete went through, so we don't need to track these for re-add
-			PendingDelete.Remove(BlueprintAsset->GeneratedClass);
-			PendingDelete.Remove(BlueprintAsset->SkeletonGeneratedClass);
 		}
+		// the delete went through, so we don't need to track these for re-add
+		PendingDelete.Remove(AssetObject);
 	}
 }
 
@@ -860,7 +693,7 @@ FBlueprintActionDatabase::FBlueprintActionDatabase()
 //------------------------------------------------------------------------------
 void FBlueprintActionDatabase::AddReferencedObjects(FReferenceCollector& Collector)
 {
-	for (auto& ActionListIt : ClassActions)
+	for (auto& ActionListIt : ActionRegistry)
 	{
 		Collector.AddReferencedObjects(ActionListIt.Value);
 	}
@@ -873,9 +706,12 @@ void FBlueprintActionDatabase::Tick(float DeltaTime)
 
 	// entries that were removed from the database, in preparation for a delete
 	// (but the user ended up not deleting the object)
-	for (UClass* Class : BlueprintActionDatabaseImpl::PendingDelete)
+	for (UObject const* AssetObj : BlueprintActionDatabaseImpl::PendingDelete)
 	{
-		RefreshClassActions(Class);
+		if (IsValid(AssetObj))
+		{
+			RefreshAssetActions(AssetObj);
+		}
 	}
 	BlueprintActionDatabaseImpl::PendingDelete.Empty();
 
@@ -884,11 +720,11 @@ void FBlueprintActionDatabase::Tick(float DeltaTime)
 	{
 		auto ActionIndex = ActionPrimingQueue.CreateIterator();	
 			
-		TWeakObjectPtr<UClass> ActionClass = ActionIndex.Key();
-		if (ActionClass.IsValid())
+		TWeakObjectPtr<UObject> ActionsKey = ActionIndex.Key();
+		if (ActionsKey.IsValid())
 		{
 			// make sure this class is still listed in the database
-			if (FActionList* ClassActionList = ClassActions.Find(ActionClass.Get()))
+			if (FActionList* ClassActionList = ActionRegistry.Find(ActionsKey.Get()))
 			{
 				int32& ActionListIndex = ActionIndex.Value();
 				for (; (ActionListIndex < ClassActionList->Num()) && (TotalFrameDuration < DurationThreshold); ++ActionListIndex)
@@ -910,17 +746,17 @@ void FBlueprintActionDatabase::Tick(float DeltaTime)
 
 				if (ActionListIndex >= ClassActionList->Num())
 				{
-					ActionPrimingQueue.Remove(ActionClass);
+					ActionPrimingQueue.Remove(ActionsKey);
 				}
 			}
 			else
 			{
-				ActionPrimingQueue.Remove(ActionClass);
+				ActionPrimingQueue.Remove(ActionsKey);
 			}
 		}
 		else
 		{
-			ActionPrimingQueue.Remove(ActionClass);
+			ActionPrimingQueue.Remove(ActionsKey);
 		}
 	}
 }
@@ -934,7 +770,9 @@ TStatId FBlueprintActionDatabase::GetStatId() const
 //------------------------------------------------------------------------------
 void FBlueprintActionDatabase::RefreshAll()
 {
-	ClassActions.Empty();
+	TGuardValue<bool> ScopedInitialization(BlueprintActionDatabaseImpl::bIsInitializing, true);
+
+	ActionRegistry.Empty();
 	for (TObjectIterator<UClass> ClassIt; ClassIt; ++ClassIt)
 	{
 		UClass* const Class = (*ClassIt);
@@ -946,45 +784,80 @@ void FBlueprintActionDatabase::RefreshAll()
 void FBlueprintActionDatabase::RefreshClassActions(UClass* const Class)
 {
 	using namespace BlueprintActionDatabaseImpl;
-
 	check(Class != nullptr);
-	bool const bIsBlueprintClass = (Cast<UBlueprintGeneratedClass>(Class) != nullptr);
-	bool const bIsSkelClass      = bIsBlueprintClass && FKismetEditorUtilities::IsClassABlueprintSkeleton(Class);
-	bool const bIsReinstClass    = Class->HasAnyClassFlags(CLASS_NewerVersionExists);
-	bool const bIsIgnorableClass = bIsReinstClass || (bIsBlueprintClass && !bIsSkelClass);
 
-	if (!bIsIgnorableClass)
+	bool const bOutOfDateClass   = Class->HasAnyClassFlags(CLASS_NewerVersionExists);
+	bool const bIsBlueprintClass = (Cast<UBlueprintGeneratedClass>(Class) != nullptr);	
+
+	if (bOutOfDateClass)
 	{
-		FActionList& ClassActionList = ClassActions.FindOrAdd(Class);
-		ClassActionList.Empty();
-		
-		// class field actions (nodes that represent and perform actions on
-		// specific fields of the class... functions, properties, etc.)
+		ActionRegistry.Remove(Class);
+		return;
+	}
+	else if (bIsBlueprintClass)
+	{
+		UBlueprint const* Blueprint = Cast<UBlueprint>(Class->ClassGeneratedBy);
+		if ((Blueprint != nullptr) && Blueprint->IsAsset())
 		{
-			AddClassFunctionActions(Class, ClassActionList);
-			AddClassPropertyActions(Class, ClassActionList);
-			AddClassEnumActions(Class, ClassActionList);
-			AddClassStructActions(Class, ClassActionList);
+			// to prevent us from hitting this twice on init (once for the skel 
+			// class, again for the generated class)
+			bool const bRefresh = !bIsInitializing || (Blueprint->SkeletonGeneratedClass == nullptr) ||
+				(Blueprint->SkeletonGeneratedClass == Class);
+
+			if (bRefresh)
+			{
+				RefreshAssetActions(Blueprint);
+			}
 		}
+	}
+	// here we account for "autonomous" standalone nodes, and any nodes that
+	// exist in a separate module; each UK2Node has a chance to append its
+	// own actions (presumably ones that would spawn that node)...
+	else if (Class->IsChildOf<UEdGraphNode>())
+	{
+		FActionList& ClassActionList = ActionRegistry.FindOrAdd(Class);
+		ClassActionList.Empty();
 
-		AddClassCastActions(Class, ClassActionList);
-		AddBlueprintGraphActions(Class, ClassActionList);
-		AddComponentClassActions(Class, ClassActionList);
-		
-		// accounts for the "autonomous" standalone nodes that can't be strongly
-		// associated with a particular class (besides the node's class)...
-		// think things like: comment nodes, custom events, the self node, etc.
-		//
-		// also should catch any actions dealing with global UFields (like
+		FBlueprintActionDatabaseRegistrar Registrar(ActionRegistry, ActionPrimingQueue, Class);
+		if (!bIsInitializing)
+		{
+			// if this a call to RefreshClassActions() from somewhere other than 
+			// RefreshAll(), then we should only add actions for this class (the
+			// node could be adding actions, probably duplicate ones for assets)
+			Registrar.ActionKeyFilter = Class;
+		}
+	
+		// also, should catch any actions dealing with global UFields (like
 		// global structs, enums, etc.; elements that wouldn't be caught
-		// normally when sifting through fields on all known classes)
-		AddAutonomousNodeActions(Class, ClassActionList);
+		// normally when sifting through fields on all known classes)		
+		GetNodeSpecificActions(Class, Registrar);
+		// don't worry, the registrar marks new actions for priming
+	}
+	else
+	{
+		FActionList& ClassActionList = ActionRegistry.FindOrAdd(Class);
+		if (!bIsInitializing)
+		{
+			ClassActionList.Empty();
+			// if we're only refreshing this class (and not init'ing the whole 
+			// database), then we have to reach out to individual nodes in case 
+			// they'd add entries for this as well
+			FBlueprintActionDatabaseRegistrar Registrar(ActionRegistry, ActionPrimingQueue);
+			Registrar.ActionKeyFilter = Class; // only want actions for this class
+
+			RegisterAllNodeActions(Registrar);
+		}
+		GetClassMemberActions(Class, ClassActionList);
 
 		// queue the newly added actions for priming
 		if (ClassActionList.Num() > 0)
 		{
 			ActionPrimingQueue.Add(Class, 0);
 		}
+		else
+		{
+			ActionRegistry.Remove(Class);
+		}
 	}
 
 	// @TODO: account for all K2ActionMenuBuilder methods...
@@ -1003,25 +876,60 @@ void FBlueprintActionDatabase::RefreshClassActions(UClass* const Class)
 }
 
 //------------------------------------------------------------------------------
-void FBlueprintActionDatabase::ClearClassActions(UClass* const Class)
+void FBlueprintActionDatabase::RefreshAssetActions(UObject const* const AssetObject)
+{
+	using namespace BlueprintActionDatabaseImpl;
+	check(AssetObject->IsAsset());
+
+	FActionList& AssetActionList = ActionRegistry.FindOrAdd(AssetObject);
+	AssetActionList.Empty();
+
+	if (UBlueprint const* Blueprint = Cast<UBlueprint>(AssetObject))
+	{
+		AddBlueprintGraphActions(Blueprint, AssetActionList);
+		if (UClass* SkeletonClass = Blueprint->SkeletonGeneratedClass)
+		{
+			GetClassMemberActions(SkeletonClass, AssetActionList);
+		}
+	}
+
+	FBlueprintActionDatabaseRegistrar Registrar(ActionRegistry, ActionPrimingQueue);
+	Registrar.ActionKeyFilter = AssetObject; // make sure actions only associated with this asset get added
+	// nodes may have actions they wish to add actions for this asset
+	RegisterAllNodeActions(Registrar);
+
+	if (AssetActionList.Num() > 0)
+	{
+		// queue these assets for priming
+		ActionPrimingQueue.Add(AssetObject, 0);
+	}
+	else
+	{
+		ClearAssetActions(AssetObject);
+	}
+}
+
+//------------------------------------------------------------------------------
+void FBlueprintActionDatabase::ClearAssetActions(UObject const* const AssetObject)
 {
-	ClassActions.Remove(Class);
+	check(AssetObject->IsAsset());
+	ActionRegistry.Remove(AssetObject);
 }
 
 //------------------------------------------------------------------------------
-FBlueprintActionDatabase::FClassActionMap const& FBlueprintActionDatabase::GetAllActions()
+FBlueprintActionDatabase::FActionRegistry const& FBlueprintActionDatabase::GetAllActions()
 {
 	// if this is the first time that we're querying for actions, generate the
 	// list before returning it
-	if (ClassActions.Num() == 0)
+	if (ActionRegistry.Num() == 0)
 	{
 		RefreshAll();
 	}
-	return ClassActions;
+	return ActionRegistry;
 }
 
 //------------------------------------------------------------------------------
-int32 FBlueprintActionDatabase::Size() const
+int32 FBlueprintActionDatabase::EstimatedSize() const
 {
 	int32 TotalSize = sizeof(*this);
 	for (TObjectIterator<UBlueprintNodeSpawner> NodeSpawnerIt; NodeSpawnerIt; ++NodeSpawnerIt)
@@ -1031,4 +939,21 @@ int32 FBlueprintActionDatabase::Size() const
 	return TotalSize;
 }
 
+//------------------------------------------------------------------------------
+void FBlueprintActionDatabase::RegisterAllNodeActions(FBlueprintActionDatabaseRegistrar& Registrar)
+{
+	// nodes may have actions they wish to add for this asset
+	for (TObjectIterator<UClass> ClassIt; ClassIt; ++ClassIt)
+	{
+		UClass* NodeClass = *ClassIt;
+		if (!NodeClass->IsChildOf<UK2Node>())
+		{
+			continue;
+		}
+
+		TGuardValue< TSubclassOf<UEdGraphNode> > ScopedNodeClass(Registrar.GeneratingClass, NodeClass);
+		BlueprintActionDatabaseImpl::GetNodeSpecificActions(NodeClass, Registrar);
+	}
+}
+
 #undef LOCTEXT_NAMESPACE
diff --git a/Engine/Source/Editor/BlueprintGraph/Private/BlueprintActionDatabaseRegistrar.cpp b/Engine/Source/Editor/BlueprintGraph/Private/BlueprintActionDatabaseRegistrar.cpp
new file mode 100644
--- /dev/null
+++ b/Engine/Source/Editor/BlueprintGraph/Private/BlueprintActionDatabaseRegistrar.cpp
@@ -0,0 +1,159 @@
+// Copyright 1998-2014 Epic Games, Inc. All Rights Reserved.
+
+#include "BlueprintGraphPrivatePCH.h"
+#include "BlueprintActionDatabaseRegistrar.h"
+#include "BlueprintNodeSpawnerUtils.h"
+
+/*******************************************************************************
+ * BlueprintActionDatabaseRegistrarImpl
+ ******************************************************************************/
+
+namespace BlueprintActionDatabaseRegistrarImpl
+{
+	static UObject const* ResolveClassKey(UClass const* ClassKey);
+	static UObject const* ResolveActionKey(UObject const* UserPassedKey);
+}
+
+//------------------------------------------------------------------------------
+static UObject const* BlueprintActionDatabaseRegistrarImpl::ResolveClassKey(UClass const* ClassKey)
+{
+	UObject const* ResolvedKey = ClassKey;
+	if (UBlueprintGeneratedClass const* BlueprintClass = Cast<UBlueprintGeneratedClass>(ClassKey))
+	{
+		ResolvedKey = CastChecked<UBlueprint>(BlueprintClass->ClassGeneratedBy);
+	}
+	return ResolvedKey;
+}
+
+//------------------------------------------------------------------------------
+static UObject const* BlueprintActionDatabaseRegistrarImpl::ResolveActionKey(UObject const* UserPassedKey)
+{
+	UObject const* ResolvedKey = nullptr;
+	if (UClass const* Class = Cast<UClass>(UserPassedKey))
+	{
+		ResolvedKey = ResolveClassKey(Class);
+	}
+	// both handled in the IsAsset() case
+// 	else if (UUserDefinedEnum const* EnumAsset = Cast<UUserDefinedEnum>(UserPassedKey))
+// 	{
+// 		ResolvedKey = EnumAsset;
+// 	}
+// 	else if (UUserDefinedStruct const* StructAsset = Cast<UUserDefinedStruct>(StructOwner))
+// 	{
+// 		ResolvedKey = StructAsset;
+// 	}
+	else if (UserPassedKey->IsAsset())
+	{
+		ResolvedKey = UserPassedKey;
+	}
+	else if (UField const* MemberField = Cast<UField>(UserPassedKey))
+	{
+		ResolvedKey = ResolveClassKey(MemberField->GetOwnerClass());
+	}
+
+	return ResolvedKey;
+}
+
+/*******************************************************************************
+ * FBlueprintActionDatabaseRegistrar
+ ******************************************************************************/
+
+//------------------------------------------------------------------------------
+FBlueprintActionDatabaseRegistrar::FBlueprintActionDatabaseRegistrar(FActionRegistry& Database, FPrimingQueue& PrimingQueue, TSubclassOf<UEdGraphNode> DefaultKey)
+	: ActionDatabase(Database)
+	, GeneratingClass(DefaultKey)
+	, ActionKeyFilter(nullptr)
+	, ActionPrimingQueue(PrimingQueue)
+{
+}
+
+//------------------------------------------------------------------------------
+void FBlueprintActionDatabaseRegistrar::AddBlueprintAction(UBlueprintNodeSpawner* NodeSpawner)
+{
+	UField const* ActionKey = GeneratingClass;
+	// if this spawner wraps some member function/property, then we want it 
+	// recorded under that class (so we can update the action if the class 
+	// changes... like, if the member is deleted, or if one is added)
+	if (UField const* MemberField = FBlueprintNodeSpawnerUtils::GetAssociatedField(NodeSpawner))
+	{
+		ActionKey = MemberField;
+	}
+	AddActionToDatabase(ActionKey, NodeSpawner);
+}
+
+//------------------------------------------------------------------------------
+void FBlueprintActionDatabaseRegistrar::AddBlueprintAction(UClass const* ClassOwner, UBlueprintNodeSpawner* NodeSpawner)
+{
+	// ResolveActionKey() is used on ClassOwner (in AddActionToDatabase), to 
+	// convert it into a proper key
+	AddActionToDatabase(ClassOwner, NodeSpawner);
+}
+
+//------------------------------------------------------------------------------
+void FBlueprintActionDatabaseRegistrar::AddBlueprintAction(UEnum const* EnumOwner, UBlueprintNodeSpawner* NodeSpawner)
+{
+	// ResolveActionKey() is used on EnumOwner (in AddActionToDatabase), to 
+	// convert it into a proper key
+	AddActionToDatabase(EnumOwner, NodeSpawner);
+}
+
+//------------------------------------------------------------------------------
+void FBlueprintActionDatabaseRegistrar::AddBlueprintAction(UScriptStruct const* StructOwner, UBlueprintNodeSpawner* NodeSpawner)
+{
+	// ResolveActionKey() is used on StructOwner (in AddActionToDatabase), to 
+	// convert it into a proper key
+	AddActionToDatabase(StructOwner, NodeSpawner);
+}
+
+//------------------------------------------------------------------------------
+void FBlueprintActionDatabaseRegistrar::AddBlueprintAction(UField const* FieldOwner, UBlueprintNodeSpawner* NodeSpawner)
+{
+	// ResolveActionKey() is used on FieldOwner (in AddActionToDatabase), to 
+	// convert it into a proper key
+	AddActionToDatabase(FieldOwner, NodeSpawner);
+}
+
+//------------------------------------------------------------------------------
+void FBlueprintActionDatabaseRegistrar::AddBlueprintAction(UObject const* AssetOwner, UBlueprintNodeSpawner* NodeSpawner)
+{
+	// cannot record an action under any ol' object (we want to associate them 
+	// with asset/class outers that are subject to change; so that we can 
+	// refresh/rebuild corresponding actions when that happens).
+	check(AssetOwner->IsAsset());
+	AddActionToDatabase(AssetOwner, NodeSpawner);
+}
+
+//------------------------------------------------------------------------------
+bool FBlueprintActionDatabaseRegistrar::IsOpenForRegistration(UObject const* OwnerKey)
+{
+	UObject const* ActionKey = BlueprintActionDatabaseRegistrarImpl::ResolveActionKey(OwnerKey);
+	if (ActionKey == nullptr)
+	{
+		ActionKey = GeneratingClass;
+	}
+	return (ActionKey != nullptr) && ((ActionKeyFilter == nullptr) || (ActionKeyFilter == ActionKey));
+}
+
+//------------------------------------------------------------------------------
+void FBlueprintActionDatabaseRegistrar::AddActionToDatabase(UObject const* ActionKey, UBlueprintNodeSpawner* NodeSpawner)
+{
+	ensureMsg(NodeSpawner->NodeClass == GeneratingClass, TEXT("We expect a nodes to add only spawners for its own type... Maybe a sub-class is adding nodes it shouldn't?"));
+	if (IsOpenForRegistration(ActionKey))
+	{
+		ActionKey = BlueprintActionDatabaseRegistrarImpl::ResolveActionKey(ActionKey);
+		if (ActionKey == nullptr)
+		{
+			ActionKey = GeneratingClass;
+		}
+		FBlueprintActionDatabase::FActionList& ActionList = ActionDatabase.FindOrAdd(ActionKey);
+		
+		int32* QueuedIndex = ActionPrimingQueue.Find(ActionKey);
+		if (QueuedIndex == nullptr)
+		{
+			int32 PrimingIndex = ActionList.Num();
+			ActionPrimingQueue.Add(ActionKey, PrimingIndex);
+		}
+
+		ActionList.Add(NodeSpawner);
+	}
+}
\ No newline at end of file
diff --git a/Engine/Source/Editor/BlueprintGraph/Private/BlueprintActionFilter.cpp b/Engine/Source/Editor/BlueprintGraph/Private/BlueprintActionFilter.cpp
--- a/Engine/Source/Editor/BlueprintGraph/Private/BlueprintActionFilter.cpp
+++ b/Engine/Source/Editor/BlueprintGraph/Private/BlueprintActionFilter.cpp
@@ -52,13 +52,14 @@ namespace BlueprintActionFilterImpl
 	static bool IsLatent(UBlueprintNodeSpawner const* NodeSpawner);
 	
 	/**
-	 * Utility method to check and see if the specified node-spawner would
-	 * produce a node associated with a global field.
+	 * Utility method to check and see if the specified field is a public global
+	 * or static field (that is has a persistent extent that span's the 
+	 * program's lifetime).
 	 *
-	 * @param  NodeSpawner	The action you want to query.
-	 * @return True if the action is associated with some global field, false if not (or unknown).
+	 * @param  Field	The field you want to check.
+	 * @return True if the field is global/static (and public), false if it has a limited extent (or is private/protected).
 	 */
-	static bool IsGlobal(UField const* Field);
+	static bool IsGloballyAccessible(UField const* Field);
 	
 	/**
 	 * Rejection test that checks to see if the supplied node-spawner would
@@ -135,40 +136,31 @@ namespace BlueprintActionFilterImpl
 	
 	/**
 	 * Rejection test that checks to see if the node-spawner has any associated
-	 * fields that are global (structs, enums, etc.), or if the field belongs to
-	 * a static library class.
+	 * "non-target" fields that are global/static.
 	 * 
-	 * @param  Context			Holds the context (unused) for this test.
+	 * @param  Filter			Holds the TagetClass context for this test.
 	 * @param  BlueprintAction	The action you wish to query.
-	 * @return True if the action would spawn a node not allowed in function graphs.
+	 * @return 
 	 */
-	static bool IsGlobal(FBlueprintActionFilter const& Filter, UBlueprintNodeSpawner const* BlueprintAction);
+	static bool IsPersistentNonTargetField(FBlueprintActionFilter const& Filter, UBlueprintNodeSpawner const* BlueprintAction);
 
 	/**
 	 * Rejection test that checks to see if the node-spawner is associated with 
-	 * a field that belongs to a class that is not white-listed (ignores global
-	 * fields).
+	 * a field that belongs to a class that is not white-listed (ignores global/
+	 * static fields).
 	 * 
-	 * @param  Context			Holds the class context for this test.
-	 * @param  BlueprintAction	The action you wish to query.
+	 * @param  Filter					Holds the class context for this test.
+	 * @param  BlueprintAction			The action you wish to query.
+	 * @param  bPermitNonTargetGlobals	Determines if this test should pass for external global/static fields.
 	 * @return True if the action is associated with a non-whitelisted class member.
 	 */
-	static bool IsExternalField(FBlueprintActionFilter const& Filter, UBlueprintNodeSpawner const* BlueprintAction);
-
-	/**
-	 * 
-	 * 
-	 * @param  Filter	
-	 * @param  BlueprintAction	
-	 * @return 
-	 */
-	static bool IsUneededMessageNode(FBlueprintActionFilter const& Filter, UBlueprintNodeSpawner const* BlueprintAction);
+	static bool IsNonTargetMemeber(FBlueprintActionFilter const& Filter, UBlueprintNodeSpawner const* BlueprintAction, bool const bPermitNonTargetGlobals);
 
 	/**
 	 * Rejection test that checks to see if the node-spawner is associated with 
 	 * a field that is hidden from the specified blueprint (via metadata).
 	 * 
-	 * @param  Context			Holds the class context for this test.
+	 * @param  Filter			Holds the class context for this test.
 	 * @param  BlueprintAction	The action you wish to query.
 	 * @return True if the action is associated with a hidden field.
 	 */
@@ -178,7 +170,7 @@ namespace BlueprintActionFilterImpl
 	 * Rejection test that checks to see if the node-spawner would produce a 
 	 * node type that isn't white-listed. 
 	 * 
-	 * @param  Context			Holds the class context for this test.
+	 * @param  Filter			Holds the class context for this test.
 	 * @param  BlueprintAction	The action you wish to query.
 	 * @param  bMatchExplicitly When true, the node's class has to perfectly match a white-listed type (sub-classes don't count).
 	 * @return True if the action would produce a non-whitelisted node.
@@ -190,7 +182,7 @@ namespace BlueprintActionFilterImpl
 	 * specific object that is not currently selected.
 	 * 
 	 * 
-	 * @param  Context			Holds the selected object context for this test.
+	 * @param  Filter			Holds the selected object context for this test.
 	 * @param  BlueprintAction	The action you wish to query.
 	 * @return True if the action would produce a bound node, tied to an object that isn't selected.
 	 */
@@ -214,6 +206,15 @@ namespace BlueprintActionFilterImpl
 	 */
 	static bool IsOutOfScopeLocalVariable(FBlueprintActionFilter const& Filter, UBlueprintNodeSpawner const* BlueprintAction);
 
+	/**
+	 * 
+	 * 
+	 * @param  Filter	
+	 * @param  BlueprintAction	
+	 * @return 
+	 */
+	static bool IsSchemaIncompatible(FBlueprintActionFilter const& Filter, UBlueprintNodeSpawner const* BlueprintAction);
+
 	/**
 	 * 
 	 * 
@@ -314,25 +315,26 @@ static bool BlueprintActionFilterImpl::IsLatent(UBlueprintNodeSpawner const* Nod
 }
 
 //------------------------------------------------------------------------------
-static bool BlueprintActionFilterImpl::IsGlobal(UField const* Field)
+static bool BlueprintActionFilterImpl::IsGloballyAccessible(UField const* Field)
 {
-	bool bIsGlobal = false;
+	bool bHasPersistentExtents = false; // is global or static
+	bool bIsPublic = Field->HasAnyFlags(RF_Public);
 	
 	UClass* ClassOuter = Cast<UClass>(Field->GetOuter());
 	// outer is probably a UPackage (for like global enums, structs, etc.)
 	if (ClassOuter == nullptr)
 	{
-		bIsGlobal = true;
+		bHasPersistentExtents = true;
 	}
 	else if (UFunction const* Function = Cast<UFunction>(Field))
 	{
-		bool const bIsPublic = !Function->HasMetaData(FBlueprintMetadata::MD_Protected) &&
+		bIsPublic |= !Function->HasMetaData(FBlueprintMetadata::MD_Protected) &&
 			!Function->HasMetaData(FBlueprintMetadata::MD_Private);
 		
-		bIsGlobal = Function->HasAnyFunctionFlags(FUNC_Static) && bIsPublic;
+		bHasPersistentExtents = Function->HasAnyFunctionFlags(FUNC_Static);
 	}
 	
-	return bIsGlobal;
+	return bIsPublic && bHasPersistentExtents;
 }
 
 //------------------------------------------------------------------------------
@@ -487,20 +489,20 @@ static bool BlueprintActionFilterImpl::IsDeprecated(FBlueprintActionFilter const
 }
 
 //------------------------------------------------------------------------------
-static bool BlueprintActionFilterImpl::IsGlobal(FBlueprintActionFilter const& Filter, UBlueprintNodeSpawner const* BlueprintAction)
+static bool BlueprintActionFilterImpl::IsPersistentNonTargetField(FBlueprintActionFilter const& Filter, UBlueprintNodeSpawner const* BlueprintAction)
 {
 	bool bIsFilteredOut = false;
 	if (UField const* ClassField = FBlueprintNodeSpawnerUtils::GetAssociatedField(BlueprintAction))
 	{
-		bIsFilteredOut = IsGlobal(ClassField);
+		bIsFilteredOut = IsGloballyAccessible(ClassField);
 		
 		UClass* FieldClass = ClassField->GetOwnerClass();
 		if (bIsFilteredOut && (FieldClass != nullptr))
 		{
-			for (UClass const* Class : Filter.OwnerClasses)
+			for (UClass const* Class : Filter.TargetClasses)
 			{
-				bool const bIsInternalFunc = Class->IsChildOf(FieldClass);
-				if (bIsInternalFunc)
+				bool const bIsInternalMemberField = Class->IsChildOf(FieldClass);
+				if (bIsInternalMemberField)
 				{
 					bIsFilteredOut = false;
 					break;
@@ -513,7 +515,7 @@ static bool BlueprintActionFilterImpl::IsGlobal(FBlueprintActionFilter const& Fi
 }
 
 //------------------------------------------------------------------------------
-static bool BlueprintActionFilterImpl::IsExternalField(FBlueprintActionFilter const& Filter, UBlueprintNodeSpawner const* BlueprintAction)
+static bool BlueprintActionFilterImpl::IsNonTargetMemeber(FBlueprintActionFilter const& Filter, UBlueprintNodeSpawner const* BlueprintAction, bool bPermitNonTargetGlobals)
 {
 	bool bIsFilteredOut = false;
 	if (UField const* ClassField = FBlueprintNodeSpawnerUtils::GetAssociatedField(BlueprintAction))
@@ -523,9 +525,9 @@ static bool BlueprintActionFilterImpl::IsExternalField(FBlueprintActionFilter co
 
 		// global (and static library) fields can stay (unless explicitly
 		// excluded... save that for a separate test)
-		if ((FieldClass != nullptr) && !IsGlobal(ClassField))
+		if ((FieldClass != nullptr) && (bPermitNonTargetGlobals || IsGloballyAccessible(ClassField)))
 		{
-			for (UClass const* Class : Filter.OwnerClasses)
+			for (UClass const* Class : Filter.TargetClasses)
 			{
 				bool const bIsInternalFunc = Class->IsChildOf(FieldClass) || (bIsInterfaceMethod && Class->ImplementsInterface(FieldClass));
 				if (!bIsInternalFunc)
@@ -540,30 +542,6 @@ static bool BlueprintActionFilterImpl::IsExternalField(FBlueprintActionFilter co
 	return bIsFilteredOut;
 }
 
-//------------------------------------------------------------------------------
-static bool BlueprintActionFilterImpl::IsUneededMessageNode(FBlueprintActionFilter const& Filter, UBlueprintNodeSpawner const* BlueprintAction)
-{
-	bool bIsFilteredOut = false;
-	if (BlueprintAction->NodeClass->IsChildOf<UK2Node_Message>())
-	{
-		UFunction const* const InterfaceFunc = FBlueprintNodeSpawnerUtils::GetAssociatedFunction(BlueprintAction);
-		checkSlow(InterfaceFunc != nullptr);
-		checkSlow(InterfaceFunc->GetOuterUClass()->IsChildOf<UInterface>());
-
-		bIsFilteredOut = (Filter.Context.Pins.Num() > 0);
-		for (UEdGraphPin const* ContextPin : Filter.Context.Pins)
-		{
-			if (!IsPinCompatibleWithTargetSelf(ContextPin, InterfaceFunc, BlueprintAction))
-			{
-				bIsFilteredOut = false;
-				break;
-			}
-		}
-
-	}
-	return bIsFilteredOut;
-}
-
 //------------------------------------------------------------------------------
 static bool BlueprintActionFilterImpl::IsFieldCategoryHidden(FBlueprintActionFilter const& Filter, UBlueprintNodeSpawner const* BlueprintAction)
 {
@@ -603,9 +581,9 @@ static bool BlueprintActionFilterImpl::IsFieldCategoryHidden(FBlueprintActionFil
 				}
 			}
 
-			for (int32 ClassIndex = 0; !bIsFilteredOut && (ClassIndex < Filter.OwnerClasses.Num()); ++ClassIndex)
+			for (int32 ClassIndex = 0; !bIsFilteredOut && (ClassIndex < Filter.TargetClasses.Num()); ++ClassIndex)
 			{
-				bIsFilteredOut = IsFieldHiddenDelegate.Execute(Filter.OwnerClasses[ClassIndex]);
+				bIsFilteredOut = IsFieldHiddenDelegate.Execute(Filter.TargetClasses[ClassIndex]);
 			}
 		}
 	}
@@ -657,12 +635,12 @@ static bool BlueprintActionFilterImpl::IsIncompatibleWithGraphType(FBlueprintAct
 //------------------------------------------------------------------------------
 static bool BlueprintActionFilterImpl::IsFilteredNodeType(FBlueprintActionFilter const& Filter, UBlueprintNodeSpawner const* BlueprintAction, bool const bExcludeChildClasses)
 {
-	bool bIsFilteredOut = (Filter.NodeTypes.Num() > 0);
+	bool bIsFilteredOut = (Filter.PermittedNodeTypes.Num() > 0);
 
 	UClass* NodeClass = BlueprintAction->NodeClass;
 	if (NodeClass != nullptr)
 	{
-		for (TSubclassOf<UEdGraphNode> AllowedClass : Filter.NodeTypes)
+		for (TSubclassOf<UEdGraphNode> AllowedClass : Filter.PermittedNodeTypes)
 		{
 			if ((NodeClass->IsChildOf(AllowedClass) && !bExcludeChildClasses) ||
 				(AllowedClass == NodeClass))
@@ -672,9 +650,9 @@ static bool BlueprintActionFilterImpl::IsFilteredNodeType(FBlueprintActionFilter
 			}
 		}
 		
-		for (int32 ClassIndx = 0; !bIsFilteredOut && (ClassIndx < Filter.ExcludedNodeTypes.Num()); ++ClassIndx)
+		for (int32 ClassIndx = 0; !bIsFilteredOut && (ClassIndx < Filter.RejectedNodeTypes.Num()); ++ClassIndx)
 		{
-			TSubclassOf<UEdGraphNode> ExcludedClass = Filter.ExcludedNodeTypes[ClassIndx];
+			TSubclassOf<UEdGraphNode> ExcludedClass = Filter.RejectedNodeTypes[ClassIndx];
 			if ((bExcludeChildClasses && (ExcludedClass == NodeClass)) ||
 				NodeClass->IsChildOf(ExcludedClass))
 			{
@@ -748,6 +726,51 @@ static bool BlueprintActionFilterImpl::IsOutOfScopeLocalVariable(FBlueprintActio
 	return bIsFilteredOut;
 }
 
+//------------------------------------------------------------------------------
+static bool BlueprintActionFilterImpl::IsSchemaIncompatible(FBlueprintActionFilter const& Filter, UBlueprintNodeSpawner const* BlueprintAction)
+{
+	bool bIsFilteredOut = false;
+	FBlueprintActionContext const& FilterContext = Filter.Context;
+
+	UEdGraphNode const* NodeCDO = BlueprintAction->NodeClass->GetDefaultObject<UEdGraphNode>();
+
+	auto IsSchemaIncompatibleLambda = [NodeCDO](TArray<UEdGraph*> const& GraphList)->bool
+	{
+		bool bIsCompatible = true;
+		for (UEdGraph const* Graph : GraphList)
+		{
+			if (!NodeCDO->CanCreateUnderSpecifiedSchema(Graph->GetSchema()))
+			{
+				bIsCompatible = false;
+				break;
+			}
+		}
+		return !bIsCompatible;
+	};
+	
+	if (FilterContext.Graphs.Num() > 0)
+	{
+		bIsFilteredOut = IsSchemaIncompatibleLambda(FilterContext.Graphs);
+	}
+	else
+	{
+		bIsFilteredOut = true;
+		for (UBlueprint const* Blueprint : FilterContext.Blueprints)
+		{
+			TArray<UEdGraph*> BpGraphList;
+			Blueprint->GetAllGraphs(BpGraphList);
+
+			if (!IsSchemaIncompatibleLambda(BpGraphList))
+			{
+				bIsFilteredOut = false;
+				break;
+			}
+		}
+	}
+
+	return bIsFilteredOut;
+}
+
 //------------------------------------------------------------------------------
 static bool BlueprintActionFilterImpl::HasMatchingPin(UBlueprintNodeSpawner const* BlueprintAction, UEdGraphPin const* Pin)
 {
@@ -1006,9 +1029,9 @@ FBlueprintActionFilter::FBlueprintActionFilter(uint32 Flags/*= 0x00*/)
 	AddIsFilteredTest(FIsFilteredDelegate::CreateStatic(IsIncompatibleImpureNode));
 	
 	AddIsFilteredTest(FIsFilteredDelegate::CreateStatic(IsFieldCategoryHidden));	
-	if (Flags & BPFILTER_ExcludeGlobalFields)
+	if (Flags & BPFILTER_RejectPersistentNonTargetFields)
 	{
-		AddIsFilteredTest(FIsFilteredDelegate::CreateStatic(IsGlobal));
+		AddIsFilteredTest(FIsFilteredDelegate::CreateStatic(IsPersistentNonTargetField));
 	}
 	
 	AddIsFilteredTest(FIsFilteredDelegate::CreateStatic(IsFieldInaccessible));
@@ -1016,14 +1039,15 @@ FBlueprintActionFilter::FBlueprintActionFilter(uint32 Flags/*= 0x00*/)
 	AddIsFilteredTest(FIsFilteredDelegate::CreateStatic(IsPermissionNotGranted));
 	AddIsFilteredTest(FIsFilteredDelegate::CreateStatic(IsRestrictedClassMember));
 	AddIsFilteredTest(FIsFilteredDelegate::CreateStatic(IsIncompatibleWithGraphType));
+	AddIsFilteredTest(FIsFilteredDelegate::CreateStatic(IsSchemaIncompatible));
 
-	if (!(Flags & BPFILTER_IncludeDeprecated))
+	if (!(Flags & BPFILTER_PermitDeprecated))
 	{
 		AddIsFilteredTest(FIsFilteredDelegate::CreateStatic(IsDeprecated));
 	}
 
-	AddIsFilteredTest(FIsFilteredDelegate::CreateStatic(IsFilteredNodeType, (Flags & BPFILTER_ExcludeChildNodeTypes) != 0));
-	AddIsFilteredTest(FIsFilteredDelegate::CreateStatic(IsExternalField));
+	AddIsFilteredTest(FIsFilteredDelegate::CreateStatic(IsFilteredNodeType, (Flags & BPFILTER_RejectPermittedNodeSubClasses)   != 0));
+	AddIsFilteredTest(FIsFilteredDelegate::CreateStatic(IsNonTargetMemeber, (Flags & BPFILTER_RejectPersistentNonTargetFields) != 0));
 	
 	// don't know why we only allow impure ("imperitive") and const functions,
 	// but this mimics legacy functionality in GetFunctionCallsOnSelectedComponents()
diff --git a/Engine/Source/Editor/BlueprintGraph/Private/K2Node_AssignmentStatement.cpp b/Engine/Source/Editor/BlueprintGraph/Private/K2Node_AssignmentStatement.cpp
--- a/Engine/Source/Editor/BlueprintGraph/Private/K2Node_AssignmentStatement.cpp
+++ b/Engine/Source/Editor/BlueprintGraph/Private/K2Node_AssignmentStatement.cpp
@@ -6,6 +6,7 @@
 #include "VariableSetHandler.h"
 #include "BlueprintNodeSpawner.h"
 #include "EditorCategoryUtils.h"
+#include "BlueprintActionDatabaseRegistrar.h"
 
 #define LOCTEXT_NAMESPACE "K2Node_AssignmentStatement"
 
@@ -174,12 +175,12 @@ FNodeHandlingFunctor* UK2Node_AssignmentStatement::CreateNodeHandler(FKismetComp
 	return new FKCHandler_AssignmentStatement(CompilerContext);
 }
 
-void UK2Node_AssignmentStatement::GetMenuActions(TArray<UBlueprintNodeSpawner*>& ActionListOut) const
+void UK2Node_AssignmentStatement::GetMenuActions(FBlueprintActionDatabaseRegistrar& ActionRegistrar) const
 {
 	UBlueprintNodeSpawner* NodeSpawner = UBlueprintNodeSpawner::Create(GetClass());
 	check(NodeSpawner != nullptr);
 
-	ActionListOut.Add(NodeSpawner);
+	ActionRegistrar.AddBlueprintAction(NodeSpawner);
 }
 
 FText UK2Node_AssignmentStatement::GetMenuCategory() const
diff --git a/Engine/Source/Editor/BlueprintGraph/Private/K2Node_BreakStruct.cpp b/Engine/Source/Editor/BlueprintGraph/Private/K2Node_BreakStruct.cpp
--- a/Engine/Source/Editor/BlueprintGraph/Private/K2Node_BreakStruct.cpp
+++ b/Engine/Source/Editor/BlueprintGraph/Private/K2Node_BreakStruct.cpp
@@ -4,6 +4,7 @@
 #include "KismetCompiler.h"
 #include "BlueprintNodeSpawner.h"
 #include "EditorCategoryUtils.h"
+#include "BlueprintActionDatabaseRegistrar.h"
 
 #define LOCTEXT_NAMESPACE "K2Node_BreakStruct"
 
@@ -291,17 +292,12 @@ FNodeHandlingFunctor* UK2Node_BreakStruct::CreateNodeHandler(class FKismetCompil
 	return new FKCHandler_BreakStruct(CompilerContext);
 }
 
-void UK2Node_BreakStruct::GetMenuActions(TArray<UBlueprintNodeSpawner*>& ActionListOut) const
+void UK2Node_BreakStruct::GetMenuActions(FBlueprintActionDatabaseRegistrar& ActionRegistrar) const
 {
 	for (TObjectIterator<UScriptStruct> StructIt; StructIt; ++StructIt)
 	{
 		UScriptStruct const* Struct = (*StructIt);
-		// we only want to add autonomous structs here; those belonging to a 
-		// certain class should instead be associated with that class (so when 
-		// the class is modified we can easily handle any structs that were changed).
-		bool bIsStandaloneStruct = Struct->GetOuter()->IsA(UPackage::StaticClass());
-
-		if (!bIsStandaloneStruct || !UEdGraphSchema_K2::IsAllowableBlueprintVariableType(Struct) || !CanBeBroken(Struct))
+		if (!UEdGraphSchema_K2::IsAllowableBlueprintVariableType(Struct) || !CanBeBroken(Struct))
 		{
 			continue;
 		}
@@ -317,7 +313,10 @@ void UK2Node_BreakStruct::GetMenuActions(TArray<UBlueprintNodeSpawner*>& ActionL
 
 		UBlueprintNodeSpawner* NodeSpawner = UBlueprintNodeSpawner::Create(GetClass());
 		check(NodeSpawner != nullptr);
-		ActionListOut.Add(NodeSpawner);
+		// this struct could belong to a class, or is a user defined struct 
+		// (asset), that's why we want to make sure to register it along with 
+		// the action (so the action knows to refresh when the class/asset is).
+		ActionRegistrar.AddBlueprintAction(Struct, NodeSpawner);
 
 		TWeakObjectPtr<UScriptStruct> StructPtr = Struct;
 		NodeSpawner->CustomizeNodeDelegate = UBlueprintNodeSpawner::FCustomizeNodeDelegate::CreateStatic(CustomizeBreakNodeLambda, StructPtr);
diff --git a/Engine/Source/Editor/BlueprintGraph/Private/K2Node_CastByteToEnum.cpp b/Engine/Source/Editor/BlueprintGraph/Private/K2Node_CastByteToEnum.cpp
--- a/Engine/Source/Editor/BlueprintGraph/Private/K2Node_CastByteToEnum.cpp
+++ b/Engine/Source/Editor/BlueprintGraph/Private/K2Node_CastByteToEnum.cpp
@@ -7,6 +7,7 @@
 #include "K2Node_CastByteToEnum.h"
 #include "BlueprintNodeSpawner.h"
 #include "EditorCategoryUtils.h"
+#include "BlueprintActionDatabaseRegistrar.h"
 
 const FString UK2Node_CastByteToEnum::ByteInputPinName = TEXT("Byte");
 
@@ -141,19 +142,12 @@ FNodeHandlingFunctor* UK2Node_CastByteToEnum::CreateNodeHandler(FKismetCompilerC
 	return new FNodeHandlingFunctor(CompilerContext);;
 }
 
-void UK2Node_CastByteToEnum::GetMenuActions(TArray<UBlueprintNodeSpawner*>& ActionListOut) const
+void UK2Node_CastByteToEnum::GetMenuActions(FBlueprintActionDatabaseRegistrar& ActionRegistrar) const
 {
 	for (TObjectIterator<UEnum> EnumIt; EnumIt; ++EnumIt)
 	{
 		UEnum const* Enum = (*EnumIt);
-		// we only want to add global "standalone" enums here; those belonging to a 
-		// certain class should instead be associated with that class (so when 
-		// the class is modified we can easily handle any enums that were changed).
-		//
-		// @TODO: don't love how this code is essentially duplicated in BlueprintActionDatabase.cpp, for class enums
-		bool bIsStandaloneEnum = Enum->GetOuter()->IsA(UPackage::StaticClass());
-
-		if (!bIsStandaloneEnum || !UEdGraphSchema_K2::IsAllowableBlueprintVariableType(Enum))
+		if (!UEdGraphSchema_K2::IsAllowableBlueprintVariableType(Enum))
 		{
 			continue;
 		}
@@ -170,7 +164,10 @@ void UK2Node_CastByteToEnum::GetMenuActions(TArray<UBlueprintNodeSpawner*>& Acti
 
 		UBlueprintNodeSpawner* NodeSpawner = UBlueprintNodeSpawner::Create(GetClass());
 		check(NodeSpawner != nullptr);
-		ActionListOut.Add(NodeSpawner);
+		// this enum could belong to a class, or is a user defined enum (asset), 
+		// that's why we want to make sure to register it along with the action 
+		// (so the action can be refreshed when the class/asset is).
+		ActionRegistrar.AddBlueprintAction(Enum, NodeSpawner);
 
 		TWeakObjectPtr<UEnum> EnumPtr = Enum;
 		NodeSpawner->CustomizeNodeDelegate = UBlueprintNodeSpawner::FCustomizeNodeDelegate::CreateStatic(CustomizeEnumNodeLambda, EnumPtr);
diff --git a/Engine/Source/Editor/BlueprintGraph/Private/K2Node_ConstructObjectFromClass.cpp b/Engine/Source/Editor/BlueprintGraph/Private/K2Node_ConstructObjectFromClass.cpp
--- a/Engine/Source/Editor/BlueprintGraph/Private/K2Node_ConstructObjectFromClass.cpp
+++ b/Engine/Source/Editor/BlueprintGraph/Private/K2Node_ConstructObjectFromClass.cpp
@@ -5,6 +5,7 @@
 #include "Runtime/Engine/Classes/Kismet/GameplayStatics.h"
 #include "BlueprintNodeSpawner.h"
 #include "EditorCategoryUtils.h"
+#include "BlueprintActionDatabaseRegistrar.h"
 
 struct FK2Node_ConstructObjectFromClassHelper
 {
@@ -287,12 +288,12 @@ void UK2Node_ConstructObjectFromClass::GetNodeAttributes( TArray<TKeyValuePair<F
 	OutNodeAttributes.Add( TKeyValuePair<FString, FString>( TEXT( "ObjectClass" ), ClassToSpawnStr ));
 }
 
-void UK2Node_ConstructObjectFromClass::GetMenuActions(TArray<UBlueprintNodeSpawner*>& ActionListOut) const
+void UK2Node_ConstructObjectFromClass::GetMenuActions(FBlueprintActionDatabaseRegistrar& ActionRegisrar) const
 {
 	UBlueprintNodeSpawner* NodeSpawner = UBlueprintNodeSpawner::Create(GetClass());
 	check(NodeSpawner != nullptr);
 
-	ActionListOut.Add(NodeSpawner);
+	ActionRegisrar.AddBlueprintAction(NodeSpawner);
 }
 
 FText UK2Node_ConstructObjectFromClass::GetMenuCategory() const
diff --git a/Engine/Source/Editor/BlueprintGraph/Private/K2Node_CustomEvent.cpp b/Engine/Source/Editor/BlueprintGraph/Private/K2Node_CustomEvent.cpp
--- a/Engine/Source/Editor/BlueprintGraph/Private/K2Node_CustomEvent.cpp
+++ b/Engine/Source/Editor/BlueprintGraph/Private/K2Node_CustomEvent.cpp
@@ -6,6 +6,7 @@
 #include "Kismet2NameValidators.h"
 #include "CompilerResultsLog.h"
 #include "BlueprintEventNodeSpawner.h"
+#include "BlueprintActionDatabaseRegistrar.h"
 
 #define LOCTEXT_NAMESPACE "K2Node_CustomEvent"
 
@@ -229,7 +230,7 @@ void UK2Node_CustomEvent::ValidateNodeDuringCompilation(class FCompilerResultsLo
 	}
 }
 
-void UK2Node_CustomEvent::GetMenuActions(TArray<UBlueprintNodeSpawner*>& ActionListOut) const
+void UK2Node_CustomEvent::GetMenuActions(FBlueprintActionDatabaseRegistrar& ActionRegistrar) const
 {
 	UBlueprintNodeSpawner* NodeSpawner = UBlueprintEventNodeSpawner::Create(GetClass(), FName());
 	check(NodeSpawner != nullptr);
@@ -248,7 +249,7 @@ void UK2Node_CustomEvent::GetMenuActions(TArray<UBlueprintNodeSpawner*>& ActionL
 	};
 
 	NodeSpawner->CustomizeNodeDelegate = UBlueprintNodeSpawner::FCustomizeNodeDelegate::CreateStatic(SetupCustomEventNodeLambda);
-	ActionListOut.Add(NodeSpawner);
+	ActionRegistrar.AddBlueprintAction(NodeSpawner);
 }
 
 void UK2Node_CustomEvent::ReconstructNode()
diff --git a/Engine/Source/Editor/BlueprintGraph/Private/K2Node_EnumLiteral.cpp b/Engine/Source/Editor/BlueprintGraph/Private/K2Node_EnumLiteral.cpp
--- a/Engine/Source/Editor/BlueprintGraph/Private/K2Node_EnumLiteral.cpp
+++ b/Engine/Source/Editor/BlueprintGraph/Private/K2Node_EnumLiteral.cpp
@@ -6,6 +6,7 @@
 #include "K2Node_EnumLiteral.h"
 #include "BlueprintNodeSpawner.h"
 #include "EditorCategoryUtils.h"
+#include "BlueprintActionDatabaseRegistrar.h"
 
 const FString& UK2Node_EnumLiteral::GetEnumInputPinName()
 {
@@ -90,19 +91,12 @@ FNodeHandlingFunctor* UK2Node_EnumLiteral::CreateNodeHandler(FKismetCompilerCont
 	return new FKCHandler_EnumLiteral(CompilerContext);
 }
 
-void UK2Node_EnumLiteral::GetMenuActions(TArray<UBlueprintNodeSpawner*>& ActionListOut) const
+void UK2Node_EnumLiteral::GetMenuActions(FBlueprintActionDatabaseRegistrar& ActionRegistrar) const
 {
 	for (TObjectIterator<UEnum> EnumIt; EnumIt; ++EnumIt)
 	{
 		UEnum const* Enum = (*EnumIt);
-		// we only want to add global "standalone" enums here; those belonging to a 
-		// certain class should instead be associated with that class (so when 
-		// the class is modified we can easily handle any enums that were changed).
-		//
-		// @TODO: don't love how this code is essentially duplicated in BlueprintActionDatabase.cpp, for class enums
-		bool bIsStandaloneEnum = Enum->GetOuter()->IsA(UPackage::StaticClass());
-
-		if (!bIsStandaloneEnum || !UEdGraphSchema_K2::IsAllowableBlueprintVariableType(Enum))
+		if (!UEdGraphSchema_K2::IsAllowableBlueprintVariableType(Enum))
 		{
 			continue;
 		}
@@ -118,7 +112,10 @@ void UK2Node_EnumLiteral::GetMenuActions(TArray<UBlueprintNodeSpawner*>& ActionL
 
 		UBlueprintNodeSpawner* NodeSpawner = UBlueprintNodeSpawner::Create(GetClass());
 		check(NodeSpawner != nullptr);
-		ActionListOut.Add(NodeSpawner);
+		// this enum could belong to a class, or is a user defined enum (asset), 
+		// that's why we want to make sure to register it along with the action 
+		// (so the action can be refreshed when the class/asset is).
+		ActionRegistrar.AddBlueprintAction(Enum, NodeSpawner);
 
 		TWeakObjectPtr<UEnum> EnumPtr = Enum;
 		NodeSpawner->CustomizeNodeDelegate = UBlueprintNodeSpawner::FCustomizeNodeDelegate::CreateStatic(CustomizeEnumNodeLambda, EnumPtr);
diff --git a/Engine/Source/Editor/BlueprintGraph/Private/K2Node_ExecutionSequence.cpp b/Engine/Source/Editor/BlueprintGraph/Private/K2Node_ExecutionSequence.cpp
--- a/Engine/Source/Editor/BlueprintGraph/Private/K2Node_ExecutionSequence.cpp
+++ b/Engine/Source/Editor/BlueprintGraph/Private/K2Node_ExecutionSequence.cpp
@@ -6,6 +6,7 @@
 #include "KismetCompilerMisc.h"
 #include "BlueprintNodeSpawner.h"
 #include "EditorCategoryUtils.h"
+#include "BlueprintActionDatabaseRegistrar.h"
 
 #define LOCTEXT_NAMESPACE "K2Node_MultiGate"
 
@@ -263,12 +264,12 @@ FNodeHandlingFunctor* UK2Node_ExecutionSequence::CreateNodeHandler(FKismetCompil
 	return new FKCHandler_ExecutionSequence(CompilerContext);
 }
 
-void UK2Node_ExecutionSequence::GetMenuActions(TArray<UBlueprintNodeSpawner*>& ActionListOut) const
+void UK2Node_ExecutionSequence::GetMenuActions(FBlueprintActionDatabaseRegistrar& ActionRegistrar) const
 {
 	UBlueprintNodeSpawner* NodeSpawner = UBlueprintNodeSpawner::Create(GetClass());
 	check(NodeSpawner != nullptr);
 
-	ActionListOut.Add(NodeSpawner);
+	ActionRegistrar.AddBlueprintAction(NodeSpawner);
 }
 
 FText UK2Node_ExecutionSequence::GetMenuCategory() const
diff --git a/Engine/Source/Editor/BlueprintGraph/Private/K2Node_ForEachElementInEnum.cpp b/Engine/Source/Editor/BlueprintGraph/Private/K2Node_ForEachElementInEnum.cpp
--- a/Engine/Source/Editor/BlueprintGraph/Private/K2Node_ForEachElementInEnum.cpp
+++ b/Engine/Source/Editor/BlueprintGraph/Private/K2Node_ForEachElementInEnum.cpp
@@ -10,6 +10,7 @@
 #include "K2Node_GetNumEnumEntries.h"
 #include "BlueprintNodeSpawner.h"
 #include "EditorCategoryUtils.h"
+#include "BlueprintActionDatabaseRegistrar.h"
 
 #define LOCTEXT_NAMESPACE "K2Node"
 
@@ -193,19 +194,12 @@ void UK2Node_ForEachElementInEnum::ExpandNode(class FKismetCompilerContext& Comp
 	}
 }
 
-void UK2Node_ForEachElementInEnum::GetMenuActions(TArray<UBlueprintNodeSpawner*>& ActionListOut) const
+void UK2Node_ForEachElementInEnum::GetMenuActions(FBlueprintActionDatabaseRegistrar& ActionRegistrar) const
 {
 	for (TObjectIterator<UEnum> EnumIt; EnumIt; ++EnumIt)
 	{
 		UEnum const* Enum = (*EnumIt);
-		// we only want to add global "standalone" enums here; those belonging to a 
-		// certain class should instead be associated with that class (so when 
-		// the class is modified we can easily handle any enums that were changed).
-		//
-		// @TODO: don't love how this code is essentially duplicated in BlueprintActionDatabase.cpp, for class enums
-		bool bIsStandaloneEnum = Enum->GetOuter()->IsA(UPackage::StaticClass());
-
-		if (!bIsStandaloneEnum || !UEdGraphSchema_K2::IsAllowableBlueprintVariableType(Enum))
+		if (!UEdGraphSchema_K2::IsAllowableBlueprintVariableType(Enum))
 		{
 			continue;
 		}
@@ -221,7 +215,10 @@ void UK2Node_ForEachElementInEnum::GetMenuActions(TArray<UBlueprintNodeSpawner*>
 
 		UBlueprintNodeSpawner* NodeSpawner = UBlueprintNodeSpawner::Create(GetClass());
 		check(NodeSpawner != nullptr);
-		ActionListOut.Add(NodeSpawner);
+		// this enum could belong to a class, or is a user defined enum (asset), 
+		// that's why we want to make sure to register it along with the action 
+		// (so the action can be refreshed when the class/asset is).
+		ActionRegistrar.AddBlueprintAction(Enum, NodeSpawner);
 
 		TWeakObjectPtr<UEnum> EnumPtr = Enum;
 		NodeSpawner->CustomizeNodeDelegate = UBlueprintNodeSpawner::FCustomizeNodeDelegate::CreateStatic(CustomizeEnumNodeLambda, EnumPtr);
diff --git a/Engine/Source/Editor/BlueprintGraph/Private/K2Node_FormatText.cpp b/Engine/Source/Editor/BlueprintGraph/Private/K2Node_FormatText.cpp
--- a/Engine/Source/Editor/BlueprintGraph/Private/K2Node_FormatText.cpp
+++ b/Engine/Source/Editor/BlueprintGraph/Private/K2Node_FormatText.cpp
@@ -8,6 +8,7 @@
 #include "ScopedTransaction.h"
 #include "BlueprintNodeSpawner.h"
 #include "EditorCategoryUtils.h"
+#include "BlueprintActionDatabaseRegistrar.h"
 
 #define LOCTEXT_NAMESPACE "K2Node_FormatText"
 
@@ -375,12 +376,12 @@ UEdGraphPin* UK2Node_FormatText::GetFormatPin() const
 	return CachedFormatPin;
 }
 
-void UK2Node_FormatText::GetMenuActions(TArray<UBlueprintNodeSpawner*>& ActionListOut) const
+void UK2Node_FormatText::GetMenuActions(FBlueprintActionDatabaseRegistrar& ActionRegistrar) const
 {
 	UBlueprintNodeSpawner* NodeSpawner = UBlueprintNodeSpawner::Create(GetClass());
 	check(NodeSpawner != nullptr);
 
-	ActionListOut.Add(NodeSpawner);
+	ActionRegistrar.AddBlueprintAction(NodeSpawner);
 }
 
 FText UK2Node_FormatText::GetMenuCategory() const
diff --git a/Engine/Source/Editor/BlueprintGraph/Private/K2Node_GetDataTableRow.cpp b/Engine/Source/Editor/BlueprintGraph/Private/K2Node_GetDataTableRow.cpp
--- a/Engine/Source/Editor/BlueprintGraph/Private/K2Node_GetDataTableRow.cpp
+++ b/Engine/Source/Editor/BlueprintGraph/Private/K2Node_GetDataTableRow.cpp
@@ -6,6 +6,7 @@
 #include "Kismet/DataTableFunctionLibrary.h"
 #include "BlueprintNodeSpawner.h"
 #include "EditorCategoryUtils.h"
+#include "BlueprintActionDatabaseRegistrar.h"
 
 #define LOCTEXT_NAMESPACE "K2Node_GetDataTableRow"
 
@@ -167,12 +168,12 @@ void UK2Node_GetDataTableRow::ReallocatePinsDuringReconstruction(TArray<UEdGraph
 	}
 }
 
-void UK2Node_GetDataTableRow::GetMenuActions(TArray<UBlueprintNodeSpawner*>& ActionListOut) const
+void UK2Node_GetDataTableRow::GetMenuActions(FBlueprintActionDatabaseRegistrar& ActionRegistrar) const
 {
 	UBlueprintNodeSpawner* NodeSpawner = UBlueprintNodeSpawner::Create(GetClass());
 	check(NodeSpawner != nullptr);
 
-	ActionListOut.Add(NodeSpawner);
+	ActionRegistrar.AddBlueprintAction(NodeSpawner);
 }
 
 FText UK2Node_GetDataTableRow::GetMenuCategory() const
diff --git a/Engine/Source/Editor/BlueprintGraph/Private/K2Node_GetEnumeratorName.cpp b/Engine/Source/Editor/BlueprintGraph/Private/K2Node_GetEnumeratorName.cpp
--- a/Engine/Source/Editor/BlueprintGraph/Private/K2Node_GetEnumeratorName.cpp
+++ b/Engine/Source/Editor/BlueprintGraph/Private/K2Node_GetEnumeratorName.cpp
@@ -8,6 +8,7 @@
 #include "K2Node_GetEnumeratorName.h"
 #include "BlueprintNodeSpawner.h"
 #include "EditorCategoryUtils.h"
+#include "BlueprintActionDatabaseRegistrar.h"
 
 FString UK2Node_GetEnumeratorName::EnumeratorPinName = TEXT("Enumerator");
 
@@ -184,12 +185,12 @@ void UK2Node_GetEnumeratorName::ExpandNode(class FKismetCompilerContext& Compile
 	}
 }
 
-void UK2Node_GetEnumeratorName::GetMenuActions(TArray<UBlueprintNodeSpawner*>& ActionListOut) const
+void UK2Node_GetEnumeratorName::GetMenuActions(FBlueprintActionDatabaseRegistrar& ActionRegistrar) const
 {
 	UBlueprintNodeSpawner* NodeSpawner = UBlueprintNodeSpawner::Create(GetClass());
 	check(NodeSpawner != nullptr);
 
-	ActionListOut.Add(NodeSpawner);
+	ActionRegistrar.AddBlueprintAction(NodeSpawner);
 }
 
 FText UK2Node_GetEnumeratorName::GetMenuCategory() const
diff --git a/Engine/Source/Editor/BlueprintGraph/Private/K2Node_GetEnumeratorNameAsString.cpp b/Engine/Source/Editor/BlueprintGraph/Private/K2Node_GetEnumeratorNameAsString.cpp
--- a/Engine/Source/Editor/BlueprintGraph/Private/K2Node_GetEnumeratorNameAsString.cpp
+++ b/Engine/Source/Editor/BlueprintGraph/Private/K2Node_GetEnumeratorNameAsString.cpp
@@ -7,6 +7,7 @@
 #include "K2Node_GetEnumeratorNameAsString.h"
 #include "BlueprintNodeSpawner.h"
 #include "EditorCategoryUtils.h"
+#include "BlueprintActionDatabaseRegistrar.h"
 
 UK2Node_GetEnumeratorNameAsString::UK2Node_GetEnumeratorNameAsString(const class FPostConstructInitializeProperties& PCIP)
 	: Super(PCIP)
@@ -31,12 +32,12 @@ FText UK2Node_GetEnumeratorNameAsString::GetNodeTitle(ENodeTitleType::Type Title
 	return NSLOCTEXT("K2Node", "GetEnumeratorNameAsString_Title", "Enum to String");
 }
 
-void UK2Node_GetEnumeratorNameAsString::GetMenuActions(TArray<UBlueprintNodeSpawner*>& ActionListOut) const
+void UK2Node_GetEnumeratorNameAsString::GetMenuActions(FBlueprintActionDatabaseRegistrar& ActionRegistrar) const
 {
 	UBlueprintNodeSpawner* NodeSpawner = UBlueprintNodeSpawner::Create(GetClass());
 	check(NodeSpawner != nullptr);
 
-	ActionListOut.Add(NodeSpawner);
+	ActionRegistrar.AddBlueprintAction(NodeSpawner);
 }
 
 FName UK2Node_GetEnumeratorNameAsString::GetFunctionName() const
diff --git a/Engine/Source/Editor/BlueprintGraph/Private/K2Node_GetInputAxisKeyValue.cpp b/Engine/Source/Editor/BlueprintGraph/Private/K2Node_GetInputAxisKeyValue.cpp
--- a/Engine/Source/Editor/BlueprintGraph/Private/K2Node_GetInputAxisKeyValue.cpp
+++ b/Engine/Source/Editor/BlueprintGraph/Private/K2Node_GetInputAxisKeyValue.cpp
@@ -8,6 +8,7 @@
 #include "Engine/InputAxisKeyDelegateBinding.h"
 #include "BlueprintEditorUtils.h"
 #include "EdGraphSchema_K2.h"
+#include "BlueprintActionDatabaseRegistrar.h"
 
 #define LOCTEXT_NAMESPACE "K2Node_GetInputAxisKeyValue"
 
@@ -113,7 +114,7 @@ FName UK2Node_GetInputAxisKeyValue::GetPaletteIcon(FLinearColor& OutColor) const
 }
 
 
-void UK2Node_GetInputAxisKeyValue::GetMenuActions(TArray<UBlueprintNodeSpawner*>& ActionListOut) const
+void UK2Node_GetInputAxisKeyValue::GetMenuActions(FBlueprintActionDatabaseRegistrar& ActionRegistrar) const
 {
 	TArray<FKey> AllKeys;
 	EKeys::GetAllKeys(AllKeys);
@@ -135,7 +136,7 @@ void UK2Node_GetInputAxisKeyValue::GetMenuActions(TArray<UBlueprintNodeSpawner*>
 		};
 		
 		NodeSpawner->CustomizeNodeDelegate = UBlueprintNodeSpawner::FCustomizeNodeDelegate::CreateStatic(CustomizeInputNodeLambda, Key);
-		ActionListOut.Add(NodeSpawner);
+		ActionRegistrar.AddBlueprintAction(NodeSpawner);
 	}
 }
 
diff --git a/Engine/Source/Editor/BlueprintGraph/Private/K2Node_GetInputAxisValue.cpp b/Engine/Source/Editor/BlueprintGraph/Private/K2Node_GetInputAxisValue.cpp
--- a/Engine/Source/Editor/BlueprintGraph/Private/K2Node_GetInputAxisValue.cpp
+++ b/Engine/Source/Editor/BlueprintGraph/Private/K2Node_GetInputAxisValue.cpp
@@ -8,6 +8,7 @@
 #include "Engine/InputAxisDelegateBinding.h"
 #include "BlueprintEditorUtils.h"
 #include "EdGraphSchema_K2.h"
+#include "BlueprintActionDatabaseRegistrar.h"
 
 #define LOCTEXT_NAMESPACE "K2Node_GetInputAxisValue"
 
@@ -89,7 +90,7 @@ void UK2Node_GetInputAxisValue::RegisterDynamicBinding(UDynamicBlueprintBinding*
 	InputAxisBindingObject->InputAxisDelegateBindings.Add(Binding);
 }
 
-void UK2Node_GetInputAxisValue::GetMenuActions(TArray<UBlueprintNodeSpawner*>& ActionListOut) const
+void UK2Node_GetInputAxisValue::GetMenuActions(FBlueprintActionDatabaseRegistrar& ActionRegistrar) const
 {
 	TArray<FName> AxisNames;
 	GetDefault<UInputSettings>()->GetAxisNames(AxisNames);
@@ -106,7 +107,7 @@ void UK2Node_GetInputAxisValue::GetMenuActions(TArray<UBlueprintNodeSpawner*>& A
 		};
 
 		NodeSpawner->CustomizeNodeDelegate = UBlueprintNodeSpawner::FCustomizeNodeDelegate::CreateStatic(CustomizeInputNodeLambda, InputAxisName);
-		ActionListOut.Add(NodeSpawner);
+		ActionRegistrar.AddBlueprintAction(NodeSpawner);
 	}
 }
 
diff --git a/Engine/Source/Editor/BlueprintGraph/Private/K2Node_GetInputVectorAxisValue.cpp b/Engine/Source/Editor/BlueprintGraph/Private/K2Node_GetInputVectorAxisValue.cpp
--- a/Engine/Source/Editor/BlueprintGraph/Private/K2Node_GetInputVectorAxisValue.cpp
+++ b/Engine/Source/Editor/BlueprintGraph/Private/K2Node_GetInputVectorAxisValue.cpp
@@ -5,6 +5,7 @@
 #include "CompilerResultsLog.h"
 #include "BlueprintNodeSpawner.h"
 #include "Engine/InputVectorAxisDelegateBinding.h"
+#include "BlueprintActionDatabaseRegistrar.h"
 
 UK2Node_GetInputVectorAxisValue::UK2Node_GetInputVectorAxisValue(const class FPostConstructInitializeProperties& PCIP)
 	: Super(PCIP)
@@ -61,7 +62,7 @@ void UK2Node_GetInputVectorAxisValue::RegisterDynamicBinding(UDynamicBlueprintBi
 	InputVectorAxisBindingObject->InputAxisKeyDelegateBindings.Add(Binding);
 }
 
-void UK2Node_GetInputVectorAxisValue::GetMenuActions(TArray<UBlueprintNodeSpawner*>& ActionListOut) const
+void UK2Node_GetInputVectorAxisValue::GetMenuActions(FBlueprintActionDatabaseRegistrar& ActionRegistrar) const
 {
 	TArray<FKey> AllKeys;
 	EKeys::GetAllKeys(AllKeys);
@@ -83,6 +84,6 @@ void UK2Node_GetInputVectorAxisValue::GetMenuActions(TArray<UBlueprintNodeSpawne
 		};
 
 		NodeSpawner->CustomizeNodeDelegate = UBlueprintNodeSpawner::FCustomizeNodeDelegate::CreateStatic(CustomizeInputNodeLambda, Key);
-		ActionListOut.Add(NodeSpawner);
+		ActionRegistrar.AddBlueprintAction(NodeSpawner);
 	}
 }
diff --git a/Engine/Source/Editor/BlueprintGraph/Private/K2Node_GetNumEnumEntries.cpp b/Engine/Source/Editor/BlueprintGraph/Private/K2Node_GetNumEnumEntries.cpp
--- a/Engine/Source/Editor/BlueprintGraph/Private/K2Node_GetNumEnumEntries.cpp
+++ b/Engine/Source/Editor/BlueprintGraph/Private/K2Node_GetNumEnumEntries.cpp
@@ -7,6 +7,7 @@
 #include "K2Node_GetNumEnumEntries.h"
 #include "BlueprintNodeSpawner.h"
 #include "EditorCategoryUtils.h"
+#include "BlueprintActionDatabaseRegistrar.h"
 
 UK2Node_GetNumEnumEntries::UK2Node_GetNumEnumEntries(const class FPostConstructInitializeProperties& PCIP)
 	: Super(PCIP)
@@ -82,19 +83,12 @@ void UK2Node_GetNumEnumEntries::ExpandNode(class FKismetCompilerContext& Compile
 	}
 }
 
-void UK2Node_GetNumEnumEntries::GetMenuActions(TArray<UBlueprintNodeSpawner*>& ActionListOut) const
+void UK2Node_GetNumEnumEntries::GetMenuActions(FBlueprintActionDatabaseRegistrar& ActionRegistrar) const
 {
 	for (TObjectIterator<UEnum> EnumIt; EnumIt; ++EnumIt)
 	{
 		UEnum const* Enum = (*EnumIt);
-		// we only want to add global "standalone" enums here; those belonging to a 
-		// certain class should instead be associated with that class (so when 
-		// the class is modified we can easily handle any enums that were changed).
-		//
-		// @TODO: don't love how this code is essentially duplicated in BlueprintActionDatabase.cpp, for class enums
-		bool bIsStandaloneEnum = Enum->GetOuter()->IsA(UPackage::StaticClass());
-
-		if (!bIsStandaloneEnum || !UEdGraphSchema_K2::IsAllowableBlueprintVariableType(Enum))
+		if (!UEdGraphSchema_K2::IsAllowableBlueprintVariableType(Enum))
 		{
 			continue;
 		}
@@ -110,7 +104,10 @@ void UK2Node_GetNumEnumEntries::GetMenuActions(TArray<UBlueprintNodeSpawner*>& A
 
 		UBlueprintNodeSpawner* NodeSpawner = UBlueprintNodeSpawner::Create(GetClass());
 		check(NodeSpawner != nullptr);
-		ActionListOut.Add(NodeSpawner);
+		// this enum could belong to a class, or is a user defined enum (asset), 
+		// that's why we want to make sure to register it along with the action 
+		// (so the action can be refreshed when the class/asset is).
+		ActionRegistrar.AddBlueprintAction(Enum, NodeSpawner);
 
 		TWeakObjectPtr<UEnum> EnumPtr = Enum;
 		NodeSpawner->CustomizeNodeDelegate = UBlueprintNodeSpawner::FCustomizeNodeDelegate::CreateStatic(CustomizeEnumNodeLambda, EnumPtr);
diff --git a/Engine/Source/Editor/BlueprintGraph/Private/K2Node_IfThenElse.cpp b/Engine/Source/Editor/BlueprintGraph/Private/K2Node_IfThenElse.cpp
--- a/Engine/Source/Editor/BlueprintGraph/Private/K2Node_IfThenElse.cpp
+++ b/Engine/Source/Editor/BlueprintGraph/Private/K2Node_IfThenElse.cpp
@@ -5,6 +5,7 @@
 #include "KismetCompiler.h"
 #include "BlueprintNodeSpawner.h"
 #include "EditorCategoryUtils.h"
+#include "BlueprintActionDatabaseRegistrar.h"
 
 #define LOCTEXT_NAMESPACE "K2Node"
 
@@ -145,12 +146,12 @@ FNodeHandlingFunctor* UK2Node_IfThenElse::CreateNodeHandler(FKismetCompilerConte
 	return new FKCHandler_Branch(CompilerContext);
 }
 
-void UK2Node_IfThenElse::GetMenuActions(TArray<UBlueprintNodeSpawner*>& ActionListOut) const
+void UK2Node_IfThenElse::GetMenuActions(FBlueprintActionDatabaseRegistrar& ActionRegistrar) const
 {
 	UBlueprintNodeSpawner* NodeSpawner = UBlueprintNodeSpawner::Create(GetClass());
 	check(NodeSpawner != nullptr);
 
-	ActionListOut.Add(NodeSpawner);
+	ActionRegistrar.AddBlueprintAction(NodeSpawner);
 }
 
 FText UK2Node_IfThenElse::GetMenuCategory() const
diff --git a/Engine/Source/Editor/BlueprintGraph/Private/K2Node_InputAction.cpp b/Engine/Source/Editor/BlueprintGraph/Private/K2Node_InputAction.cpp
--- a/Engine/Source/Editor/BlueprintGraph/Private/K2Node_InputAction.cpp
+++ b/Engine/Source/Editor/BlueprintGraph/Private/K2Node_InputAction.cpp
@@ -8,6 +8,7 @@
 #include "EditorCategoryUtils.h"
 #include "BlueprintEditorUtils.h"
 #include "EdGraphSchema_K2.h"
+#include "BlueprintActionDatabaseRegistrar.h"
 
 #define LOCTEXT_NAMESPACE "K2Node_InputAction"
 
@@ -133,7 +134,7 @@ void UK2Node_InputAction::ExpandNode(FKismetCompilerContext& CompilerContext, UE
 	}
 }
 
-void UK2Node_InputAction::GetMenuActions(TArray<UBlueprintNodeSpawner*>& ActionListOut) const
+void UK2Node_InputAction::GetMenuActions(FBlueprintActionDatabaseRegistrar& ActionRegistrar) const
 {
 	TArray<FName> ActionNames;
 	GetDefault<UInputSettings>()->GetActionNames(ActionNames);
@@ -150,7 +151,7 @@ void UK2Node_InputAction::GetMenuActions(TArray<UBlueprintNodeSpawner*>& ActionL
 		};
 
 		NodeSpawner->CustomizeNodeDelegate = UBlueprintNodeSpawner::FCustomizeNodeDelegate::CreateStatic(CustomizeInputNodeLambda, InputActionName);
-		ActionListOut.Add(NodeSpawner);
+		ActionRegistrar.AddBlueprintAction(NodeSpawner);
 	}
 }
 
diff --git a/Engine/Source/Editor/BlueprintGraph/Private/K2Node_InputAxisEvent.cpp b/Engine/Source/Editor/BlueprintGraph/Private/K2Node_InputAxisEvent.cpp
--- a/Engine/Source/Editor/BlueprintGraph/Private/K2Node_InputAxisEvent.cpp
+++ b/Engine/Source/Editor/BlueprintGraph/Private/K2Node_InputAxisEvent.cpp
@@ -5,6 +5,7 @@
 #include "BlueprintNodeSpawner.h"
 #include "EditorCategoryUtils.h"
 #include "Engine/InputAxisDelegateBinding.h"
+#include "BlueprintActionDatabaseRegistrar.h"
 
 #define LOCTEXT_NAMESPACE "K2Node_InputAxisEvent"
 
@@ -101,7 +102,7 @@ bool UK2Node_InputAxisEvent::IsCompatibleWithGraph(const UEdGraph* TargetGraph)
 	return bIsCompatible && Super::IsCompatibleWithGraph(TargetGraph);
 }
 
-void UK2Node_InputAxisEvent::GetMenuActions(TArray<UBlueprintNodeSpawner*>& ActionListOut) const
+void UK2Node_InputAxisEvent::GetMenuActions(FBlueprintActionDatabaseRegistrar& ActionRegistrar) const
 {
 	TArray<FName> AxisNames;
 	GetDefault<UInputSettings>()->GetAxisNames(AxisNames);
@@ -118,7 +119,7 @@ void UK2Node_InputAxisEvent::GetMenuActions(TArray<UBlueprintNodeSpawner*>& Acti
 		};
 
 		NodeSpawner->CustomizeNodeDelegate = UBlueprintNodeSpawner::FCustomizeNodeDelegate::CreateStatic(CustomizeInputNodeLambda, InputAxisName);
-		ActionListOut.Add(NodeSpawner);
+		ActionRegistrar.AddBlueprintAction(NodeSpawner);
 	}
 }
 
diff --git a/Engine/Source/Editor/BlueprintGraph/Private/K2Node_InputAxisKeyEvent.cpp b/Engine/Source/Editor/BlueprintGraph/Private/K2Node_InputAxisKeyEvent.cpp
--- a/Engine/Source/Editor/BlueprintGraph/Private/K2Node_InputAxisKeyEvent.cpp
+++ b/Engine/Source/Editor/BlueprintGraph/Private/K2Node_InputAxisKeyEvent.cpp
@@ -8,6 +8,7 @@
 #include "Engine/InputAxisKeyDelegateBinding.h"
 #include "BlueprintEditorUtils.h"
 #include "EdGraphSchema_K2.h"
+#include "BlueprintActionDatabaseRegistrar.h"
 
 #define LOCTEXT_NAMESPACE "UK2Node_InputAxisKeyEvent"
 
@@ -110,7 +111,7 @@ bool UK2Node_InputAxisKeyEvent::IsCompatibleWithGraph(const UEdGraph* TargetGrap
 	return bIsCompatible && Super::IsCompatibleWithGraph(TargetGraph);
 }
 
-void UK2Node_InputAxisKeyEvent::GetMenuActions(TArray<UBlueprintNodeSpawner*>& ActionListOut) const
+void UK2Node_InputAxisKeyEvent::GetMenuActions(FBlueprintActionDatabaseRegistrar& ActionRegistrar) const
 {
 	TArray<FKey> AllKeys;
 	EKeys::GetAllKeys(AllKeys);
@@ -132,7 +133,7 @@ void UK2Node_InputAxisKeyEvent::GetMenuActions(TArray<UBlueprintNodeSpawner*>& A
 		};
 
 		NodeSpawner->CustomizeNodeDelegate = UBlueprintNodeSpawner::FCustomizeNodeDelegate::CreateStatic(CustomizeInputNodeLambda, Key);
-		ActionListOut.Add(NodeSpawner);
+		ActionRegistrar.AddBlueprintAction(NodeSpawner);
 	}
 }
 
diff --git a/Engine/Source/Editor/BlueprintGraph/Private/K2Node_InputKey.cpp b/Engine/Source/Editor/BlueprintGraph/Private/K2Node_InputKey.cpp
--- a/Engine/Source/Editor/BlueprintGraph/Private/K2Node_InputKey.cpp
+++ b/Engine/Source/Editor/BlueprintGraph/Private/K2Node_InputKey.cpp
@@ -8,6 +8,7 @@
 #include "EditorCategoryUtils.h"
 #include "BlueprintEditorUtils.h"
 #include "EdGraphSchema_K2.h"
+#include "BlueprintActionDatabaseRegistrar.h"
 
 #define LOCTEXT_NAMESPACE "UK2Node_InputKey"
 
@@ -357,7 +358,7 @@ void UK2Node_InputKey::ExpandNode(FKismetCompilerContext& CompilerContext, UEdGr
 	}
 }
 
-void UK2Node_InputKey::GetMenuActions(TArray<UBlueprintNodeSpawner*>& ActionListOut) const
+void UK2Node_InputKey::GetMenuActions(FBlueprintActionDatabaseRegistrar& ActionRegistrar) const
 {
 	TArray<FKey> AllKeys;
 	EKeys::GetAllKeys(AllKeys);
@@ -380,7 +381,7 @@ void UK2Node_InputKey::GetMenuActions(TArray<UBlueprintNodeSpawner*>& ActionList
 		};
 
 		NodeSpawner->CustomizeNodeDelegate = UBlueprintNodeSpawner::FCustomizeNodeDelegate::CreateStatic(CustomizeInputNodeLambda, Key);
-		ActionListOut.Add(NodeSpawner);
+		ActionRegistrar.AddBlueprintAction(NodeSpawner);
 	}
 }
 
diff --git a/Engine/Source/Editor/BlueprintGraph/Private/K2Node_InputTouch.cpp b/Engine/Source/Editor/BlueprintGraph/Private/K2Node_InputTouch.cpp
--- a/Engine/Source/Editor/BlueprintGraph/Private/K2Node_InputTouch.cpp
+++ b/Engine/Source/Editor/BlueprintGraph/Private/K2Node_InputTouch.cpp
@@ -6,6 +6,7 @@
 #include "KismetCompiler.h"
 #include "BlueprintNodeSpawner.h"
 #include "EditorCategoryUtils.h"
+#include "BlueprintActionDatabaseRegistrar.h"
 
 UK2Node_InputTouch::UK2Node_InputTouch(const class FPostConstructInitializeProperties& PCIP)
 	: Super(PCIP)
@@ -67,12 +68,12 @@ FText UK2Node_InputTouch::GetNodeTitle(ENodeTitleType::Type TitleType) const
 	return Title;
 }
 
-void UK2Node_InputTouch::GetMenuActions(TArray<UBlueprintNodeSpawner*>& ActionListOut) const
+void UK2Node_InputTouch::GetMenuActions(FBlueprintActionDatabaseRegistrar& ActionRegistrar) const
 {
 	UBlueprintNodeSpawner* NodeSpawner = UBlueprintNodeSpawner::Create(GetClass());
 	check(NodeSpawner != nullptr);
 
-	ActionListOut.Add(NodeSpawner);
+	ActionRegistrar.AddBlueprintAction(NodeSpawner);
 }
 
 FText UK2Node_InputTouch::GetMenuCategory() const
diff --git a/Engine/Source/Editor/BlueprintGraph/Private/K2Node_InputVectorAxisEvent.cpp b/Engine/Source/Editor/BlueprintGraph/Private/K2Node_InputVectorAxisEvent.cpp
--- a/Engine/Source/Editor/BlueprintGraph/Private/K2Node_InputVectorAxisEvent.cpp
+++ b/Engine/Source/Editor/BlueprintGraph/Private/K2Node_InputVectorAxisEvent.cpp
@@ -5,6 +5,7 @@
 #include "CompilerResultsLog.h"
 #include "BlueprintNodeSpawner.h"
 #include "Engine/InputVectorAxisDelegateBinding.h"
+#include "BlueprintActionDatabaseRegistrar.h"
 
 UK2Node_InputVectorAxisEvent::UK2Node_InputVectorAxisEvent(const class FPostConstructInitializeProperties& PCIP)
 : Super(PCIP)
@@ -50,7 +51,7 @@ void UK2Node_InputVectorAxisEvent::RegisterDynamicBinding(UDynamicBlueprintBindi
 	InputVectorAxisBindingObject->InputAxisKeyDelegateBindings.Add(Binding);
 }
 
-void UK2Node_InputVectorAxisEvent::GetMenuActions(TArray<UBlueprintNodeSpawner*>& ActionListOut) const
+void UK2Node_InputVectorAxisEvent::GetMenuActions(FBlueprintActionDatabaseRegistrar& ActionRegistrar) const
 {
 	TArray<FKey> AllKeys;
 	EKeys::GetAllKeys(AllKeys);
@@ -72,6 +73,6 @@ void UK2Node_InputVectorAxisEvent::GetMenuActions(TArray<UBlueprintNodeSpawner*>
 		};
 		
 		NodeSpawner->CustomizeNodeDelegate = UBlueprintNodeSpawner::FCustomizeNodeDelegate::CreateStatic(CustomizeInputNodeLambda, Key);
-		ActionListOut.Add(NodeSpawner);
+		ActionRegistrar.AddBlueprintAction(NodeSpawner);
 	}
 }
diff --git a/Engine/Source/Editor/BlueprintGraph/Private/K2Node_Knot.cpp b/Engine/Source/Editor/BlueprintGraph/Private/K2Node_Knot.cpp
--- a/Engine/Source/Editor/BlueprintGraph/Private/K2Node_Knot.cpp
+++ b/Engine/Source/Editor/BlueprintGraph/Private/K2Node_Knot.cpp
@@ -4,6 +4,7 @@
 #include "K2ActionMenuBuilder.h"
 #include "Kismet2NameValidators.h"
 #include "BlueprintNodeSpawner.h"
+#include "BlueprintActionDatabaseRegistrar.h"
 
 #define LOCTEXT_NAMESPACE "K2Node_Knot"
 
@@ -145,12 +146,12 @@ void UK2Node_Knot::PostReconstructNode()
 	}
 }
 
-void UK2Node_Knot::GetMenuActions(TArray<UBlueprintNodeSpawner*>& ActionListOut) const
+void UK2Node_Knot::GetMenuActions(FBlueprintActionDatabaseRegistrar& ActionRegistrar) const
 {
 	UBlueprintNodeSpawner* NodeSpawner = UBlueprintNodeSpawner::Create(GetClass());
 	check(NodeSpawner != nullptr);
 
-	ActionListOut.Add(NodeSpawner);
+	ActionRegistrar.AddBlueprintAction(NodeSpawner);
 }
 
 bool UK2Node_Knot::ShouldOverridePinNames() const
diff --git a/Engine/Source/Editor/BlueprintGraph/Private/K2Node_MakeArray.cpp b/Engine/Source/Editor/BlueprintGraph/Private/K2Node_MakeArray.cpp
--- a/Engine/Source/Editor/BlueprintGraph/Private/K2Node_MakeArray.cpp
+++ b/Engine/Source/Editor/BlueprintGraph/Private/K2Node_MakeArray.cpp
@@ -9,6 +9,7 @@
 #include "KismetCompiler.h"
 #include "BlueprintNodeSpawner.h"
 #include "EditorCategoryUtils.h"
+#include "BlueprintActionDatabaseRegistrar.h"
 
 static const FString OutputPinName = FString(TEXT("Array"));
 
@@ -403,12 +404,12 @@ bool UK2Node_MakeArray::IsConnectionDisallowed(const UEdGraphPin* MyPin, const U
 	return false;
 }
 
-void UK2Node_MakeArray::GetMenuActions(TArray<UBlueprintNodeSpawner*>& ActionListOut) const
+void UK2Node_MakeArray::GetMenuActions(FBlueprintActionDatabaseRegistrar& ActionRegistrar) const
 {
 	UBlueprintNodeSpawner* NodeSpawner = UBlueprintNodeSpawner::Create(GetClass());
 	check(NodeSpawner != nullptr);
 
-	ActionListOut.Add(NodeSpawner);
+	ActionRegistrar.AddBlueprintAction(NodeSpawner);
 }
 
 FText UK2Node_MakeArray::GetMenuCategory() const
diff --git a/Engine/Source/Editor/BlueprintGraph/Private/K2Node_MakeStruct.cpp b/Engine/Source/Editor/BlueprintGraph/Private/K2Node_MakeStruct.cpp
--- a/Engine/Source/Editor/BlueprintGraph/Private/K2Node_MakeStruct.cpp
+++ b/Engine/Source/Editor/BlueprintGraph/Private/K2Node_MakeStruct.cpp
@@ -5,6 +5,7 @@
 #include "MakeStructHandler.h"
 #include "BlueprintNodeSpawner.h"
 #include "EditorCategoryUtils.h"
+#include "BlueprintActionDatabaseRegistrar.h"
 
 #define LOCTEXT_NAMESPACE "K2Node_MakeStruct"
 
@@ -247,17 +248,12 @@ UK2Node::ERedirectType UK2Node_MakeStruct::DoPinsMatchForReconstruction(const UE
 	return Result;
 }
 
-void UK2Node_MakeStruct::GetMenuActions(TArray<UBlueprintNodeSpawner*>& ActionListOut) const
+void UK2Node_MakeStruct::GetMenuActions(FBlueprintActionDatabaseRegistrar& ActionRegistrar) const
 {
 	for (TObjectIterator<UScriptStruct> StructIt; StructIt; ++StructIt)
 	{
 		UScriptStruct const* Struct = (*StructIt);
-		// we only want to add autonomous structs here; those belonging to a 
-		// certain class should instead be associated with that class (so when 
-		// the class is modified we can easily handle any structs that were changed).
-		bool bIsStandaloneStruct = Struct->GetOuter()->IsA(UPackage::StaticClass());
-
-		if (!bIsStandaloneStruct || !UEdGraphSchema_K2::IsAllowableBlueprintVariableType(Struct) || !CanBeMade(Struct))
+		if (!UEdGraphSchema_K2::IsAllowableBlueprintVariableType(Struct) || !CanBeMade(Struct))
 		{
 			continue;
 		}
@@ -273,7 +269,10 @@ void UK2Node_MakeStruct::GetMenuActions(TArray<UBlueprintNodeSpawner*>& ActionLi
 
 		UBlueprintNodeSpawner* NodeSpawner = UBlueprintNodeSpawner::Create(GetClass());
 		check(NodeSpawner != nullptr);
-		ActionListOut.Add(NodeSpawner);
+		// this struct could belong to a class, or is a user defined struct 
+		// (asset), that's why we want to make sure to register it along with 
+		// the action (so the action knows to refresh when the class/asset is).
+		ActionRegistrar.AddBlueprintAction(Struct, NodeSpawner);
 
 		TWeakObjectPtr<UScriptStruct> StructPtr = Struct;
 		NodeSpawner->CustomizeNodeDelegate = UBlueprintNodeSpawner::FCustomizeNodeDelegate::CreateStatic(CustomizeMakeNodeLambda, StructPtr);
diff --git a/Engine/Source/Editor/BlueprintGraph/Private/K2Node_MathExpression.cpp b/Engine/Source/Editor/BlueprintGraph/Private/K2Node_MathExpression.cpp
--- a/Engine/Source/Editor/BlueprintGraph/Private/K2Node_MathExpression.cpp
+++ b/Engine/Source/Editor/BlueprintGraph/Private/K2Node_MathExpression.cpp
@@ -9,6 +9,7 @@
 #include "BasicTokenParser.h"
 #include "UnrealMathUtility.h"
 #include "BlueprintEditorUtils.h"
+#include "BlueprintActionDatabaseRegistrar.h"
 
 #define LOCTEXT_NAMESPACE "K2Node"
 
@@ -2301,12 +2302,12 @@ void UK2Node_MathExpression::GetMenuEntries(FGraphContextMenuBuilder& ContextMen
 }
 
 //------------------------------------------------------------------------------
-void UK2Node_MathExpression::GetMenuActions(TArray<UBlueprintNodeSpawner*>& ActionListOut) const
+void UK2Node_MathExpression::GetMenuActions(FBlueprintActionDatabaseRegistrar& ActionRegistrar) const
 {
 	UBlueprintNodeSpawner* NodeSpawner = UBlueprintNodeSpawner::Create(GetClass());
 	check(NodeSpawner != nullptr);
 	
-	ActionListOut.Add(NodeSpawner);
+	ActionRegistrar.AddBlueprintAction(NodeSpawner);
 }
 
 //------------------------------------------------------------------------------
diff --git a/Engine/Source/Editor/BlueprintGraph/Private/K2Node_MultiGate.cpp b/Engine/Source/Editor/BlueprintGraph/Private/K2Node_MultiGate.cpp
--- a/Engine/Source/Editor/BlueprintGraph/Private/K2Node_MultiGate.cpp
+++ b/Engine/Source/Editor/BlueprintGraph/Private/K2Node_MultiGate.cpp
@@ -7,6 +7,7 @@
 #include "Kismet/KismetSystemLibrary.h"
 #include "K2Node_MultiGate.h"
 #include "BlueprintNodeSpawner.h"
+#include "BlueprintActionDatabaseRegistrar.h"
 
 #define LOCTEXT_NAMESPACE "K2Node_MultiGate"
 
@@ -656,12 +657,12 @@ void UK2Node_MultiGate::ExpandNode(class FKismetCompilerContext& CompilerContext
 	}
 }
 
-void UK2Node_MultiGate::GetMenuActions(TArray<UBlueprintNodeSpawner*>& ActionListOut) const
+void UK2Node_MultiGate::GetMenuActions(FBlueprintActionDatabaseRegistrar& ActionRegistrar) const
 {
 	UBlueprintNodeSpawner* NodeSpawner = UBlueprintNodeSpawner::Create(GetClass());
 	check(NodeSpawner != nullptr);
 
-	ActionListOut.Add(NodeSpawner);
+	ActionRegistrar.AddBlueprintAction(NodeSpawner);
 }
 
 #undef LOCTEXT_NAMESPACE
diff --git a/Engine/Source/Editor/BlueprintGraph/Private/K2Node_Select.cpp b/Engine/Source/Editor/BlueprintGraph/Private/K2Node_Select.cpp
--- a/Engine/Source/Editor/BlueprintGraph/Private/K2Node_Select.cpp
+++ b/Engine/Source/Editor/BlueprintGraph/Private/K2Node_Select.cpp
@@ -6,6 +6,7 @@
 #include "KismetCompiler.h"
 #include "BlueprintNodeSpawner.h"
 #include "EditorCategoryUtils.h"
+#include "BlueprintActionDatabaseRegistrar.h"
 
 #define LOCTEXT_NAMESPACE "K2Node_Select"
 
@@ -774,12 +775,12 @@ FNodeHandlingFunctor* UK2Node_Select::CreateNodeHandler(FKismetCompilerContext&
 	return new FKCHandler_Select(CompilerContext);
 }
 
-void UK2Node_Select::GetMenuActions(TArray<UBlueprintNodeSpawner*>& ActionListOut) const
+void UK2Node_Select::GetMenuActions(FBlueprintActionDatabaseRegistrar& ActionRegistrar) const
 {
 	UBlueprintNodeSpawner* NodeSpawner = UBlueprintNodeSpawner::Create(GetClass());
 	check(NodeSpawner != nullptr);
 
-	ActionListOut.Add(NodeSpawner);
+	ActionRegistrar.AddBlueprintAction(NodeSpawner);
 }
 
 FText UK2Node_Select::GetMenuCategory() const
diff --git a/Engine/Source/Editor/BlueprintGraph/Private/K2Node_SelectEnum.cpp b/Engine/Source/Editor/BlueprintGraph/Private/K2Node_SelectEnum.cpp
--- a/Engine/Source/Editor/BlueprintGraph/Private/K2Node_SelectEnum.cpp
+++ b/Engine/Source/Editor/BlueprintGraph/Private/K2Node_SelectEnum.cpp
@@ -4,6 +4,7 @@
 #include "BlueprintGraphPrivatePCH.h"
 #include "K2Node_SelectEnum.h"
 #include "BlueprintNodeSpawner.h"
+#include "BlueprintActionDatabaseRegistrar.h"
 
 UDEPRECATED_K2Node_SelectEnum::UDEPRECATED_K2Node_SelectEnum(const class FPostConstructInitializeProperties& PCIP)
 	: Super(PCIP)
@@ -62,7 +63,7 @@ void UDEPRECATED_K2Node_SelectEnum::NotifyPinConnectionListChanged(UEdGraphPin*
 	// don't do anything
 }
 
-void UDEPRECATED_K2Node_SelectEnum::GetMenuActions(TArray<UBlueprintNodeSpawner*>& ActionListOut) const
+void UDEPRECATED_K2Node_SelectEnum::GetMenuActions(FBlueprintActionDatabaseRegistrar& ActionRegistrar) const
 {
 	for (TObjectIterator<UEnum> EnumIt; EnumIt; ++EnumIt)
 	{
@@ -78,7 +79,7 @@ void UDEPRECATED_K2Node_SelectEnum::GetMenuActions(TArray<UBlueprintNodeSpawner*
 		
 		UBlueprintNodeSpawner* NodeSpawner = UBlueprintNodeSpawner::Create(GetClass());
 		check(NodeSpawner != nullptr);
-		ActionListOut.Add(NodeSpawner);
+		ActionRegistrar.AddBlueprintAction(NodeSpawner);
 		
 		TWeakObjectPtr<UEnum> EnumPtr = Enum;
 		NodeSpawner->CustomizeNodeDelegate = UBlueprintNodeSpawner::FCustomizeNodeDelegate::CreateStatic(CustomizeEnumNodeLambda, EnumPtr);
diff --git a/Engine/Source/Editor/BlueprintGraph/Private/K2Node_Self.cpp b/Engine/Source/Editor/BlueprintGraph/Private/K2Node_Self.cpp
--- a/Engine/Source/Editor/BlueprintGraph/Private/K2Node_Self.cpp
+++ b/Engine/Source/Editor/BlueprintGraph/Private/K2Node_Self.cpp
@@ -5,6 +5,7 @@
 #include "KismetCompiler.h"
 #include "BlueprintNodeSpawner.h"
 #include "EditorCategoryUtils.h"
+#include "BlueprintActionDatabaseRegistrar.h"
 
 #define LOCTEXT_NAMESPACE "K2Node_Self"
 
@@ -81,12 +82,12 @@ void UK2Node_Self::ValidateNodeDuringCompilation(class FCompilerResultsLog& Mess
 	}
 }
 
-void UK2Node_Self::GetMenuActions(TArray<UBlueprintNodeSpawner*>& ActionListOut) const
+void UK2Node_Self::GetMenuActions(FBlueprintActionDatabaseRegistrar& ActionRegistrar) const
 {
 	UBlueprintNodeSpawner* NodeSpawner = UBlueprintNodeSpawner::Create(GetClass());
 	check(NodeSpawner != nullptr);
 
-	ActionListOut.Add(NodeSpawner);
+	ActionRegistrar.AddBlueprintAction(NodeSpawner);
 }
 
 FText UK2Node_Self::GetMenuCategory() const
diff --git a/Engine/Source/Editor/BlueprintGraph/Private/K2Node_SpawnActorFromClass.cpp b/Engine/Source/Editor/BlueprintGraph/Private/K2Node_SpawnActorFromClass.cpp
--- a/Engine/Source/Editor/BlueprintGraph/Private/K2Node_SpawnActorFromClass.cpp
+++ b/Engine/Source/Editor/BlueprintGraph/Private/K2Node_SpawnActorFromClass.cpp
@@ -5,6 +5,7 @@
 #include "Runtime/Engine/Classes/Kismet/GameplayStatics.h"
 #include "BlueprintNodeSpawner.h"
 #include "EditorCategoryUtils.h"
+#include "BlueprintActionDatabaseRegistrar.h"
 
 struct FK2Node_SpawnActorFromClassHelper
 {
@@ -314,12 +315,12 @@ void UK2Node_SpawnActorFromClass::GetNodeAttributes( TArray<TKeyValuePair<FStrin
 	OutNodeAttributes.Add( TKeyValuePair<FString, FString>( TEXT( "ActorClass" ), ClassToSpawnStr ));
 }
 
-void UK2Node_SpawnActorFromClass::GetMenuActions(TArray<UBlueprintNodeSpawner*>& ActionListOut) const
+void UK2Node_SpawnActorFromClass::GetMenuActions(FBlueprintActionDatabaseRegistrar& ActionRegistrar) const
 {
 	UBlueprintNodeSpawner* NodeSpawner = UBlueprintNodeSpawner::Create(GetClass());
 	check(NodeSpawner != nullptr);
 
-	ActionListOut.Add(NodeSpawner);
+	ActionRegistrar.AddBlueprintAction(NodeSpawner);
 }
 
 FText UK2Node_SpawnActorFromClass::GetMenuCategory() const
diff --git a/Engine/Source/Editor/BlueprintGraph/Private/K2Node_SwitchEnum.cpp b/Engine/Source/Editor/BlueprintGraph/Private/K2Node_SwitchEnum.cpp
--- a/Engine/Source/Editor/BlueprintGraph/Private/K2Node_SwitchEnum.cpp
+++ b/Engine/Source/Editor/BlueprintGraph/Private/K2Node_SwitchEnum.cpp
@@ -4,6 +4,8 @@
 #include "Kismet/KismetMathLibrary.h"
 #include "K2Node_SwitchEnum.h"
 #include "BlueprintNodeSpawner.h"
+#include "BlueprintActionDatabaseRegistrar.h"
+#include "BlueprintActionDatabaseRegistrar.h"
 
 #define LOCTEXT_NAMESPACE "K2Node"
 
@@ -62,19 +64,12 @@ FString UK2Node_SwitchEnum::GetTooltip() const
 	return NSLOCTEXT("K2Node", "SwitchEnum_ToolTip", "Selects an output that matches the input value").ToString();
 }
 
-void UK2Node_SwitchEnum::GetMenuActions(TArray<UBlueprintNodeSpawner*>& ActionListOut) const
+void UK2Node_SwitchEnum::GetMenuActions(FBlueprintActionDatabaseRegistrar& ActionRegistrar) const
 {
 	for (TObjectIterator<UEnum> EnumIt; EnumIt; ++EnumIt)
 	{
 		UEnum const* Enum = (*EnumIt);
-		// we only want to add global "standalone" enums here; those belonging to a 
-		// certain class should instead be associated with that class (so when 
-		// the class is modified we can easily handle any enums that were changed).
-		//
-		// @TODO: don't love how this code is essentially duplicated in BlueprintActionDatabase.cpp, for class enums
-		bool bIsStandaloneEnum = Enum->GetOuter()->IsA(UPackage::StaticClass());
-
-		if (!bIsStandaloneEnum || !UEdGraphSchema_K2::IsAllowableBlueprintVariableType(Enum))
+		if (!UEdGraphSchema_K2::IsAllowableBlueprintVariableType(Enum))
 		{
 			continue;
 		}
@@ -90,7 +85,11 @@ void UK2Node_SwitchEnum::GetMenuActions(TArray<UBlueprintNodeSpawner*>& ActionLi
 
 		UBlueprintNodeSpawner* NodeSpawner = UBlueprintNodeSpawner::Create(GetClass());
 		check(NodeSpawner != nullptr);
-		ActionListOut.Add(NodeSpawner);
+
+		// this enum could belong to a class, or is a user defined enum (asset), 
+		// that's why we want to make sure to register it along with the action 
+		// (so the action can be refreshed when the class/asset is).
+		ActionRegistrar.AddBlueprintAction(Enum, NodeSpawner);
 
 		TWeakObjectPtr<UEnum> EnumPtr = Enum;
 		NodeSpawner->CustomizeNodeDelegate = UBlueprintNodeSpawner::FCustomizeNodeDelegate::CreateStatic(CustomizeEnumNodeLambda, EnumPtr);
diff --git a/Engine/Source/Editor/BlueprintGraph/Private/K2Node_SwitchInteger.cpp b/Engine/Source/Editor/BlueprintGraph/Private/K2Node_SwitchInteger.cpp
--- a/Engine/Source/Editor/BlueprintGraph/Private/K2Node_SwitchInteger.cpp
+++ b/Engine/Source/Editor/BlueprintGraph/Private/K2Node_SwitchInteger.cpp
@@ -3,6 +3,7 @@
 #include "BlueprintGraphPrivatePCH.h"
 #include "../../../Runtime/Engine/Classes/Kismet/KismetMathLibrary.h"
 #include "BlueprintNodeSpawner.h"
+#include "BlueprintActionDatabaseRegistrar.h"
 
 UK2Node_SwitchInteger::UK2Node_SwitchInteger(const class FPostConstructInitializeProperties& PCIP)
 	: Super(PCIP)
@@ -23,12 +24,12 @@ FString UK2Node_SwitchInteger::GetTooltip() const
 	return NSLOCTEXT("K2Node", "SwitchInt_ToolTip", "Selects an output that matches the input value").ToString();
 }
 
-void UK2Node_SwitchInteger::GetMenuActions(TArray<UBlueprintNodeSpawner*>& ActionListOut) const
+void UK2Node_SwitchInteger::GetMenuActions(FBlueprintActionDatabaseRegistrar& ActionRegistrar) const
 {
 	UBlueprintNodeSpawner* NodeSpawner = UBlueprintNodeSpawner::Create(GetClass());
 	check(NodeSpawner != nullptr);
 
-	ActionListOut.Add(NodeSpawner);
+	ActionRegistrar.AddBlueprintAction(NodeSpawner);
 }
 
 void UK2Node_SwitchInteger::CreateSelectionPin()
diff --git a/Engine/Source/Editor/BlueprintGraph/Private/K2Node_SwitchName.cpp b/Engine/Source/Editor/BlueprintGraph/Private/K2Node_SwitchName.cpp
--- a/Engine/Source/Editor/BlueprintGraph/Private/K2Node_SwitchName.cpp
+++ b/Engine/Source/Editor/BlueprintGraph/Private/K2Node_SwitchName.cpp
@@ -3,6 +3,7 @@
 #include "BlueprintGraphPrivatePCH.h"
 #include "../../../Runtime/Engine/Classes/Kismet/KismetMathLibrary.h"
 #include "BlueprintNodeSpawner.h"
+#include "BlueprintActionDatabaseRegistrar.h"
 
 UK2Node_SwitchName::UK2Node_SwitchName(const class FPostConstructInitializeProperties& PCIP)
 	: Super(PCIP)
@@ -38,12 +39,12 @@ FString UK2Node_SwitchName::GetTooltip() const
 	return NSLOCTEXT("K2Node", "SwitchName_ToolTip", "Selects an output that matches the input value").ToString();
 }
 
-void UK2Node_SwitchName::GetMenuActions(TArray<UBlueprintNodeSpawner*>& ActionListOut) const
+void UK2Node_SwitchName::GetMenuActions(FBlueprintActionDatabaseRegistrar& ActionRegistrar) const
 {
 	UBlueprintNodeSpawner* NodeSpawner = UBlueprintNodeSpawner::Create(GetClass());
 	check(NodeSpawner != nullptr);
 
-	ActionListOut.Add(NodeSpawner);
+	ActionRegistrar.AddBlueprintAction(NodeSpawner);
 }
 
 void UK2Node_SwitchName::CreateSelectionPin()
diff --git a/Engine/Source/Editor/BlueprintGraph/Private/K2Node_SwitchString.cpp b/Engine/Source/Editor/BlueprintGraph/Private/K2Node_SwitchString.cpp
--- a/Engine/Source/Editor/BlueprintGraph/Private/K2Node_SwitchString.cpp
+++ b/Engine/Source/Editor/BlueprintGraph/Private/K2Node_SwitchString.cpp
@@ -4,6 +4,7 @@
 #include "Kismet/KismetMathLibrary.h"
 #include "K2Node_SwitchString.h"
 #include "BlueprintNodeSpawner.h"
+#include "BlueprintActionDatabaseRegistrar.h"
 
 UK2Node_SwitchString::UK2Node_SwitchString(const class FPostConstructInitializeProperties& PCIP)
 	: Super(PCIP)
@@ -50,12 +51,12 @@ FString UK2Node_SwitchString::GetTooltip() const
 	return NSLOCTEXT("K2Node", "SwitchString_ToolTip", "Selects an output that matches the input value").ToString();
 }
 
-void UK2Node_SwitchString::GetMenuActions(TArray<UBlueprintNodeSpawner*>& ActionListOut) const
+void UK2Node_SwitchString::GetMenuActions(FBlueprintActionDatabaseRegistrar& ActionRegistrar) const
 {
 	UBlueprintNodeSpawner* NodeSpawner = UBlueprintNodeSpawner::Create(GetClass());
 	check(NodeSpawner != nullptr);
 
-	ActionListOut.Add(NodeSpawner);
+	ActionRegistrar.AddBlueprintAction(NodeSpawner);
 }
 
 void UK2Node_SwitchString::CreateSelectionPin()
diff --git a/Engine/Source/Editor/BlueprintGraph/Private/K2Node_TemporaryVariable.cpp b/Engine/Source/Editor/BlueprintGraph/Private/K2Node_TemporaryVariable.cpp
--- a/Engine/Source/Editor/BlueprintGraph/Private/K2Node_TemporaryVariable.cpp
+++ b/Engine/Source/Editor/BlueprintGraph/Private/K2Node_TemporaryVariable.cpp
@@ -5,6 +5,7 @@
 #include "KismetCompiler.h"
 #include "BlueprintNodeSpawner.h"
 #include "EditorCategoryUtils.h"
+#include "BlueprintActionDatabaseRegistrar.h"
 
 class FKCHandler_TemporaryVariable : public FNodeHandlingFunctor
 {
@@ -116,7 +117,7 @@ FNodeHandlingFunctor* UK2Node_TemporaryVariable::CreateNodeHandler(FKismetCompil
 	return new FKCHandler_TemporaryVariable(CompilerContext);
 }
 
-void UK2Node_TemporaryVariable::GetMenuActions(TArray<UBlueprintNodeSpawner*>& ActionListOut) const
+void UK2Node_TemporaryVariable::GetMenuActions(FBlueprintActionDatabaseRegistrar& ActionRegistrar) const
 {
 	auto MakeTempVarNodeSpawner = [](FEdGraphPinType const& VarType, bool bIsPersistent) -> UBlueprintNodeSpawner*
 	{
@@ -136,36 +137,36 @@ void UK2Node_TemporaryVariable::GetMenuActions(TArray<UBlueprintNodeSpawner*>& A
 
 	UEdGraphSchema_K2 const* K2Schema = GetDefault<UEdGraphSchema_K2>();
 
-	ActionListOut.Add(MakeTempVarNodeSpawner(FEdGraphPinType(K2Schema->PC_Int,     TEXT(""), nullptr, /*bIsArray =*/false, /*bIsReference =*/false), /*bIsPersistent =*/false));
-	ActionListOut.Add(MakeTempVarNodeSpawner(FEdGraphPinType(K2Schema->PC_Int,     TEXT(""), nullptr, /*bIsArray =*/ true, /*bIsReference =*/false), /*bIsPersistent =*/false));
-	ActionListOut.Add(MakeTempVarNodeSpawner(FEdGraphPinType(K2Schema->PC_Float,   TEXT(""), nullptr, /*bIsArray =*/false, /*bIsReference =*/false), /*bIsPersistent =*/false));
-	ActionListOut.Add(MakeTempVarNodeSpawner(FEdGraphPinType(K2Schema->PC_Float,   TEXT(""), nullptr, /*bIsArray =*/ true, /*bIsReference =*/false), /*bIsPersistent =*/false));
-	ActionListOut.Add(MakeTempVarNodeSpawner(FEdGraphPinType(K2Schema->PC_Boolean, TEXT(""), nullptr, /*bIsArray =*/false, /*bIsReference =*/false), /*bIsPersistent =*/false));
-	ActionListOut.Add(MakeTempVarNodeSpawner(FEdGraphPinType(K2Schema->PC_Boolean, TEXT(""), nullptr, /*bIsArray =*/ true, /*bIsReference =*/false), /*bIsPersistent =*/false));
-	ActionListOut.Add(MakeTempVarNodeSpawner(FEdGraphPinType(K2Schema->PC_String,  TEXT(""), nullptr, /*bIsArray =*/false, /*bIsReference =*/false), /*bIsPersistent =*/false));
-	ActionListOut.Add(MakeTempVarNodeSpawner(FEdGraphPinType(K2Schema->PC_String,  TEXT(""), nullptr, /*bIsArray =*/ true, /*bIsReference =*/false), /*bIsPersistent =*/false));
-	ActionListOut.Add(MakeTempVarNodeSpawner(FEdGraphPinType(K2Schema->PC_Text,    TEXT(""), nullptr, /*bIsArray =*/false, /*bIsReference =*/false), /*bIsPersistent =*/false));
-	ActionListOut.Add(MakeTempVarNodeSpawner(FEdGraphPinType(K2Schema->PC_Text,    TEXT(""), nullptr, /*bIsArray =*/ true, /*bIsReference =*/false), /*bIsPersistent =*/false));
+	ActionRegistrar.AddBlueprintAction(MakeTempVarNodeSpawner(FEdGraphPinType(K2Schema->PC_Int,     TEXT(""), nullptr, /*bIsArray =*/false, /*bIsReference =*/false), /*bIsPersistent =*/false));
+	ActionRegistrar.AddBlueprintAction(MakeTempVarNodeSpawner(FEdGraphPinType(K2Schema->PC_Int,     TEXT(""), nullptr, /*bIsArray =*/ true, /*bIsReference =*/false), /*bIsPersistent =*/false));
+	ActionRegistrar.AddBlueprintAction(MakeTempVarNodeSpawner(FEdGraphPinType(K2Schema->PC_Float,   TEXT(""), nullptr, /*bIsArray =*/false, /*bIsReference =*/false), /*bIsPersistent =*/false));
+	ActionRegistrar.AddBlueprintAction(MakeTempVarNodeSpawner(FEdGraphPinType(K2Schema->PC_Float,   TEXT(""), nullptr, /*bIsArray =*/ true, /*bIsReference =*/false), /*bIsPersistent =*/false));
+	ActionRegistrar.AddBlueprintAction(MakeTempVarNodeSpawner(FEdGraphPinType(K2Schema->PC_Boolean, TEXT(""), nullptr, /*bIsArray =*/false, /*bIsReference =*/false), /*bIsPersistent =*/false));
+	ActionRegistrar.AddBlueprintAction(MakeTempVarNodeSpawner(FEdGraphPinType(K2Schema->PC_Boolean, TEXT(""), nullptr, /*bIsArray =*/ true, /*bIsReference =*/false), /*bIsPersistent =*/false));
+	ActionRegistrar.AddBlueprintAction(MakeTempVarNodeSpawner(FEdGraphPinType(K2Schema->PC_String,  TEXT(""), nullptr, /*bIsArray =*/false, /*bIsReference =*/false), /*bIsPersistent =*/false));
+	ActionRegistrar.AddBlueprintAction(MakeTempVarNodeSpawner(FEdGraphPinType(K2Schema->PC_String,  TEXT(""), nullptr, /*bIsArray =*/ true, /*bIsReference =*/false), /*bIsPersistent =*/false));
+	ActionRegistrar.AddBlueprintAction(MakeTempVarNodeSpawner(FEdGraphPinType(K2Schema->PC_Text,    TEXT(""), nullptr, /*bIsArray =*/false, /*bIsReference =*/false), /*bIsPersistent =*/false));
+	ActionRegistrar.AddBlueprintAction(MakeTempVarNodeSpawner(FEdGraphPinType(K2Schema->PC_Text,    TEXT(""), nullptr, /*bIsArray =*/ true, /*bIsReference =*/false), /*bIsPersistent =*/false));
 
 	UScriptStruct* VectorStruct  = FindObjectChecked<UScriptStruct>(UObject::StaticClass(), TEXT("Vector"));
-	ActionListOut.Add(MakeTempVarNodeSpawner(FEdGraphPinType(K2Schema->PC_Struct, TEXT("Vector"), VectorStruct, /*bIsArray =*/false, /*bIsReference =*/false), /*bIsPersistent =*/false));
-	ActionListOut.Add(MakeTempVarNodeSpawner(FEdGraphPinType(K2Schema->PC_Struct, TEXT("Vector"), VectorStruct, /*bIsArray =*/ true, /*bIsReference =*/false), /*bIsPersistent =*/false));
+	ActionRegistrar.AddBlueprintAction(MakeTempVarNodeSpawner(FEdGraphPinType(K2Schema->PC_Struct, TEXT("Vector"), VectorStruct, /*bIsArray =*/false, /*bIsReference =*/false), /*bIsPersistent =*/false));
+	ActionRegistrar.AddBlueprintAction(MakeTempVarNodeSpawner(FEdGraphPinType(K2Schema->PC_Struct, TEXT("Vector"), VectorStruct, /*bIsArray =*/ true, /*bIsReference =*/false), /*bIsPersistent =*/false));
 	
 	UScriptStruct* RotatorStruct = FindObjectChecked<UScriptStruct>(UObject::StaticClass(), TEXT("Rotator"));
-	ActionListOut.Add(MakeTempVarNodeSpawner(FEdGraphPinType(K2Schema->PC_Struct, TEXT("Rotator"), RotatorStruct, /*bIsArray =*/false, /*bIsReference =*/false), /*bIsPersistent =*/false));
-	ActionListOut.Add(MakeTempVarNodeSpawner(FEdGraphPinType(K2Schema->PC_Struct, TEXT("Rotator"), RotatorStruct, /*bIsArray =*/ true, /*bIsReference =*/false), /*bIsPersistent =*/false));
+	ActionRegistrar.AddBlueprintAction(MakeTempVarNodeSpawner(FEdGraphPinType(K2Schema->PC_Struct, TEXT("Rotator"), RotatorStruct, /*bIsArray =*/false, /*bIsReference =*/false), /*bIsPersistent =*/false));
+	ActionRegistrar.AddBlueprintAction(MakeTempVarNodeSpawner(FEdGraphPinType(K2Schema->PC_Struct, TEXT("Rotator"), RotatorStruct, /*bIsArray =*/ true, /*bIsReference =*/false), /*bIsPersistent =*/false));
 	
 	UScriptStruct* TransformStruct = FindObjectChecked<UScriptStruct>(UObject::StaticClass(), TEXT("Transform"));
-	ActionListOut.Add(MakeTempVarNodeSpawner(FEdGraphPinType(K2Schema->PC_Struct, TEXT("Transform"), TransformStruct, /*bIsArray =*/false, /*bIsReference =*/false), /*bIsPersistent =*/false));
-	ActionListOut.Add(MakeTempVarNodeSpawner(FEdGraphPinType(K2Schema->PC_Struct, TEXT("Transform"), TransformStruct, /*bIsArray =*/ true, /*bIsReference =*/false), /*bIsPersistent =*/false));
+	ActionRegistrar.AddBlueprintAction(MakeTempVarNodeSpawner(FEdGraphPinType(K2Schema->PC_Struct, TEXT("Transform"), TransformStruct, /*bIsArray =*/false, /*bIsReference =*/false), /*bIsPersistent =*/false));
+	ActionRegistrar.AddBlueprintAction(MakeTempVarNodeSpawner(FEdGraphPinType(K2Schema->PC_Struct, TEXT("Transform"), TransformStruct, /*bIsArray =*/ true, /*bIsReference =*/false), /*bIsPersistent =*/false));
 
 	UScriptStruct* BlendSampleStruct = FindObjectChecked<UScriptStruct>(ANY_PACKAGE, TEXT("BlendSampleData"));
-	ActionListOut.Add(MakeTempVarNodeSpawner(FEdGraphPinType(K2Schema->PC_Struct, TEXT("BlendSampleData"), BlendSampleStruct, /*bIsArray =*/false, /*bIsReference =*/false), /*bIsPersistent =*/false));
-	ActionListOut.Add(MakeTempVarNodeSpawner(FEdGraphPinType(K2Schema->PC_Struct, TEXT("BlendSampleData"), BlendSampleStruct, /*bIsArray =*/ true, /*bIsReference =*/false), /*bIsPersistent =*/false));
+	ActionRegistrar.AddBlueprintAction(MakeTempVarNodeSpawner(FEdGraphPinType(K2Schema->PC_Struct, TEXT("BlendSampleData"), BlendSampleStruct, /*bIsArray =*/false, /*bIsReference =*/false), /*bIsPersistent =*/false));
+	ActionRegistrar.AddBlueprintAction(MakeTempVarNodeSpawner(FEdGraphPinType(K2Schema->PC_Struct, TEXT("BlendSampleData"), BlendSampleStruct, /*bIsArray =*/ true, /*bIsReference =*/false), /*bIsPersistent =*/false));
 
 	// add persistent bool and int types (for macro graphs)
-	ActionListOut.Add(MakeTempVarNodeSpawner(FEdGraphPinType(K2Schema->PC_Int,     TEXT(""), nullptr, /*bIsArray =*/false, /*bIsReference =*/false), /*bIsPersistent =*/true));
-	ActionListOut.Add(MakeTempVarNodeSpawner(FEdGraphPinType(K2Schema->PC_Boolean, TEXT(""), nullptr, /*bIsArray =*/false, /*bIsReference =*/false), /*bIsPersistent =*/true));
+	ActionRegistrar.AddBlueprintAction(MakeTempVarNodeSpawner(FEdGraphPinType(K2Schema->PC_Int,     TEXT(""), nullptr, /*bIsArray =*/false, /*bIsReference =*/false), /*bIsPersistent =*/true));
+	ActionRegistrar.AddBlueprintAction(MakeTempVarNodeSpawner(FEdGraphPinType(K2Schema->PC_Boolean, TEXT(""), nullptr, /*bIsArray =*/false, /*bIsReference =*/false), /*bIsPersistent =*/true));
 }
 
 FText UK2Node_TemporaryVariable::GetMenuCategory() const
diff --git a/Engine/Source/Editor/BlueprintGraph/Private/K2Node_Timeline.cpp b/Engine/Source/Editor/BlueprintGraph/Private/K2Node_Timeline.cpp
--- a/Engine/Source/Editor/BlueprintGraph/Private/K2Node_Timeline.cpp
+++ b/Engine/Source/Editor/BlueprintGraph/Private/K2Node_Timeline.cpp
@@ -6,6 +6,7 @@
 #include "Kismet2NameValidators.h"
 #include "KismetCompiler.h"
 #include "BlueprintNodeSpawner.h"
+#include "BlueprintActionDatabaseRegistrar.h"
 
 #define LOCTEXT_NAMESPACE "K2Node_Timeline"
 
@@ -566,7 +567,7 @@ void UK2Node_Timeline::GetNodeAttributes( TArray<TKeyValuePair<FString, FString>
 	OutNodeAttributes.Add( TKeyValuePair<FString, FString>( TEXT( "Name" ), GetName() ));
 }
 
-void UK2Node_Timeline::GetMenuActions(TArray<UBlueprintNodeSpawner*>& ActionListOut) const
+void UK2Node_Timeline::GetMenuActions(FBlueprintActionDatabaseRegistrar& ActionRegistrar) const
 {
 	UBlueprintNodeSpawner* NodeSpawner = UBlueprintNodeSpawner::Create(GetClass());
 	check(NodeSpawner != nullptr);
@@ -589,7 +590,7 @@ void UK2Node_Timeline::GetMenuActions(TArray<UBlueprintNodeSpawner*>& ActionList
 	};
 
 	NodeSpawner->CustomizeNodeDelegate = UBlueprintNodeSpawner::FCustomizeNodeDelegate::CreateStatic(CustomizeTimelineNodeLambda);
-	ActionListOut.Add(NodeSpawner);
+	ActionRegistrar.AddBlueprintAction(NodeSpawner);
 }
 
 #undef LOCTEXT_NAMESPACE
diff --git a/Engine/Source/Editor/BlueprintGraph/Private/K2Node_VariableSetRef.cpp b/Engine/Source/Editor/BlueprintGraph/Private/K2Node_VariableSetRef.cpp
--- a/Engine/Source/Editor/BlueprintGraph/Private/K2Node_VariableSetRef.cpp
+++ b/Engine/Source/Editor/BlueprintGraph/Private/K2Node_VariableSetRef.cpp
@@ -7,6 +7,7 @@
 #include "K2Node_VariableSetRef.h"
 #include "BlueprintNodeSpawner.h"
 #include "EditorCategoryUtils.h"
+#include "BlueprintActionDatabaseRegistrar.h"
 
 static FString TargetVarPinName(TEXT("Target"));
 static FString VarValuePinName(TEXT("Value"));
@@ -204,12 +205,12 @@ FNodeHandlingFunctor* UK2Node_VariableSetRef::CreateNodeHandler(FKismetCompilerC
 	return new FKCHandler_VariableSetRef(CompilerContext);
 }
 
-void UK2Node_VariableSetRef::GetMenuActions(TArray<UBlueprintNodeSpawner*>& ActionListOut) const
+void UK2Node_VariableSetRef::GetMenuActions(FBlueprintActionDatabaseRegistrar& ActionRegistrar) const
 {
 	UBlueprintNodeSpawner* NodeSpawner = UBlueprintNodeSpawner::Create(GetClass());
 	check(NodeSpawner != nullptr);
 
-	ActionListOut.Add(NodeSpawner);
+	ActionRegistrar.AddBlueprintAction(NodeSpawner);
 }
 
 FText UK2Node_VariableSetRef::GetMenuCategory() const
diff --git a/Engine/Source/Editor/BlueprintGraph/Public/BlueprintActionDatabase.h b/Engine/Source/Editor/BlueprintGraph/Public/BlueprintActionDatabase.h
--- a/Engine/Source/Editor/BlueprintGraph/Public/BlueprintActionDatabase.h
+++ b/Engine/Source/Editor/BlueprintGraph/Public/BlueprintActionDatabase.h
@@ -10,6 +10,7 @@ class UBlueprint;
 class FBlueprintActionFilter;
 class UBlueprintNodeSpawner;
 class FReferenceCollector;
+class FBlueprintActionDatabaseRegistrar;
 
 /**
  * Serves as a container for all available blueprint actions (no matter the 
@@ -23,9 +24,10 @@ class FReferenceCollector;
 class BLUEPRINTGRAPH_API FBlueprintActionDatabase : public FGCObject, public FTickableEditorObject
 {
 public:
-	typedef TArray<UBlueprintNodeSpawner*> FActionList;
-	typedef TMap<UClass*, FActionList>     FClassActionMap;
-
+	typedef TMap<TWeakObjectPtr<UObject>, int32> FPrimingQueue;
+	typedef TArray<UBlueprintNodeSpawner*>       FActionList;
+	typedef TMap<UObject const*, FActionList>    FActionRegistry;
+	
 public:
 	/**
 	 * Getter to access the database singleton. Will populate the database first 
@@ -45,6 +47,20 @@ public:
 	virtual void AddReferencedObjects(FReferenceCollector& Collector) override;
 	// End FGCObject interface
 
+	/**
+	 * Will populate the database first if it hasn't been created yet, and then 
+	 * returns it in its entirety. 
+	 * 
+	 * Each node spawner is categorized by a class orasset. A spawner that 
+	 * corresponds to a specific class field (like a function, property, enum, 
+	 * etc.) will be listed under that field's class owner. Remaining spawners 
+	 * that can't be categorized this way will be registered by asset or node 
+	 * type.
+	 *
+	 * @return The entire action database, which maps specific objects to arrays of associated node-spawners.
+	 */
+	FActionRegistry const& GetAllActions();
+
 	/**
 	 * Populates the action database from scratch. Loops over every known class
 	 * and records a set of node-spawners associated with each.
@@ -53,57 +69,59 @@ public:
 
 	/**
 	 * Finds the database entry for the specified class and wipes it, 
-	 * repopulating it with a fresh set of associated node-spawners.
+	 * repopulating it with a fresh set of associated node-spawners. 
 	 *
 	 * @param  Class	The class entry you want rebuilt.
 	 */
 	void RefreshClassActions(UClass* const Class);
 
 	/**
-	 * Finds the database entry for the specified class and wipes it. The entry 
-	 * won't be rebuilt, unless RefreshClassActions() is explicitly called after.
+	 * Finds the database entry for the specified asset and wipes it,
+	 * repopulating it with a fresh set of associated node-spawners.
 	 * 
-	 * @param  Class	The class entry you want emptied.
+	 * @param  AssetObject	The asset entry you want rebuilt.
 	 */
-	void ClearClassActions(UClass* const Class);
-	
+	void RefreshAssetActions(UObject const* const AssetObject);
+
 	/**
-	 * Will populate the database first if it hasn't been created yet, and then 
-	 * returns it in its entirety. 
+	 * Finds the database entry for the specified class and wipes it. The entry 
+	 * won't be rebuilt, unless RefreshAssetActions() is explicitly called after.
 	 * 
-	 * Each node spawner is categorized by a class. A spawner that corresponds
-	 * to a specific class field (like a function, property, enum, etc.) will
-	 * be listed under that field's class outer. Remaining spawners that can't 
-	 * be categorized this way will be filed by node type; for example: with 
-	 * global enums (enums that don't belong to a specific class), the enum's 
-	 * UK2Node_EnumLiteral spawner is found under the UK2Node_EnumLiteral class.
-	 *
-	 * @return The entire action database, which maps specific classes to arrays of associated node-spawners.
+	 * @param  AssetObject	
 	 */
-	FClassActionMap const& GetAllActions();
-
+	void ClearAssetActions(UObject const* const AssetObject);
+	
 	/**
+	 * 
 	 * @return An estimated memory footprint of this database (in bytes).
 	 */
-	int32 Size() const;
+	int32 EstimatedSize() const;
 
 private:
 	/** Private constructor for singleton purposes. */
 	FBlueprintActionDatabase();
 
+	/**
+	 * 
+	 * 
+	 * @param  Registrar	
+	 */
+	void RegisterAllNodeActions(FBlueprintActionDatabaseRegistrar& Registrar);
+
+private:
 	/** 
-	 * A map of associated node-spawners for each class. A spawner that 
+	 * A map of associated node-spawners for each class/asset. A spawner that 
 	 * corresponds to a specific class field (like a function, property, enum, 
 	 * etc.) will be mapped under that field's class outer. Other spawners (that
 	 * can't be associated with a class outer), will be filed under the desired
-	 * node's type.
+	 * node's type, or an associated asset.
 	 */
-	FClassActionMap ClassActions;
+	FActionRegistry ActionRegistry;
 
 	/** 
 	 * References newly allocated actions that need to be "primed". Priming is 
 	 * something we do on Tick() aimed at speeding up performance (like pre-
 	 * caching each spawner's template-node, etc.).
 	 */
-	TMap<TWeakObjectPtr<UClass>, int32> ActionPrimingQueue;
+	FPrimingQueue ActionPrimingQueue;
 };
diff --git a/Engine/Source/Editor/BlueprintGraph/Public/BlueprintActionDatabaseRegistrar.h b/Engine/Source/Editor/BlueprintGraph/Public/BlueprintActionDatabaseRegistrar.h
new file mode 100644
--- /dev/null
+++ b/Engine/Source/Editor/BlueprintGraph/Public/BlueprintActionDatabaseRegistrar.h
@@ -0,0 +1,91 @@
+// Copyright 1998-2014 Epic Games, Inc. All Rights Reserved.
+
+#pragma once
+
+#include "BlueprintActionDatabase.h" // for FActionRegistry
+
+/**
+ * From the BlueprintActionDatabase, passed around to all UK2Nodes, giving each 
+ * a chance to register its own actions (specifically for UK2Nodes in other
+ * modules that the database doesn't have access to).
+ */
+class BLUEPRINTGRAPH_API FBlueprintActionDatabaseRegistrar
+{
+public:
+	/**
+	 * Attempts to suss out the key that this action should be registered under;
+	 * if it doesn't find a better one then it associates the action with the 
+	 * node filling this out.
+	 * 
+	 * @param  NodeSpawner	The new node spawner that you wish to register.
+	 */
+	void AddBlueprintAction(UBlueprintNodeSpawner* NodeSpawner);
+
+	/**
+	 * Each action should be recorded under a specific UField key; primarily to 
+	 * refresh those actions when the corresponding asset is updated (Blueprint
+	 * regenerated, struct added/deleted, etc.).
+	 * 
+	 * @param  FieldOwner	A field object that the node is associated with (flags when this action should be updated).
+	 * @param  NodeSpawner	The new node spawner that you wish to register.
+	 */
+	void AddBlueprintAction(UClass const* ClassOwner, UBlueprintNodeSpawner* NodeSpawner);
+	void AddBlueprintAction(UEnum const* EnumOwner, UBlueprintNodeSpawner* NodeSpawner);
+	void AddBlueprintAction(UScriptStruct const* StructOwner, UBlueprintNodeSpawner* NodeSpawner);
+	void AddBlueprintAction(UField const* FieldOwner, UBlueprintNodeSpawner* NodeSpawner);
+
+	/**
+	 * Special case for asset bound actions (we want to clean-up/refresh these 
+	 * when the corresponding asset is updated). AssetOwner must be an asset!
+	 * 
+	 * @param  AssetOwner	An asset object that the node is associated with (flags when this action should be updated).
+	 * @param  NodeSpawner	The new node spawner that you wish to register.
+	 */
+	void AddBlueprintAction(UObject const* AssetOwner, UBlueprintNodeSpawner* NodeSpawner);
+
+	/**
+	 * Occasionally (when an asset is added/refreshed), this registrar will be 
+	 * passed around to gather only specific keyed actions (see ActionKeyFilter). 
+	 * In that case, it will block registration of all unwanted keys.  
+	 * Functionality wise this doesn't matter to UK2Node, but UK2Node may be  
+	 * able to save on some work/allocations if it knew this beforehand.
+	 * 
+	 * @param  OwnerKey		The key you wish to register your action(s) under.
+	 * @return True if the OwnerKey would is allowed to register actions, false if it would be blocked. 
+	 */
+	bool IsOpenForRegistration(UObject const* OwnerKey);
+
+private:
+	typedef FBlueprintActionDatabase::FActionRegistry FActionRegistry;
+	typedef FBlueprintActionDatabase::FPrimingQueue   FPrimingQueue;
+
+	/** Only FBlueprintActionDatabase can spawn and distribute this. */
+	friend class FBlueprintActionDatabase;
+	FBlueprintActionDatabaseRegistrar(FActionRegistry& Database, FPrimingQueue& PrimingQueue, TSubclassOf<UEdGraphNode> DefaultKey = nullptr);
+
+	/**
+	 * Internal method that actually adds the action to the database.
+	 * 
+	 * @param  ActionKey	The key you want the action listed under.
+	 * @param  NodeSpawner	The new node spawner that you wish to register.
+	 */
+	void AddActionToDatabase(UObject const* ActionKey, UBlueprintNodeSpawner* NodeSpawner);
+
+private:
+	/** The node type that this is currently passed to (acts as a fallback database key). */
+	TSubclassOf<UEdGraphNode> GeneratingClass;
+
+	/** A reference to FBlueprintActionDatabase's internal map */
+	FActionRegistry& ActionDatabase;
+
+	/** 
+	 * When an asset is added/updated we want to poll nodes again for 
+	 * new/refreshed actions, but only for that specific asset. This serves as
+	 * a filter in that scenario (to keep nodes from adding unrelated, probably
+	 * duplicate, actions).
+	 */
+	UObject const* ActionKeyFilter;
+
+	/** A reference to FBlueprintActionDatabase's action priming queue */
+	FPrimingQueue& ActionPrimingQueue;
+};
diff --git a/Engine/Source/Editor/BlueprintGraph/Public/BlueprintActionFilter.h b/Engine/Source/Editor/BlueprintGraph/Public/BlueprintActionFilter.h
--- a/Engine/Source/Editor/BlueprintGraph/Public/BlueprintActionFilter.h
+++ b/Engine/Source/Editor/BlueprintGraph/Public/BlueprintActionFilter.h
@@ -35,8 +35,8 @@ struct FBlueprintActionContext
 
 	/** 
 	 * A list of objects the user currently has selected (things like blueprint 
-	 * properties, level actors, content-browser assets, etc.). Objects listed 
-	 * here will result in certain bound actions passing the filter.
+	 * properties, level actors, content-browser assets, etc.). Bound actions
+	 * have to be tied to one of these objects in order to pass the filter.
 	 */
 	TArray<UObject*> SelectedObjects;
 };
@@ -56,50 +56,52 @@ public:
 	enum EFlags // Flags, which configure certain rejection tests.
 	{
 		/** Deprecated class actions will not be filtered out*/
-		BPFILTER_IncludeDeprecated		= 0x00000001,
+		BPFILTER_PermitDeprecated				= 0x00000001,
 
 		/** 
-		 * Excludes nodes associated with global field types (functions, 
-		 * structs, enums, etc.); NOTE: that static function library members are 
-		 * counted as "global". 
+		 * Rejects actions associated with global/static ("persistent") fields.
+		 * However, static members of TargetClasses are kept, and not outright 
+		 * rejected by this (if TargetClasses is left empty, then all 
+		 * "persistent" fields are rejected.
 		 */
-		BPFILTER_ExcludeGlobalFields	= 0x00000002,
+		BPFILTER_RejectPersistentNonTargetFields= 0x00000002,
 
 		/**
-		 * Makes NodeTypes tests more aggressive by excluding nodes sub-classes, 
-		 * nodes have to explicitly match the class listed. 
+		 * Makes PermittedNodeType tests more aggressive by rejecting node 
+		 * sub-classes, (actions would have to explicitly match a class listed
+		 * in PermittedNodeTypes). 
 		 */
-		BPFILTER_ExcludeChildNodeTypes	= 0x00000004,
+		BPFILTER_RejectPermittedNodeSubClasses	= 0x00000004,
 	};
 	FBlueprintActionFilter(uint32 const Flags = 0x00);
 	
 	/**
-	 * A list of allowed node types. If an action's NodeClass isn't one of
-	 * these types, then it is filtered out. Use the "ExcludeChildNodeTypes"
-	 * flag to aggressively filter out child classes (must be an explicit match).
+	 * Contains the full blueprint/graph/pin context that this is filtering 
+	 * actions for.
 	 */
 	FBlueprintActionContext Context;	
 	
 	/**
-	 * A list of allowed node types. If an action's NodeClass isn't one of
-	 * these types, then it is filtered out. Use the "ExcludeChildNodeTypes"
-	 * flag to aggressively filter out child classes (must be an explicit match).
+	 * A list of allowed node types. If a spawner's NodeClass isn't one of
+	 * these types, then it is filtered out. Use the "RejectPermittedNodeSubClasses"
+	 * flag to aggressively filter out child classes as well (enforcing an 
+	 * explicit match).
 	 */
-	TArray< TSubclassOf<UEdGraphNode> > NodeTypes;
+	TArray< TSubclassOf<UEdGraphNode> > PermittedNodeTypes;
 	
 	/**
 	 * A list of node types that should be filtered out. If a node class is 
-	 * listed both here and in NodeTypes, then the exclusion wins (and it will 
-	 * be filtered out).
+	 * listed both here and in PermittedNodeTypes, then the exclusion wins (it 
+	 * will be filtered out).
 	 */
-	TArray< TSubclassOf<UEdGraphNode> > ExcludedNodeTypes;
+	TArray< TSubclassOf<UEdGraphNode> > RejectedNodeTypes;
 	
 	/**
-	 * A list of classes that you want members for. If an action does not have 
-	 * an associated field, or the associated field doesn't belong to any of 
-	 * these classes, then it is filtered out.
+	 * A list of classes that you want members for. If an action would produce
+	 * a node with a TargetPin, and that pin is incompatible with one of these
+	 * classes, then the action is filtered out.
 	 */
-	TArray<UClass*> OwnerClasses;
+	TArray<UClass*> TargetClasses;
 
 	/**
 	 * Users can extend the filter and add their own rejection tests with this
diff --git a/Engine/Source/Editor/GameplayAbilitiesEditor/Classes/K2Node_GameplayCueEvent.h b/Engine/Source/Editor/GameplayAbilitiesEditor/Classes/K2Node_GameplayCueEvent.h
--- a/Engine/Source/Editor/GameplayAbilitiesEditor/Classes/K2Node_GameplayCueEvent.h
+++ b/Engine/Source/Editor/GameplayAbilitiesEditor/Classes/K2Node_GameplayCueEvent.h
@@ -18,6 +18,6 @@ class UK2Node_GameplayCueEvent : public UK2Node_Event
 	// End UEdGraphNode interface
 	
 	// Begin UK2Node interface
-	virtual void GetMenuActions(TArray<UBlueprintNodeSpawner*>& ActionListOut) const override;
+	virtual void GetMenuActions(FBlueprintActionDatabaseRegistrar& ActionRegistrar) const override;
 	// End UK2Node interface
 };
\ No newline at end of file
diff --git a/Engine/Source/Editor/GameplayAbilitiesEditor/Classes/K2Node_LatentAbilityCall.h b/Engine/Source/Editor/GameplayAbilitiesEditor/Classes/K2Node_LatentAbilityCall.h
--- a/Engine/Source/Editor/GameplayAbilitiesEditor/Classes/K2Node_LatentAbilityCall.h
+++ b/Engine/Source/Editor/GameplayAbilitiesEditor/Classes/K2Node_LatentAbilityCall.h
@@ -16,7 +16,7 @@ class UK2Node_LatentAbilityCall : public UK2Node_BaseAsyncTask
 	// UEdGraphNode interface
 	virtual bool CanCreateUnderSpecifiedSchema(const UEdGraphSchema* DesiredSchema) const;
 	virtual void GetMenuEntries(FGraphContextMenuBuilder& ContextMenuBuilder) const override;
-	virtual void GetMenuActions(TArray<UBlueprintNodeSpawner*>& ActionListOut) const override;
+	virtual void GetMenuActions(FBlueprintActionDatabaseRegistrar& ActionRegistrar) const override;
 	virtual bool IsCompatibleWithGraph(UEdGraph const* TargetGraph) const override;
 	// End of UEdGraphNode interface
 
diff --git a/Engine/Source/Editor/GameplayAbilitiesEditor/Private/K2Node_GameplayCueEvent.cpp b/Engine/Source/Editor/GameplayAbilitiesEditor/Private/K2Node_GameplayCueEvent.cpp
--- a/Engine/Source/Editor/GameplayAbilitiesEditor/Private/K2Node_GameplayCueEvent.cpp
+++ b/Engine/Source/Editor/GameplayAbilitiesEditor/Private/K2Node_GameplayCueEvent.cpp
@@ -9,6 +9,7 @@
 #include "GameplayCueInterface.h"
 #include "BlueprintEventNodeSpawner.h"
 #include "BlueprintEditorUtils.h"
+#include "BlueprintActionDatabaseRegistrar.h"
 
 #define LOCTEXT_NAMESPACE "K2Node_GameplayCueEvent"
 
@@ -87,7 +88,7 @@ void UK2Node_GameplayCueEvent::GetMenuEntries(FGraphContextMenuBuilder& Context)
 	}	
 }
 
-void UK2Node_GameplayCueEvent::GetMenuActions(TArray<UBlueprintNodeSpawner*>& ActionListOut) const
+void UK2Node_GameplayCueEvent::GetMenuActions(FBlueprintActionDatabaseRegistrar& ActionRegistrar) const
 {
 	auto CustomizeCueNodeLambda = [](UEdGraphNode* NewNode, bool bIsTemplateNode, FName TagName)
 	{
@@ -107,7 +108,7 @@ void UK2Node_GameplayCueEvent::GetMenuActions(TArray<UBlueprintNodeSpawner*>& Ac
 		check(NodeSpawner != nullptr);
 		NodeSpawner->CustomizeNodeDelegate = PostSpawnDelegate;
 		
-		ActionListOut.Add(NodeSpawner);
+		ActionRegistrar.AddBlueprintAction(NodeSpawner);
 	}
 }
 
diff --git a/Engine/Source/Editor/GameplayAbilitiesEditor/Private/K2Node_LatentAbilityCall.cpp b/Engine/Source/Editor/GameplayAbilitiesEditor/Private/K2Node_LatentAbilityCall.cpp
--- a/Engine/Source/Editor/GameplayAbilitiesEditor/Private/K2Node_LatentAbilityCall.cpp
+++ b/Engine/Source/Editor/GameplayAbilitiesEditor/Private/K2Node_LatentAbilityCall.cpp
@@ -9,6 +9,7 @@
 #include "GameplayAbilityGraphSchema.h"
 #include "K2Node_EnumLiteral.h"
 #include "BlueprintFunctionNodeSpawner.h"
+#include "BlueprintActionDatabaseRegistrar.h"
 
 #define LOCTEXT_NAMESPACE "K2Node"
 
@@ -80,7 +81,7 @@ void UK2Node_LatentAbilityCall::GetMenuEntries(FGraphContextMenuBuilder& Context
 	}
 }
 
-void UK2Node_LatentAbilityCall::GetMenuActions(TArray<UBlueprintNodeSpawner*>& ActionListOut) const
+void UK2Node_LatentAbilityCall::GetMenuActions(FBlueprintActionDatabaseRegistrar& ActionRegistrar) const
 {
 	// these nested loops are combing over the same classes/functions the
 	// FBlueprintActionDatabase does; ideally we save on perf and fold this in
@@ -130,7 +131,7 @@ void UK2Node_LatentAbilityCall::GetMenuActions(TArray<UBlueprintNodeSpawner*>& A
 				
 				// @TODO: since this can't be folded into FBlueprintActionDatabase, we
 				//        need a way to associate these spawners with a certain class
-				ActionListOut.Add(NodeSpawner);
+				ActionRegistrar.AddBlueprintAction(NodeSpawner);
 			}
 		}
 	}
diff --git a/Engine/Source/Editor/GameplayTagsEditor/Classes/GameplayTagsK2Node_LiteralGameplayTag.h b/Engine/Source/Editor/GameplayTagsEditor/Classes/GameplayTagsK2Node_LiteralGameplayTag.h
--- a/Engine/Source/Editor/GameplayTagsEditor/Classes/GameplayTagsK2Node_LiteralGameplayTag.h
+++ b/Engine/Source/Editor/GameplayTagsEditor/Classes/GameplayTagsK2Node_LiteralGameplayTag.h
@@ -25,7 +25,7 @@ class UGameplayTagsK2Node_LiteralGameplayTag : public UK2Node
 	virtual bool ShouldShowNodeProperties() const override { return true; }
 	virtual bool IsNodeSafeToIgnore() const override { return true; }
 	virtual void ExpandNode(class FKismetCompilerContext& CompilerContext, UEdGraph* SourceGraph) override;
-	virtual void GetMenuActions(TArray<UBlueprintNodeSpawner*>& ActionListOut) const override;
+	virtual void GetMenuActions(FBlueprintActionDatabaseRegistrar& ActionRegistrar) const override;
 	virtual FText GetMenuCategory() const override;
 	// End UK2Node interface
 #endif
diff --git a/Engine/Source/Editor/Kismet/Private/BlueprintActionMenuBuilder.cpp b/Engine/Source/Editor/Kismet/Private/BlueprintActionMenuBuilder.cpp
--- a/Engine/Source/Editor/Kismet/Private/BlueprintActionMenuBuilder.cpp
+++ b/Engine/Source/Editor/Kismet/Private/BlueprintActionMenuBuilder.cpp
@@ -537,9 +537,9 @@ void FBlueprintActionMenuBuilderImpl::AppendLegacyItems(FMenuSectionDefinition c
 		FBlueprintPaletteListBuilder LegacyBuilder(Blueprint, MenuDef.GetSectionHeading().ToString());
 		
 		UClass* ClassFilter = nullptr;
-		if (MenuFilter.OwnerClasses.Num() > 0)
+		if (MenuFilter.TargetClasses.Num() > 0)
 		{
-			ClassFilter = MenuFilter.OwnerClasses[0];
+			ClassFilter = MenuFilter.TargetClasses[0];
 		}
 		
 		UEdGraphSchema_K2 const* K2Schema = GetDefault<UEdGraphSchema_K2>();
@@ -593,12 +593,13 @@ void FBlueprintActionMenuBuilder::RebuildActionList()
 	
 	if (!GetDefault<UEdGraphSchema_K2>()->bUseLegacyActionMenus)
 	{
-		FBlueprintActionDatabase::FClassActionMap const& ClassActions = FBlueprintActionDatabase::Get().GetAllActions();
-		for (auto ActionEntry : ClassActions)
+		FBlueprintActionDatabase::FActionRegistry const& ActionDatabase = FBlueprintActionDatabase::Get().GetAllActions();
+		for (auto ActionEntry : ActionDatabase)
 		{
-			UClass* Class = ActionEntry.Key;
 			for (UBlueprintNodeSpawner* NodeSpawner : ActionEntry.Value)
 			{
+				// @TODO: could probably have a super filter that spreads across 
+				//        all MenuSctions (to pair down on performance)
 				for (TSharedRef<FMenuSectionDefinition> MenuSection : MenuSections)
 				{
 					TSharedPtr<FEdGraphSchemaAction> MenuEntry = MenuSection->MakeMenuItem(BlueprintEditorPtr, NodeSpawner);
diff --git a/Engine/Source/Editor/Kismet/Private/BlueprintActionMenuUtils.cpp b/Engine/Source/Editor/Kismet/Private/BlueprintActionMenuUtils.cpp
--- a/Engine/Source/Editor/Kismet/Private/BlueprintActionMenuUtils.cpp
+++ b/Engine/Source/Editor/Kismet/Private/BlueprintActionMenuUtils.cpp
@@ -45,7 +45,8 @@ void FBlueprintActionMenuUtils::MakePaletteMenu(FBlueprintActionContext const& C
 	uint32 FilterFlags = 0x00;
 	if (FilterClass != nullptr)
 	{
-		FilterFlags = FBlueprintActionFilter::BPFILTER_ExcludeGlobalFields;
+		// make sure we exclude global and static library actions
+		FilterFlags = FBlueprintActionFilter::BPFILTER_RejectPersistentNonTargetFields;
 	}
 	
 	FBlueprintActionFilter MenuFilter(FilterFlags);
@@ -53,12 +54,12 @@ void FBlueprintActionMenuUtils::MakePaletteMenu(FBlueprintActionContext const& C
 	
 	// self member variables can be accessed through the MyBluprint panel (even
 	// inherited ones)... external variables can be accessed through the context
-	// menu (don't want to clutter the palette?)
-	MenuFilter.ExcludedNodeTypes.Add(UK2Node_Variable::StaticClass());
+	// menu (don't want to clutter the palette, I guess?)
+	MenuFilter.RejectedNodeTypes.Add(UK2Node_Variable::StaticClass());
 	
 	if (FilterClass != nullptr)
 	{
-		MenuFilter.OwnerClasses.Add(FilterClass);
+		MenuFilter.TargetClasses.Add(FilterClass);
 	}
 
 	MenuOut.AddMenuSection(MenuFilter, LOCTEXT("PaletteRoot", "Library"), /*MenuOrder =*/0, FBlueprintActionMenuBuilder::ConsolidatePropertyActions);
@@ -75,8 +76,8 @@ void FBlueprintActionMenuUtils::MakeContextMenu(FBlueprintActionContext const& C
 
 	FBlueprintActionFilter ComponentsFilter;
 	ComponentsFilter.Context = Context;
-	ComponentsFilter.NodeTypes.Add(UK2Node_ComponentBoundEvent::StaticClass());
-	ComponentsFilter.NodeTypes.Add(UK2Node_CallFunction::StaticClass());
+	ComponentsFilter.PermittedNodeTypes.Add(UK2Node_ComponentBoundEvent::StaticClass());
+	ComponentsFilter.PermittedNodeTypes.Add(UK2Node_CallFunction::StaticClass());
 	ComponentsFilter.AddIsFilteredTest(FBlueprintActionFilter::FIsFilteredDelegate::CreateStatic(IsUnBoundSpawner));
 	
 	bool bHasComponentsSection = false;
@@ -84,7 +85,7 @@ void FBlueprintActionMenuUtils::MakeContextMenu(FBlueprintActionContext const& C
 	{
 		if (UObjectProperty* ObjProperty = Cast<UObjectProperty>(Selection))
 		{
-			ComponentsFilter.OwnerClasses.Add(ObjProperty->PropertyClass);
+			ComponentsFilter.TargetClasses.Add(ObjProperty->PropertyClass);
 			bHasComponentsSection = true;
 		}
 	}
@@ -107,8 +108,7 @@ void FBlueprintActionMenuUtils::MakeContextMenu(FBlueprintActionContext const& C
 			for (UBlueprint* Blueprint : Context.Blueprints)
 			{
 				UClass* BlueprintClass = (Blueprint->SkeletonGeneratedClass != nullptr) ? Blueprint->SkeletonGeneratedClass : Blueprint->ParentClass;
-				// @TODO: are we sure we only want members in the unconnected context menu? 
-				MenuFilter.OwnerClasses.Add(BlueprintClass);
+				MenuFilter.TargetClasses.Add(BlueprintClass);
 			}
 		}
 	}
diff --git a/Engine/Source/Editor/OnlineBlueprintSupport/Classes/K2Node_LatentOnlineCall.h b/Engine/Source/Editor/OnlineBlueprintSupport/Classes/K2Node_LatentOnlineCall.h
--- a/Engine/Source/Editor/OnlineBlueprintSupport/Classes/K2Node_LatentOnlineCall.h
+++ b/Engine/Source/Editor/OnlineBlueprintSupport/Classes/K2Node_LatentOnlineCall.h
@@ -12,6 +12,6 @@ class ONLINEBLUEPRINTSUPPORT_API UK2Node_LatentOnlineCall : public UK2Node_BaseA
 	
 	// UEdGraphNode interface
 	virtual void GetMenuEntries(FGraphContextMenuBuilder& ContextMenuBuilder) const override;
-	virtual void GetMenuActions(TArray<UBlueprintNodeSpawner*>& ActionListOut) const override;
+	virtual void GetMenuActions(FBlueprintActionDatabaseRegistrar& ActionRegistrar) const override;
 	// End of UEdGraphNode interface
 };
diff --git a/Engine/Source/Editor/OnlineBlueprintSupport/Private/K2Node_LatentOnlineCall.cpp b/Engine/Source/Editor/OnlineBlueprintSupport/Private/K2Node_LatentOnlineCall.cpp
--- a/Engine/Source/Editor/OnlineBlueprintSupport/Private/K2Node_LatentOnlineCall.cpp
+++ b/Engine/Source/Editor/OnlineBlueprintSupport/Private/K2Node_LatentOnlineCall.cpp
@@ -4,6 +4,7 @@
 #include "OnlineBlueprintCallProxyBase.h"
 #include "K2Node_LatentOnlineCall.h"
 #include "BlueprintFunctionNodeSpawner.h"
+#include "BlueprintActionDatabaseRegistrar.h"
 
 #define LOCTEXT_NAMESPACE "K2Node"
 
@@ -50,7 +51,7 @@ void UK2Node_LatentOnlineCall::GetMenuEntries(FGraphContextMenuBuilder& ContextM
 	}
 }
 
-void UK2Node_LatentOnlineCall::GetMenuActions(TArray<UBlueprintNodeSpawner*>& ActionListOut) const
+void UK2Node_LatentOnlineCall::GetMenuActions(FBlueprintActionDatabaseRegistrar& ActionRegistrar) const
 {
 	// these nested loops are combing over the same classes/functions the
 	// FBlueprintActionDatabase does; ideally we save on perf and fold this in
@@ -100,7 +101,7 @@ void UK2Node_LatentOnlineCall::GetMenuActions(TArray<UBlueprintNodeSpawner*>& Ac
 				
 				// @TODO: since this can't be folded into FBlueprintActionDatabase, we
 				//        need a way to associate these spawners with a certain class
-				ActionListOut.Add(NodeSpawner);
+				ActionRegistrar.AddBlueprintAction(NodeSpawner);
 			}
 		}
 	}
diff --git a/Engine/Source/Editor/UnrealEd/Private/Commandlets/DumpBlueprintsInfoCommandlet.cpp b/Engine/Source/Editor/UnrealEd/Private/Commandlets/DumpBlueprintsInfoCommandlet.cpp
--- a/Engine/Source/Editor/UnrealEd/Private/Commandlets/DumpBlueprintsInfoCommandlet.cpp
+++ b/Engine/Source/Editor/UnrealEd/Private/Commandlets/DumpBlueprintsInfoCommandlet.cpp
@@ -1041,7 +1041,7 @@ static bool DumpBlueprintInfoUtils::DumpActionDatabaseInfo(uint32 Indent, FArchi
 			}
 		}
 
-		float const DatabaseKbSize = Database.Size() / 1024.f;
+		float const DatabaseKbSize = Database.EstimatedSize() / 1024.f;
 		float EstimatedCacheKbSize = 0.0;
 		for (UBlueprint* CacheBlueprint : TemplateOuters)
 		{
