Missed file from latest checkin.

diff --git a/Engine/Source/Editor/BlueprintGraph/Classes/K2Node_GetDataTableRow.h b/Engine/Source/Editor/BlueprintGraph/Classes/K2Node_GetDataTableRow.h
--- a/Engine/Source/Editor/BlueprintGraph/Classes/K2Node_GetDataTableRow.h
+++ b/Engine/Source/Editor/BlueprintGraph/Classes/K2Node_GetDataTableRow.h
@@ -6,7 +6,7 @@
 #include "K2Node_GetDataTableRow.generated.h"
 
 UCLASS()
-class UK2Node_GetDataTableRow : public UK2Node
+class BLUEPRINTGRAPH_API UK2Node_GetDataTableRow : public UK2Node
 {
 	GENERATED_UCLASS_BODY()
 
@@ -30,7 +30,7 @@ class UK2Node_GetDataTableRow : public UK2Node
 	// Begin UK2Node interface
 	virtual bool IsNodeSafeToIgnore() const override { return true; }
 	virtual void ReallocatePinsDuringReconstruction(TArray<UEdGraphPin*>& OldPins) override;
-	virtual void GetMenuActions(TArray<UBlueprintNodeSpawner*>& ActionListOut) const override;
+	virtual void GetMenuActions(FBlueprintActionDatabaseRegistrar& ActionRegistrar) const override;
 	virtual FText GetMenuCategory() const override;
 	// End UK2Node interface
 
