Fixing up various build errors.

diff --git a/Engine/Source/Editor/BlueprintGraph/Classes/K2Node_MathExpression.h b/Engine/Source/Editor/BlueprintGraph/Classes/K2Node_MathExpression.h
--- a/Engine/Source/Editor/BlueprintGraph/Classes/K2Node_MathExpression.h
+++ b/Engine/Source/Editor/BlueprintGraph/Classes/K2Node_MathExpression.h
@@ -5,14 +5,14 @@
 #include "K2Node_MathExpression.generated.h"
 
 /**
-* This node type acts like our collapsed nodes, a single node that represents
+* This node type acts like a collapsed node, a single node that represents
 * a larger sub-network of nodes (contained within a sub-graph). This node will
-* take the expression it was named with, and attempt to convert it into a
-* series of math node. If it is unsuccessful, then it generates a series of
-* errors that are tacked onto the node for the user to see.
+* take the math expression it was named with, and attempt to convert it into a
+* series of math nodes. If it is unsuccessful, then it generates a series of
+* actionable errors.
 */
 UCLASS()
-class UK2Node_MathExpression : public UK2Node_Composite
+class BLUEPRINTGRAPH_API UK2Node_MathExpression : public UK2Node_Composite
 {
 	GENERATED_UCLASS_BODY()
 
@@ -38,7 +38,7 @@ public:
 	// UK2Node interface
 	virtual void GetMenuEntries(FGraphContextMenuBuilder& ContextMenuBuilder) const override;
 	virtual void ValidateNodeDuringCompilation(class FCompilerResultsLog& MessageLog) const override;
-	virtual void GetMenuActions(TArray<UBlueprintNodeSpawner*>& ActionListOut) const override;
+	virtual void GetMenuActions(FBlueprintActionDatabaseRegistrar& ActionRegistrar) const override;
 	// End of UK2Node interface
 
 private:
diff --git a/Engine/Source/Editor/BlueprintGraph/Private/BlueprintActionDatabaseRegistrar.cpp b/Engine/Source/Editor/BlueprintGraph/Private/BlueprintActionDatabaseRegistrar.cpp
--- a/Engine/Source/Editor/BlueprintGraph/Private/BlueprintActionDatabaseRegistrar.cpp
+++ b/Engine/Source/Editor/BlueprintGraph/Private/BlueprintActionDatabaseRegistrar.cpp
@@ -60,8 +60,8 @@ static UObject const* BlueprintActionDatabaseRegistrarImpl::ResolveActionKey(UOb
 
 //------------------------------------------------------------------------------
 FBlueprintActionDatabaseRegistrar::FBlueprintActionDatabaseRegistrar(FActionRegistry& Database, FPrimingQueue& PrimingQueue, TSubclassOf<UEdGraphNode> DefaultKey)
-	: ActionDatabase(Database)
-	, GeneratingClass(DefaultKey)
+	: GeneratingClass(DefaultKey)
+	, ActionDatabase(Database)
 	, ActionKeyFilter(nullptr)
 	, ActionPrimingQueue(PrimingQueue)
 {
diff --git a/Engine/Source/Editor/BlueprintGraph/Private/K2Node_BaseAsyncTask.cpp b/Engine/Source/Editor/BlueprintGraph/Private/K2Node_BaseAsyncTask.cpp
--- a/Engine/Source/Editor/BlueprintGraph/Private/K2Node_BaseAsyncTask.cpp
+++ b/Engine/Source/Editor/BlueprintGraph/Private/K2Node_BaseAsyncTask.cpp
@@ -7,6 +7,7 @@
 #include "KismetCompiler.h"
 #include "K2ActionMenuBuilder.h"
 #include "BlueprintNodeSpawner.h"
+#include "BlueprintActionDatabaseRegistrar.h"
 
 #define LOCTEXT_NAMESPACE "UK2Node_BaseAsyncTask"
 
@@ -374,12 +375,12 @@ FText UK2Node_BaseAsyncTask::GetMenuCategory() const
 	return FText::FromString(UK2Node_CallFunction::GetDefaultCategoryForFunction(TargetFunction, TEXT("Call Function")));
 }
 
-void UK2Node_BaseAsyncTask::GetMenuActions(TArray<UBlueprintNodeSpawner*>& ActionListOut) const
+void UK2Node_BaseAsyncTask::GetMenuActions(FBlueprintActionDatabaseRegistrar& ActionRegistrar) const
 {
 	UBlueprintNodeSpawner* NodeSpawner = UBlueprintNodeSpawner::Create(GetClass());
 	check(NodeSpawner != nullptr);
 	
-	ActionListOut.Add(NodeSpawner);
+	ActionRegistrar.AddBlueprintAction(NodeSpawner);
 }
 
 UFunction* UK2Node_BaseAsyncTask::GetFactoryFunction() const
diff --git a/Engine/Source/Editor/BlueprintGraph/Private/K2Node_DoOnceMultiInput.cpp b/Engine/Source/Editor/BlueprintGraph/Private/K2Node_DoOnceMultiInput.cpp
--- a/Engine/Source/Editor/BlueprintGraph/Private/K2Node_DoOnceMultiInput.cpp
+++ b/Engine/Source/Editor/BlueprintGraph/Private/K2Node_DoOnceMultiInput.cpp
@@ -8,6 +8,7 @@
 #include "KismetCompiler.h"
 #include "BlueprintNodeSpawner.h"
 #include "EditorCategoryUtils.h"
+#include "BlueprintActionDatabaseRegistrar.h"
 
 #define LOCTEXT_NAMESPACE "K2Node"
 
@@ -375,12 +376,12 @@ void UK2Node_DoOnceMultiInput::ExpandNode(FKismetCompilerContext& CompilerContex
 	}
 }
 
-void UK2Node_DoOnceMultiInput::GetMenuActions(TArray<UBlueprintNodeSpawner*>& ActionListOut) const
+void UK2Node_DoOnceMultiInput::GetMenuActions(FBlueprintActionDatabaseRegistrar& ActionRegistrar) const
 {
 	UBlueprintNodeSpawner* NodeSpawner = UBlueprintNodeSpawner::Create(GetClass());
 	check(NodeSpawner != nullptr);
 	
-	ActionListOut.Add(NodeSpawner);
+	ActionRegistrar.AddBlueprintAction(NodeSpawner);
 }
 
 FText UK2Node_DoOnceMultiInput::GetMenuCategory() const
diff --git a/Engine/Source/Editor/BlueprintGraph/Private/K2Node_EnumEquality.cpp b/Engine/Source/Editor/BlueprintGraph/Private/K2Node_EnumEquality.cpp
--- a/Engine/Source/Editor/BlueprintGraph/Private/K2Node_EnumEquality.cpp
+++ b/Engine/Source/Editor/BlueprintGraph/Private/K2Node_EnumEquality.cpp
@@ -8,6 +8,7 @@
 #include "K2Node_EnumEquality.h"
 #include "BlueprintNodeSpawner.h"
 #include "EditorCategoryUtils.h"
+#include "BlueprintActionDatabaseRegistrar.h"
 
 #define LOCTEXT_NAMESPACE "K2Node_EnumEquality"
 
@@ -259,12 +260,12 @@ void UK2Node_EnumEquality::ExpandNode(class FKismetCompilerContext& CompilerCont
 	}
 }
 
-void UK2Node_EnumEquality::GetMenuActions(TArray<UBlueprintNodeSpawner*>& ActionListOut) const
+void UK2Node_EnumEquality::GetMenuActions(FBlueprintActionDatabaseRegistrar& ActionRegistrar) const
 {
 	UBlueprintNodeSpawner* NodeSpawner = UBlueprintNodeSpawner::Create(GetClass());
 	check(NodeSpawner != nullptr);
 	
-	ActionListOut.Add(NodeSpawner);
+	ActionRegistrar.AddBlueprintAction(NodeSpawner);
 }
 
 FText UK2Node_EnumEquality::GetMenuCategory() const
diff --git a/Engine/Source/Editor/GameplayTagsEditor/Private/GameplayTagsK2Node_LiteralGameplayTag.cpp b/Engine/Source/Editor/GameplayTagsEditor/Private/GameplayTagsK2Node_LiteralGameplayTag.cpp
--- a/Engine/Source/Editor/GameplayTagsEditor/Private/GameplayTagsK2Node_LiteralGameplayTag.cpp
+++ b/Engine/Source/Editor/GameplayTagsEditor/Private/GameplayTagsK2Node_LiteralGameplayTag.cpp
@@ -6,6 +6,7 @@
 #include "K2ActionMenuBuilder.h" // for FK2ActionMenuBuilder::AddNewNodeAction()
 #include "BlueprintNodeSpawner.h"
 #include "EditorCategoryUtils.h"
+#include "BlueprintActionDatabaseRegistrar.h"
 
 #define LOCTEXT_NAMESPACE "GameplayTagsK2Node_LiteralGameplayTag"
 
@@ -148,12 +149,12 @@ void UGameplayTagsK2Node_LiteralGameplayTag::ExpandNode(class FKismetCompilerCon
 	}
 }
 
-void UGameplayTagsK2Node_LiteralGameplayTag::GetMenuActions(TArray<UBlueprintNodeSpawner*>& ActionListOut) const
+void UGameplayTagsK2Node_LiteralGameplayTag::GetMenuActions(FBlueprintActionDatabaseRegistrar& ActionRegistrar) const
 {
 	UBlueprintNodeSpawner* NodeSpawner = UBlueprintNodeSpawner::Create(GetClass());
 	check(NodeSpawner != nullptr);
 	
-	ActionListOut.Add(NodeSpawner);
+	ActionRegistrar.AddBlueprintAction(NodeSpawner);
 }
 
 FText UGameplayTagsK2Node_LiteralGameplayTag::GetMenuCategory() const
