UE4: Fix for non-unity builds

diff --git a/Engine/Source/Editor/BlueprintGraph/Private/BlueprintActionDatabaseRegistrar.cpp b/Engine/Source/Editor/BlueprintGraph/Private/BlueprintActionDatabaseRegistrar.cpp
--- a/Engine/Source/Editor/BlueprintGraph/Private/BlueprintActionDatabaseRegistrar.cpp
+++ b/Engine/Source/Editor/BlueprintGraph/Private/BlueprintActionDatabaseRegistrar.cpp
@@ -2,6 +2,7 @@
 
 #include "BlueprintGraphPrivatePCH.h"
 #include "BlueprintActionDatabaseRegistrar.h"
+#include "BlueprintNodeSpawner.h"
 #include "BlueprintNodeSpawnerUtils.h"
 
 /*******************************************************************************
