BP: UK2Node_SetFieldsInStruct::RemoveFieldPins - obfuscated code refactored

diff --git a/Engine/Source/Editor/BlueprintGraph/Private/K2Node_SetFieldsInStruct.cpp b/Engine/Source/Editor/BlueprintGraph/Private/K2Node_SetFieldsInStruct.cpp
--- a/Engine/Source/Editor/BlueprintGraph/Private/K2Node_SetFieldsInStruct.cpp
+++ b/Engine/Source/Editor/BlueprintGraph/Private/K2Node_SetFieldsInStruct.cpp
@@ -114,10 +114,13 @@ void UK2Node_SetFieldsInStruct::RemoveFieldPins(const UEdGraphPin* Pin, EPinsToR
 			Pin = Pin->ParentPin;
 		}
 
+		const bool bHideSelected = (Selection == EPinsToRemove::GivenPin);
+		const bool bHideNotSelected = (Selection == EPinsToRemove::AllOtherPins);
 		bool bWasChanged = false;
 		for (FOptionalPinFromProperty& OptionalProperty : ShowPinForProperties)
 		{
-			const bool bHide = ((Pin->PinName == OptionalProperty.PropertyName.ToString()) == (Selection == EPinsToRemove::GivenPin));
+			const bool bSelected = (Pin->PinName == OptionalProperty.PropertyName.ToString());
+			const bool bHide = (bSelected && bHideSelected) || (!bSelected && bHideNotSelected);
 			if (OptionalProperty.bShowPin && bHide)
 			{
 				bWasChanged = true;
