Don't display the [REDACTED_HANDLE] and [REDACTED_HANDLE] portions of a blueprint callable function in its tooltip

diff --git a/Engine/Source/Editor/BlueprintGraph/Private/K2Node_CallFunction.cpp b/Engine/Source/Editor/BlueprintGraph/Private/K2Node_CallFunction.cpp
--- a/Engine/Source/Editor/BlueprintGraph/Private/K2Node_CallFunction.cpp
+++ b/Engine/Source/Editor/BlueprintGraph/Private/K2Node_CallFunction.cpp
@@ -754,6 +754,11 @@ FString UK2Node_CallFunction::GetTooltip() const
 	{
 		Tooltip = GetDefaultTooltipForFunction(Function);
 
+		// Strip off the @param nastiness
+		Tooltip.Split(TEXT("@param"), &Tooltip, nullptr, ESearchCase::IgnoreCase, ESearchDir::FromStart);
+		Tooltip.Split(TEXT("@return"), &Tooltip, nullptr, ESearchCase::IgnoreCase, ESearchDir::FromStart);
+		Tooltip.Trim();
+		while (Tooltip.RemoveFromEnd(TEXT("\n"))) ;
 
 		FString ClientString;
 
