Better fix for removing [REDACTED_HANDLE] and [REDACTED_HANDLE] from call function tooltips (works for menu as well as actual nodes)

diff --git a/Engine/Source/Editor/BlueprintGraph/Private/K2Node_CallFunction.cpp b/Engine/Source/Editor/BlueprintGraph/Private/K2Node_CallFunction.cpp
--- a/Engine/Source/Editor/BlueprintGraph/Private/K2Node_CallFunction.cpp
+++ b/Engine/Source/Editor/BlueprintGraph/Private/K2Node_CallFunction.cpp
@@ -421,8 +421,7 @@ void UK2Node_CallFunction::GeneratePinHoverText(const UEdGraphPin& Pin, FString&
 	}
 
 	// grab the the function's comment block for us to parse
-	FString FunctionToolTipText = GetDefaultTooltipForFunction(Function);
-
+	FString FunctionToolTipText = Function->GetToolTipText().ToString();
 	
 	int32 CurStrPos = INDEX_NONE;
 	int32 FullToolTipLen = FunctionToolTipText.Len();
@@ -754,12 +753,6 @@ FString UK2Node_CallFunction::GetTooltip() const
 	{
 		Tooltip = GetDefaultTooltipForFunction(Function);
 
-		// Strip off the @param nastiness
-		Tooltip.Split(TEXT("@param"), &Tooltip, nullptr, ESearchCase::IgnoreCase, ESearchDir::FromStart);
-		Tooltip.Split(TEXT("@return"), &Tooltip, nullptr, ESearchCase::IgnoreCase, ESearchDir::FromStart);
-		Tooltip.Trim();
-		while (Tooltip.RemoveFromEnd(TEXT("\n"))) ;
-
 		FString ClientString;
 
 		if (Function->HasAllFunctionFlags(FUNC_BlueprintAuthorityOnly))
@@ -799,10 +792,16 @@ FString UK2Node_CallFunction::GetUserFacingFunctionName(const UFunction* Functio
 
 FString UK2Node_CallFunction::GetDefaultTooltipForFunction(const UFunction* Function)
 {
-	const FString Tooltip = Function->GetToolTipText().ToString();
+	FString Tooltip = Function->GetToolTipText().ToString();
 
 	if (!Tooltip.IsEmpty())
 	{
+		// Strip off the @param nastiness
+		Tooltip.Split(TEXT("@param"), &Tooltip, nullptr, ESearchCase::IgnoreCase, ESearchDir::FromStart);
+		Tooltip.Split(TEXT("@return"), &Tooltip, nullptr, ESearchCase::IgnoreCase, ESearchDir::FromStart);
+		Tooltip.Trim();
+		while (Tooltip.RemoveFromEnd(TEXT("\n"))) ;
+
 		return Tooltip;
 	}
 	else
