Fix for a case-sensitiveness issue.

diff --git a/Engine/Source/Editor/BlueprintGraph/Private/BlueprintFunctionNodeSpawner.cpp b/Engine/Source/Editor/BlueprintGraph/Private/BlueprintFunctionNodeSpawner.cpp
--- a/Engine/Source/Editor/BlueprintGraph/Private/BlueprintFunctionNodeSpawner.cpp
+++ b/Engine/Source/Editor/BlueprintGraph/Private/BlueprintFunctionNodeSpawner.cpp
@@ -3,7 +3,7 @@
 #include "BlueprintGraphPrivatePCH.h"
 #include "BlueprintFunctionNodeSpawner.h"
 #include "EdGraphSchema_K2.h"	// for FBlueprintMetadata
-#include "Gameframework/Actor.h"
+#include "GameFramework/Actor.h"
 #include "BlueprintVariableNodeSpawner.h"
 
 #define LOCTEXT_NAMESPACE "BlueprintFunctionNodeSpawner"
