Cleaning up how certain blueprint nodes are listed (for the new action menu system).

diff --git a/Engine/Plugins/Experimental/Paper2D/Source/Paper2D/Classes/PaperFlipbookComponent.h b/Engine/Plugins/Experimental/Paper2D/Source/Paper2D/Classes/PaperFlipbookComponent.h
--- a/Engine/Plugins/Experimental/Paper2D/Source/Paper2D/Classes/PaperFlipbookComponent.h
+++ b/Engine/Plugins/Experimental/Paper2D/Source/Paper2D/Classes/PaperFlipbookComponent.h
@@ -8,7 +8,7 @@
 // Event for a non-looping flipbook finishing play
 DECLARE_DYNAMIC_MULTICAST_DELEGATE(FFlipbookFinishedPlaySignature);
 
-UCLASS(ShowCategories=(Mobility, ComponentReplication), EarlyAccessPreview, meta=(BlueprintSpawnableComponent))
+UCLASS(ShowCategories=(Mobility, ComponentReplication), ClassGroup=Paper2D, EarlyAccessPreview, meta=(BlueprintSpawnableComponent))
 class PAPER2D_API UPaperFlipbookComponent : public UPrimitiveComponent
 {
 	GENERATED_UCLASS_BODY()
diff --git a/Engine/Plugins/Experimental/Paper2D/Source/Paper2D/Classes/PaperSpriteComponent.h b/Engine/Plugins/Experimental/Paper2D/Source/Paper2D/Classes/PaperSpriteComponent.h
--- a/Engine/Plugins/Experimental/Paper2D/Source/Paper2D/Classes/PaperSpriteComponent.h
+++ b/Engine/Plugins/Experimental/Paper2D/Source/Paper2D/Classes/PaperSpriteComponent.h
@@ -6,7 +6,7 @@
 
 #include "PaperSpriteComponent.generated.h"
 
-UCLASS(ShowCategories=(Mobility), EarlyAccessPreview, meta=(BlueprintSpawnableComponent))
+UCLASS(ShowCategories=(Mobility), ClassGroup=Paper2D, EarlyAccessPreview, meta=(BlueprintSpawnableComponent))
 class PAPER2D_API UPaperSpriteComponent : public UMeshComponent
 {
 	GENERATED_UCLASS_BODY()
diff --git a/Engine/Plugins/Experimental/Paper2D/Source/Paper2D/Classes/PaperTileMapRenderComponent.h b/Engine/Plugins/Experimental/Paper2D/Source/Paper2D/Classes/PaperTileMapRenderComponent.h
--- a/Engine/Plugins/Experimental/Paper2D/Source/Paper2D/Classes/PaperTileMapRenderComponent.h
+++ b/Engine/Plugins/Experimental/Paper2D/Source/Paper2D/Classes/PaperTileMapRenderComponent.h
@@ -6,7 +6,7 @@
 
 #include "PaperTileMapRenderComponent.generated.h"
 
-UCLASS(hideCategories=Object, Experimental, meta=(BlueprintSpawnableComponent))
+UCLASS(hideCategories=Object, ClassGroup=Paper2D, Experimental, meta=(BlueprintSpawnableComponent))
 class PAPER2D_API UPaperTileMapRenderComponent : public UPrimitiveComponent
 {
 	GENERATED_UCLASS_BODY()
diff --git a/Engine/Source/Editor/BlueprintGraph/Classes/K2Node_ClassDynamicCast.h b/Engine/Source/Editor/BlueprintGraph/Classes/K2Node_ClassDynamicCast.h
--- a/Engine/Source/Editor/BlueprintGraph/Classes/K2Node_ClassDynamicCast.h
+++ b/Engine/Source/Editor/BlueprintGraph/Classes/K2Node_ClassDynamicCast.h
@@ -13,6 +13,7 @@ class UK2Node_ClassDynamicCast : public UK2Node_DynamicCast
 	// Begin UEdGraphNode interface
 	virtual void AllocateDefaultPins() override;
 	virtual FLinearColor GetNodeTitleColor() const override;
+	virtual FText GetNodeTitle(ENodeTitleType::Type TitleType) const override;
 	// End UEdGraphNode interface
 
 	// UK2Node interface
diff --git a/Engine/Source/Editor/BlueprintGraph/Private/K2Node_ClassDynamicCast.cpp b/Engine/Source/Editor/BlueprintGraph/Private/K2Node_ClassDynamicCast.cpp
--- a/Engine/Source/Editor/BlueprintGraph/Private/K2Node_ClassDynamicCast.cpp
+++ b/Engine/Source/Editor/BlueprintGraph/Private/K2Node_ClassDynamicCast.cpp
@@ -5,6 +5,8 @@
 #include "K2Node_ClassDynamicCast.h"
 #include "DynamicCastHandler.h"
 
+#define LOCTEXT_NAMESPACE "K2Node_ClassDynamicCast"
+
 struct FClassDynamicCastHelper
 {
 	static const FString& GetClassToCastName()
@@ -53,6 +55,11 @@ FLinearColor UK2Node_ClassDynamicCast::GetNodeTitleColor() const
 	return GetDefault<UGraphEditorSettings>()->ClassPinTypeColor;
 }
 
+FText UK2Node_ClassDynamicCast::GetNodeTitle(ENodeTitleType::Type TitleType) const
+{
+	return FText::Format(LOCTEXT("NodeTitle", "{0} Class"), Super::GetNodeTitle(TitleType));
+}
+
 UEdGraphPin* UK2Node_ClassDynamicCast::GetCastSourcePin() const
 {
 	UEdGraphPin* Pin = FindPinChecked(FClassDynamicCastHelper::GetClassToCastName());
@@ -64,3 +71,5 @@ FNodeHandlingFunctor* UK2Node_ClassDynamicCast::CreateNodeHandler(FKismetCompile
 {
 	return new FKCHandler_DynamicCast(CompilerContext, KCST_MetaCast);
 }
+
+#undef LOCTEXT_NAMESPACE
\ No newline at end of file
diff --git a/Engine/Source/Editor/BlueprintGraph/Private/K2Node_DoOnceMultiInput.cpp b/Engine/Source/Editor/BlueprintGraph/Private/K2Node_DoOnceMultiInput.cpp
--- a/Engine/Source/Editor/BlueprintGraph/Private/K2Node_DoOnceMultiInput.cpp
+++ b/Engine/Source/Editor/BlueprintGraph/Private/K2Node_DoOnceMultiInput.cpp
@@ -20,7 +20,7 @@ void UK2Node_DoOnceMultiInput::GetMenuEntries(FGraphContextMenuBuilder& ContextM
 	// Create a node template for this factory method
 	UK2Node_DoOnceMultiInput* NodeTemplate = NewObject<UK2Node_DoOnceMultiInput>(ContextMenuBuilder.OwnerOfTemporaries);
 	
-	TSharedPtr<FEdGraphSchemaAction_K2NewNode> NodeAction = FK2ActionMenuBuilder::AddNewNodeAction(ContextMenuBuilder, LOCTEXT("FlowCategory", "Flow Control").ToString(), NodeTemplate->GetNodeTitle(ENodeTitleType::ListView), NodeTemplate->GetTooltip(), 0, NodeTemplate->GetKeywords());
+	TSharedPtr<FEdGraphSchemaAction_K2NewNode> NodeAction = FK2ActionMenuBuilder::AddNewNodeAction(ContextMenuBuilder, LOCTEXT("FlowCategory", "Utilities|Flow Control").ToString(), NodeTemplate->GetNodeTitle(ENodeTitleType::ListView), NodeTemplate->GetTooltip(), 0, NodeTemplate->GetKeywords());
 
 	NodeAction->NodeTemplate = NodeTemplate;
 }
diff --git a/Engine/Source/Editor/BlueprintGraph/Private/K2Node_Message.cpp b/Engine/Source/Editor/BlueprintGraph/Private/K2Node_Message.cpp
--- a/Engine/Source/Editor/BlueprintGraph/Private/K2Node_Message.cpp
+++ b/Engine/Source/Editor/BlueprintGraph/Private/K2Node_Message.cpp
@@ -18,10 +18,17 @@ FText UK2Node_Message::GetNodeTitle(ENodeTitleType::Type TitleType) const
 	{
 		FString NodeNameStr = UK2Node_CallFunction::GetUserFacingFunctionName(Function);
 
-		FFormatNamedArguments Args;
-		Args.Add(TEXT("NodeName"), FText::FromString(NodeNameStr));
-		Args.Add(TEXT("OuterClassName"), FText::FromString(Function->GetOuterUClass()->GetName()));
-		NodeName = FText::Format(NSLOCTEXT("K2Node", "CallInterfaceContext", "{NodeName}\nUsing Interface {OuterClassName}"), Args);
+		if (TitleType == ENodeTitleType::ListView)
+		{
+			NodeName = FText::Format(LOCTEXT("ListTitle", "{0} (Message)"), FText::FromString(NodeNameStr));
+		}
+		else
+		{
+			FFormatNamedArguments Args;
+			Args.Add(TEXT("NodeName"), FText::FromString(NodeNameStr));
+			Args.Add(TEXT("OuterClassName"), FText::FromString(Function->GetOuterUClass()->GetName()));
+			NodeName = FText::Format(NSLOCTEXT("K2Node", "CallInterfaceContext", "{NodeName}\nUsing Interface {OuterClassName}"), Args);
+		}
 	}
 	else
 	{
diff --git a/Engine/Source/Runtime/Engine/Private/EdGraph/EdGraphNode_Comment.cpp b/Engine/Source/Runtime/Engine/Private/EdGraph/EdGraphNode_Comment.cpp
--- a/Engine/Source/Runtime/Engine/Private/EdGraph/EdGraphNode_Comment.cpp
+++ b/Engine/Source/Runtime/Engine/Private/EdGraph/EdGraphNode_Comment.cpp
@@ -74,7 +74,14 @@ FText UEdGraphNode_Comment::GetNodeTitle(ENodeTitleType::Type TitleType) const
 {
 	if(TitleType == ENodeTitleType::ListView)
 	{
-		return NSLOCTEXT("K2Node", "CommentBlock_Title", "Comment");
+		if (NodeComment.IsEmpty())
+		{
+			return NSLOCTEXT("K2Node", "NoComment_ListTitle", "Add Comment...");
+		}
+		else 
+		{
+			return NSLOCTEXT("K2Node", "CommentBlock_ListTitle", "Comment");
+		}		
 	}
 
 	return FText::FromString(NodeComment);
diff --git a/Engine/Source/Runtime/Engine/Private/EdGraph/EdGraphNode_Documentation.cpp b/Engine/Source/Runtime/Engine/Private/EdGraph/EdGraphNode_Documentation.cpp
--- a/Engine/Source/Runtime/Engine/Private/EdGraph/EdGraphNode_Documentation.cpp
+++ b/Engine/Source/Runtime/Engine/Private/EdGraph/EdGraphNode_Documentation.cpp
@@ -37,7 +37,17 @@ FString UEdGraphNode_Documentation::GetTooltip() const
 
 FText UEdGraphNode_Documentation::GetNodeTitle( ENodeTitleType::Type TitleType ) const
 {
-	return NSLOCTEXT( "K2Node", "DocumentationBlock_Title", "UDN Documentation Excerpt" );
+	FText NodeTitle;
+	if (TitleType == ENodeTitleType::ListView)
+	{
+		NodeTitle = NSLOCTEXT("K2Node", "DocumentationBlock_ListTitle", "Add Documentation Node...");
+	}
+	else
+	{
+		NodeTitle = NSLOCTEXT("K2Node", "DocumentationBlock_Title", "UDN Documentation Excerpt");
+	}
+
+	return NodeTitle;
 }
 
 FString UEdGraphNode_Documentation::GetPinNameOverride(const UEdGraphPin& Pin) const
