Removing "Call Function" root category (requested through the NUX trello board).

diff --git a/Engine/Source/Editor/BlueprintGraph/Private/K2ActionMenuBuilder.cpp b/Engine/Source/Editor/BlueprintGraph/Private/K2ActionMenuBuilder.cpp
--- a/Engine/Source/Editor/BlueprintGraph/Private/K2ActionMenuBuilder.cpp
+++ b/Engine/Source/Editor/BlueprintGraph/Private/K2ActionMenuBuilder.cpp
@@ -103,6 +103,7 @@ namespace K2ActionCategories
 	static FString const GameCategory	      = LOCTEXT("GameCategory",         "Game").ToString();
 	
 
+	static FString const RootFunctionCategory = TEXT("");
 	static FString const GenericFunctionCategory   = LOCTEXT("FunctionCategory",         "Call Function").ToString();
 	static FString const CallFuncForEachCategory   = LOCTEXT("CallFuncForEachCategory",  "Call Function For Each").ToString();
 	static FString const CallFuncOnPrefix	       = LOCTEXT("CallFuncOnPrefix",         "Call Function on ").ToString();
@@ -610,6 +611,10 @@ void FK2ActionMenuBuilder::AddSpawnInfoForFunction(const UFunction* Function, bo
 	if (!bIsParentContext && !bPlaceOnTop)
 	{
 		NodeCategory = UK2Node_CallFunction::GetDefaultCategoryForFunction(Function, BaseCategory);
+		if (NodeCategory.IsEmpty())
+		{
+			NodeCategory = K2ActionCategories::GenericFunctionCategory;
+		}
 	}
 
 	TSharedPtr<FEdGraphSchemaAction_K2NewNode> Action;
@@ -839,13 +844,13 @@ void FK2ActionMenuBuilder::GetContextAllowedNodeTypes(FBlueprintGraphActionListB
 			FunctionTypes |= UEdGraphSchema_K2::FT_Imperative;
 		}
 
-		GetFuncNodesForClass(ContextMenuBuilder, Class, Class, ContextMenuBuilder.CurrentGraph, FunctionTypes, K2ActionCategories::GenericFunctionCategory);
-		GetFuncNodesForClass(ContextMenuBuilder, NULL, Class, ContextMenuBuilder.CurrentGraph, FunctionTypes, K2ActionCategories::GenericFunctionCategory);
+		GetFuncNodesForClass(ContextMenuBuilder, Class, Class, ContextMenuBuilder.CurrentGraph, FunctionTypes, K2ActionCategories::RootFunctionCategory);
+		GetFuncNodesForClass(ContextMenuBuilder, NULL, Class, ContextMenuBuilder.CurrentGraph, FunctionTypes, K2ActionCategories::RootFunctionCategory);
 
 		// spawn actor from archetype node
 		if (!bIsConstructionScript)
 		{
-			AddSpawnActorNodeAction(ContextMenuBuilder, K2ActionCategories::GenericFunctionCategory);
+			AddSpawnActorNodeAction(ContextMenuBuilder, K2ActionCategories::RootFunctionCategory);
 
 			if ( Blueprint->SupportsInputEvents() )
 			{
@@ -853,7 +858,7 @@ void FK2ActionMenuBuilder::GetContextAllowedNodeTypes(FBlueprintGraphActionListB
 			}
 		}
 		
-		AddGetDataTableRowNodeAction(ContextMenuBuilder, K2ActionCategories::GenericFunctionCategory);
+		AddGetDataTableRowNodeAction(ContextMenuBuilder, K2ActionCategories::RootFunctionCategory);
 
 		// Add struct make/break nodes
 		GetStructActions( ContextMenuBuilder );
@@ -1067,7 +1072,7 @@ struct FClassDynamicCastMenuUtils
 
 	static void SpawnNode(UClass* TargetType, FBlueprintGraphActionListBuilder& ContextMenuBuilder)
 	{
-		static const FString CastCategory = K2ActionCategories::GenericFunctionCategory + "|" + K2ActionCategories::CastingCategory;
+		static const FString CastCategory = K2ActionCategories::RootFunctionCategory + "|" + K2ActionCategories::CastingCategory;
 		UK2Node_DynamicCast* CastNode = ContextMenuBuilder.CreateTemplateNode<UK2Node_ClassDynamicCast>();
 		CastNode->TargetType = TargetType;
 
@@ -1134,8 +1139,8 @@ void FK2ActionMenuBuilder::GetPinAllowedNodeTypes(FBlueprintGraphActionListBuild
 			const bool bWantBindableDelegates = (GraphType == GT_Ubergraph) || (GraphType == GT_Function) ||
 				((bAllowUnsafeCommands) && (Blueprint->BlueprintType == BPTYPE_MacroLibrary));
 
-			GetFuncNodesForClass(ContextMenuBuilder, Class, Class, ContextMenuBuilder.CurrentGraph, UEdGraphSchema_K2::FT_Imperative | UEdGraphSchema_K2::FT_Protected, K2ActionCategories::GenericFunctionCategory);
-			GetFuncNodesForClass(ContextMenuBuilder, NULL, Class, ContextMenuBuilder.CurrentGraph, UEdGraphSchema_K2::FT_Imperative | UEdGraphSchema_K2::FT_Protected, K2ActionCategories::GenericFunctionCategory);
+			GetFuncNodesForClass(ContextMenuBuilder, Class, Class, ContextMenuBuilder.CurrentGraph, UEdGraphSchema_K2::FT_Imperative | UEdGraphSchema_K2::FT_Protected, K2ActionCategories::RootFunctionCategory);
+			GetFuncNodesForClass(ContextMenuBuilder, NULL, Class, ContextMenuBuilder.CurrentGraph, UEdGraphSchema_K2::FT_Imperative | UEdGraphSchema_K2::FT_Protected, K2ActionCategories::RootFunctionCategory);
 
 			GetVariableGettersSettersForClass(ContextMenuBuilder, Class, FromPin.PinType, true, false, true, bWantBindableDelegates, bRequireDesiredPinType);
 
@@ -1151,7 +1156,7 @@ void FK2ActionMenuBuilder::GetPinAllowedNodeTypes(FBlueprintGraphActionListBuild
 
 			if (!K2Schema->IsConstructionScript(ContextMenuBuilder.CurrentGraph))
 			{
-				AddSpawnActorNodeAction(ContextMenuBuilder, K2ActionCategories::GenericFunctionCategory);
+				AddSpawnActorNodeAction(ContextMenuBuilder, K2ActionCategories::RootFunctionCategory);
 
 				if (Class->IsChildOf(AActor::StaticClass()))
 				{
@@ -1160,7 +1165,7 @@ void FK2ActionMenuBuilder::GetPinAllowedNodeTypes(FBlueprintGraphActionListBuild
 				}
 			}
 			
-			AddGetDataTableRowNodeAction(ContextMenuBuilder, K2ActionCategories::GenericFunctionCategory);
+			AddGetDataTableRowNodeAction(ContextMenuBuilder, K2ActionCategories::RootFunctionCategory);
 
 			// for output pins, add macro flow control as a connection option
 			if ( FromPin.Direction == EGPD_Output )
@@ -1205,7 +1210,7 @@ void FK2ActionMenuBuilder::GetPinAllowedNodeTypes(FBlueprintGraphActionListBuild
 							FunctionTypes |= UEdGraphSchema_K2::FT_Imperative;
 						}
 
-						GetFuncNodesForClass(ContextMenuBuilder, PinClass, PinClass, ContextMenuBuilder.CurrentGraph, FunctionTypes, K2ActionCategories::GenericFunctionCategory);
+						GetFuncNodesForClass(ContextMenuBuilder, PinClass, PinClass, ContextMenuBuilder.CurrentGraph, FunctionTypes, K2ActionCategories::RootFunctionCategory);
 
 						if( FromPin.PinType.bIsReference )
 						{
@@ -1238,7 +1243,7 @@ void FK2ActionMenuBuilder::GetPinAllowedNodeTypes(FBlueprintGraphActionListBuild
 
 				// Dynamic casts
 				{
-					FString const CastCategory = K2ActionCategories::GenericFunctionCategory + "|" + K2ActionCategories::CastingCategory;
+					FString const CastCategory = K2ActionCategories::RootFunctionCategory + "|" + K2ActionCategories::CastingCategory;
 
 					for (TObjectIterator<UClass> ClassIt; ClassIt; ++ClassIt)
 					{
@@ -2533,7 +2538,7 @@ void FK2ActionMenuBuilder::GetFuncNodesWithPinType(FBlueprintGraphActionListBuil
 				{
 					const bool bShowMakeOnTop = bUseNativeMake && Function->HasMetaData(NativeMakeFunc);
 					const bool bShowBrakeOnTop = bUseNativeBrake && Function->HasMetaData(NativeBrakeFunc);
-					FK2ActionMenuBuilder::AddSpawnInfoForFunction(Function, false, FFunctionTargetInfo(), FMemberReference(), K2ActionCategories::GenericFunctionCategory, K2Schema->AG_LevelReference, ContextMenuBuilder, false, bShowMakeOnTop || bShowBrakeOnTop);
+					FK2ActionMenuBuilder::AddSpawnInfoForFunction(Function, false, FFunctionTargetInfo(), FMemberReference(), K2ActionCategories::RootFunctionCategory, K2Schema->AG_LevelReference, ContextMenuBuilder, false, bShowMakeOnTop || bShowBrakeOnTop);
 				}
 			}
 		}
diff --git a/Engine/Source/Editor/BlueprintGraph/Private/K2Node_BaseAsyncTask.cpp b/Engine/Source/Editor/BlueprintGraph/Private/K2Node_BaseAsyncTask.cpp
--- a/Engine/Source/Editor/BlueprintGraph/Private/K2Node_BaseAsyncTask.cpp
+++ b/Engine/Source/Editor/BlueprintGraph/Private/K2Node_BaseAsyncTask.cpp
@@ -373,7 +373,7 @@ FName UK2Node_BaseAsyncTask::GetCornerIcon() const
 FText UK2Node_BaseAsyncTask::GetMenuCategory() const
 {	
 	UFunction* TargetFunction = GetFactoryFunction();
-	return FText::FromString(UK2Node_CallFunction::GetDefaultCategoryForFunction(TargetFunction, TEXT("Call Function")));
+	return FText::FromString(UK2Node_CallFunction::GetDefaultCategoryForFunction(TargetFunction, TEXT("")));
 }
 
 void UK2Node_BaseAsyncTask::GetMenuActions(FBlueprintActionDatabaseRegistrar& ActionRegistrar) const
diff --git a/Engine/Source/Editor/BlueprintGraph/Private/K2Node_EnumEquality.cpp b/Engine/Source/Editor/BlueprintGraph/Private/K2Node_EnumEquality.cpp
--- a/Engine/Source/Editor/BlueprintGraph/Private/K2Node_EnumEquality.cpp
+++ b/Engine/Source/Editor/BlueprintGraph/Private/K2Node_EnumEquality.cpp
@@ -204,12 +204,9 @@ void UK2Node_EnumEquality::GetMenuEntries(FGraphContextMenuBuilder& Context) con
 
 	if (bShowEnumEquality)
 	{
-		//@TODO: Promote the categories into the schema and remove this duplication with the code in EdGraphSchema_K2.cpp
-		const FString FunctionCategory(TEXT("Call Function"));
-
 		UK2Node* EnumNodeTemplate = Context.CreateTemplateNode<UK2Node_EnumEquality>();
 
-		const FString Category = FunctionCategory + TEXT("|Utilities| Enum");
+		const FString Category = TEXT("Utilities| Enum");
 		const FText MenuDesc = EnumNodeTemplate->GetNodeTitle(ENodeTitleType::ListView);
 		const FString Tooltip = EnumNodeTemplate->GetTooltip();
 		const FString Keywords = EnumNodeTemplate->GetKeywords();
diff --git a/Engine/Source/Editor/GameplayTagsEditor/Private/GameplayTagsK2Node_LiteralGameplayTag.cpp b/Engine/Source/Editor/GameplayTagsEditor/Private/GameplayTagsK2Node_LiteralGameplayTag.cpp
--- a/Engine/Source/Editor/GameplayTagsEditor/Private/GameplayTagsK2Node_LiteralGameplayTag.cpp
+++ b/Engine/Source/Editor/GameplayTagsEditor/Private/GameplayTagsK2Node_LiteralGameplayTag.cpp
@@ -42,12 +42,10 @@ bool UGameplayTagsK2Node_LiteralGameplayTag::CanCreateUnderSpecifiedSchema( cons
 void UGameplayTagsK2Node_LiteralGameplayTag::GetMenuEntries( FGraphContextMenuBuilder& ContextMenuBuilder ) const
 {
 	Super::GetMenuEntries(ContextMenuBuilder);
-	
-	const FString FunctionCategory(TEXT("Call Function"));
 
 	UK2Node* EnumNodeTemplate = ContextMenuBuilder.CreateTemplateNode<UGameplayTagsK2Node_LiteralGameplayTag>();
 
-	const FString Category = FunctionCategory + TEXT("|Gameplay Tags|Tag Container");
+	const FString Category = TEXT("Gameplay Tags|Tag Container");
 	const FText MenuDesc = EnumNodeTemplate->GetNodeTitle(ENodeTitleType::ListView);
 	const FString Tooltip = EnumNodeTemplate->GetTooltip();
 	const FString Keywords = EnumNodeTemplate->GetKeywords();
@@ -159,7 +157,7 @@ void UGameplayTagsK2Node_LiteralGameplayTag::GetMenuActions(FBlueprintActionData
 
 FText UGameplayTagsK2Node_LiteralGameplayTag::GetMenuCategory() const
 {
-	return LOCTEXT("ActionMenuCategory", "Gameplay Tags");
+	return LOCTEXT("ActionMenuCategory", "Gameplay Tags|Tag Container");
 }
 
 #undef LOCTEXT_NAMESPACE
diff --git a/Engine/Source/Editor/IntroTutorials/Private/K2Node_TutorialExcerptComplete.cpp b/Engine/Source/Editor/IntroTutorials/Private/K2Node_TutorialExcerptComplete.cpp
--- a/Engine/Source/Editor/IntroTutorials/Private/K2Node_TutorialExcerptComplete.cpp
+++ b/Engine/Source/Editor/IntroTutorials/Private/K2Node_TutorialExcerptComplete.cpp
@@ -39,12 +39,11 @@ void UK2Node_TutorialExcerptComplete::GetMenuEntries(FGraphContextMenuBuilder& C
 	UK2Node_TutorialExcerptComplete* TemplateNode = ContextMenuBuilder.CreateTemplateNode<UK2Node_TutorialExcerptComplete>();
 	TemplateNode->BindFunction();
 
-	const FString Category = LOCTEXT("TutorialExcerptComplete_Category", "Call Function").ToString();
-	const FString SubCategory = LOCTEXT("TutorialExcerptComplete_SubCategory", "Tutorials").ToString();
+	const FString Category = LOCTEXT("TutorialExcerptComplete_SubCategory", "Tutorials").ToString();
 	const FText MenuDesc = LOCTEXT("TutorialExcerptComplete_Desc", "Tutorial Excerpt Complete");
 	const FString Tooltip = LOCTEXT("TutorialExcerptComplete_Tooltip", "Add a node to complete a tutorial excerpt").ToString();
 
-	TSharedPtr<FEdGraphSchemaAction_K2NewNode> NodeAction = FK2ActionMenuBuilder::AddNewNodeAction(ContextMenuBuilder, Category + TEXT("|") + SubCategory, MenuDesc, Tooltip);
+	TSharedPtr<FEdGraphSchemaAction_K2NewNode> NodeAction = FK2ActionMenuBuilder::AddNewNodeAction(ContextMenuBuilder, Category, MenuDesc, Tooltip);
 	NodeAction->NodeTemplate = TemplateNode;
 }
 
