Integration of 388 from user moritz-wundke - New interpolation methods in the math library, exposed to blueprints via custom node

diff --git a/Engine/Source/Editor/BlueprintGraph/Classes/K2Node_EaseFunction.h b/Engine/Source/Editor/BlueprintGraph/Classes/K2Node_EaseFunction.h
new file mode 100644
--- /dev/null
+++ b/Engine/Source/Editor/BlueprintGraph/Classes/K2Node_EaseFunction.h
@@ -0,0 +1,114 @@
+// Copyright 1998-2014 Epic Games, Inc. All Rights Reserved.
+
+#pragma once
+#include "K2Node.h"
+#include "K2Node_EaseFunction.generated.h"
+
+/////////////////////////////////////////////////////
+// UK2Node_EaseFunction
+
+USTRUCT()
+struct FCustomPin
+{
+	GENERATED_USTRUCT_BODY()
+
+	/** Name of the pin within the current node */
+	UPROPERTY()
+	FString PinName;
+
+	/** Name of the pin of the call function */
+	UPROPERTY()
+	FString CallFuncPinName;
+
+	/** If TRUE this is a custom pin that does come from the type of wildcard pin connected */
+	UPROPERTY()
+	bool bValuePin;
+};
+
+UCLASS(MinimalAPI)
+class UK2Node_EaseFunction : public UK2Node
+{
+	GENERATED_UCLASS_BODY()
+
+	// Begin UEdGraphNode interface.
+	virtual void AllocateDefaultPins() override;
+	virtual FText GetNodeTitle(ENodeTitleType::Type TitleType) const override;
+	void GetContextMenuActions(const FGraphNodeContextMenuBuilder& Context) const override;
+	virtual void PinDefaultValueChanged(UEdGraphPin* Pin) override;
+	virtual void PinTypeChanged(UEdGraphPin* Pin) override;
+	virtual FString GetTooltip() const override;
+	virtual FName GetPaletteIcon(FLinearColor& OutColor) const override;
+	// End UEdGraphNode interface.
+
+	// Begin UK2Node interface.
+	virtual void ExpandNode(class FKismetCompilerContext& CompilerContext, UEdGraph* SourceGraph) override;
+	virtual void NotifyPinConnectionListChanged(UEdGraphPin* Pin) override;
+	virtual void PostReconstructNode() override;
+	virtual bool IsConnectionDisallowed(const UEdGraphPin* MyPin, const UEdGraphPin* OtherPin, FString& OutReason) const override;
+	virtual void GetMenuActions(FBlueprintActionDatabaseRegistrar& ActionRegistrar) const override;
+	virtual FText GetMenuCategory() const override;
+	// End UK2Node interface.
+
+private:
+	/** Reset the A, B and Result pins to its initial wildcard state */
+	void ResetToWildcards();
+	
+	/** returns the EaseFunc pin */
+	UEdGraphPin* GetEaseFuncPin() const;
+
+	/** Returns TRUE if the selected curve can feature customization, FALSE otherwise */
+	bool CanCustomizeCurve() const
+	{
+		const auto EaseFuncPin = GetEaseFuncPin();
+		return EaseFuncPin && EaseFuncPin->LinkedTo.Num() == 0;
+	}
+
+	/**
+	* Notification from the editor that the user wants to change the PinType on a selected pin
+	*
+	* @param Pin	The pin the user wants to change the type for
+	*/
+	void ChangePinType(UEdGraphPin* Pin);
+
+	/** Generates the available custom pins for the current ease function if applicable */
+	bool ConditionalGenerateCustomPins();
+
+	/** Try to remove all custom pins that exist */
+	bool RemoveCustomPins(const bool bRemoveValuePins);
+
+	/** Generate custom pins used by the current ease function */
+	bool GenerateCustomPins(const int32 NewEasingFunc);
+
+	/** Update MyPin with data from another pin. Return FALSE in case nothing was changed, TRUE otherwise. */
+	bool UpdatePin(UEdGraphPin* MyPin, UEdGraphPin* OtherPin);
+
+	/**
+	* Takes the specified "MutatablePin" and sets its 'PinToolTip' field (according
+	* to the specified description)
+	*
+	* @param   MutatablePin	The pin you want to set tool-tip text on
+	* @param   PinDescription	A string describing the pin's purpose
+	*/
+	void SetPinToolTip(UEdGraphPin& MutatablePin, const FText& PinDescription) const;
+
+	/** Called right after the wildcard pins has been assigned a value */
+	void GenerateExtraPins();
+
+	/** Custom pins that are curently used */
+	UPROPERTY()
+	TArray<FCustomPin> CustomPinNames;
+
+	/** Previous easing function, used to only recreate custom pins if required */
+	UPROPERTY()
+	int32 OldEasingFunc;
+
+	/** Name of the kismet ease function to be called */
+	UPROPERTY()
+	FString EaseFunctionName;
+
+	/** Tooltip text for this node. */
+	FString NodeTooltip;
+
+	/** The "EaseFunc" input pin, used to enable/disable easing customization */
+	UEdGraphPin* CachedEaseFuncPin;
+};
diff --git a/Engine/Source/Editor/BlueprintGraph/Private/K2ActionMenuBuilder.cpp b/Engine/Source/Editor/BlueprintGraph/Private/K2ActionMenuBuilder.cpp
--- a/Engine/Source/Editor/BlueprintGraph/Private/K2ActionMenuBuilder.cpp
+++ b/Engine/Source/Editor/BlueprintGraph/Private/K2ActionMenuBuilder.cpp
@@ -84,6 +84,8 @@ namespace K2ActionCategories
 	static FString const MakeStructSubCategory  = LOCTEXT("MakeStruct",      "Make Struct").ToString();
 	static FString const BreakStructSubCategory = LOCTEXT("BreakStruct",     "Break Struct").ToString();
 	static FString const TextSubCategory        = LOCTEXT("TextCategory",    "Text").ToString();
+	static FString const MathCategory           = LOCTEXT("MathCategory",    "Math").ToString();
+	static FString const InterpSubCategory      = LOCTEXT("InterpCategory", "Interpolation").ToString();
 
 	static FString const ArrayCategory         = UtilitiesCategory + SubCategoryDelim + ArraySubCategory;
 	static FString const CastingCategory       = UtilitiesCategory + SubCategoryDelim + CastingSubCategory;
@@ -97,6 +99,8 @@ namespace K2ActionCategories
 	static FString const TextCategory          = UtilitiesCategory + SubCategoryDelim + TextSubCategory;
 	static FString const FallbackMacroCategory = UtilitiesCategory;
 
+	static FString const InterpCategory = MathCategory + SubCategoryDelim + InterpSubCategory;
+
 
 	static FString const AddComponentCategory = LOCTEXT("AddComponentCategory", "Add Component").ToString();
 	static FString const VariablesCategory    = LOCTEXT("VariablesCategory",    "Variables").ToString();
@@ -792,6 +796,11 @@ void FK2ActionMenuBuilder::GetPaletteActions(FBlueprintPaletteListBuilder& Actio
 	GetMacroTools(ActionMenuBuilder, true, NULL);
 
 	GetRequestedAction<UK2Node_FormatText>(ActionMenuBuilder, K2ActionCategories::TextCategory);
+
+	{
+		// Actions that are part of Math
+		GetRequestedAction<UK2Node_EaseFunction>(ActionMenuBuilder, K2ActionCategories::InterpCategory);
+	}
 }
 
 /*******************************************************************************
@@ -835,6 +844,11 @@ void FK2ActionMenuBuilder::GetContextAllowedNodeTypes(FBlueprintGraphActionListB
 
 		GetRequestedAction<UK2Node_FormatText>(ContextMenuBuilder, K2ActionCategories::TextCategory);
 
+		{
+			// Actions that are part of Math
+			GetRequestedAction<UK2Node_EaseFunction>(ContextMenuBuilder, K2ActionCategories::InterpCategory);
+		}
+
 		GetAnimNotifyMenuItems(ContextMenuBuilder);
 		GetVariableGettersSettersForClass(ContextMenuBuilder, Class, true, bAllowEvents);
 
@@ -1380,6 +1394,14 @@ void FK2ActionMenuBuilder::GetPinAllowedNodeTypes(FBlueprintGraphActionListBuild
 		{
 			GetRequestedAction<UK2Node_FormatText>(ContextMenuBuilder, K2ActionCategories::ArrayCategory);
 		}
+		else if (FromPin.PinType.PinCategory == K2Schema->PC_Float ||
+			(FromPin.PinType.PinCategory == K2Schema->PC_Struct &&
+			(FromPin.PinType.PinSubCategoryObject.Get()->GetName() == TEXT("Vector") ||
+			FromPin.PinType.PinSubCategoryObject.Get()->GetName() == TEXT("Rotator") ||
+			FromPin.PinType.PinSubCategoryObject.Get()->GetName() == TEXT("Transform"))))
+		{
+			GetRequestedAction<UK2Node_EaseFunction>(ContextMenuBuilder, K2ActionCategories::InterpCategory);
+		}
 
 		UEnum* Enum = Cast<UEnum>(FromPin.PinType.PinSubCategoryObject.Get());
 		if ((FromPin.PinType.PinCategory == K2Schema->PC_Byte) && Enum)
diff --git a/Engine/Source/Editor/BlueprintGraph/Private/K2Node_EaseFunction.cpp b/Engine/Source/Editor/BlueprintGraph/Private/K2Node_EaseFunction.cpp
new file mode 100644
--- /dev/null
+++ b/Engine/Source/Editor/BlueprintGraph/Private/K2Node_EaseFunction.cpp
@@ -0,0 +1,601 @@
+// Copyright 1998-2014 Epic Games, Inc. All Rights Reserved.
+
+#include"BlueprintGraphPrivatePCH.h"
+
+#include "BlueprintActionDatabaseRegistrar.h"
+#include "BlueprintGraphPrivatePCH.h"
+#include "BlueprintNodeSpawner.h"
+#include "EditorCategoryUtils.h"
+#include "Editor/GraphEditor/Public/DiffResults.h"
+#include "K2Node_EaseFunction.h"
+#include "Kismet2NameValidators.h"
+#include "KismetCompiler.h"
+#include "ScopedTransaction.h"
+
+#define LOCTEXT_NAMESPACE "K2Node_EaseFunction"
+
+struct FEaseFunctionNodeHelper
+{
+	static const FString& GetEaseFuncPinName()
+	{
+		static const FString EaseFuncPinName(TEXT("Function"));
+		return EaseFuncPinName;
+	}
+
+	static const FString& GetAlphaPinName()
+	{
+		static const FString AlphaPinName(TEXT("Alpha"));
+		return AlphaPinName;
+	}
+
+	static const FString& GetAPinName()
+	{
+		static const FString APinFloatName(TEXT("A"));
+		return APinFloatName;
+	}
+
+	static const FString& GetBPinName()
+	{
+		static const FString BPinFloatName(TEXT("B"));
+		return BPinFloatName;
+	}
+
+	static const FString& GetResultPinName()
+	{
+		static const FString ResultPinName(TEXT("Result"));
+		return ResultPinName;
+	}
+
+	static const FString& GetBlendExpPinName()
+	{
+		static const FString BlendExpPinName(TEXT("BlendExp"));
+		return BlendExpPinName;
+	}
+
+	static const FString& GetStepsPinName()
+	{
+		static const FString StepsPinName(TEXT("Steps"));
+		return StepsPinName;
+	}
+
+	static const FString& GetShortestPathPinName()
+	{
+		static const FString ShortestPathPinName(TEXT("ShortestPath"));
+		return ShortestPathPinName;
+	}
+
+};
+
+UK2Node_EaseFunction::UK2Node_EaseFunction(const class FPostConstructInitializeProperties& PCIP)
+	: Super(PCIP)
+	, CachedEaseFuncPin(NULL)
+{
+	NodeTooltip = LOCTEXT("NodeTooltip", "Interpolates from value A to value B using a user specified easing function").ToString();
+	OldEasingFunc = INDEX_NONE;
+	EaseFunctionName = TEXT("");
+}
+
+void UK2Node_EaseFunction::AllocateDefaultPins()
+{
+	Super::AllocateDefaultPins();
+
+	const UEdGraphSchema_K2* K2Schema = GetDefault<UEdGraphSchema_K2>();
+
+	// Add the first pin representing all available easing functions. If EEasingFunc changes its name this will fail a runtime check!
+	UEnum * EaseFuncEnum = FindObject<UEnum>(ANY_PACKAGE, TEXT("EEasingFunc"), true);
+	check(EaseFuncEnum != NULL);
+	CachedEaseFuncPin = CreatePin(EGPD_Input, K2Schema->PC_Byte, TEXT(""), EaseFuncEnum, false, false, FEaseFunctionNodeHelper::GetEaseFuncPinName());
+	SetPinToolTip(*CachedEaseFuncPin, LOCTEXT("EaseFunsPinDescription", "Specifies the desired ease function to be applied. If connected no customization is possible."));
+
+	// Make sure that the default value is set correctly if none has been set
+	K2Schema->SetPinDefaultValue(CachedEaseFuncPin);
+
+	UEdGraphPin *AlphaPin = CreatePin(EGPD_Input, K2Schema->PC_Float, TEXT(""), NULL, false, false, FEaseFunctionNodeHelper::GetAlphaPinName());
+	SetPinToolTip(*AlphaPin, LOCTEXT("AlphaPinTooltip", "Alpha value used to specify the easing in time."));
+
+	// Add wildcard pins for A, B and the return Pin
+	UEdGraphPin *APin = CreatePin(EGPD_Input, K2Schema->PC_Wildcard, TEXT(""), NULL, false, false, FEaseFunctionNodeHelper::GetAPinName());
+	SetPinToolTip(*APin, LOCTEXT("APinDescription", "Easing start value"));
+
+	UEdGraphPin *BPin = CreatePin(EGPD_Input, K2Schema->PC_Wildcard, TEXT(""), NULL, false, false, FEaseFunctionNodeHelper::GetBPinName());
+	SetPinToolTip(*BPin, LOCTEXT("BPinDescription", "Easing end value"));
+
+	UEdGraphPin *ResultPin = CreatePin(EGPD_Output, K2Schema->PC_Wildcard, TEXT(""), NULL, false, false, FEaseFunctionNodeHelper::GetResultPinName());
+	SetPinToolTip(*ResultPin, LOCTEXT("ResultPinDescription", "Easing result value"));
+
+	// Try to generate any custom pins that applies
+	ConditionalGenerateCustomPins();
+}
+
+FText UK2Node_EaseFunction::GetNodeTitle(ENodeTitleType::Type TitleType) const
+{
+	return LOCTEXT("EaseFunction_Title", "Ease");
+}
+
+FString UK2Node_EaseFunction::GetTooltip() const
+{
+	return NodeTooltip;
+}
+
+FName UK2Node_EaseFunction::GetPaletteIcon(FLinearColor& OutColor) const
+{
+	// The function icon seams the best choice!
+	OutColor = GetNodeTitleColor();
+	return TEXT("Kismet.AllClasses.FunctionIcon");
+}
+
+void UK2Node_EaseFunction::SetPinToolTip(UEdGraphPin& MutatablePin, const FText& PinDescription) const
+{
+	MutatablePin.PinToolTip = UEdGraphSchema_K2::TypeToString(MutatablePin.PinType);
+
+	UEdGraphSchema_K2 const* const K2Schema = Cast<const UEdGraphSchema_K2>(GetSchema());
+	if (K2Schema != nullptr)
+	{
+		MutatablePin.PinToolTip += TEXT(" ");
+		MutatablePin.PinToolTip += K2Schema->GetPinDisplayName(&MutatablePin);
+	}
+
+	MutatablePin.PinToolTip += FString(TEXT("\n")) + PinDescription.ToString();
+}
+
+void UK2Node_EaseFunction::NotifyPinConnectionListChanged(UEdGraphPin* Pin)
+{
+	Super::NotifyPinConnectionListChanged(Pin);
+
+	const auto EaseFuncPin = GetEaseFuncPin();
+	if (Pin == EaseFuncPin)
+	{
+		bool bDirty = false;
+
+		// Check in which state we are at the moment
+		if (CanCustomizeCurve())
+		{
+			OldEasingFunc = INDEX_NONE;
+			bDirty = ConditionalGenerateCustomPins();
+		}
+		else
+		{
+			bDirty = RemoveCustomPins(false);
+		}
+
+		if (bDirty)
+		{
+			GetGraph()->NotifyGraphChanged();
+		}
+	}
+	else
+	{
+		PinTypeChanged(Pin);
+	}
+}
+
+
+void UK2Node_EaseFunction::PinDefaultValueChanged(UEdGraphPin* Pin)
+{
+	const auto EaseFuncPin = GetEaseFuncPin();
+	if (Pin == EaseFuncPin && ConditionalGenerateCustomPins())
+	{
+		GetGraph()->NotifyGraphChanged();
+	}
+}
+
+void UK2Node_EaseFunction::PostReconstructNode()
+{
+	Super::PostReconstructNode();
+
+	CustomPinNames.Empty();
+	// Check in which state we are at the moment
+	if (CanCustomizeCurve())
+	{
+		OldEasingFunc = INDEX_NONE;
+		ConditionalGenerateCustomPins();
+	}
+	GenerateExtraPins();
+
+	// Find a pin that has connections to use to jumpstart the wildcard process
+	for (int32 PinIndex = 0; PinIndex < Pins.Num(); ++PinIndex)
+	{
+		if (Pins[PinIndex]->PinName == FEaseFunctionNodeHelper::GetAPinName() ||
+			Pins[PinIndex]->PinName == FEaseFunctionNodeHelper::GetBPinName() ||
+			Pins[PinIndex]->PinName == FEaseFunctionNodeHelper::GetResultPinName())
+		{
+			// Take default pin values into account in case we can securly convert the string default value into a PinType
+			// this is currently not the case but could be considered in the future.
+			if (Pins[PinIndex]->LinkedTo.Num() > 0)
+			{
+				PinTypeChanged(Pins[PinIndex]);
+				break;
+			}
+		}
+	}
+}
+
+bool UK2Node_EaseFunction::IsConnectionDisallowed(const UEdGraphPin* MyPin, const UEdGraphPin* OtherPin, FString& OutReason) const
+{
+	const UEdGraphSchema_K2* K2Schema = GetDefault<UEdGraphSchema_K2>();
+
+	// Check the pin name and see if it matches on of our three base pins
+	if (MyPin->PinName == FEaseFunctionNodeHelper::GetAPinName() || 
+		MyPin->PinName == FEaseFunctionNodeHelper::GetBPinName() || 
+		MyPin->PinName == FEaseFunctionNodeHelper::GetResultPinName())
+	{
+		const bool bConnectionOk = (OtherPin->PinType.PinCategory == K2Schema->PC_Float ||
+			(OtherPin->PinType.PinCategory == K2Schema->PC_Struct &&
+			(OtherPin->PinType.PinSubCategoryObject.Get()->GetName() == TEXT("Vector") ||
+			OtherPin->PinType.PinSubCategoryObject.Get()->GetName() == TEXT("Rotator") ||
+			OtherPin->PinType.PinSubCategoryObject.Get()->GetName() == TEXT("Transform"))));
+		if (!bConnectionOk)
+		{
+			OutReason = LOCTEXT("PinConnectionDisallowed", "Pin type is not supported by function.").ToString();
+			return true;
+		}
+	}
+
+	return Super::IsConnectionDisallowed(MyPin, OtherPin, OutReason);
+}
+
+void UK2Node_EaseFunction::GetMenuActions(FBlueprintActionDatabaseRegistrar& ActionRegistrar) const
+{
+	UBlueprintNodeSpawner* NodeSpawner = UBlueprintNodeSpawner::Create(GetClass());
+	check(NodeSpawner != nullptr);
+
+	ActionRegistrar.AddBlueprintAction(NodeSpawner);
+}
+
+FText UK2Node_EaseFunction::GetMenuCategory() const
+{
+	return FEditorCategoryUtils::BuildCategoryString(FCommonEditorCategory::Input, LOCTEXT("InterpCategory", "Interpolation"));
+}
+
+void UK2Node_EaseFunction::ChangePinType(UEdGraphPin* Pin)
+{
+	PinTypeChanged(Pin);
+}
+
+void UK2Node_EaseFunction::PinTypeChanged(UEdGraphPin* Pin)
+{
+	const UEdGraphSchema_K2* Schema = GetDefault<UEdGraphSchema_K2>();
+	bool bChanged = false;
+
+	if (Pin->PinName == FEaseFunctionNodeHelper::GetAPinName() ||
+		Pin->PinName == FEaseFunctionNodeHelper::GetBPinName() ||
+		Pin->PinName == FEaseFunctionNodeHelper::GetResultPinName())
+	{
+		// Get pin refs
+		UEdGraphPin* APin = FindPin(FEaseFunctionNodeHelper::GetAPinName());
+		UEdGraphPin* BPin = FindPin(FEaseFunctionNodeHelper::GetBPinName());
+		UEdGraphPin* ResultPin = FindPin(FEaseFunctionNodeHelper::GetResultPinName());
+
+		// Propagate the type change or reset to wildcard PinType
+		if (Pin->LinkedTo.Num() > 0)
+		{
+			UEdGraphPin* InstigatorPin = Pin->LinkedTo[0];
+
+			bChanged |= UpdatePin(APin, InstigatorPin);
+			bChanged |= UpdatePin(BPin, InstigatorPin);
+			bChanged |= UpdatePin(ResultPin, InstigatorPin);
+
+			if (bChanged)
+			{
+				// Just in case we switch to an invalid function clean it first
+				EaseFunctionName = TEXT("");
+
+				// Generate the right function name
+				if (InstigatorPin->PinType.PinCategory == Schema->PC_Float)
+				{
+					EaseFunctionName = TEXT("Ease");
+				}
+				else if (InstigatorPin->PinType.PinCategory == Schema->PC_Struct)
+				{
+					if (InstigatorPin->PinType.PinSubCategoryObject.Get()->GetName() == TEXT("Vector"))
+					{
+						EaseFunctionName = TEXT("VEase");
+					}
+					else if (InstigatorPin->PinType.PinSubCategoryObject.Get()->GetName() == TEXT("Rotator"))
+					{
+						EaseFunctionName = TEXT("REase");
+					}
+					else if (InstigatorPin->PinType.PinSubCategoryObject.Get()->GetName() == TEXT("Transform"))
+					{
+						EaseFunctionName = TEXT("TEase");
+					}
+				}
+			}
+		}
+		else
+		{	
+			if (APin->GetDefaultAsString().IsEmpty() && APin->LinkedTo.Num() == 0 &&
+				BPin->GetDefaultAsString().IsEmpty() && BPin->LinkedTo.Num() == 0 &&
+				ResultPin->LinkedTo.Num() == 0)
+			{
+				// Restore wild card pin
+				APin->PinType.PinCategory = Schema->PC_Wildcard;
+				APin->PinType.PinSubCategory = TEXT("");
+				APin->PinType.PinSubCategoryObject = NULL;
+
+				// Propagate change
+				UpdatePin(BPin, APin);
+				UpdatePin(ResultPin, APin);
+
+				// Make sure the function name is nulled out
+				EaseFunctionName = TEXT("");
+				bChanged = true;			
+			}
+		}
+
+		// Pin connections and data changed in some way
+		if (bChanged)
+		{
+			SetPinToolTip(*APin, LOCTEXT("APinDescription", "Easing start value"));
+			SetPinToolTip(*BPin, LOCTEXT("BPinDescription", "Easing end value"));
+			SetPinToolTip(*ResultPin, LOCTEXT("ResultPinDescription", "Easing result value"));
+
+			// Let our subclasses generate some pins if required, this way we can add any aditional pins required by some types for examples
+			GenerateExtraPins();
+
+			// Let the graph know to refresh
+			GetGraph()->NotifyGraphChanged();
+
+			UBlueprint* Blueprint = GetBlueprint();
+			if (!Blueprint->bBeingCompiled)
+			{
+				FBlueprintEditorUtils::MarkBlueprintAsModified(Blueprint);
+				Blueprint->BroadcastChanged();
+			}
+		}
+	}
+
+	Super::PinTypeChanged(Pin);
+}
+
+bool UK2Node_EaseFunction::UpdatePin(UEdGraphPin* MyPin, UEdGraphPin* OtherPin)
+{
+	// If the data type changed break links first
+	if (MyPin->PinType != OtherPin->PinType)
+	{
+		MyPin->PinType = OtherPin->PinType;
+		return true;
+	}	
+	return false;
+}
+
+void UK2Node_EaseFunction::ExpandNode(class FKismetCompilerContext& CompilerContext, UEdGraph* SourceGraph)
+{
+	Super::ExpandNode(CompilerContext, SourceGraph);
+
+	if (CompilerContext.bIsFullCompile)
+	{
+		/**
+		At the end of this, the UK2Node_EaseFunction will not be a part of the Blueprint, it merely handles connecting
+		the other nodes into the Blueprint.
+		*/
+
+		UFunction* function = UKismetMathLibrary::StaticClass()->FindFunctionByName(*EaseFunctionName);
+		if (function == NULL)
+		{
+			CompilerContext.MessageLog.Error(*LOCTEXT("InvalidFunctionName", "BaseAsyncTask: Type not supported or not initialized. @@").ToString(), this);
+			return;
+		}
+
+		const UEdGraphSchema_K2* Schema = CompilerContext.GetSchema();
+
+		// The call function does all the real work, each child class implementing easing for  a given type provides
+		// the name of the desired function
+		UK2Node_CallFunction* CallFunction = CompilerContext.SpawnIntermediateNode<UK2Node_CallFunction>(this, SourceGraph);
+		
+		CallFunction->SetFromFunction(function);
+		CallFunction->AllocateDefaultPins();
+		CompilerContext.MessageLog.NotifyIntermediateObjectCreation(CallFunction, this);
+
+		// Move the ease function and the alpha connections from us to the call function
+		CompilerContext.MovePinLinksToIntermediate(*FindPin(FEaseFunctionNodeHelper::GetEaseFuncPinName()), *CallFunction->FindPin(TEXT("EasingFunc")));
+		CompilerContext.MovePinLinksToIntermediate(*FindPin(FEaseFunctionNodeHelper::GetAlphaPinName()), *CallFunction->FindPin(TEXT("Alpha")));
+
+		// Move base connections to the call function's connections
+		CompilerContext.MovePinLinksToIntermediate(*FindPin(FEaseFunctionNodeHelper::GetAPinName()), *CallFunction->FindPin(TEXT("A")));
+		CompilerContext.MovePinLinksToIntermediate(*FindPin(FEaseFunctionNodeHelper::GetBPinName()), *CallFunction->FindPin(TEXT("B")));
+		CompilerContext.MovePinLinksToIntermediate(*FindPin(FEaseFunctionNodeHelper::GetResultPinName()), *CallFunction->GetReturnValuePin());
+
+		// Now move the custom pins to their new locations
+		for (int32 ArgIdx = 0; ArgIdx < CustomPinNames.Num(); ++ArgIdx)
+		{
+			CompilerContext.MovePinLinksToIntermediate(*FindPin(CustomPinNames[ArgIdx].PinName), *CallFunction->FindPin(CustomPinNames[ArgIdx].CallFuncPinName));
+		}
+
+		// Cleanup links to ourselfs and we are done!
+		BreakAllNodeLinks();
+	}
+}
+
+UEdGraphPin* UK2Node_EaseFunction::GetEaseFuncPin() const
+{
+	if (!CachedEaseFuncPin)
+	{
+		const_cast<UK2Node_EaseFunction*>(this)->CachedEaseFuncPin = FindPinChecked(FEaseFunctionNodeHelper::GetEaseFuncPinName());
+	}
+	return CachedEaseFuncPin;
+}
+
+bool UK2Node_EaseFunction::ConditionalGenerateCustomPins()
+{
+	// Only generate custom values if we are able to use them
+	if (CanCustomizeCurve())
+	{
+		const auto EaseFuncPin = GetEaseFuncPin();
+		if (EaseFuncPin)
+		{
+			UEnum * Enum = FindObject<UEnum>(ANY_PACKAGE, TEXT("EEasingFunc"), true);
+			check(Enum != NULL);
+			const int32 EnumIndex = Enum->FindEnumIndex(*EaseFuncPin->DefaultValue);
+			return GenerateCustomPins(EnumIndex);
+		}
+	}
+	return false;
+}
+
+bool UK2Node_EaseFunction::RemoveCustomPins(const bool bRemoveValuePins)
+{
+	bool bResult = false;
+	// Remove old pins
+	for (int32 ArgIdx = CustomPinNames.Num()-1; ArgIdx >= 0; --ArgIdx)
+	{
+		UEdGraphPin* CustomPin = FindPin(CustomPinNames[ArgIdx].PinName);
+
+		if (CustomPinNames[ArgIdx].bValuePin == bRemoveValuePins)
+		{
+			// Make sure the pin exists
+			if (CustomPin)
+			{
+				CustomPin->BreakAllPinLinks();
+				Pins.Remove(CustomPin);
+				CustomPinNames.RemoveAt(ArgIdx, 1);
+				bResult = true;
+
+			}
+			else
+			{
+				// When loaded from the disk this function is called before pins have been created:
+				CustomPinNames.RemoveAt(ArgIdx, 1);
+			}
+		}
+	}
+	return bResult;
+}
+
+void UK2Node_EaseFunction::GenerateExtraPins()
+{
+	bool bResult = false;
+	const UEdGraphSchema_K2* K2Schema = Cast<const UEdGraphSchema_K2>(GetSchema());
+
+	// First try to remove old pins
+	RemoveCustomPins(true);
+
+	// Add pins based on the pin type
+	const UEdGraphPin* APin = FindPin(FEaseFunctionNodeHelper::GetAPinName());
+	const bool bIsRotator = APin->PinType.PinCategory == K2Schema->PC_Struct && APin->PinType.PinSubCategoryObject.Get()->GetName() == TEXT("Rotator");
+	if (bIsRotator)
+	{
+		// Easing a rotator lets us choose which path in the rotation to choose
+		UEdGraphPin *ShortestPathPin = CreatePin(EGPD_Input, K2Schema->PC_Boolean, TEXT(""), NULL, false, false, FEaseFunctionNodeHelper::GetShortestPathPinName());
+		SetPinToolTip(*ShortestPathPin, LOCTEXT("ShortestPathPinTooltip", "When interpolating the shortest path should be taken."));
+		if (ShortestPathPin->DefaultValue.Len() == 0)
+		{
+			ShortestPathPin->DefaultValue = TEXT("true");
+		}
+
+		// Add the pin
+		CustomPinNames.AddZeroed();
+		CustomPinNames[CustomPinNames.Num() - 1].PinName = FEaseFunctionNodeHelper::GetShortestPathPinName();
+		CustomPinNames[CustomPinNames.Num() - 1].CallFuncPinName = TEXT("bShortestPath");
+		CustomPinNames[CustomPinNames.Num() - 1].bValuePin = true;
+	}
+}
+
+bool UK2Node_EaseFunction::GenerateCustomPins(const int32 NewEasingFunc)
+{
+	// Early exit in case no changes are required
+	if (OldEasingFunc == NewEasingFunc ||
+		((OldEasingFunc == EEasingFunc::EaseIn || OldEasingFunc == EEasingFunc::EaseOut || OldEasingFunc == EEasingFunc::EaseInOut) &&
+		(NewEasingFunc == EEasingFunc::EaseIn || NewEasingFunc == EEasingFunc::EaseOut || NewEasingFunc == EEasingFunc::EaseInOut)))
+	{
+		OldEasingFunc = NewEasingFunc;
+		return false;
+	}
+
+	bool bResult = false;
+	const UEdGraphSchema_K2* K2Schema = Cast<const UEdGraphSchema_K2>(GetSchema());
+
+	// First try to remove old pins
+	bResult |= RemoveCustomPins(false);
+
+	// Add new pins (if applicable)
+	UEdGraphPin *CustomPin = NULL;
+	switch (NewEasingFunc)
+	{
+	case EEasingFunc::EaseIn:
+	case EEasingFunc::EaseOut:
+	case EEasingFunc::EaseInOut:
+		// Generate the BlendExp custom pin
+		CustomPin = CreatePin(EGPD_Input, K2Schema->PC_Float, TEXT(""), NULL, false, false, FEaseFunctionNodeHelper::GetBlendExpPinName());
+		SetPinToolTip(*CustomPin, LOCTEXT("BlendExpPinDescription", "Blend Exponent for basic ease functions"));
+		if (CustomPin->DefaultValue.Len() == 0)
+		{
+			CustomPin->DefaultValue = TEXT("2.0");
+		}
+
+		// Add the pin
+		CustomPinNames.AddZeroed();
+		CustomPinNames[CustomPinNames.Num() - 1].PinName = FEaseFunctionNodeHelper::GetBlendExpPinName();
+		CustomPinNames[CustomPinNames.Num() - 1].CallFuncPinName = FEaseFunctionNodeHelper::GetBlendExpPinName();
+		CustomPinNames[CustomPinNames.Num() - 1].bValuePin = false;
+		bResult |= true;
+		break;
+
+	case EEasingFunc::Step:
+		// Generate the BlendExp custom pin
+		CustomPin = CreatePin(EGPD_Input, K2Schema->PC_Int, TEXT(""), NULL, false, false, FEaseFunctionNodeHelper::GetStepsPinName());
+		SetPinToolTip(*CustomPin, LOCTEXT("StepsPinDescription", "Number of steps required to go from A to B"));
+		if (CustomPin->DefaultValue.Len() == 0)
+		{
+			CustomPin->DefaultValue = TEXT("2");
+		}
+
+		// Add the pin
+		CustomPinNames.AddZeroed();
+		CustomPinNames[CustomPinNames.Num() - 1].PinName = FEaseFunctionNodeHelper::GetStepsPinName();
+		CustomPinNames[CustomPinNames.Num() - 1].CallFuncPinName = FEaseFunctionNodeHelper::GetStepsPinName();
+		CustomPinNames[CustomPinNames.Num() - 1].bValuePin = false;
+		bResult |= true;
+		break;
+	}
+
+	OldEasingFunc = NewEasingFunc;
+	return bResult;
+}
+
+void UK2Node_EaseFunction::ResetToWildcards()
+{
+	FScopedTransaction Transaction(LOCTEXT("ResetToDefaultsTx", "ResetToDefaults"));
+	Modify();
+
+	// Get pin refs
+	UEdGraphPin* APin = FindPin(FEaseFunctionNodeHelper::GetAPinName());
+	UEdGraphPin* BPin = FindPin(FEaseFunctionNodeHelper::GetBPinName());
+	UEdGraphPin* ResultPin = FindPin(FEaseFunctionNodeHelper::GetResultPinName());
+	
+	// Set all to defaults and break links
+	APin->DefaultValue = TEXT("");
+	BPin->DefaultValue = TEXT("");
+	APin->BreakAllPinLinks();
+	BPin->BreakAllPinLinks();
+	ResultPin->BreakAllPinLinks();
+
+	// Do the rest of the work, we will not recompile because the wildcard pins will prevent it
+	PinTypeChanged(APin);
+}
+
+void UK2Node_EaseFunction::GetContextMenuActions(const FGraphNodeContextMenuBuilder& Context) const
+{
+	Super::GetContextMenuActions(Context);
+
+	if (!Context.bIsDebugging)
+	{
+		if (Context.Pin == NULL)
+		{
+			Context.MenuBuilder->BeginSection("UK2Node_EaseFunction", LOCTEXT("ContextMenuHeader", "Ease"));
+			{
+				Context.MenuBuilder->AddMenuEntry(
+					LOCTEXT("AddPin", "Reset to Wildcards"),
+					LOCTEXT("AddPinTooltip", "Resets A, B and Results pins to its default wildcard state"),
+					FSlateIcon(),
+					FUIAction(
+					FExecuteAction::CreateUObject(this, &UK2Node_EaseFunction::ResetToWildcards)
+					)
+				);
+			}
+			Context.MenuBuilder->EndSection();
+		}
+	}
+}
+
+#undef LOCTEXT_NAMESPACE
diff --git a/Engine/Source/Editor/BlueprintGraph/Public/BlueprintGraphDefinitions.h b/Engine/Source/Editor/BlueprintGraph/Public/BlueprintGraphDefinitions.h
--- a/Engine/Source/Editor/BlueprintGraph/Public/BlueprintGraphDefinitions.h
+++ b/Engine/Source/Editor/BlueprintGraph/Public/BlueprintGraphDefinitions.h
@@ -27,6 +27,7 @@
 #include "K2Node_CustomEvent.h"
 #include "K2Node_DelegateSet.h"
 #include "K2Node_DynamicCast.h"
+#include "K2Node_EaseFunction.h"
 #include "K2Node_Event.h"
 #include "K2Node_ExecutionSequence.h"
 #include "K2Node_FormatText.h"
diff --git a/Engine/Source/Runtime/Core/Public/Math/UnrealMathUtility.h b/Engine/Source/Runtime/Core/Public/Math/UnrealMathUtility.h
--- a/Engine/Source/Runtime/Core/Public/Math/UnrealMathUtility.h
+++ b/Engine/Source/Runtime/Core/Public/Math/UnrealMathUtility.h
@@ -535,17 +535,122 @@ struct FMath : public FPlatformMath
 		return (a * A) + b;
 	}
 
+	/** Interpolate between A and B, applying an ease in function.  Exp controls the degree of the curve. */
+	template< class T >
+	static FORCEINLINE_DEBUGGABLE T InterpEaseIn(const T& A, const T& B, float Alpha, float Exp)
+	{
+		float const ModifiedAlpha = FMath::Pow(Alpha, Exp);
+		return Lerp<T>(A, B, ModifiedAlpha);
+	}
+
+	/** Interpolate between A and B, applying an ease out function.  Exp controls the degree of the curve. */
+	template< class T >
+	static FORCEINLINE_DEBUGGABLE T InterpEaseOut(const T& A, const T& B, float Alpha, float Exp)
+	{
+		float const ModifiedAlpha = 1.f - FMath::Pow(1.f - Alpha, Exp);
+		return Lerp<T>(A, B, ModifiedAlpha);
+	}
+
 	/** Interpolate between A and B, applying an ease in/out function.  Exp controls the degree of the curve. */
 	template< class T > 
 	static FORCEINLINE_DEBUGGABLE T InterpEaseInOut( const T& A, const T& B, float Alpha, float Exp )
 	{
-		float const ModifiedAlpha = ( Alpha < 0.5f ) ?
-			0.5f * Pow(2.f * Alpha, Exp) :
-		1.f - 0.5f * Pow(2.f * (1.f - Alpha), Exp);
+		return (Alpha < 0.5f) ?
+			InterpEaseIn(A, B, Alpha * 2.f, Exp) * 0.5f :
+			InterpEaseOut(A, B, Alpha * 2.f - 1.f, Exp) * 0.5f + 0.5f;
+	}
+
+	/** Interpolation between A and B, applying a step function. */
+	template< class T >
+	static FORCEINLINE_DEBUGGABLE T InterpStep(const T& A, const T& B, float Alpha, int32 Steps)
+	{
+		if (Steps <= 1)
+		{
+			return A;
+		}
+
+		const float StepsAsFloat = static_cast<float>(Steps);
+		const float NumIntervals = StepsAsFloat - 1.f;
+		float const ModifiedAlpha = FloorToFloat(Alpha * StepsAsFloat) / NumIntervals;
+		return Lerp<T>(A, B, ModifiedAlpha);
+	}
+
+	/** Interpolation between A and B, applying a sinusoidal in function. */
+	template< class T >
+	static FORCEINLINE_DEBUGGABLE T InterpSinIn(const T& A, const T& B, float Alpha)
+	{
+		float const ModifiedAlpha = -1.f * cos(Alpha * HALF_PI) + 1.f;
+		return Lerp<T>(A, B, ModifiedAlpha);
+	}
+	
+	/** Interpolation between A and B, applying a sinusoidal out function. */
+	template< class T >
+	static FORCEINLINE_DEBUGGABLE T InterpSinOut(const T& A, const T& B, float Alpha)
+	{
+		float const ModifiedAlpha = sin(Alpha * HALF_PI);
+		return Lerp<T>(A, B, ModifiedAlpha);
+	}
+
+	/** Interpolation between A and B, applying a sinusoidal in/out function. */
+	template< class T >
+	static FORCEINLINE_DEBUGGABLE T InterpSinInOut(const T& A, const T& B, float Alpha)
+	{
+		return (Alpha < 0.5f) ?
+			InterpSinIn(A, B, Alpha * 2.f) * 0.5f :
+			InterpSinOut(A, B, Alpha * 2.f - 1.f) * 0.5f + 0.5f;
+	}
+
+	/** Interpolation between A and B, applying an exponential in function. */
+	template< class T >
+	static FORCEINLINE_DEBUGGABLE T InterpExpoIn(const T& A, const T& B, float Alpha)
+	{
+		float const ModifiedAlpha = (Alpha == 0.f) ? 0.f : pow(2.f, 10.f * (Alpha - 1.f));
+		return Lerp<T>(A, B, ModifiedAlpha);
+	}
 
+	/** Interpolation between A and B, applying an exponential out function. */
+	template< class T >
+	static FORCEINLINE_DEBUGGABLE T InterpExpoOut(const T& A, const T& B, float Alpha)
+	{
+		float const ModifiedAlpha = (Alpha == 1.f) ? 1.f : -pow(2.f, -10.f * Alpha) + 1.f;
 		return Lerp<T>(A, B, ModifiedAlpha);
 	}
 
+	/** Interpolation between A and B, applying an exponential in/out function. */
+	template< class T >
+	static FORCEINLINE_DEBUGGABLE T InterpExpoInOut(const T& A, const T& B, float Alpha)
+	{
+		return (Alpha < 0.5f) ?
+			InterpExpoIn(A, B, Alpha * 2.f) * 0.5f :
+			InterpExpoOut(A, B, Alpha * 2.f - 1.f) * 0.5f + 0.5f;
+	}
+
+	/** Interpolation between A and B, applying a circular in function. */
+	template< class T >
+	static FORCEINLINE_DEBUGGABLE T InterpCircularIn(const T& A, const T& B, float Alpha)
+	{
+		float const ModifiedAlpha = -1.f * (Sqrt(1.f - Alpha * Alpha) - 1.f);
+		return Lerp<T>(A, B, ModifiedAlpha);
+	}
+
+	/** Interpolation between A and B, applying a circular out function. */
+	template< class T >
+	static FORCEINLINE_DEBUGGABLE T InterpCircularOut(const T& A, const T& B, float Alpha)
+	{
+		Alpha -= 1.f;
+		float const ModifiedAlpha = Sqrt(1.f - Alpha  * Alpha);
+		return Lerp<T>(A, B, ModifiedAlpha);
+	}
+
+	/** Interpolation between A and B, applying a circular in/out function. */
+	template< class T >
+	static FORCEINLINE_DEBUGGABLE T InterpCircularInOut(const T& A, const T& B, float Alpha)
+	{
+		return (Alpha < 0.5f) ?
+			InterpCircularIn(A, B, Alpha * 2.f) * 0.5f :
+			InterpCircularOut(A, B, Alpha * 2.f - 1.f) * 0.5f + 0.5f;
+	}
+
 	// Rotator specific interpolation
 	template< class U > static FRotator Lerp( const FRotator& A, const FRotator& B, const U& Alpha);
 
diff --git a/Engine/Source/Runtime/Engine/Classes/Kismet/KismetMathLibrary.h b/Engine/Source/Runtime/Engine/Classes/Kismet/KismetMathLibrary.h
--- a/Engine/Source/Runtime/Engine/Classes/Kismet/KismetMathLibrary.h
+++ b/Engine/Source/Runtime/Engine/Classes/Kismet/KismetMathLibrary.h
@@ -4,6 +4,57 @@
 
 #include "KismetMathLibrary.generated.h"
 
+/** Provides different easing functions that can be used in blueprints */
+UENUM(BlueprintType)
+namespace EEasingFunc
+{
+	enum Type
+	{
+		/** Simple linear interpolation. */
+		Linear,
+
+		/** Simple step interpolation. */
+		Step,
+
+		/** Sinusoidal in interpolation. */
+		SinusoidalIn,
+
+		/** Sinusoidal out interpolation. */
+		SinusoidalOut,
+
+		/** Sinusoidal in/out interpolation. */
+		SinusoidalInOut,
+
+		/** Immediately accelerates, but smoothly decelerates into the target.  Ease amount controlled by BlendExp. */
+		EaseIn,
+
+		/** Smoothly accelerates, but does not decelerate into the target.  Ease amount controlled by BlendExp. */
+		EaseOut,
+
+		/** Smoothly accelerates and decelerates.  Ease amount controlled by BlendExp. */
+		EaseInOut,
+
+		/** Easing in using an exponential */
+		ExpoIn,
+
+		/** Easing out using an exponential */
+		ExpoOut,
+
+		/** Easing in/out using an exponential method */
+		ExpoInOut,
+
+		/** Easing is based on a half circle. */
+		CircularIn,
+
+		/** Easing is based on an inverted half circle. */
+		CircularOut,
+
+		/** Easing is based on two half circles. */
+		CircularInOut,
+
+	};
+}
+
 UCLASS(MinimalAPI)
 class UKismetMathLibrary : public UBlueprintFunctionLibrary
 {
@@ -383,6 +434,10 @@ class UKismetMathLibrary : public UBlueprintFunctionLibrary
 	UFUNCTION(BlueprintPure, Category="Math|Float")
 	static float Lerp(float A, float B, float Alpha);
 
+	/* Easeing  between A and B using a specified easing function */
+	UFUNCTION(BlueprintPure, meta = (FriendlyName = "Ease", BlueprintInternalUseOnly = "true"), Category = "Math|Interpolation")
+	static float Ease(float A, float B, float Alpha, TEnumAsByte<EEasingFunc::Type> EasingFunc, float BlendExp = 2, int32 Steps = 2);
+
 	/* Rounds A to the nearest integer */
 	UFUNCTION(BlueprintPure, Category="Math|Float")
 	static int32 Round(float A);
@@ -502,6 +557,10 @@ class UKismetMathLibrary : public UBlueprintFunctionLibrary
 	UFUNCTION(BlueprintPure, meta=(FriendlyName = "Lerp (vector)"), Category="Math|Vector")
 	static FVector VLerp(FVector A, FVector B, float Alpha);
 
+	/* Easeing  between A and B using a specified easing function */
+	UFUNCTION(BlueprintPure, meta = (FriendlyName = "Ease (vector)", BlueprintInternalUseOnly = "true"), Category = "Math|Interpolation")
+	static FVector VEase(FVector A, FVector B, float Alpha, TEnumAsByte<EEasingFunc::Type> EasingFunc, float BlendExp = 2, int32 Steps = 2);
+
 	/* Returns a random vector with length of 1 */
 	UFUNCTION(BlueprintPure, Category="Math|Random")
 	static FVector RandomUnitVector();
@@ -586,6 +645,10 @@ class UKismetMathLibrary : public UBlueprintFunctionLibrary
 	UFUNCTION(BlueprintPure, meta=(FriendlyName = "Lerp (Rotator)"), Category="Math|Rotator")
 	static FRotator RLerp(FRotator A, FRotator B, float Alpha, bool bShortestPath);
 
+	/* Easeing  between A and B using a specified easing function */
+	UFUNCTION(BlueprintPure, meta = (FriendlyName = "Ease (Rotator)", BlueprintInternalUseOnly = "true"), Category = "Math|Interpolation")
+	static FRotator REase(FRotator A, FRotator B, float Alpha, bool bShortestPath, TEnumAsByte<EEasingFunc::Type> EasingFunc, float BlendExp = 2, int32 Steps = 2);
+
 	/** Normalized A-B */
 	UFUNCTION(BlueprintPure, meta=(FriendlyName = "Delta (Rotator)"), Category="Math|Rotator")
 	static FRotator NormalizedDeltaRotator(FRotator A, FRotator B);
@@ -1171,6 +1234,10 @@ class UKismetMathLibrary : public UBlueprintFunctionLibrary
 	UFUNCTION(BlueprintPure, meta=(FriendlyName = "Lerp (Transform)"), Category="Math|Transform")
 	static FTransform TLerp(const FTransform& A, const FTransform& B, float Alpha);
 
+	/* Easeing  between A and B using a specified easing function */
+	UFUNCTION(BlueprintPure, meta = (FriendlyName = "Ease (Transform)", BlueprintInternalUseOnly = "true"), Category = "Math|Interpolation")
+	static FTransform TEase(const FTransform& A, const FTransform& B, float Alpha, TEnumAsByte<EEasingFunc::Type> EasingFunc, float BlendExp = 2, int32 Steps = 2);
+
 	/** Tries to reach a target transform */
 	UFUNCTION(BlueprintPure, Category="Math|Interpolation")
 	static FTransform TInterpTo(const FTransform& Current, const FTransform& Target, float DeltaTime, float InterpSpeed);
diff --git a/Engine/Source/Runtime/Engine/Private/KismetMathLibrary.cpp b/Engine/Source/Runtime/Engine/Private/KismetMathLibrary.cpp
--- a/Engine/Source/Runtime/Engine/Private/KismetMathLibrary.cpp
+++ b/Engine/Source/Runtime/Engine/Private/KismetMathLibrary.cpp
@@ -3,6 +3,28 @@
 #include "EnginePrivate.h"
 #include "CoreStats.h"
 
+/** Interpolate a linear alpha value using an ease mode and function, BlendExp used in cases where the easing is exponential */
+float EaseAlpha(float InAlpha, uint8 EasingFunc, float BlendExp, int32 Steps)
+{
+	switch (EasingFunc)
+	{
+	case EEasingFunc::Step:					return FMath::Clamp<float>(FMath::InterpStep(0.f, 1.f, InAlpha, Steps), 0.f, 1.f);
+	case EEasingFunc::SinusoidalIn:			return FMath::Clamp<float>(FMath::InterpSinIn<float>(0.f, 1.f, InAlpha), 0.f, 1.f);
+	case EEasingFunc::SinusoidalOut:		return FMath::Clamp<float>(FMath::InterpSinOut<float>(0.f, 1.f, InAlpha), 0.f, 1.f);
+	case EEasingFunc::SinusoidalInOut:		return FMath::Clamp<float>(FMath::InterpSinInOut<float>(0.f, 1.f, InAlpha), 0.f, 1.f);
+	case EEasingFunc::EaseIn:				return FMath::Clamp<float>(FMath::InterpEaseIn<float>(0.f, 1.f, InAlpha, BlendExp), 0.f, 1.f);
+	case EEasingFunc::EaseOut:				return FMath::Clamp<float>(FMath::InterpEaseOut<float>(0.f, 1.f, InAlpha, BlendExp), 0.f, 1.f);
+	case EEasingFunc::EaseInOut:			return FMath::Clamp<float>(FMath::InterpEaseInOut<float>(0.f, 1.f, InAlpha, BlendExp), 0.f, 1.f);
+	case EEasingFunc::ExpoIn:				return FMath::Clamp<float>(FMath::InterpExpoIn<float>(0.f, 1.f, InAlpha), 0.f, 1.f);
+	case EEasingFunc::ExpoOut:				return FMath::Clamp<float>(FMath::InterpExpoOut<float>(0.f, 1.f, InAlpha), 0.f, 1.f);
+	case EEasingFunc::ExpoInOut:			return FMath::Clamp<float>(FMath::InterpExpoInOut<float>(0.f, 1.f, InAlpha), 0.f, 1.f);
+	case EEasingFunc::CircularIn:			return FMath::Clamp<float>(FMath::InterpCircularIn<float>(0.f, 1.f, InAlpha), 0.f, 1.f);
+	case EEasingFunc::CircularOut:			return FMath::Clamp<float>(FMath::InterpCircularOut<float>(0.f, 1.f, InAlpha), 0.f, 1.f);
+	case EEasingFunc::CircularInOut:		return FMath::Clamp<float>(FMath::InterpCircularInOut<float>(0.f, 1.f, InAlpha), 0.f, 1.f);
+	}
+	return InAlpha;
+}
+
 UKismetMathLibrary::UKismetMathLibrary(const class FPostConstructInitializeProperties& PCIP)
 	: Super(PCIP)
 {
@@ -575,6 +597,11 @@ float UKismetMathLibrary::Lerp(float A, float B, float V)
 	return A + V*(B-A);
 }	
 
+float UKismetMathLibrary::Ease(float A, float B, float Alpha, TEnumAsByte<EEasingFunc::Type> EasingFunc, float BlendExp, int32 Steps)
+{
+	return Lerp(A, B, EaseAlpha(Alpha, EasingFunc, BlendExp, Steps));
+}
+
 float UKismetMathLibrary::FInterpTo(float Current, float Target, float DeltaTime, float InterpSpeed)
 {
 	return FMath::FInterpTo(Current, Target, DeltaTime, InterpSpeed);
@@ -700,6 +727,11 @@ FVector UKismetMathLibrary::VLerp(FVector A, FVector B, float V)
 	return A + V*(B-A);
 }	
 
+FVector UKismetMathLibrary::VEase(FVector A, FVector B, float Alpha, TEnumAsByte<EEasingFunc::Type> EasingFunc, float BlendExp, int32 Steps)
+{
+	return VLerp(A, B, EaseAlpha(Alpha, EasingFunc, BlendExp, Steps));
+}
+
 FVector UKismetMathLibrary::VInterpTo(FVector Current, FVector Target, float DeltaTime, float InterpSpeed)
 {
 	return FMath::VInterpTo( Current, Target, DeltaTime, InterpSpeed );
@@ -855,6 +887,11 @@ FRotator UKismetMathLibrary::RLerp(FRotator A, FRotator B, float Alpha, bool bSh
 	return A + Alpha*DeltaAngle;
 }
 
+FRotator UKismetMathLibrary::REase(FRotator A, FRotator B, float Alpha, bool bShortestPath, TEnumAsByte<EEasingFunc::Type> EasingFunc, float BlendExp, int32 Steps)
+{
+	return RLerp(A, B, EaseAlpha(Alpha, EasingFunc, BlendExp, Steps), bShortestPath);
+}
+
 FRotator UKismetMathLibrary::NormalizedDeltaRotator(FRotator A, FRotator B)
 {
 	FRotator Delta = A - B;
@@ -941,6 +978,11 @@ FTransform UKismetMathLibrary::TLerp(const FTransform& A, const FTransform& B, f
 	return Result;
 }
 
+FTransform UKismetMathLibrary::TEase(const FTransform& A, const FTransform& B, float Alpha, TEnumAsByte<EEasingFunc::Type> EasingFunc, float BlendExp, int32 Steps)
+{
+	return TLerp(A, B, EaseAlpha(Alpha, EasingFunc, BlendExp, Steps));
+}
+
 FTransform UKismetMathLibrary::TInterpTo(const FTransform& Current, const FTransform& Target, float DeltaTime, float InterpSpeed)
 {
 	if( InterpSpeed <= 0.f )
