Fix Function Entry nodes when function has a friendly name

diff --git a/Engine/Source/Editor/BlueprintGraph/Private/EdGraphSchema_K2.cpp b/Engine/Source/Editor/BlueprintGraph/Private/EdGraphSchema_K2.cpp
--- a/Engine/Source/Editor/BlueprintGraph/Private/EdGraphSchema_K2.cpp
+++ b/Engine/Source/Editor/BlueprintGraph/Private/EdGraphSchema_K2.cpp
@@ -3475,10 +3475,11 @@ void UEdGraphSchema_K2::GetGraphDisplayInformation(const UEdGraph& Graph, /*out*
 		}
 		else
 		{
-			// If we found a function from this graph, grab its tooltip
+			// If we found a function from this graph..
 			if (Function)
 			{
-				DisplayInfo.Tooltip = UK2Node_CallFunction::GetDefaultTooltipForFunction(Function);
+				DisplayInfo.PlainName = FText::FromString(UK2Node_CallFunction::GetUserFacingFunctionName(Function)); // grab friendly function name
+				DisplayInfo.Tooltip = UK2Node_CallFunction::GetDefaultTooltipForFunction(Function); // grab its tooltip
 			}
 			else
 			{
