A k2node should be able to tell, if it's wildcard pin can accept array.
ttp345111 EDITOR: LIVE: BP: CRASH: Connecting an Array Get node to the wrong pin of an Array Add node crashes the editor.

diff --git a/Engine/Source/Editor/BlueprintGraph/Classes/K2Node.h b/Engine/Source/Editor/BlueprintGraph/Classes/K2Node.h
--- a/Engine/Source/Editor/BlueprintGraph/Classes/K2Node.h
+++ b/Engine/Source/Editor/BlueprintGraph/Classes/K2Node.h
@@ -318,6 +318,7 @@ class UK2Node : public UEdGraphNode
 
 	BLUEPRINTGRAPH_API virtual int32 GetNodeRefreshPriority() const { return EBaseNodeRefreshPriority::Normal; }
 
+	BLUEPRINTGRAPH_API virtual bool DoesSingleInputWildcardPinAcceptsArray(const UEdGraphPin* Pin) const { return true; }
 protected:
 	/** 
 	 * A mapping from old property and function names to new ones.  Get primed from INI files, and should contain entries for properties, functions, and delegates that get moved, so they can be fixed up
diff --git a/Engine/Source/Editor/BlueprintGraph/Classes/K2Node_CallArrayFunction.h b/Engine/Source/Editor/BlueprintGraph/Classes/K2Node_CallArrayFunction.h
--- a/Engine/Source/Editor/BlueprintGraph/Classes/K2Node_CallArrayFunction.h
+++ b/Engine/Source/Editor/BlueprintGraph/Classes/K2Node_CallArrayFunction.h
@@ -28,6 +28,7 @@ class UK2Node_CallArrayFunction : public UK2Node_CallFunction
 
 	// Begin UK2Node interface
 	virtual void NotifyPinConnectionListChanged(UEdGraphPin* Pin) override;
+	virtual bool DoesSingleInputWildcardPinAcceptsArray(const UEdGraphPin* Pin) const override { return false; }
 	// End UK2Node interface
 
 	/** Retrieves the target pin for the function */
diff --git a/Engine/Source/Editor/BlueprintGraph/Private/EdGraphSchema_K2.cpp b/Engine/Source/Editor/BlueprintGraph/Private/EdGraphSchema_K2.cpp
--- a/Engine/Source/Editor/BlueprintGraph/Private/EdGraphSchema_K2.cpp
+++ b/Engine/Source/Editor/BlueprintGraph/Private/EdGraphSchema_K2.cpp
@@ -2735,15 +2735,39 @@ void UEdGraphSchema_K2::GetVariableSubtypes(const FString& Type, TArray<UObject*
 	SubtypesList.Sort();
 }
 
+struct FWildcardArrayPinHelper
+{
+	static bool CheckArrayCompatibility(const UEdGraphPin* OutputPin, const UEdGraphPin* InputPin, bool bIgnoreArray)
+	{
+		if (bIgnoreArray)
+		{
+			return true;
+		}
+
+		const UK2Node* OwningNode = InputPin ? Cast<UK2Node>(InputPin->GetOwningNode()) : NULL;
+		const bool bInputWildcardPinAcceptsArray = !OwningNode || OwningNode->DoesSingleInputWildcardPinAcceptsArray(InputPin);
+		if (bInputWildcardPinAcceptsArray)
+		{
+			return true;
+		}
+
+		const bool bCheckInputPin = (InputPin->PinType.PinCategory == GetDefault<UEdGraphSchema_K2>()->PC_Wildcard) && !InputPin->PinType.bIsArray;
+		const bool bArrayOutputPin = OutputPin && OutputPin->PinType.bIsArray;
+		return !(bCheckInputPin && bArrayOutputPin);
+	}
+};
+
 bool UEdGraphSchema_K2::ArePinsCompatible(const UEdGraphPin* PinA, const UEdGraphPin* PinB, const UClass* CallingContext, bool bIgnoreArray /*= false*/) const
 {
 	if ((PinA->Direction == EGPD_Input) && (PinB->Direction == EGPD_Output))
 	{
-		return ArePinTypesCompatible(PinB->PinType, PinA->PinType, CallingContext, bIgnoreArray);
+		return FWildcardArrayPinHelper::CheckArrayCompatibility(PinB, PinA, bIgnoreArray)
+			&& ArePinTypesCompatible(PinB->PinType, PinA->PinType, CallingContext, bIgnoreArray);
 	}
 	else if ((PinB->Direction == EGPD_Input) && (PinA->Direction == EGPD_Output))
 	{
-		return ArePinTypesCompatible(PinA->PinType, PinB->PinType, CallingContext, bIgnoreArray);
+		return FWildcardArrayPinHelper::CheckArrayCompatibility(PinA, PinB, bIgnoreArray)
+			&& ArePinTypesCompatible(PinA->PinType, PinB->PinType, CallingContext, bIgnoreArray);
 	}
 	else
 	{
