BP: Function renamed: DoesInputWildcardPinAcceptArray

diff --git a/Engine/Source/Editor/BlueprintGraph/Classes/K2Node.h b/Engine/Source/Editor/BlueprintGraph/Classes/K2Node.h
--- a/Engine/Source/Editor/BlueprintGraph/Classes/K2Node.h
+++ b/Engine/Source/Editor/BlueprintGraph/Classes/K2Node.h
@@ -318,7 +318,7 @@ class UK2Node : public UEdGraphNode
 
 	BLUEPRINTGRAPH_API virtual int32 GetNodeRefreshPriority() const { return EBaseNodeRefreshPriority::Normal; }
 
-	BLUEPRINTGRAPH_API virtual bool DoesSingleInputWildcardPinAcceptsArray(const UEdGraphPin* Pin) const { return true; }
+	BLUEPRINTGRAPH_API virtual bool DoesInputWildcardPinAcceptArray(const UEdGraphPin* Pin) const { return true; }
 protected:
 	/** 
 	 * A mapping from old property and function names to new ones.  Get primed from INI files, and should contain entries for properties, functions, and delegates that get moved, so they can be fixed up
diff --git a/Engine/Source/Editor/BlueprintGraph/Classes/K2Node_CallArrayFunction.h b/Engine/Source/Editor/BlueprintGraph/Classes/K2Node_CallArrayFunction.h
--- a/Engine/Source/Editor/BlueprintGraph/Classes/K2Node_CallArrayFunction.h
+++ b/Engine/Source/Editor/BlueprintGraph/Classes/K2Node_CallArrayFunction.h
@@ -28,7 +28,7 @@ class UK2Node_CallArrayFunction : public UK2Node_CallFunction
 
 	// Begin UK2Node interface
 	virtual void NotifyPinConnectionListChanged(UEdGraphPin* Pin) override;
-	virtual bool DoesSingleInputWildcardPinAcceptsArray(const UEdGraphPin* Pin) const override { return false; }
+	virtual bool DoesInputWildcardPinAcceptArray(const UEdGraphPin* Pin) const override { return false; }
 	// End UK2Node interface
 
 	/** Retrieves the target pin for the function */
diff --git a/Engine/Source/Editor/BlueprintGraph/Private/EdGraphSchema_K2.cpp b/Engine/Source/Editor/BlueprintGraph/Private/EdGraphSchema_K2.cpp
--- a/Engine/Source/Editor/BlueprintGraph/Private/EdGraphSchema_K2.cpp
+++ b/Engine/Source/Editor/BlueprintGraph/Private/EdGraphSchema_K2.cpp
@@ -2745,7 +2745,7 @@ struct FWildcardArrayPinHelper
 		}
 
 		const UK2Node* OwningNode = InputPin ? Cast<UK2Node>(InputPin->GetOwningNode()) : NULL;
-		const bool bInputWildcardPinAcceptsArray = !OwningNode || OwningNode->DoesSingleInputWildcardPinAcceptsArray(InputPin);
+		const bool bInputWildcardPinAcceptsArray = !OwningNode || OwningNode->DoesInputWildcardPinAcceptArray(InputPin);
 		if (bInputWildcardPinAcceptsArray)
 		{
 			return true;
