FBlueprintGraphPanelPinFactory::CreatePin - fix for a crash.

diff --git a/Engine/Source/Editor/BlueprintGraph/Private/BlueprintGraphPanelPinFactory.cpp b/Engine/Source/Editor/BlueprintGraph/Private/BlueprintGraphPanelPinFactory.cpp
--- a/Engine/Source/Editor/BlueprintGraph/Private/BlueprintGraphPanelPinFactory.cpp
+++ b/Engine/Source/Editor/BlueprintGraph/Private/BlueprintGraphPanelPinFactory.cpp
@@ -20,8 +20,11 @@ TSharedPtr<class SGraphPin> FBlueprintGraphPanelPinFactory::CreatePin(class UEdG
 			if (CallFunctionNode)
 			{
 				const UFunction* FunctionToCall = CallFunctionNode->GetTargetFunction();
-				FString DataTablePinName = FunctionToCall->GetMetaData(FBlueprintMetadata::MD_DataTablePin);
-				DataTablePin = CallFunctionNode->FindPin(DataTablePinName);
+				if (FunctionToCall)
+				{
+					FString DataTablePinName = FunctionToCall->GetMetaData(FBlueprintMetadata::MD_DataTablePin);
+					DataTablePin = CallFunctionNode->FindPin(DataTablePinName);
+				}
 			}
 		}
 		else if (Outer->IsA(UK2Node_GetDataTableRow::StaticClass()))
