Using the new BP menu system, dragging off pins for UObjects will no longer present the user with cast options for every single UObject type.

diff --git a/Engine/Source/Editor/BlueprintGraph/Private/BlueprintActionFilter.cpp b/Engine/Source/Editor/BlueprintGraph/Private/BlueprintActionFilter.cpp
--- a/Engine/Source/Editor/BlueprintGraph/Private/BlueprintActionFilter.cpp
+++ b/Engine/Source/Editor/BlueprintGraph/Private/BlueprintActionFilter.cpp
@@ -18,6 +18,7 @@
 #include "K2Node_MakeArray.h"
 #include "K2Node_Message.h"
 #include "K2Node_ExecutionSequence.h"
+#include "K2Node_DynamicCast.h"
 
 /*******************************************************************************
  * Static BlueprintActionFilter Helpers
@@ -251,6 +252,15 @@ namespace BlueprintActionFilterImpl
 	 * @return 
 	 */
 	static bool IsMissingMatchingPinParam(FBlueprintActionFilter const& Filter, UBlueprintNodeSpawner const* BlueprintAction);
+
+	/**
+	 * Dynamic casts should only show results for casting to classes that the context pin is a child of (and not itself)
+	 * 
+	 * @param  Filter	
+	 * @param  BlueprintAction	
+	 * @return 
+	 */
+	static bool IsNotSubClassCast(FBlueprintActionFilter const& Filter, UBlueprintNodeSpawner const* BlueprintAction);
 };
 
 
@@ -968,6 +978,38 @@ static bool BlueprintActionFilterImpl::IsMissingMatchingPinParam(FBlueprintActio
 	return bIsFilteredOut;
 }
 
+//------------------------------------------------------------------------------
+static bool BlueprintActionFilterImpl::IsNotSubClassCast(FBlueprintActionFilter const& Filter, UBlueprintNodeSpawner const* BlueprintAction)
+{
+	bool bIsFilteredOut = false;
+
+	if(BlueprintAction->NodeClass && BlueprintAction->NodeClass == UK2Node_DynamicCast::StaticClass())
+	{
+		for(UEdGraphPin const* ContextPin : Filter.Context.Pins)
+		{
+			// Only worry about removing cast nodes when dragging off output pins
+			if(ContextPin->Direction == EGPD_Output)
+			{
+				UEdGraph* OuterGraph = ContextPin->GetOwningNode()->GetGraph();
+				if (UK2Node_DynamicCast* CastNode = Cast<UK2Node_DynamicCast>(BlueprintAction->GetTemplateNode(OuterGraph)))
+				{
+					if( ContextPin->PinType.PinSubCategoryObject.IsValid() && CastNode->TargetType)
+					{
+						UClass* ContextPinClass = Cast<UClass>(ContextPin->PinType.PinSubCategoryObject.Get());
+						UClass* ResultPinClass = CastNode->TargetType;
+						if( ContextPinClass == ResultPinClass || !ResultPinClass->IsChildOf(ContextPinClass) )
+						{
+							bIsFilteredOut = true;
+							break;
+						}
+					}
+				}
+			}
+		}
+	}
+	return bIsFilteredOut;
+}
+
 /*******************************************************************************
  * FBlueprintActionFilter
  ******************************************************************************/
@@ -1001,6 +1043,7 @@ FBlueprintActionFilter::FBlueprintActionFilter(uint32 Flags/*= 0x00*/)
 	}
 	
 	AddIsFilteredTest(FIsFilteredDelegate::CreateStatic(IsFieldInaccessible));
+	AddIsFilteredTest(FIsFilteredDelegate::CreateStatic(IsNotSubClassCast));
 	AddIsFilteredTest(FIsFilteredDelegate::CreateStatic(IsEventUnimplementable));
 	AddIsFilteredTest(FIsFilteredDelegate::CreateStatic(IsPermissionNotGranted));
 	AddIsFilteredTest(FIsFilteredDelegate::CreateStatic(IsRestrictedClassMember));
