Accidently adding blueprint nodes twice to the outer graph (in the new, experimental menu system).

diff --git a/Engine/Source/Editor/BlueprintGraph/Private/BlueprintNodeSpawner.cpp b/Engine/Source/Editor/BlueprintGraph/Private/BlueprintNodeSpawner.cpp
--- a/Engine/Source/Editor/BlueprintGraph/Private/BlueprintNodeSpawner.cpp
+++ b/Engine/Source/Editor/BlueprintGraph/Private/BlueprintNodeSpawner.cpp
@@ -162,11 +162,6 @@ UEdGraphNode* UBlueprintNodeSpawner::Invoke(UEdGraph* ParentGraph, FVector2D con
 			NewNode->SetFlags(RF_Transactional);
 			NewNode->AllocateDefaultPins();
 			NewNode->PostPlacedNewNode();
-
-			ParentGraph->Modify();
-			// selections of this node happen up, out at the 
-			// FBlueprintActionMenuItem level
-			ParentGraph->AddNode(NewNode, /*bFromUI =*/true, /*bSelectNewNode =*/false);
 		}
 	}
 
