Optimizing BP node titles (caching any that use FText::Format(), which is slow).

diff --git a/Engine/Source/Editor/BlueprintGraph/Classes/BlueprintComponentNodeSpawner.h b/Engine/Source/Editor/BlueprintGraph/Classes/BlueprintComponentNodeSpawner.h
--- a/Engine/Source/Editor/BlueprintGraph/Classes/BlueprintComponentNodeSpawner.h
+++ b/Engine/Source/Editor/BlueprintGraph/Classes/BlueprintComponentNodeSpawner.h
@@ -3,6 +3,7 @@
 #pragma once
 
 #include "BlueprintNodeSpawner.h"
+#include "EdGraph/EdGraphNodeUtils.h" // for FNodeTextCache
 #include "BlueprintComponentNodeSpawner.generated.h"
 
 // Forward declarations
@@ -48,4 +49,7 @@ private:
 	/** The component class to configure new nodes with. */
 	UPROPERTY()
 	TSubclassOf<UActorComponent> ComponentClass;
+
+	/** Constructing FText strings can be costly, so we cache the default menu name */
+	FNodeTextCache CachedMenuName;
 };
diff --git a/Engine/Source/Editor/BlueprintGraph/Classes/BlueprintEventNodeSpawner.h b/Engine/Source/Editor/BlueprintGraph/Classes/BlueprintEventNodeSpawner.h
--- a/Engine/Source/Editor/BlueprintGraph/Classes/BlueprintEventNodeSpawner.h
+++ b/Engine/Source/Editor/BlueprintGraph/Classes/BlueprintEventNodeSpawner.h
@@ -3,6 +3,7 @@
 #pragma once
 
 #include "BlueprintNodeSpawner.h"
+#include "EdGraph/EdGraphNodeUtils.h" // for FNodeTextCache
 #include "BlueprintEventNodeSpawner.generated.h"
 
 class UK2Node_Event;
@@ -71,7 +72,9 @@ public:
 	 */
 	virtual UK2Node_Event const* FindPreExistingEvent(UBlueprint* Blueprint) const;
 
-	
+protected:
+	/** Constructing FText strings can be costly, so we cache the default menu name */
+	FNodeTextCache CachedMenuName;
 
 private:
 	/** The function to configure new nodes with. */
@@ -80,5 +83,5 @@ private:
 
 	/** The custom name to configure new event nodes with. */
 	UPROPERTY()
-	FName CustomEventName;
+	FName CustomEventName;	
 };
diff --git a/Engine/Source/Editor/BlueprintGraph/Classes/BlueprintVariableNodeSpawner.h b/Engine/Source/Editor/BlueprintGraph/Classes/BlueprintVariableNodeSpawner.h
--- a/Engine/Source/Editor/BlueprintGraph/Classes/BlueprintVariableNodeSpawner.h
+++ b/Engine/Source/Editor/BlueprintGraph/Classes/BlueprintVariableNodeSpawner.h
@@ -3,6 +3,7 @@
 #pragma once
 
 #include "BlueprintNodeSpawner.h"
+#include "EdGraph/EdGraphNodeUtils.h" // for FNodeTextCache
 #include "BlueprintVariableNodeSpawner.generated.h"
 
 // Forward declarations
@@ -106,4 +107,7 @@ private:
 	/** */
 	UPROPERTY()
 	FBPVariableDescription    LocalVarDesc;
+
+	/** Constructing FText strings can be costly, so we cache the default menu name */
+	FNodeTextCache CachedMenuName;
 };
\ No newline at end of file
diff --git a/Engine/Source/Editor/BlueprintGraph/Classes/K2Node_ActorBoundEvent.h b/Engine/Source/Editor/BlueprintGraph/Classes/K2Node_ActorBoundEvent.h
--- a/Engine/Source/Editor/BlueprintGraph/Classes/K2Node_ActorBoundEvent.h
+++ b/Engine/Source/Editor/BlueprintGraph/Classes/K2Node_ActorBoundEvent.h
@@ -3,6 +3,7 @@
 
 #pragma once
 #include "K2Node_Event.h"
+#include "EdGraph/EdGraphNodeUtils.h" // for FNodeTextCache
 #include "K2Node_ActorBoundEvent.generated.h"
 
 UCLASS(MinimalAPI)
@@ -51,5 +52,9 @@ class UK2Node_ActorBoundEvent : public UK2Node_Event
 
 	/** Returns the delegate that this event is bound in */
 	BLUEPRINTGRAPH_API FMulticastScriptDelegate* GetTargetDelegate();
+
+private:
+	/** Constructing FText strings can be costly, so we cache the node's title */
+	FNodeTextCache CachedNodeTitle;
 };
 
diff --git a/Engine/Source/Editor/BlueprintGraph/Classes/K2Node_AddComponent.h b/Engine/Source/Editor/BlueprintGraph/Classes/K2Node_AddComponent.h
--- a/Engine/Source/Editor/BlueprintGraph/Classes/K2Node_AddComponent.h
+++ b/Engine/Source/Editor/BlueprintGraph/Classes/K2Node_AddComponent.h
@@ -3,6 +3,7 @@
 
 #pragma once
 #include "K2Node_CallFunction.h"
+#include "EdGraph/EdGraphNodeUtils.h" // for FNodeTextCache
 #include "K2Node_AddComponent.generated.h"
 
 UCLASS(MinimalAPI)
@@ -70,6 +71,9 @@ private:
 	}
 
 	const UClass* GetSpawnedType() const;
+
+	/** Constructing FText strings can be costly, so we cache the node's title */
+	FNodeTextCache CachedNodeTitle;
 };
 
 
diff --git a/Engine/Source/Editor/BlueprintGraph/Classes/K2Node_AssignDelegate.h b/Engine/Source/Editor/BlueprintGraph/Classes/K2Node_AssignDelegate.h
--- a/Engine/Source/Editor/BlueprintGraph/Classes/K2Node_AssignDelegate.h
+++ b/Engine/Source/Editor/BlueprintGraph/Classes/K2Node_AssignDelegate.h
@@ -2,6 +2,7 @@
 #pragma once
 
 #include "K2Node_AddDelegate.h"
+#include "EdGraph/EdGraphNodeUtils.h" // for FNodeTextCache
 #include "K2Node_AssignDelegate.generated.h"
 
 /**
@@ -21,4 +22,8 @@ public:
 	virtual bool IsCompatibleWithGraph(const UEdGraph* TargetGraph) const override;
 	virtual void PostPlacedNewNode() override;
 	// End of UEdGraphNode interface
+
+private:
+	/** Constructing FText strings can be costly, so we cache the node's title */
+	FNodeTextCache CachedListTitle;
 };
diff --git a/Engine/Source/Editor/BlueprintGraph/Classes/K2Node_BaseMCDelegate.h b/Engine/Source/Editor/BlueprintGraph/Classes/K2Node_BaseMCDelegate.h
--- a/Engine/Source/Editor/BlueprintGraph/Classes/K2Node_BaseMCDelegate.h
+++ b/Engine/Source/Editor/BlueprintGraph/Classes/K2Node_BaseMCDelegate.h
@@ -1,6 +1,7 @@
 // Copyright 1998-2014 Epic Games, Inc. All Rights Reserved.
 #pragma once
 #include "K2Node.h"
+#include "EdGraph/EdGraphNodeUtils.h" // for FNodeTextCache
 #include "K2Node_BaseMCDelegate.generated.h"
 
 UCLASS(MinimalAPI, abstract)
@@ -53,4 +54,8 @@ public:
 
 	/** Is the delegate BlueprintAuthorityOnly */
 	BLUEPRINTGRAPH_API bool IsAuthorityOnly() const;
+
+protected:
+	/** Constructing FText strings can be costly, so we cache the node's title */
+	FNodeTextCache CachedNodeTitle;
 };
diff --git a/Engine/Source/Editor/BlueprintGraph/Classes/K2Node_BreakStruct.h b/Engine/Source/Editor/BlueprintGraph/Classes/K2Node_BreakStruct.h
--- a/Engine/Source/Editor/BlueprintGraph/Classes/K2Node_BreakStruct.h
+++ b/Engine/Source/Editor/BlueprintGraph/Classes/K2Node_BreakStruct.h
@@ -2,6 +2,7 @@
 
 #pragma once
 #include "K2Node_StructMemberGet.h"
+#include "EdGraph/EdGraphNodeUtils.h" // for FNodeTextCache
 #include "K2Node_BreakStruct.generated.h"
 
 UCLASS(MinimalAPI)
@@ -29,5 +30,9 @@ class UK2Node_BreakStruct : public UK2Node_StructMemberGet
 	virtual void GetMenuActions(FBlueprintActionDatabaseRegistrar& ActionRegistrar) const override;
 	virtual FText GetMenuCategory() const override;
 	// End K2Node interface
+
+private:
+	/** Constructing FText strings can be costly, so we cache the node's title */
+	FNodeTextCache CachedNodeTitle;
 };
 
diff --git a/Engine/Source/Editor/BlueprintGraph/Classes/K2Node_ComponentBoundEvent.h b/Engine/Source/Editor/BlueprintGraph/Classes/K2Node_ComponentBoundEvent.h
--- a/Engine/Source/Editor/BlueprintGraph/Classes/K2Node_ComponentBoundEvent.h
+++ b/Engine/Source/Editor/BlueprintGraph/Classes/K2Node_ComponentBoundEvent.h
@@ -3,6 +3,7 @@
 
 #pragma once
 #include "K2Node_Event.h"
+#include "EdGraph/EdGraphNodeUtils.h" // for FNodeTextCache
 #include "K2Node_ComponentBoundEvent.generated.h"
 
 UCLASS(MinimalAPI)
@@ -37,5 +38,9 @@ class UK2Node_ComponentBoundEvent : public UK2Node_Event
 	BLUEPRINTGRAPH_API UMulticastDelegateProperty* GetTargetDelegateProperty() const;
 
 	BLUEPRINTGRAPH_API void InitializeComponentBoundEventParams(UObjectProperty const* InComponentProperty, const UMulticastDelegateProperty* InDelegateProperty);
+
+private:
+	/** Constructing FText strings can be costly, so we cache the node's title */
+	FNodeTextCache CachedNodeTitle;
 };
 
diff --git a/Engine/Source/Editor/BlueprintGraph/Classes/K2Node_Composite.h b/Engine/Source/Editor/BlueprintGraph/Classes/K2Node_Composite.h
--- a/Engine/Source/Editor/BlueprintGraph/Classes/K2Node_Composite.h
+++ b/Engine/Source/Editor/BlueprintGraph/Classes/K2Node_Composite.h
@@ -1,6 +1,7 @@
 // Copyright 1998-2014 Epic Games, Inc. All Rights Reserved.
 #pragma once
 #include "K2Node_Tunnel.h"
+#include "EdGraph/EdGraphNodeUtils.h" // for FNodeTextCache
 #include "K2Node_Composite.generated.h"
 
 UCLASS(MinimalAPI)
@@ -45,6 +46,9 @@ private:
 
 	/** Determine if the name already used by another graph in composite nodes chain */
 	bool IsCompositeNameAvailable( const FString& NewName );
+
+	/** Constructing FText strings can be costly, so we cache the node's title */
+	FNodeTextCache CachedNodeTitle;
 };
 
 
diff --git a/Engine/Source/Editor/BlueprintGraph/Classes/K2Node_ConstructObjectFromClass.h b/Engine/Source/Editor/BlueprintGraph/Classes/K2Node_ConstructObjectFromClass.h
--- a/Engine/Source/Editor/BlueprintGraph/Classes/K2Node_ConstructObjectFromClass.h
+++ b/Engine/Source/Editor/BlueprintGraph/Classes/K2Node_ConstructObjectFromClass.h
@@ -3,6 +3,7 @@
 #pragma once
 
 #include "K2Node.h"
+#include "EdGraph/EdGraphNodeUtils.h" // for FNodeTextCache
 #include "K2Node_ConstructObjectFromClass.generated.h"
 
 UCLASS(abstract)
@@ -67,4 +68,8 @@ protected:
 
 	/** Tooltip text for this node. */
 	FString NodeTooltip;
+
+private:
+	/** Constructing FText strings can be costly, so we cache the node's title */
+	FNodeTextCache CachedNodeTitle;
 };
diff --git a/Engine/Source/Editor/BlueprintGraph/Classes/K2Node_CustomEvent.h b/Engine/Source/Editor/BlueprintGraph/Classes/K2Node_CustomEvent.h
--- a/Engine/Source/Editor/BlueprintGraph/Classes/K2Node_CustomEvent.h
+++ b/Engine/Source/Editor/BlueprintGraph/Classes/K2Node_CustomEvent.h
@@ -3,6 +3,7 @@
 
 #pragma once
 #include "K2Node_Event.h"
+#include "EdGraph/EdGraphNodeUtils.h" // for FNodeTextCache
 #include "K2Node_CustomEvent.generated.h"
 
 UCLASS(MinimalAPI)
@@ -60,6 +61,10 @@ class UK2Node_CustomEvent : public UK2Node_Event
 	 * @return If this CustomEvent is an override, then this is the super's net flags, otherwise it's from the FunctionFlags set on this node.
 	 */
 	BLUEPRINTGRAPH_API uint32 GetNetFlags() const;
+
+private:
+	/** Constructing FText strings can be costly, so we cache the node's title */
+	FNodeTextCache CachedNodeTitle;
 };
 
 
diff --git a/Engine/Source/Editor/BlueprintGraph/Classes/K2Node_DelegateSet.h b/Engine/Source/Editor/BlueprintGraph/Classes/K2Node_DelegateSet.h
--- a/Engine/Source/Editor/BlueprintGraph/Classes/K2Node_DelegateSet.h
+++ b/Engine/Source/Editor/BlueprintGraph/Classes/K2Node_DelegateSet.h
@@ -3,6 +3,7 @@
 
 #pragma once
 #include "K2Node.h"
+#include "EdGraph/EdGraphNodeUtils.h" // for FNodeTextCache
 #include "K2Node_DelegateSet.generated.h"
 
 UCLASS(MinimalAPI)
@@ -48,5 +49,9 @@ class UK2Node_DelegateSet : public UK2Node
 	// @todo document
 	BLUEPRINTGRAPH_API UFunction* GetDelegateSignature();
 	BLUEPRINTGRAPH_API UFunction* GetDelegateSignature() const;
+
+private:
+	/** Constructing FText strings can be costly, so we cache the node's title */
+	FNodeTextCache CachedNodeTitle;
 };
 
diff --git a/Engine/Source/Editor/BlueprintGraph/Classes/K2Node_DynamicCast.h b/Engine/Source/Editor/BlueprintGraph/Classes/K2Node_DynamicCast.h
--- a/Engine/Source/Editor/BlueprintGraph/Classes/K2Node_DynamicCast.h
+++ b/Engine/Source/Editor/BlueprintGraph/Classes/K2Node_DynamicCast.h
@@ -3,6 +3,7 @@
 
 #pragma once
 #include "K2Node.h"
+#include "EdGraph/EdGraphNodeUtils.h" // for FNodeTextCache
 #include "K2Node_DynamicCast.generated.h"
 
 UCLASS(MinimalAPI)
@@ -39,5 +40,9 @@ class UK2Node_DynamicCast : public UK2Node
 
 	/** Get the input object to be casted pin */
 	BLUEPRINTGRAPH_API virtual UEdGraphPin* GetCastSourcePin() const;
+
+protected:
+	/** Constructing FText strings can be costly, so we cache the node's title */
+	FNodeTextCache CachedNodeTitle;
 };
 
diff --git a/Engine/Source/Editor/BlueprintGraph/Classes/K2Node_EnumLiteral.h b/Engine/Source/Editor/BlueprintGraph/Classes/K2Node_EnumLiteral.h
--- a/Engine/Source/Editor/BlueprintGraph/Classes/K2Node_EnumLiteral.h
+++ b/Engine/Source/Editor/BlueprintGraph/Classes/K2Node_EnumLiteral.h
@@ -3,6 +3,7 @@
 #pragma once
 #include "K2Node.h"
 #include "NodeDependingOnEnumInterface.h"
+#include "EdGraph/EdGraphNodeUtils.h" // for FNodeTextCache
 #include "K2Node_EnumLiteral.generated.h"
 
 UCLASS(MinimalAPI)
@@ -33,5 +34,9 @@ class UK2Node_EnumLiteral : public UK2Node, public INodeDependingOnEnumInterface
 	virtual class UEnum* GetEnum() const override { return Enum; }
 	virtual bool ShouldBeReconstructedAfterEnumChanged() const override { return true; }
 	// End of INodeDependingOnEnumInterface
+
+private:
+	/** Constructing FText strings can be costly, so we cache the node's tootltip */
+	FNodeTextCache CachedTooltip;
 };
 
diff --git a/Engine/Source/Editor/BlueprintGraph/Classes/K2Node_ForEachElementInEnum.h b/Engine/Source/Editor/BlueprintGraph/Classes/K2Node_ForEachElementInEnum.h
--- a/Engine/Source/Editor/BlueprintGraph/Classes/K2Node_ForEachElementInEnum.h
+++ b/Engine/Source/Editor/BlueprintGraph/Classes/K2Node_ForEachElementInEnum.h
@@ -2,6 +2,7 @@
 
 #pragma once
 #include "K2Node.h"
+#include "EdGraph/EdGraphNodeUtils.h" // for FNodeTextCache
 #include "K2Node_ForEachElementInEnum.generated.h"
 
 UCLASS(MinimalAPI)
@@ -29,5 +30,9 @@ class UK2Node_ForEachElementInEnum : public UK2Node
 	virtual void GetMenuActions(FBlueprintActionDatabaseRegistrar& ActionRegistrar) const override;
 	virtual FText GetMenuCategory() const override;
 	// End UK2Node interface
+
+private:
+	/** Constructing FText strings can be costly, so we cache the node's title */
+	FNodeTextCache CachedNodeTitle;
 };
 
diff --git a/Engine/Source/Editor/BlueprintGraph/Classes/K2Node_GetDataTableRow.h b/Engine/Source/Editor/BlueprintGraph/Classes/K2Node_GetDataTableRow.h
--- a/Engine/Source/Editor/BlueprintGraph/Classes/K2Node_GetDataTableRow.h
+++ b/Engine/Source/Editor/BlueprintGraph/Classes/K2Node_GetDataTableRow.h
@@ -3,6 +3,7 @@
 
 #pragma once
 #include "K2Node.h"
+#include "EdGraph/EdGraphNodeUtils.h" // for FNodeTextCache
 #include "K2Node_GetDataTableRow.generated.h"
 
 UCLASS()
@@ -69,4 +70,7 @@ private:
 
 	/** Tooltip text for this node. */
 	FString NodeTooltip;
+
+	/** Constructing FText strings can be costly, so we cache the node's title */
+	FNodeTextCache CachedNodeTitle;
 };
diff --git a/Engine/Source/Editor/BlueprintGraph/Classes/K2Node_GetInputAxisKeyValue.h b/Engine/Source/Editor/BlueprintGraph/Classes/K2Node_GetInputAxisKeyValue.h
--- a/Engine/Source/Editor/BlueprintGraph/Classes/K2Node_GetInputAxisKeyValue.h
+++ b/Engine/Source/Editor/BlueprintGraph/Classes/K2Node_GetInputAxisKeyValue.h
@@ -2,6 +2,7 @@
 
 #pragma once
 #include "K2Node_CallFunction.h"
+#include "EdGraph/EdGraphNodeUtils.h" // for FNodeTextCache
 #include "K2Node_GetInputAxisKeyValue.generated.h"
 
 UCLASS(MinimalAPI)
@@ -38,4 +39,8 @@ class UK2Node_GetInputAxisKeyValue : public UK2Node_CallFunction
 	// End UK2Node interface
 	
 	void Initialize(const FKey AxisKey);
+
+private:
+	/** Constructing FText strings can be costly, so we cache the node's title */
+	FNodeTextCache CachedNodeTitle;
 };
\ No newline at end of file
diff --git a/Engine/Source/Editor/BlueprintGraph/Classes/K2Node_GetInputAxisValue.h b/Engine/Source/Editor/BlueprintGraph/Classes/K2Node_GetInputAxisValue.h
--- a/Engine/Source/Editor/BlueprintGraph/Classes/K2Node_GetInputAxisValue.h
+++ b/Engine/Source/Editor/BlueprintGraph/Classes/K2Node_GetInputAxisValue.h
@@ -2,6 +2,7 @@
 
 #pragma once
 #include "K2Node_CallFunction.h"
+#include "EdGraph/EdGraphNodeUtils.h" // for FNodeTextCache
 #include "K2Node_GetInputAxisValue.generated.h"
 
 UCLASS(MinimalAPI)
@@ -37,4 +38,8 @@ class UK2Node_GetInputAxisValue : public UK2Node_CallFunction
 	// End UK2Node interface
 	
 	void Initialize(const FName AxisName);
+
+private:
+	/** Constructing FText strings can be costly, so we cache the node's title */
+	FNodeTextCache CachedNodeTitle;
 };
diff --git a/Engine/Source/Editor/BlueprintGraph/Classes/K2Node_GetNumEnumEntries.h b/Engine/Source/Editor/BlueprintGraph/Classes/K2Node_GetNumEnumEntries.h
--- a/Engine/Source/Editor/BlueprintGraph/Classes/K2Node_GetNumEnumEntries.h
+++ b/Engine/Source/Editor/BlueprintGraph/Classes/K2Node_GetNumEnumEntries.h
@@ -2,6 +2,7 @@
 
 #pragma once
 #include "NodeDependingOnEnumInterface.h"
+#include "EdGraph/EdGraphNodeUtils.h" // for FNodeTextCache
 #include "K2Node_GetNumEnumEntries.generated.h"
 
 UCLASS(MinimalAPI)
@@ -30,5 +31,9 @@ class UK2Node_GetNumEnumEntries : public UK2Node, public INodeDependingOnEnumInt
 	virtual class UEnum* GetEnum() const override { return Enum; }
 	virtual bool ShouldBeReconstructedAfterEnumChanged() const override {return false;}
 	// End of INodeDependingOnEnumInterface
+
+private:
+	/** Constructing FText strings can be costly, so we cache the node's title */
+	FNodeTextCache CachedNodeTitle;
 };
 
diff --git a/Engine/Source/Editor/BlueprintGraph/Classes/K2Node_InputAction.h b/Engine/Source/Editor/BlueprintGraph/Classes/K2Node_InputAction.h
--- a/Engine/Source/Editor/BlueprintGraph/Classes/K2Node_InputAction.h
+++ b/Engine/Source/Editor/BlueprintGraph/Classes/K2Node_InputAction.h
@@ -3,6 +3,7 @@
 
 #pragma once
 #include "K2Node.h"
+#include "EdGraph/EdGraphNodeUtils.h" // for FNodeTextCache
 #include "K2Node_InputAction.generated.h"
 
 UCLASS(MinimalAPI)
@@ -54,4 +55,7 @@ class UK2Node_InputAction : public UK2Node
 
 private:
 	void CreateInputActionEvent(FKismetCompilerContext& CompilerContext, UEdGraph* SourceGraph, UEdGraphPin* InputActionPin, const EInputEvent InputKeyEvent);
+
+	/** Constructing FText strings can be costly, so we cache the node's title */
+	FNodeTextCache CachedNodeTitle;
 };
diff --git a/Engine/Source/Editor/BlueprintGraph/Classes/K2Node_InputAxisEvent.h b/Engine/Source/Editor/BlueprintGraph/Classes/K2Node_InputAxisEvent.h
--- a/Engine/Source/Editor/BlueprintGraph/Classes/K2Node_InputAxisEvent.h
+++ b/Engine/Source/Editor/BlueprintGraph/Classes/K2Node_InputAxisEvent.h
@@ -2,6 +2,7 @@
 
 #pragma once
 #include "K2Node_Event.h"
+#include "EdGraph/EdGraphNodeUtils.h" // for FNodeTextCache
 #include "K2Node_InputAxisEvent.generated.h"
 
 UCLASS(MinimalAPI)
@@ -44,4 +45,8 @@ class UK2Node_InputAxisEvent : public UK2Node_Event
 	// End UK2Node interface
 
 	void Initialize(const FName AxisName);
+
+private:
+	/** Constructing FText strings can be costly, so we cache the node's title */
+	FNodeTextCache CachedNodeTitle;
 };
diff --git a/Engine/Source/Editor/BlueprintGraph/Classes/K2Node_InputKey.h b/Engine/Source/Editor/BlueprintGraph/Classes/K2Node_InputKey.h
--- a/Engine/Source/Editor/BlueprintGraph/Classes/K2Node_InputKey.h
+++ b/Engine/Source/Editor/BlueprintGraph/Classes/K2Node_InputKey.h
@@ -3,6 +3,7 @@
 
 #pragma once
 #include "K2Node.h"
+#include "EdGraph/EdGraphNodeUtils.h" // for FNodeTextCache
 #include "K2Node_InputKey.generated.h"
 
 UCLASS(MinimalAPI)
@@ -78,4 +79,6 @@ class UK2Node_InputKey : public UK2Node
 private:
 	void CreateInputKeyEvent(FKismetCompilerContext& CompilerContext, UEdGraph* SourceGraph, UEdGraphPin* InputKeyPin, const EInputEvent KeyEvent);
 
+	/** Constructing FText strings can be costly, so we cache the node's title */
+	FNodeTextCache CachedNodeTitle;
 };
diff --git a/Engine/Source/Editor/BlueprintGraph/Classes/K2Node_MakeStruct.h b/Engine/Source/Editor/BlueprintGraph/Classes/K2Node_MakeStruct.h
--- a/Engine/Source/Editor/BlueprintGraph/Classes/K2Node_MakeStruct.h
+++ b/Engine/Source/Editor/BlueprintGraph/Classes/K2Node_MakeStruct.h
@@ -2,6 +2,7 @@
 
 #pragma once
 #include "K2Node_StructMemberSet.h"
+#include "EdGraph/EdGraphNodeUtils.h" // for FNodeTextCache
 #include "K2Node_MakeStruct.generated.h"
 
 // Pure kismet node that creates a struct with specified values for each member
@@ -41,4 +42,8 @@ protected:
 		virtual void CustomizePinData(UEdGraphPin* Pin, FName SourcePropertyName, int32 ArrayIndex, UProperty* Property) const override;
 		virtual bool CanTreatPropertyAsOptional(UProperty* TestProperty) const override;
 	};
+
+private:
+	/** Constructing FText strings can be costly, so we cache the node's title */
+	FNodeTextCache CachedNodeTitle;
 };
diff --git a/Engine/Source/Editor/BlueprintGraph/Classes/K2Node_MathExpression.h b/Engine/Source/Editor/BlueprintGraph/Classes/K2Node_MathExpression.h
--- a/Engine/Source/Editor/BlueprintGraph/Classes/K2Node_MathExpression.h
+++ b/Engine/Source/Editor/BlueprintGraph/Classes/K2Node_MathExpression.h
@@ -2,6 +2,8 @@
 
 #pragma once
 
+#include "K2Node_Composite.h"
+#include "EdGraph/EdGraphNodeUtils.h" // for FNodeTextCache
 #include "K2Node_MathExpression.generated.h"
 
 /**
@@ -59,8 +61,10 @@ private:
 
 private:
 	/** Cached so we don't have to regenerate it when the graph is recompiled */
-	//UPROPERTY(Transient)
 	TSharedPtr<class FCompilerResultsLog> CachedMessageLog;
+
+	/** Constructing FText strings can be costly, so we cache the node's title */
+	FNodeTextCache CachedNodeTitle;
 };
 
 
diff --git a/Engine/Source/Editor/BlueprintGraph/Classes/K2Node_Message.h b/Engine/Source/Editor/BlueprintGraph/Classes/K2Node_Message.h
--- a/Engine/Source/Editor/BlueprintGraph/Classes/K2Node_Message.h
+++ b/Engine/Source/Editor/BlueprintGraph/Classes/K2Node_Message.h
@@ -3,6 +3,7 @@
 
 #pragma once
 #include "K2Node_CallFunction.h"
+#include "EdGraph/EdGraphNodeUtils.h" // for FNodeTitleTextTable
 #include "K2Node_Message.generated.h"
 
 UCLASS(MinimalAPI)
@@ -28,5 +29,8 @@ class UK2Node_Message : public UK2Node_CallFunction
 protected:
 	virtual void EnsureFunctionIsInBlueprint() override;
 	// End K2Node_CallFunction Interface.
+
+	/** Constructing FText strings can be costly, so we cache the node's title */
+	FNodeTitleTextTable CachedNodeTitles;
 };
 
diff --git a/Engine/Source/Editor/BlueprintGraph/Classes/K2Node_SetFieldsInStruct.h b/Engine/Source/Editor/BlueprintGraph/Classes/K2Node_SetFieldsInStruct.h
--- a/Engine/Source/Editor/BlueprintGraph/Classes/K2Node_SetFieldsInStruct.h
+++ b/Engine/Source/Editor/BlueprintGraph/Classes/K2Node_SetFieldsInStruct.h
@@ -2,6 +2,7 @@
 
 #pragma once
 #include "K2Node_MakeStruct.h"
+#include "EdGraph/EdGraphNodeUtils.h" // for FNodeTextCache
 #include "K2Node_SetFieldsInStruct.generated.h"
 
 // Pure kismet node that creates a struct with specified values for each member
@@ -33,4 +34,8 @@ class UK2Node_SetFieldsInStruct : public UK2Node_MakeStruct
 	BLUEPRINTGRAPH_API void RemoveFieldPins(const UEdGraphPin* InGraphPin, EPinsToRemove Selection);
 	BLUEPRINTGRAPH_API bool AllPinsAreShown() const;
 	BLUEPRINTGRAPH_API void RestoreAllPins();
+
+private:
+	/** Constructing FText strings can be costly, so we cache the node's title */
+	FNodeTextCache CachedNodeTitle;
 };
diff --git a/Engine/Source/Editor/BlueprintGraph/Classes/K2Node_SpawnActor.h b/Engine/Source/Editor/BlueprintGraph/Classes/K2Node_SpawnActor.h
--- a/Engine/Source/Editor/BlueprintGraph/Classes/K2Node_SpawnActor.h
+++ b/Engine/Source/Editor/BlueprintGraph/Classes/K2Node_SpawnActor.h
@@ -3,6 +3,7 @@
 
 #pragma once
 #include "K2Node.h"
+#include "EdGraph/EdGraphNodeUtils.h" // for FNodeTextCache
 #include "K2Node_SpawnActor.generated.h"
 
 UCLASS(MinimalAPI)
@@ -57,4 +58,7 @@ private:
 
 	/** Tooltip text for this node. */
 	FString NodeTooltip;
+
+	/** Constructing FText strings can be costly, so we cache the node's title */
+	FNodeTextCache CachedNodeTitle;
 };
diff --git a/Engine/Source/Editor/BlueprintGraph/Classes/K2Node_SpawnActorFromClass.h b/Engine/Source/Editor/BlueprintGraph/Classes/K2Node_SpawnActorFromClass.h
--- a/Engine/Source/Editor/BlueprintGraph/Classes/K2Node_SpawnActorFromClass.h
+++ b/Engine/Source/Editor/BlueprintGraph/Classes/K2Node_SpawnActorFromClass.h
@@ -3,6 +3,7 @@
 
 #pragma once
 #include "K2Node.h"
+#include "EdGraph/EdGraphNodeUtils.h" // for FNodeTextCache
 #include "K2Node_SpawnActorFromClass.generated.h"
 
 UCLASS(MinimalAPI)
@@ -65,4 +66,7 @@ private:
 
 	/** Tooltip text for this node. */
 	FString NodeTooltip;
+
+	/** Constructing FText strings can be costly, so we cache the node's title */
+	FNodeTextCache CachedNodeTitle;
 };
diff --git a/Engine/Source/Editor/BlueprintGraph/Classes/K2Node_StructMemberGet.h b/Engine/Source/Editor/BlueprintGraph/Classes/K2Node_StructMemberGet.h
--- a/Engine/Source/Editor/BlueprintGraph/Classes/K2Node_StructMemberGet.h
+++ b/Engine/Source/Editor/BlueprintGraph/Classes/K2Node_StructMemberGet.h
@@ -2,6 +2,7 @@
 
 #pragma once
 #include "K2Node_StructOperation.h"
+#include "EdGraph/EdGraphNodeUtils.h" // for FNodeTextCache
 #include "K2Node_StructMemberGet.generated.h"
 
 // Pure kismet node that gets one or more member variables of a struct
@@ -31,5 +32,9 @@ class UK2Node_StructMemberGet : public UK2Node_StructOperation
 
 	// AllocateDefaultPins with just one member set
 	BLUEPRINTGRAPH_API void AllocatePinsForSingleMemberGet(FName MemberName);
+
+private:
+	/** Constructing FText strings can be costly, so we cache the node's title */
+	FNodeTextCache CachedNodeTitle;
 };
 
diff --git a/Engine/Source/Editor/BlueprintGraph/Classes/K2Node_StructMemberSet.h b/Engine/Source/Editor/BlueprintGraph/Classes/K2Node_StructMemberSet.h
--- a/Engine/Source/Editor/BlueprintGraph/Classes/K2Node_StructMemberSet.h
+++ b/Engine/Source/Editor/BlueprintGraph/Classes/K2Node_StructMemberSet.h
@@ -2,6 +2,7 @@
 
 #pragma once
 #include "K2Node_StructOperation.h"
+#include "EdGraph/EdGraphNodeUtils.h" // for FNodeTextCache
 #include "K2Node_StructMemberSet.generated.h"
 
 // Imperative kismet node that sets one or more member variables of a struct
@@ -29,5 +30,9 @@ class UK2Node_StructMemberSet : public UK2Node_StructOperation
 	virtual ERedirectType DoPinsMatchForReconstruction(const UEdGraphPin* NewPin, int32 NewPinIndex, const UEdGraphPin* OldPin, int32 OldPinIndex)  const override;
 	virtual class FNodeHandlingFunctor* CreateNodeHandler(class FKismetCompilerContext& CompilerContext) const override;
 	// End of UK2Node interface
+
+private:
+	/** Constructing FText strings can be costly, so we cache the node's title */
+	FNodeTextCache CachedNodeTitle;
 };
 
diff --git a/Engine/Source/Editor/BlueprintGraph/Classes/K2Node_SwitchEnum.h b/Engine/Source/Editor/BlueprintGraph/Classes/K2Node_SwitchEnum.h
--- a/Engine/Source/Editor/BlueprintGraph/Classes/K2Node_SwitchEnum.h
+++ b/Engine/Source/Editor/BlueprintGraph/Classes/K2Node_SwitchEnum.h
@@ -4,6 +4,7 @@
 #pragma once
 #include "K2Node_Switch.h"
 #include "NodeDependingOnEnumInterface.h"
+#include "EdGraph/EdGraphNodeUtils.h" // for FNodeTextCache
 #include "K2Node_SwitchEnum.generated.h"
 
 UCLASS(MinimalAPI)
@@ -57,4 +58,8 @@ protected:
 	
 	/** Don't support removing pins from an enum */
 	virtual void RemovePin(UEdGraphPin* TargetPin) override {}
+
+private:
+	/** Constructing FText strings can be costly, so we cache the node's title */
+	FNodeTextCache CachedNodeTitle;
 };
diff --git a/Engine/Source/Editor/BlueprintGraph/Classes/K2Node_TemporaryVariable.h b/Engine/Source/Editor/BlueprintGraph/Classes/K2Node_TemporaryVariable.h
--- a/Engine/Source/Editor/BlueprintGraph/Classes/K2Node_TemporaryVariable.h
+++ b/Engine/Source/Editor/BlueprintGraph/Classes/K2Node_TemporaryVariable.h
@@ -2,6 +2,7 @@
 
 #pragma once
 #include "K2Node.h"
+#include "EdGraph/EdGraphNodeUtils.h" // for FNodeTextCache
 #include "K2Node_TemporaryVariable.generated.h"
 
 UCLASS(MinimalAPI)
@@ -33,5 +34,9 @@ class UK2Node_TemporaryVariable : public UK2Node
 	virtual void GetMenuActions(FBlueprintActionDatabaseRegistrar& ActionRegistrar) const override;
 	virtual FText GetMenuCategory() const override;
 	// End UK2Node interface.
+
+private:
+	/** Constructing FText strings can be costly, so we cache the node's title */
+	FNodeTextCache CachedNodeTitle;
 };
 
diff --git a/Engine/Source/Editor/BlueprintGraph/Classes/K2Node_VariableGet.h b/Engine/Source/Editor/BlueprintGraph/Classes/K2Node_VariableGet.h
--- a/Engine/Source/Editor/BlueprintGraph/Classes/K2Node_VariableGet.h
+++ b/Engine/Source/Editor/BlueprintGraph/Classes/K2Node_VariableGet.h
@@ -3,6 +3,7 @@
 
 #pragma once
 #include "K2Node_Variable.h"
+#include "EdGraph/EdGraphNodeUtils.h" // for FNodeTextCache
 #include "K2Node_VariableGet.generated.h"
 
 UCLASS()
@@ -22,5 +23,8 @@ class BLUEPRINTGRAPH_API UK2Node_VariableGet : public UK2Node_Variable
 	virtual bool ShouldShowNodeProperties() const override { return true; }
 	virtual class FNodeHandlingFunctor* CreateNodeHandler(class FKismetCompilerContext& CompilerContext) const override;
 	// End K2Node interface
+private:
+	/** Constructing FText strings can be costly, so we cache the node's title */
+	FNodeTextCache CachedNodeTitle;
 };
 
diff --git a/Engine/Source/Editor/BlueprintGraph/Classes/K2Node_VariableSet.h b/Engine/Source/Editor/BlueprintGraph/Classes/K2Node_VariableSet.h
--- a/Engine/Source/Editor/BlueprintGraph/Classes/K2Node_VariableSet.h
+++ b/Engine/Source/Editor/BlueprintGraph/Classes/K2Node_VariableSet.h
@@ -3,6 +3,7 @@
 
 #pragma once
 #include "K2Node_Variable.h"
+#include "EdGraph/EdGraphNodeUtils.h" // for FNodeTextCache
 #include "K2Node_VariableSet.generated.h"
 
 UCLASS(MinimalAPI)
@@ -26,5 +27,9 @@ class UK2Node_VariableSet : public UK2Node_Variable
 	BLUEPRINTGRAPH_API bool HasLocalRepNotify() const;
 	BLUEPRINTGRAPH_API FName GetRepNotifyName() const;
 	BLUEPRINTGRAPH_API bool ShouldFlushDormancyOnSet() const;
+
+private:
+	/** Constructing FText strings can be costly, so we cache the node's title */
+	FNodeTextCache CachedNodeTitle;
 };
 
diff --git a/Engine/Source/Editor/BlueprintGraph/Classes/K2Node_VariableSetRef.h b/Engine/Source/Editor/BlueprintGraph/Classes/K2Node_VariableSetRef.h
--- a/Engine/Source/Editor/BlueprintGraph/Classes/K2Node_VariableSetRef.h
+++ b/Engine/Source/Editor/BlueprintGraph/Classes/K2Node_VariableSetRef.h
@@ -2,6 +2,8 @@
 
 
 #pragma once
+#include "K2Node.h"
+#include "EdGraph/EdGraphNodeUtils.h" // for FNodeTextCache
 #include "K2Node_VariableSetRef.generated.h"
 
 UCLASS(MinimalAPI)
@@ -36,5 +38,9 @@ class UK2Node_VariableSetRef : public UK2Node
 
 	/** Returns the pin that specifies the value to set */
 	BLUEPRINTGRAPH_API UEdGraphPin* GetValuePin() const;
+
+private:
+	/** Constructing FText strings can be costly, so we cache the node's title */
+	FNodeTextCache CachedNodeTitle;
 };
 
diff --git a/Engine/Source/Editor/BlueprintGraph/Private/BlueprintActionDatabase.cpp b/Engine/Source/Editor/BlueprintGraph/Private/BlueprintActionDatabase.cpp
--- a/Engine/Source/Editor/BlueprintGraph/Private/BlueprintActionDatabase.cpp
+++ b/Engine/Source/Editor/BlueprintGraph/Private/BlueprintActionDatabase.cpp
@@ -298,6 +298,13 @@ namespace BlueprintActionDatabaseImpl
 	 */
 	static void GetNodeSpecificActions(TSubclassOf<UEdGraphNode const> const NodeClass, FBlueprintActionDatabaseRegistrar& Registrar);
 
+	/**
+	 * 
+	 * 
+	 * @param  BlueprintAction	
+	 */
+	static void PrimeBlueprintAction(UBlueprintNodeSpawner* BlueprintAction);
+
 	/**
 	 * Callback to refresh the database when a blueprint has been altered 
 	 * (clears the database entries for the blueprint's classes and recreates 
@@ -575,6 +582,31 @@ static void BlueprintActionDatabaseImpl::GetNodeSpecificActions(TSubclassOf<UEdG
 	}
 }
 
+//------------------------------------------------------------------------------
+static void BlueprintActionDatabaseImpl::PrimeBlueprintAction(UBlueprintNodeSpawner* BlueprintAction)
+{
+	// @TODO: for UBlueprintFunctionNodeSpawner, we almost completely don't need
+	//        a template node (we can gather all the info from the UFunction)...
+	//        so we should save on template cache and not do this for function
+	//        node spawners (have to fill the holes first).
+	UEdGraphNode* TemplateNode = BlueprintAction->GetTemplateNode();
+	if (TemplateNode != nullptr)
+	{
+		// since we're priming incrementally, someone could have already
+		// requested this template, and allocated its pins (don't need to do 
+		// redundant work)
+		if (TemplateNode->Pins.Num() == 0)
+		{
+			// in certain scenarios we need pin information from the 
+			// spawner (to help filter by pin context)
+			TemplateNode->AllocateDefaultPins();
+		}
+
+		// cache any FText::Format() calls 
+		TemplateNode->GetNodeTitle(ENodeTitleType::ListView);
+	}
+}
+
 //------------------------------------------------------------------------------
 static void BlueprintActionDatabaseImpl::OnBlueprintChanged(UBlueprint* Blueprint)
 {
@@ -727,16 +759,7 @@ void FBlueprintActionDatabase::Tick(float DeltaTime)
 					FScopedDurationTimer ScopedTimer(TotalFrameDuration);
 
 					UBlueprintNodeSpawner* Action = (*ClassActionList)[ActionListIndex];
-					UEdGraphNode* TemplateNode = Action->GetTemplateNode();
-					// since we're doing this incrementally, someone could have 
-					// already requested this template, and allocated its pins 
-					// (don't need to do redundant work)
-					if ((TemplateNode != nullptr) && (TemplateNode->Pins.Num() == 0))
-					{
-						// in certain scenarios we need pin information from the 
-						// spawner (to help filter by pin context)
-						TemplateNode->AllocateDefaultPins();
-					}
+					BlueprintActionDatabaseImpl::PrimeBlueprintAction(Action);
 				}
 
 				if (ActionListIndex >= ClassActionList->Num())
diff --git a/Engine/Source/Editor/BlueprintGraph/Private/BlueprintBoundEventNodeSpawner.cpp b/Engine/Source/Editor/BlueprintGraph/Private/BlueprintBoundEventNodeSpawner.cpp
--- a/Engine/Source/Editor/BlueprintGraph/Private/BlueprintBoundEventNodeSpawner.cpp
+++ b/Engine/Source/Editor/BlueprintGraph/Private/BlueprintBoundEventNodeSpawner.cpp
@@ -49,8 +49,13 @@ UEdGraphNode* UBlueprintBoundEventNodeSpawner::Invoke(UEdGraph* ParentGraph, FVe
 //------------------------------------------------------------------------------
 FText UBlueprintBoundEventNodeSpawner::GetDefaultMenuName() const
 {
-	FText const DelegateName = FText::FromName(EventDelegate->GetFName());
-	return FText::Format(LOCTEXT("ComponentEventName", "Add {0}"), DelegateName);
+	if (CachedMenuName.IsOutOfDate())
+	{
+		FText const DelegateName = FText::FromName(EventDelegate->GetFName());
+		// FText::Format() is slow, so we cache this to save on performance
+		CachedMenuName = FText::Format(LOCTEXT("ComponentEventName", "Add {0}"), DelegateName);
+	}
+	return CachedMenuName;
 }
 
 //------------------------------------------------------------------------------
diff --git a/Engine/Source/Editor/BlueprintGraph/Private/BlueprintComponentNodeSpawner.cpp b/Engine/Source/Editor/BlueprintGraph/Private/BlueprintComponentNodeSpawner.cpp
--- a/Engine/Source/Editor/BlueprintGraph/Private/BlueprintComponentNodeSpawner.cpp
+++ b/Engine/Source/Editor/BlueprintGraph/Private/BlueprintComponentNodeSpawner.cpp
@@ -94,7 +94,12 @@ UEdGraphNode* UBlueprintComponentNodeSpawner::Invoke(UEdGraph* ParentGraph, FVec
 FText UBlueprintComponentNodeSpawner::GetDefaultMenuName() const
 {
 	check(ComponentClass != nullptr);
-	return FText::Format(LOCTEXT("AddComponentMenuName", "Add {0}"), FText::FromName(ComponentClass->GetFName()));
+	if (CachedMenuName.IsOutOfDate())
+	{
+		// FText::Format() is slow, so we cache this to save on performance
+		CachedMenuName = FText::Format(LOCTEXT("AddComponentMenuName", "Add {0}"), FText::FromName(ComponentClass->GetFName()));
+	}
+	return CachedMenuName;
 }
 
 //------------------------------------------------------------------------------
diff --git a/Engine/Source/Editor/BlueprintGraph/Private/BlueprintEventNodeSpawner.cpp b/Engine/Source/Editor/BlueprintGraph/Private/BlueprintEventNodeSpawner.cpp
--- a/Engine/Source/Editor/BlueprintGraph/Private/BlueprintEventNodeSpawner.cpp
+++ b/Engine/Source/Editor/BlueprintGraph/Private/BlueprintEventNodeSpawner.cpp
@@ -140,21 +140,21 @@ UEdGraphNode* UBlueprintEventNodeSpawner::Invoke(UEdGraph* ParentGraph, FVector2
 //------------------------------------------------------------------------------
 FText UBlueprintEventNodeSpawner::GetDefaultMenuName() const
 {
-	FText EventName;
-	if (EventFunc != nullptr)
-	{
-		EventName = FText::FromString(UEdGraphSchema_K2::GetFriendlySignitureName(EventFunc));
-	}
-	else if (!CustomEventName.IsNone())
+	if (CachedMenuName.IsOutOfDate())
 	{
-		EventName = FText::FromName(CustomEventName);
-	}
+		FText EventName;
+		if (EventFunc != nullptr)
+		{
+			EventName = FText::FromString(UEdGraphSchema_K2::GetFriendlySignitureName(EventFunc));
+		}
+		else if (!CustomEventName.IsNone())
+		{
+			EventName = FText::FromName(CustomEventName);
+		}
 
-	if (!EventName.IsEmpty())
-	{
-		EventName = FText::Format(LOCTEXT("EventWithSignatureName", "Event {0}"), EventName);
+		CachedMenuName = FText::Format(LOCTEXT("EventWithSignatureName", "Event {0}"), EventName);
 	}
-	return EventName;
+	return CachedMenuName;
 }
 
 //------------------------------------------------------------------------------
diff --git a/Engine/Source/Editor/BlueprintGraph/Private/BlueprintVariableNodeSpawner.cpp b/Engine/Source/Editor/BlueprintGraph/Private/BlueprintVariableNodeSpawner.cpp
--- a/Engine/Source/Editor/BlueprintGraph/Private/BlueprintVariableNodeSpawner.cpp
+++ b/Engine/Source/Editor/BlueprintGraph/Private/BlueprintVariableNodeSpawner.cpp
@@ -94,20 +94,22 @@ UEdGraphNode* UBlueprintVariableNodeSpawner::Invoke(UEdGraph* ParentGraph, FVect
 //------------------------------------------------------------------------------
 FText UBlueprintVariableNodeSpawner::GetDefaultMenuName() const
 {	
-	FText MenuName = GetVariableName();
-	if (NodeClass != nullptr)
+	if (CachedMenuName.IsOutOfDate())
 	{
+		FText VarName = GetVariableName();
+		check(NodeClass != nullptr);
+
+		// FText::Format() is slow, so we cache this to save on performance
 		if (NodeClass->IsChildOf<UK2Node_VariableGet>())
 		{
-			MenuName = FText::Format(LOCTEXT("GetterMenuName", "Get {0}"), MenuName);
+			CachedMenuName = FText::Format(LOCTEXT("GetterMenuName", "Get {0}"), VarName);
 		}
 		else if (NodeClass->IsChildOf<UK2Node_VariableSet>())
 		{
-			MenuName = FText::Format(LOCTEXT("SetterMenuName", "Set {0}"), MenuName);
+			CachedMenuName = FText::Format(LOCTEXT("SetterMenuName", "Set {0}"), VarName);
 		}
 	}
-
-	return MenuName;
+	return CachedMenuName;
 }
 
 //------------------------------------------------------------------------------
diff --git a/Engine/Source/Editor/BlueprintGraph/Private/K2Node_ActorBoundEvent.cpp b/Engine/Source/Editor/BlueprintGraph/Private/K2Node_ActorBoundEvent.cpp
--- a/Engine/Source/Editor/BlueprintGraph/Private/K2Node_ActorBoundEvent.cpp
+++ b/Engine/Source/Editor/BlueprintGraph/Private/K2Node_ActorBoundEvent.cpp
@@ -79,17 +79,22 @@ void UK2Node_ActorBoundEvent::DestroyNode()
 
 FText UK2Node_ActorBoundEvent::GetNodeTitle(ENodeTitleType::Type TitleType) const
 {
-	FText TargetName = LOCTEXT("None", "None");
-	if( EventOwner )
+	if (EventOwner == nullptr)
 	{
-		TargetName = FText::FromString(EventOwner->GetActorLabel());		
+		FFormatNamedArguments Args;
+		Args.Add(TEXT("DelegatePropertyName"), FText::FromName(DelegatePropertyName));
+		return FText::Format(LOCTEXT("ActorBoundEventTitle", "{DelegatePropertyName} (None)"), Args);
 	}
+	else if (CachedNodeTitle.IsOutOfDate())
+	{
+		FFormatNamedArguments Args;
+		Args.Add(TEXT("DelegatePropertyName"), FText::FromName(DelegatePropertyName));
+		Args.Add(TEXT("TargetName"), FText::FromString(EventOwner->GetActorLabel()));
 
-	FFormatNamedArguments Args;
-	Args.Add(TEXT("DelegatePropertyName"), FText::FromName(DelegatePropertyName));
-	Args.Add(TEXT("TargetName"), TargetName);
-
-	return FText::Format(LOCTEXT("ActorBoundEventTitle", "{DelegatePropertyName} ({TargetName})"), Args);
+		// FText::Format() is slow, so we cache this to save on performance
+		CachedNodeTitle = FText::Format(LOCTEXT("ActorBoundEventTitle", "{DelegatePropertyName} ({TargetName})"), Args);
+	}
+	return CachedNodeTitle;
 }
 
 FString UK2Node_ActorBoundEvent::GetTooltip() const
diff --git a/Engine/Source/Editor/BlueprintGraph/Private/K2Node_AddComponent.cpp b/Engine/Source/Editor/BlueprintGraph/Private/K2Node_AddComponent.cpp
--- a/Engine/Source/Editor/BlueprintGraph/Private/K2Node_AddComponent.cpp
+++ b/Engine/Source/Editor/BlueprintGraph/Private/K2Node_AddComponent.cpp
@@ -324,7 +324,7 @@ void UK2Node_AddComponent::PostPasteNode()
 FText UK2Node_AddComponent::GetNodeTitle(ENodeTitleType::Type TitleType) const
 {
 	UEdGraphPin* TemplateNamePin = GetTemplateNamePin();
-	if (TemplateNamePin != NULL)
+	if (CachedNodeTitle.IsOutOfDate() && (TemplateNamePin != NULL))
 	{
 		FString TemplateName = TemplateNamePin->DefaultValue;
 		UBlueprint* Blueprint = GetBlueprint();
@@ -340,35 +340,40 @@ FText UK2Node_AddComponent::GetNodeTitle(ENodeTitleType::Type TitleType) const
 			{
 				FFormatNamedArguments Args;
 				Args.Add(TEXT("StaticMeshName"), FText::FromString(StaticMeshComp->StaticMesh->GetName()));
-				return FText::Format(LOCTEXT("AddStaticMesh", "Add StaticMesh {StaticMeshName}"), Args);
+				CachedNodeTitle = FText::Format(LOCTEXT("AddStaticMesh", "Add StaticMesh {StaticMeshName}"), Args);
 			}
 			else if(SkelMeshComp != NULL && SkelMeshComp->SkeletalMesh != NULL)
 			{
 				FFormatNamedArguments Args;
 				Args.Add(TEXT("SkeletalMeshName"), FText::FromString(SkelMeshComp->SkeletalMesh->GetName()));
-				return FText::Format(LOCTEXT("AddSkeletalMesh", "Add SkeletalMesh {SkeletalMeshName}"), Args);
+				CachedNodeTitle = FText::Format(LOCTEXT("AddSkeletalMesh", "Add SkeletalMesh {SkeletalMeshName}"), Args);
 			}
 			else if(PSysComp != NULL && PSysComp->Template != NULL)
 			{
 				FFormatNamedArguments Args;
 				Args.Add(TEXT("ParticleSystemName"), FText::FromString(PSysComp->Template->GetName()));
-				return FText::Format(LOCTEXT("AddParticleSystem", "Add ParticleSystem {ParticleSystemName}"), Args);
+				CachedNodeTitle = FText::Format(LOCTEXT("AddParticleSystem", "Add ParticleSystem {ParticleSystemName}"), Args);
 			}
 			else if (SubActorComp && SubActorComp->ChildActorClass)
 			{
 				FFormatNamedArguments Args;
 				Args.Add(TEXT("ComponentClassName"), FText::FromString(SubActorComp->ChildActorClass->GetName()));
-				return FText::Format(LOCTEXT("AddChildActorComponent", "Add ChildActorComponent {ComponentClassName}"), Args);
+				CachedNodeTitle = FText::Format(LOCTEXT("AddChildActorComponent", "Add ChildActorComponent {ComponentClassName}"), Args);
 			}
 			else
 			{
 				FFormatNamedArguments Args;
 				Args.Add(TEXT("ClassName"), FText::FromString(SourceTemplate->GetClass()->GetName()));
-				return FText::Format(LOCTEXT("AddClass", "Add {ClassName}"), Args);		
+				CachedNodeTitle = FText::Format(LOCTEXT("AddClass", "Add {ClassName}"), Args);
 			}
 		}
 	}
 
+	// FText::Format() is slow, so we cache the title to save on performance
+	if (!CachedNodeTitle.IsOutOfDate())
+	{
+		return CachedNodeTitle;
+	}
 	return Super::GetNodeTitle(TitleType);
 }
 
diff --git a/Engine/Source/Editor/BlueprintGraph/Private/K2Node_AssignDelegate.cpp b/Engine/Source/Editor/BlueprintGraph/Private/K2Node_AssignDelegate.cpp
--- a/Engine/Source/Editor/BlueprintGraph/Private/K2Node_AssignDelegate.cpp
+++ b/Engine/Source/Editor/BlueprintGraph/Private/K2Node_AssignDelegate.cpp
@@ -20,27 +20,28 @@ UK2Node_AssignDelegate::UK2Node_AssignDelegate(const class FPostConstructInitial
 //------------------------------------------------------------------------------
 FText UK2Node_AssignDelegate::GetNodeTitle(ENodeTitleType::Type TitleType) const
 {
-	FText NodeTitle;
 	if (TitleType == ENodeTitleType::ListView)
 	{
-		FText PropertyName;
-		if (UProperty* Property = GetProperty())
+		if (CachedListTitle.IsOutOfDate())
 		{
-			bool const bShowFriendlyNames = GetDefault<UEditorStyleSettings>()->bShowFriendlyNames;
-			PropertyName = FText::FromString(bShowFriendlyNames ? UEditorEngine::GetFriendlyName(Property) : Property->GetName());
-		}
-		else
-		{
-			PropertyName = LOCTEXT("InvalidPropertyTitle", "<invalid delegate>");
+			FText PropertyName;
+			if (UProperty* Property = GetProperty())
+			{
+				bool const bShowFriendlyNames = GetDefault<UEditorStyleSettings>()->bShowFriendlyNames;
+				PropertyName = FText::FromString(bShowFriendlyNames ? UEditorEngine::GetFriendlyName(Property) : Property->GetName());
+
+				// FText::Format() is slow, so we cache this to save on performance
+				CachedNodeTitle = FText::Format(LOCTEXT("AssignDelegateTitle", "Assign {0}"), PropertyName);
+			}
+			else
+			{
+				return LOCTEXT("InvalidPropertyTitle", "Assign <invalid delegate>");
+			}
 		}
-		NodeTitle = FText::Format(LOCTEXT("AssignDelegateTitle", "Assign {0}"), PropertyName);
-	}
-	else
-	{
-		NodeTitle = Super::GetNodeTitle(TitleType);
+		return CachedNodeTitle;
 	}
-
-	return NodeTitle;
+	
+	return Super::GetNodeTitle(TitleType);
 }
 
 //------------------------------------------------------------------------------
diff --git a/Engine/Source/Editor/BlueprintGraph/Private/K2Node_BreakStruct.cpp b/Engine/Source/Editor/BlueprintGraph/Private/K2Node_BreakStruct.cpp
--- a/Engine/Source/Editor/BlueprintGraph/Private/K2Node_BreakStruct.cpp
+++ b/Engine/Source/Editor/BlueprintGraph/Private/K2Node_BreakStruct.cpp
@@ -194,9 +194,15 @@ void UK2Node_BreakStruct::AllocateDefaultPins()
 
 FText UK2Node_BreakStruct::GetNodeTitle(ENodeTitleType::Type TitleType) const
 {
-	FFormatNamedArguments Args;
-	Args.Add(TEXT("StructName"), FText::FromString(StructType ? StructType->GetName() : FString()));
-	return FText::Format(LOCTEXT("BreakNodeTitle", "Break {StructName}"), Args);
+	if (CachedNodeTitle.IsOutOfDate())
+	{
+		FFormatNamedArguments Args;
+		Args.Add(TEXT("StructName"), FText::FromString(StructType ? StructType->GetName() : FString()));
+
+		// FText::Format() is slow, so we cache this to save on performance
+		CachedNodeTitle = FText::Format(LOCTEXT("BreakNodeTitle", "Break {StructName}"), Args);
+	}
+	return CachedNodeTitle;
 }
 
 FString UK2Node_BreakStruct::GetTooltip() const
diff --git a/Engine/Source/Editor/BlueprintGraph/Private/K2Node_ClassDynamicCast.cpp b/Engine/Source/Editor/BlueprintGraph/Private/K2Node_ClassDynamicCast.cpp
--- a/Engine/Source/Editor/BlueprintGraph/Private/K2Node_ClassDynamicCast.cpp
+++ b/Engine/Source/Editor/BlueprintGraph/Private/K2Node_ClassDynamicCast.cpp
@@ -57,7 +57,11 @@ FLinearColor UK2Node_ClassDynamicCast::GetNodeTitleColor() const
 
 FText UK2Node_ClassDynamicCast::GetNodeTitle(ENodeTitleType::Type TitleType) const
 {
-	return FText::Format(LOCTEXT("NodeTitle", "{0} Class"), Super::GetNodeTitle(TitleType));
+	if (CachedNodeTitle.IsOutOfDate())
+	{
+		CachedNodeTitle = FText::Format(LOCTEXT("NodeTitle", "{0} Class"), Super::GetNodeTitle(TitleType));
+	}
+	return CachedNodeTitle;
 }
 
 UEdGraphPin* UK2Node_ClassDynamicCast::GetCastSourcePin() const
diff --git a/Engine/Source/Editor/BlueprintGraph/Private/K2Node_ComponentBoundEvent.cpp b/Engine/Source/Editor/BlueprintGraph/Private/K2Node_ComponentBoundEvent.cpp
--- a/Engine/Source/Editor/BlueprintGraph/Private/K2Node_ComponentBoundEvent.cpp
+++ b/Engine/Source/Editor/BlueprintGraph/Private/K2Node_ComponentBoundEvent.cpp
@@ -12,10 +12,16 @@ UK2Node_ComponentBoundEvent::UK2Node_ComponentBoundEvent(const class FPostConstr
 
 FText UK2Node_ComponentBoundEvent::GetNodeTitle(ENodeTitleType::Type TitleType) const
 {
-	FFormatNamedArguments Args;
-	Args.Add(TEXT("DelegatePropertyName"), FText::FromName(DelegatePropertyName));
-	Args.Add(TEXT("ComponentPropertyName"), FText::FromName(ComponentPropertyName));
-	return FText::Format(LOCTEXT("ComponentBoundEvent_Title", "{DelegatePropertyName} ({ComponentPropertyName})"), Args);
+	if (CachedNodeTitle.IsOutOfDate())
+	{
+		FFormatNamedArguments Args;
+		Args.Add(TEXT("DelegatePropertyName"), FText::FromName(DelegatePropertyName));
+		Args.Add(TEXT("ComponentPropertyName"), FText::FromName(ComponentPropertyName));
+
+		// FText::Format() is slow, so we cache this to save on performance
+		CachedNodeTitle = FText::Format(LOCTEXT("ComponentBoundEvent_Title", "{DelegatePropertyName} ({ComponentPropertyName})"), Args);
+	}
+	return CachedNodeTitle;
 }
 
 void UK2Node_ComponentBoundEvent::InitializeComponentBoundEventParams(UObjectProperty const* InComponentProperty, const UMulticastDelegateProperty* InDelegateProperty)
diff --git a/Engine/Source/Editor/BlueprintGraph/Private/K2Node_Composite.cpp b/Engine/Source/Editor/BlueprintGraph/Private/K2Node_Composite.cpp
--- a/Engine/Source/Editor/BlueprintGraph/Private/K2Node_Composite.cpp
+++ b/Engine/Source/Editor/BlueprintGraph/Private/K2Node_Composite.cpp
@@ -164,16 +164,22 @@ FLinearColor UK2Node_Composite::GetNodeTitleColor() const
 
 FText UK2Node_Composite::GetNodeTitle(ENodeTitleType::Type TitleType) const
 {
-	if(TitleType == ENodeTitleType::FullTitle)
+	if (BoundGraph == nullptr)
 	{
-		FFormatNamedArguments Args;
-		Args.Add(TEXT("BoundGraphName"), (BoundGraph)? FText::FromString(BoundGraph->GetName()) : LOCTEXT("InvalidGraph", "Invalid Graph"));
-		return FText::Format(LOCTEXT("Collapsed_Name", "{BoundGraphName}\nCollapsed Graph"), Args);
+		return LOCTEXT("InvalidGraph", "Invalid Graph");
+	}
+	else if (TitleType != ENodeTitleType::FullTitle)
+	{
+		return FText::FromString(BoundGraph->GetName());
 	}
-	else
+	else if (CachedNodeTitle.IsOutOfDate()) // TitleType == ENodeTitleType::FullTitle
 	{
-		return (BoundGraph)? FText::FromString(BoundGraph->GetName()) : LOCTEXT("InvalidGraph", "Invalid Graph");
+		FFormatNamedArguments Args;
+		Args.Add(TEXT("BoundGraphName"), (BoundGraph) ? FText::FromString(BoundGraph->GetName()) : LOCTEXT("InvalidGraph", "Invalid Graph"));
+		// FText::Format() is slow, so we cache this to save on performance
+		CachedNodeTitle = FText::Format(LOCTEXT("Collapsed_Name", "{BoundGraphName}\nCollapsed Graph"), Args);
 	}
+	return CachedNodeTitle;
 }
 
 bool UK2Node_Composite::CanUserDeleteNode() const
diff --git a/Engine/Source/Editor/BlueprintGraph/Private/K2Node_ConstructObjectFromClass.cpp b/Engine/Source/Editor/BlueprintGraph/Private/K2Node_ConstructObjectFromClass.cpp
--- a/Engine/Source/Editor/BlueprintGraph/Private/K2Node_ConstructObjectFromClass.cpp
+++ b/Engine/Source/Editor/BlueprintGraph/Private/K2Node_ConstructObjectFromClass.cpp
@@ -246,30 +246,35 @@ FText UK2Node_ConstructObjectFromClass::GetNodeTitleFormat() const
 
 FText UK2Node_ConstructObjectFromClass::GetNodeTitle(ENodeTitleType::Type TitleType) const
 {
-	FText NodeTitle = GetBaseNodeTitle();
-	if (TitleType != ENodeTitleType::ListView)
+	if (TitleType == ENodeTitleType::ListView)
 	{
-		FText SpawnString = NSLOCTEXT("K2Node", "None", "NONE");
-		if (UEdGraphPin* ClassPin = GetClassPin())
+		return GetBaseNodeTitle();
+	}
+	else if (UEdGraphPin* ClassPin = GetClassPin())
+	{
+		if (ClassPin->LinkedTo.Num() > 0)
 		{
-			if(ClassPin->LinkedTo.Num() > 0)
-			{
-				// Blueprint will be determined dynamically, so we don't have the name in this case
-				SpawnString = FText::GetEmpty();
-			}
-			else if(ClassPin->DefaultObject != NULL)
-			{
-				SpawnString = FText::FromString(ClassPin->DefaultObject->GetName());
-			}
+			// Blueprint will be determined dynamically, so we don't have the name in this case
+			return NSLOCTEXT("K2Node", "ConstructObject_Title_Unknown", "Construct");
 		}
-		
-		FFormatNamedArguments Args;
-		Args.Add(TEXT("ClassName"), SpawnString);
-		
-		NodeTitle = FText::Format(GetNodeTitleFormat(), Args);
+		else if (ClassPin->DefaultObject == nullptr)
+		{
+			return NSLOCTEXT("K2Node", "ConstructObject_Title_NONE", "Construct NONE");
+		}
+		else if (CachedNodeTitle.IsOutOfDate())
+		{
+			FFormatNamedArguments Args;
+			Args.Add(TEXT("ClassName"), FText::FromString(ClassPin->DefaultObject->GetName()));
+			// FText::Format() is slow, so we cache this to save on performance
+			CachedNodeTitle = FText::Format(GetNodeTitleFormat(), Args);
+		}
+	}
+	else
+	{
+		return NSLOCTEXT("K2Node", "ConstructObject_Title_NONE", "Construct NONE");
 	}
 
-	return NodeTitle;
+	return CachedNodeTitle;
 }
 
 bool UK2Node_ConstructObjectFromClass::IsCompatibleWithGraph(const UEdGraph* TargetGraph) const 
diff --git a/Engine/Source/Editor/BlueprintGraph/Private/K2Node_CustomEvent.cpp b/Engine/Source/Editor/BlueprintGraph/Private/K2Node_CustomEvent.cpp
--- a/Engine/Source/Editor/BlueprintGraph/Private/K2Node_CustomEvent.cpp
+++ b/Engine/Source/Editor/BlueprintGraph/Private/K2Node_CustomEvent.cpp
@@ -106,23 +106,27 @@ UK2Node_CustomEvent::UK2Node_CustomEvent(const class FPostConstructInitializePro
 
 FText UK2Node_CustomEvent::GetNodeTitle(ENodeTitleType::Type TitleType) const
 {
-	if (CustomFunctionName.IsNone() && (TitleType == ENodeTitleType::ListView))
+	if (CustomFunctionName.IsNone())
 	{
 		return LOCTEXT("ActionMenuTitle", "Custom Event...");
 	}
-	else if ((TitleType == ENodeTitleType::EditableTitle) || (TitleType == ENodeTitleType::ListView))
+	else if (TitleType != ENodeTitleType::FullTitle)
 	{
 		return FText::FromName(CustomFunctionName);
 	}
-	else
+	else if (CachedNodeTitle.IsOutOfDate())
 	{
 		FString RPCString = UK2Node_Event::GetLocalizedNetString(FunctionFlags, false);
-
+		
 		FFormatNamedArguments Args;
 		Args.Add(TEXT("FunctionName"), FText::FromName(CustomFunctionName));
 		Args.Add(TEXT("RPCString"), FText::FromString(RPCString));
-		return FText::Format(NSLOCTEXT("K2Node", "CustomEvent_Name", "{FunctionName}{RPCString}\nCustom Event"), Args);
+
+		// FText::Format() is slow, so we cache this to save on performance
+		CachedNodeTitle = FText::Format(NSLOCTEXT("K2Node", "CustomEvent_Name", "{FunctionName}{RPCString}\nCustom Event"), Args);
 	}
+
+	return CachedNodeTitle;
 }
 
 UEdGraphPin* UK2Node_CustomEvent::CreatePinFromUserDefinition(const TSharedPtr<FUserPinInfo> NewPinInfo)
@@ -160,6 +164,7 @@ void UK2Node_CustomEvent::RenameCustomEventCloseToName(int32 StartIndex)
 void UK2Node_CustomEvent::OnRenameNode(const FString& NewName)
 {
 	CustomFunctionName = *NewName;
+	CachedNodeTitle.MarkDirty();
 	FBlueprintEditorUtils::MarkBlueprintAsStructurallyModified(GetBlueprint());
 }
 
diff --git a/Engine/Source/Editor/BlueprintGraph/Private/K2Node_DelegateSet.cpp b/Engine/Source/Editor/BlueprintGraph/Private/K2Node_DelegateSet.cpp
--- a/Engine/Source/Editor/BlueprintGraph/Private/K2Node_DelegateSet.cpp
+++ b/Engine/Source/Editor/BlueprintGraph/Private/K2Node_DelegateSet.cpp
@@ -164,9 +164,14 @@ FString UK2Node_DelegateSet::GetTooltip() const
 
 FText UK2Node_DelegateSet::GetNodeTitle(ENodeTitleType::Type TitleType) const
 {
-	FFormatNamedArguments Args;
-	Args.Add(TEXT("DelegatePropertyName"), FText::FromName(DelegatePropertyName));
-	return FText::Format(NSLOCTEXT("K2Node", "Assign_Name", "Assign {DelegatePropertyName}"), Args);
+	if (CachedNodeTitle.IsOutOfDate())
+	{
+		FFormatNamedArguments Args;
+		Args.Add(TEXT("DelegatePropertyName"), FText::FromName(DelegatePropertyName));
+		// FText::Format() is slow, so we cache this to save on performance
+		CachedNodeTitle = FText::Format(NSLOCTEXT("K2Node", "Assign_Name", "Assign {DelegatePropertyName}"), Args);
+	}
+	return CachedNodeTitle;
 }
 
 UEdGraphPin* UK2Node_DelegateSet::GetDelegateOwner() const
diff --git a/Engine/Source/Editor/BlueprintGraph/Private/K2Node_DynamicCast.cpp b/Engine/Source/Editor/BlueprintGraph/Private/K2Node_DynamicCast.cpp
--- a/Engine/Source/Editor/BlueprintGraph/Private/K2Node_DynamicCast.cpp
+++ b/Engine/Source/Editor/BlueprintGraph/Private/K2Node_DynamicCast.cpp
@@ -54,12 +54,16 @@ FLinearColor UK2Node_DynamicCast::GetNodeTitleColor() const
 
 FText UK2Node_DynamicCast::GetNodeTitle(ENodeTitleType::Type TitleType) const
 {
-	if(TargetType != NULL)
+	if (TargetType == nullptr)
+	{
+		return NSLOCTEXT("K2Node_DynamicCast", "BadCastNode", "Bad cast node");
+	}
+	else if (CachedNodeTitle.IsOutOfDate())
 	{
 		// If casting to BP class, use BP name not class name (ie. remove the _C)
 		FString TargetName;
 		UBlueprint* CastToBP = UBlueprint::GetBlueprintFromClass(TargetType);
-		if(CastToBP != NULL)
+		if (CastToBP != NULL)
 		{
 			TargetName = CastToBP->GetName();
 		}
@@ -70,13 +74,11 @@ FText UK2Node_DynamicCast::GetNodeTitle(ENodeTitleType::Type TitleType) const
 
 		FFormatNamedArguments Args;
 		Args.Add(TEXT("TargetName"), FText::FromString(TargetName));
-		return FText::Format(NSLOCTEXT("K2Node_DynamicCast", "CastTo", "Cast To {TargetName}"), Args);
-	}
-	// No target type, bad node
-	else
-	{
-		return NSLOCTEXT("K2Node_DynamicCast", "BadCastNode", "Bad cast node");
+
+		// FText::Format() is slow, so we cache this to save on performance
+		CachedNodeTitle = FText::Format(NSLOCTEXT("K2Node_DynamicCast", "CastTo", "Cast To {TargetName}"), Args);
 	}
+	return CachedNodeTitle;
 }
 
 UEdGraphPin* UK2Node_DynamicCast::GetValidCastPin() const
diff --git a/Engine/Source/Editor/BlueprintGraph/Private/K2Node_EnumLiteral.cpp b/Engine/Source/Editor/BlueprintGraph/Private/K2Node_EnumLiteral.cpp
--- a/Engine/Source/Editor/BlueprintGraph/Private/K2Node_EnumLiteral.cpp
+++ b/Engine/Source/Editor/BlueprintGraph/Private/K2Node_EnumLiteral.cpp
@@ -41,16 +41,26 @@ FName UK2Node_EnumLiteral::GetPaletteIcon(FLinearColor& OutColor) const
 
 FString UK2Node_EnumLiteral::GetTooltip() const
 {
-	const FText EnumName = Enum ? FText::FromString(Enum->GetName()) : NSLOCTEXT("K2Node", "BadEnum", "(bad enum)");
-
-	FFormatNamedArguments Args;
-	Args.Add(TEXT("EnumName"), EnumName);
-	return FText::Format(NSLOCTEXT("K2Node", "EnumLiteral_Tooltip", "Literal enum {EnumName}"), Args).ToString();
+	if (Enum == nullptr)
+	{
+		return NSLOCTEXT("K2Node", "BadEnumLiteral_Tooltip", "Literal enum (bad enum)").ToString();
+	}
+	else if (CachedTooltip.IsOutOfDate())
+	{
+		FFormatNamedArguments Args;
+		Args.Add(TEXT("EnumName"), FText::FromName(Enum->GetFName()));
+		CachedTooltip = FText::Format(NSLOCTEXT("K2Node", "EnumLiteral_Tooltip", "Literal enum {EnumName}"), Args);
+	}
+	return ((FText&)CachedTooltip).ToString();	
 }
 
 FText UK2Node_EnumLiteral::GetNodeTitle(ENodeTitleType::Type TitleType) const
 {
-	return FText::FromString(GetTooltip());
+	if (CachedTooltip.IsOutOfDate())
+	{
+		return FText::FromString(GetTooltip());
+	}
+	return CachedTooltip;
 }
 
 class FKCHandler_EnumLiteral : public FNodeHandlingFunctor
diff --git a/Engine/Source/Editor/BlueprintGraph/Private/K2Node_ForEachElementInEnum.cpp b/Engine/Source/Editor/BlueprintGraph/Private/K2Node_ForEachElementInEnum.cpp
--- a/Engine/Source/Editor/BlueprintGraph/Private/K2Node_ForEachElementInEnum.cpp
+++ b/Engine/Source/Editor/BlueprintGraph/Private/K2Node_ForEachElementInEnum.cpp
@@ -138,9 +138,18 @@ FString UK2Node_ForEachElementInEnum::GetTooltip() const
 
 FText UK2Node_ForEachElementInEnum::GetNodeTitle(ENodeTitleType::Type TitleType) const
 {
-	FFormatNamedArguments Args;
-	Args.Add(TEXT("EnumName"), Enum ? FText::FromString(Enum->GetName()) : LOCTEXT("Unknown", "UNKNOWN"));
-	return FText::Format(LOCTEXT("ForEachElementInEnum_Title", "ForEach {EnumName}"), Args);
+	if (Enum == nullptr)
+	{
+		return LOCTEXT("ForEachElementInUnknownEnum_Title", "ForEach UNKNOWN");
+	}
+	else if (CachedNodeTitle.IsOutOfDate())
+	{
+		FFormatNamedArguments Args;
+		Args.Add(TEXT("EnumName"), FText::FromName(Enum->GetFName()));
+		// FText::Format() is slow, so we cache this to save on performance
+		CachedNodeTitle = FText::Format(LOCTEXT("ForEachElementInEnum_Title", "ForEach {EnumName}"), Args);
+	}
+	return CachedNodeTitle;
 }
 
 void UK2Node_ForEachElementInEnum::ExpandNode(class FKismetCompilerContext& CompilerContext, UEdGraph* SourceGraph)
diff --git a/Engine/Source/Editor/BlueprintGraph/Private/K2Node_GetDataTableRow.cpp b/Engine/Source/Editor/BlueprintGraph/Private/K2Node_GetDataTableRow.cpp
--- a/Engine/Source/Editor/BlueprintGraph/Private/K2Node_GetDataTableRow.cpp
+++ b/Engine/Source/Editor/BlueprintGraph/Private/K2Node_GetDataTableRow.cpp
@@ -267,32 +267,35 @@ UEdGraphPin* UK2Node_GetDataTableRow::GetResultPin() const
 
 FText UK2Node_GetDataTableRow::GetNodeTitle(ENodeTitleType::Type TitleType) const
 {
-	UEdGraphPin* DataTablePin = GetDataTablePin();
-
-	FText DataTableString = NSLOCTEXT("K2Node", "None", "NONE");
-	if (DataTablePin != NULL)
+	if (TitleType == ENodeTitleType::ListView)
+	{
+		return LOCTEXT("ListViewTitle", "Get Data Table Row");
+	}
+	else if (UEdGraphPin* DataTablePin = GetDataTablePin())
 	{
-		if(DataTablePin->LinkedTo.Num() > 0)
+		if (DataTablePin->LinkedTo.Num() > 0)
 		{
-			// Blueprint will be determined dynamically, so we don't have the name in this case
-			DataTableString = FText::GetEmpty();
+			return NSLOCTEXT("K2Node", "DataTable_Title_Unknown", "Get Data Table Row");
 		}
-		else if(DataTablePin->DefaultObject != NULL)
+		else if (DataTablePin->DefaultObject == nullptr)
 		{
-			DataTableString = FText::FromString(DataTablePin->DefaultObject->GetName());
+			return NSLOCTEXT("K2Node", "DataTable_Title_None", "Get Data Table Row NONE");
 		}
-	}
+		else if (CachedNodeTitle.IsOutOfDate())
+		{
+			FFormatNamedArguments Args;
+			Args.Add(TEXT("DataTableName"), FText::FromString(DataTablePin->DefaultObject->GetName()));
 
-	FText LocFormat = NSLOCTEXT("K2Node", "DataTable", "Get Data Table Row {DataTableName}");
-	if (TitleType == ENodeTitleType::ListView)
-	{
-		LocFormat = LOCTEXT("ListViewTitle", "Get Data Table Row");
+			FText LocFormat = NSLOCTEXT("K2Node", "DataTable", "Get Data Table Row {DataTableName}");
+			// FText::Format() is slow, so we cache this to save on performance
+			CachedNodeTitle = FText::Format(LocFormat, Args);
+		}
 	}
-
-
-	FFormatNamedArguments Args;
-	Args.Add(TEXT("DataTableName"), DataTableString);
-	return FText::Format(LocFormat, Args);
+	else
+	{
+		return NSLOCTEXT("K2Node", "DataTable_Title_None", "Get Data Table Row NONE");
+	}	
+	return CachedNodeTitle;
 }
 
 void UK2Node_GetDataTableRow::ExpandNode(class FKismetCompilerContext& CompilerContext, UEdGraph* SourceGraph)
diff --git a/Engine/Source/Editor/BlueprintGraph/Private/K2Node_GetInputAxisKeyValue.cpp b/Engine/Source/Editor/BlueprintGraph/Private/K2Node_GetInputAxisKeyValue.cpp
--- a/Engine/Source/Editor/BlueprintGraph/Private/K2Node_GetInputAxisKeyValue.cpp
+++ b/Engine/Source/Editor/BlueprintGraph/Private/K2Node_GetInputAxisKeyValue.cpp
@@ -30,19 +30,24 @@ void UK2Node_GetInputAxisKeyValue::Initialize(const FKey AxisKey)
 {
 	InputAxisKey = AxisKey;
 	SetFromFunction(AActor::StaticClass()->FindFunctionByName(TEXT("GetInputAxisKeyValue")));
+	CachedNodeTitle.MarkDirty();
 }
 
 FText UK2Node_GetInputAxisKeyValue::GetNodeTitle(ENodeTitleType::Type TitleType) const
 {
-	FFormatNamedArguments Args;
-	Args.Add(TEXT("AxisKey"), InputAxisKey.GetDisplayName());
-	
-	FText TitleFormat = NSLOCTEXT("K2Node", "GetInputAxisKey_Name", "Get {AxisKey}");
 	if (TitleType == ENodeTitleType::ListView)
 	{
-		TitleFormat = NSLOCTEXT("K2Node", "GetInputAxisKey_ListTitle", "{AxisKey}");
+		return InputAxisKey.GetDisplayName();
 	}
-	return FText::Format(TitleFormat, Args);
+	else if (CachedNodeTitle.IsOutOfDate())
+	{
+		FFormatNamedArguments Args;
+		Args.Add(TEXT("AxisKey"), InputAxisKey.GetDisplayName());
+		// FText::Format() is slow, so we cache this to save on performance
+		CachedNodeTitle = FText::Format(NSLOCTEXT("K2Node", "GetInputAxisKey_Name", "Get {AxisKey}"), Args);
+	}
+
+	return CachedNodeTitle;
 }
 
 FString UK2Node_GetInputAxisKeyValue::GetTooltip() const
diff --git a/Engine/Source/Editor/BlueprintGraph/Private/K2Node_GetInputAxisValue.cpp b/Engine/Source/Editor/BlueprintGraph/Private/K2Node_GetInputAxisValue.cpp
--- a/Engine/Source/Editor/BlueprintGraph/Private/K2Node_GetInputAxisValue.cpp
+++ b/Engine/Source/Editor/BlueprintGraph/Private/K2Node_GetInputAxisValue.cpp
@@ -34,16 +34,21 @@ void UK2Node_GetInputAxisValue::Initialize(const FName AxisName)
 
 FText UK2Node_GetInputAxisValue::GetNodeTitle(ENodeTitleType::Type TitleType) const
 {
-	FFormatNamedArguments Args;
-	Args.Add(TEXT("InputAxisName"), FText::FromName(InputAxisName));
-
-	FText LocFormat = NSLOCTEXT("K2Node", "GetInputAxis_Name", "Get {InputAxisName}");
 	if (TitleType == ENodeTitleType::ListView)
 	{
-		LocFormat = NSLOCTEXT("K2Node", "GetInputAxis_ListTitle", "{InputAxisName}");
+		return FText::FromName(InputAxisName);
+	}
+	else if (CachedNodeTitle.IsOutOfDate())
+	{
+		FFormatNamedArguments Args;
+		Args.Add(TEXT("InputAxisName"), FText::FromName(InputAxisName));
+
+		FText LocFormat = NSLOCTEXT("K2Node", "GetInputAxis_Name", "Get {InputAxisName}");
+		// FText::Format() is slow, so we cache this to save on performance
+		CachedNodeTitle = FText::Format(LocFormat, Args);
 	}
 
-	return FText::Format(LocFormat, Args);
+	return CachedNodeTitle;
 }
 
 FString UK2Node_GetInputAxisValue::GetTooltip() const
diff --git a/Engine/Source/Editor/BlueprintGraph/Private/K2Node_GetNumEnumEntries.cpp b/Engine/Source/Editor/BlueprintGraph/Private/K2Node_GetNumEnumEntries.cpp
--- a/Engine/Source/Editor/BlueprintGraph/Private/K2Node_GetNumEnumEntries.cpp
+++ b/Engine/Source/Editor/BlueprintGraph/Private/K2Node_GetNumEnumEntries.cpp
@@ -33,11 +33,17 @@ FString UK2Node_GetNumEnumEntries::GetTooltip() const
 
 FText UK2Node_GetNumEnumEntries::GetNodeTitle(ENodeTitleType::Type TitleType) const
 {
-	const FText EnumName = (Enum != NULL) ? FText::FromString(Enum->GetName()) : NSLOCTEXT("K2Node", "BadEnum", "(bad enum)");
-
-	FFormatNamedArguments Args;
-	Args.Add(TEXT("EnumName"), EnumName);
-	return FText::Format(NSLOCTEXT("K2Node", "GetNumEnumEntries_Title", "Get number of entries in {EnumName}"), Args);
+	if (Enum == nullptr)
+	{
+		return NSLOCTEXT("K2Node", "GetNumEnumEntries_BadEnumTitle", "Get number of entries in (bad enum)");
+	}
+	else if (CachedNodeTitle.IsOutOfDate())
+	{
+		FFormatNamedArguments Args;
+		Args.Add(TEXT("EnumName"), FText::FromString(Enum->GetName()));
+		CachedNodeTitle = FText::Format(NSLOCTEXT("K2Node", "GetNumEnumEntries_Title", "Get number of entries in {EnumName}"), Args);
+	}
+	return CachedNodeTitle;
 }
 
 void UK2Node_GetNumEnumEntries::ExpandNode(class FKismetCompilerContext& CompilerContext, UEdGraph* SourceGraph)
diff --git a/Engine/Source/Editor/BlueprintGraph/Private/K2Node_InputAction.cpp b/Engine/Source/Editor/BlueprintGraph/Private/K2Node_InputAction.cpp
--- a/Engine/Source/Editor/BlueprintGraph/Private/K2Node_InputAction.cpp
+++ b/Engine/Source/Editor/BlueprintGraph/Private/K2Node_InputAction.cpp
@@ -47,16 +47,21 @@ FLinearColor UK2Node_InputAction::GetNodeTitleColor() const
 
 FText UK2Node_InputAction::GetNodeTitle(ENodeTitleType::Type TitleType) const
 {
-	FFormatNamedArguments Args;
-	Args.Add(TEXT("InputActionName"), FText::FromName(InputActionName));
-
-	FText LocFormat = NSLOCTEXT("K2Node", "InputAction_Name", "InputAction {InputActionName}");
 	if (TitleType == ENodeTitleType::ListView)
 	{
-		LocFormat = NSLOCTEXT("K2Node", "InputAction_ListTitle", "{InputActionName}");
+		return FText::FromName(InputActionName);
+	}
+	else if (CachedNodeTitle.IsOutOfDate())
+	{
+		FFormatNamedArguments Args;
+		Args.Add(TEXT("InputActionName"), FText::FromName(InputActionName));
+
+		FText LocFormat = NSLOCTEXT("K2Node", "InputAction_Name", "InputAction {InputActionName}");
+		// FText::Format() is slow, so we cache this to save on performance
+		CachedNodeTitle = FText::Format(LocFormat, Args);
 	}
 
-	return FText::Format(LocFormat, Args);
+	return CachedNodeTitle;
 }
 
 FString UK2Node_InputAction::GetTooltip() const
diff --git a/Engine/Source/Editor/BlueprintGraph/Private/K2Node_InputAxisEvent.cpp b/Engine/Source/Editor/BlueprintGraph/Private/K2Node_InputAxisEvent.cpp
--- a/Engine/Source/Editor/BlueprintGraph/Private/K2Node_InputAxisEvent.cpp
+++ b/Engine/Source/Editor/BlueprintGraph/Private/K2Node_InputAxisEvent.cpp
@@ -39,16 +39,21 @@ void UK2Node_InputAxisEvent::Initialize(const FName AxisName)
 
 FText UK2Node_InputAxisEvent::GetNodeTitle(ENodeTitleType::Type TitleType) const
 {
-	FFormatNamedArguments Args;
-	Args.Add(TEXT("InputAxisName"), FText::FromName(InputAxisName));
-
-	FText LocFormat = NSLOCTEXT("K2Node", "InputAxis_Name", "InputAxis {InputAxisName}");
 	if (TitleType == ENodeTitleType::ListView)
 	{
-		LocFormat = NSLOCTEXT("K2Node", "InputAxis_ListTitle", "{InputAxisName}");
+		return FText::FromName(InputAxisName);
+	}
+	else if (CachedNodeTitle.IsOutOfDate())
+	{
+		FFormatNamedArguments Args;
+		Args.Add(TEXT("InputAxisName"), FText::FromName(InputAxisName));
+
+		FText LocFormat = NSLOCTEXT("K2Node", "InputAxis_Name", "InputAxis {InputAxisName}");
+		// FText::Format() is slow, so we cache this to save on performance
+		CachedNodeTitle = FText::Format(LocFormat, Args);
 	}
 
-	return FText::Format(LocFormat, Args);
+	return CachedNodeTitle;
 }
 
 FString UK2Node_InputAxisEvent::GetTooltip() const
diff --git a/Engine/Source/Editor/BlueprintGraph/Private/K2Node_InputKey.cpp b/Engine/Source/Editor/BlueprintGraph/Private/K2Node_InputKey.cpp
--- a/Engine/Source/Editor/BlueprintGraph/Private/K2Node_InputKey.cpp
+++ b/Engine/Source/Editor/BlueprintGraph/Private/K2Node_InputKey.cpp
@@ -231,10 +231,16 @@ FText UK2Node_InputKey::GetNodeTitle(ENodeTitleType::Type TitleType) const
 {
 	if (bControl || bAlt || bShift)
 	{
-		FFormatNamedArguments Args;
-		Args.Add(TEXT("ModifierKey"), GetModifierText());
-		Args.Add(TEXT("Key"), GetKeyText());
-		return FText::Format(NSLOCTEXT("K2Node", "InputKey_Name_WithModifiers", "{ModifierKey} {Key}"), Args);
+		if (CachedNodeTitle.IsOutOfDate())
+		{
+			FFormatNamedArguments Args;
+			Args.Add(TEXT("ModifierKey"), GetModifierText());
+			Args.Add(TEXT("Key"), GetKeyText());
+			
+			// FText::Format() is slow, so we cache this to save on performance
+			CachedNodeTitle = FText::Format(NSLOCTEXT("K2Node", "InputKey_Name_WithModifiers", "{ModifierKey} {Key}"), Args);
+		}
+		return CachedNodeTitle;
 	}
 	else
 	{
diff --git a/Engine/Source/Editor/BlueprintGraph/Private/K2Node_MCDelegate.cpp b/Engine/Source/Editor/BlueprintGraph/Private/K2Node_MCDelegate.cpp
--- a/Engine/Source/Editor/BlueprintGraph/Private/K2Node_MCDelegate.cpp
+++ b/Engine/Source/Editor/BlueprintGraph/Private/K2Node_MCDelegate.cpp
@@ -223,9 +223,14 @@ void UK2Node_AddDelegate::AllocateDefaultPins()
 
 FText UK2Node_AddDelegate::GetNodeTitle(ENodeTitleType::Type TitleType) const
 {
-	FFormatNamedArguments Args;
-	Args.Add(TEXT("PropertyName"), FText::FromName(GetPropertyName()));
-	return FText::Format(NSLOCTEXT("K2Node", "AddDelegate", "Bind Event to {PropertyName}"), Args);
+	if (CachedNodeTitle.IsOutOfDate())
+	{
+		FFormatNamedArguments Args;
+		Args.Add(TEXT("PropertyName"), FText::FromName(GetPropertyName()));
+		// FText::Format() is slow, so we cache this to save on performance
+		CachedNodeTitle = FText::Format(NSLOCTEXT("K2Node", "AddDelegate", "Bind Event to {PropertyName}"), Args);
+	}
+	return CachedNodeTitle;
 }
 
 FNodeHandlingFunctor* UK2Node_AddDelegate::CreateNodeHandler(FKismetCompilerContext& CompilerContext) const
@@ -249,9 +254,14 @@ UK2Node_ClearDelegate::UK2Node_ClearDelegate(const class FPostConstructInitializ
 
 FText UK2Node_ClearDelegate::GetNodeTitle(ENodeTitleType::Type TitleType) const
 {
-	FFormatNamedArguments Args;
-	Args.Add(TEXT("PropertyName"), FText::FromName(GetPropertyName()));
-	return FText::Format(NSLOCTEXT("K2Node", "ClearDelegate", "Unbind all Events from {PropertyName}"), Args);
+	if (CachedNodeTitle.IsOutOfDate())
+	{
+		FFormatNamedArguments Args;
+		Args.Add(TEXT("PropertyName"), FText::FromName(GetPropertyName()));
+		// FText::Format() is slow, so we cache this to save on performance
+		CachedNodeTitle = FText::Format(NSLOCTEXT("K2Node", "ClearDelegate", "Unbind all Events from {PropertyName}"), Args);
+	}
+	return CachedNodeTitle;
 }
 
 FNodeHandlingFunctor* UK2Node_ClearDelegate::CreateNodeHandler(FKismetCompilerContext& CompilerContext) const
@@ -281,9 +291,14 @@ void UK2Node_RemoveDelegate::AllocateDefaultPins()
 
 FText UK2Node_RemoveDelegate::GetNodeTitle(ENodeTitleType::Type TitleType) const
 {
-	FFormatNamedArguments Args;
-	Args.Add(TEXT("PropertyName"), FText::FromName(GetPropertyName()));
-	return FText::Format(NSLOCTEXT("K2Node", "RemoveDelegate", "Unbind Event from {PropertyName}"), Args);
+	if (CachedNodeTitle.IsOutOfDate())
+	{
+		FFormatNamedArguments Args;
+		Args.Add(TEXT("PropertyName"), FText::FromName(GetPropertyName()));
+		// FText::Format() is slow, so we cache this to save on performance
+		CachedNodeTitle = FText::Format(NSLOCTEXT("K2Node", "RemoveDelegate", "Unbind Event from {PropertyName}"), Args);
+	}
+	return CachedNodeTitle;
 }
 
 FNodeHandlingFunctor* UK2Node_RemoveDelegate::CreateNodeHandler(FKismetCompilerContext& CompilerContext) const
@@ -329,9 +344,14 @@ void UK2Node_CallDelegate::AllocateDefaultPins()
 
 FText UK2Node_CallDelegate::GetNodeTitle(ENodeTitleType::Type TitleType) const
 {
-	FFormatNamedArguments Args;
-	Args.Add(TEXT("PropertyName"), FText::FromName(GetPropertyName()));
-	return FText::Format(NSLOCTEXT("K2Node", "CallDelegate", "Call {PropertyName}"), Args);
+	if (CachedNodeTitle.IsOutOfDate())
+	{
+		FFormatNamedArguments Args;
+		Args.Add(TEXT("PropertyName"), FText::FromName(GetPropertyName()));
+		// FText::Format() is slow, so we cache this to save on performance
+		CachedNodeTitle = FText::Format(NSLOCTEXT("K2Node", "CallDelegate", "Call {PropertyName}"), Args);
+	}
+	return CachedNodeTitle;
 }
 
 void UK2Node_CallDelegate::ValidateNodeDuringCompilation(class FCompilerResultsLog& MessageLog) const
diff --git a/Engine/Source/Editor/BlueprintGraph/Private/K2Node_MakeStruct.cpp b/Engine/Source/Editor/BlueprintGraph/Private/K2Node_MakeStruct.cpp
--- a/Engine/Source/Editor/BlueprintGraph/Private/K2Node_MakeStruct.cpp
+++ b/Engine/Source/Editor/BlueprintGraph/Private/K2Node_MakeStruct.cpp
@@ -152,9 +152,14 @@ void UK2Node_MakeStruct::ValidateNodeDuringCompilation(class FCompilerResultsLog
 
 FText UK2Node_MakeStruct::GetNodeTitle(ENodeTitleType::Type TitleType) const 
 {
-	FFormatNamedArguments Args;
-	Args.Add(TEXT("StructName"), FText::FromString(StructType ? StructType->GetName() : FString()));
-	return FText::Format(LOCTEXT("MakeNodeTitle", "Make {StructName}"), Args);
+	if (CachedNodeTitle.IsOutOfDate())
+	{
+		FFormatNamedArguments Args;
+		Args.Add(TEXT("StructName"), FText::FromString(StructType ? StructType->GetName() : FString()));
+		// FText::Format() is slow, so we cache this to save on performance
+		CachedNodeTitle = FText::Format(LOCTEXT("MakeNodeTitle", "Make {StructName}"), Args);
+	}
+	return CachedNodeTitle;
 }
 
 FString UK2Node_MakeStruct::GetTooltip() const
diff --git a/Engine/Source/Editor/BlueprintGraph/Private/K2Node_MathExpression.cpp b/Engine/Source/Editor/BlueprintGraph/Private/K2Node_MathExpression.cpp
--- a/Engine/Source/Editor/BlueprintGraph/Private/K2Node_MathExpression.cpp
+++ b/Engine/Source/Editor/BlueprintGraph/Private/K2Node_MathExpression.cpp
@@ -2322,6 +2322,7 @@ TSharedPtr<class INameValidatorInterface> UK2Node_MathExpression::MakeNameValida
 void UK2Node_MathExpression::OnRenameNode(const FString& NewName)
 {
 	RebuildExpression(NewName);
+	CachedNodeTitle.MarkDirty();
 }
 
 //------------------------------------------------------------------------------
@@ -2440,21 +2441,20 @@ void UK2Node_MathExpression::ValidateNodeDuringCompilation(FCompilerResultsLog&
 //------------------------------------------------------------------------------
 FText UK2Node_MathExpression::GetNodeTitle(ENodeTitleType::Type TitleType) const
 {
-	FText Result = FText::FromString(Expression);
-	if (Expression.IsEmpty())
+	if (Expression.IsEmpty() && (TitleType == ENodeTitleType::ListView))
 	{
-		//Result = LOCTEXT("NoExpressionTitle", "[Empty]").ToString();
+		return LOCTEXT("AddMathExprMenuOption", "Add Math Expression...");
 	}
-	
-	if (TitleType == ENodeTitleType::FullTitle)
+	else if (TitleType != ENodeTitleType::FullTitle)
 	{
-		Result = FText::Format(LOCTEXT("MathExpressionSecondTitleLine", "{0}\nMath Expression"), Result);
+		return FText::FromString(Expression);
 	}
-	else if ((TitleType == ENodeTitleType::ListView) && Expression.IsEmpty())
+	else if (CachedNodeTitle.IsOutOfDate())
 	{
-		Result = LOCTEXT("AddMathExprMenuOption", "Add Math Expression...");
+		// FText::Format() is slow, so we cache this to save on performance
+		CachedNodeTitle = FText::Format(LOCTEXT("MathExpressionSecondTitleLine", "{0}\nMath Expression"), FText::FromString(Expression));
 	}
-	return Result;
+	return CachedNodeTitle;
 }
 
 //------------------------------------------------------------------------------
diff --git a/Engine/Source/Editor/BlueprintGraph/Private/K2Node_Message.cpp b/Engine/Source/Editor/BlueprintGraph/Private/K2Node_Message.cpp
--- a/Engine/Source/Editor/BlueprintGraph/Private/K2Node_Message.cpp
+++ b/Engine/Source/Editor/BlueprintGraph/Private/K2Node_Message.cpp
@@ -16,26 +16,31 @@ FText UK2Node_Message::GetNodeTitle(ENodeTitleType::Type TitleType) const
 	FText NodeName;
 	if (UFunction* Function = GetTargetFunction())
 	{
-		FString NodeNameStr = UK2Node_CallFunction::GetUserFacingFunctionName(Function);
-
-		if (TitleType == ENodeTitleType::ListView)
-		{
-			NodeName = FText::Format(LOCTEXT("ListTitle", "{0} (Message)"), FText::FromString(NodeNameStr));
-		}
-		else
+		if (!CachedNodeTitles.IsTitleCached(TitleType))
 		{
-			FFormatNamedArguments Args;
-			Args.Add(TEXT("NodeName"), FText::FromString(NodeNameStr));
-			Args.Add(TEXT("OuterClassName"), FText::FromString(Function->GetOuterUClass()->GetName()));
-			NodeName = FText::Format(NSLOCTEXT("K2Node", "CallInterfaceContext", "{NodeName}\nUsing Interface {OuterClassName}"), Args);
+			FText NodeNameText = FText::FromString(UK2Node_CallFunction::GetUserFacingFunctionName(Function));
+			if (TitleType == ENodeTitleType::ListView)
+			{
+				// FText::Format() is slow, so we cache this to save on performance
+				CachedNodeTitles.SetCachedTitle(TitleType, FText::Format(LOCTEXT("ListTitle", "{0} (Message)"), NodeNameText));
+			}
+			else
+			{
+				FFormatNamedArguments Args;
+				Args.Add(TEXT("NodeName"), NodeNameText);
+				Args.Add(TEXT("OuterClassName"), FText::FromString(Function->GetOuterUClass()->GetName()));
+
+				FText NodeTitle = FText::Format(NSLOCTEXT("K2Node", "CallInterfaceContext", "{NodeName}\nUsing Interface {OuterClassName}"), Args);
+				// FText::Format() is slow, so we cache this to save on performance
+				CachedNodeTitles.SetCachedTitle(TitleType, NodeTitle);
+			}
 		}
 	}
 	else
 	{
-		NodeName = NSLOCTEXT("K2Node", "InvalidMessageNode", "Invalid Message Node");
+		return NSLOCTEXT("K2Node", "InvalidMessageNode", "Invalid Message Node");
 	}
-
-	return NodeName;
+	return CachedNodeTitles.GetCachedTitle(TitleType);
 }
 
 UEdGraphPin* UK2Node_Message::CreateSelfPin(const UFunction* Function)
diff --git a/Engine/Source/Editor/BlueprintGraph/Private/K2Node_SetFieldsInStruct.cpp b/Engine/Source/Editor/BlueprintGraph/Private/K2Node_SetFieldsInStruct.cpp
--- a/Engine/Source/Editor/BlueprintGraph/Private/K2Node_SetFieldsInStruct.cpp
+++ b/Engine/Source/Editor/BlueprintGraph/Private/K2Node_SetFieldsInStruct.cpp
@@ -55,9 +55,14 @@ void UK2Node_SetFieldsInStruct::AllocateDefaultPins()
 
 FText UK2Node_SetFieldsInStruct::GetNodeTitle(ENodeTitleType::Type TitleType) const
 {
-	FFormatNamedArguments Args;
-	Args.Add(TEXT("StructName"), FText::FromString(StructType ? StructType->GetName() : FString()));
-	return FText::Format(LOCTEXT("SetFieldsInStructNodeTitle", "Set members in {StructName}"), Args);
+	if (CachedNodeTitle.IsOutOfDate())
+	{
+		FFormatNamedArguments Args;
+		Args.Add(TEXT("StructName"), FText::FromString(StructType ? StructType->GetName() : FString()));
+		// FText::Format() is slow, so we cache this to save on performance
+		CachedNodeTitle = FText::Format(LOCTEXT("SetFieldsInStructNodeTitle", "Set members in {StructName}"), Args);
+	}
+	return CachedNodeTitle;
 }
 
 FString UK2Node_SetFieldsInStruct::GetTooltip() const
diff --git a/Engine/Source/Editor/BlueprintGraph/Private/K2Node_SpawnActor.cpp b/Engine/Source/Editor/BlueprintGraph/Private/K2Node_SpawnActor.cpp
--- a/Engine/Source/Editor/BlueprintGraph/Private/K2Node_SpawnActor.cpp
+++ b/Engine/Source/Editor/BlueprintGraph/Private/K2Node_SpawnActor.cpp
@@ -228,24 +228,23 @@ FLinearColor UK2Node_SpawnActor::GetNodeTitleColor() const
 FText UK2Node_SpawnActor::GetNodeTitle(ENodeTitleType::Type TitleType) const
 {
 	UEdGraphPin* BlueprintPin = GetBlueprintPin();
-
-	FText SpawnString = NSLOCTEXT("K2Node", "None", "NONE");
-	if(BlueprintPin != NULL)
+	if (BlueprintPin == NULL)
 	{
-		if(BlueprintPin->LinkedTo.Num() > 0)
-		{
-			// Blueprint will be determined dynamically, so we don't have the name in this case
-			SpawnString = FText::GetEmpty();
-		}
-		else if(BlueprintPin->DefaultObject != NULL)
-		{
-			SpawnString = FText::FromString(BlueprintPin->DefaultObject->GetName());
-		}
+		return NSLOCTEXT("K2Node", "SpawnActorNone_Title", "SpawnActor NONE");
 	}
-
-	FFormatNamedArguments Args;
-	Args.Add(TEXT("ActorName"), SpawnString);
-	return FText::Format(NSLOCTEXT("K2Node", "SpawnActor", "SpawnActor {ActorName}"), Args);
+	else if (BlueprintPin->LinkedTo.Num() > 0)
+	{
+		// Blueprint will be determined dynamically, so we don't have the name in this case
+		return NSLOCTEXT("K2Node", "SpawnActorUnknown_Title", "SpawnActor");
+	}
+	else if (CachedNodeTitle.IsOutOfDate())
+	{
+		FFormatNamedArguments Args;
+		Args.Add(TEXT("ActorName"), FText::FromString(BlueprintPin->DefaultObject->GetName()));
+		// FText::Format() is slow, so we cache this to save on performance
+		CachedNodeTitle = FText::Format(NSLOCTEXT("K2Node", "SpawnActor", "SpawnActor {ActorName}"), Args);
+	}
+	return CachedNodeTitle;
 }
 
 bool UK2Node_SpawnActor::IsCompatibleWithGraph(const UEdGraph* TargetGraph) const 
diff --git a/Engine/Source/Editor/BlueprintGraph/Private/K2Node_SpawnActorFromClass.cpp b/Engine/Source/Editor/BlueprintGraph/Private/K2Node_SpawnActorFromClass.cpp
--- a/Engine/Source/Editor/BlueprintGraph/Private/K2Node_SpawnActorFromClass.cpp
+++ b/Engine/Source/Editor/BlueprintGraph/Private/K2Node_SpawnActorFromClass.cpp
@@ -279,23 +279,32 @@ FText UK2Node_SpawnActorFromClass::GetNodeTitle(ENodeTitleType::Type TitleType)
 		FText SpawnString = NSLOCTEXT("K2Node", "None", "NONE");
 		if (UEdGraphPin* ClassPin = GetClassPin())
 		{
-			if(ClassPin->LinkedTo.Num() > 0)
+			if (ClassPin->LinkedTo.Num() > 0)
 			{
 				// Blueprint will be determined dynamically, so we don't have the name in this case
-				SpawnString = FText::GetEmpty();
+				NodeTitle = NSLOCTEXT("K2Node", "SpawnActor_Title_Unknown", "SpawnActor");
 			}
-			else if(ClassPin->DefaultObject != NULL)
+			else if (ClassPin->DefaultObject == nullptr)
 			{
-				SpawnString = FText::FromString(ClassPin->DefaultObject->GetName());
+				NodeTitle = NSLOCTEXT("K2Node", "SpawnActor_Title_NONE", "SpawnActor NONE");
 			}
+			else
+			{
+				if (CachedNodeTitle.IsOutOfDate())
+				{
+					FFormatNamedArguments Args;
+					Args.Add(TEXT("ClassName"), FText::FromString(ClassPin->DefaultObject->GetName()));
+					// FText::Format() is slow, so we cache this to save on performance
+					CachedNodeTitle = FText::Format(NSLOCTEXT("K2Node", "SpawnActor", "SpawnActor {ClassName}"), Args);
+				}
+				NodeTitle = CachedNodeTitle;
+			} 
+		}
+		else
+		{
+			NodeTitle = NSLOCTEXT("K2Node", "SpawnActor_Title_NONE", "SpawnActor NONE");
 		}
-		
-		FFormatNamedArguments Args;
-		Args.Add(TEXT("ClassName"), SpawnString);
-		
-		NodeTitle = FText::Format(NSLOCTEXT("K2Node", "SpawnActor", "SpawnActor {ClassName}"), Args);
 	}
-
 	return NodeTitle;
 }
 
diff --git a/Engine/Source/Editor/BlueprintGraph/Private/K2Node_StructMemberGet.cpp b/Engine/Source/Editor/BlueprintGraph/Private/K2Node_StructMemberGet.cpp
--- a/Engine/Source/Editor/BlueprintGraph/Private/K2Node_StructMemberGet.cpp
+++ b/Engine/Source/Editor/BlueprintGraph/Private/K2Node_StructMemberGet.cpp
@@ -78,9 +78,14 @@ FString UK2Node_StructMemberGet::GetTooltip() const
 
 FText UK2Node_StructMemberGet::GetNodeTitle(ENodeTitleType::Type TitleType) const
 {
-	FFormatNamedArguments Args;
-	Args.Add(TEXT("VariableName"), FText::FromString(GetVarNameString()));
-	return FText::Format(LOCTEXT("GetMembersInVariable", "Get members in {VariableName}"), Args);
+	if (CachedNodeTitle.IsOutOfDate())
+	{
+		FFormatNamedArguments Args;
+		Args.Add(TEXT("VariableName"), FText::FromString(GetVarNameString()));
+		// FText::Format() is slow, so we cache this to save on performance
+		CachedNodeTitle = FText::Format(LOCTEXT("GetMembersInVariable", "Get members in {VariableName}"), Args);
+	}
+	return CachedNodeTitle;
 }
 
 FNodeHandlingFunctor* UK2Node_StructMemberGet::CreateNodeHandler(FKismetCompilerContext& CompilerContext) const
diff --git a/Engine/Source/Editor/BlueprintGraph/Private/K2Node_StructMemberSet.cpp b/Engine/Source/Editor/BlueprintGraph/Private/K2Node_StructMemberSet.cpp
--- a/Engine/Source/Editor/BlueprintGraph/Private/K2Node_StructMemberSet.cpp
+++ b/Engine/Source/Editor/BlueprintGraph/Private/K2Node_StructMemberSet.cpp
@@ -51,9 +51,14 @@ FString UK2Node_StructMemberSet::GetTooltip() const
 
 FText UK2Node_StructMemberSet::GetNodeTitle(ENodeTitleType::Type TitleType) const
 {
-	FFormatNamedArguments Args;
-	Args.Add(TEXT("VariableName"), FText::FromName(VariableReference.GetMemberName()));
-	return FText::Format(LOCTEXT("SetMembersInVariable", "Set members in {VariableName}"), Args);
+	if (CachedNodeTitle.IsOutOfDate())
+	{
+		FFormatNamedArguments Args;
+		Args.Add(TEXT("VariableName"), FText::FromName(VariableReference.GetMemberName()));
+		// FText::Format() is slow, so we cache this to save on performance
+		CachedNodeTitle = FText::Format(LOCTEXT("SetMembersInVariable", "Set members in {VariableName}"), Args);
+	}
+	return CachedNodeTitle;
 }
 
 UK2Node::ERedirectType UK2Node_StructMemberSet::DoPinsMatchForReconstruction(const UEdGraphPin* NewPin, int32 NewPinIndex, const UEdGraphPin* OldPin, int32 OldPinIndex)  const
diff --git a/Engine/Source/Editor/BlueprintGraph/Private/K2Node_SwitchEnum.cpp b/Engine/Source/Editor/BlueprintGraph/Private/K2Node_SwitchEnum.cpp
--- a/Engine/Source/Editor/BlueprintGraph/Private/K2Node_SwitchEnum.cpp
+++ b/Engine/Source/Editor/BlueprintGraph/Private/K2Node_SwitchEnum.cpp
@@ -51,12 +51,18 @@ void UK2Node_SwitchEnum::SetEnum(UEnum* InEnum)
 
 FText UK2Node_SwitchEnum::GetNodeTitle(ENodeTitleType::Type TitleType) const 
 {
-	FText EnumName = (Enum != NULL) ? FText::FromString(Enum->GetName()) : LOCTEXT("BadEnum", "(bad enum)");
-
-	FFormatNamedArguments Args;
-	Args.Add(TEXT("EnumName"), EnumName);
-
-	return FText::Format(NSLOCTEXT("K2Node", "Switch_Enum", "Switch on {EnumName}"), Args);
+	if (Enum == nullptr)
+	{
+		return LOCTEXT("SwitchEnum_BadEnumTitle", "Switch on (bad enum)");
+	}
+	else if (CachedNodeTitle.IsOutOfDate())
+	{
+		FFormatNamedArguments Args;
+		Args.Add(TEXT("EnumName"), FText::FromString(Enum->GetName()));
+		// FText::Format() is slow, so we cache this to save on performance
+		CachedNodeTitle = FText::Format(NSLOCTEXT("K2Node", "Switch_Enum", "Switch on {EnumName}"), Args);
+	}
+	return CachedNodeTitle;
 }
 
 FString UK2Node_SwitchEnum::GetTooltip() const
diff --git a/Engine/Source/Editor/BlueprintGraph/Private/K2Node_TemporaryVariable.cpp b/Engine/Source/Editor/BlueprintGraph/Private/K2Node_TemporaryVariable.cpp
--- a/Engine/Source/Editor/BlueprintGraph/Private/K2Node_TemporaryVariable.cpp
+++ b/Engine/Source/Editor/BlueprintGraph/Private/K2Node_TemporaryVariable.cpp
@@ -56,12 +56,17 @@ FString UK2Node_TemporaryVariable::GetTooltip() const
 
 FText UK2Node_TemporaryVariable::GetNodeTitle(ENodeTitleType::Type TitleType) const
 {
-	FText Result = !bIsPersistent ? NSLOCTEXT("K2Node", "LocalVariable", "Local {VariableType}") : NSLOCTEXT("K2Node", "PersistentLocalVariable", "Persistent Local {VariableType}");
+	if (CachedNodeTitle.IsOutOfDate())
+	{
+		FFormatNamedArguments Args;
+		Args.Add(TEXT("VariableType"), UEdGraphSchema_K2::TypeToText(VariableType));
 
-	FFormatNamedArguments Args;
-	Args.Add(TEXT("VariableType"), UEdGraphSchema_K2::TypeToText(VariableType));
-	Result = FText::Format(Result, Args);
-	return Result;
+		FText TitleFormat = !bIsPersistent ? NSLOCTEXT("K2Node", "LocalVariable", "Local {VariableType}") : NSLOCTEXT("K2Node", "PersistentLocalVariable", "Persistent Local {VariableType}");
+		// FText::Format() is slow, so we cache this to save on performance
+		CachedNodeTitle = FText::Format(TitleFormat, Args);
+	}
+	
+	return CachedNodeTitle;
 }
 
 bool UK2Node_TemporaryVariable::IsNodePure() const
diff --git a/Engine/Source/Editor/BlueprintGraph/Private/K2Node_VariableGet.cpp b/Engine/Source/Editor/BlueprintGraph/Private/K2Node_VariableGet.cpp
--- a/Engine/Source/Editor/BlueprintGraph/Private/K2Node_VariableGet.cpp
+++ b/Engine/Source/Editor/BlueprintGraph/Private/K2Node_VariableGet.cpp
@@ -128,8 +128,6 @@ FString UK2Node_VariableGet::GetTooltip() const
 
 FText UK2Node_VariableGet::GetNodeTitle(ENodeTitleType::Type TitleType) const
 {
-	FText Result;
-
 	// If there is only one variable being read, the title can be made the variable name
 	FString OutputPinName;
 	int32 NumOutputsFound = 0;
@@ -144,18 +142,18 @@ FText UK2Node_VariableGet::GetNodeTitle(ENodeTitleType::Type TitleType) const
 		}
 	}
 
-	if (NumOutputsFound == 1)
+	if (NumOutputsFound != 1)
 	{
-		FFormatNamedArguments Args;
-		Args.Add(TEXT("PinName"), FText::FromString(OutputPinName));
-		Result = FText::Format(LOCTEXT("GetPinName", "Get {PinName}"), Args);
+		return LOCTEXT("Get", "Get");
 	}
-	else
+	else if (CachedNodeTitle.IsOutOfDate())
 	{
-		Result = LOCTEXT("Get", "Get");
+		FFormatNamedArguments Args;
+		Args.Add(TEXT("PinName"), FText::FromString(OutputPinName));
+		// FText::Format() is slow, so we cache this to save on performance
+		CachedNodeTitle = FText::Format(LOCTEXT("GetPinName", "Get {PinName}"), Args);
 	}
-
-	return Result;
+	return CachedNodeTitle;
 }
 
 FNodeHandlingFunctor* UK2Node_VariableGet::CreateNodeHandler(FKismetCompilerContext& CompilerContext) const
diff --git a/Engine/Source/Editor/BlueprintGraph/Private/K2Node_VariableSet.cpp b/Engine/Source/Editor/BlueprintGraph/Private/K2Node_VariableSet.cpp
--- a/Engine/Source/Editor/BlueprintGraph/Private/K2Node_VariableSet.cpp
+++ b/Engine/Source/Editor/BlueprintGraph/Private/K2Node_VariableSet.cpp
@@ -88,8 +88,6 @@ FText UK2Node_VariableSet::GetNodeTitle(ENodeTitleType::Type TitleType) const
 {
 	const UEdGraphSchema_K2* K2Schema = GetDefault<UEdGraphSchema_K2>();
 
-	FText Result = HasLocalRepNotify() ? NSLOCTEXT("K2Node", "SetWithNotify", "Set with Notify") : NSLOCTEXT("K2Node", "Set", "Set");
-
 	// If there is only one variable being written (one non-meta input pin), the title can be made the variable name
 	FString InputPinName;
 	int32 NumInputsFound = 0;
@@ -104,26 +102,26 @@ FText UK2Node_VariableSet::GetNodeTitle(ENodeTitleType::Type TitleType) const
 		}
 	}
 
-	if (NumInputsFound == 1)
+	if (NumInputsFound != 1)
+	{
+		return HasLocalRepNotify() ? NSLOCTEXT("K2Node", "SetWithNotify", "Set with Notify") : NSLOCTEXT("K2Node", "Set", "Set");
+	}
+	else if (CachedNodeTitle.IsOutOfDate())
 	{
 		FFormatNamedArguments Args;
 		Args.Add(TEXT("PinName"), FText::FromString(InputPinName));
 
-		if(HasLocalRepNotify())
+		// FText::Format() is slow, so we cache this to save on performance
+		if (HasLocalRepNotify())
 		{
-			Result = FText::Format(NSLOCTEXT("K2Node", "SetWithNotifyPinName", "Set with Notify {PinName}"), Args);
+			CachedNodeTitle = FText::Format(NSLOCTEXT("K2Node", "SetWithNotifyPinName", "Set with Notify {PinName}"), Args);
 		}
 		else
 		{
-			Result = FText::Format(NSLOCTEXT("K2Node", "SetPinName", "Set {PinName}"), Args);
+			CachedNodeTitle = FText::Format(NSLOCTEXT("K2Node", "SetPinName", "Set {PinName}"), Args);
 		}
 	}
-	else
-	{
-		Result = HasLocalRepNotify() ? NSLOCTEXT("K2Node", "SetWithNotify", "Set with Notify") : NSLOCTEXT("K2Node", "Set", "Set");
-	}
-
-	return Result;
+	return CachedNodeTitle;
 }
 
 /** Returns true if the variable we are setting has a RepNotify AND was defined in a blueprint
diff --git a/Engine/Source/Editor/BlueprintGraph/Private/K2Node_VariableSetRef.cpp b/Engine/Source/Editor/BlueprintGraph/Private/K2Node_VariableSetRef.cpp
--- a/Engine/Source/Editor/BlueprintGraph/Private/K2Node_VariableSetRef.cpp
+++ b/Engine/Source/Editor/BlueprintGraph/Private/K2Node_VariableSetRef.cpp
@@ -131,17 +131,18 @@ FText UK2Node_VariableSetRef::GetNodeTitle(ENodeTitleType::Type TitleType) const
 	const UEdGraphSchema_K2* Schema = GetDefault<UEdGraphSchema_K2>();
 
 	UEdGraphPin* TargetPin = GetTargetPin();
-
-	if( TargetPin && TargetPin->PinType.PinCategory != Schema->PC_Wildcard )
+	if ((TargetPin == nullptr) || (TargetPin->PinType.PinCategory == Schema->PC_Wildcard))
 	{
-		FFormatNamedArguments Args;
-		Args.Add(TEXT("PinType"), Schema->TypeToText(TargetPin->PinType));
-		return FText::Format(NSLOCTEXT("K2Node", "SetRefVarNodeTitle_Typed", "Set {PinType}"), Args);
+		return NSLOCTEXT("K2Node", "SetRefVarNodeTitle", "Set By-Ref Var");
 	}
-	else
+	else if (CachedNodeTitle.IsOutOfDate())
 	{
-		return NSLOCTEXT("K2Node", "SetRefVarNodeTitle", "Set By-Ref Var");
+		FFormatNamedArguments Args;
+		Args.Add(TEXT("PinType"), Schema->TypeToText(TargetPin->PinType));
+		// FText::Format() is slow, so we cache this to save on performance
+		CachedNodeTitle = FText::Format(NSLOCTEXT("K2Node", "SetRefVarNodeTitle_Typed", "Set {PinType}"), Args);
 	}
+	return CachedNodeTitle;
 }
 
 void UK2Node_VariableSetRef::NotifyPinConnectionListChanged(UEdGraphPin* Pin)
@@ -187,6 +188,8 @@ void UK2Node_VariableSetRef::CoerceTypeFromPin(const UEdGraphPin* Pin)
 		ValuePin->PinType.PinSubCategory = TEXT("");
 		ValuePin->PinType.PinSubCategoryObject = NULL;
 		ValuePin->BreakAllPinLinks();
+
+		CachedNodeTitle.MarkDirty();
 	}
 }
 
diff --git a/Engine/Source/Runtime/Engine/Classes/EdGraph/EdGraphNode.h b/Engine/Source/Runtime/Engine/Classes/EdGraph/EdGraphNode.h
--- a/Engine/Source/Runtime/Engine/Classes/EdGraph/EdGraphNode.h
+++ b/Engine/Source/Runtime/Engine/Classes/EdGraph/EdGraphNode.h
@@ -31,7 +31,9 @@ namespace ENodeTitleType
 		// More concise, single line title
 		ListView,
 		// Returns the editable title (which might not be a title at all)
-		EditableTitle
+		EditableTitle,
+
+		MAX_TitleTypes,
 	};
 }
 
diff --git a/Engine/Source/Runtime/Engine/Public/EdGraph/EdGraphNodeUtils.h b/Engine/Source/Runtime/Engine/Public/EdGraph/EdGraphNodeUtils.h
new file mode 100644
--- /dev/null
+++ b/Engine/Source/Runtime/Engine/Public/EdGraph/EdGraphNodeUtils.h
@@ -0,0 +1,140 @@
+// Copyright 1998-2014 Epic Games, Inc. All Rights Reserved.
+
+#pragma once
+
+#include "EdGraph/EdGraphNode.h" // for ENodeTitleType
+
+/*******************************************************************************
+ * FNodeTextCache
+ ******************************************************************************/
+
+/** 
+ * Constructing FText strings every frame can be costly, so this struct provides 
+ * a way to easily cache those strings (node title, tooptip, etc.) for reuse.
+ */
+struct FNodeTextCache
+{
+public:
+	/** */
+	FORCEINLINE bool IsOutOfDate() const
+	{
+		return CachedText.IsEmpty();
+	}
+
+	/** */
+	FORCEINLINE void SetCachedText(FText const& Text) const
+	{
+		CachedText = Text;
+	}
+
+	/** */
+	FORCEINLINE FText& GetCachedText() const
+	{
+		return CachedText;
+	}
+
+	/** */
+	FORCEINLINE void MarkDirty() const
+	{
+		Clear();
+	}
+
+	/** */
+	FORCEINLINE void Clear() const
+	{
+		CachedText = FText::GetEmpty();
+	}
+
+	/** */
+	FORCEINLINE FText& operator=(FText const& Text) const
+	{
+		SetCachedText(Text);
+		return CachedText;
+	}
+
+	/** */
+	FORCEINLINE operator FText&() const
+	{
+		return GetCachedText();
+	}
+
+private:
+	/** Mutable so that SetCachedText() can remain const (callable by the node's GetNodeTitle() method)*/
+	mutable FText CachedText;
+};
+
+/*******************************************************************************
+ * FNodeTitleTextTable
+ ******************************************************************************/
+
+struct FNodeTitleTextTable
+{
+public:
+	/** */
+	FORCEINLINE bool IsTitleCached(ENodeTitleType::Type TitleType) const
+	{
+		return !CachedNodeTitles[TitleType].IsOutOfDate();
+	}
+
+	/** */
+	FORCEINLINE void SetCachedTitle(ENodeTitleType::Type TitleType, FText const& Text) const
+	{
+		CachedNodeTitles[TitleType].SetCachedText(Text);
+	}
+
+	/** */
+	FORCEINLINE FText GetCachedTitle(ENodeTitleType::Type TitleType) const
+	{
+		return CachedNodeTitles[TitleType].GetCachedText();
+	}
+
+	/** */
+	FORCEINLINE void MarkDirty() const
+	{
+		for (uint32 TitleIndex = 0; TitleIndex < ENodeTitleType::MAX_TitleTypes; ++TitleIndex)
+		{
+			CachedNodeTitles[TitleIndex].MarkDirty();
+		}
+	}
+
+private:
+	/** */
+	FNodeTextCache CachedNodeTitles[ENodeTitleType::MAX_TitleTypes];
+};
+
+/*******************************************************************************
+* FNodeTextTable
+******************************************************************************/
+
+struct FNodeTextTable : FNodeTitleTextTable
+{
+public:
+	/** */
+	FORCEINLINE bool IsTooltipCached() const
+	{
+		return !CachedTooltip.IsOutOfDate();
+	}
+
+	/** */
+	FORCEINLINE void SetCachedTooltip(FText const& Text) const
+	{
+		CachedTooltip.SetCachedText(Text);
+	}
+
+	/** */
+	FORCEINLINE FText GetCachedTooltip() const
+	{
+		return CachedTooltip.GetCachedText();
+	}
+
+	/** */
+	FORCEINLINE void MarkDirty() const
+	{
+		FNodeTitleTextTable::MarkDirty();
+		CachedTooltip.MarkDirty();
+	}
+
+private:
+	/** */
+	FNodeTextCache CachedTooltip;
+};
