Fix crash when splitting certain vector/rotator pins that had their type coerced via another wildcard pin being connected

diff --git a/Engine/Source/Editor/BlueprintGraph/Private/EdGraphSchema_K2.cpp b/Engine/Source/Editor/BlueprintGraph/Private/EdGraphSchema_K2.cpp
--- a/Engine/Source/Editor/BlueprintGraph/Private/EdGraphSchema_K2.cpp
+++ b/Engine/Source/Editor/BlueprintGraph/Private/EdGraphSchema_K2.cpp
@@ -4768,6 +4768,11 @@ void UEdGraphSchema_K2::SplitPin(UEdGraphPin* Pin) const
 			{
 				Default = FString::SanitizeFloat(FCString::Atof(*Default));
 			}
+			// In some cases (particularly wildcards) the default value may not accurately reflect the normal component elements
+			while (OriginalDefaults.Num() < 3)
+			{
+				OriginalDefaults.Add(TEXT("0.0"));
+			}
 		}
 		else if (StructType == FindObjectChecked<UScriptStruct>(UObject::StaticClass(), TEXT("Vector2D")))
 		{
