(WIP) Optimizing new blueprint menu filtering (caching associated fields/functions/etc. that we query for in multiple rejection tests).

diff --git a/Engine/Source/Editor/BlueprintGraph/Private/BlueprintActionFilter.cpp b/Engine/Source/Editor/BlueprintGraph/Private/BlueprintActionFilter.cpp
--- a/Engine/Source/Editor/BlueprintGraph/Private/BlueprintActionFilter.cpp
+++ b/Engine/Source/Editor/BlueprintGraph/Private/BlueprintActionFilter.cpp
@@ -41,7 +41,7 @@ namespace BlueprintActionFilterImpl
 	 * @param  NodeSpawner	The action you want to query.
 	 * @return True if the action will spawn an impure node, false if not (or unknown).
 	 */
-	static bool IsImpure(UBlueprintNodeSpawner const* NodeSpawner);
+	static bool IsImpure(FBlueprintActionInfo& NodeSpawner);
 
 	/**
 	 * Utility method to check and see if the specified node-spawner would 
@@ -50,7 +50,7 @@ namespace BlueprintActionFilterImpl
 	 * @param  NodeSpawner	The action you want to query.
 	 * @return True if the action will spawn a latent node, false if not (or unknown).
 	 */
-	static bool IsLatent(UBlueprintNodeSpawner const* NodeSpawner);
+	static bool IsLatent(FBlueprintActionInfo& NodeSpawner);
 	
 	/**
 	 * Utility method to check and see if the specified field is a public global
@@ -70,7 +70,7 @@ namespace BlueprintActionFilterImpl
 	 * @param  BlueprintAction	The action you wish to query.
 	 * @return True if the node-spawner would produce an event incompatible with the specified blueprint(s).
 	 */
-	static bool IsEventUnimplementable(FBlueprintActionFilter const& Filter, UBlueprintNodeSpawner const* BlueprintAction);
+	static bool IsEventUnimplementable(FBlueprintActionFilter const& Filter, FBlueprintActionInfo& BlueprintAction);
 
 	/**
 	 * Rejection test that checks to see if the supplied node-spawner has an 
@@ -81,7 +81,7 @@ namespace BlueprintActionFilterImpl
 	 * @param  BlueprintAction	The action you wish to query.
 	 * @return True if the node-spawner is associated with a private field.
 	 */
-	static bool IsFieldInaccessible(FBlueprintActionFilter const& Filter, UBlueprintNodeSpawner const* BlueprintAction);
+	static bool IsFieldInaccessible(FBlueprintActionFilter const& Filter, FBlueprintActionInfo& BlueprintAction);
 
 	/**
 	 * Rejection test that checks to see if the supplied node-spawner has an 
@@ -92,7 +92,7 @@ namespace BlueprintActionFilterImpl
 	 * @param  BlueprintAction	The action you wish to query.
 	 * @return True if the node-spawner belongs to a class that is restricted to certain blueprints (not this one).
 	 */
-	static bool IsRestrictedClassMember(FBlueprintActionFilter const& Filter, UBlueprintNodeSpawner const* BlueprintAction);
+	static bool IsRestrictedClassMember(FBlueprintActionFilter const& Filter, FBlueprintActionInfo& BlueprintAction);
 	
 	/**
 	 * Rejection test that checks to see if the supplied node-spawner would 
@@ -103,7 +103,7 @@ namespace BlueprintActionFilterImpl
 	 * @param  BlueprintAction	The action you wish to query.
 	 * @return True if the action would spawn a variable-set node for a read-only property.
 	 */
-	static bool IsPermissionNotGranted(FBlueprintActionFilter const& Filter, UBlueprintNodeSpawner const* BlueprintAction);
+	static bool IsPermissionNotGranted(FBlueprintActionFilter const& Filter, FBlueprintActionInfo& BlueprintAction);
 
 	/**
 	 * Rejection test that checks to see if the supplied node-spawner would 
@@ -113,7 +113,7 @@ namespace BlueprintActionFilterImpl
 	 * @param  BlueprintAction	The action you wish to query.
 	 * @return True if the action would spawn a node that is deprecated.
 	 */
-	static bool IsDeprecated(FBlueprintActionFilter const& Filter, UBlueprintNodeSpawner const* BlueprintAction);
+	static bool IsDeprecated(FBlueprintActionFilter const& Filter, FBlueprintActionInfo& BlueprintAction);
 	
 	/**
 	 * Rejection test that checks to see if the supplied node-spawner would 
@@ -123,7 +123,7 @@ namespace BlueprintActionFilterImpl
 	 * @param  BlueprintAction	The action you wish to query.
 	 * @return True if the action would spawn an impure node (and the graph doesn't allow them).
 	 */
-	static bool IsIncompatibleImpureNode(FBlueprintActionFilter const& Filter, UBlueprintNodeSpawner const* BlueprintAction);
+	static bool IsIncompatibleImpureNode(FBlueprintActionFilter const& Filter, FBlueprintActionInfo& BlueprintAction);
 	
 	/**
 	 * Rejection test that checks to see if the supplied node-spawner would 
@@ -133,7 +133,7 @@ namespace BlueprintActionFilterImpl
 	 * @param  BlueprintAction	The action you wish to query.
 	 * @return True if the action would spawn a node (and the graph wouldn't allow it).
 	 */
-	static bool IsIncompatibleWithGraphType(FBlueprintActionFilter const& Filter, UBlueprintNodeSpawner const* BlueprintAction);
+	static bool IsIncompatibleWithGraphType(FBlueprintActionFilter const& Filter, FBlueprintActionInfo& BlueprintAction);
 	
 	/**
 	 * Rejection test that checks to see if the node-spawner has any associated
@@ -143,7 +143,7 @@ namespace BlueprintActionFilterImpl
 	 * @param  BlueprintAction	The action you wish to query.
 	 * @return 
 	 */
-	static bool IsPersistentNonTargetField(FBlueprintActionFilter const& Filter, UBlueprintNodeSpawner const* BlueprintAction);
+	static bool IsPersistentNonTargetField(FBlueprintActionFilter const& Filter, FBlueprintActionInfo& BlueprintAction);
 
 	/**
 	 * Rejection test that checks to see if the node-spawner is associated with 
@@ -155,7 +155,7 @@ namespace BlueprintActionFilterImpl
 	 * @param  bPermitNonTargetGlobals	Determines if this test should pass for external global/static fields.
 	 * @return True if the action is associated with a non-whitelisted class member.
 	 */
-	static bool IsNonTargetMemeber(FBlueprintActionFilter const& Filter, UBlueprintNodeSpawner const* BlueprintAction, bool const bPermitNonTargetGlobals);
+	static bool IsNonTargetMemeber(FBlueprintActionFilter const& Filter, FBlueprintActionInfo& BlueprintAction, bool const bPermitNonTargetGlobals);
 
 	/**
 	 * Rejection test that checks to see if the node-spawner is associated with 
@@ -165,7 +165,7 @@ namespace BlueprintActionFilterImpl
 	 * @param  BlueprintAction	The action you wish to query.
 	 * @return True if the action is associated with a hidden field.
 	 */
-	static bool IsFieldCategoryHidden(FBlueprintActionFilter const& Filter, UBlueprintNodeSpawner const* BlueprintAction);
+	static bool IsFieldCategoryHidden(FBlueprintActionFilter const& Filter, FBlueprintActionInfo& BlueprintAction);
 
 	/**
 	 * Rejection test that checks to see if the node-spawner would produce a 
@@ -176,7 +176,7 @@ namespace BlueprintActionFilterImpl
 	 * @param  bMatchExplicitly When true, the node's class has to perfectly match a white-listed type (sub-classes don't count).
 	 * @return True if the action would produce a non-whitelisted node.
 	 */
-	static bool IsFilteredNodeType(FBlueprintActionFilter const& Filter, UBlueprintNodeSpawner const* BlueprintAction, bool const bMatchExplicitly = false);
+	static bool IsFilteredNodeType(FBlueprintActionFilter const& Filter, FBlueprintActionInfo& BlueprintAction, bool const bMatchExplicitly = false);
 
 	/**
 	 * Rejection test that checks to see if the node-spawner is tied to a 
@@ -187,7 +187,7 @@ namespace BlueprintActionFilterImpl
 	 * @param  BlueprintAction	The action you wish to query.
 	 * @return True if the action would produce a bound node, tied to an object that isn't selected.
 	 */
-	static bool IsBoundToUnselectedObject(FBlueprintActionFilter const& Filter, UBlueprintNodeSpawner const* BlueprintAction);
+	static bool IsBoundToUnselectedObject(FBlueprintActionFilter const& Filter, FBlueprintActionInfo& BlueprintAction);
 
 	/**
 	 * 
@@ -196,7 +196,7 @@ namespace BlueprintActionFilterImpl
 	 * @param  BlueprintAction	
 	 * @return 
 	 */
-	static bool IsOutOfScopeLocalVariable(FBlueprintActionFilter const& Filter, UBlueprintNodeSpawner const* BlueprintAction);
+	static bool IsOutOfScopeLocalVariable(FBlueprintActionFilter const& Filter, FBlueprintActionInfo& BlueprintAction);
 
 	/**
 	 * 
@@ -205,7 +205,7 @@ namespace BlueprintActionFilterImpl
 	 * @param  BlueprintAction	
 	 * @return 
 	 */
-	static bool IsSchemaIncompatible(FBlueprintActionFilter const& Filter, UBlueprintNodeSpawner const* BlueprintAction);
+	static bool IsSchemaIncompatible(FBlueprintActionFilter const& Filter, FBlueprintActionInfo& BlueprintAction);
 
 	/**
 	 * 
@@ -214,7 +214,7 @@ namespace BlueprintActionFilterImpl
 	 * @param  Pin	
 	 * @return 
 	 */
-	static bool HasMatchingPin(UBlueprintNodeSpawner const* BlueprintAction, UEdGraphPin const* Pin);
+	static bool HasMatchingPin(FBlueprintActionInfo& BlueprintAction, UEdGraphPin const* Pin);
 
 	/**
 	 * 
@@ -224,7 +224,7 @@ namespace BlueprintActionFilterImpl
 	 * @param  MemberNodeSpawner
 	 * @return 
 	 */
-	static bool IsPinCompatibleWithTargetSelf(UEdGraphPin const* Pin, UField const* MemberField, UBlueprintNodeSpawner const* MemberNodeSpawner);
+	static bool IsPinCompatibleWithTargetSelf(UEdGraphPin const* Pin, UField const* MemberField, FBlueprintActionInfo& MemberNodeSpawner);
 
 	/**
 	 * 
@@ -233,7 +233,7 @@ namespace BlueprintActionFilterImpl
 	 * @param  BlueprintAction	
 	 * @return 
 	 */
-	static bool IsFunctionMissingPinParam(FBlueprintActionFilter const& Filter, UBlueprintNodeSpawner const* BlueprintAction);
+	static bool IsFunctionMissingPinParam(FBlueprintActionFilter const& Filter, FBlueprintActionInfo& BlueprintAction);
 
 	/**
 	 * 
@@ -242,7 +242,7 @@ namespace BlueprintActionFilterImpl
 	 * @param  BlueprintAction	
 	 * @return 
 	 */
-	static bool IsMissmatchedPropertyType(FBlueprintActionFilter const& Filter, UBlueprintNodeSpawner const* BlueprintAction);
+	static bool IsMissmatchedPropertyType(FBlueprintActionFilter const& Filter, FBlueprintActionInfo& BlueprintAction);
 
 	/**
 	 * 
@@ -251,16 +251,17 @@ namespace BlueprintActionFilterImpl
 	 * @param  BlueprintAction	
 	 * @return 
 	 */
-	static bool IsMissingMatchingPinParam(FBlueprintActionFilter const& Filter, UBlueprintNodeSpawner const* BlueprintAction);
+	static bool IsMissingMatchingPinParam(FBlueprintActionFilter const& Filter, FBlueprintActionInfo& BlueprintAction);
 
 	/**
-	 * Dynamic casts should only show results for casting to classes that the context pin is a child of (and not itself)
+	 * Dynamic casts should only show results for casting to classes that the 
+	 * context pin is a child of (and not itself).
 	 * 
 	 * @param  Filter	
 	 * @param  BlueprintAction	
 	 * @return 
 	 */
-	static bool IsNotSubClassCast(FBlueprintActionFilter const& Filter, UBlueprintNodeSpawner const* BlueprintAction);
+	static bool IsNotSubClassCast(FBlueprintActionFilter const& Filter, FBlueprintActionInfo& BlueprintAction);
 };
 
 
@@ -270,20 +271,20 @@ UClass* BlueprintActionFilterImpl::GetBlueprintClass(UBlueprint const* const Blu
 	return (Blueprint->SkeletonGeneratedClass != nullptr) ? Blueprint->SkeletonGeneratedClass : Blueprint->ParentClass;
 }
 
-
-
 //------------------------------------------------------------------------------
-static bool BlueprintActionFilterImpl::IsImpure(UBlueprintNodeSpawner const* NodeSpawner)
+static bool BlueprintActionFilterImpl::IsImpure(FBlueprintActionInfo& BlueprintAction)
 {
 	bool bIsImpure = false;
 	
-	if (UFunction const* Function = FBlueprintNodeSpawnerUtils::GetAssociatedFunction(NodeSpawner))
+	if (UFunction const* Function = BlueprintAction.GetAssociatedFunction())
 	{
 		bIsImpure = (Function->HasAnyFunctionFlags(FUNC_BlueprintPure) == false);
 	}
-	else if (NodeSpawner->NodeClass != nullptr)
+	else
 	{
-		UClass* NodeClass = NodeSpawner->NodeClass;
+		UClass const* NodeClass = BlueprintAction.GetNodeClass();
+		check(NodeClass != nullptr);
+
 		// TODO: why are some of these "impure"?... we shouldn't have hardcoded 
 		//       node types here (game modules cannot add their node types here,
 		//       so we should find another way of identifying "pure" node types)
@@ -298,16 +299,16 @@ static bool BlueprintActionFilterImpl::IsImpure(UBlueprintNodeSpawner const* Nod
 }
 
 //------------------------------------------------------------------------------
-static bool BlueprintActionFilterImpl::IsLatent(UBlueprintNodeSpawner const* NodeSpawner)
+static bool BlueprintActionFilterImpl::IsLatent(FBlueprintActionInfo& BlueprintAction)
 {
 	bool bIsLatent = false;
 
-	UClass* NodeClass = NodeSpawner->NodeClass;
+	UClass const* NodeClass = BlueprintAction.GetNodeClass();
 	if ((NodeClass != nullptr) && NodeClass->IsChildOf<UK2Node_BaseAsyncTask>())
 	{
 		bIsLatent = true;
 	}
-	else if (UFunction const* Function = FBlueprintNodeSpawnerUtils::GetAssociatedFunction(NodeSpawner))
+	else if (UFunction const* Function = BlueprintAction.GetAssociatedFunction())
 	{
 		bIsLatent = (Function->HasMetaData(FBlueprintMetadata::MD_Latent) != false);
 	}
@@ -339,12 +340,12 @@ static bool BlueprintActionFilterImpl::IsGloballyAccessible(UField const* Field)
 }
 
 //------------------------------------------------------------------------------
-static bool BlueprintActionFilterImpl::IsEventUnimplementable(FBlueprintActionFilter const& Filter, UBlueprintNodeSpawner const* BlueprintAction)
+static bool BlueprintActionFilterImpl::IsEventUnimplementable(FBlueprintActionFilter const& Filter, FBlueprintActionInfo& BlueprintAction)
 {
 	bool bIsFilteredOut = false;
 	FBlueprintActionContext const& FilterContext = Filter.Context;
 	
-	if (UBlueprintEventNodeSpawner const* EventSpawner = Cast<UBlueprintEventNodeSpawner>(BlueprintAction))
+	if (UBlueprintEventNodeSpawner const* EventSpawner = Cast<UBlueprintEventNodeSpawner>(BlueprintAction.NodeSpawner))
 	{
 		if (UFunction const* EventFunc = EventSpawner->GetEventFunction())
 		{
@@ -379,17 +380,17 @@ static bool BlueprintActionFilterImpl::IsEventUnimplementable(FBlueprintActionFi
 }
 
 //------------------------------------------------------------------------------
-static bool BlueprintActionFilterImpl::IsFieldInaccessible(FBlueprintActionFilter const& Filter, UBlueprintNodeSpawner const* BlueprintAction)
+static bool BlueprintActionFilterImpl::IsFieldInaccessible(FBlueprintActionFilter const& Filter, FBlueprintActionInfo& BlueprintAction)
 {
 	bool bIsFilteredOut = false;
 	FBlueprintActionContext const& FilterContext = Filter.Context;
 	
-	if (UField const* ClassField = FBlueprintNodeSpawnerUtils::GetAssociatedField(BlueprintAction))
+	if (UField const* ClassField = BlueprintAction.GetAssociatedMemberField())
 	{
 		bool bIsProtected = ClassField->HasMetaData(FBlueprintMetadata::MD_Protected);
 		bool bIsPrivate   = ClassField->HasMetaData(FBlueprintMetadata::MD_Private);
 
-		UClass* FieldOwner = Cast<UClass>(ClassField->GetOuter());
+		UClass* FieldOwner = ClassField->GetOwnerClass();
 		for (UBlueprint const* Blueprint : FilterContext.Blueprints)
 		{
 			UClass* BpClass = Blueprint->GeneratedClass;
@@ -413,12 +414,12 @@ static bool BlueprintActionFilterImpl::IsFieldInaccessible(FBlueprintActionFilte
 }
 
 //------------------------------------------------------------------------------
-static bool BlueprintActionFilterImpl::IsRestrictedClassMember(FBlueprintActionFilter const& Filter, UBlueprintNodeSpawner const* BlueprintAction)
+static bool BlueprintActionFilterImpl::IsRestrictedClassMember(FBlueprintActionFilter const& Filter, FBlueprintActionInfo& BlueprintAction)
 {
 	bool bIsFilteredOut = false;
 	FBlueprintActionContext const& FilterContext = Filter.Context;
 	
-	if (UClass const* ActionClass = FBlueprintNodeSpawnerUtils::GetAssociatedClass(BlueprintAction))
+	if (UClass const* ActionClass = BlueprintAction.GetOwnerClass())
 	{
 		if (ActionClass->HasMetaData(FBlueprintMetadata::MD_RestrictedToClasses))
 		{
@@ -453,14 +454,14 @@ static bool BlueprintActionFilterImpl::IsRestrictedClassMember(FBlueprintActionF
 }
 
 //------------------------------------------------------------------------------
-static bool BlueprintActionFilterImpl::IsPermissionNotGranted(FBlueprintActionFilter const& Filter, UBlueprintNodeSpawner const* BlueprintAction)
+static bool BlueprintActionFilterImpl::IsPermissionNotGranted(FBlueprintActionFilter const& Filter, FBlueprintActionInfo& BlueprintAction)
 {
 	bool bIsFilteredOut = false;
 	FBlueprintActionContext const& FilterContext = Filter.Context;
 
-	if (UProperty const* const Property = FBlueprintNodeSpawnerUtils::GetAssociatedProperty(BlueprintAction))
+	if (UProperty const* const Property = BlueprintAction.GetAssociatedProperty())
 	{
-		UClass const* const NodeClass = BlueprintAction->NodeClass;
+		UClass const* const NodeClass = BlueprintAction.GetNodeClass();
 		for (UBlueprint const* Blueprint : FilterContext.Blueprints)
 		{
 			bool const bIsReadOnly = FBlueprintEditorUtils::IsPropertyReadOnlyInCurrentBlueprint(Blueprint, Property);
@@ -475,14 +476,16 @@ static bool BlueprintActionFilterImpl::IsPermissionNotGranted(FBlueprintActionFi
 }
 
 //------------------------------------------------------------------------------
-static bool BlueprintActionFilterImpl::IsDeprecated(FBlueprintActionFilter const& /*Filter*/, UBlueprintNodeSpawner const* BlueprintAction)
+static bool BlueprintActionFilterImpl::IsDeprecated(FBlueprintActionFilter const& /*Filter*/, FBlueprintActionInfo& BlueprintAction)
 {
+	checkSlow(BlueprintAction.GetNodeClass() != nullptr);
+
 	bool bIsFilteredOut = false;	
-	if ((BlueprintAction->NodeClass != nullptr) && BlueprintAction->NodeClass->HasAnyClassFlags(CLASS_Deprecated))
+	if (BlueprintAction.GetNodeClass()->HasAnyClassFlags(CLASS_Deprecated))
 	{
 		bIsFilteredOut = true;
 	}
-	else if (UClass* ActionClass = FBlueprintNodeSpawnerUtils::GetAssociatedClass(BlueprintAction))
+	else if (UClass const* ActionClass = BlueprintAction.GetOwnerClass())
 	{
 		bIsFilteredOut = ActionClass->HasAnyClassFlags(CLASS_Deprecated);
 	}
@@ -490,10 +493,10 @@ static bool BlueprintActionFilterImpl::IsDeprecated(FBlueprintActionFilter const
 }
 
 //------------------------------------------------------------------------------
-static bool BlueprintActionFilterImpl::IsPersistentNonTargetField(FBlueprintActionFilter const& Filter, UBlueprintNodeSpawner const* BlueprintAction)
+static bool BlueprintActionFilterImpl::IsPersistentNonTargetField(FBlueprintActionFilter const& Filter, FBlueprintActionInfo& BlueprintAction)
 {
 	bool bIsFilteredOut = false;
-	if (UField const* ClassField = FBlueprintNodeSpawnerUtils::GetAssociatedField(BlueprintAction))
+	if (UField const* ClassField = BlueprintAction.GetAssociatedMemberField())
 	{
 		bIsFilteredOut = IsGloballyAccessible(ClassField);
 		
@@ -516,10 +519,10 @@ static bool BlueprintActionFilterImpl::IsPersistentNonTargetField(FBlueprintActi
 }
 
 //------------------------------------------------------------------------------
-static bool BlueprintActionFilterImpl::IsNonTargetMemeber(FBlueprintActionFilter const& Filter, UBlueprintNodeSpawner const* BlueprintAction, bool bPermitNonTargetGlobals)
+static bool BlueprintActionFilterImpl::IsNonTargetMemeber(FBlueprintActionFilter const& Filter, FBlueprintActionInfo& BlueprintAction, bool bPermitNonTargetGlobals)
 {
 	bool bIsFilteredOut = false;
-	if (UField const* ClassField = FBlueprintNodeSpawnerUtils::GetAssociatedField(BlueprintAction))
+	if (UField const* ClassField = BlueprintAction.GetAssociatedMemberField())
 	{
 		UClass* FieldClass = Cast<UClass>(ClassField->GetOuter());
 		bool const bIsInterfaceMethod = FieldClass->IsChildOf<UInterface>();
@@ -545,17 +548,17 @@ static bool BlueprintActionFilterImpl::IsNonTargetMemeber(FBlueprintActionFilter
 }
 
 //------------------------------------------------------------------------------
-static bool BlueprintActionFilterImpl::IsFieldCategoryHidden(FBlueprintActionFilter const& Filter, UBlueprintNodeSpawner const* BlueprintAction)
+static bool BlueprintActionFilterImpl::IsFieldCategoryHidden(FBlueprintActionFilter const& Filter, FBlueprintActionInfo& BlueprintAction)
 {
 	bool bIsFilteredOut = false;
 	FBlueprintActionContext const& FilterContext = Filter.Context;
 
-	if (UField const* AssociatedField = FBlueprintNodeSpawnerUtils::GetAssociatedField(BlueprintAction))
+	if (UField const* AssociatedField = BlueprintAction.GetAssociatedMemberField())
 	{
 		DECLARE_DELEGATE_RetVal_OneParam(bool, FIsFieldHiddenDelegate, UClass*);
 		FIsFieldHiddenDelegate IsFieldHiddenDelegate;
 
-		if (UFunction const* Function = Cast<UFunction>(AssociatedField))
+		if (UFunction const* Function = BlueprintAction.GetAssociatedFunction())
 		{
 			auto IsFunctionHiddenLambda = [](UClass* Class, UFunction const* Function)->bool
 			{
@@ -563,7 +566,7 @@ static bool BlueprintActionFilterImpl::IsFieldCategoryHidden(FBlueprintActionFil
 			};
 			IsFieldHiddenDelegate = FIsFieldHiddenDelegate::CreateStatic(IsFunctionHiddenLambda, Function);
 		}
-		else if (UProperty const* Property = Cast<UProperty>(AssociatedField))
+		else if (UProperty const* Property = BlueprintAction.GetAssociatedProperty())
 		{
 			auto IsPropertyHiddenLambda = [](UClass* Class, UProperty const* Property)->bool
 			{
@@ -594,7 +597,7 @@ static bool BlueprintActionFilterImpl::IsFieldCategoryHidden(FBlueprintActionFil
 }
 
 //------------------------------------------------------------------------------
-static bool BlueprintActionFilterImpl::IsIncompatibleImpureNode(FBlueprintActionFilter const& Filter, UBlueprintNodeSpawner const* BlueprintAction)
+static bool BlueprintActionFilterImpl::IsIncompatibleImpureNode(FBlueprintActionFilter const& Filter, FBlueprintActionInfo& BlueprintAction)
 {
 	bool bAllowImpureNodes = true;
 	FBlueprintActionContext const& FilterContext = Filter.Context;
@@ -612,14 +615,14 @@ static bool BlueprintActionFilterImpl::IsIncompatibleImpureNode(FBlueprintAction
 }
 
 //------------------------------------------------------------------------------
-static bool BlueprintActionFilterImpl::IsIncompatibleWithGraphType(FBlueprintActionFilter const& Filter, UBlueprintNodeSpawner const* BlueprintAction)
+static bool BlueprintActionFilterImpl::IsIncompatibleWithGraphType(FBlueprintActionFilter const& Filter, FBlueprintActionInfo& BlueprintAction)
 {
 	bool bIsFilteredOut = false;
 	FBlueprintActionContext const& FilterContext = Filter.Context;
 	
-	if (UClass* NodeClass = BlueprintAction->NodeClass)
+	if (UClass const* NodeClass = BlueprintAction.GetNodeClass())
 	{
-		if (UEdGraphNode* NodeCDO = NodeClass->GetDefaultObject<UEdGraphNode>())
+		if (UEdGraphNode const* NodeCDO = CastChecked<UEdGraphNode>(NodeClass->ClassDefaultObject))
 		{
 			for (UEdGraph* Graph : FilterContext.Graphs)
 			{
@@ -635,11 +638,11 @@ static bool BlueprintActionFilterImpl::IsIncompatibleWithGraphType(FBlueprintAct
 }
 
 //------------------------------------------------------------------------------
-static bool BlueprintActionFilterImpl::IsFilteredNodeType(FBlueprintActionFilter const& Filter, UBlueprintNodeSpawner const* BlueprintAction, bool const bExcludeChildClasses)
+static bool BlueprintActionFilterImpl::IsFilteredNodeType(FBlueprintActionFilter const& Filter, FBlueprintActionInfo& BlueprintAction, bool const bExcludeChildClasses)
 {
 	bool bIsFilteredOut = (Filter.PermittedNodeTypes.Num() > 0);
 
-	UClass* NodeClass = BlueprintAction->NodeClass;
+	UClass const* NodeClass = BlueprintAction.GetNodeClass();
 	if (NodeClass != nullptr)
 	{
 		for (TSubclassOf<UEdGraphNode> AllowedClass : Filter.PermittedNodeTypes)
@@ -667,10 +670,10 @@ static bool BlueprintActionFilterImpl::IsFilteredNodeType(FBlueprintActionFilter
 }
 
 //------------------------------------------------------------------------------
-static bool BlueprintActionFilterImpl::IsBoundToUnselectedObject(FBlueprintActionFilter const& Filter, UBlueprintNodeSpawner const* BlueprintAction)
+static bool BlueprintActionFilterImpl::IsBoundToUnselectedObject(FBlueprintActionFilter const& Filter, FBlueprintActionInfo& BlueprintAction)
 {
 	bool bIsFilteredOut = false;
-	if (UBlueprintBoundEventNodeSpawner const* BoundSpawner = Cast<UBlueprintBoundEventNodeSpawner>(BlueprintAction))
+	if (UBlueprintBoundEventNodeSpawner const* BoundSpawner = Cast<UBlueprintBoundEventNodeSpawner>(BlueprintAction.NodeSpawner))
 	{
 		bIsFilteredOut = true;
 		for (UObject const* SelectedObj : Filter.Context.SelectedObjects)
@@ -686,10 +689,10 @@ static bool BlueprintActionFilterImpl::IsBoundToUnselectedObject(FBlueprintActio
 }
 
 //------------------------------------------------------------------------------
-static bool BlueprintActionFilterImpl::IsOutOfScopeLocalVariable(FBlueprintActionFilter const& Filter, UBlueprintNodeSpawner const* BlueprintAction)
+static bool BlueprintActionFilterImpl::IsOutOfScopeLocalVariable(FBlueprintActionFilter const& Filter, FBlueprintActionInfo& BlueprintAction)
 {
 	bool bIsFilteredOut = false;
-	if (UBlueprintVariableNodeSpawner const* VarSpawner = Cast<UBlueprintVariableNodeSpawner>(BlueprintAction))
+	if (UBlueprintVariableNodeSpawner const* VarSpawner = Cast<UBlueprintVariableNodeSpawner>(BlueprintAction.NodeSpawner))
 	{
 		if (VarSpawner->IsLocalVariable())
 		{
@@ -710,12 +713,14 @@ static bool BlueprintActionFilterImpl::IsOutOfScopeLocalVariable(FBlueprintActio
 }
 
 //------------------------------------------------------------------------------
-static bool BlueprintActionFilterImpl::IsSchemaIncompatible(FBlueprintActionFilter const& Filter, UBlueprintNodeSpawner const* BlueprintAction)
+static bool BlueprintActionFilterImpl::IsSchemaIncompatible(FBlueprintActionFilter const& Filter, FBlueprintActionInfo& BlueprintAction)
 {
 	bool bIsFilteredOut = false;
 	FBlueprintActionContext const& FilterContext = Filter.Context;
 
-	UEdGraphNode const* NodeCDO = BlueprintAction->NodeClass->GetDefaultObject<UEdGraphNode>();
+	UClass const* NodeClass = BlueprintAction.GetNodeClass();
+	UEdGraphNode const* NodeCDO = CastChecked<UEdGraphNode>(NodeClass->ClassDefaultObject);
+	checkSlow(NodeCDO != nullptr);
 
 	auto IsSchemaIncompatibleLambda = [NodeCDO](TArray<UEdGraph*> const& GraphList)->bool
 	{
@@ -755,12 +760,12 @@ static bool BlueprintActionFilterImpl::IsSchemaIncompatible(FBlueprintActionFilt
 }
 
 //------------------------------------------------------------------------------
-static bool BlueprintActionFilterImpl::HasMatchingPin(UBlueprintNodeSpawner const* BlueprintAction, UEdGraphPin const* Pin)
+static bool BlueprintActionFilterImpl::HasMatchingPin(FBlueprintActionInfo& BlueprintAction, UEdGraphPin const* Pin)
 {
 	bool bHasCompatiblePin = false;
 
 	UEdGraph* OuterGraph = Pin->GetOwningNode()->GetGraph();
-	if (UEdGraphNode* TemplateNode = BlueprintAction->GetTemplateNode(OuterGraph))
+	if (UEdGraphNode* TemplateNode = BlueprintAction.NodeSpawner->GetTemplateNode(OuterGraph))
 	{
 		if (TemplateNode->Pins.Num() == 0)
 		{
@@ -796,10 +801,10 @@ static bool BlueprintActionFilterImpl::HasMatchingPin(UBlueprintNodeSpawner cons
 }
 
 //------------------------------------------------------------------------------
-static bool BlueprintActionFilterImpl::IsPinCompatibleWithTargetSelf(UEdGraphPin const* Pin, UField const* MemberField, UBlueprintNodeSpawner const* MemberNodeSpawner)
+static bool BlueprintActionFilterImpl::IsPinCompatibleWithTargetSelf(UEdGraphPin const* Pin, UField const* MemberField, FBlueprintActionInfo& BlueprintAction)
 {
 	bool bIsCompatible = false;
-	checkSlow(FBlueprintNodeSpawnerUtils::GetAssociatedField(MemberNodeSpawner) == MemberField);
+	checkSlow(BlueprintAction.GetAssociatedMemberField() == MemberField);
 
 	UClass* OwnerClass = Cast<UClass>(MemberField->GetOuter());
 	if ((Pin->Direction == EGPD_Output) && (OwnerClass != nullptr))
@@ -833,7 +838,7 @@ static bool BlueprintActionFilterImpl::IsPinCompatibleWithTargetSelf(UEdGraphPin
 					else
 					{
 						UEdGraph* OuterGraph = Pin->GetOwningNode()->GetGraph();
-						if (UK2Node* TemplateNode = Cast<UK2Node>(MemberNodeSpawner->GetTemplateNode(OuterGraph)))
+						if (UK2Node* TemplateNode = Cast<UK2Node>(BlueprintAction.NodeSpawner->GetTemplateNode(OuterGraph)))
 						{
 							bIsCompatible = TemplateNode->AllowMultipleSelfs(/*bInputAsArray =*/true);
 						}	
@@ -847,13 +852,13 @@ static bool BlueprintActionFilterImpl::IsPinCompatibleWithTargetSelf(UEdGraphPin
 }
 
 //------------------------------------------------------------------------------
-static bool BlueprintActionFilterImpl::IsFunctionMissingPinParam(FBlueprintActionFilter const& Filter, UBlueprintNodeSpawner const* BlueprintAction)
+static bool BlueprintActionFilterImpl::IsFunctionMissingPinParam(FBlueprintActionFilter const& Filter, FBlueprintActionInfo& BlueprintAction)
 {
 	bool bIsFilteredOut = false;
-	if (UFunction const* AssociatedFunc = FBlueprintNodeSpawnerUtils::GetAssociatedFunction(BlueprintAction))
+	if (UFunction const* AssociatedFunc = BlueprintAction.GetAssociatedFunction())
 	{
 		UEdGraphSchema_K2 const* K2Schema = GetDefault<UEdGraphSchema_K2>();
-		bool const bIsEventSpawner = BlueprintAction->NodeClass->IsChildOf<UK2Node_Event>();
+		bool const bIsEventSpawner = BlueprintAction.GetNodeClass()->IsChildOf<UK2Node_Event>();
 
 		for (int32 PinIndex = 0; !bIsFilteredOut && (PinIndex < Filter.Context.Pins.Num()); ++PinIndex)
 		{
@@ -892,17 +897,17 @@ static bool BlueprintActionFilterImpl::IsFunctionMissingPinParam(FBlueprintActio
 }
 
 //------------------------------------------------------------------------------
-static bool BlueprintActionFilterImpl::IsMissmatchedPropertyType(FBlueprintActionFilter const& Filter, UBlueprintNodeSpawner const* BlueprintAction)
+static bool BlueprintActionFilterImpl::IsMissmatchedPropertyType(FBlueprintActionFilter const& Filter, FBlueprintActionInfo& BlueprintAction)
 {
 	bool bIsFilteredOut = false;
-	if (UProperty const* Property = FBlueprintNodeSpawnerUtils::GetAssociatedProperty(BlueprintAction))
+	if (UProperty const* Property = BlueprintAction.GetAssociatedProperty())
 	{
 		TArray<UEdGraphPin*> const& ContextPins = Filter.Context.Pins;
 		if (ContextPins.Num() > 0)
 		{
 			bool const bIsDelegate = Property->IsA<UMulticastDelegateProperty>();
-			bool const bIsGetter   = BlueprintAction->NodeClass->IsChildOf<UK2Node_VariableGet>();
-			bool const bIsSetter   = BlueprintAction->NodeClass->IsChildOf<UK2Node_VariableSet>();
+			bool const bIsGetter   = BlueprintAction.GetNodeClass()->IsChildOf<UK2Node_VariableGet>();
+			bool const bIsSetter   = BlueprintAction.GetNodeClass()->IsChildOf<UK2Node_VariableSet>();
 
 			//K2Schema->ConvertPropertyToPinType(VariableProperty, /*out*/ VariablePin->PinType);
 			for (int32 PinIndex = 0; !bIsFilteredOut && PinIndex < ContextPins.Num(); ++PinIndex)
@@ -957,11 +962,11 @@ static bool BlueprintActionFilterImpl::IsMissmatchedPropertyType(FBlueprintActio
 }
 
 //------------------------------------------------------------------------------
-static bool BlueprintActionFilterImpl::IsMissingMatchingPinParam(FBlueprintActionFilter const& Filter, UBlueprintNodeSpawner const* BlueprintAction)
+static bool BlueprintActionFilterImpl::IsMissingMatchingPinParam(FBlueprintActionFilter const& Filter, FBlueprintActionInfo& BlueprintAction)
 {
 	bool bIsFilteredOut = false;
 
-	UField const* AssociatedField = FBlueprintNodeSpawnerUtils::GetAssociatedField(BlueprintAction);
+	UField const* AssociatedField = BlueprintAction.GetAssociatedMemberField();
 	// we have a separate tests for field nodes (IsFunctionMissingPinParam/IsMissmatchedPropertyType)
 	if (AssociatedField == nullptr)
 	{
@@ -979,30 +984,31 @@ static bool BlueprintActionFilterImpl::IsMissingMatchingPinParam(FBlueprintActio
 }
 
 //------------------------------------------------------------------------------
-static bool BlueprintActionFilterImpl::IsNotSubClassCast(FBlueprintActionFilter const& Filter, UBlueprintNodeSpawner const* BlueprintAction)
+static bool BlueprintActionFilterImpl::IsNotSubClassCast(FBlueprintActionFilter const& Filter, FBlueprintActionInfo& BlueprintAction)
 {
 	bool bIsFilteredOut = false;
 
-	if(BlueprintAction->NodeClass && BlueprintAction->NodeClass == UK2Node_DynamicCast::StaticClass())
+	if (BlueprintAction.GetNodeClass() == UK2Node_DynamicCast::StaticClass())
 	{
 		for(UEdGraphPin const* ContextPin : Filter.Context.Pins)
 		{
 			// Only worry about removing cast nodes when dragging off output pins
-			if(ContextPin->Direction == EGPD_Output)
+			if (ContextPin->Direction != EGPD_Output)
+			{
+				bIsFilteredOut = true;
+				break;
+			}
+			
+			if (ContextPin->PinType.PinSubCategoryObject.IsValid())
 			{
-				UEdGraph* OuterGraph = ContextPin->GetOwningNode()->GetGraph();
-				if (UK2Node_DynamicCast* CastNode = Cast<UK2Node_DynamicCast>(BlueprintAction->GetTemplateNode(OuterGraph)))
+				UClass const* CastClass = BlueprintAction.GetOwnerClass();
+				checkSlow(CastClass != nullptr);
+
+				UClass const* ContextPinClass = Cast<UClass>(ContextPin->PinType.PinSubCategoryObject.Get());
+				if ((ContextPinClass == CastClass) || !CastClass->IsChildOf(ContextPinClass))
 				{
-					if( ContextPin->PinType.PinSubCategoryObject.IsValid() && CastNode->TargetType)
-					{
-						UClass* ContextPinClass = Cast<UClass>(ContextPin->PinType.PinSubCategoryObject.Get());
-						UClass* ResultPinClass = CastNode->TargetType;
-						if( ContextPinClass == ResultPinClass || !ResultPinClass->IsChildOf(ContextPinClass) )
-						{
-							bIsFilteredOut = true;
-							break;
-						}
-					}
+					bIsFilteredOut = true;
+					break;
 				}
 			}
 		}
@@ -1010,6 +1016,118 @@ static bool BlueprintActionFilterImpl::IsNotSubClassCast(FBlueprintActionFilter
 	return bIsFilteredOut;
 }
 
+/*******************************************************************************
+ * FBlueprintActionInfo
+ ******************************************************************************/
+
+/** */
+namespace EBlueprintActionInfoFlags
+{
+	enum 
+	{
+		CachedClass    = (1 << 0),
+		CachedField    = (1 << 1),
+		CachedProperty = (1 << 2),
+		CachedFunction = (1 << 3),
+	};
+}
+
+//------------------------------------------------------------------------------
+FBlueprintActionInfo::FBlueprintActionInfo(UObject const* ActionOwnerIn, UBlueprintNodeSpawner const* Action)
+	: NodeSpawner(Action)
+	, ActionOwner(ActionOwnerIn)
+	, CacheFlags(0)
+	, CachedOwnerClass(nullptr)
+	, CachedActionField(nullptr)
+	, CachedActionProperty(nullptr)
+	, CachedActionFunction(nullptr)
+{
+	checkSlow(NodeSpawner != nullptr);
+}
+
+//------------------------------------------------------------------------------
+UObject const* FBlueprintActionInfo::GetActionOwner()
+{
+	return ActionOwner;
+}
+
+//------------------------------------------------------------------------------
+UClass const* FBlueprintActionInfo::GetOwnerClass()
+{
+	if ((CacheFlags & EBlueprintActionInfoFlags::CachedClass) == 0)
+	{
+		CachedOwnerClass = Cast<UClass>(ActionOwner);
+		if (CachedOwnerClass == GetNodeClass())
+		{
+			CachedOwnerClass = nullptr;
+		}
+
+		if (CachedOwnerClass == nullptr)
+		{
+			CachedOwnerClass = GetAssociatedMemberField()->GetOwnerClass();
+		}
+		CacheFlags |= EBlueprintActionInfoFlags::CachedClass;
+	}
+	return CachedOwnerClass;
+}
+
+//------------------------------------------------------------------------------
+UClass const* FBlueprintActionInfo::GetNodeClass()
+{
+	UClass const* NodeClass = NodeSpawner->NodeClass;
+	checkSlow(NodeClass != nullptr);
+	return NodeClass;
+}
+
+//------------------------------------------------------------------------------
+UField const* FBlueprintActionInfo::GetAssociatedMemberField()
+{
+	if ((CacheFlags & EBlueprintActionInfoFlags::CachedField) == 0)
+	{
+		CachedActionField = FBlueprintNodeSpawnerUtils::GetAssociatedField(NodeSpawner);
+		CacheFlags |= EBlueprintActionInfoFlags::CachedField;
+	}
+	return CachedActionField;
+}
+
+//------------------------------------------------------------------------------
+UProperty const* FBlueprintActionInfo::GetAssociatedProperty()
+{
+	if ((CacheFlags & EBlueprintActionInfoFlags::CachedProperty) == 0)
+	{
+		if (CacheFlags & EBlueprintActionInfoFlags::CachedField)
+		{
+			CachedActionProperty = Cast<UProperty>(CachedActionField);
+		}
+		else
+		{
+			CachedActionProperty = FBlueprintNodeSpawnerUtils::GetAssociatedProperty(NodeSpawner);
+			CachedActionField = CachedActionProperty;
+		}
+		CacheFlags |= (EBlueprintActionInfoFlags::CachedProperty | EBlueprintActionInfoFlags::CachedField);
+	}
+	return CachedActionProperty;
+}
+
+//------------------------------------------------------------------------------
+UFunction const* FBlueprintActionInfo::GetAssociatedFunction()
+{
+	if ((CacheFlags & EBlueprintActionInfoFlags::CachedFunction) == 0)
+	{
+		if (CacheFlags & EBlueprintActionInfoFlags::CachedField)
+		{
+			CachedActionFunction = Cast<UFunction>(CachedActionField);
+		}
+		else
+		{
+			CachedActionFunction = FBlueprintNodeSpawnerUtils::GetAssociatedFunction(NodeSpawner);
+			CachedActionField = CachedActionFunction;
+		}
+		CacheFlags |= (EBlueprintActionInfoFlags::CachedFunction | EBlueprintActionInfoFlags::CachedField);
+	}
+	return CachedActionFunction;
+}
+
 /*******************************************************************************
  * FBlueprintActionFilter
  ******************************************************************************/
@@ -1031,34 +1149,34 @@ FBlueprintActionFilter::FBlueprintActionFilter(uint32 Flags/*= 0x00*/)
 	//
 	// this test in-particular spawns a template-node and then calls 
 	// AllocateDefaultPins() which is costly, so it should be very last!
-	AddIsFilteredTest(FIsFilteredDelegate::CreateStatic(IsMissingMatchingPinParam));
-	AddIsFilteredTest(FIsFilteredDelegate::CreateStatic(IsMissmatchedPropertyType));
-	AddIsFilteredTest(FIsFilteredDelegate::CreateStatic(IsFunctionMissingPinParam));
-	AddIsFilteredTest(FIsFilteredDelegate::CreateStatic(IsIncompatibleImpureNode));
+	AddRejectionTest(FRejectionTestDelegate::CreateStatic(IsMissingMatchingPinParam));
+	AddRejectionTest(FRejectionTestDelegate::CreateStatic(IsMissmatchedPropertyType));
+	AddRejectionTest(FRejectionTestDelegate::CreateStatic(IsFunctionMissingPinParam));
+	AddRejectionTest(FRejectionTestDelegate::CreateStatic(IsIncompatibleImpureNode));
 	
-	AddIsFilteredTest(FIsFilteredDelegate::CreateStatic(IsFieldCategoryHidden));	
+	AddRejectionTest(FRejectionTestDelegate::CreateStatic(IsFieldCategoryHidden));
 	if (Flags & BPFILTER_RejectPersistentNonTargetFields)
 	{
-		AddIsFilteredTest(FIsFilteredDelegate::CreateStatic(IsPersistentNonTargetField));
+		AddRejectionTest(FRejectionTestDelegate::CreateStatic(IsPersistentNonTargetField));
 	}
 	
-	AddIsFilteredTest(FIsFilteredDelegate::CreateStatic(IsFieldInaccessible));
-	AddIsFilteredTest(FIsFilteredDelegate::CreateStatic(IsNotSubClassCast));
-	AddIsFilteredTest(FIsFilteredDelegate::CreateStatic(IsEventUnimplementable));
-	AddIsFilteredTest(FIsFilteredDelegate::CreateStatic(IsPermissionNotGranted));
-	AddIsFilteredTest(FIsFilteredDelegate::CreateStatic(IsRestrictedClassMember));
-	AddIsFilteredTest(FIsFilteredDelegate::CreateStatic(IsIncompatibleWithGraphType));
-	AddIsFilteredTest(FIsFilteredDelegate::CreateStatic(IsSchemaIncompatible));
+	AddRejectionTest(FRejectionTestDelegate::CreateStatic(IsFieldInaccessible));
+	AddRejectionTest(FRejectionTestDelegate::CreateStatic(IsNotSubClassCast));
+	AddRejectionTest(FRejectionTestDelegate::CreateStatic(IsEventUnimplementable));
+	AddRejectionTest(FRejectionTestDelegate::CreateStatic(IsPermissionNotGranted));
+	AddRejectionTest(FRejectionTestDelegate::CreateStatic(IsRestrictedClassMember));
+	AddRejectionTest(FRejectionTestDelegate::CreateStatic(IsIncompatibleWithGraphType));
+	AddRejectionTest(FRejectionTestDelegate::CreateStatic(IsSchemaIncompatible));
 
 	if (!(Flags & BPFILTER_PermitDeprecated))
 	{
-		AddIsFilteredTest(FIsFilteredDelegate::CreateStatic(IsDeprecated));
+		AddRejectionTest(FRejectionTestDelegate::CreateStatic(IsDeprecated));
 	}
 
-	AddIsFilteredTest(FIsFilteredDelegate::CreateStatic(IsFilteredNodeType, (Flags & BPFILTER_RejectPermittedNodeSubClasses)   != 0));
-	AddIsFilteredTest(FIsFilteredDelegate::CreateStatic(IsNonTargetMemeber, (Flags & BPFILTER_RejectPersistentNonTargetFields) == 0));
-	AddIsFilteredTest(FIsFilteredDelegate::CreateStatic(IsBoundToUnselectedObject));
-	AddIsFilteredTest(FIsFilteredDelegate::CreateStatic(IsOutOfScopeLocalVariable));
+	AddRejectionTest(FRejectionTestDelegate::CreateStatic(IsFilteredNodeType, (Flags & BPFILTER_RejectPermittedNodeSubClasses) != 0));
+	AddRejectionTest(FRejectionTestDelegate::CreateStatic(IsNonTargetMemeber, (Flags & BPFILTER_RejectPersistentNonTargetFields) == 0));
+	AddRejectionTest(FRejectionTestDelegate::CreateStatic(IsBoundToUnselectedObject));
+	AddRejectionTest(FRejectionTestDelegate::CreateStatic(IsOutOfScopeLocalVariable));
 
 	// @TODO: account for all K2ActionMenuBuilder checks...
 	//    NodeCDO->CanCreateUnderSpecifiedSchema(this)
@@ -1083,7 +1201,7 @@ FBlueprintActionFilter::FBlueprintActionFilter(uint32 Flags/*= 0x00*/)
 }
 
 //------------------------------------------------------------------------------
-void FBlueprintActionFilter::AddIsFilteredTest(FIsFilteredDelegate IsFilteredDelegate)
+void FBlueprintActionFilter::AddRejectionTest(FRejectionTestDelegate IsFilteredDelegate)
 {
 	if (IsFilteredDelegate.IsBound())
 	{
@@ -1092,7 +1210,7 @@ void FBlueprintActionFilter::AddIsFilteredTest(FIsFilteredDelegate IsFilteredDel
 }
 
 //------------------------------------------------------------------------------
-bool FBlueprintActionFilter::IsFiltered(UBlueprintNodeSpawner const* BlueprintAction)
+bool FBlueprintActionFilter::IsFiltered(FBlueprintActionInfo& BlueprintAction)
 {
 	bool bIsFiltered = IsFilteredByThis(BlueprintAction);
 	if (!bIsFiltered)
@@ -1133,7 +1251,7 @@ FBlueprintActionFilter const& FBlueprintActionFilter::operator&=(FBlueprintActio
 }
 
 //------------------------------------------------------------------------------
-bool FBlueprintActionFilter::IsFilteredByThis(UBlueprintNodeSpawner const* BlueprintAction) const
+bool FBlueprintActionFilter::IsFilteredByThis(FBlueprintActionInfo& BlueprintAction) const
 {
 	FBlueprintActionFilter const& FilterRef = *this;
 
@@ -1142,10 +1260,10 @@ bool FBlueprintActionFilter::IsFilteredByThis(UBlueprintNodeSpawner const* Bluep
 	// internal tests are ran last).
 	for (int32 TestIndex = FilterTests.Num()-1; TestIndex >= 0; --TestIndex)
 	{
-		FIsFilteredDelegate const& IsFilteredDelegate = FilterTests[TestIndex];
-		check(IsFilteredDelegate.IsBound());
+		FRejectionTestDelegate const& RejectionTestDelegate = FilterTests[TestIndex];
+		check(RejectionTestDelegate.IsBound());
 
-		if (IsFilteredDelegate.Execute(FilterRef, BlueprintAction))
+		if (RejectionTestDelegate.Execute(FilterRef, BlueprintAction))
 		{
 			bIsFiltered = true;
 			break;
diff --git a/Engine/Source/Editor/BlueprintGraph/Public/BlueprintActionFilter.h b/Engine/Source/Editor/BlueprintGraph/Public/BlueprintActionFilter.h
--- a/Engine/Source/Editor/BlueprintGraph/Public/BlueprintActionFilter.h
+++ b/Engine/Source/Editor/BlueprintGraph/Public/BlueprintActionFilter.h
@@ -41,6 +41,94 @@ struct FBlueprintActionContext
 	TArray<UObject*> SelectedObjects;
 };
 
+/*******************************************************************************
+ * FBlueprintActionInfo
+ ******************************************************************************/
+
+/**
+ * Info struct passed around to filter rejection tests. Wraps a 
+ * UBlueprintNodeSpawner, and caches associated fields/files/etc. as they're 
+ * requested (to optimize duplicated queries cross rejection tests).
+ */
+struct BLUEPRINTGRAPH_API FBlueprintActionInfo
+{
+	/** */
+	FBlueprintActionInfo(UObject const* ActionOwner, UBlueprintNodeSpawner const* Action);
+
+	/**
+	 * Retrieves the key that the wrapped action is associated with in the 
+	 * FBlueprintActionDatabase (either a UClass, or asset object).
+	 * 
+	 * @return The class/asset that this action conceptually belongs to.
+	 */
+	UObject const* GetActionOwner();
+
+	/**
+	 * Retrieves a class associated with the wrapped action. Intended to be the 
+	 * action's class "owner". Could be null if the action is keyed to an asset.
+	 * Will not be the spawner's NodeClass (even if it is keyed to it in the 
+	 * database).
+	 * 
+	 * @return The class that this action conceptually belongs to (null if the action is keyed to an asset that we cannot derive a class from).
+	 */
+	UClass const* GetOwnerClass();
+
+	/**
+	 * Retrieves the node class that the wrapped action will spawn (assume to 
+	 * be not null).
+	 * 
+	 * @return The node type that the action will spawn.
+	 */
+	UClass const* GetNodeClass();
+
+	/**
+	 * Certain actions are associated with specific member fields (a member 
+	 * function call, a variable get/set, etc.) This retrieves that member field 
+	 * if there is one (not all actions have an associated field).
+	 * 
+	 * @return The member field associated with the wrapped action (null if there isn't one).
+	 */
+	UField const* GetAssociatedMemberField();
+
+	/**
+	 * Certain actions are associated with specific properties (like delegate  
+	 * node spawners, or variable get/set spawners) This retrieves that property
+	 * from the wrapped action if it can (not all actions have an associated 
+	 * property).
+	 * 
+	 * @return The property associated with the wrapped action (null if there isn't one).
+	 */
+	UProperty const* GetAssociatedProperty();
+
+	/**
+	 * Certain actions are associated with specific functions (like function    
+	 * call spawners, or event spawners) This retrieves the function from the 
+	 * wrapped action if it can (not all actions have an associated function).
+	 * 
+	 * @return The function associated with the wrapped action (null if there isn't one).
+	 */
+	UFunction const* GetAssociatedFunction();
+
+	/** The raw action that this struct represent (const so we don't mutate the database) */
+	UBlueprintNodeSpawner const* const NodeSpawner;
+
+private:
+	/** The class or asset-object that the NodeSpawner action is keyed to (in the action database)*/
+	UObject const* ActionOwner;
+
+	/** Keeps track of the fields we've cached (needed in case one turns out to be null) */
+	uint32 CacheFlags;
+	
+	/** */
+	UClass const*    CachedOwnerClass;
+	/** */
+	UField const*    CachedActionField;
+	/** */
+	UProperty const* CachedActionProperty;
+	/** */
+	UFunction const* CachedActionFunction;
+};
+
 
 /*******************************************************************************
  * FBlueprintActionFilter
@@ -50,7 +138,7 @@ class BLUEPRINTGRAPH_API FBlueprintActionFilter
 {
 public:
 	/** The filter uses a series of rejection tests matching */
-	DECLARE_DELEGATE_RetVal_TwoParams(bool, FIsFilteredDelegate, FBlueprintActionFilter const&, UBlueprintNodeSpawner const*);
+	DECLARE_DELEGATE_RetVal_TwoParams(bool, FRejectionTestDelegate, FBlueprintActionFilter const&, FBlueprintActionInfo&);
 
 public:
 	enum EFlags // Flags, which configure certain rejection tests.
@@ -108,9 +196,9 @@ public:
 	 * method. We use rejection "IsFiltered" tests rather than inclusive tests 
 	 * because it is more optimal to whittle down the list of actions early.
 	 * 
-	 * @param  IsFilteredDelegate	The rejection test you wish to add to this filter.
+	 * @param  RejectionTestDelegate	The rejection test you wish to add to this filter.
 	 */
-	void AddIsFilteredTest(FIsFilteredDelegate IsFilteredDelegate);
+	void AddRejectionTest(FRejectionTestDelegate RejectionTestDelegate);
 
 	/**
 	 * Query to check and see if the specified action gets filtered out by this 
@@ -120,7 +208,7 @@ public:
 	 * @param  BlueprintAction	The node-spawner you wish to test.
 	 * @return False if the action passes the filter, otherwise false (the action got filtered out).
 	 */
-	bool IsFiltered(UBlueprintNodeSpawner const* BlueprintAction);
+	bool IsFiltered(FBlueprintActionInfo& BlueprintAction);
 
 	/**
 	 * Appends another filter to be utilized in IsFiltered() queries, extending  
@@ -154,10 +242,10 @@ private:
 	 * @param  BlueprintAction	The node-spawner you wish to test.
 	 * @return False if the action passes the filter, otherwise false (the action got filtered out).
 	 */
-	bool IsFilteredByThis(UBlueprintNodeSpawner const* BlueprintAction) const;
+	bool IsFilteredByThis(FBlueprintActionInfo& BlueprintAction) const;
 
 	/** Set of rejection tests for this specific filter. */
-	TArray<FIsFilteredDelegate> FilterTests;
+	TArray<FRejectionTestDelegate> FilterTests;
 
 	/** Filters to be logically and'd in with the IsFilteredByThis() result. */
 	TArray<FBlueprintActionFilter> AndFilters;
diff --git a/Engine/Source/Editor/Kismet/Private/BlueprintActionMenuBuilder.cpp b/Engine/Source/Editor/Kismet/Private/BlueprintActionMenuBuilder.cpp
--- a/Engine/Source/Editor/Kismet/Private/BlueprintActionMenuBuilder.cpp
+++ b/Engine/Source/Editor/Kismet/Private/BlueprintActionMenuBuilder.cpp
@@ -158,8 +158,8 @@ TSharedPtr<FEdGraphSchemaAction> FBlueprintActionMenuItemFactory::MakeActionMenu
 	FBlueprintActionMenuItem * NewMenuItem = new FBlueprintActionMenuItem(Action, FEditorStyle::GetBrush(IconBrushName), IconTint, MenuGrouping);
 	NewMenuItem->MenuDescription    = GetMenuNameForAction(EditorContext, Action);
 	NewMenuItem->TooltipDescription = GetTooltipForAction(EditorContext, Action).ToString();
-	NewMenuItem->Category			= GetCategoryForAction(EditorContext, Action).ToString();
-	NewMenuItem->Keywords			= GetSearchKeywordsForAction(EditorContext, Action).ToString();
+	NewMenuItem->Category           = GetCategoryForAction(EditorContext, Action).ToString();
+	NewMenuItem->Keywords           = GetSearchKeywordsForAction(EditorContext, Action).ToString();
 
 	NewMenuItem->Category = FString::Printf(TEXT("%s|%s"), *RootCategory.ToString(), *NewMenuItem->Category);	
 	return MakeShareable(NewMenuItem);
@@ -379,7 +379,7 @@ namespace FBlueprintActionMenuBuilderImpl
 		 * @param  DatabaseAction	The node-spawner that the new menu item should wrap.
 		 * @return An empty TSharedPtr if the action was filtered out, otherwise a newly allocated FBlueprintActionMenuItem.
 		 */
-		TSharedPtr<FEdGraphSchemaAction> MakeMenuItem(TWeakPtr<FBlueprintEditor> EditorContext, UBlueprintNodeSpawner const* const DatabaseAction);
+		TSharedPtr<FEdGraphSchemaAction> MakeMenuItem(TWeakPtr<FBlueprintEditor> EditorContext, FBlueprintActionInfo& DatabaseAction);
 
 		/**
 		 * 
@@ -389,7 +389,7 @@ namespace FBlueprintActionMenuBuilderImpl
 		 * @param  Bindings	
 		 * @return 
 		 */
-		TSharedPtr<FEdGraphSchemaAction> MakeBoundMenuItem(TWeakPtr<FBlueprintEditor> EditorContext, UBlueprintNodeSpawner const* const DatabaseAction, TArray<UObject*> const& Bindings);
+		TSharedPtr<FEdGraphSchemaAction> MakeBoundMenuItem(TWeakPtr<FBlueprintEditor> EditorContext, FBlueprintActionInfo& DatabaseAction, TArray<UObject*> const& Bindings);
 		
 		/**
 		 * Clears out any consolidated properties that this may have been 
@@ -443,11 +443,10 @@ void FBlueprintActionMenuBuilderImpl::FMenuSectionDefinition::SetSectionSortOrde
 }
 // 
 //------------------------------------------------------------------------------
-TSharedPtr<FEdGraphSchemaAction> FBlueprintActionMenuBuilderImpl::FMenuSectionDefinition::MakeBoundMenuItem(TWeakPtr<FBlueprintEditor> EditorContext, UBlueprintNodeSpawner const* const DatabaseAction, TArray<UObject*> const& Bindings)
+TSharedPtr<FEdGraphSchemaAction> FBlueprintActionMenuBuilderImpl::FMenuSectionDefinition::MakeBoundMenuItem(TWeakPtr<FBlueprintEditor> EditorContext, FBlueprintActionInfo& DatabaseAction, TArray<UObject*> const& Bindings)
 {
  	TSharedPtr<FBlueprintBoundMenuItem> MenuItem;
-
-	UBlueprintNodeSpawner* BoundAction = nullptr;
+	
 	// we don't want the blueprint database growing out of control with an entry 
 	// for every object you could ever possibly bind to, so each 
 	// UBlueprintNodeSpawner comes with an interface to bind through... however,
@@ -456,52 +455,63 @@ TSharedPtr<FEdGraphSchemaAction> FBlueprintActionMenuBuilderImpl::FMenuSectionDe
 	for (auto BindingIt = Bindings.CreateConstIterator(); BindingIt;)
 	{
 		UObject const* BindingObj = *BindingIt;
-		++BindingIt;
-
 		// check to see if this object can be bound to this action
-		if (DatabaseAction->CanBind(BindingObj))
+		if (DatabaseAction.NodeSpawner->CanBind(BindingObj))
 		{
 			// we don't want binding to mutate the database, so we clone the 
 			// action and pass that to FBlueprintBoundMenuItem for binding;
 			// if we haven't created a mutable BoundAction yet, we need that...
 			// we do this prior to filtering because there may be tests that 
 			// reject based off bindings
-			if (BoundAction == nullptr)
+			//
+			// don't worry about outer, if this makes it through the filter,
+			// then the menu item will add a reference to keep it around
+			UBlueprintNodeSpawner* BoundAction = DuplicateObject<UBlueprintNodeSpawner>(DatabaseAction.NodeSpawner, GetTransientPackage());
+			BoundAction->CustomizeNodeDelegate = DatabaseAction.NodeSpawner->CustomizeNodeDelegate; // delegate isn't duplicated
+
+			FBlueprintActionInfo BoundActionInfo(DatabaseAction.GetActionOwner(), BoundAction);
+			do
 			{
-				// don't worry about outer, if this makes it through the filter,
-				// then the menu item will add a reference to keep it around
-				BoundAction = DuplicateObject<UBlueprintNodeSpawner>(DatabaseAction, GetTransientPackage());
-				BoundAction->CustomizeNodeDelegate = DatabaseAction->CustomizeNodeDelegate; // delegate isn't duplicated
-			}
-
-			// add bindings before filtering (in case tests accept/reject based off of this)
-			BoundAction->AddBinding(BindingObj);
+				// add bindings before filtering (in case tests accept/reject based off of this)
+				BoundAction->AddBinding(BindingObj);
+				
+				++BindingIt;
+				bool const bIsLastBinding = !BindingIt;
 
-			bool const bIsLastBinding = !BindingIt;
-			// if BoundAction is now "full" (meaning it can take any more 
-			// bindings), or if this is the last binding to test...
-			if (!BoundAction->CanBindMultipleObjects() || bIsLastBinding)
-			{
-				// have to check IsFiltered() for every "fully bound" action (in
-				// case there are tests that reject based off of this), we may 
-				// test this multiple times per action (we have to make sure 
-				// that every set of bound objects pass before folding them into
-				// MenuItem)
-				bool const bPassedFilter = !Filter.IsFiltered(BoundAction);
-				if (bPassedFilter)
+				// if BoundAction is now "full" (meaning it can take any more 
+				// bindings), or if this is the last binding to test...
+				if (!BoundAction->CanBindMultipleObjects() || bIsLastBinding)
 				{
-					if (!MenuItem.IsValid())
+					// have to check IsFiltered() for every "fully bound" action (in
+					// case there are tests that reject based off of this), we may 
+					// test this multiple times per action (we have to make sure 
+					// that every set of bound objects pass before folding them into
+					// MenuItem)
+					bool const bPassedFilter = !Filter.IsFiltered(BoundActionInfo);
+					if (bPassedFilter)
 					{
-						MenuItem = ItemFactory.MakeBoundMenuItem(EditorContext, BoundAction);
+						if (!MenuItem.IsValid())
+						{
+							MenuItem = ItemFactory.MakeBoundMenuItem(EditorContext, BoundAction);
+						}
+
+						for (TWeakObjectPtr<UObject> Object : BoundAction->GetBindings())
+						{
+							MenuItem->AddBinding(Object.Get());
+						}
 					}
+					BoundAction->ClearBindings();
+				}
 
-					for (TWeakObjectPtr<UObject> Object : BoundAction->GetBindings())
-					{
-						MenuItem->AddBinding(Object.Get());
-					}
+				if (!bIsLastBinding)
+				{
+					BindingObj = *BindingIt;
 				}
-				BoundAction->ClearBindings();
-			}
+			} while (BindingIt);
+		}
+		else
+		{
+			++BindingIt;
 		}
 	}
 
@@ -509,9 +519,9 @@ TSharedPtr<FEdGraphSchemaAction> FBlueprintActionMenuBuilderImpl::FMenuSectionDe
 }
 
 //------------------------------------------------------------------------------
-TSharedPtr<FEdGraphSchemaAction> FBlueprintActionMenuBuilderImpl::FMenuSectionDefinition::MakeMenuItem(TWeakPtr<FBlueprintEditor> EditorContext, UBlueprintNodeSpawner const* const DatabaseAction)
+TSharedPtr<FEdGraphSchemaAction> FBlueprintActionMenuBuilderImpl::FMenuSectionDefinition::MakeMenuItem(TWeakPtr<FBlueprintEditor> EditorContext, FBlueprintActionInfo& DatabaseAction)
 {	
-	bool bPassedFilter = (DatabaseAction != nullptr) && !Filter.IsFiltered(DatabaseAction);
+	bool bPassedFilter = !Filter.IsFiltered(DatabaseAction);
 
 	bool bCreateDragDropItem = false;
 	// if the caller wants to consolidate all property actions, then we have to 
@@ -520,13 +530,13 @@ TSharedPtr<FEdGraphSchemaAction> FBlueprintActionMenuBuilderImpl::FMenuSectionDe
 	if (bPassedFilter && (Flags & FBlueprintActionMenuBuilder::ConsolidatePropertyActions))
 	{
 		UProperty const* ActionProperty = nullptr;
-		if (UBlueprintVariableNodeSpawner const* VariableSpawner = Cast<UBlueprintVariableNodeSpawner>(DatabaseAction))
+		if (UBlueprintVariableNodeSpawner const* VariableSpawner = Cast<UBlueprintVariableNodeSpawner>(DatabaseAction.NodeSpawner))
 		{
 			bCreateDragDropItem = true;
 			ActionProperty = VariableSpawner->GetVarProperty();
 			bPassedFilter = (ActionProperty != nullptr);
 		}
-		else if (UBlueprintDelegateNodeSpawner const* DelegateSpawner = Cast<UBlueprintDelegateNodeSpawner>(DatabaseAction))
+		else if (UBlueprintDelegateNodeSpawner const* DelegateSpawner = Cast<UBlueprintDelegateNodeSpawner>(DatabaseAction.NodeSpawner))
 		{
 			bCreateDragDropItem = true;
 			ActionProperty = DelegateSpawner->GetProperty();
@@ -547,11 +557,11 @@ TSharedPtr<FEdGraphSchemaAction> FBlueprintActionMenuBuilderImpl::FMenuSectionDe
 	{
 		if (bCreateDragDropItem)
 		{
-			MenuEntry = ItemFactory.MakeDragDropMenuItem(DatabaseAction);
+			MenuEntry = ItemFactory.MakeDragDropMenuItem(DatabaseAction.NodeSpawner);
 		}
 		else
 		{
-			MenuEntry = ItemFactory.MakeActionMenuItem(EditorContext, DatabaseAction);
+			MenuEntry = ItemFactory.MakeActionMenuItem(EditorContext, DatabaseAction.NodeSpawner);
 		}
 	}
 
@@ -667,13 +677,16 @@ void FBlueprintActionMenuBuilder::RebuildActionList()
 		FBlueprintActionDatabase::FActionRegistry const& ActionDatabase = FBlueprintActionDatabase::Get().GetAllActions();
 		for (auto const& ActionEntry : ActionDatabase)
 		{
-			for (UBlueprintNodeSpawner* NodeSpawner : ActionEntry.Value)
+			for (UBlueprintNodeSpawner const* NodeSpawner : ActionEntry.Value)
 			{
+
+				FBlueprintActionInfo BlueprintAction(ActionEntry.Key, NodeSpawner);
+
 				// @TODO: could probably have a super filter that spreads across 
-				//        all MenuSctions (to pair down on performance)
-				for (TSharedRef<FMenuSectionDefinition>& MenuSection : MenuSections)
+				//        all MenuSctions (to pair down on performance?)
+				for (TSharedRef<FMenuSectionDefinition> const& MenuSection : MenuSections)
 				{
-					TSharedPtr<FEdGraphSchemaAction> MenuEntry = MenuSection->MakeMenuItem(BlueprintEditorPtr, NodeSpawner);
+					TSharedPtr<FEdGraphSchemaAction> MenuEntry = MenuSection->MakeMenuItem(BlueprintEditorPtr, BlueprintAction);
 					if (MenuEntry.IsValid())
 					{
 						AddAction(MenuEntry);
@@ -683,14 +696,14 @@ void FBlueprintActionMenuBuilder::RebuildActionList()
 					// @TODO: this only makes one "BoundMenuEntry" for all 
 					//        bindings, we could have a non-consolidated mode if
 					//        we want (like we have for property actions)
-					TSharedPtr<FEdGraphSchemaAction> BoundMenuEntry = MenuSection->MakeBoundMenuItem(BlueprintEditorPtr, NodeSpawner, BindingObjects);
+					TSharedPtr<FEdGraphSchemaAction> BoundMenuEntry = MenuSection->MakeBoundMenuItem(BlueprintEditorPtr, BlueprintAction, BindingObjects);
 					if (BoundMenuEntry.IsValid())
 					{
 						AddAction(BoundMenuEntry);
 					}
 				}
 			}
-		}
+		}	
 	}
 	else if (MenuSections.Num() > 0)
 	{
