Fixing a crash in the new BP menu system, when you save a new BP right after creating it.

diff --git a/Engine/Source/Editor/BlueprintGraph/Private/BlueprintActionDatabase.cpp b/Engine/Source/Editor/BlueprintGraph/Private/BlueprintActionDatabase.cpp
--- a/Engine/Source/Editor/BlueprintGraph/Private/BlueprintActionDatabase.cpp
+++ b/Engine/Source/Editor/BlueprintGraph/Private/BlueprintActionDatabase.cpp
@@ -640,8 +640,20 @@ static void BlueprintActionDatabaseImpl::OnAssetAdded(FAssetData const& NewAsset
 		UObject* AssetObject = NewAssetInfo.GetAsset();
 		if (UBlueprint* NewBlueprint = Cast<UBlueprint>(AssetObject))
 		{
+			FBlueprintActionDatabase& ActionDatabase = FBlueprintActionDatabase::Get();
+			// this callback can be triggered twice from the asset registry (once 
+			// when it is initially created, and again when it is permanently 
+			// saved)
+			bool const bNewlyCreated = (ActionDatabase.GetAllActions().Find(NewBlueprint) == nullptr);
+
 			OnBlueprintChanged(NewBlueprint);
-			NewBlueprint->OnChanged().AddStatic(&BlueprintActionDatabaseImpl::OnBlueprintChanged);
+			// have to be careful not to register this callback twice for the 
+			// blueprint (since this function can be triggered multiple times 
+			// for the same asset)
+			if (bNewlyCreated)
+			{
+				NewBlueprint->OnChanged().AddStatic(&BlueprintActionDatabaseImpl::OnBlueprintChanged);
+			}			
 		}
 		else 
 		{
