BP - Adding additional keywords to the Select node.

diff --git a/Engine/Source/Editor/BlueprintGraph/Classes/K2Node_Select.h b/Engine/Source/Editor/BlueprintGraph/Classes/K2Node_Select.h
--- a/Engine/Source/Editor/BlueprintGraph/Classes/K2Node_Select.h
+++ b/Engine/Source/Editor/BlueprintGraph/Classes/K2Node_Select.h
@@ -37,6 +37,7 @@ class UK2Node_Select : public UK2Node
 	// Begin UEdGraphNode interface
 	virtual void AllocateDefaultPins() override;
 	virtual FText GetTooltipText() const override;
+	virtual FString GetKeywords() const override;
 	virtual FText GetNodeTitle(ENodeTitleType::Type TitleType) const override;
 	virtual void NodeConnectionListChanged() override;
 	virtual void PinTypeChanged(UEdGraphPin* Pin) override;
@@ -46,7 +47,7 @@ class UK2Node_Select : public UK2Node
 
 	// Begin UK2Node interface
 	virtual void ReallocatePinsDuringReconstruction(TArray<UEdGraphPin*>& OldPins) override;
- 	virtual void PostReconstructNode() override;
+	virtual void PostReconstructNode() override;
 	virtual bool IsNodeSafeToIgnore() const override { return true; }
 	virtual bool IsNodePure() const override { return true; }
 	virtual void NotifyPinConnectionListChanged(UEdGraphPin* Pin) override;
diff --git a/Engine/Source/Editor/BlueprintGraph/Private/K2Node_Select.cpp b/Engine/Source/Editor/BlueprintGraph/Private/K2Node_Select.cpp
--- a/Engine/Source/Editor/BlueprintGraph/Private/K2Node_Select.cpp
+++ b/Engine/Source/Editor/BlueprintGraph/Private/K2Node_Select.cpp
@@ -261,6 +261,11 @@ FText UK2Node_Select::GetTooltipText() const
 	return LOCTEXT("SelectNodeTooltip", "Return the option at Index, (first option is indexed at 0)");
 }
 
+FString UK2Node_Select::GetKeywords() const
+{
+	return TEXT("Ternary If");
+}
+
 FText UK2Node_Select::GetNodeTitle(ENodeTitleType::Type TitleType) const
 {
 	return LOCTEXT("Select", "Select");
