Missed include for non-unity builds.

diff --git a/Engine/Source/Editor/BlueprintGraph/Classes/K2Node_CallFunction.h b/Engine/Source/Editor/BlueprintGraph/Classes/K2Node_CallFunction.h
--- a/Engine/Source/Editor/BlueprintGraph/Classes/K2Node_CallFunction.h
+++ b/Engine/Source/Editor/BlueprintGraph/Classes/K2Node_CallFunction.h
@@ -3,6 +3,7 @@
 #pragma once
 
 #include "K2Node.h"
+#include "EdGraph/EdGraphNodeUtils.h" // for FNodeTextCache
 #include "K2Node_CallFunction.generated.h"
 
 UCLASS()
