BP: when object is passed by string, GetPathName should be used, instead of GetFullName
ttp 344824 LIVE: BP: Checking the contents of an ActorClass array will alway return True even when checking for a class that is not part à?

diff --git a/Engine/Source/Editor/BlueprintGraph/Private/CallFunctionHandler.cpp b/Engine/Source/Editor/BlueprintGraph/Private/CallFunctionHandler.cpp
--- a/Engine/Source/Editor/BlueprintGraph/Private/CallFunctionHandler.cpp
+++ b/Engine/Source/Editor/BlueprintGraph/Private/CallFunctionHandler.cpp
@@ -581,7 +581,7 @@ void FKCHandler_CallFunction::RegisterNets(FKismetFunctionContext& Context, UEdG
 						Term->bIsLocal = true;
 						if (!Pin->PinType.bIsArray)
 						{
-							Term->PropertyDefault = Pin->DefaultObject ? Pin->DefaultObject->GetFullName() : Pin->DefaultValue;
+							Term->PropertyDefault = Pin->DefaultObject ? Pin->DefaultObject->GetPathName() : Pin->DefaultValue;
 						}
 
 						Context.NetMap.Add(Pin, Term);
diff --git a/Engine/Source/Runtime/Engine/Private/EdGraph/EdGraphPin.cpp b/Engine/Source/Runtime/Engine/Private/EdGraph/EdGraphPin.cpp
--- a/Engine/Source/Runtime/Engine/Private/EdGraph/EdGraphPin.cpp
+++ b/Engine/Source/Runtime/Engine/Private/EdGraph/EdGraphPin.cpp
@@ -202,7 +202,7 @@ FString UEdGraphPin::GetDefaultAsString() const
 {
 	if(DefaultObject != NULL)
 	{
-		return DefaultObject->GetFullName();
+		return DefaultObject->GetPathName();
 	}
 	else
 	{
