Optimizing the new BP menu system. Caching menu categories where we can.

diff --git a/Engine/Source/Editor/BlueprintGraph/Classes/BlueprintComponentNodeSpawner.h b/Engine/Source/Editor/BlueprintGraph/Classes/BlueprintComponentNodeSpawner.h
--- a/Engine/Source/Editor/BlueprintGraph/Classes/BlueprintComponentNodeSpawner.h
+++ b/Engine/Source/Editor/BlueprintGraph/Classes/BlueprintComponentNodeSpawner.h
@@ -50,6 +50,7 @@ private:
 	UPROPERTY()
 	TSubclassOf<UActorComponent> ComponentClass;
 
-	/** Constructing FText strings can be costly, so we cache the default menu name */
+	/** Constructing FText strings can be costly, so we cache the default menu name/category */
 	FNodeTextCache CachedMenuName;
+	FNodeTextCache CachedCategory;
 };
diff --git a/Engine/Source/Editor/BlueprintGraph/Classes/BlueprintEventNodeSpawner.h b/Engine/Source/Editor/BlueprintGraph/Classes/BlueprintEventNodeSpawner.h
--- a/Engine/Source/Editor/BlueprintGraph/Classes/BlueprintEventNodeSpawner.h
+++ b/Engine/Source/Editor/BlueprintGraph/Classes/BlueprintEventNodeSpawner.h
@@ -47,6 +47,7 @@ public:
 	// UBlueprintNodeSpawner interface
 	virtual UEdGraphNode* Invoke(UEdGraph* ParentGraph, FVector2D const Location) const override;
 	virtual FText GetDefaultMenuName() const override;
+	virtual FText GetDefaultMenuCategory() const override;
 	virtual FText GetDefaultSearchKeywords() const override;
 	// End UBlueprintNodeSpawner interface
 
diff --git a/Engine/Source/Editor/BlueprintGraph/Classes/BlueprintVariableNodeSpawner.h b/Engine/Source/Editor/BlueprintGraph/Classes/BlueprintVariableNodeSpawner.h
--- a/Engine/Source/Editor/BlueprintGraph/Classes/BlueprintVariableNodeSpawner.h
+++ b/Engine/Source/Editor/BlueprintGraph/Classes/BlueprintVariableNodeSpawner.h
@@ -108,6 +108,7 @@ private:
 	UPROPERTY()
 	FBPVariableDescription    LocalVarDesc;
 
-	/** Constructing FText strings can be costly, so we cache the default menu name */
+	/** Constructing FText strings can be costly, so we cache the default menu name/category */
 	FNodeTextCache CachedMenuName;
+	FNodeTextCache CachedCategory;
 };
\ No newline at end of file
diff --git a/Engine/Source/Editor/BlueprintGraph/Private/BlueprintComponentNodeSpawner.cpp b/Engine/Source/Editor/BlueprintGraph/Private/BlueprintComponentNodeSpawner.cpp
--- a/Engine/Source/Editor/BlueprintGraph/Private/BlueprintComponentNodeSpawner.cpp
+++ b/Engine/Source/Editor/BlueprintGraph/Private/BlueprintComponentNodeSpawner.cpp
@@ -107,22 +107,26 @@ FText UBlueprintComponentNodeSpawner::GetDefaultMenuCategory() const
 {
 	check(ComponentClass != nullptr);
 
-	FText ClassGroup;
-	TArray<FString> ClassGroupNames;
-	ComponentClass->GetClassGroupNames(ClassGroupNames);
-
-	static FText const DefaultClassGroup(LOCTEXT("DefaultClassGroup", "Common"));
-	// 'Common' takes priority over other class groups
-	if (ClassGroupNames.Contains(DefaultClassGroup.ToString()) || (ClassGroupNames.Num() == 0)) 
-	{
-		ClassGroup = DefaultClassGroup;
-	}
-	else
+	if (CachedCategory.IsOutOfDate())
 	{
-		ClassGroup = FText::FromString(ClassGroupNames[0]);
-	}
+		FText ClassGroup;
+		TArray<FString> ClassGroupNames;
+		ComponentClass->GetClassGroupNames(ClassGroupNames);
 
-	return FText::Format(LOCTEXT("ComponentCategory", "Add Component|{0}"), ClassGroup);
+		static FText const DefaultClassGroup(LOCTEXT("DefaultClassGroup", "Common"));
+		// 'Common' takes priority over other class groups
+		if (ClassGroupNames.Contains(DefaultClassGroup.ToString()) || (ClassGroupNames.Num() == 0))
+		{
+			ClassGroup = DefaultClassGroup;
+		}
+		else
+		{
+			ClassGroup = FText::FromString(ClassGroupNames[0]);
+		}
+		// FText::Format() is slow, so we cache this to save on performance
+		CachedCategory = FText::Format(LOCTEXT("ComponentCategory", "Add Component|{0}"), ClassGroup);
+	}
+	return CachedCategory;
 }
 
 //------------------------------------------------------------------------------
diff --git a/Engine/Source/Editor/BlueprintGraph/Private/BlueprintEventNodeSpawner.cpp b/Engine/Source/Editor/BlueprintGraph/Private/BlueprintEventNodeSpawner.cpp
--- a/Engine/Source/Editor/BlueprintGraph/Private/BlueprintEventNodeSpawner.cpp
+++ b/Engine/Source/Editor/BlueprintGraph/Private/BlueprintEventNodeSpawner.cpp
@@ -157,6 +157,15 @@ FText UBlueprintEventNodeSpawner::GetDefaultMenuName() const
 	return CachedMenuName;
 }
 
+//------------------------------------------------------------------------------
+FText UBlueprintEventNodeSpawner::GetDefaultMenuCategory() const
+{
+	// certain events can have specialized categories (like input events), so 
+	// we choose to leave this in the node's hands (returns an empty text string
+	// so that the menu builder polls the node instead).
+	return UBlueprintNodeSpawner::GetDefaultMenuCategory();
+}
+
 //------------------------------------------------------------------------------
 FText UBlueprintEventNodeSpawner::GetDefaultSearchKeywords() const
 {
diff --git a/Engine/Source/Editor/BlueprintGraph/Private/BlueprintVariableNodeSpawner.cpp b/Engine/Source/Editor/BlueprintGraph/Private/BlueprintVariableNodeSpawner.cpp
--- a/Engine/Source/Editor/BlueprintGraph/Private/BlueprintVariableNodeSpawner.cpp
+++ b/Engine/Source/Editor/BlueprintGraph/Private/BlueprintVariableNodeSpawner.cpp
@@ -115,16 +115,20 @@ FText UBlueprintVariableNodeSpawner::GetDefaultMenuName() const
 //------------------------------------------------------------------------------
 FText UBlueprintVariableNodeSpawner::GetDefaultMenuCategory() const
 {
-	FText VarSubCategory;
-	if (IsLocalVariable())
-	{
-		VarSubCategory = FText::FromName(LocalVarDesc.Category);
-	}
-	else if (MemberVariable.IsValid())
+	if (CachedCategory.IsOutOfDate())
 	{
-		VarSubCategory = FText::FromString(FObjectEditorUtils::GetCategory(MemberVariable.Get()));
+		FText VarSubCategory;
+		if (IsLocalVariable())
+		{
+			VarSubCategory = FText::FromName(LocalVarDesc.Category);
+		}
+		else if (MemberVariable.IsValid())
+		{
+			VarSubCategory = FText::FromString(FObjectEditorUtils::GetCategory(MemberVariable.Get()));
+		}
+		CachedCategory = FEditorCategoryUtils::BuildCategoryString(FCommonEditorCategory::Variables, VarSubCategory);
 	}
-	return FEditorCategoryUtils::BuildCategoryString(FCommonEditorCategory::Variables, VarSubCategory);
+	return CachedCategory;
 }
 
 //------------------------------------------------------------------------------
diff --git a/Engine/Source/Editor/BlueprintGraph/Private/K2Node_DynamicCast.cpp b/Engine/Source/Editor/BlueprintGraph/Private/K2Node_DynamicCast.cpp
--- a/Engine/Source/Editor/BlueprintGraph/Private/K2Node_DynamicCast.cpp
+++ b/Engine/Source/Editor/BlueprintGraph/Private/K2Node_DynamicCast.cpp
@@ -166,7 +166,13 @@ bool UK2Node_DynamicCast::HasExternalBlueprintDependencies(TArray<class UStruct*
 
 FText UK2Node_DynamicCast::GetMenuCategory() const
 {
-	return FEditorCategoryUtils::BuildCategoryString(FCommonEditorCategory::Utilities, LOCTEXT("ActionMenuCategory", "Casting"));
+	static FNodeTextCache CachedCategory;
+	if (CachedCategory.IsOutOfDate())
+	{
+		// FText::Format() is slow, so we cache this to save on performance
+		CachedCategory = FEditorCategoryUtils::BuildCategoryString(FCommonEditorCategory::Utilities, LOCTEXT("ActionMenuCategory", "Casting"));
+	}
+	return CachedCategory;
 }
 
 #undef LOCTEXT_NAMESPACE
diff --git a/Engine/Source/Editor/BlueprintGraph/Private/K2Node_EaseFunction.cpp b/Engine/Source/Editor/BlueprintGraph/Private/K2Node_EaseFunction.cpp
--- a/Engine/Source/Editor/BlueprintGraph/Private/K2Node_EaseFunction.cpp
+++ b/Engine/Source/Editor/BlueprintGraph/Private/K2Node_EaseFunction.cpp
@@ -11,6 +11,7 @@
 #include "Kismet2NameValidators.h"
 #include "KismetCompiler.h"
 #include "ScopedTransaction.h"
+#include "EdGraph/EdGraphNodeUtils.h" // for FNodeTextCache
 
 #define LOCTEXT_NAMESPACE "K2Node_EaseFunction"
 
@@ -244,7 +245,13 @@ void UK2Node_EaseFunction::GetMenuActions(FBlueprintActionDatabaseRegistrar& Act
 
 FText UK2Node_EaseFunction::GetMenuCategory() const
 {
-	return FEditorCategoryUtils::BuildCategoryString(FCommonEditorCategory::Input, LOCTEXT("InterpCategory", "Interpolation"));
+	static FNodeTextCache CachedCategory;
+	if (CachedCategory.IsOutOfDate())
+	{
+		// FText::Format() is slow, so we cache this to save on performance
+		CachedCategory = FEditorCategoryUtils::BuildCategoryString(FCommonEditorCategory::Math, LOCTEXT("InterpCategory", "Interpolation"));
+	}
+	return CachedCategory;
 }
 
 void UK2Node_EaseFunction::ChangePinType(UEdGraphPin* Pin)
diff --git a/Engine/Source/Editor/BlueprintGraph/Private/K2Node_Event.cpp b/Engine/Source/Editor/BlueprintGraph/Private/K2Node_Event.cpp
--- a/Engine/Source/Editor/BlueprintGraph/Private/K2Node_Event.cpp
+++ b/Engine/Source/Editor/BlueprintGraph/Private/K2Node_Event.cpp
@@ -639,19 +639,14 @@ void UK2Node_Event::GetNodeAttributes( TArray<TKeyValuePair<FString, FString>>&
 
 FText UK2Node_Event::GetMenuCategory() const
 {
-	FText FuncCategory;
+	FString FuncSubCategory;
 	if (UFunction* Function = FindField<UFunction>(EventSignatureClass, EventSignatureName))
 	{
-		FuncCategory = FText::FromString(UK2Node_CallFunction::GetDefaultCategoryForFunction(Function, TEXT("")));
+		FuncSubCategory = UK2Node_CallFunction::GetDefaultCategoryForFunction(Function, TEXT(""));
 	}
 
-	FText Category = LOCTEXT("AddEventCategory", "Add Event");
-	if (!FuncCategory.IsEmpty())
-	{
-		Category = FText::Format(LOCTEXT("ConcatenatedCategory", "{0}|{1}"), Category, FuncCategory);
-	}
-
-	return Category;
+	FString RootCategory = LOCTEXT("AddEventCategory", "Add Event").ToString();
+	return FText::FromString(RootCategory + TEXT("|") + FuncSubCategory);
 }
 
 bool UK2Node_Event::IsDeprecated() const
diff --git a/Engine/Source/Editor/BlueprintGraph/Private/K2Node_GetInputAxisKeyValue.cpp b/Engine/Source/Editor/BlueprintGraph/Private/K2Node_GetInputAxisKeyValue.cpp
--- a/Engine/Source/Editor/BlueprintGraph/Private/K2Node_GetInputAxisKeyValue.cpp
+++ b/Engine/Source/Editor/BlueprintGraph/Private/K2Node_GetInputAxisKeyValue.cpp
@@ -154,21 +154,40 @@ void UK2Node_GetInputAxisKeyValue::GetMenuActions(FBlueprintActionDatabaseRegist
 
 FText UK2Node_GetInputAxisKeyValue::GetMenuCategory() const
 {
+	enum EAxisKeyCategory
+	{
+		GamepadKeyCategory,
+		MouseButtonCategory,
+		KeyValueCategory,
+		AxisKeyCategory_MAX,
+	};
+	static FNodeTextCache CachedCategories[AxisKeyCategory_MAX];
+
 	FText SubCategory;
+	EAxisKeyCategory CategoryIndex = AxisKeyCategory_MAX;
+
 	if (InputAxisKey.IsGamepadKey())
 	{
 		SubCategory = LOCTEXT("GamepadSubCategory", "Gamepad Values");
+		CategoryIndex = GamepadKeyCategory;
 	}
 	else if (InputAxisKey.IsMouseButton())
 	{
 		SubCategory = LOCTEXT("MouseSubCategory", "Mouse Values");
+		CategoryIndex = MouseButtonCategory;
 	}
 	else
 	{
 		SubCategory = LOCTEXT("KeySubCategory", "Key Values");
+		CategoryIndex = KeyValueCategory;
 	}
-	
-	return FEditorCategoryUtils::BuildCategoryString(FCommonEditorCategory::Input, SubCategory);
+
+	if (CachedCategories[CategoryIndex].IsOutOfDate())
+	{
+		// FText::Format() is slow, so we cache this to save on performance
+		CachedCategories[CategoryIndex] = FEditorCategoryUtils::BuildCategoryString(FCommonEditorCategory::Input, SubCategory);
+	}
+	return CachedCategories[CategoryIndex];
 }
 
 #undef LOCTEXT_NAMESPACE
\ No newline at end of file
diff --git a/Engine/Source/Editor/BlueprintGraph/Private/K2Node_GetInputAxisValue.cpp b/Engine/Source/Editor/BlueprintGraph/Private/K2Node_GetInputAxisValue.cpp
--- a/Engine/Source/Editor/BlueprintGraph/Private/K2Node_GetInputAxisValue.cpp
+++ b/Engine/Source/Editor/BlueprintGraph/Private/K2Node_GetInputAxisValue.cpp
@@ -123,7 +123,13 @@ void UK2Node_GetInputAxisValue::GetMenuActions(FBlueprintActionDatabaseRegistrar
 
 FText UK2Node_GetInputAxisValue::GetMenuCategory() const
 {
-	return FEditorCategoryUtils::BuildCategoryString(FCommonEditorCategory::Input, LOCTEXT("ActionMenuCategory", "Axis Values"));
+	static FNodeTextCache CachedCategory;
+	if (CachedCategory.IsOutOfDate())
+	{
+		// FText::Format() is slow, so we cache this to save on performance
+		CachedCategory = FEditorCategoryUtils::BuildCategoryString(FCommonEditorCategory::Input, LOCTEXT("ActionMenuCategory", "Axis Values"));
+	}
+	return CachedCategory;
 }
 
 #undef LOCTEXT_NAMESPACE
diff --git a/Engine/Source/Editor/BlueprintGraph/Private/K2Node_InputAction.cpp b/Engine/Source/Editor/BlueprintGraph/Private/K2Node_InputAction.cpp
--- a/Engine/Source/Editor/BlueprintGraph/Private/K2Node_InputAction.cpp
+++ b/Engine/Source/Editor/BlueprintGraph/Private/K2Node_InputAction.cpp
@@ -167,7 +167,13 @@ void UK2Node_InputAction::GetMenuActions(FBlueprintActionDatabaseRegistrar& Acti
 
 FText UK2Node_InputAction::GetMenuCategory() const
 {
-	return FEditorCategoryUtils::BuildCategoryString(FCommonEditorCategory::Input, LOCTEXT("ActionMenuCategory", "Action Events"));
+	static FNodeTextCache CachedCategory;
+	if (CachedCategory.IsOutOfDate())
+	{
+		// FText::Format() is slow, so we cache this to save on performance
+		CachedCategory = FEditorCategoryUtils::BuildCategoryString(FCommonEditorCategory::Input, LOCTEXT("ActionMenuCategory", "Action Events"));
+	}
+	return CachedCategory;
 }
 
 #undef LOCTEXT_NAMESPACE
diff --git a/Engine/Source/Editor/BlueprintGraph/Private/K2Node_InputAxisEvent.cpp b/Engine/Source/Editor/BlueprintGraph/Private/K2Node_InputAxisEvent.cpp
--- a/Engine/Source/Editor/BlueprintGraph/Private/K2Node_InputAxisEvent.cpp
+++ b/Engine/Source/Editor/BlueprintGraph/Private/K2Node_InputAxisEvent.cpp
@@ -135,7 +135,13 @@ void UK2Node_InputAxisEvent::GetMenuActions(FBlueprintActionDatabaseRegistrar& A
 
 FText UK2Node_InputAxisEvent::GetMenuCategory() const
 {
-	return FEditorCategoryUtils::BuildCategoryString(FCommonEditorCategory::Input, LOCTEXT("ActionMenuCategory", "Axis Events"));
+	static FNodeTextCache CachedCategory;
+	if (CachedCategory.IsOutOfDate())
+	{
+		// FText::Format() is slow, so we cache this to save on performance
+		CachedCategory = FEditorCategoryUtils::BuildCategoryString(FCommonEditorCategory::Input, LOCTEXT("ActionMenuCategory", "Axis Events"));
+	}
+	return CachedCategory;
 }
 
 #undef LOCTEXT_NAMESPACE
diff --git a/Engine/Source/Editor/BlueprintGraph/Private/K2Node_InputAxisKeyEvent.cpp b/Engine/Source/Editor/BlueprintGraph/Private/K2Node_InputAxisKeyEvent.cpp
--- a/Engine/Source/Editor/BlueprintGraph/Private/K2Node_InputAxisKeyEvent.cpp
+++ b/Engine/Source/Editor/BlueprintGraph/Private/K2Node_InputAxisKeyEvent.cpp
@@ -144,22 +144,40 @@ void UK2Node_InputAxisKeyEvent::GetMenuActions(FBlueprintActionDatabaseRegistrar
 
 FText UK2Node_InputAxisKeyEvent::GetMenuCategory() const
 {
+	enum EAxisKeyCategory
+	{
+		GamepadKeyCategory,
+		MouseButtonCategory,
+		KeyEventCategory,
+		AxisKeyCategory_MAX,
+	};
+	static FNodeTextCache CachedCategories[AxisKeyCategory_MAX];
+
 	FText SubCategory;
+	EAxisKeyCategory CategoryIndex = AxisKeyCategory_MAX;
+
 	if (AxisKey.IsGamepadKey())
 	{
 		SubCategory = LOCTEXT("GamepadCategory", "Gamepad Events");
+		CategoryIndex = GamepadKeyCategory;
 	}
 	else if (AxisKey.IsMouseButton())
 	{
 		SubCategory = LOCTEXT("MouseCategory", "Mouse Events");
+		CategoryIndex = MouseButtonCategory;
 	}
 	else
 	{
 		SubCategory = LOCTEXT("KeyEventsCategory", "Key Events");
+		CategoryIndex = KeyEventCategory;
 	}
 
-	return FEditorCategoryUtils::BuildCategoryString(FCommonEditorCategory::Input, SubCategory);
-
+	if (CachedCategories[CategoryIndex].IsOutOfDate())
+	{
+		// FText::Format() is slow, so we cache this to save on performance
+		CachedCategories[CategoryIndex] = FEditorCategoryUtils::BuildCategoryString(FCommonEditorCategory::Input, SubCategory);
+	}
+	return CachedCategories[CategoryIndex];
 }
 
 #undef LOCTEXT_NAMESPACE
diff --git a/Engine/Source/Editor/BlueprintGraph/Private/K2Node_InputKey.cpp b/Engine/Source/Editor/BlueprintGraph/Private/K2Node_InputKey.cpp
--- a/Engine/Source/Editor/BlueprintGraph/Private/K2Node_InputKey.cpp
+++ b/Engine/Source/Editor/BlueprintGraph/Private/K2Node_InputKey.cpp
@@ -400,22 +400,40 @@ void UK2Node_InputKey::GetMenuActions(FBlueprintActionDatabaseRegistrar& ActionR
 
 FText UK2Node_InputKey::GetMenuCategory() const
 {
+	enum EAxisKeyCategory
+	{
+		GamepadKeyCategory,
+		MouseButtonCategory,
+		KeyEventCategory,
+		AxisKeyCategory_MAX,
+	};
+	static FNodeTextCache CachedCategories[AxisKeyCategory_MAX];
+
 	FText SubCategory;
+	EAxisKeyCategory CategoryIndex = AxisKeyCategory_MAX;
+
 	if (InputKey.IsGamepadKey())
 	{
 		SubCategory = LOCTEXT("GamepadCategory", "Gamepad Events");
+		CategoryIndex = GamepadKeyCategory;
 	}
 	else if (InputKey.IsMouseButton())
 	{
 		SubCategory = LOCTEXT("MouseCategory", "Mouse Events");
+		CategoryIndex = MouseButtonCategory;
 	}
 	else
 	{
 		SubCategory = LOCTEXT("KeyEventsCategory", "Key Events");
+		CategoryIndex = KeyEventCategory;
 	}
 
-	return FEditorCategoryUtils::BuildCategoryString(FCommonEditorCategory::Input, SubCategory);
-
+	if (CachedCategories[CategoryIndex].IsOutOfDate())
+	{
+		// FText::Format() is slow, so we cache this to save on performance
+		CachedCategories[CategoryIndex] = FEditorCategoryUtils::BuildCategoryString(FCommonEditorCategory::Input, SubCategory);
+	}
+	return CachedCategories[CategoryIndex];
 }
 
 #undef LOCTEXT_NAMESPACE
diff --git a/Engine/Source/Editor/BlueprintGraph/Private/K2Node_MakeArray.cpp b/Engine/Source/Editor/BlueprintGraph/Private/K2Node_MakeArray.cpp
--- a/Engine/Source/Editor/BlueprintGraph/Private/K2Node_MakeArray.cpp
+++ b/Engine/Source/Editor/BlueprintGraph/Private/K2Node_MakeArray.cpp
@@ -10,6 +10,7 @@
 #include "BlueprintNodeSpawner.h"
 #include "EditorCategoryUtils.h"
 #include "BlueprintActionDatabaseRegistrar.h"
+#include "EdGraph/EdGraphNodeUtils.h" // for FNodeTextCache
 
 static const FString OutputPinName = FString(TEXT("Array"));
 
@@ -414,7 +415,13 @@ void UK2Node_MakeArray::GetMenuActions(FBlueprintActionDatabaseRegistrar& Action
 
 FText UK2Node_MakeArray::GetMenuCategory() const
 {
-	return FEditorCategoryUtils::BuildCategoryString(FCommonEditorCategory::Utilities, LOCTEXT("ActionMenuCategory", "Array"));
+	static FNodeTextCache CachedCategory;
+	if (CachedCategory.IsOutOfDate())
+	{
+		// FText::Format() is slow, so we cache this to save on performance
+		CachedCategory = FEditorCategoryUtils::BuildCategoryString(FCommonEditorCategory::Utilities, LOCTEXT("ActionMenuCategory", "Array"));
+	}
+	return CachedCategory;
 }
 
 #undef LOCTEXT_NAMESPACE
diff --git a/Engine/Source/Editor/BlueprintGraph/Private/K2Node_Switch.cpp b/Engine/Source/Editor/BlueprintGraph/Private/K2Node_Switch.cpp
--- a/Engine/Source/Editor/BlueprintGraph/Private/K2Node_Switch.cpp
+++ b/Engine/Source/Editor/BlueprintGraph/Private/K2Node_Switch.cpp
@@ -6,6 +6,7 @@
 #include "../../../Runtime/Engine/Classes/Kismet/KismetMathLibrary.h"
 #include "K2Node_SwitchEnum.h"
 #include "EditorCategoryUtils.h"
+#include "EdGraph/EdGraphNodeUtils.h" // for FNodeTextCache
 
 #define LOCTEXT_NAMESPACE "K2Node_Switch"
 
@@ -328,7 +329,13 @@ FNodeHandlingFunctor* UK2Node_Switch::CreateNodeHandler(FKismetCompilerContext&
 
 FText UK2Node_Switch::GetMenuCategory() const
 {
-	return FEditorCategoryUtils::BuildCategoryString(FCommonEditorCategory::FlowControl, LOCTEXT("ActionMenuCategory", "Switch"));
+	static FNodeTextCache CachedCategory;
+	if (CachedCategory.IsOutOfDate())
+	{
+		// FText::Format() is slow, so we cache this to save on performance
+		CachedCategory = FEditorCategoryUtils::BuildCategoryString(FCommonEditorCategory::FlowControl, LOCTEXT("ActionMenuCategory", "Switch"));
+	}
+	return CachedCategory;
 }
 
 #undef LOCTEXT_NAMESPACE
