Cleaning up more deprecation spam for QAGame.

diff --git a/Engine/Source/Editor/BlueprintGraph/Private/K2ActionMenuBuilder.cpp b/Engine/Source/Editor/BlueprintGraph/Private/K2ActionMenuBuilder.cpp
--- a/Engine/Source/Editor/BlueprintGraph/Private/K2ActionMenuBuilder.cpp
+++ b/Engine/Source/Editor/BlueprintGraph/Private/K2ActionMenuBuilder.cpp
@@ -121,7 +121,7 @@ template<class ActionType>
 static void GetRequestedAction(FGraphActionListBuilderBase& ActionMenuBuilder, const FString& Category)
 {
 	UK2Node* NodeTemplate = NewObject<ActionType>(ActionMenuBuilder.OwnerOfTemporaries);
-	TSharedPtr<FEdGraphSchemaAction_K2NewNode> Action = FK2ActionMenuBuilder::AddNewNodeAction(ActionMenuBuilder, Category, NodeTemplate->GetNodeTitle(ENodeTitleType::ListView), NodeTemplate->GetTooltip(), 0, NodeTemplate->GetKeywords());
+	TSharedPtr<FEdGraphSchemaAction_K2NewNode> Action = FK2ActionMenuBuilder::AddNewNodeAction(ActionMenuBuilder, Category, NodeTemplate->GetNodeTitle(ENodeTitleType::ListView), NodeTemplate->GetTooltipText().ToString(), 0, NodeTemplate->GetKeywords());
 	Action->NodeTemplate = NodeTemplate;
 }
 
@@ -133,7 +133,7 @@ static void AddSpawnActorNodeAction(FGraphActionListBuilderBase& ContextMenuBuil
 	// Add the new SpawnActorFromClass node
 	{
 		UK2Node* NodeTemplate = ContextMenuBuilder.CreateTemplateNode<UK2Node_SpawnActorFromClass>();
-		TSharedPtr<FEdGraphSchemaAction_K2NewNode> Action = FK2ActionMenuBuilder::AddNewNodeAction(ContextMenuBuilder, SpawnActorCategory, LOCTEXT("SpawnActorFromClass", "Spawn Actor from Class"), NodeTemplate->GetTooltip(), 0, NodeTemplate->GetKeywords());
+		TSharedPtr<FEdGraphSchemaAction_K2NewNode> Action = FK2ActionMenuBuilder::AddNewNodeAction(ContextMenuBuilder, SpawnActorCategory, LOCTEXT("SpawnActorFromClass", "Spawn Actor from Class"), NodeTemplate->GetTooltipText().ToString(), 0, NodeTemplate->GetKeywords());
 		Action->NodeTemplate = NodeTemplate;
 	}
 }
@@ -146,7 +146,7 @@ static void AddGetDataTableRowNodeAction(FGraphActionListBuilderBase& ContextMen
 	// Add the new GetDataTableRow node
 	{
 		UK2Node* NodeTemplate = ContextMenuBuilder.CreateTemplateNode<UK2Node_GetDataTableRow>();
-		TSharedPtr<FEdGraphSchemaAction_K2NewNode> Action = FK2ActionMenuBuilder::AddNewNodeAction(ContextMenuBuilder, GetDataTableRowCategory, LOCTEXT("GetDataTableRow", "Get Data Table Row"), NodeTemplate->GetTooltip(), 0, NodeTemplate->GetKeywords());
+		TSharedPtr<FEdGraphSchemaAction_K2NewNode> Action = FK2ActionMenuBuilder::AddNewNodeAction(ContextMenuBuilder, GetDataTableRowCategory, LOCTEXT("GetDataTableRow", "Get Data Table Row"), NodeTemplate->GetTooltipText().ToString(), 0, NodeTemplate->GetKeywords());
 		Action->NodeTemplate = NodeTemplate;
 	}
 }
@@ -328,7 +328,7 @@ static void GetEnumUtilitiesNodes(FBlueprintGraphActionListBuilder& ContextMenuB
 					EnumTemplate->Enum = Enum;
 					UK2Node* NodeTemplate = EnumTemplate;
 					TSharedPtr<FEdGraphSchemaAction_K2NewNode> Action = FK2ActionMenuBuilder::AddNewNodeAction(
-						InContextMenuBuilder, Category, NodeTemplate->GetNodeTitle(ENodeTitleType::ListView),  NodeTemplate->GetTooltip(), 0, NodeTemplate->GetKeywords());
+						InContextMenuBuilder, Category, NodeTemplate->GetNodeTitle(ENodeTitleType::ListView), NodeTemplate->GetTooltipText().ToString(), 0, NodeTemplate->GetKeywords());
 					Action->NodeTemplate = NodeTemplate;
 				}
 
@@ -338,7 +338,7 @@ static void GetEnumUtilitiesNodes(FBlueprintGraphActionListBuilder& ContextMenuB
 					EnumTemplate->Enum = Enum;
 					UK2Node* NodeTemplate = EnumTemplate;
 					TSharedPtr<FEdGraphSchemaAction_K2NewNode> Action = FK2ActionMenuBuilder::AddNewNodeAction(
-						InContextMenuBuilder, Category, NodeTemplate->GetNodeTitle(ENodeTitleType::ListView),  NodeTemplate->GetTooltip(), 0, NodeTemplate->GetKeywords());
+						InContextMenuBuilder, Category, NodeTemplate->GetNodeTitle(ENodeTitleType::ListView), NodeTemplate->GetTooltipText().ToString(), 0, NodeTemplate->GetKeywords());
 					Action->NodeTemplate = NodeTemplate;
 				}
 
@@ -349,7 +349,7 @@ static void GetEnumUtilitiesNodes(FBlueprintGraphActionListBuilder& ContextMenuB
 					EnumTemplate->bSafe = true;
 					UK2Node* NodeTemplate = EnumTemplate;
 					TSharedPtr<FEdGraphSchemaAction_K2NewNode> Action = FK2ActionMenuBuilder::AddNewNodeAction(
-						InContextMenuBuilder, Category, NodeTemplate->GetNodeTitle(ENodeTitleType::ListView),  NodeTemplate->GetTooltip(), 0, NodeTemplate->GetKeywords());
+						InContextMenuBuilder, Category, NodeTemplate->GetNodeTitle(ENodeTitleType::ListView), NodeTemplate->GetTooltipText().ToString(), 0, NodeTemplate->GetKeywords());
 					Action->NodeTemplate = NodeTemplate;
 				}
 
@@ -359,7 +359,7 @@ static void GetEnumUtilitiesNodes(FBlueprintGraphActionListBuilder& ContextMenuB
 					EnumTemplate->Enum = Enum;
 					UK2Node* NodeTemplate = EnumTemplate;
 					TSharedPtr<FEdGraphSchemaAction_K2NewNode> Action = FK2ActionMenuBuilder::AddNewNodeAction(
-						InContextMenuBuilder, Category, NodeTemplate->GetNodeTitle(ENodeTitleType::ListView), NodeTemplate->GetTooltip(), 0, NodeTemplate->GetKeywords());
+						InContextMenuBuilder, Category, NodeTemplate->GetNodeTitle(ENodeTitleType::ListView), NodeTemplate->GetTooltipText().ToString(), 0, NodeTemplate->GetKeywords());
 					Action->NodeTemplate = NodeTemplate;
 				}
 			}
@@ -395,7 +395,7 @@ static void GetInterfaceMessages(UClass* CurrentInterface, const FString& Messag
 			{
 				UK2Node_Message* MessageTemplate = ContextMenuBuilder.CreateTemplateNode<UK2Node_Message>();
 				MessageTemplate->FunctionReference.SetExternalMember(Function->GetFName(), CurrentInterface);
-				TSharedPtr<FEdGraphSchemaAction_K2NewNode> Action = FK2ActionMenuBuilder::AddNewNodeAction(ContextMenuBuilder, MessageSubCategory, FText::FromString(Function->GetName()), MessageTemplate->GetTooltip(), 0, MessageTemplate->GetKeywords());
+				TSharedPtr<FEdGraphSchemaAction_K2NewNode> Action = FK2ActionMenuBuilder::AddNewNodeAction(ContextMenuBuilder, MessageSubCategory, FText::FromString(Function->GetName()), MessageTemplate->GetTooltipText().ToString(), 0, MessageTemplate->GetKeywords());
 				Action->NodeTemplate = MessageTemplate;
 			}
 		}
@@ -408,7 +408,7 @@ static void GetReferenceSetterItems(FBlueprintGraphActionListBuilder& ContextMen
 	const UEdGraphPin* AssociatedPin = ContextMenuBuilder.FromPin;
 
 	UK2Node_VariableSetRef* SetRefTemplate = ContextMenuBuilder.CreateTemplateNode<UK2Node_VariableSetRef>();
-	TSharedPtr<FEdGraphSchemaAction_K2NewNode> Action = FK2ActionMenuBuilder::AddNewNodeAction(ContextMenuBuilder, K2ActionCategories::VariablesCategory, SetRefTemplate->GetNodeTitle(ENodeTitleType::ListView), SetRefTemplate->GetTooltip(), 0, SetRefTemplate->GetKeywords());
+	TSharedPtr<FEdGraphSchemaAction_K2NewNode> Action = FK2ActionMenuBuilder::AddNewNodeAction(ContextMenuBuilder, K2ActionCategories::VariablesCategory, SetRefTemplate->GetNodeTitle(ENodeTitleType::ListView), SetRefTemplate->GetTooltipText().ToString(), 0, SetRefTemplate->GetKeywords());
 
 	Action->NodeTemplate = SetRefTemplate;
 }
@@ -574,7 +574,7 @@ void FK2ActionMenuBuilder::AddEventForFunction(FGraphActionListBuilderBase& Acti
 
 		const FString Category = UK2Node_CallFunction::GetDefaultCategoryForFunction(Function, K2ActionCategories::AddEventCategory);
 
-		TSharedPtr<FEdGraphSchemaAction_K2AddEvent> Action(new FEdGraphSchemaAction_K2AddEvent(Category, FText::FromString(EventName), EventNodeTemplate->GetTooltip(), 0));
+		TSharedPtr<FEdGraphSchemaAction_K2AddEvent> Action(new FEdGraphSchemaAction_K2AddEvent(Category, FText::FromString(EventName), EventNodeTemplate->GetTooltipText().ToString(), 0));
 		Action->Keywords     = EventNodeTemplate->GetKeywords();
 		Action->NodeTemplate = EventNodeTemplate;
 
@@ -592,7 +592,7 @@ void FK2ActionMenuBuilder::AttachMacroGraphAction(FGraphActionListBuilderBase& A
 
 		UK2Node_MacroInstance* MacroTemplate = ActionMenuBuilder.CreateTemplateNode<UK2Node_MacroInstance>();
 		MacroTemplate->SetMacroGraph(MacroGraph);
-		TSharedPtr<FEdGraphSchemaAction_K2NewNode> Action = FK2ActionMenuBuilder::AddNewNodeAction(ActionMenuBuilder, CategoryName, MacroTemplate->GetNodeTitle(ENodeTitleType::ListView), MacroTemplate->GetTooltip());
+		TSharedPtr<FEdGraphSchemaAction_K2NewNode> Action = FK2ActionMenuBuilder::AddNewNodeAction(ActionMenuBuilder, CategoryName, MacroTemplate->GetNodeTitle(ENodeTitleType::ListView), MacroTemplate->GetTooltipText().ToString());
 		Action->NodeTemplate = MacroTemplate;
 	}
 }
@@ -962,7 +962,7 @@ void FK2ActionMenuBuilder::GetContextAllowedNodeTypes(FBlueprintGraphActionListB
 			// Assignment statement
 			{
 				UK2Node* NodeTemplate = ContextMenuBuilder.CreateTemplateNode<UK2Node_AssignmentStatement>();
-				TSharedPtr<FEdGraphSchemaAction_K2NewNode> Action = FK2ActionMenuBuilder::AddNewNodeAction(ContextMenuBuilder, K2ActionCategories::MacroToolsCategory, NodeTemplate->GetNodeTitle(ENodeTitleType::ListView), NodeTemplate->GetTooltip(), 0, NodeTemplate->GetKeywords());
+				TSharedPtr<FEdGraphSchemaAction_K2NewNode> Action = FK2ActionMenuBuilder::AddNewNodeAction(ContextMenuBuilder, K2ActionCategories::MacroToolsCategory, NodeTemplate->GetNodeTitle(ENodeTitleType::ListView), NodeTemplate->GetTooltipText().ToString(), 0, NodeTemplate->GetKeywords());
 				Action->NodeTemplate = NodeTemplate;
 			}
 
@@ -1020,7 +1020,7 @@ void FK2ActionMenuBuilder::GetContextAllowedNodeTypes(FBlueprintGraphActionListB
 				{
 					UK2Node_TemporaryVariable* NodeTemplate = ContextMenuBuilder.CreateTemplateNode<UK2Node_TemporaryVariable>();
 					NodeTemplate->VariableType = *TypeIt;
-					TSharedPtr<FEdGraphSchemaAction_K2NewNode> Action = FK2ActionMenuBuilder::AddNewNodeAction(ContextMenuBuilder, K2ActionCategories::MacroToolsCategory, NodeTemplate->GetNodeTitle(ENodeTitleType::ListView), NodeTemplate->GetTooltip(), 0, NodeTemplate->GetKeywords());
+					TSharedPtr<FEdGraphSchemaAction_K2NewNode> Action = FK2ActionMenuBuilder::AddNewNodeAction(ContextMenuBuilder, K2ActionCategories::MacroToolsCategory, NodeTemplate->GetNodeTitle(ENodeTitleType::ListView), NodeTemplate->GetTooltipText().ToString(), 0, NodeTemplate->GetKeywords());
 					Action->NodeTemplate = NodeTemplate;
 				}
 			}
@@ -1032,14 +1032,14 @@ void FK2ActionMenuBuilder::GetContextAllowedNodeTypes(FBlueprintGraphActionListB
 				UK2Node_TemporaryVariable* PersistentIntNodeTemplate = ContextMenuBuilder.CreateTemplateNode<UK2Node_TemporaryVariable>();
 				PersistentIntNodeTemplate->VariableType = FEdGraphPinType(K2Schema->PC_Int, TEXT(""), NULL, false, false);
 				PersistentIntNodeTemplate->bIsPersistent = true;
-				TSharedPtr<FEdGraphSchemaAction_K2NewNode> IntAction = FK2ActionMenuBuilder::AddNewNodeAction(ContextMenuBuilder, K2ActionCategories::MacroToolsCategory, PersistentIntNodeTemplate->GetNodeTitle(ENodeTitleType::ListView), PersistentIntNodeTemplate->GetTooltip(), 0, PersistentIntNodeTemplate->GetKeywords());
+				TSharedPtr<FEdGraphSchemaAction_K2NewNode> IntAction = FK2ActionMenuBuilder::AddNewNodeAction(ContextMenuBuilder, K2ActionCategories::MacroToolsCategory, PersistentIntNodeTemplate->GetNodeTitle(ENodeTitleType::ListView), PersistentIntNodeTemplate->GetTooltipText().ToString(), 0, PersistentIntNodeTemplate->GetKeywords());
 				IntAction->NodeTemplate = PersistentIntNodeTemplate;
 
 				// Persistent Bool
 				UK2Node_TemporaryVariable* PersistentBoolNodeTemplate = ContextMenuBuilder.CreateTemplateNode<UK2Node_TemporaryVariable>();
 				PersistentBoolNodeTemplate->VariableType = FEdGraphPinType(K2Schema->PC_Boolean, TEXT(""), NULL, false, false);
 				PersistentBoolNodeTemplate->bIsPersistent = true;
-				TSharedPtr<FEdGraphSchemaAction_K2NewNode> BoolAction = FK2ActionMenuBuilder::AddNewNodeAction(ContextMenuBuilder, K2ActionCategories::MacroToolsCategory, PersistentBoolNodeTemplate->GetNodeTitle(ENodeTitleType::ListView), PersistentBoolNodeTemplate->GetTooltip(), 0, PersistentBoolNodeTemplate->GetKeywords());
+				TSharedPtr<FEdGraphSchemaAction_K2NewNode> BoolAction = FK2ActionMenuBuilder::AddNewNodeAction(ContextMenuBuilder, K2ActionCategories::MacroToolsCategory, PersistentBoolNodeTemplate->GetNodeTitle(ENodeTitleType::ListView), PersistentBoolNodeTemplate->GetTooltipText().ToString(), 0, PersistentBoolNodeTemplate->GetKeywords());
 				BoolAction->NodeTemplate = PersistentBoolNodeTemplate;
 			}
 		}
@@ -1062,7 +1062,7 @@ void FK2ActionMenuBuilder::GetContextAllowedNodeTypes(FBlueprintGraphActionListB
 
 		{
 			UK2Node* NodeTemplate = ContextMenuBuilder.CreateTemplateNode<UK2Node_Select>();
-			TSharedPtr<FEdGraphSchemaAction_K2NewNode> SelectionAction = FK2ActionMenuBuilder::AddNewNodeAction(ContextMenuBuilder, K2ActionCategories::UtilitiesCategory, NodeTemplate->GetNodeTitle(ENodeTitleType::ListView), NodeTemplate->GetTooltip(), 0, NodeTemplate->GetKeywords());
+			TSharedPtr<FEdGraphSchemaAction_K2NewNode> SelectionAction = FK2ActionMenuBuilder::AddNewNodeAction(ContextMenuBuilder, K2ActionCategories::UtilitiesCategory, NodeTemplate->GetNodeTitle(ENodeTitleType::ListView), NodeTemplate->GetTooltipText().ToString(), 0, NodeTemplate->GetKeywords());
 			SelectionAction->NodeTemplate = NodeTemplate;
 		}
 
@@ -1346,14 +1346,14 @@ void FK2ActionMenuBuilder::GetPinAllowedNodeTypes(FBlueprintGraphActionListBuild
 				{
 					UK2Node_BreakStruct* NodeTemplate = ContextMenuBuilder.CreateTemplateNode<UK2Node_BreakStruct>();
 					NodeTemplate->StructType = PinStruct;
-					TSharedPtr<FEdGraphSchemaAction_K2NewNode> Action = FK2ActionMenuBuilder::AddNewNodeAction(ContextMenuBuilder, FString(), NodeTemplate->GetNodeTitle(ENodeTitleType::ListView), NodeTemplate->GetTooltip(), 0, NodeTemplate->GetKeywords());
+					TSharedPtr<FEdGraphSchemaAction_K2NewNode> Action = FK2ActionMenuBuilder::AddNewNodeAction(ContextMenuBuilder, FString(), NodeTemplate->GetNodeTitle(ENodeTitleType::ListView), NodeTemplate->GetTooltipText().ToString(), 0, NodeTemplate->GetKeywords());
 					Action->NodeTemplate = NodeTemplate;
 				}
 				else if ((FromPin.Direction == EGPD_Input) && bCanBeMade)
 				{
 					UK2Node_MakeStruct* NodeTemplate = ContextMenuBuilder.CreateTemplateNode<UK2Node_MakeStruct>();
 					NodeTemplate->StructType = PinStruct;
-					TSharedPtr<FEdGraphSchemaAction_K2NewNode> Action = FK2ActionMenuBuilder::AddNewNodeAction(ContextMenuBuilder, FString(), NodeTemplate->GetNodeTitle(ENodeTitleType::ListView), NodeTemplate->GetTooltip(), 0, NodeTemplate->GetKeywords());
+					TSharedPtr<FEdGraphSchemaAction_K2NewNode> Action = FK2ActionMenuBuilder::AddNewNodeAction(ContextMenuBuilder, FString(), NodeTemplate->GetNodeTitle(ENodeTitleType::ListView), NodeTemplate->GetTooltipText().ToString(), 0, NodeTemplate->GetKeywords());
 					Action->NodeTemplate = NodeTemplate;
 				}
 
@@ -1361,7 +1361,7 @@ void FK2ActionMenuBuilder::GetPinAllowedNodeTypes(FBlueprintGraphActionListBuild
 				{
 					UK2Node_SetFieldsInStruct* NodeTemplate = ContextMenuBuilder.CreateTemplateNode<UK2Node_SetFieldsInStruct>();
 					NodeTemplate->StructType = PinStruct;
-					TSharedPtr<FEdGraphSchemaAction_K2NewNode> Action = FK2ActionMenuBuilder::AddNewNodeAction(ContextMenuBuilder, FString(), NodeTemplate->GetNodeTitle(ENodeTitleType::ListView), NodeTemplate->GetTooltip(), 0, NodeTemplate->GetKeywords());
+					TSharedPtr<FEdGraphSchemaAction_K2NewNode> Action = FK2ActionMenuBuilder::AddNewNodeAction(ContextMenuBuilder, FString(), NodeTemplate->GetNodeTitle(ENodeTitleType::ListView), NodeTemplate->GetTooltipText().ToString(), 0, NodeTemplate->GetKeywords());
 					Action->NodeTemplate = NodeTemplate;
 				}
 			}
@@ -1369,7 +1369,7 @@ void FK2ActionMenuBuilder::GetPinAllowedNodeTypes(FBlueprintGraphActionListBuild
 			if ((FromPin.Direction == EGPD_Input) && (FromPin.PinType.PinCategory == K2Schema->PC_Delegate))
 			{
 				UK2Node_CreateDelegate* NodeTemplate = ContextMenuBuilder.CreateTemplateNode<UK2Node_CreateDelegate>();
-				TSharedPtr<FEdGraphSchemaAction_K2NewNode> Action = FK2ActionMenuBuilder::AddNewNodeAction(ContextMenuBuilder, FString(), NodeTemplate->GetNodeTitle(ENodeTitleType::ListView), NodeTemplate->GetTooltip(), 0, NodeTemplate->GetKeywords());
+				TSharedPtr<FEdGraphSchemaAction_K2NewNode> Action = FK2ActionMenuBuilder::AddNewNodeAction(ContextMenuBuilder, FString(), NodeTemplate->GetNodeTitle(ENodeTitleType::ListView), NodeTemplate->GetTooltipText().ToString(), 0, NodeTemplate->GetKeywords());
 				Action->NodeTemplate = NodeTemplate;
 
 				UFunction* Signature = FMemberReference::ResolveSimpleMemberReference<UFunction>(FromPin.PinType.PinSubCategoryMemberReference);
@@ -1623,7 +1623,7 @@ void FK2ActionMenuBuilder::GetSwitchMenuItems(FBlueprintPaletteListBuilder& Acti
 				EnumTemplate->SetEnum(CurrentEnum);
 
 				UK2Node* NodeTemplate = EnumTemplate;
-				TSharedPtr<FEdGraphSchemaAction_K2NewNode> Action = FK2ActionMenuBuilder::AddNewNodeAction(ActionMenuBuilder, SwitchControlCategory, NodeTemplate->GetNodeTitle(ENodeTitleType::ListView), NodeTemplate->GetTooltip(), 0, NodeTemplate->GetKeywords());
+				TSharedPtr<FEdGraphSchemaAction_K2NewNode> Action = FK2ActionMenuBuilder::AddNewNodeAction(ActionMenuBuilder, SwitchControlCategory, NodeTemplate->GetNodeTitle(ENodeTitleType::ListView), NodeTemplate->GetTooltipText().ToString(), 0, NodeTemplate->GetKeywords());
 				Action->NodeTemplate = NodeTemplate;
 			}
 		}
@@ -1970,7 +1970,7 @@ void FK2ActionMenuBuilder::GetStructActions(FBlueprintGraphActionListBuilder& Co
 				if( BreakNodeTemplate != NULL )
 				{
 					BreakNodeTemplate->StructType = Struct;
-					TSharedPtr<FEdGraphSchemaAction_K2NewNode> Action = FK2ActionMenuBuilder::AddNewNodeAction(ContextMenuBuilder, K2ActionCategories::BreakStructCategory, BreakNodeTemplate->GetNodeTitle(ENodeTitleType::ListView), BreakNodeTemplate->GetTooltip(), 0, BreakNodeTemplate->GetKeywords());
+					TSharedPtr<FEdGraphSchemaAction_K2NewNode> Action = FK2ActionMenuBuilder::AddNewNodeAction(ContextMenuBuilder, K2ActionCategories::BreakStructCategory, BreakNodeTemplate->GetNodeTitle(ENodeTitleType::ListView), BreakNodeTemplate->GetTooltipText().ToString(), 0, BreakNodeTemplate->GetKeywords());
 					Action->NodeTemplate = BreakNodeTemplate;
 				}
 			}
@@ -1982,7 +1982,7 @@ void FK2ActionMenuBuilder::GetStructActions(FBlueprintGraphActionListBuilder& Co
 					if (MakeNodeTemplate != NULL)
 					{
 						MakeNodeTemplate->StructType = Struct;
-						TSharedPtr<FEdGraphSchemaAction_K2NewNode> Action = FK2ActionMenuBuilder::AddNewNodeAction(ContextMenuBuilder, K2ActionCategories::MakeStructCategory, MakeNodeTemplate->GetNodeTitle(ENodeTitleType::ListView), MakeNodeTemplate->GetTooltip(), 0, MakeNodeTemplate->GetKeywords());
+						TSharedPtr<FEdGraphSchemaAction_K2NewNode> Action = FK2ActionMenuBuilder::AddNewNodeAction(ContextMenuBuilder, K2ActionCategories::MakeStructCategory, MakeNodeTemplate->GetNodeTitle(ENodeTitleType::ListView), MakeNodeTemplate->GetTooltipText().ToString(), 0, MakeNodeTemplate->GetKeywords());
 						Action->NodeTemplate = MakeNodeTemplate;
 					}
 				}
@@ -1992,7 +1992,7 @@ void FK2ActionMenuBuilder::GetStructActions(FBlueprintGraphActionListBuilder& Co
 					if (SetFieldsTemplate != NULL)
 					{
 						SetFieldsTemplate->StructType = Struct;
-						TSharedPtr<FEdGraphSchemaAction_K2NewNode> Action = FK2ActionMenuBuilder::AddNewNodeAction(ContextMenuBuilder, K2ActionCategories::MakeStructCategory, SetFieldsTemplate->GetNodeTitle(ENodeTitleType::ListView), SetFieldsTemplate->GetTooltip(), 0, SetFieldsTemplate->GetKeywords());
+						TSharedPtr<FEdGraphSchemaAction_K2NewNode> Action = FK2ActionMenuBuilder::AddNewNodeAction(ContextMenuBuilder, K2ActionCategories::MakeStructCategory, SetFieldsTemplate->GetNodeTitle(ENodeTitleType::ListView), SetFieldsTemplate->GetTooltipText().ToString(), 0, SetFieldsTemplate->GetKeywords());
 						Action->NodeTemplate = SetFieldsTemplate;
 					}
 				}
@@ -2603,7 +2603,7 @@ void FK2ActionMenuBuilder::GetAllInterfaceMessageActions(FGraphActionListBuilder
 
 					UK2Node_Message* MessageTemplate = ActionMenuBuilder.CreateTemplateNode<UK2Node_Message>();
 					MessageTemplate->FunctionReference.SetExternalMember(Function->GetFName(), CurrentInterface);
-					TSharedPtr<FEdGraphSchemaAction_K2NewNode> Action = FK2ActionMenuBuilder::AddNewNodeAction(ActionMenuBuilder, MessageCategory, FText::FromString(Function->GetName()), MessageTemplate->GetTooltip(), 0, MessageTemplate->GetKeywords());
+					TSharedPtr<FEdGraphSchemaAction_K2NewNode> Action = FK2ActionMenuBuilder::AddNewNodeAction(ActionMenuBuilder, MessageCategory, FText::FromString(Function->GetName()), MessageTemplate->GetTooltipText().ToString(), 0, MessageTemplate->GetKeywords());
 					Action->NodeTemplate = MessageTemplate;
 				}
 			}
