Push EdGraphSchema constant variables out of UPROPERTY in to proper static variables
Remove unused PC_Meta

diff --git a/Engine/Source/Editor/BlueprintGraph/Classes/EdGraphSchema_K2.h b/Engine/Source/Editor/BlueprintGraph/Classes/EdGraphSchema_K2.h
--- a/Engine/Source/Editor/BlueprintGraph/Classes/EdGraphSchema_K2.h
+++ b/Engine/Source/Editor/BlueprintGraph/Classes/EdGraphSchema_K2.h
@@ -216,126 +216,51 @@ class BLUEPRINTGRAPH_API UEdGraphSchema_K2 : public UEdGraphSchema
 	GENERATED_UCLASS_BODY()
 
 	// Allowable PinType.PinCategory values
-	UPROPERTY()
-	FString PC_Exec;
-
-	UPROPERTY()
-	FString PC_Meta;
-
-	// PC_Array - not implemented yet
-	UPROPERTY()
-	FString PC_Boolean;
-
-	UPROPERTY()
-	FString PC_Byte;
-
-	UPROPERTY()
-	FString PC_Class;    // SubCategoryObject is the MetaClass of the Class passed thru this pin, or SubCategory can be 'self'. The DefaultValue string should always be empty, use DefaultObject.
-
-	UPROPERTY()
-	FString PC_Int;
-
-	UPROPERTY()
-	FString PC_Float;
-
-	UPROPERTY()
-	FString PC_Name;
-
-	UPROPERTY()
-	FString PC_Delegate;    // SubCategoryObject is the UFunction of the delegate signature
-
-	UPROPERTY()
-	FString PC_MCDelegate;  // SubCategoryObject is the UFunction of the delegate signature
-
-	UPROPERTY()
-	FString PC_Object;    // SubCategoryObject is the Class of the object passed thru this pin, or SubCategory can be 'self'. The DefaultValue string should always be empty, use DefaultObject.
-
-	UPROPERTY()
-	FString PC_Interface;	// SubCategoryObject is the Class of the object passed thru this pin.
-
-	UPROPERTY()
-	FString PC_String;
-
-	UPROPERTY()
-	FString PC_Text;
-
-	UPROPERTY()
-	FString PC_Struct;    // SubCategoryObject is the ScriptStruct of the struct passed thru this pin, 'self' is not a valid SubCategory. DefaultObject should always be empty, the DefaultValue string may be used for supported structs.
-
-	UPROPERTY()
-	FString PC_Wildcard;    // Special matching rules are imposed by the node itself
+	static const FString PC_Exec;
+	static const FString PC_Boolean;
+	static const FString PC_Byte;
+	static const FString PC_Class;    // SubCategoryObject is the MetaClass of the Class passed thru this pin, or SubCategory can be 'self'. The DefaultValue string should always be empty, use DefaultObject.
+	static const FString PC_Int;
+	static const FString PC_Float;
+	static const FString PC_Name;
+	static const FString PC_Delegate;    // SubCategoryObject is the UFunction of the delegate signature
+	static const FString PC_MCDelegate;  // SubCategoryObject is the UFunction of the delegate signature
+	static const FString PC_Object;    // SubCategoryObject is the Class of the object passed thru this pin, or SubCategory can be 'self'. The DefaultValue string should always be empty, use DefaultObject.
+	static const FString PC_Interface;	// SubCategoryObject is the Class of the object passed thru this pin.
+	static const FString PC_String;
+	static const FString PC_Text;
+	static const FString PC_Struct;    // SubCategoryObject is the ScriptStruct of the struct passed thru this pin, 'self' is not a valid SubCategory. DefaultObject should always be empty, the DefaultValue string may be used for supported structs.
+	static const FString PC_Wildcard;    // Special matching rules are imposed by the node itself
 
 	// Common PinType.PinSubCategory values
-	UPROPERTY()
-	FString PSC_Self;    // Category=PC_Object or PC_Class, indicates the class being compiled
+	static const FString PSC_Self;    // Category=PC_Object or PC_Class, indicates the class being compiled
 
-	UPROPERTY()
-	FString PSC_Index;	// Category=PC_Wildcard, indicates the wildcard will only accept Int, Bool, Byte and Enum pins (used when a pin represents indexing a list)
+	static const FString PSC_Index;	// Category=PC_Wildcard, indicates the wildcard will only accept Int, Bool, Byte and Enum pins (used when a pin represents indexing a list)
 
 	// Pin names that have special meaning and required types in some contexts (depending on the node type)
-	UPROPERTY()
-	FString PN_Execute;    // Category=PC_Exec, singleton, input
-
-	UPROPERTY()
-	FString PN_Then;    // Category=PC_Exec, singleton, output
-
-	UPROPERTY()
-	FString PN_Completed;    // Category=PC_Exec, singleton, output
-
-	UPROPERTY()
-	FString PN_DelegateEntry;    // Category=PC_Exec, singleton, output; entry point for a dynamically bound delegate
-
-	UPROPERTY()
-	FString PN_EntryPoint;	// entry point to the ubergraph
-
-	UPROPERTY()
-	FString PN_Self;    // Category=PC_Object, singleton, input
-
-	UPROPERTY()
-	FString PN_Else;    // Category=PC_Exec, singleton, output
-
-	UPROPERTY()
-	FString PN_Loop;    // Category=PC_Exec, singleton, output
-
-	UPROPERTY()
-	FString PN_After;    // Category=PC_Exec, singleton, output
-
-	UPROPERTY()
-	FString PN_ReturnValue;		// Category=PC_Object, singleton, output
-
-	UPROPERTY()
-	FString PN_ObjectToCast;    // Category=PC_Object, singleton, input
-
-	UPROPERTY()
-	FString PN_Condition;    // Category=PC_Boolean, singleton, input
-
-	UPROPERTY()
-	FString PN_Start;    // Category=PC_Int, singleton, input
-
-	UPROPERTY()
-	FString PN_Stop;    // Category=PC_Int, singleton, input
-
-	UPROPERTY()
-	FString PN_Index;    // Category=PC_Int, singleton, output
-
-	UPROPERTY()
-	FString PN_CastSucceeded;    // Category=PC_Exec, singleton, output
-
-	UPROPERTY()
-	FString PN_CastFailed;    // Category=PC_Exec, singleton, output
-
-	UPROPERTY()
-	FString PN_CastedValuePrefix;    // Category=PC_Object, singleton, output; actual pin name varies depending on the type to be casted to, this is just a prefix
-
-	UPROPERTY()
-	FString PN_MatineeFinished;    // Category=PC_Exec, singleton, output
+	static const FString PN_Execute;    // Category=PC_Exec, singleton, input
+	static const FString PN_Then;    // Category=PC_Exec, singleton, output
+	static const FString PN_Completed;    // Category=PC_Exec, singleton, output
+	static const FString PN_DelegateEntry;    // Category=PC_Exec, singleton, output; entry point for a dynamically bound delegate
+	static const FString PN_EntryPoint;	// entry point to the ubergraph
+	static const FString PN_Self;    // Category=PC_Object, singleton, input
+	static const FString PN_Else;    // Category=PC_Exec, singleton, output
+	static const FString PN_Loop;    // Category=PC_Exec, singleton, output
+	static const FString PN_After;    // Category=PC_Exec, singleton, output
+	static const FString PN_ReturnValue;		// Category=PC_Object, singleton, output
+	static const FString PN_ObjectToCast;    // Category=PC_Object, singleton, input
+	static const FString PN_Condition;    // Category=PC_Boolean, singleton, input
+	static const FString PN_Start;    // Category=PC_Int, singleton, input
+	static const FString PN_Stop;    // Category=PC_Int, singleton, input
+	static const FString PN_Index;    // Category=PC_Int, singleton, output
+	static const FString PN_CastSucceeded;    // Category=PC_Exec, singleton, output
+	static const FString PN_CastFailed;    // Category=PC_Exec, singleton, output
+	static const FString PN_CastedValuePrefix;    // Category=PC_Object, singleton, output; actual pin name varies depending on the type to be casted to, this is just a prefix
+	static const FString PN_MatineeFinished;    // Category=PC_Exec, singleton, output
 
 	// construction script function names
-	UPROPERTY()
-	FName FN_UserConstructionScript;
-
-	UPROPERTY()
-	FName FN_ExecuteUbergraphBase;
+	static const FName FN_UserConstructionScript;
+	static const FName FN_ExecuteUbergraphBase;
 
 	// metadata keys
 
@@ -343,19 +268,14 @@ class BLUEPRINTGRAPH_API UEdGraphSchema_K2 : public UEdGraphSchema
 
 
 	// graph names
-	UPROPERTY()
-	FName GN_EventGraph;
-
-	UPROPERTY()
-	FName GN_AnimGraph;
+	static const FName GN_EventGraph;
+	static const FName GN_AnimGraph;
 
 	// variable names
-	UPROPERTY()
-	FName VR_DefaultCategory;
+	static const FName VR_DefaultCategory;
 
 	// action grouping values
-	UPROPERTY()
-	int32 AG_LevelReference;
+	static const int32 AG_LevelReference;
 
 	/** Whether this schema should use the old (legacy) menu building, or the new (experimental) menu building */
 	UPROPERTY(GlobalConfig)
diff --git a/Engine/Source/Editor/BlueprintGraph/Private/EdGraphSchema_K2.cpp b/Engine/Source/Editor/BlueprintGraph/Private/EdGraphSchema_K2.cpp
--- a/Engine/Source/Editor/BlueprintGraph/Private/EdGraphSchema_K2.cpp
+++ b/Engine/Source/Editor/BlueprintGraph/Private/EdGraphSchema_K2.cpp
@@ -248,54 +248,55 @@ FString UEdGraphSchema_K2::FPinTypeTreeInfo::GetDescription() const
 	}
 }
 
+
+const FString UEdGraphSchema_K2::PC_Exec(TEXT("exec"));
+const FString UEdGraphSchema_K2::PC_Boolean(TEXT("bool"));
+const FString UEdGraphSchema_K2::PC_Byte(TEXT("byte"));
+const FString UEdGraphSchema_K2::PC_Class(TEXT("class"));
+const FString UEdGraphSchema_K2::PC_Int(TEXT("int"));
+const FString UEdGraphSchema_K2::PC_Float(TEXT("float"));
+const FString UEdGraphSchema_K2::PC_Name(TEXT("name"));
+const FString UEdGraphSchema_K2::PC_Delegate(TEXT("delegate"));
+const FString UEdGraphSchema_K2::PC_MCDelegate(TEXT("mcdelegate"));
+const FString UEdGraphSchema_K2::PC_Object(TEXT("object"));
+const FString UEdGraphSchema_K2::PC_Interface(TEXT("interface"));
+const FString UEdGraphSchema_K2::PC_String(TEXT("string"));
+const FString UEdGraphSchema_K2::PC_Text(TEXT("text"));
+const FString UEdGraphSchema_K2::PC_Struct(TEXT("struct"));
+const FString UEdGraphSchema_K2::PC_Wildcard(TEXT("wildcard"));
+const FString UEdGraphSchema_K2::PSC_Self(TEXT("self"));
+const FString UEdGraphSchema_K2::PSC_Index(TEXT("index"));
+const FString UEdGraphSchema_K2::PN_Execute(TEXT("execute"));
+const FString UEdGraphSchema_K2::PN_Then(TEXT("then"));
+const FString UEdGraphSchema_K2::PN_Completed(TEXT("Completed"));
+const FString UEdGraphSchema_K2::PN_DelegateEntry(TEXT("delegate"));
+const FString UEdGraphSchema_K2::PN_EntryPoint(TEXT("EntryPoint"));
+const FString UEdGraphSchema_K2::PN_Self(TEXT("self"));
+const FString UEdGraphSchema_K2::PN_Else(TEXT("else"));
+const FString UEdGraphSchema_K2::PN_Loop(TEXT("loop"));
+const FString UEdGraphSchema_K2::PN_After(TEXT("after"));
+const FString UEdGraphSchema_K2::PN_ReturnValue(TEXT("ReturnValue"));
+const FString UEdGraphSchema_K2::PN_ObjectToCast(TEXT("Object"));
+const FString UEdGraphSchema_K2::PN_Condition(TEXT("Condition"));
+const FString UEdGraphSchema_K2::PN_Start(TEXT("Start"));
+const FString UEdGraphSchema_K2::PN_Stop(TEXT("Stop"));
+const FString UEdGraphSchema_K2::PN_Index(TEXT("Index"));
+const FString UEdGraphSchema_K2::PN_CastSucceeded(TEXT("then"));
+const FString UEdGraphSchema_K2::PN_CastFailed(TEXT("CastFailed"));
+const FString UEdGraphSchema_K2::PN_CastedValuePrefix(TEXT("As"));
+const FString UEdGraphSchema_K2::PN_MatineeFinished(TEXT("Finished"));
+
+const FName UEdGraphSchema_K2::FN_UserConstructionScript(TEXT("UserConstructionScript"));
+const FName UEdGraphSchema_K2::FN_ExecuteUbergraphBase(TEXT("ExecuteUbergraph"));
+const FName UEdGraphSchema_K2::GN_EventGraph(TEXT("EventGraph"));
+const FName UEdGraphSchema_K2::GN_AnimGraph(TEXT("AnimGraph"));
+const FName UEdGraphSchema_K2::VR_DefaultCategory(TEXT("Default"));
+
+const int32 UEdGraphSchema_K2::AG_LevelReference = 100;
+
 UEdGraphSchema_K2::UEdGraphSchema_K2(const class FPostConstructInitializeProperties& PCIP)
 	: Super(PCIP)
 {
-
-	PC_Exec = TEXT("exec");
-	PC_Meta = TEXT("meta");
-	PC_Boolean = TEXT("bool");
-	PC_Byte = TEXT("byte");
-	PC_Class = TEXT("class");
-	PC_Int = TEXT("int");
-	PC_Float = TEXT("float");
-	PC_Name = TEXT("name");
-	PC_Delegate = TEXT("delegate");
-	PC_MCDelegate = TEXT("mcdelegate");
-	PC_Object = TEXT("object");
-	PC_Interface = TEXT("interface");
-	PC_String = TEXT("string");
-	PC_Text = TEXT("text");
-	PC_Struct = TEXT("struct");
-	PC_Wildcard = TEXT("wildcard");
-	PSC_Self = TEXT("self");
-	PSC_Index = TEXT("index");
-	PN_Execute = TEXT("execute");
-	PN_Then = TEXT("then");
-	PN_Completed = TEXT("Completed");
-	PN_DelegateEntry = TEXT("delegate");
-	PN_EntryPoint = TEXT("EntryPoint");
-	PN_Self = TEXT("self");
-	PN_Else = TEXT("else");
-	PN_Loop = TEXT("loop");
-	PN_After = TEXT("after");
-	PN_ReturnValue = TEXT("ReturnValue");
-	PN_ObjectToCast = TEXT("Object");
-	PN_Condition = TEXT("Condition");
-	PN_Start = TEXT("Start");
-	PN_Stop = TEXT("Stop");
-	PN_Index = TEXT("Index");
-	PN_CastSucceeded = TEXT("then");
-	PN_CastFailed = TEXT("CastFailed");
-	PN_CastedValuePrefix = TEXT("As");
-	PN_MatineeFinished = TEXT("Finished");
-	FN_UserConstructionScript = TEXT("UserConstructionScript");
-	FN_ExecuteUbergraphBase = TEXT("ExecuteUbergraph");
-	GN_EventGraph = TEXT("EventGraph");
-	GN_AnimGraph = TEXT("AnimGraph");
-	VR_DefaultCategory = TEXT("Default");
-	AG_LevelReference = 100;
-
 	VectorStruct = FindObjectChecked<UScriptStruct>(UObject::StaticClass(), TEXT("Vector"));
 	RotatorStruct = FindObjectChecked<UScriptStruct>(UObject::StaticClass(), TEXT("Rotator"));
 	TransformStruct = FindObjectChecked<UScriptStruct>(UObject::StaticClass(), TEXT("Transform"));
diff --git a/Engine/Source/Editor/UnrealEd/Private/Kismet2/StructureEditorUtils.cpp b/Engine/Source/Editor/UnrealEd/Private/Kismet2/StructureEditorUtils.cpp
--- a/Engine/Source/Editor/UnrealEd/Private/Kismet2/StructureEditorUtils.cpp
+++ b/Engine/Source/Editor/UnrealEd/Private/Kismet2/StructureEditorUtils.cpp
@@ -163,7 +163,6 @@ bool FStructureEditorUtils::CanHaveAMemberVariableOfType(const UUserDefinedStruc
 		}
 	}
 	else if ((VarType.PinCategory == K2Schema->PC_Exec) 
-		|| (VarType.PinCategory == K2Schema->PC_Meta) 
 		|| (VarType.PinCategory == K2Schema->PC_Wildcard)
 		|| (VarType.PinCategory == K2Schema->PC_MCDelegate)
 		|| (VarType.PinCategory == K2Schema->PC_Delegate))
