Make K2Schema well-known struct cache, static and used consistently

diff --git a/Engine/Source/Editor/BlueprintGraph/Classes/EdGraphSchema_K2.h b/Engine/Source/Editor/BlueprintGraph/Classes/EdGraphSchema_K2.h
--- a/Engine/Source/Editor/BlueprintGraph/Classes/EdGraphSchema_K2.h
+++ b/Engine/Source/Editor/BlueprintGraph/Classes/EdGraphSchema_K2.h
@@ -925,10 +925,10 @@ private:
 	 */
 	bool DoesFunctionHaveOutParameters( const UFunction* Function ) const;
 
-	const UScriptStruct* VectorStruct;
-	const UScriptStruct* RotatorStruct;
-	const UScriptStruct* TransformStruct;
-	const UScriptStruct* LinearColorStruct;
-	const UScriptStruct* ColorStruct;
+	static const UScriptStruct* VectorStruct;
+	static const UScriptStruct* RotatorStruct;
+	static const UScriptStruct* TransformStruct;
+	static const UScriptStruct* LinearColorStruct;
+	static const UScriptStruct* ColorStruct;
 };
 
diff --git a/Engine/Source/Editor/BlueprintGraph/Private/EdGraphSchema_K2.cpp b/Engine/Source/Editor/BlueprintGraph/Private/EdGraphSchema_K2.cpp
--- a/Engine/Source/Editor/BlueprintGraph/Private/EdGraphSchema_K2.cpp
+++ b/Engine/Source/Editor/BlueprintGraph/Private/EdGraphSchema_K2.cpp
@@ -294,18 +294,26 @@ const FName UEdGraphSchema_K2::VR_DefaultCategory(TEXT("Default"));
 
 const int32 UEdGraphSchema_K2::AG_LevelReference = 100;
 
+const UScriptStruct* UEdGraphSchema_K2::VectorStruct = nullptr;
+const UScriptStruct* UEdGraphSchema_K2::RotatorStruct = nullptr;
+const UScriptStruct* UEdGraphSchema_K2::TransformStruct = nullptr;
+const UScriptStruct* UEdGraphSchema_K2::LinearColorStruct = nullptr;
+const UScriptStruct* UEdGraphSchema_K2::ColorStruct = nullptr;
+
 UEdGraphSchema_K2::UEdGraphSchema_K2(const class FPostConstructInitializeProperties& PCIP)
 	: Super(PCIP)
 {
-	VectorStruct = FindObjectChecked<UScriptStruct>(UObject::StaticClass(), TEXT("Vector"));
-	RotatorStruct = FindObjectChecked<UScriptStruct>(UObject::StaticClass(), TEXT("Rotator"));
-	TransformStruct = FindObjectChecked<UScriptStruct>(UObject::StaticClass(), TEXT("Transform"));
-	LinearColorStruct = FindObjectChecked<UScriptStruct>(UObject::StaticClass(), TEXT("LinearColor"));
-	ColorStruct = FindObjectChecked<UScriptStruct>(UObject::StaticClass(), TEXT("Color"));
+	// Initialize cached static references to well-known struct types
+	if (VectorStruct == nullptr)
+	{
+		VectorStruct = FindObjectChecked<UScriptStruct>(UObject::StaticClass(), TEXT("Vector"));
+		RotatorStruct = FindObjectChecked<UScriptStruct>(UObject::StaticClass(), TEXT("Rotator"));
+		TransformStruct = FindObjectChecked<UScriptStruct>(UObject::StaticClass(), TEXT("Transform"));
+		LinearColorStruct = FindObjectChecked<UScriptStruct>(UObject::StaticClass(), TEXT("LinearColor"));
+		ColorStruct = FindObjectChecked<UScriptStruct>(UObject::StaticClass(), TEXT("Color"));
+	}
 }
 
-
-
 bool UEdGraphSchema_K2::DoesFunctionHaveOutParameters( const UFunction* Function ) const
 {
 	if ( Function != NULL )
@@ -1424,12 +1432,6 @@ bool UEdGraphSchema_K2::SearchForAutocastFunction(const UEdGraphPin* OutputPin,
 
 	TargetFunction = NAME_None;
 
-	const UScriptStruct* VectorStruct = FindObjectChecked<UScriptStruct>(UObject::StaticClass(), TEXT("Vector"));
-	const UScriptStruct* RotatorStruct = FindObjectChecked<UScriptStruct>(UObject::StaticClass(), TEXT("Rotator"));
-	const UScriptStruct* TransformStruct = FindObjectChecked<UScriptStruct>(UObject::StaticClass(), TEXT("Transform"));
-	const UScriptStruct* LinearColorStruct = FindObjectChecked<UScriptStruct>(UObject::StaticClass(), TEXT("LinearColor"));
-	const UScriptStruct* ColorStruct = FindObjectChecked<UScriptStruct>(UObject::StaticClass(), TEXT("Color"));
-
 	const UScriptStruct* InputStructType = Cast<const UScriptStruct>(InputPin->PinType.PinSubCategoryObject.Get());
 	const UScriptStruct* OutputStructType = Cast<const UScriptStruct>(OutputPin->PinType.PinSubCategoryObject.Get());
 
@@ -2014,11 +2016,6 @@ bool UEdGraphSchema_K2::DefaultValueSimpleValidation(const FEdGraphPinType& PinT
 		}
 		else if(!NewDefaultValue.IsEmpty())
 		{
-			UScriptStruct* VectorStruct = FindObjectChecked<UScriptStruct>(UObject::StaticClass(), TEXT("Vector"));
-			UScriptStruct* RotatorStruct = FindObjectChecked<UScriptStruct>(UObject::StaticClass(), TEXT("Rotator"));
-			UScriptStruct* TransformStruct = FindObjectChecked<UScriptStruct>(UObject::StaticClass(), TEXT("Transform"));
-			UScriptStruct* LinearColorStruct = FindObjectChecked<UScriptStruct>(UObject::StaticClass(), TEXT("LinearColor"));
-
 			if (StructType == VectorStruct)
 			{
 				if (!FDefaultValueHelper::IsStringValidVector(NewDefaultValue))
@@ -2108,10 +2105,6 @@ FLinearColor UEdGraphSchema_K2::GetPinTypeColor(const FEdGraphPinType& PinType)
 	}
 	else if (TypeString == PC_Struct)
 	{
-		static UScriptStruct* VectorStruct = FindObjectChecked<UScriptStruct>(UObject::StaticClass(), TEXT("Vector"));
-		static UScriptStruct* RotatorStruct = FindObjectChecked<UScriptStruct>(UObject::StaticClass(), TEXT("Rotator"));
-		static UScriptStruct* TransformStruct = FindObjectChecked<UScriptStruct>(UObject::StaticClass(), TEXT("Transform"));
-
 		if (PinType.PinSubCategoryObject == VectorStruct)
 		{
 			// vector
@@ -2618,9 +2611,6 @@ void UEdGraphSchema_K2::GetVariableTypeTree( TArray< TSharedPtr<FPinTypeTreeInfo
 	TypeTree.Add( MakeShareable( new FPinTypeTreeInfo(PC_Text, this, LOCTEXT("TextType", "A localizable text string").ToString()) ) );
 
 	// Add in special first-class struct types
-	UScriptStruct* VectorStruct = FindObjectChecked<UScriptStruct>(UObject::StaticClass(), TEXT("Vector"));
-	UScriptStruct* RotatorStruct = FindObjectChecked<UScriptStruct>(UObject::StaticClass(), TEXT("Rotator"));
-	UScriptStruct* TransformStruct = FindObjectChecked<UScriptStruct>(UObject::StaticClass(), TEXT("Transform"));
 	TypeTree.Add( MakeShareable( new FPinTypeTreeInfo(PC_Struct, VectorStruct, LOCTEXT("VectorType", "A 3D vector").ToString()) ) );
 	TypeTree.Add( MakeShareable( new FPinTypeTreeInfo(PC_Struct, RotatorStruct, LOCTEXT("RotatorType", "A 3D rotation").ToString()) ) );
 	TypeTree.Add( MakeShareable( new FPinTypeTreeInfo(PC_Struct, TransformStruct, LOCTEXT("TransformType", "A 3D transformation, including translation, rotation and 3D scale.").ToString()) ) );
@@ -3196,9 +3186,6 @@ void UEdGraphSchema_K2::SetPinDefaultValue(UEdGraphPin* Pin, const UFunction* Fu
 
 void UEdGraphSchema_K2::SetPinDefaultValueBasedOnType(UEdGraphPin* Pin) const
 {
-	UScriptStruct* VectorStruct = FindObjectChecked<UScriptStruct>(UObject::StaticClass(), TEXT("Vector"));
-	UScriptStruct* RotatorStruct = FindObjectChecked<UScriptStruct>(UObject::StaticClass(), TEXT("Rotator"));
-
 	// Create a useful default value based on the pin type
 	if(Pin->PinType.bIsArray)
 	{
@@ -3294,11 +3281,6 @@ UFunction* UEdGraphSchema_K2::FindSetVariableByNameFunction(const FEdGraphPinTyp
 
 	const UEdGraphSchema_K2* K2Schema = GetDefault<UEdGraphSchema_K2>();
 
-	UScriptStruct* VectorStruct = FindObjectChecked<UScriptStruct>(UObject::StaticClass(), TEXT("Vector"));
-	UScriptStruct* RotatorStruct = FindObjectChecked<UScriptStruct>(UObject::StaticClass(), TEXT("Rotator"));
-	UScriptStruct* ColorStruct = FindObjectChecked<UScriptStruct>(UObject::StaticClass(), TEXT("LinearColor"));
-	UScriptStruct* TransformStruct = FindObjectChecked<UScriptStruct>(UObject::StaticClass(), TEXT("Transform"));
-
 	FName SetFunctionName = NAME_None;
 	if(PinType.PinCategory == K2Schema->PC_Int)
 	{
