BP: select node is updated, when index enum is changed. New option pin in select node has proper type.

diff --git a/Engine/Source/Editor/BlueprintGraph/Classes/K2Node_Select.h b/Engine/Source/Editor/BlueprintGraph/Classes/K2Node_Select.h
--- a/Engine/Source/Editor/BlueprintGraph/Classes/K2Node_Select.h
+++ b/Engine/Source/Editor/BlueprintGraph/Classes/K2Node_Select.h
@@ -3,10 +3,11 @@
 
 #pragma once
 #include "K2Node.h"
+#include "NodeDependingOnEnumInterface.h"
 #include "K2Node_Select.generated.h"
 
 UCLASS(MinimalAPI)
-class UK2Node_Select : public UK2Node
+class UK2Node_Select : public UK2Node, public INodeDependingOnEnumInterface
 {
 	GENERATED_UCLASS_BODY()
 
@@ -31,7 +32,7 @@ class UK2Node_Select : public UK2Node
 	TArray<FName> EnumEntryFriendlyNames;
 
 	/** Whether we need to reconstruct the node after the pins have changed */
-	UPROPERTY()
+	UPROPERTY(Transient)
 	bool bReconstructNode;
 
 	// Begin UEdGraphNode interface
@@ -58,6 +59,11 @@ class UK2Node_Select : public UK2Node
 	virtual int32 GetNodeRefreshPriority() const override { return EBaseNodeRefreshPriority::Low_UsesDependentWildcard; }
 	// End UK2Node interface
 
+	// INodeDependingOnEnumInterface
+	virtual class UEnum* GetEnum() const override { return Enum; }
+	virtual bool ShouldBeReconstructedAfterEnumChanged() const override { return true; }
+	// End of INodeDependingOnEnumInterface
+
 	/** Get the return value pin */
 	BLUEPRINTGRAPH_API UEdGraphPin* GetReturnValuePin() const;
 	/** Get the condition pin */
@@ -93,6 +99,6 @@ class UK2Node_Select : public UK2Node
 	BLUEPRINTGRAPH_API bool CanChangePinType(UEdGraphPin* Pin) const;
 
 	// Bind the options to a named enum 
-	virtual void SetEnum(UEnum* InEnum);
+	virtual void SetEnum(UEnum* InEnum, bool bForceRegenerate = false);
 };
 
diff --git a/Engine/Source/Editor/BlueprintGraph/Classes/K2Node_SelectEnum.h b/Engine/Source/Editor/BlueprintGraph/Classes/K2Node_SelectEnum.h
--- a/Engine/Source/Editor/BlueprintGraph/Classes/K2Node_SelectEnum.h
+++ b/Engine/Source/Editor/BlueprintGraph/Classes/K2Node_SelectEnum.h
@@ -28,6 +28,6 @@ class UDEPRECATED_K2Node_SelectEnum : public UK2Node_Select
 	// End UK2Node_Select interface
 
 	// Bind the option to a named enum 
-	BLUEPRINTGRAPH_API void SetEnum(UEnum* InEnum) override;
+	BLUEPRINTGRAPH_API void SetEnum(UEnum* InEnum, bool bForceRegenerate = false) override;
 };
 
diff --git a/Engine/Source/Editor/BlueprintGraph/Private/K2Node_Select.cpp b/Engine/Source/Editor/BlueprintGraph/Private/K2Node_Select.cpp
--- a/Engine/Source/Editor/BlueprintGraph/Private/K2Node_Select.cpp
+++ b/Engine/Source/Editor/BlueprintGraph/Private/K2Node_Select.cpp
@@ -207,7 +207,7 @@ void UK2Node_Select::AllocateDefaultPins()
 	const UEdGraphSchema_K2* Schema = GetDefault<UEdGraphSchema_K2>();
 
 	// To refresh, just in case it changed
-	SetEnum(Enum);
+	SetEnum(Enum, true);
 
 	if (Enum)
 	{
@@ -344,13 +344,25 @@ void UK2Node_Select::PostReconstructNode()
 {
 	bReconstructNode = false;
 
+	const UEdGraphSchema_K2* Schema = Cast<UEdGraphSchema_K2>(GetSchema());
+
 	UEdGraphPin* ReturnPin = GetReturnValuePin();
 	PinConnectionListChanged(ReturnPin);
+	const bool bFillTypeFromReturn = Schema && ReturnPin && (ReturnPin->PinType.PinCategory != Schema->PC_Wildcard);
 
 	TArray<UEdGraphPin*> OptionPins;
 	GetOptionPins(OptionPins);
 	for (auto It = OptionPins.CreateConstIterator(); It; It++)
 	{
+		UEdGraphPin* Pin = *It;
+		const bool bTypeShouldBeFilled = Schema && Pin && (Pin->PinType.PinCategory == Schema->PC_Wildcard);
+		if (bTypeShouldBeFilled && bFillTypeFromReturn)
+		{
+			Pin->Modify();
+			Pin->PinType = ReturnPin->PinType;
+			UEdGraphSchema_K2::ValidateExistingConnections(Pin);
+		}
+
 		PinConnectionListChanged(*It);
 	}
 
@@ -558,14 +570,14 @@ void UK2Node_Select::RemoveOptionPinToNode()
 	ReconstructNode();
 }
 
-void UK2Node_Select::SetEnum(UEnum* InEnum)
+void UK2Node_Select::SetEnum(UEnum* InEnum, bool bForceRegenerate)
 {
 	const UEdGraphSchema_K2* Schema = GetDefault<UEdGraphSchema_K2>();
 
 	UEnum* PrevEnum = Enum;
 	Enum = InEnum;
 
-	if (PrevEnum != Enum)
+	if ((PrevEnum != Enum) || bForceRegenerate)
 	{
 		// regenerate enum name list
 		EnumEntries.Empty();
diff --git a/Engine/Source/Editor/BlueprintGraph/Private/K2Node_SelectEnum.cpp b/Engine/Source/Editor/BlueprintGraph/Private/K2Node_SelectEnum.cpp
--- a/Engine/Source/Editor/BlueprintGraph/Private/K2Node_SelectEnum.cpp
+++ b/Engine/Source/Editor/BlueprintGraph/Private/K2Node_SelectEnum.cpp
@@ -100,7 +100,7 @@ void UDEPRECATED_K2Node_SelectEnum::GetOptionPins(TArray<UEdGraphPin*>& OptionPi
 	}
 }
 
-void UDEPRECATED_K2Node_SelectEnum::SetEnum(UEnum* InEnum)
+void UDEPRECATED_K2Node_SelectEnum::SetEnum(UEnum* InEnum, bool bForceRegenerate)
 {
 	Enum = InEnum;
 
