Change how types are displayed in blueprint pin tooltips
Deprecate TypeToString
Make ConstructBasicPinTooltip take FText for Description
Make SpawnActorFromClass pin tooltips consistent with other nodes

diff --git a/Engine/Plugins/Experimental/LiveEditor/Source/LiveEditor/Private/K2Node_LiveEditObject.cpp b/Engine/Plugins/Experimental/LiveEditor/Source/LiveEditor/Private/K2Node_LiveEditObject.cpp
--- a/Engine/Plugins/Experimental/LiveEditor/Source/LiveEditor/Private/K2Node_LiveEditObject.cpp
+++ b/Engine/Plugins/Experimental/LiveEditor/Source/LiveEditor/Private/K2Node_LiveEditObject.cpp
@@ -778,7 +778,7 @@ UEdGraphPin *UK2Node_LiveEditObject::GetClampMaxPin() const
 
 void UK2Node_LiveEditObject::SetPinToolTip(UEdGraphPin& MutatablePin, const FText& PinDescription) const
 {
-	MutatablePin.PinToolTip = UEdGraphSchema_K2::TypeToString(MutatablePin.PinType);
+	MutatablePin.PinToolTip = UEdGraphSchema_K2::TypeToText(MutatablePin.PinType).ToString();
 
 	UEdGraphSchema_K2 const* const K2Schema = Cast<const UEdGraphSchema_K2>(GetSchema());
 	if (K2Schema != nullptr)
diff --git a/Engine/Source/Editor/BlueprintGraph/Classes/EdGraphSchema_K2.h b/Engine/Source/Editor/BlueprintGraph/Classes/EdGraphSchema_K2.h
--- a/Engine/Source/Editor/BlueprintGraph/Classes/EdGraphSchema_K2.h
+++ b/Engine/Source/Editor/BlueprintGraph/Classes/EdGraphSchema_K2.h
@@ -376,7 +376,7 @@ public:
 	virtual bool ShouldShowAssetPickerForPin(UEdGraphPin* Pin) const override;
 	virtual FLinearColor GetPinTypeColor(const FEdGraphPinType& PinType) const override;
 	virtual FString GetPinDisplayName(const UEdGraphPin* Pin) const override;
-	virtual void ConstructBasicPinTooltip(const UEdGraphPin& Pin, const FString& PinDescription, FString& TooltipOut) const override;
+	virtual void ConstructBasicPinTooltip(const UEdGraphPin& Pin, const FText& PinDescription, FString& TooltipOut) const override;
 	virtual EGraphType GetGraphType(const UEdGraph* TestEdGraph) const override;
 	virtual bool IsTitleBarPin(const UEdGraphPin& Pin) const override;
 	virtual void BreakNodeLinks(UEdGraphNode& TargetNode) const override;
@@ -723,6 +723,7 @@ public:
 	 *
 	 * @return	The converted type string.
 	 */
+	DEPRECATED(4.5, "UEdGraphSchema_K2::TypeToString is deprecated.  Use TypeToText instead.")
 	static FString TypeToString(const FEdGraphPinType& Type);
 
 	/**
@@ -732,8 +733,18 @@ public:
 	 *
 	 * @return	The converted type string.
 	 */
+	DEPRECATED(4.5, "UEdGraphSchema_K2::TypeToString is deprecated.  Use TypeToText instead.")
 	static FString TypeToString(UProperty* const Property);
 
+	/**
+	 * Converts the type of a property into a fully qualified string (e.g., object'ObjectName').
+	 *
+	 * @param	Property	The property to convert into a string.
+	 *
+	 * @return	The converted type string.
+	 */
+	static FText TypeToText(UProperty* const Property);
+
 	/**
 	 * Converts a pin type into a fully qualified FText (e.g., object'ObjectName').
 	 *
@@ -743,6 +754,16 @@ public:
 	 */
 	static FText TypeToText(const FEdGraphPinType& Type);
 
+	/**
+	 * Returns the FText to use for a given schema category
+	 *
+	 * @param	Category	The category to convert into a FText.
+	 * @param	bForMenu	Indicates if this is for display in tooltips or menu
+	 *
+	 * @return	The text to display for the category.
+	 */
+	static FText GetCategoryText(const FString& Category, const bool bForMenu = false);
+
 	/**
 	 * Get the type tree for all of the property types valid for this schema
 	 *
diff --git a/Engine/Source/Editor/BlueprintGraph/Classes/K2Node_SpawnActorFromClass.h b/Engine/Source/Editor/BlueprintGraph/Classes/K2Node_SpawnActorFromClass.h
--- a/Engine/Source/Editor/BlueprintGraph/Classes/K2Node_SpawnActorFromClass.h
+++ b/Engine/Source/Editor/BlueprintGraph/Classes/K2Node_SpawnActorFromClass.h
@@ -55,15 +55,6 @@ class UK2Node_SpawnActorFromClass : public UK2Node
 	BLUEPRINTGRAPH_API UClass* GetClassToSpawn(const TArray<UEdGraphPin*>* InPinsToSearch=NULL) const;
 
 private:
-	/**
-	 * Takes the specified "MutatablePin" and sets its 'PinToolTip' field (according
-	 * to the specified description)
-	 * 
-	 * @param   MutatablePin	The pin you want to set tool-tip text on
-	 * @param   PinDescription	A string describing the pin's purpose
-	 */
-	void SetPinToolTip(UEdGraphPin& MutatablePin, const FText& PinDescription) const;
-
 	/** Tooltip text for this node. */
 	FText NodeTooltip;
 
diff --git a/Engine/Source/Editor/BlueprintGraph/Private/EdGraphSchema_K2.cpp b/Engine/Source/Editor/BlueprintGraph/Private/EdGraphSchema_K2.cpp
--- a/Engine/Source/Editor/BlueprintGraph/Private/EdGraphSchema_K2.cpp
+++ b/Engine/Source/Editor/BlueprintGraph/Private/EdGraphSchema_K2.cpp
@@ -1389,7 +1389,7 @@ const FPinConnectionResponse UEdGraphSchema_K2::CanCreateConnection(const UEdGra
 
 		if (bCanAutocast || bCanAutoConvert)
 		{
-			return FPinConnectionResponse(CONNECT_RESPONSE_MAKE_WITH_CONVERSION_NODE, FString::Printf(TEXT("Convert %s to %s"), *TypeToString(OutputPin->PinType), *TypeToString(InputPin->PinType)));
+			return FPinConnectionResponse(CONNECT_RESPONSE_MAKE_WITH_CONVERSION_NODE, FString::Printf(TEXT("Convert %s to %s"), *TypeToText(OutputPin->PinType).ToString(), *TypeToText(InputPin->PinType).ToString()));
 		}
 		else
 		{
@@ -1405,7 +1405,7 @@ const FPinConnectionResponse UEdGraphSchema_K2::CanCreateConnection(const UEdGra
 				}
 				return FPinConnectionResponse(CONNECT_RESPONSE_DISALLOW, Msg);
 			}
-			return FPinConnectionResponse(CONNECT_RESPONSE_DISALLOW, FString::Printf(TEXT("%s is not compatible with %s"), *TypeToString(PinA->PinType), *TypeToString(PinB->PinType)));
+			return FPinConnectionResponse(CONNECT_RESPONSE_DISALLOW, FString::Printf(TEXT("%s is not compatible with %s"), *TypeToText(PinA->PinType).ToString(), *TypeToText(PinB->PinType).ToString()));
 		}
 	}
 }
@@ -2061,8 +2061,8 @@ bool UEdGraphSchema_K2::DefaultValueSimpleValidation(const FEdGraphPinType& PinT
 	}
 	else
 	{
-		//@TODO: MessageLog.Error(*FString::Printf(TEXT("Unsupported type %s on @@"), *UEdGraphSchema_K2::TypeToString(Type)), SourceObject);
-		DVSV_RETURN_MSG( FString::Printf(TEXT("Unsupported type %s on pin %s"), *UEdGraphSchema_K2::TypeToString(PinType), *(PinName)) ); 
+		//@TODO: MessageLog.Error(*FString::Printf(TEXT("Unsupported type %s on @@"), *UEdGraphSchema_K2::TypeToText(Type).ToString()), SourceObject);
+		DVSV_RETURN_MSG( FString::Printf(TEXT("Unsupported type %s on pin %s"), *UEdGraphSchema_K2::TypeToText(PinType).ToString(), *(PinName)) ); 
 	}
 
 #undef DVSV_RETURN_MSG
@@ -2192,26 +2192,39 @@ FString UEdGraphSchema_K2::GetPinDisplayName(const UEdGraphPin* Pin) const
 	return DisplayName;
 }
 
-void UEdGraphSchema_K2::ConstructBasicPinTooltip(const UEdGraphPin& Pin, const FString& PinDescription, FString& TooltipOut) const
+void UEdGraphSchema_K2::ConstructBasicPinTooltip(const UEdGraphPin& Pin, const FText& PinDescription, FString& TooltipOut) const
 {
-	// using a local FString so users can use the same variable for PinDescription and TooltipOut
-	FString ConstructedTooltip = UEdGraphSchema_K2::TypeToString(Pin.PinType);
+	FFormatNamedArguments Args;
+	Args.Add(TEXT("PinType"), TypeToText(Pin.PinType));
 
 	if (UEdGraphNode* PinNode = Pin.GetOwningNode())
 	{
 		UEdGraphSchema_K2 const* const K2Schema = Cast<const UEdGraphSchema_K2>(PinNode->GetSchema());
 		if (ensure(K2Schema != NULL)) // ensure that this node belongs to this schema
 		{
-			ConstructedTooltip += TEXT(" ");
-			ConstructedTooltip += K2Schema->GetPinDisplayName(&Pin);
+			Args.Add(TEXT("DisplayName"), FText::FromString(GetPinDisplayName(&Pin)));
+			Args.Add(TEXT("LineFeed1"), FText::FromString(TEXT("\n")));
 		}
-	}	
+	}
+	else
+	{
+			Args.Add(TEXT("DisplayName"), FText::GetEmpty());
+			Args.Add(TEXT("LineFeed1"), FText::GetEmpty());
+	}
+
 
 	if (!PinDescription.IsEmpty())
 	{
-	ConstructedTooltip += FString(TEXT("\n")) + PinDescription;
+		Args.Add(TEXT("Description"), PinDescription);
+		Args.Add(TEXT("LineFeed2"), FText::FromString(TEXT("\n\n")));
 	}
-	TooltipOut = ConstructedTooltip; // using a local FString, so PinDescription and TooltipOut can be the same variable
+	else
+	{
+		Args.Add(TEXT("Description"), FText::GetEmpty());
+		Args.Add(TEXT("LineFeed2"), FText::GetEmpty());
+	}
+	
+	TooltipOut = FText::Format(LOCTEXT("PinTooltip", "{DisplayName}{LineFeed1}{PinType}{LineFeed2}{Description}"), Args).ToString(); 
 }
 
 EGraphType UEdGraphSchema_K2::GetGraphType(const UEdGraph* TestEdGraph) const
@@ -2479,53 +2492,99 @@ FString UEdGraphSchema_K2::TypeToString(const FEdGraphPinType& Type)
 }
 
 FString UEdGraphSchema_K2::TypeToString(UProperty* const Property)
+{
+	return TypeToText(Property).ToString();
+}
+
+FText UEdGraphSchema_K2::TypeToText(UProperty* const Property)
 {
 	if (UStructProperty* Struct = Cast<UStructProperty>(Property))
 	{
 		if (Struct->Struct)
 		{
-		return FString::Printf(TEXT("struct'%s'"), *Struct->Struct->GetName());
-	}
+			FEdGraphPinType PinType;
+			PinType.PinCategory = PC_Struct;
+			PinType.PinSubCategoryObject = Struct->Struct;
+			return TypeToText(PinType);
+		}
 	}
 	else if (UClassProperty* Class = Cast<UClassProperty>(Property))
 	{
-		if (Class->MetaClass != nullptr)
+		if (Class->MetaClass)
 		{
-		return FString::Printf(TEXT("class'%s'"), *Class->MetaClass->GetName());
-	}
+			FEdGraphPinType PinType;
+			PinType.PinCategory = PC_Class;
+			PinType.PinSubCategoryObject = Class->MetaClass;
+			return TypeToText(PinType);
+		}
 	}
 	else if (UInterfaceProperty* Interface = Cast<UInterfaceProperty>(Property))
 	{
 		if (Interface->InterfaceClass != nullptr)
 		{
-		return FString::Printf(TEXT("interface'%s'"), *Interface->InterfaceClass->GetName());
-	}
+			FEdGraphPinType PinType;
+			PinType.PinCategory = PC_Interface;
+			PinType.PinSubCategoryObject = Interface->InterfaceClass;
+			return TypeToText(PinType);
+		}
 	}
 	else if (UObjectPropertyBase* Obj = Cast<UObjectPropertyBase>(Property))
 	{
 		if( Obj->PropertyClass )
 		{
-			if( Property->IsA(UWeakObjectProperty::StaticClass()) )
-			{
-				return FString::Printf(TEXT("weak_ptr_object'%s'"), *Obj->PropertyClass->GetName());
-			}
-			else
-			{
-				return FString::Printf(TEXT("object'%s'"), *Obj->PropertyClass->GetName());
-			}
+			FEdGraphPinType PinType;
+			PinType.PinCategory = PC_Object;
+			PinType.PinSubCategoryObject = Obj->PropertyClass;
+			PinType.bIsWeakPointer = Property->IsA(UWeakObjectProperty::StaticClass());
+			return TypeToText(PinType);
 		}
 
-		return TEXT("");
+		return FText::GetEmpty();
 	}
 	else if (UArrayProperty* Array = Cast<UArrayProperty>(Property))
 	{
 		if (Array->Inner)
 		{
-		return FString::Printf(TEXT("array[%s]"), *TypeToString(Array->Inner)); 
-	}
+			FFormatNamedArguments Args;
+			Args.Add(TEXT("ArrayType"), TypeToText(Array->Inner));
+			return FText::Format(LOCTEXT("ArrayPropertyText", "Array of {ArrayType}"), Args); 
+		}
 	}
 	
-		return Property->GetClass()->GetName();
+	return FText::FromString(Property->GetClass()->GetName());
+}
+
+FText UEdGraphSchema_K2::GetCategoryText(const FString& Category, const bool bForMenu)
+{
+	static TMap<FString, FText> CategoryDescriptions;
+	if (CategoryDescriptions.Num() == 0)
+	{
+		CategoryDescriptions.Add(PC_Exec, LOCTEXT("Exec", "Exec"));
+		CategoryDescriptions.Add(PC_Boolean, LOCTEXT("BoolCategory","Boolean"));
+		CategoryDescriptions.Add(PC_Byte, LOCTEXT("ByteCategory","Byte"));
+		CategoryDescriptions.Add(PC_Class, LOCTEXT("ClassCategory","Class"));
+		CategoryDescriptions.Add(PC_Int, LOCTEXT("IntCategory","Integer"));
+		CategoryDescriptions.Add(PC_Float, LOCTEXT("FloatCategory","Float"));
+		CategoryDescriptions.Add(PC_Name, LOCTEXT("NameCategory","Name"));
+		CategoryDescriptions.Add(PC_Delegate, LOCTEXT("DelegateCategory","Delegate"));
+		CategoryDescriptions.Add(PC_MCDelegate, LOCTEXT("MulticastDelegateCategory","Multicast Delegate"));
+		CategoryDescriptions.Add(PC_Object, LOCTEXT("ObjectCategory","Reference"));
+		CategoryDescriptions.Add(PC_Interface, LOCTEXT("InterfaceCategory","Interface"));
+		CategoryDescriptions.Add(PC_String, LOCTEXT("StringCategory","String"));
+		CategoryDescriptions.Add(PC_Text, LOCTEXT("TextCategory","Text"));
+		CategoryDescriptions.Add(PC_Struct, LOCTEXT("StructCategory","Structure"));
+		CategoryDescriptions.Add(PC_Wildcard, LOCTEXT("WildcardCategory","Wildcard"));
+	}
+
+	if (bForMenu)
+	{
+		if (Category == PC_Object)
+		{
+			return LOCTEXT("ObjectCategoryForMenu", "Object Reference");
+		}
+	}
+
+	return CategoryDescriptions.FindChecked(Category);
 }
 
 FText UEdGraphSchema_K2::TypeToText(const FEdGraphPinType& Type)
@@ -2533,13 +2592,13 @@ FText UEdGraphSchema_K2::TypeToText(const FEdGraphPinType& Type)
 	FText PropertyText;
 
 	if (Type.PinSubCategoryObject != NULL)
-{
+	{
 		const UEdGraphSchema_K2* Schema = GetDefault<UEdGraphSchema_K2>();
 		if (Type.PinCategory == Schema->PC_Byte)
 		{
 			FFormatNamedArguments Args;
 			Args.Add(TEXT("EnumName"), FText::FromString(Type.PinSubCategoryObject->GetName()));
-			PropertyText = FText::Format(LOCTEXT("EnumAsText", "enum'{EnumName}'"), Args);
+			PropertyText = FText::Format(LOCTEXT("EnumAsText", "{EnumName} Enum"), Args);
 		}
 		else
 		{
@@ -2548,39 +2607,45 @@ FText UEdGraphSchema_K2::TypeToText(const FEdGraphPinType& Type)
 				UClass* PSCOAsClass = Cast<UClass>(Type.PinSubCategoryObject.Get());
 				const bool bIsInterface = PSCOAsClass && PSCOAsClass->HasAnyClassFlags(CLASS_Interface);
 
-				FText CategoryDesc = !bIsInterface ? FText::FromString(Type.PinCategory) : LOCTEXT("Interface", "interface");
-
 				FFormatNamedArguments Args;
-				Args.Add(TEXT("Category"), CategoryDesc);
-				Args.Add(TEXT("ObjectName"), FText::FromString(Type.PinSubCategoryObject.Get()->GetName()));
-				PropertyText = FText::Format(LOCTEXT("ObjectAsText", "{Category}'{ObjectName}'"), Args);
+				// Don't display the category for "well-known" struct types
+				if (Type.PinCategory == PC_Struct && (Type.PinSubCategoryObject == VectorStruct || Type.PinSubCategoryObject == RotatorStruct || Type.PinSubCategoryObject == TransformStruct))
+				{
+					Args.Add(TEXT("Category"), FText::GetEmpty());
+				}
+				else
+				{
+					Args.Add(TEXT("Category"), (!bIsInterface ? GetCategoryText(Type.PinCategory) : GetCategoryText(PC_Interface)));
+				}
+				Args.Add(TEXT("ObjectName"), FText::FromString(FName::NameToDisplayString(Type.PinSubCategoryObject.Get()->GetName(), false)));
+				PropertyText = FText::Format(LOCTEXT("ObjectAsText", "{ObjectName} {Category}"), Args);
 			}
 			else
 			{
 				FFormatNamedArguments Args;
 				Args.Add(TEXT("Category"), FText::FromString(Type.PinCategory));
 				Args.Add(TEXT("ObjectName"), FText::FromString(Type.PinSubCategoryObject.Get()->GetName()));
-				PropertyText = FText::Format(LOCTEXT("WeakPtrAsText", "weak_ptr_{Category}'{ObjectName}'"), Args);
+				PropertyText = FText::Format(LOCTEXT("WeakPtrAsText", "{ObjectName} Weak {Category}"), Args);
 			}
 		}
 	}
 	else if (Type.PinSubCategory != TEXT(""))
 	{
 		FFormatNamedArguments Args;
-		Args.Add(TEXT("Category"), FText::FromString(Type.PinCategory));
-		Args.Add(TEXT("ObjectName"), FText::FromString(Type.PinSubCategory));
-		PropertyText = FText::Format(LOCTEXT("ObjectAsText", "{Category}'{ObjectName}'"), Args);
+		Args.Add(TEXT("Category"), GetCategoryText(Type.PinCategory));
+		Args.Add(TEXT("ObjectName"), FText::FromString(FName::NameToDisplayString(Type.PinSubCategory, false)));
+		PropertyText = FText::Format(LOCTEXT("ObjectAsText", "{ObjectName} {Category}"), Args);
 	}
 	else
 	{
-		PropertyText = FText::FromString(Type.PinCategory);
+		PropertyText = GetCategoryText(Type.PinCategory);
 	}
 
 	if (Type.bIsArray)
 	{
 		FFormatNamedArguments Args;
 		Args.Add(TEXT("PropertyTitle"), PropertyText);
-		PropertyText = FText::Format(LOCTEXT("ArrayAsText", "array[{PropertyTitle}]"), Args);
+		PropertyText = FText::Format(LOCTEXT("ArrayAsText", "Array of {PropertyTitle}"), Args);
 	}
 	else if (Type.bIsReference)
 	{
diff --git a/Engine/Source/Editor/BlueprintGraph/Private/K2Node_AddComponent.cpp b/Engine/Source/Editor/BlueprintGraph/Private/K2Node_AddComponent.cpp
--- a/Engine/Source/Editor/BlueprintGraph/Private/K2Node_AddComponent.cpp
+++ b/Engine/Source/Editor/BlueprintGraph/Private/K2Node_AddComponent.cpp
@@ -74,10 +74,10 @@ void UK2Node_AddComponent::AllocateDefaultPins()
 
 	UEdGraphSchema const* Schema = GetSchema();
 	check(Schema != NULL);
-	Schema->ConstructBasicPinTooltip(*ManualAttachmentPin, LOCTEXT("ManualAttachmentPinTooltip", "Defines whether the component should attach to the root automatically, or be left unattached for the user to manually attach later.").ToString(), ManualAttachmentPin->PinToolTip);
+	Schema->ConstructBasicPinTooltip(*ManualAttachmentPin, LOCTEXT("ManualAttachmentPinTooltip", "Defines whether the component should attach to the root automatically, or be left unattached for the user to manually attach later."), ManualAttachmentPin->PinToolTip);
 
 	UEdGraphPin* TransformPin = GetRelativeTransformPin();
-	Schema->ConstructBasicPinTooltip(*TransformPin, LOCTEXT("TransformPinTooltip", "Defines where to position the component (relative to its parent). If the component is left unattached, then the transform is relative to the world.").ToString(), TransformPin->PinToolTip);
+	Schema->ConstructBasicPinTooltip(*TransformPin, LOCTEXT("TransformPinTooltip", "Defines where to position the component (relative to its parent). If the component is left unattached, then the transform is relative to the world."), TransformPin->PinToolTip);
 }
 
 void UK2Node_AddComponent::AllocatePinsForExposedVariables()
diff --git a/Engine/Source/Editor/BlueprintGraph/Private/K2Node_CallFunction.cpp b/Engine/Source/Editor/BlueprintGraph/Private/K2Node_CallFunction.cpp
--- a/Engine/Source/Editor/BlueprintGraph/Private/K2Node_CallFunction.cpp
+++ b/Engine/Source/Editor/BlueprintGraph/Private/K2Node_CallFunction.cpp
@@ -396,7 +396,7 @@ void UK2Node_CallFunction::GeneratePinHoverText(const UEdGraphPin& Pin, FString&
 
 	if (K2Schema == NULL)
 	{
-		Schema->ConstructBasicPinTooltip(Pin, TEXT(""), HoverTextOut);
+		Schema->ConstructBasicPinTooltip(Pin, FText::GetEmpty(), HoverTextOut);
 		return;
 	}
 	
@@ -416,7 +416,7 @@ void UK2Node_CallFunction::GeneratePinHoverText(const UEdGraphPin& Pin, FString&
 	UFunction* Function = GetTargetFunction();
 	if (Function == NULL)
 	{
-		Schema->ConstructBasicPinTooltip(Pin, TEXT(""), HoverTextOut);
+		Schema->ConstructBasicPinTooltip(Pin, FText::GetEmpty(), HoverTextOut);
 		return;
 	}
 
@@ -502,13 +502,13 @@ void UK2Node_CallFunction::GeneratePinHoverText(const UEdGraphPin& Pin, FString&
 		// if we came up with a valid description for the param/return-val
 		if (!ParamDesc.IsEmpty())
 		{
-			HoverTextOut += FString(TEXT("\n")) + ParamDesc;
+			HoverTextOut += ParamDesc;
 			break; // we found a match, so there's no need to continue
 		}
 
 	} while (CurStrPos < FullToolTipLen);
 
-	Schema->ConstructBasicPinTooltip(Pin, HoverTextOut, HoverTextOut);
+	Schema->ConstructBasicPinTooltip(Pin, FText::FromString(HoverTextOut), HoverTextOut);
 }
 
 bool UK2Node_CallFunction::CreatePinsForFunctionCall(const UFunction* Function)
diff --git a/Engine/Source/Editor/BlueprintGraph/Private/K2Node_ConstructObjectFromClass.cpp b/Engine/Source/Editor/BlueprintGraph/Private/K2Node_ConstructObjectFromClass.cpp
--- a/Engine/Source/Editor/BlueprintGraph/Private/K2Node_ConstructObjectFromClass.cpp
+++ b/Engine/Source/Editor/BlueprintGraph/Private/K2Node_ConstructObjectFromClass.cpp
@@ -56,7 +56,7 @@ void UK2Node_ConstructObjectFromClass::AllocateDefaultPins()
 
 void UK2Node_ConstructObjectFromClass::SetPinToolTip(UEdGraphPin& MutatablePin, const FText& PinDescription) const
 {
-	MutatablePin.PinToolTip = UEdGraphSchema_K2::TypeToString(MutatablePin.PinType);
+	MutatablePin.PinToolTip = UEdGraphSchema_K2::TypeToText(MutatablePin.PinType).ToString();
 
 	UEdGraphSchema_K2 const* const K2Schema = Cast<const UEdGraphSchema_K2>(GetSchema());
 	if (K2Schema != nullptr)
diff --git a/Engine/Source/Editor/BlueprintGraph/Private/K2Node_EaseFunction.cpp b/Engine/Source/Editor/BlueprintGraph/Private/K2Node_EaseFunction.cpp
--- a/Engine/Source/Editor/BlueprintGraph/Private/K2Node_EaseFunction.cpp
+++ b/Engine/Source/Editor/BlueprintGraph/Private/K2Node_EaseFunction.cpp
@@ -127,7 +127,7 @@ FName UK2Node_EaseFunction::GetPaletteIcon(FLinearColor& OutColor) const
 
 void UK2Node_EaseFunction::SetPinToolTip(UEdGraphPin& MutatablePin, const FText& PinDescription) const
 {
-	MutatablePin.PinToolTip = UEdGraphSchema_K2::TypeToString(MutatablePin.PinType);
+	MutatablePin.PinToolTip = UEdGraphSchema_K2::TypeToText(MutatablePin.PinType).ToString();
 
 	UEdGraphSchema_K2 const* const K2Schema = Cast<const UEdGraphSchema_K2>(GetSchema());
 	if (K2Schema != nullptr)
diff --git a/Engine/Source/Editor/BlueprintGraph/Private/K2Node_GetDataTableRow.cpp b/Engine/Source/Editor/BlueprintGraph/Private/K2Node_GetDataTableRow.cpp
--- a/Engine/Source/Editor/BlueprintGraph/Private/K2Node_GetDataTableRow.cpp
+++ b/Engine/Source/Editor/BlueprintGraph/Private/K2Node_GetDataTableRow.cpp
@@ -55,7 +55,7 @@ void UK2Node_GetDataTableRow::AllocateDefaultPins()
 
 void UK2Node_GetDataTableRow::SetPinToolTip(UEdGraphPin& MutatablePin, const FText& PinDescription) const
 {
-	MutatablePin.PinToolTip = UEdGraphSchema_K2::TypeToString(MutatablePin.PinType);
+	MutatablePin.PinToolTip = UEdGraphSchema_K2::TypeToText(MutatablePin.PinType).ToString();
 
 	UEdGraphSchema_K2 const* const K2Schema = Cast<const UEdGraphSchema_K2>(GetSchema());
 	if (K2Schema != nullptr)
diff --git a/Engine/Source/Editor/BlueprintGraph/Private/K2Node_SpawnActor.cpp b/Engine/Source/Editor/BlueprintGraph/Private/K2Node_SpawnActor.cpp
--- a/Engine/Source/Editor/BlueprintGraph/Private/K2Node_SpawnActor.cpp
+++ b/Engine/Source/Editor/BlueprintGraph/Private/K2Node_SpawnActor.cpp
@@ -33,20 +33,20 @@ void UK2Node_SpawnActor::AllocateDefaultPins()
 
 	// Add blueprint pin
 	UEdGraphPin* BlueprintPin = CreatePin(EGPD_Input, K2Schema->PC_Object, TEXT(""), UBlueprint::StaticClass(), false, false, BlueprintPinName);
-	K2Schema->ConstructBasicPinTooltip(*BlueprintPin, LOCTEXT("BlueprintPinDescription", "The blueprint Actor you want to spawn").ToString(), BlueprintPin->PinToolTip);
+	K2Schema->ConstructBasicPinTooltip(*BlueprintPin, LOCTEXT("BlueprintPinDescription", "The blueprint Actor you want to spawn"), BlueprintPin->PinToolTip);
 
 	// Transform pin
 	UScriptStruct* TransformStruct = FindObjectChecked<UScriptStruct>(UObject::StaticClass(), TEXT("Transform"));
 	UEdGraphPin* TransformPin = CreatePin(EGPD_Input, K2Schema->PC_Struct, TEXT(""), TransformStruct, false, false, SpawnTransformPinName);
-	K2Schema->ConstructBasicPinTooltip(*TransformPin, LOCTEXT("TransformPinDescription", "The transform to spawn the Actor with").ToString(), TransformPin->PinToolTip);
+	K2Schema->ConstructBasicPinTooltip(*TransformPin, LOCTEXT("TransformPinDescription", "The transform to spawn the Actor with"), TransformPin->PinToolTip);
 
 	// bNoCollisionFail pin
 	UEdGraphPin* NoCollisionFailPin = CreatePin(EGPD_Input, K2Schema->PC_Boolean, TEXT(""), NULL, false, false, NoCollisionFailPinName);
-	K2Schema->ConstructBasicPinTooltip(*NoCollisionFailPin, LOCTEXT("NoCollisionFailPinDescription", "Determines if the Actor should be spawned when the location is blocked by a collision").ToString(), NoCollisionFailPin->PinToolTip);
+	K2Schema->ConstructBasicPinTooltip(*NoCollisionFailPin, LOCTEXT("NoCollisionFailPinDescription", "Determines if the Actor should be spawned when the location is blocked by a collision"), NoCollisionFailPin->PinToolTip);
 
 	// Result pin
 	UEdGraphPin* ResultPin = CreatePin(EGPD_Output, K2Schema->PC_Object, TEXT(""), AActor::StaticClass(), false, false, K2Schema->PN_ReturnValue);
-	K2Schema->ConstructBasicPinTooltip(*ResultPin, LOCTEXT("ResultPinDescription", "The spawned Actor").ToString(), ResultPin->PinToolTip);
+	K2Schema->ConstructBasicPinTooltip(*ResultPin, LOCTEXT("ResultPinDescription", "The spawned Actor"), ResultPin->PinToolTip);
 
 	Super::AllocateDefaultPins();
 }
diff --git a/Engine/Source/Editor/BlueprintGraph/Private/K2Node_SpawnActorFromClass.cpp b/Engine/Source/Editor/BlueprintGraph/Private/K2Node_SpawnActorFromClass.cpp
--- a/Engine/Source/Editor/BlueprintGraph/Private/K2Node_SpawnActorFromClass.cpp
+++ b/Engine/Source/Editor/BlueprintGraph/Private/K2Node_SpawnActorFromClass.cpp
@@ -44,38 +44,24 @@ void UK2Node_SpawnActorFromClass::AllocateDefaultPins()
 
 	// Add blueprint pin
 	UEdGraphPin* ClassPin = CreatePin(EGPD_Input, K2Schema->PC_Class, TEXT(""), AActor::StaticClass(), false, false, FK2Node_SpawnActorFromClassHelper::ClassPinName);
-	SetPinToolTip(*ClassPin, LOCTEXT("ClassPinDescription", "The Actor class you want to spawn"));
+	K2Schema->ConstructBasicPinTooltip(*ClassPin, LOCTEXT("ClassPinDescription", "The Actor class you want to spawn"), ClassPin->PinToolTip);
 
 	// Transform pin
 	UScriptStruct* TransformStruct = FindObjectChecked<UScriptStruct>(UObject::StaticClass(), TEXT("Transform"));
 	UEdGraphPin* TransformPin = CreatePin(EGPD_Input, K2Schema->PC_Struct, TEXT(""), TransformStruct, false, false, FK2Node_SpawnActorFromClassHelper::SpawnTransformPinName);
-	SetPinToolTip(*TransformPin, LOCTEXT("TransformPinDescription", "The transform to spawn the Actor with"));
+	K2Schema->ConstructBasicPinTooltip(*TransformPin, LOCTEXT("TransformPinDescription", "The transform to spawn the Actor with"), TransformPin->PinToolTip);
 
 	// bNoCollisionFail pin
 	UEdGraphPin* NoCollisionFailPin = CreatePin(EGPD_Input, K2Schema->PC_Boolean, TEXT(""), NULL, false, false, FK2Node_SpawnActorFromClassHelper::NoCollisionFailPinName);
-	SetPinToolTip(*NoCollisionFailPin, LOCTEXT("NoCollisionFailPinDescription", "Determines if the Actor should be spawned when the location is blocked by a collision"));
+	K2Schema->ConstructBasicPinTooltip(*NoCollisionFailPin, LOCTEXT("NoCollisionFailPinDescription", "Determines if the Actor should be spawned when the location is blocked by a collision"), NoCollisionFailPin->PinToolTip);
 
 	// Result pin
 	UEdGraphPin* ResultPin = CreatePin(EGPD_Output, K2Schema->PC_Object, TEXT(""), AActor::StaticClass(), false, false, K2Schema->PN_ReturnValue);
-	SetPinToolTip(*ResultPin, LOCTEXT("ResultPinDescription", "The spawned Actor"));
+	K2Schema->ConstructBasicPinTooltip(*ResultPin, LOCTEXT("ResultPinDescription", "The spawned Actor"), ResultPin->PinToolTip);
 
 	Super::AllocateDefaultPins();
 }
 
-void UK2Node_SpawnActorFromClass::SetPinToolTip(UEdGraphPin& MutatablePin, const FText& PinDescription) const
-{
-	MutatablePin.PinToolTip = UEdGraphSchema_K2::TypeToString(MutatablePin.PinType);
-
-	UEdGraphSchema_K2 const* const K2Schema = Cast<const UEdGraphSchema_K2>(GetSchema());
-	if (K2Schema != nullptr)
-	{
-		MutatablePin.PinToolTip += TEXT(" ");
-		MutatablePin.PinToolTip += K2Schema->GetPinDisplayName(&MutatablePin);
-	}
-
-	MutatablePin.PinToolTip += FString(TEXT("\n")) + PinDescription.ToString();
-}
-
 void UK2Node_SpawnActorFromClass::CreatePinsForClass(UClass* InClass)
 {
 	check(InClass != NULL);
diff --git a/Engine/Source/Editor/GraphEditor/Private/KismetNodes/SGraphNodeK2Base.cpp b/Engine/Source/Editor/GraphEditor/Private/KismetNodes/SGraphNodeK2Base.cpp
--- a/Engine/Source/Editor/GraphEditor/Private/KismetNodes/SGraphNodeK2Base.cpp
+++ b/Engine/Source/Editor/GraphEditor/Private/KismetNodes/SGraphNodeK2Base.cpp
@@ -433,7 +433,7 @@ void SGraphNodeK2Base::GetNodeInfoPopups(FNodeInfoContext* Context, TArray<FGrap
 						PinnedWatchText += TEXT("\n");
 					}
 
-					FString PinName = UEdGraphSchema_K2::TypeToString(WatchPin->PinType);
+					FString PinName = UEdGraphSchema_K2::TypeToText(WatchPin->PinType).ToString();
 					PinName += TEXT(" ");
 					PinName += Schema->GetPinDisplayName(WatchPin);
 
diff --git a/Engine/Source/Editor/Kismet/Private/SBlueprintActionMenu.cpp b/Engine/Source/Editor/Kismet/Private/SBlueprintActionMenu.cpp
--- a/Engine/Source/Editor/Kismet/Private/SBlueprintActionMenu.cpp
+++ b/Engine/Source/Editor/Kismet/Private/SBlueprintActionMenu.cpp
@@ -353,9 +353,9 @@ FText SBlueprintActionMenu::GetSearchContextDesc() const
 		else
 		{
 			// Get the type string
-			const FString TypeStringRaw = UEdGraphSchema_K2::TypeToString(OnePin->PinType);
+			const FString TypeStringRaw = UEdGraphSchema_K2::TypeToText(OnePin->PinType).ToString();
 
-			//@TODO: Add a parameter to TypeToString indicating the kind of formating requested
+			//@TODO: Add a parameter to TypeToText indicating the kind of formating requested
 			const FString TypeString = (TypeStringRaw.Replace(TEXT("'"), TEXT(" "))).TrimTrailing();
 
 			if (OnePin->Direction == EGPD_Input)
diff --git a/Engine/Source/Editor/Kismet/Private/SBlueprintPalette.cpp b/Engine/Source/Editor/Kismet/Private/SBlueprintPalette.cpp
--- a/Engine/Source/Editor/Kismet/Private/SBlueprintPalette.cpp
+++ b/Engine/Source/Editor/Kismet/Private/SBlueprintPalette.cpp
@@ -85,7 +85,7 @@ static FString GetVarType(UStruct* VarScope, FName VarName, bool bUseObjToolTip,
 					}
 					else
 					{
-						VarDesc = UEdGraphSchema_K2::TypeToString(PinType);
+						VarDesc = UEdGraphSchema_K2::TypeToText(PinType).ToString();
 					}
 				}
 			}
diff --git a/Engine/Source/Editor/KismetCompiler/Private/KismetCompiler.cpp b/Engine/Source/Editor/KismetCompiler/Private/KismetCompiler.cpp
--- a/Engine/Source/Editor/KismetCompiler/Private/KismetCompiler.cpp
+++ b/Engine/Source/Editor/KismetCompiler/Private/KismetCompiler.cpp
@@ -373,7 +373,7 @@ UProperty* FKismetCompilerContext::CreateVariable(const FName VarName, const FEd
 	else
 	{
 		MessageLog.Error(*FString::Printf(*LOCTEXT("VariableInvalidType_Error", "The variable %s declared in @@ has an invalid type %s").ToString(),
-			*VarName.ToString(), *UEdGraphSchema_K2::TypeToString(VarType)), Blueprint);
+			*VarName.ToString(), *UEdGraphSchema_K2::TypeToText(VarType).ToString()), Blueprint);
 	}
 
 	return NewProperty;
@@ -583,7 +583,7 @@ void FKismetCompilerContext::CreatePropertiesFromList(UStruct* Scope, UField**&
 			{
 				continue;
 			}
-			MessageLog.Warning(*FString::Printf(*LOCTEXT("AssociatedVarProperty_Error", "AssociatedVarProperty property overriden %s from @@ type (%s)").ToString(), *Term.Name, *UEdGraphSchema_K2::TypeToString(Term.Type)), Term.Source);
+			MessageLog.Warning(*FString::Printf(*LOCTEXT("AssociatedVarProperty_Error", "AssociatedVarProperty property overriden %s from @@ type (%s)").ToString(), *Term.Name, *UEdGraphSchema_K2::TypeToText(Term.Type).ToString()), Term.Source);
 		}
 
 		UProperty* NewProperty = FKismetCompilerUtilities::CreatePropertyOnScope(Scope, FName(*Term.Name), Term.Type, NewClass, PropertyFlags, Schema, MessageLog);
@@ -677,7 +677,7 @@ void FKismetCompilerContext::CreatePropertiesFromList(UStruct* Scope, UField**&
 		}
 		else
 		{
-			MessageLog.Error(*FString::Printf(*LOCTEXT("FailedCreateProperty_Error", "Failed to create property %s from @@ due to a bad or unknown type (%s)").ToString(), *Term.Name, *UEdGraphSchema_K2::TypeToString(Term.Type)),
+			MessageLog.Error(*FString::Printf(*LOCTEXT("FailedCreateProperty_Error", "Failed to create property %s from @@ due to a bad or unknown type (%s)").ToString(), *Term.Name, *UEdGraphSchema_K2::TypeToText(Term.Type).ToString()),
 				Term.Source);
 		}
 
diff --git a/Engine/Source/Editor/KismetCompiler/Private/KismetCompilerMisc.cpp b/Engine/Source/Editor/KismetCompiler/Private/KismetCompilerMisc.cpp
--- a/Engine/Source/Editor/KismetCompiler/Private/KismetCompilerMisc.cpp
+++ b/Engine/Source/Editor/KismetCompiler/Private/KismetCompilerMisc.cpp
@@ -247,15 +247,15 @@ bool FKismetCompilerUtilities::IsTypeCompatibleWithProperty(UEdGraphPin* SourceP
 	}
 	else
 	{
-		MessageLog.Error(*FString::Printf(*LOCTEXT("UnsupportedTypeForPin", "Unsupported type (%s) on @@").ToString(), *UEdGraphSchema_K2::TypeToString(Type)), SourcePin);
+		MessageLog.Error(*FString::Printf(*LOCTEXT("UnsupportedTypeForPin", "Unsupported type (%s) on @@").ToString(), *UEdGraphSchema_K2::TypeToText(Type).ToString()), SourcePin);
 	}
 
 	if (bTypeMismatch)
 	{
 		MessageLog.Error(*FString::Printf(*LOCTEXT("TypeDoesNotMatchPropertyOfType_Error", "@@ of type %s doesn't match the property %s of type %s").ToString(),
-			*UEdGraphSchema_K2::TypeToString(Type),
+			*UEdGraphSchema_K2::TypeToText(Type).ToString(),
 			*Property->GetName(),
-			*UEdGraphSchema_K2::TypeToString(Property)),
+			*UEdGraphSchema_K2::TypeToText(Property).ToString()),
 			SourcePin);
 	}
 
diff --git a/Engine/Source/Editor/KismetCompiler/Private/UserDefinedStructureCompilerUtils.cpp b/Engine/Source/Editor/KismetCompiler/Private/UserDefinedStructureCompilerUtils.cpp
--- a/Engine/Source/Editor/KismetCompiler/Private/UserDefinedStructureCompilerUtils.cpp
+++ b/Engine/Source/Editor/KismetCompiler/Private/UserDefinedStructureCompilerUtils.cpp
@@ -142,7 +142,7 @@ struct FUserDefinedStructureCompilerInner
 			else
 			{
 				LogError(Struct, MessageLog, FString::Printf(*LOCTEXT("VariableInvalidType_Error", "The variable %s declared in %s has an invalid type %s").ToString(),
-					*VarDesc.VarName.ToString(), *Struct->GetName(), *UEdGraphSchema_K2::TypeToString(VarType)));
+					*VarDesc.VarName.ToString(), *Struct->GetName(), *UEdGraphSchema_K2::TypeToText(VarType).ToString()));
 				continue;
 			}
 
diff --git a/Engine/Source/Editor/KismetCompiler/Public/KismetCompiledFunctionContext.h b/Engine/Source/Editor/KismetCompiler/Public/KismetCompiledFunctionContext.h
--- a/Engine/Source/Editor/KismetCompiler/Public/KismetCompiledFunctionContext.h
+++ b/Engine/Source/Editor/KismetCompiler/Public/KismetCompiledFunctionContext.h
@@ -528,7 +528,7 @@ public:
 			}
 			else
 			{
-				MessageLog.Error(*FString::Printf(TEXT("Expected @@ to %s instead of %s"), *Schema->TypeToString(TestType), *Schema->TypeToString(Pin->PinType)), Pin);
+				MessageLog.Error(*FString::Printf(TEXT("Expected @@ to %s instead of %s"), *Schema->TypeToText(TestType).ToString(), *Schema->TypeToText(Pin->PinType).ToString()), Pin);
 				return false;
 			}
 		}
diff --git a/Engine/Source/Editor/UnrealEd/Private/Commandlets/DumpBlueprintsInfoCommandlet.cpp b/Engine/Source/Editor/UnrealEd/Private/Commandlets/DumpBlueprintsInfoCommandlet.cpp
--- a/Engine/Source/Editor/UnrealEd/Private/Commandlets/DumpBlueprintsInfoCommandlet.cpp
+++ b/Engine/Source/Editor/UnrealEd/Private/Commandlets/DumpBlueprintsInfoCommandlet.cpp
@@ -1680,7 +1680,7 @@ static void DumpBlueprintInfoUtils::DumpContextActionList(uint32 Indent, FBluepr
 		{
 			PinTypeLog = "OUTPUT";
 		}
-		PinTypeLog += UEdGraphSchema_K2::TypeToString(ActionBuilder.FromPin->PinType);
+		PinTypeLog += UEdGraphSchema_K2::TypeToText(ActionBuilder.FromPin->PinType).ToString();
 
 		UE_LOG(LogBlueprintInfoDump, Display, TEXT("%sDumping pin actions: %s"), *BuildIndentString(Indent, true), *PinTypeLog);
 	}
@@ -1765,7 +1765,7 @@ static void DumpBlueprintInfoUtils::DumpContextInfo(uint32 Indent, FBlueprintGra
 	}
 	else
 	{
-		ContextEntry += UEdGraphSchema_K2::TypeToString(ActionBuilder.FromPin->PinType);
+		ContextEntry += UEdGraphSchema_K2::TypeToText(ActionBuilder.FromPin->PinType).ToString();
 
 		ContextEntry += "\"," + IndentedNewline + "\"PinDirection\" : \"";
 		if (ActionBuilder.FromPin->Direction == EGPD_Input)
diff --git a/Engine/Source/Runtime/Engine/Classes/EdGraph/EdGraphSchema.h b/Engine/Source/Runtime/Engine/Classes/EdGraph/EdGraphSchema.h
--- a/Engine/Source/Runtime/Engine/Classes/EdGraph/EdGraphSchema.h
+++ b/Engine/Source/Runtime/Engine/Classes/EdGraph/EdGraphSchema.h
@@ -575,7 +575,7 @@ class ENGINE_API UEdGraphSchema : public UObject
 	 * @param   PinDescription	A detailed description, describing the pin's purpose
 	 * @param   TooltipOut		The constructed tool-tip (out)
 	 */
-	virtual void ConstructBasicPinTooltip(UEdGraphPin const& Pin, FString const& PinDescription, FString& TooltipOut) const;
+	virtual void ConstructBasicPinTooltip(UEdGraphPin const& Pin, FText const& PinDescription, FString& TooltipOut) const;
 
 	/** @return     The type of graph (function vs. ubergraph) that this that TestEdGraph is. */
 	//@TODO: This is too K2-specific to be included in EdGraphSchema and should be refactored
diff --git a/Engine/Source/Runtime/Engine/Private/EdGraph/EdGraphSchema.cpp b/Engine/Source/Runtime/Engine/Private/EdGraph/EdGraphSchema.cpp
--- a/Engine/Source/Runtime/Engine/Private/EdGraph/EdGraphSchema.cpp
+++ b/Engine/Source/Runtime/Engine/Private/EdGraph/EdGraphSchema.cpp
@@ -482,9 +482,9 @@ FString UEdGraphSchema::GetPinDisplayName(const UEdGraphPin* Pin) const
 	return !Pin->PinFriendlyName.IsEmpty() ? Pin->PinFriendlyName.ToString() : Pin->PinName;
 }
 
-void UEdGraphSchema::ConstructBasicPinTooltip(UEdGraphPin const& Pin, FString const& PinDescription, FString& TooltipOut) const
+void UEdGraphSchema::ConstructBasicPinTooltip(UEdGraphPin const& Pin, FText const& PinDescription, FString& TooltipOut) const
 {
-	TooltipOut = PinDescription;
+	TooltipOut = PinDescription.ToString();
 }
 
 void UEdGraphSchema::GetGraphContextActions(FGraphContextMenuBuilder& ContextMenuBuilder) const
