Missed file from CL# 2286449
Remove the frequently incorrect and not very informative (NO DEBUGGING INFORMATION GENERATED, NEED TO RECOMPILE THE BLUEPRINT) from graph pin tooltips

diff --git a/Engine/Source/Editor/BlueprintGraph/Private/K2Node.cpp b/Engine/Source/Editor/BlueprintGraph/Private/K2Node.cpp
--- a/Engine/Source/Editor/BlueprintGraph/Private/K2Node.cpp
+++ b/Engine/Source/Editor/BlueprintGraph/Private/K2Node.cpp
@@ -1009,7 +1009,7 @@ void FOptionalPinManager::CreateVisiblePins(TArray<FOptionalPinFromProperty>& Pr
 							const FString PinName = PinFriendlyName.ToString();
 							NewPin = TargetNode->CreatePin(Direction, PinType.PinCategory, PinType.PinSubCategory, PinType.PinSubCategoryObject.Get(), PinType.bIsArray, PinType.bIsReference, PinName);
 							NewPin->PinFriendlyName = PinFriendlyName;
-							Schema->ConstructBasicPinTooltip(*NewPin, TEXT("\n") + PropertyEntry.PropertyTooltip.ToString(), NewPin->PinToolTip);
+							Schema->ConstructBasicPinTooltip(*NewPin, PropertyEntry.PropertyTooltip, NewPin->PinToolTip);
 
 							// Allow the derived class to customize the created pin
 							CustomizePinData(NewPin, PropertyEntry.PropertyName, Index, InnerProperty);
@@ -1042,7 +1042,7 @@ void FOptionalPinManager::CreateVisiblePins(TArray<FOptionalPinFromProperty>& Pr
 						const FString PinName = PropertyEntry.PropertyName.ToString();
 						NewPin = TargetNode->CreatePin(Direction, PinType.PinCategory, PinType.PinSubCategory, PinType.PinSubCategoryObject.Get(), PinType.bIsArray, PinType.bIsReference, PinName);
 						NewPin->PinFriendlyName = PropertyEntry.PropertyFriendlyName.IsEmpty() ? FText::FromString(PinName) : FText::FromString(PropertyEntry.PropertyFriendlyName);
-						Schema->ConstructBasicPinTooltip(*NewPin, TEXT("\n") + PropertyEntry.PropertyTooltip.ToString(), NewPin->PinToolTip);
+						Schema->ConstructBasicPinTooltip(*NewPin, PropertyEntry.PropertyTooltip, NewPin->PinToolTip);
 
 						// Allow the derived class to customize the created pin
 						CustomizePinData(NewPin, PropertyEntry.PropertyName, INDEX_NONE, OuterProperty);
@@ -1135,7 +1135,7 @@ void UK2Node::GetPinHoverText(const UEdGraphPin& Pin, FString& HoverTextOut) con
 	{
 		UEdGraphSchema const* Schema = GetSchema();
 		check(Schema != NULL);
-		Schema->ConstructBasicPinTooltip(Pin, TEXT(""), HoverTextOut);
+		Schema->ConstructBasicPinTooltip(Pin, FText::GetEmpty(), HoverTextOut);
 	}	
 
 	UBlueprint* Blueprint = FBlueprintEditorUtils::FindBlueprintForNodeChecked(this);
@@ -1162,10 +1162,10 @@ void UK2Node::GetPinHoverText(const UEdGraphPin& Pin, FString& HoverTextOut) con
 	}
 
 	// if the blueprint doesn't have debug data, notify the user
-	if (!FKismetDebugUtilities::HasDebuggingData(Blueprint))
+	/*if (!FKismetDebugUtilities::HasDebuggingData(Blueprint))
 	{
 		HoverTextOut += TEXT("\n(NO DEBUGGING INFORMATION GENERATED, NEED TO RECOMPILE THE BLUEPRINT)");
-	}
+	}*/
 
 	//@TODO: For exec pins, show when they were last executed
 
