Export Type text from Generate Blueprint API Commandlet
Display type text separate from tooltip in blueprint API docs
Fix filename for json output when no explicit filename provided

diff --git a/Engine/Source/Editor/BlueprintGraph/Classes/EdGraphSchema_K2.h b/Engine/Source/Editor/BlueprintGraph/Classes/EdGraphSchema_K2.h
--- a/Engine/Source/Editor/BlueprintGraph/Classes/EdGraphSchema_K2.h
+++ b/Engine/Source/Editor/BlueprintGraph/Classes/EdGraphSchema_K2.h
@@ -277,6 +277,9 @@ class BLUEPRINTGRAPH_API UEdGraphSchema_K2 : public UEdGraphSchema
 	// action grouping values
 	static const int32 AG_LevelReference;
 
+	// Somewhat hacky mechanism to prevent tooltips created for pins from including the display name and type when generating BP API documentation
+	static bool bGeneratingDocumentation;
+
 	/** Whether this schema should use the old (legacy) menu building, or the new (experimental) menu building */
 	UPROPERTY(GlobalConfig)
 	bool bUseLegacyActionMenus;
diff --git a/Engine/Source/Editor/BlueprintGraph/Private/EdGraphSchema_K2.cpp b/Engine/Source/Editor/BlueprintGraph/Private/EdGraphSchema_K2.cpp
--- a/Engine/Source/Editor/BlueprintGraph/Private/EdGraphSchema_K2.cpp
+++ b/Engine/Source/Editor/BlueprintGraph/Private/EdGraphSchema_K2.cpp
@@ -300,6 +300,8 @@ const UScriptStruct* UEdGraphSchema_K2::TransformStruct = nullptr;
 const UScriptStruct* UEdGraphSchema_K2::LinearColorStruct = nullptr;
 const UScriptStruct* UEdGraphSchema_K2::ColorStruct = nullptr;
 
+bool UEdGraphSchema_K2::bGeneratingDocumentation = false;
+
 UEdGraphSchema_K2::UEdGraphSchema_K2(const class FPostConstructInitializeProperties& PCIP)
 	: Super(PCIP)
 {
@@ -2194,37 +2196,44 @@ FString UEdGraphSchema_K2::GetPinDisplayName(const UEdGraphPin* Pin) const
 
 void UEdGraphSchema_K2::ConstructBasicPinTooltip(const UEdGraphPin& Pin, const FText& PinDescription, FString& TooltipOut) const
 {
-	FFormatNamedArguments Args;
-	Args.Add(TEXT("PinType"), TypeToText(Pin.PinType));
-
-	if (UEdGraphNode* PinNode = Pin.GetOwningNode())
+	if (bGeneratingDocumentation)
 	{
-		UEdGraphSchema_K2 const* const K2Schema = Cast<const UEdGraphSchema_K2>(PinNode->GetSchema());
-		if (ensure(K2Schema != NULL)) // ensure that this node belongs to this schema
-		{
-			Args.Add(TEXT("DisplayName"), FText::FromString(GetPinDisplayName(&Pin)));
-			Args.Add(TEXT("LineFeed1"), FText::FromString(TEXT("\n")));
-		}
+		TooltipOut = PinDescription.ToString();
 	}
 	else
 	{
-			Args.Add(TEXT("DisplayName"), FText::GetEmpty());
-			Args.Add(TEXT("LineFeed1"), FText::GetEmpty());
-	}
+		FFormatNamedArguments Args;
+		Args.Add(TEXT("PinType"), TypeToText(Pin.PinType));
 
+		if (UEdGraphNode* PinNode = Pin.GetOwningNode())
+		{
+			UEdGraphSchema_K2 const* const K2Schema = Cast<const UEdGraphSchema_K2>(PinNode->GetSchema());
+			if (ensure(K2Schema != NULL)) // ensure that this node belongs to this schema
+			{
+				Args.Add(TEXT("DisplayName"), FText::FromString(GetPinDisplayName(&Pin)));
+				Args.Add(TEXT("LineFeed1"), FText::FromString(TEXT("\n")));
+			}
+		}
+		else
+		{
+				Args.Add(TEXT("DisplayName"), FText::GetEmpty());
+				Args.Add(TEXT("LineFeed1"), FText::GetEmpty());
+		}
 
-	if (!PinDescription.IsEmpty())
-	{
-		Args.Add(TEXT("Description"), PinDescription);
-		Args.Add(TEXT("LineFeed2"), FText::FromString(TEXT("\n\n")));
-	}
-	else
-	{
-		Args.Add(TEXT("Description"), FText::GetEmpty());
-		Args.Add(TEXT("LineFeed2"), FText::GetEmpty());
-	}
+
+		if (!PinDescription.IsEmpty())
+		{
+			Args.Add(TEXT("Description"), PinDescription);
+			Args.Add(TEXT("LineFeed2"), FText::FromString(TEXT("\n\n")));
+		}
+		else
+		{
+			Args.Add(TEXT("Description"), FText::GetEmpty());
+			Args.Add(TEXT("LineFeed2"), FText::GetEmpty());
+		}
 	
-	TooltipOut = FText::Format(LOCTEXT("PinTooltip", "{DisplayName}{LineFeed1}{PinType}{LineFeed2}{Description}"), Args).ToString(); 
+		TooltipOut = FText::Format(LOCTEXT("PinTooltip", "{DisplayName}{LineFeed1}{PinType}{LineFeed2}{Description}"), Args).ToString(); 
+	}
 }
 
 EGraphType UEdGraphSchema_K2::GetGraphType(const UEdGraph* TestEdGraph) const
diff --git a/Engine/Source/Editor/UnrealEd/Private/Commandlets/GenerateBlueprintAPICommandlet.cpp b/Engine/Source/Editor/UnrealEd/Private/Commandlets/GenerateBlueprintAPICommandlet.cpp
--- a/Engine/Source/Editor/UnrealEd/Private/Commandlets/GenerateBlueprintAPICommandlet.cpp
+++ b/Engine/Source/Editor/UnrealEd/Private/Commandlets/GenerateBlueprintAPICommandlet.cpp
@@ -462,7 +462,7 @@ static NodeType* GenerateBlueprintAPIUtils::AddNodeToGraph(UEdGraph* Graph)
 //------------------------------------------------------------------------------
 static FString GenerateBlueprintAPIUtils::BuildDumpFilePath(UClass* BlueprintClass)
 {
-	FString CommandletSaveDir, Pathname;
+	FString CommandletSaveDir, Filename;
 
 	if (CommandOptions.SaveDir.IsEmpty())
 	{
@@ -477,19 +477,19 @@ static FString GenerateBlueprintAPIUtils::BuildDumpFilePath(UClass* BlueprintCla
 	
 	if (CommandOptions.Filename.IsEmpty())
 	{
-		FString Pathname = FString::Printf(TEXT("GenerateBlueprintAPI_%s"), FPlatformTime::StrTimestamp());
-		Pathname = Pathname.Replace(TEXT(" "), TEXT("_"));
-		Pathname = Pathname.Replace(TEXT("/"), TEXT("-"));
-		Pathname = Pathname.Replace(TEXT(":"), TEXT("."));
+		Filename = FString::Printf(TEXT("GenerateBlueprintAPI_%s"), FPlatformTime::StrTimestamp());
+		Filename = Filename.Replace(TEXT(" "), TEXT("_"));
+		Filename = Filename.Replace(TEXT("/"), TEXT("-"));
+		Filename = Filename.Replace(TEXT(":"), TEXT("."));
 	}
 	else
 	{
-		Pathname = CommandOptions.Filename;
+		Filename = CommandOptions.Filename;
 	}
 
-	Pathname += ".json";
+	Filename += ".json";
 
-	return CommandletSaveDir / *Pathname;
+	return CommandletSaveDir / *Filename;
 }
 
 //------------------------------------------------------------------------------
@@ -849,6 +849,8 @@ static void GenerateBlueprintAPIUtils::DumpActionMenuItem(uint32 Indent, FGraphA
 					}
 					ActionEntry += PinDetailsIndentedNewline + "\"Direction\"            : \"" + (Pin->Direction == EGPD_Input ? "input" : "output") + "\"";
 
+					ActionEntry += PinDetailsIndentedNewline + "\"TypeText\"             : \"" + UEdGraphSchema_K2::TypeToText(Pin->PinType).ToString() + "\"";
+
 					if (!Pin->PinToolTip.IsEmpty())
 					{
 						const FString PinTooltipFieldLabel("\"Tooltip\"              : \"");
@@ -913,6 +915,8 @@ UGenerateBlueprintAPICommandlet::UGenerateBlueprintAPICommandlet(class FPostCons
 //------------------------------------------------------------------------------
 int32 UGenerateBlueprintAPICommandlet::Main(FString const& Params)
 {
+	UEdGraphSchema_K2::bGeneratingDocumentation = true;
+
 	TArray<FString> Tokens, Switches;
 	ParseCommandLine(*Params, Tokens, Switches);
 
diff --git a/Engine/Source/Programs/UnrealDocTool/APIDocTool/APIDocTool/APIAction.cs b/Engine/Source/Programs/UnrealDocTool/APIDocTool/APIDocTool/APIAction.cs
--- a/Engine/Source/Programs/UnrealDocTool/APIDocTool/APIDocTool/APIAction.cs
+++ b/Engine/Source/Programs/UnrealDocTool/APIDocTool/APIDocTool/APIAction.cs
@@ -37,6 +37,14 @@ namespace APIDocTool
 			}
 		}
 
+		private void WritePins(UdnWriter Writer, IEnumerable<APIActionPin> Pins)
+		{
+			foreach(var Pin in Pins)
+			{
+				Pin.WritePin(Writer);
+			}
+		}
+
 		public override void WritePage(UdnManifest Manifest, string OutputPath)
 		{
 			using (UdnWriter Writer = new UdnWriter(OutputPath))
@@ -50,12 +58,16 @@ namespace APIDocTool
 				{
 					// Inputs
 					Writer.EnterSection("inputs", "Inputs");
-					Writer.WriteList(Pins.Where(x => x.bInputPin).Select(x => x.GetListItem()));
+					Writer.WriteObject("MemberIconListHeadBlank");
+					WritePins(Writer, Pins.Where(x => x.bInputPin));
+					Writer.WriteObject("MemberIconListTail");
 					Writer.LeaveSection();
 
 					// Outputs
 					Writer.EnterSection("outputs", "Outputs");
-					Writer.WriteList(Pins.Where(x => !x.bInputPin).Select(x => x.GetListItem()));
+					Writer.WriteObject("MemberIconListHeadBlank");
+					WritePins(Writer, Pins.Where(x => !x.bInputPin));
+					Writer.WriteObject("MemberIconListTail");
 					Writer.LeaveSection();
 				}
 			}
diff --git a/Engine/Source/Programs/UnrealDocTool/APIDocTool/APIDocTool/APIActionPin.cs b/Engine/Source/Programs/UnrealDocTool/APIDocTool/APIDocTool/APIActionPin.cs
--- a/Engine/Source/Programs/UnrealDocTool/APIDocTool/APIDocTool/APIActionPin.cs
+++ b/Engine/Source/Programs/UnrealDocTool/APIDocTool/APIDocTool/APIActionPin.cs
@@ -10,6 +10,7 @@ namespace APIDocTool
 	class APIActionPin : APIPage
 	{
 		public readonly string Tooltip;
+		public readonly string TypeText;
 
 		public readonly string PinCategory;
 		public readonly string PinSubCategory;
@@ -48,6 +49,7 @@ namespace APIDocTool
 
 			bInputPin = (string)PinProperties["Direction"] == "input";
 			PinCategory = (string)PinProperties["PinCategory"];
+			TypeText = (string)PinProperties["TypeText"];
 
 			if (PinProperties.TryGetValue("Tooltip", out Value))
 			{
@@ -132,10 +134,10 @@ namespace APIDocTool
 			Manifest.Add(Name, this);
 		}
 
-		public UdnIconListItem GetListItem()
+		public void WritePin(UdnWriter Writer)
 		{
+			// Get all the icons
 			List<Icon> ItemIcons = new List<Icon>();
-
 			if (bIsArray)
 			{
 				ItemIcons.Add(Icons.ArrayVariablePinIcons[GetPinCategory()]);
@@ -145,7 +147,7 @@ namespace APIDocTool
 				ItemIcons.Add(Icons.VariablePinIcons[GetPinCategory()]);
 			}
 
-			return new UdnIconListItem(ItemIcons, Name, Tooltip, null);
+			Writer.WriteObject("ActionPinListItem", "icons", ItemIcons, "name", Name, "type", TypeText, "tooltip", Tooltip);
 		}
 	}
 }
