PinSelector drop down now uses same labels for variables as the pin tooltips do

diff --git a/Engine/Source/Editor/BlueprintGraph/Classes/EdGraphSchema_K2.h b/Engine/Source/Editor/BlueprintGraph/Classes/EdGraphSchema_K2.h
--- a/Engine/Source/Editor/BlueprintGraph/Classes/EdGraphSchema_K2.h
+++ b/Engine/Source/Editor/BlueprintGraph/Classes/EdGraphSchema_K2.h
@@ -299,7 +299,7 @@ public:
 		{
 		}
 
-		void Init(const FString& FriendlyCategoryName, const FString& CategoryName, const UEdGraphSchema_K2* Schema, const FString& InTooltip, bool bInReadOnly);
+		void Init(const FText& FriendlyCategoryName, const FString& CategoryName, const UEdGraphSchema_K2* Schema, const FText& InTooltip, bool bInReadOnly);
 	
 	private:
 		/** The pin type corresponding to the schema type */
@@ -316,10 +316,10 @@ public:
 		bool bReadOnly;
 
 		/** Friendly display name of pin type; also used to see if it has subtypes */
-		FString FriendlyName;
+		FText FriendlyName;
 
 		/** Text for regular tooltip */
-		FString Tooltip;
+		FText Tooltip;
 
 	public:
 		const FEdGraphPinType& GetPinType(bool bForceLoadedSubCategoryObject);
@@ -328,10 +328,9 @@ public:
 			PinType.PinSubCategory = SubCategory;
 		}
 
-		FPinTypeTreeInfo(const FString& FriendlyName, const FString& CategoryName, const UEdGraphSchema_K2* Schema, const FString& InTooltip, bool bInReadOnly = false);
-		FPinTypeTreeInfo(const FString& CategoryName, const UEdGraphSchema_K2* Schema, const FString& InTooltip, bool bInReadOnly = false);
-		FPinTypeTreeInfo(const FString& CategoryName, UObject* SubCategoryObject, const FString& InTooltip, bool bInReadOnly = false);
-		FPinTypeTreeInfo(const FString& CategoryName, const FStringAssetReference& SubCategoryObject, const FString& InTooltip, bool bInReadOnly = false);
+		FPinTypeTreeInfo(const FText& FriendlyName, const FString& CategoryName, const UEdGraphSchema_K2* Schema, const FText& InTooltip, bool bInReadOnly = false);
+		FPinTypeTreeInfo(const FString& CategoryName, UObject* SubCategoryObject, const FText& InTooltip, bool bInReadOnly = false);
+		FPinTypeTreeInfo(const FString& CategoryName, const FStringAssetReference& SubCategoryObject, const FText& InTooltip, bool bInReadOnly = false);
 
 		FPinTypeTreeInfo(TSharedPtr<FPinTypeTreeInfo> InInfo)
 		{
@@ -343,17 +342,17 @@ public:
 		}
 		
 		/** Returns a succinct menu description of this type */
-		FString GetDescription() const;
+		FText GetDescription() const;
 
-		FString GetToolTip() const
+		FText GetToolTip() const
 		{
 			if (PinType.PinSubCategoryObject.IsValid())
 			{
-				if (Tooltip.IsEmpty() || (PinType.PinSubCategoryObject->GetName() == Tooltip))
+				if (Tooltip.IsEmpty())
 				{
-					if ( (PinType.PinCategory == TEXT("struct")) && PinType.PinSubCategoryObject->IsA(UScriptStruct::StaticClass()) )
+					if ( (PinType.PinCategory == UEdGraphSchema_K2::PC_Struct) && PinType.PinSubCategoryObject->IsA(UScriptStruct::StaticClass()) )
 					{
-						return PinType.PinSubCategoryObject->GetPathName();
+						return FText::FromString(PinType.PinSubCategoryObject->GetPathName());
 					}
 				}
 			}
diff --git a/Engine/Source/Editor/BlueprintGraph/Private/EdGraphSchema_K2.cpp b/Engine/Source/Editor/BlueprintGraph/Private/EdGraphSchema_K2.cpp
--- a/Engine/Source/Editor/BlueprintGraph/Private/EdGraphSchema_K2.cpp
+++ b/Engine/Source/Editor/BlueprintGraph/Private/EdGraphSchema_K2.cpp
@@ -89,16 +89,11 @@ const FName FBlueprintMetadata::MD_DataTablePin(TEXT("DataTablePin"));
 
 #define LOCTEXT_NAMESPACE "KismetSchema"
 
-UEdGraphSchema_K2::FPinTypeTreeInfo::FPinTypeTreeInfo(const FString& InFriendlyCategoryName, const FString& CategoryName, const UEdGraphSchema_K2* Schema, const FString& InTooltip, bool bInReadOnly/*=false*/)
+UEdGraphSchema_K2::FPinTypeTreeInfo::FPinTypeTreeInfo(const FText& InFriendlyCategoryName, const FString& CategoryName, const UEdGraphSchema_K2* Schema, const FText& InTooltip, bool bInReadOnly/*=false*/)
 {
 	Init(InFriendlyCategoryName, CategoryName, Schema, InTooltip, bInReadOnly);
 }
 
-UEdGraphSchema_K2::FPinTypeTreeInfo::FPinTypeTreeInfo(const FString& CategoryName, const UEdGraphSchema_K2* Schema, const FString& InTooltip, bool bInReadOnly/*=false*/)
-{
-	Init(CategoryName, CategoryName, Schema, InTooltip, bInReadOnly);
-}
-
 struct FGatherStructTypesFromAssetsHelper
 {
 	typedef TSharedPtr<UEdGraphSchema_K2::FPinTypeTreeInfo> FPinTypeTreeInfoPtr;
@@ -106,7 +101,7 @@ struct FGatherStructTypesFromAssetsHelper
 	{
 		FORCEINLINE bool operator()(const FPinTypeTreeInfoPtr A, const FPinTypeTreeInfoPtr B) const
 		{
-			return (A->GetDescription() < B->GetDescription());
+			return (A->GetDescription().ToString() < B->GetDescription().ToString());
 		}
 	};
 
@@ -123,8 +118,8 @@ struct FGatherStructTypesFromAssetsHelper
 				const FString* pDescription = Asset.TagsAndValues.Find(TEXT("Tooltip"));
 				const FString Tooltip = (pDescription && !pDescription->IsEmpty()) ? *pDescription : Asset.ObjectPath.ToString();
 
-				FPinTypeTreeInfoPtr TypeTreeInfo = MakeShareable(new UEdGraphSchema_K2::FPinTypeTreeInfo(CategoryName, Asset.ToStringReference(), Tooltip));
-				TypeTreeInfo->FriendlyName = Asset.AssetName.ToString();
+				FPinTypeTreeInfoPtr TypeTreeInfo = MakeShareable(new UEdGraphSchema_K2::FPinTypeTreeInfo(CategoryName, Asset.ToStringReference(), FText::FromString(Tooltip)));
+				TypeTreeInfo->FriendlyName = FText::FromName(Asset.AssetName);
 				OutChildren.Add(TypeTreeInfo);
 			}
 		}
@@ -153,43 +148,43 @@ const FEdGraphPinType& UEdGraphSchema_K2::FPinTypeTreeInfo::GetPinType(bool bFor
 	return PinType;
 }
 
-void UEdGraphSchema_K2::FPinTypeTreeInfo::Init(const FString& InFriendlyName, const FString& CategoryName, const UEdGraphSchema_K2* Schema, const FString& InTooltip, bool bInReadOnly)
+void UEdGraphSchema_K2::FPinTypeTreeInfo::Init(const FText& InFriendlyName, const FString& CategoryName, const UEdGraphSchema_K2* Schema, const FText& InTooltip, bool bInReadOnly)
 {
 	check( !CategoryName.IsEmpty() );
 	check( Schema );
 
 	FriendlyName = InFriendlyName;
 	Tooltip = InTooltip;
-	PinType.PinCategory = CategoryName;
+	PinType.PinCategory = (CategoryName == TEXT("Enum") ? PC_Byte : CategoryName);
 	PinType.PinSubCategory = TEXT("");
 	PinType.PinSubCategoryObject = NULL;
 
 	bReadOnly = bInReadOnly;
 
-	if (Schema->DoesTypeHaveSubtypes(FriendlyName))
+	if (Schema->DoesTypeHaveSubtypes(CategoryName))
 	{
 		TArray<UObject*> Subtypes;
-		Schema->GetVariableSubtypes(FriendlyName, Subtypes);
+		Schema->GetVariableSubtypes(CategoryName, Subtypes);
 		for (auto it = Subtypes.CreateIterator(); it; ++it)
 		{
-			FString SubtypeTooltip = CategoryName;
+			FText SubtypeTooltip;
 			UStruct* Struct = Cast<UStruct>(*it);
 			if(Struct != NULL)
 			{
-				SubtypeTooltip = Struct->GetToolTipText().ToString();
+				SubtypeTooltip = (Struct ? Struct->GetToolTipText() : InFriendlyName);
 			}
 
-			Children.Add( MakeShareable(new FPinTypeTreeInfo(CategoryName, *it, SubtypeTooltip)) );
+			Children.Add( MakeShareable(new FPinTypeTreeInfo(PinType.PinCategory, *it, SubtypeTooltip)) );
 		}
 
-		if (Schema->PC_Struct == FriendlyName)
+		if (Schema->PC_Struct == CategoryName)
 		{
 			FGatherStructTypesFromAssetsHelper::Gather(CategoryName, Schema, Children);
 		}
 	}
 }
 
-UEdGraphSchema_K2::FPinTypeTreeInfo::FPinTypeTreeInfo(const FString& CategoryName, UObject* SubCategoryObject, const FString& InTooltip, bool bInReadOnly/*=false*/)
+UEdGraphSchema_K2::FPinTypeTreeInfo::FPinTypeTreeInfo(const FString& CategoryName, UObject* SubCategoryObject, const FText& InTooltip, bool bInReadOnly/*=false*/)
 {
 	check( !CategoryName.IsEmpty() );
 	check( SubCategoryObject );
@@ -201,7 +196,7 @@ UEdGraphSchema_K2::FPinTypeTreeInfo::FPinTypeTreeInfo(const FString& CategoryNam
 	bReadOnly = bInReadOnly;
 }
 
-UEdGraphSchema_K2::FPinTypeTreeInfo::FPinTypeTreeInfo(const FString& CategoryName, const FStringAssetReference& SubCategoryObject, const FString& InTooltip, bool bInReadOnly)
+UEdGraphSchema_K2::FPinTypeTreeInfo::FPinTypeTreeInfo(const FString& CategoryName, const FStringAssetReference& SubCategoryObject, const FText& InTooltip, bool bInReadOnly)
 {
 	check(!CategoryName.IsEmpty());
 	check(SubCategoryObject.IsValid());
@@ -215,9 +210,9 @@ UEdGraphSchema_K2::FPinTypeTreeInfo::FPinTypeTreeInfo(const FString& CategoryNam
 	bReadOnly = bInReadOnly;
 }
 
-FString UEdGraphSchema_K2::FPinTypeTreeInfo::GetDescription() const
+FText UEdGraphSchema_K2::FPinTypeTreeInfo::GetDescription() const
 {
-	if ((PinType.PinCategory != FriendlyName) && !FriendlyName.IsEmpty())
+	if (!FriendlyName.IsEmpty())
 	{
 		return FriendlyName;
 	}
@@ -230,25 +225,14 @@ FString UEdGraphSchema_K2::FPinTypeTreeInfo::GetDescription() const
 			DisplayName.RemoveFromEnd(TEXT("_C"));
 		}
 
-		//@todo:  Fix this once the XX_YYYY names in the schema are static!  This is mirrored to PC_Class
-		if ((PinType.PinCategory == TEXT("class")) && SubCategoryClass)
-		{
-			DisplayName = FString::Printf(TEXT("class'%s'"), *DisplayName);
-		}
-
-		return DisplayName;
-	}
-	else if (!PinType.PinCategory.IsEmpty())
-	{
-		return PinType.PinCategory;
+		return FText::FromString(DisplayName);
 	}
 	else
 	{
-		return TEXT("Error!");
+		return LOCTEXT("PinDescriptionError", "Error!");
 	}
 }
 
-
 const FString UEdGraphSchema_K2::PC_Exec(TEXT("exec"));
 const FString UEdGraphSchema_K2::PC_Boolean(TEXT("bool"));
 const FString UEdGraphSchema_K2::PC_Byte(TEXT("byte"));
@@ -2673,34 +2657,34 @@ void UEdGraphSchema_K2::GetVariableTypeTree( TArray< TSharedPtr<FPinTypeTreeInfo
 
 	if( bAllowExec )
 	{
-		TypeTree.Add( MakeShareable( new FPinTypeTreeInfo(PC_Exec, this, LOCTEXT("ExecType", "Execution pin").ToString()) ) );
+		TypeTree.Add( MakeShareable( new FPinTypeTreeInfo(GetCategoryText(PC_Exec, true), PC_Exec, this, LOCTEXT("ExecType", "Execution pin")) ) );
 	}
 
-	TypeTree.Add( MakeShareable( new FPinTypeTreeInfo(PC_Boolean, this, LOCTEXT("BooleanType", "True or false value").ToString()) ) );
-	TypeTree.Add( MakeShareable( new FPinTypeTreeInfo(PC_Byte, this, LOCTEXT("ByteType", "8 bit number").ToString()) ) );
-	TypeTree.Add( MakeShareable( new FPinTypeTreeInfo(PC_Int, this, LOCTEXT("IntegerType", "Integer number").ToString()) ) );
-	TypeTree.Add( MakeShareable( new FPinTypeTreeInfo(PC_Float, this, LOCTEXT("FloatType", "Floating point number").ToString()) ) );
-	TypeTree.Add( MakeShareable( new FPinTypeTreeInfo(PC_Name, this, LOCTEXT("NameType", "A text name").ToString()) ) );
-	TypeTree.Add( MakeShareable( new FPinTypeTreeInfo(PC_String, this, LOCTEXT("StringType", "A text string").ToString()) ) );
-	TypeTree.Add( MakeShareable( new FPinTypeTreeInfo(PC_Text, this, LOCTEXT("TextType", "A localizable text string").ToString()) ) );
+	TypeTree.Add( MakeShareable( new FPinTypeTreeInfo(GetCategoryText(PC_Boolean, true), PC_Boolean, this, LOCTEXT("BooleanType", "True or false value")) ) );
+	TypeTree.Add( MakeShareable( new FPinTypeTreeInfo(GetCategoryText(PC_Byte, true), PC_Byte, this, LOCTEXT("ByteType", "8 bit number")) ) );
+	TypeTree.Add( MakeShareable( new FPinTypeTreeInfo(GetCategoryText(PC_Int, true), PC_Int, this, LOCTEXT("IntegerType", "Integer number")) ) );
+	TypeTree.Add( MakeShareable( new FPinTypeTreeInfo(GetCategoryText(PC_Float, true), PC_Float, this, LOCTEXT("FloatType", "Floating point number")) ) );
+	TypeTree.Add( MakeShareable( new FPinTypeTreeInfo(GetCategoryText(PC_Name, true), PC_Name, this, LOCTEXT("NameType", "A text name")) ) );
+	TypeTree.Add( MakeShareable( new FPinTypeTreeInfo(GetCategoryText(PC_String, true), PC_String, this, LOCTEXT("StringType", "A text string")) ) );
+	TypeTree.Add( MakeShareable( new FPinTypeTreeInfo(GetCategoryText(PC_Text, true), PC_Text, this, LOCTEXT("TextType", "A localizable text string")) ) );
 
 	// Add in special first-class struct types
-	TypeTree.Add( MakeShareable( new FPinTypeTreeInfo(PC_Struct, VectorStruct, LOCTEXT("VectorType", "A 3D vector").ToString()) ) );
-	TypeTree.Add( MakeShareable( new FPinTypeTreeInfo(PC_Struct, RotatorStruct, LOCTEXT("RotatorType", "A 3D rotation").ToString()) ) );
-	TypeTree.Add( MakeShareable( new FPinTypeTreeInfo(PC_Struct, TransformStruct, LOCTEXT("TransformType", "A 3D transformation, including translation, rotation and 3D scale.").ToString()) ) );
+	TypeTree.Add( MakeShareable( new FPinTypeTreeInfo(PC_Struct, VectorStruct, LOCTEXT("VectorType", "A 3D vector")) ) );
+	TypeTree.Add( MakeShareable( new FPinTypeTreeInfo(PC_Struct, RotatorStruct, LOCTEXT("RotatorType", "A 3D rotation")) ) );
+	TypeTree.Add( MakeShareable( new FPinTypeTreeInfo(PC_Struct, TransformStruct, LOCTEXT("TransformType", "A 3D transformation, including translation, rotation and 3D scale.")) ) );
 	
 	// Add wildcard type
 	if (bAllowWildCard)
 	{
-		TypeTree.Add( MakeShareable( new FPinTypeTreeInfo(PC_Wildcard, this, LOCTEXT("WildcardType", "Wildcard type (unspecified).").ToString()) ) );
+		TypeTree.Add( MakeShareable( new FPinTypeTreeInfo(GetCategoryText(PC_Wildcard, true), PC_Wildcard, this, LOCTEXT("WildcardType", "Wildcard type (unspecified).")) ) );
 	}
 
 	// Add the types that have subtrees
-	TypeTree.Add( MakeShareable( new FPinTypeTreeInfo(PC_Struct, this, LOCTEXT("StructType", "Struct (value) types.").ToString(), true) ) );
-	TypeTree.Add( MakeShareable( new FPinTypeTreeInfo(PC_Object, this, LOCTEXT("ObjectType", "Object pointer.").ToString(), true) ) );
-	TypeTree.Add( MakeShareable( new FPinTypeTreeInfo(PC_Interface, this, LOCTEXT("InterfaceType", "Interface pointer.").ToString(), true) ) );
-	TypeTree.Add( MakeShareable( new FPinTypeTreeInfo(PC_Class, this, LOCTEXT("ClassType", "Class pointers.").ToString(), true) ) );
-	TypeTree.Add( MakeShareable( new FPinTypeTreeInfo(TEXT("Enum"), PC_Byte, this, LOCTEXT("EnumType", "Enumeration types.").ToString(), true) ) );
+	TypeTree.Add( MakeShareable( new FPinTypeTreeInfo(GetCategoryText(PC_Struct, true), PC_Struct, this, LOCTEXT("StructType", "Struct (value) types."), true) ) );
+	TypeTree.Add( MakeShareable( new FPinTypeTreeInfo(GetCategoryText(PC_Object, true), PC_Object, this, LOCTEXT("ObjectType", "Object pointer."), true) ) );
+	TypeTree.Add( MakeShareable( new FPinTypeTreeInfo(GetCategoryText(PC_Interface, true), PC_Interface, this, LOCTEXT("InterfaceType", "Interface pointer."), true) ) );
+	TypeTree.Add( MakeShareable( new FPinTypeTreeInfo(GetCategoryText(PC_Class, true), PC_Class, this, LOCTEXT("ClassType", "Class pointers."), true) ) );
+	TypeTree.Add( MakeShareable( new FPinTypeTreeInfo(LOCTEXT("EnumCategory", "Enum"), TEXT("Enum"), this, LOCTEXT("EnumType", "Enumeration types."), true) ) );
 }
 
 void UEdGraphSchema_K2::GetVariableIndexTypeTree( TArray< TSharedPtr<FPinTypeTreeInfo> >& TypeTree, bool bAllowExec, bool bAllowWildcard ) const
@@ -2710,26 +2694,26 @@ void UEdGraphSchema_K2::GetVariableIndexTypeTree( TArray< TSharedPtr<FPinTypeTre
 
 	if( bAllowExec )
 	{
-		TypeTree.Add( MakeShareable( new FPinTypeTreeInfo(PC_Exec, this, LOCTEXT("ExecIndexType", "Execution pin").ToString()) ) );
+		TypeTree.Add( MakeShareable( new FPinTypeTreeInfo(GetCategoryText(PC_Exec, true), PC_Exec, this, LOCTEXT("ExecIndexType", "Execution pin")) ) );
 	}
 
-	TypeTree.Add( MakeShareable( new FPinTypeTreeInfo(PC_Boolean, this, LOCTEXT("BooleanIndexType", "True or false value").ToString()) ) );
-	TypeTree.Add( MakeShareable( new FPinTypeTreeInfo(PC_Byte, this, LOCTEXT("ByteIndexType", "8 bit number").ToString()) ) );
-	TypeTree.Add( MakeShareable( new FPinTypeTreeInfo(PC_Int, this, LOCTEXT("IntegerIndexType", "Integer number").ToString()) ) );
+	TypeTree.Add( MakeShareable( new FPinTypeTreeInfo(GetCategoryText(PC_Boolean, true), PC_Boolean, this, LOCTEXT("BooleanIndexType", "True or false value")) ) );
+	TypeTree.Add( MakeShareable( new FPinTypeTreeInfo(GetCategoryText(PC_Byte, true), PC_Byte, this, LOCTEXT("ByteIndexType", "8 bit number")) ) );
+	TypeTree.Add( MakeShareable( new FPinTypeTreeInfo(GetCategoryText(PC_Int, true), PC_Int, this, LOCTEXT("IntegerIndexType", "Integer number")) ) );
 
 	// Add wildcard type
 	if (bAllowWildcard)
 	{
-		TypeTree.Add( MakeShareable( new FPinTypeTreeInfo(PC_Wildcard, this, LOCTEXT("WildcardIndexType", "Wildcard type (unspecified).").ToString()) ) );
+		TypeTree.Add( MakeShareable( new FPinTypeTreeInfo(GetCategoryText(PC_Wildcard, true), PC_Wildcard, this, LOCTEXT("WildcardIndexType", "Wildcard type (unspecified).")) ) );
 	}
 
 	// Add the types that have subtrees
-	TypeTree.Add( MakeShareable( new FPinTypeTreeInfo(TEXT("Enum"), PC_Byte, this, LOCTEXT("EnumIndexType", "Enumeration types.").ToString(), true) ) );
+	TypeTree.Add( MakeShareable( new FPinTypeTreeInfo(LOCTEXT("EnumCategory", "Enum"), TEXT("Enum"), this, LOCTEXT("EnumIndexType", "Enumeration types."), true) ) );
 }
 
-bool UEdGraphSchema_K2::DoesTypeHaveSubtypes(const FString& FriendlyTypeName) const
+bool UEdGraphSchema_K2::DoesTypeHaveSubtypes(const FString& Category) const
 {
-	return (FriendlyTypeName == PC_Struct) || (FriendlyTypeName == PC_Object) || (FriendlyTypeName == PC_Interface) || (FriendlyTypeName == PC_Class) || (FriendlyTypeName == TEXT("Enum"));
+	return (Category == PC_Struct) || (Category == PC_Object) || (Category == PC_Interface) || (Category == PC_Class) || (Category == TEXT("Enum"));
 }
 
 void UEdGraphSchema_K2::GetVariableSubtypes(const FString& Type, TArray<UObject*>& SubtypesList) const
diff --git a/Engine/Source/Editor/KismetWidgets/Private/SPinTypeSelector.cpp b/Engine/Source/Editor/KismetWidgets/Private/SPinTypeSelector.cpp
--- a/Engine/Source/Editor/KismetWidgets/Private/SPinTypeSelector.cpp
+++ b/Engine/Source/Editor/KismetWidgets/Private/SPinTypeSelector.cpp
@@ -120,7 +120,7 @@ FText SPinTypeSelector::GetTypeDescription() const
 	}
 	else
 	{
-		return FText::FromString(TargetPinType.Get().PinCategory);
+		return UEdGraphSchema_K2::GetCategoryText(TargetPinType.Get().PinCategory, true);
 	}
 }
 
@@ -152,17 +152,20 @@ void SPinTypeSelector::OnArrayCheckStateChanged(ESlateCheckBoxState::Type NewSta
 TSharedRef<ITableRow> SPinTypeSelector::GenerateTypeTreeRow(FPinTypeTreeItem InItem, const TSharedRef<STableViewBase>& OwnerTree)
 {
 	const bool bHasChildren = (InItem->Children.Num() > 0);
-	const FString Description = InItem->GetDescription();
+	const FText Description = InItem->GetDescription();
 
 	// Determine the best icon the to represents this item
 	const FSlateBrush* IconBrush = GetIconFromPin(InItem->GetPinType(false));
 
 	// Use tooltip if supplied, otherwise just repeat description
-	const FString OrgTooltip = InItem->GetToolTip();
-	const FString Tooltip = !OrgTooltip.IsEmpty() ? OrgTooltip : Description;
+	const FText OrgTooltip = InItem->GetToolTip();
+	const FText Tooltip = !OrgTooltip.IsEmpty() ? OrgTooltip : Description;
+
+	const FEdGraphPinType& PinType = InItem->GetPinType(true);
+	const FString PinTooltipExcerpt = ((PinType.PinCategory != UEdGraphSchema_K2::PC_Byte || PinType.PinSubCategoryObject == nullptr) ? PinType.PinCategory : TEXT("Enum")); 
 
 	return SNew( SComboRow<FPinTypeTreeItem>, OwnerTree )
-		.ToolTip( IDocumentation::Get()->CreateToolTip( FText::FromString( Tooltip ), NULL, *BigTooltipDocLink, *Description) )
+		.ToolTip( IDocumentation::Get()->CreateToolTip( Tooltip, NULL, *BigTooltipDocLink, PinTooltipExcerpt) )
 		[
 			SNew(SHorizontalBox)
 			+SHorizontalBox::Slot()
@@ -354,7 +357,7 @@ bool SPinTypeSelector::GetChildrenMatchingSearch(const FString& InSearchText, co
 		// Have to run GetChildrenMatchingSearch first, so that we can make sure we get valid children for the list!
 		if( GetChildrenMatchingSearch(InSearchText, Item->Children, ValidChildren)
 			|| InSearchText.IsEmpty()
-			|| Item->GetDescription().Contains(InSearchText) )
+			|| Item->GetDescription().ToString().Contains(InSearchText) )
 		{
 			NewInfo->Children = ValidChildren;
 			OutFilteredList.Add(NewInfo);
