Fixed conflicting LOCTEXT for InvalidDynamicCastClass_Error.

diff --git a/Engine/Source/Editor/BlueprintGraph/Private/DynamicCastHandler.cpp b/Engine/Source/Editor/BlueprintGraph/Private/DynamicCastHandler.cpp
--- a/Engine/Source/Editor/BlueprintGraph/Private/DynamicCastHandler.cpp
+++ b/Engine/Source/Editor/BlueprintGraph/Private/DynamicCastHandler.cpp
@@ -80,7 +80,7 @@ void FKCHandler_DynamicCast::Compile(FKismetFunctionContext& Context, UEdGraphNo
 	FBPTerminal** CastResultTerm = Context.NetMap.Find(CastOutputPin);
 	if (!CastResultTerm || !(*CastResultTerm))
 	{
-		CompilerContext.MessageLog.Error(*LOCTEXT("InvalidDynamicCastClass_Error", "Node @@ has an invalid target class. (Inner compiler error?)").ToString(), Node);
+		CompilerContext.MessageLog.Error(*LOCTEXT("InvalidDynamicCastClass_CompilerError", "Node @@ has an invalid target class. (Inner compiler error?)").ToString(), Node);
 		return;
 	}
 
