Fixed issue from refactor with cast nodes being filtered out in the new Blueprint menu when dragging off input pins.

diff --git a/Engine/Source/Editor/BlueprintGraph/Private/BlueprintActionFilter.cpp b/Engine/Source/Editor/BlueprintGraph/Private/BlueprintActionFilter.cpp
--- a/Engine/Source/Editor/BlueprintGraph/Private/BlueprintActionFilter.cpp
+++ b/Engine/Source/Editor/BlueprintGraph/Private/BlueprintActionFilter.cpp
@@ -995,7 +995,6 @@ static bool BlueprintActionFilterImpl::IsNotSubClassCast(FBlueprintActionFilter
 			// Only worry about removing cast nodes when dragging off output pins
 			if (ContextPin->Direction != EGPD_Output)
 			{
-				bIsFilteredOut = true;
 				break;
 			}
 			
