BP: AutoCreateRefTerm meta-data affects only the selected param.
ttp345754 LIVE: BLUEPRINT: CRASH: 'Get Array' node disconnected itself, causes crashes when reconnected.

diff --git a/Engine/Source/Editor/BlueprintGraph/Private/EdGraphSchema_K2.cpp b/Engine/Source/Editor/BlueprintGraph/Private/EdGraphSchema_K2.cpp
--- a/Engine/Source/Editor/BlueprintGraph/Private/EdGraphSchema_K2.cpp
+++ b/Engine/Source/Editor/BlueprintGraph/Private/EdGraphSchema_K2.cpp
@@ -3178,9 +3178,9 @@ bool UEdGraphSchema_K2::IsAutoCreateRefTerm(const UEdGraphPin* Pin) const
 	if (FuncNode)
 	{
 		UFunction* TargetFunction = FuncNode->GetTargetFunction();
-		if (TargetFunction)
+		if (TargetFunction && !Pin->PinName.IsEmpty())
 		{
-			bIsAutoCreateRefTerm = TargetFunction->HasMetaData(FBlueprintMetadata::MD_AutoCreateRefTerm);
+			bIsAutoCreateRefTerm = (Pin->PinName == TargetFunction->GetMetaData(FBlueprintMetadata::MD_AutoCreateRefTerm));
 		}
 	}
 
