UMG/BP - The construct object from class no longer automatically breaks the output pin if the class changes.  Letting the compiler catch it instead if the code is no longer valid.

diff --git a/Engine/Source/Editor/BlueprintGraph/Private/K2Node_ConstructObjectFromClass.cpp b/Engine/Source/Editor/BlueprintGraph/Private/K2Node_ConstructObjectFromClass.cpp
--- a/Engine/Source/Editor/BlueprintGraph/Private/K2Node_ConstructObjectFromClass.cpp
+++ b/Engine/Source/Editor/BlueprintGraph/Private/K2Node_ConstructObjectFromClass.cpp
@@ -151,8 +151,8 @@ void UK2Node_ConstructObjectFromClass::PinDefaultValueChanged(UEdGraphPin* Chang
 		const UEdGraphSchema_K2* K2Schema = GetDefault<UEdGraphSchema_K2>();
 
 		// Because the archetype has changed, we break the output link as the output pin type will change
-		UEdGraphPin* ResultPin = GetResultPin();
-		ResultPin->BreakAllPinLinks();
+		//UEdGraphPin* ResultPin = GetResultPin();
+		//ResultPin->BreakAllPinLinks();
 
 		// Remove all pins related to archetype variables
 		TArray<UEdGraphPin*> OldPins = Pins;
