BP: Fix for BP Editor assertion when BP pin category name is empty

diff --git a/Engine/Source/Editor/BlueprintGraph/Private/EdGraphSchema_K2.cpp b/Engine/Source/Editor/BlueprintGraph/Private/EdGraphSchema_K2.cpp
--- a/Engine/Source/Editor/BlueprintGraph/Private/EdGraphSchema_K2.cpp
+++ b/Engine/Source/Editor/BlueprintGraph/Private/EdGraphSchema_K2.cpp
@@ -2549,6 +2549,11 @@ FText UEdGraphSchema_K2::TypeToText(UProperty* const Property)
 
 FText UEdGraphSchema_K2::GetCategoryText(const FString& Category, const bool bForMenu)
 {
+	if (Category.IsEmpty())
+	{
+		return FText::GetEmpty();
+	}
+
 	static TMap<FString, FText> CategoryDescriptions;
 	if (CategoryDescriptions.Num() == 0)
 	{
