Fixed 'Add Timeline' showing up for non-actor based blueprints when using the new blueprint context menu

diff --git a/Engine/Source/Editor/BlueprintGraph/Private/K2Node_Timeline.cpp b/Engine/Source/Editor/BlueprintGraph/Private/K2Node_Timeline.cpp
--- a/Engine/Source/Editor/BlueprintGraph/Private/K2Node_Timeline.cpp
+++ b/Engine/Source/Editor/BlueprintGraph/Private/K2Node_Timeline.cpp
@@ -2,11 +2,13 @@
 
 
 #include "BlueprintGraphPrivatePCH.h"
+
+#include "BlueprintActionDatabaseRegistrar.h"
+#include "BlueprintEditorUtils.h"
+#include "BlueprintNodeSpawner.h"
 #include "Editor/GraphEditor/Public/DiffResults.h"
 #include "Kismet2NameValidators.h"
 #include "KismetCompiler.h"
-#include "BlueprintNodeSpawner.h"
-#include "BlueprintActionDatabaseRegistrar.h"
 
 #define LOCTEXT_NAMESPACE "K2Node_Timeline"
 
@@ -222,7 +224,7 @@ bool UK2Node_Timeline::IsCompatibleWithGraph(const UEdGraph* TargetGraph) const
 
 			if(bAllowEvents)
 			{
-				return true;
+				return FBlueprintEditorUtils::DoesSupportTimelines(Blueprint);
 			}
 			else
 			{
@@ -248,7 +250,7 @@ bool UK2Node_Timeline::IsCompatibleWithGraph(const UEdGraph* TargetGraph) const
 						}
 					}
 				}
-				return bCompositeOfUbberGraph;
+				return bCompositeOfUbberGraph ? FBlueprintEditorUtils::DoesSupportTimelines(Blueprint) : false;
 			}
 		}
 	}
