Change ShowHiddenSelfPins to ShowWorldContextPin

diff --git a/Engine/Plugins/Experimental/LiveEditor/Source/LiveEditor/Private/K2Node_LiveEditObject.cpp b/Engine/Plugins/Experimental/LiveEditor/Source/LiveEditor/Private/K2Node_LiveEditObject.cpp
--- a/Engine/Plugins/Experimental/LiveEditor/Source/LiveEditor/Private/K2Node_LiveEditObject.cpp
+++ b/Engine/Plugins/Experimental/LiveEditor/Source/LiveEditor/Private/K2Node_LiveEditObject.cpp
@@ -151,7 +151,7 @@ void UK2Node_LiveEditObject::AllocateDefaultPins()
 	CreatePin(EGPD_Output, K2Schema->PC_Exec, TEXT(""), NULL, false, false, K2Schema->PN_Then);
 
 	// If required add the world context pin
-	if (GetBlueprint()->ParentClass->HasMetaData(FBlueprintMetadata::MD_ShowHiddenSelfPins))
+	if (GetBlueprint()->ParentClass->HasMetaData(FBlueprintMetadata::MD_ShowWorldContextPin))
 	{
 		CreatePin(EGPD_Input, K2Schema->PC_Object, TEXT(""), UObject::StaticClass(), false, false, UK2Node_LiveEditObjectStatics::WorldContextPinName);
 	}
diff --git a/Engine/Source/Editor/BlueprintGraph/Classes/EdGraphSchema_K2.h b/Engine/Source/Editor/BlueprintGraph/Classes/EdGraphSchema_K2.h
--- a/Engine/Source/Editor/BlueprintGraph/Classes/EdGraphSchema_K2.h
+++ b/Engine/Source/Editor/BlueprintGraph/Classes/EdGraphSchema_K2.h
@@ -101,8 +101,8 @@ public:
 	/** If true, an unconnected pin will generate a UPROPERTY under the hood to connect as the input, which will be set to the literal value for the pin.  Only valid for reference parameters. */
 	static const FName MD_AutoCreateRefTerm;
 
-	/** If true, the hidden default to self pins will be visible when the function is placed in a child blueprint of the class. */
-	static const FName MD_ShowHiddenSelfPins;
+	/** If true, the hidden world context pin will be visible when the function is placed in a child blueprint of the class. */
+	static const FName MD_ShowWorldContextPin;
 
 	static const FName MD_BlueprintInternalUseOnly;
 	static const FName MD_NeedsLatentFixup;
diff --git a/Engine/Source/Editor/BlueprintGraph/Private/CallFunctionHandler.cpp b/Engine/Source/Editor/BlueprintGraph/Private/CallFunctionHandler.cpp
--- a/Engine/Source/Editor/BlueprintGraph/Private/CallFunctionHandler.cpp
+++ b/Engine/Source/Editor/BlueprintGraph/Private/CallFunctionHandler.cpp
@@ -512,7 +512,7 @@ void FKCHandler_CallFunction::RegisterNets(FKismetFunctionContext& Context, UEdG
 			FString const DefaltToSelfPinName = Function->GetMetaData(FBlueprintMetadata::MD_DefaultToSelf);
 			AutoCreateRefTermPinNames.Remove(DefaltToSelfPinName);
 
-			if (Context.Blueprint->ParentClass->HasMetaData(FBlueprintMetadata::MD_ShowHiddenSelfPins))
+			if (Context.Blueprint->ParentClass->HasMetaData(FBlueprintMetadata::MD_ShowWorldContextPin))
 			{
 				RequiresSetValue.Add(DefaltToSelfPinName);
 			}
diff --git a/Engine/Source/Editor/BlueprintGraph/Private/EdGraphSchema_K2.cpp b/Engine/Source/Editor/BlueprintGraph/Private/EdGraphSchema_K2.cpp
--- a/Engine/Source/Editor/BlueprintGraph/Private/EdGraphSchema_K2.cpp
+++ b/Engine/Source/Editor/BlueprintGraph/Private/EdGraphSchema_K2.cpp
@@ -58,7 +58,7 @@ const FName FBlueprintMetadata::MD_DefaultToSelf(TEXT("DefaultToSelf"));
 const FName FBlueprintMetadata::MD_WorldContext(TEXT("WorldContext"));
 const FName FBlueprintMetadata::MD_AutoCreateRefTerm(TEXT("AutoCreateRefTerm"));
 
-const FName FBlueprintMetadata::MD_ShowHiddenSelfPins(TEXT("ShowHiddenSelfPins"));
+const FName FBlueprintMetadata::MD_ShowWorldContextPin(TEXT("ShowWorldContextPin"));
 const FName FBlueprintMetadata::MD_Private(TEXT("BlueprintPrivate"));
 
 const FName FBlueprintMetadata::MD_BlueprintInternalUseOnly(TEXT("BlueprintInternalUseOnly"));
diff --git a/Engine/Source/Editor/BlueprintGraph/Private/K2Node_CallFunction.cpp b/Engine/Source/Editor/BlueprintGraph/Private/K2Node_CallFunction.cpp
--- a/Engine/Source/Editor/BlueprintGraph/Private/K2Node_CallFunction.cpp
+++ b/Engine/Source/Editor/BlueprintGraph/Private/K2Node_CallFunction.cpp
@@ -554,7 +554,7 @@ bool UK2Node_CallFunction::CreatePinsForFunctionCall(const UFunction* Function)
 	TSet<FString> PinsToHide;
 	FBlueprintEditorUtils::GetHiddenPinsForFunction(BP, Function, PinsToHide);
 
-	const bool bShowHiddenSelfPins = ((PinsToHide.Num() > 0) && BP && BP->ParentClass && BP->ParentClass->HasMetaData(FBlueprintMetadata::MD_ShowHiddenSelfPins));
+	const bool bShowWorldContextPin = ((PinsToHide.Num() > 0) && BP && BP->ParentClass && BP->ParentClass->HasMetaData(FBlueprintMetadata::MD_ShowWorldContextPin));
 
 	// Create the inputs and outputs
 	bool bAllPinsGood = true;
@@ -592,7 +592,7 @@ bool UK2Node_CallFunction::CreatePinsForFunctionCall(const UFunction* Function)
 				FString const WorldContextMetaValue  = Function->GetMetaData(FBlueprintMetadata::MD_WorldContext);
 				bool bIsSelfPin = ((Pin->PinName == DefaultToSelfMetaValue) || (Pin->PinName == WorldContextMetaValue));
 
-				if (!bShowHiddenSelfPins || !bIsSelfPin)
+				if (!bShowWorldContextPin || !bIsSelfPin)
 				{
 					Pin->bHidden = true;
 					K2Schema->SetPinDefaultValueBasedOnType(Pin);
@@ -1078,7 +1078,7 @@ void UK2Node_CallFunction::PostPasteNode()
 		TSet<FString> PinsToHide;
 		FBlueprintEditorUtils::GetHiddenPinsForFunction(GetBlueprint(), Function, PinsToHide);
 
-		const bool bShowHiddenSelfPins = ((PinsToHide.Num() > 0) && GetBlueprint()->ParentClass->HasMetaData(FBlueprintMetadata::MD_ShowHiddenSelfPins));
+		const bool bShowWorldContextPin = ((PinsToHide.Num() > 0) && GetBlueprint()->ParentClass->HasMetaData(FBlueprintMetadata::MD_ShowWorldContextPin));
 
 		FString const DefaultToSelfMetaValue = Function->GetMetaData(FBlueprintMetadata::MD_DefaultToSelf);
 		FString const WorldContextMetaValue  = Function->GetMetaData(FBlueprintMetadata::MD_WorldContext);
@@ -1089,7 +1089,7 @@ void UK2Node_CallFunction::PostPasteNode()
 			UEdGraphPin* Pin = Pins[PinIndex];
 
 			bool bIsSelfPin = ((Pin->PinName == DefaultToSelfMetaValue) || (Pin->PinName == WorldContextMetaValue));
-			bool bPinShouldBeHidden = PinsToHide.Contains(Pin->PinName) && (!bShowHiddenSelfPins || !bIsSelfPin);
+			bool bPinShouldBeHidden = PinsToHide.Contains(Pin->PinName) && (!bShowWorldContextPin || !bIsSelfPin);
 
 			if (bPinShouldBeHidden && !Pin->bHidden)
 			{
diff --git a/Engine/Source/Editor/BlueprintGraph/Private/K2Node_ConstructObjectFromClass.cpp b/Engine/Source/Editor/BlueprintGraph/Private/K2Node_ConstructObjectFromClass.cpp
--- a/Engine/Source/Editor/BlueprintGraph/Private/K2Node_ConstructObjectFromClass.cpp
+++ b/Engine/Source/Editor/BlueprintGraph/Private/K2Node_ConstructObjectFromClass.cpp
@@ -38,7 +38,7 @@ void UK2Node_ConstructObjectFromClass::AllocateDefaultPins()
 	CreatePin(EGPD_Output, K2Schema->PC_Exec, TEXT(""), NULL, false, false, K2Schema->PN_Then);
 
 	// If required add the world context pin
-	if ( GetBlueprint()->ParentClass->HasMetaData(FBlueprintMetadata::MD_ShowHiddenSelfPins) )
+	if ( GetBlueprint()->ParentClass->HasMetaData(FBlueprintMetadata::MD_ShowWorldContextPin) )
 	{
 		CreatePin(EGPD_Input, K2Schema->PC_Object, TEXT(""), UObject::StaticClass(), false, false, FK2Node_ConstructObjectFromClassHelper::WorldContextPinName);
 	}
diff --git a/Engine/Source/Editor/BlueprintGraph/Private/K2Node_SpawnActor.cpp b/Engine/Source/Editor/BlueprintGraph/Private/K2Node_SpawnActor.cpp
--- a/Engine/Source/Editor/BlueprintGraph/Private/K2Node_SpawnActor.cpp
+++ b/Engine/Source/Editor/BlueprintGraph/Private/K2Node_SpawnActor.cpp
@@ -26,7 +26,7 @@ void UK2Node_SpawnActor::AllocateDefaultPins()
 	CreatePin(EGPD_Output, K2Schema->PC_Exec, TEXT(""), NULL, false, false, K2Schema->PN_Then);
 
 	// If required add the world context pin
-	if (GetBlueprint()->ParentClass->HasMetaData(FBlueprintMetadata::MD_ShowHiddenSelfPins))
+	if (GetBlueprint()->ParentClass->HasMetaData(FBlueprintMetadata::MD_ShowWorldContextPin))
 	{
 		CreatePin(EGPD_Input, K2Schema->PC_Object, TEXT(""), UObject::StaticClass(), false, false, WorldContextPinName);
 	}
diff --git a/Engine/Source/Editor/BlueprintGraph/Private/K2Node_SpawnActorFromClass.cpp b/Engine/Source/Editor/BlueprintGraph/Private/K2Node_SpawnActorFromClass.cpp
--- a/Engine/Source/Editor/BlueprintGraph/Private/K2Node_SpawnActorFromClass.cpp
+++ b/Engine/Source/Editor/BlueprintGraph/Private/K2Node_SpawnActorFromClass.cpp
@@ -37,7 +37,7 @@ void UK2Node_SpawnActorFromClass::AllocateDefaultPins()
 	CreatePin(EGPD_Output, K2Schema->PC_Exec, TEXT(""), NULL, false, false, K2Schema->PN_Then);
 
 	// If required add the world context pin
-	if (GetBlueprint()->ParentClass->HasMetaData(FBlueprintMetadata::MD_ShowHiddenSelfPins))
+	if (GetBlueprint()->ParentClass->HasMetaData(FBlueprintMetadata::MD_ShowWorldContextPin))
 	{
 		CreatePin(EGPD_Input, K2Schema->PC_Object, TEXT(""), UObject::StaticClass(), false, false, FK2Node_SpawnActorFromClassHelper::WorldContextPinName);
 	}
diff --git a/Engine/Source/Runtime/CoreUObject/Public/UObject/ObjectBase.h b/Engine/Source/Runtime/CoreUObject/Public/UObject/ObjectBase.h
--- a/Engine/Source/Runtime/CoreUObject/Public/UObject/ObjectBase.h
+++ b/Engine/Source/Runtime/CoreUObject/Public/UObject/ObjectBase.h
@@ -875,9 +875,9 @@ namespace UM
 		/// [ClassMetadata] Used by BlueprintFunctionLibrary classes to restrict the graphs the functions in the library can be used in to the classes specified.
 		RestrictedToClasses,
 
-		/// [ClassMetadata] Indicates that when placing blueprint nodes in graphs owned by this class that any hidden self pins (generally used to determine the world) should be visible because the self context of the class cannot
+		/// [ClassMetadata] Indicates that when placing blueprint nodes in graphs owned by this class that the hidden world context pin should be visible because the self context of the class cannot
 		///                 provide the world context and it must be wired in manually
-		ShowHiddenSelfPins,
+		ShowWorldContextPin,
 
 	};
 
@@ -1080,7 +1080,7 @@ namespace UM
 		/// [FunctionMetadata] For BlueprintCallable functions indicates that the function should be displayed the same as the implicit Make Struct nodes
 		NativeMakeFunc,
 
-		// [FunctionMetadata] Used by BlueprintCallable functions to indicate that this function is not to be allowed in the Construction Script.
+		/// [FunctionMetadata] Used by BlueprintCallable functions to indicate that this function is not to be allowed in the Construction Script.
 		UnsafeDuringActorConstruction,
 
 	};
diff --git a/Engine/Source/Runtime/Engine/Classes/Animation/AnimNotifies/AnimNotifyState.h b/Engine/Source/Runtime/Engine/Classes/Animation/AnimNotifies/AnimNotifyState.h
--- a/Engine/Source/Runtime/Engine/Classes/Animation/AnimNotifies/AnimNotifyState.h
+++ b/Engine/Source/Runtime/Engine/Classes/Animation/AnimNotifies/AnimNotifyState.h
@@ -9,7 +9,7 @@ class UAnimSequence;
 class UAnimSequenceBase;
 struct FAnimNotifyEvent;
 
-UCLASS(abstract, editinlinenew, Blueprintable, const, hidecategories=Object, collapsecategories, meta=(ShowHiddenSelfPins))
+UCLASS(abstract, editinlinenew, Blueprintable, const, hidecategories=Object, collapsecategories, meta=(ShowWorldContextPin))
 class ENGINE_API UAnimNotifyState : public UObject
 {
 	GENERATED_UCLASS_BODY()
diff --git a/Engine/Source/Runtime/Engine/Classes/Animation/AnimNotifies/AnimNotifyState_Trail.h b/Engine/Source/Runtime/Engine/Classes/Animation/AnimNotifies/AnimNotifyState_Trail.h
--- a/Engine/Source/Runtime/Engine/Classes/Animation/AnimNotifies/AnimNotifyState_Trail.h
+++ b/Engine/Source/Runtime/Engine/Classes/Animation/AnimNotifies/AnimNotifyState_Trail.h
@@ -6,7 +6,7 @@
 
 DECLARE_LOG_CATEGORY_EXTERN(LogAnimTrails, Log, All);
 
-UCLASS(editinlinenew, Blueprintable, const, hidecategories = Object, collapsecategories, MinimalAPI, meta = (ShowHiddenSelfPins, DisplayName = "Trail"))
+UCLASS(editinlinenew, Blueprintable, const, hidecategories = Object, collapsecategories, MinimalAPI, meta = (ShowWorldContextPin, DisplayName = "Trail"))
 class UAnimNotifyState_Trail : public UAnimNotifyState
 {
 	GENERATED_UCLASS_BODY()
