Reorganize CanFunctionBeUsedInClass to make to easier to understand and modify

diff --git a/Engine/Source/Editor/BlueprintGraph/Private/EdGraphSchema_K2.cpp b/Engine/Source/Editor/BlueprintGraph/Private/EdGraphSchema_K2.cpp
--- a/Engine/Source/Editor/BlueprintGraph/Private/EdGraphSchema_K2.cpp
+++ b/Engine/Source/Editor/BlueprintGraph/Private/EdGraphSchema_K2.cpp
@@ -316,60 +316,106 @@ bool UEdGraphSchema_K2::DoesFunctionHaveOutParameters( const UFunction* Function
 	return false;
 }
 
-bool UEdGraphSchema_K2::CanFunctionBeUsedInClass(const UClass* InClass, UFunction* InFunction, const UEdGraph* InDestGraph, uint32 InFunctionTypes, bool bInShowInherited, bool bInCalledForEach, const FFunctionTargetInfo& InTargetInfo) const
+bool UEdGraphSchema_K2::CanFunctionBeUsedInClass(const UClass* InClass, UFunction* InFunction, const UEdGraph* InDestGraph, uint32 InAllowedFunctionTypes, bool bInShowInherited, bool bInCalledForEach, const FFunctionTargetInfo& InTargetInfo) const
 {
-	bool bLatentFuncs = true;
-	bool bIsConstructionScript = false;
-
-	if(InDestGraph != NULL)
-	{
-		bLatentFuncs = (GetGraphType(InDestGraph) == GT_Ubergraph || (GetGraphType(InDestGraph) == GT_Macro));
-		bIsConstructionScript = IsConstructionScript(InDestGraph);
-	}
-
 	if (CanUserKismetCallFunction(InFunction))
 	{
-		// See if this is the kind of function we are looking for
-		const bool bPureFuncs = (InFunctionTypes & FT_Pure) != 0;
-		const bool bImperativeFuncs = (InFunctionTypes & FT_Imperative) != 0;
-		const bool bConstFuncs = (InFunctionTypes & FT_Const) != 0;
-		const bool bProtectedFuncs = (InFunctionTypes & FT_Protected) != 0;
+		bool bLatentFuncsAllowed = true;
+		bool bIsConstructionScript = false;
+
+		if(InDestGraph != NULL)
+		{
+			bLatentFuncsAllowed = (GetGraphType(InDestGraph) == GT_Ubergraph || (GetGraphType(InDestGraph) == GT_Macro));
+			bIsConstructionScript = IsConstructionScript(InDestGraph);
+		}
 
 		const bool bIsPureFunc = (InFunction->HasAnyFunctionFlags(FUNC_BlueprintPure) != false);
+		if (bIsPureFunc)
+		{
+			const bool bAllowPureFuncs = (InAllowedFunctionTypes & FT_Pure) != 0;
+			if (!bAllowPureFuncs)
+			{
+				return false;
+			}
+		}
+		else
+		{
+			const bool bAllowImperativeFuncs = (InAllowedFunctionTypes & FT_Imperative) != 0;
+			if (!bAllowImperativeFuncs)
+			{
+				return false;
+			}
+		}
+
 		const bool bIsConstFunc = (InFunction->HasAnyFunctionFlags(FUNC_Const) != false);
+		const bool bAllowConstFuncs = (InAllowedFunctionTypes & FT_Const) != 0;
+		if (bIsConstFunc && !bAllowConstFuncs)
+		{
+			return false;
+		}
+
 		const bool bIsLatent = InFunction->HasMetaData(FBlueprintMetadata::MD_Latent);
-		const bool bIsBlueprintProtected = InFunction->GetBoolMetaData(FBlueprintMetadata::MD_Protected);
+		if (bIsLatent && !bLatentFuncsAllowed)
+		{
+			return false;
+		}
+
+		const bool bIsProtected = InFunction->GetBoolMetaData(FBlueprintMetadata::MD_Protected);
+		const bool bFuncBelongsToSubClass = InClass->IsChildOf(InFunction->GetOuterUClass());
+		if (bIsProtected)
+		{
+			const bool bAllowProtectedFuncs = (InAllowedFunctionTypes & FT_Protected) != 0;
+			if (!bAllowProtectedFuncs)
+			{
+				return false;
+			}
+
+			if (!bFuncBelongsToSubClass)
+			{
+				return false;
+			}
+		}
+
+		const bool bIsPrivate = InFunction->GetBoolMetaData(FBlueprintMetadata::MD_Private);
+		const bool bFuncBelongsToClass = bFuncBelongsToSubClass && (InFunction->GetOuterUClass() == InClass);
+		if (bIsPrivate && !bFuncBelongsToClass)
+		{
+			return false;
+		}
+
 		const bool bIsUnsafeForConstruction = InFunction->GetBoolMetaData(FBlueprintMetadata::MD_UnsafeForConstructionScripts);	
+		if (bIsUnsafeForConstruction && bIsConstructionScript)
+		{
+			return false;
+		}
+
 		const bool bFunctionHidden = FObjectEditorUtils::IsFunctionHiddenFromClass(InFunction, InClass);
-		const bool bIsPrivate = InFunction->GetBoolMetaData(FBlueprintMetadata::MD_Private);
+		if (bFunctionHidden)
+		{
+			return false;
+		}
 
 		const bool bFunctionStatic = InFunction->HasAllFunctionFlags(FUNC_Static);
 		const bool bHasReturnParams = (InFunction->GetReturnProperty() != NULL);
 		const bool bHasArrayPointerParms = InFunction->HasMetaData(TEXT("ArrayParm"));
+
 		const bool bAllowForEachCall = !bFunctionStatic && !bIsLatent && !bIsPureFunc && !bIsConstFunc && !bHasReturnParams && !bHasArrayPointerParms;
+		if (bInCalledForEach && !bAllowForEachCall)
+		{
+			return false;
+		}
 
 		const bool bClassIsAnActor = InClass->IsChildOf( AActor::StaticClass() );
-		const bool bClassIsAComponent = InClass->IsChildOf( UActorComponent::StaticClass() );
-
-		const bool bFuncBelongsToSubClass = InClass->IsChildOf(InFunction->GetOuterUClass());
-		const bool bFuncBelongsToClass    = bFuncBelongsToSubClass && (InFunction->GetOuterUClass() == InClass);
-
-		const bool bFunctionHasReturnOrOutParameters = bHasReturnParams || DoesFunctionHaveOutParameters(InFunction);
 
 		// This will evaluate to false if there are multiple actors selected and the function has a return value or out parameters
+		const bool bFunctionHasReturnOrOutParameters = bHasReturnParams || DoesFunctionHaveOutParameters(InFunction);
 		const bool bAllowReturnValuesForNoneOrSingleActors = !bClassIsAnActor || InTargetInfo.Actors.Num() <= 1 || !bFunctionHasReturnOrOutParameters;
-
-		if (((bIsPureFunc && bPureFuncs) || (!bIsPureFunc && bImperativeFuncs) || (bIsConstFunc && bConstFuncs))
-			&& (!bIsLatent || bLatentFuncs)
-			&& (!bIsBlueprintProtected || (bProtectedFuncs && bFuncBelongsToSubClass))
-			&& (!bIsUnsafeForConstruction || !bIsConstructionScript)
-			&& !bFunctionHidden
-			&& (bAllowForEachCall || !bInCalledForEach)
-			&& bAllowReturnValuesForNoneOrSingleActors 
-			&& (!bIsPrivate || bFuncBelongsToClass) )
+		if (!bAllowReturnValuesForNoneOrSingleActors)
 		{
-			return true;
+			return false;
 		}
+
+		return true;
 	}
 
 	return false;
