Prevent functions that require a World Context from being placed in a Graph that does not implement GetWorld() and report a warning for functions that already exist
A function with a World Context pin can also be flagged with new metadata CallableWithoutWorldContext which indicates it can be used in any context, even if GetWorld() is not implemented, the called function is responsible for providing a working alternative when World is null
World Centric functions can also be used in classes tagged with the ShowWorldContextPin metadata which indicates the class has a well defined method of supplying the world context that will be connected

diff --git a/Engine/Source/Editor/BlueprintGraph/Classes/EdGraphSchema_K2.h b/Engine/Source/Editor/BlueprintGraph/Classes/EdGraphSchema_K2.h
--- a/Engine/Source/Editor/BlueprintGraph/Classes/EdGraphSchema_K2.h
+++ b/Engine/Source/Editor/BlueprintGraph/Classes/EdGraphSchema_K2.h
@@ -98,6 +98,9 @@ public:
 	/** The specified parameter should be used as the context object when retrieving a UWorld pointer (implies hidden and default-to-self) */
 	static const FName MD_WorldContext;
 
+	/** For functions that have the MD_WorldContext metadata but are safe to be called from contexts that do not have the ability to provide the world context (either through GetWorld() or ShowWorldContextPin class metadata */
+	static const FName MD_CallableWithoutWorldContext;
+
 	/** If true, an unconnected pin will generate a UPROPERTY under the hood to connect as the input, which will be set to the literal value for the pin.  Only valid for reference parameters. */
 	static const FName MD_AutoCreateRefTerm;
 
diff --git a/Engine/Source/Editor/BlueprintGraph/Private/CallFunctionHandler.cpp b/Engine/Source/Editor/BlueprintGraph/Private/CallFunctionHandler.cpp
--- a/Engine/Source/Editor/BlueprintGraph/Private/CallFunctionHandler.cpp
+++ b/Engine/Source/Editor/BlueprintGraph/Private/CallFunctionHandler.cpp
@@ -525,7 +525,7 @@ void FKCHandler_CallFunction::RegisterNets(FKismetFunctionContext& Context, UEdG
 			{
 				DefaultToSelfParamNames.Add(WorldContextPinName);
 			}
-			else
+			else if (!Function->HasMetaData(FBlueprintMetadata::MD_CallableWithoutWorldContext))
 			{
 				RequiresSetValue.Add(WorldContextPinName);
 			}
diff --git a/Engine/Source/Editor/BlueprintGraph/Private/EdGraphSchema_K2.cpp b/Engine/Source/Editor/BlueprintGraph/Private/EdGraphSchema_K2.cpp
--- a/Engine/Source/Editor/BlueprintGraph/Private/EdGraphSchema_K2.cpp
+++ b/Engine/Source/Editor/BlueprintGraph/Private/EdGraphSchema_K2.cpp
@@ -56,6 +56,7 @@ const FName FBlueprintMetadata::MD_FriendlyName(TEXT("FriendlyName"));
 const FName FBlueprintMetadata::MD_ExposeOnSpawn(TEXT("ExposeOnSpawn"));
 const FName FBlueprintMetadata::MD_DefaultToSelf(TEXT("DefaultToSelf"));
 const FName FBlueprintMetadata::MD_WorldContext(TEXT("WorldContext"));
+const FName FBlueprintMetadata::MD_CallableWithoutWorldContext(TEXT("CallableWithoutWorldContext"));
 const FName FBlueprintMetadata::MD_AutoCreateRefTerm(TEXT("AutoCreateRefTerm"));
 
 const FName FBlueprintMetadata::MD_ShowWorldContextPin(TEXT("ShowWorldContextPin"));
@@ -389,6 +390,19 @@ bool UEdGraphSchema_K2::CanFunctionBeUsedInClass(const UClass* InClass, UFunctio
 			return false;
 		}
 
+		const bool bRequiresWorldContext = InFunction->HasMetaData(FBlueprintMetadata::MD_WorldContext);
+		if (bRequiresWorldContext)
+		{
+			if (InDestGraph && !InFunction->HasMetaData(FBlueprintMetadata::MD_CallableWithoutWorldContext))
+			{
+				UClass* ParentClass = FBlueprintEditorUtils::FindBlueprintForGraphChecked(InDestGraph)->ParentClass;
+				if (!ParentClass->GetDefaultObject()->ImplementsGetWorld() && !ParentClass->HasMetaData(FBlueprintMetadata::MD_ShowWorldContextPin))
+				{
+					return false;
+				}
+			}
+		}
+
 		const bool bFunctionHidden = FObjectEditorUtils::IsFunctionHiddenFromClass(InFunction, InClass);
 		if (bFunctionHidden)
 		{
@@ -3260,6 +3274,15 @@ bool UEdGraphSchema_K2::ShouldShowAssetPickerForPin(UEdGraphPin* Pin) const
 		{
 			// Don't show literal buttons for component type objects
 			bShow = !ObjectClass->IsChildOf(UActorComponent::StaticClass());
+
+			if (bShow)
+			{
+				if (UK2Node_CallFunction* CallFunctionNode = Cast<UK2Node_CallFunction>(Pin->GetOwningNode()))
+				{
+					const UEdGraphPin* WorldContextPin = CallFunctionNode->FindPin(CallFunctionNode->GetTargetFunction()->GetMetaData(FBlueprintMetadata::MD_WorldContext));
+					bShow = (WorldContextPin != Pin);
+				}
+			}
 		}
 	}
 	return bShow;
diff --git a/Engine/Source/Editor/BlueprintGraph/Private/K2Node_CallFunction.cpp b/Engine/Source/Editor/BlueprintGraph/Private/K2Node_CallFunction.cpp
--- a/Engine/Source/Editor/BlueprintGraph/Private/K2Node_CallFunction.cpp
+++ b/Engine/Source/Editor/BlueprintGraph/Private/K2Node_CallFunction.cpp
@@ -1126,35 +1126,49 @@ void UK2Node_CallFunction::ValidateNodeDuringCompilation(class FCompilerResultsL
 		MessageLog.Warning(*FString::Printf(*LOCTEXT("EnumToExecExpansionFailed", "Unable to find enum parameter with name '%s' to expand for @@").ToString(), *EnumParamName), this);
 	}
 
-	// enforce UnsafeDuringActorConstruction keyword
-	if (Function && Function->HasMetaData(FBlueprintMetadata::MD_UnsafeForConstructionScripts))
+	if (Function)
 	{
-		// emit warning if we are in a construction script
-		UEdGraph const* const Graph = GetGraph();
-		UEdGraphSchema_K2 const* const Schema = Cast<const UEdGraphSchema_K2>(GetSchema());
-		bool bNodeIsInConstructionScript = Schema && Schema->IsConstructionScript(Graph);
-
-		if (bNodeIsInConstructionScript == false)
+		// enforce UnsafeDuringActorConstruction keyword
+		if (Function->HasMetaData(FBlueprintMetadata::MD_UnsafeForConstructionScripts))
 		{
-			// IsConstructionScript() can return false if graph was cloned from the construction script
-			// in that case, check the function entry
-			TArray<const UK2Node_FunctionEntry*> EntryPoints;
-			Graph->GetNodesOfClass(EntryPoints);
+			// emit warning if we are in a construction script
+			UEdGraph const* const Graph = GetGraph();
+			UEdGraphSchema_K2 const* const Schema = Cast<const UEdGraphSchema_K2>(GetSchema());
+			bool bNodeIsInConstructionScript = Schema && Schema->IsConstructionScript(Graph);
 
-			if (EntryPoints.Num() == 1)
+			if (bNodeIsInConstructionScript == false)
 			{
-				UK2Node_FunctionEntry const* const Node = EntryPoints[0];
-				if (Node)
+				// IsConstructionScript() can return false if graph was cloned from the construction script
+				// in that case, check the function entry
+				TArray<const UK2Node_FunctionEntry*> EntryPoints;
+				Graph->GetNodesOfClass(EntryPoints);
+
+				if (EntryPoints.Num() == 1)
 				{
-					UFunction* const SignatureFunction = FindField<UFunction>(Node->SignatureClass, Node->SignatureName);
-					bNodeIsInConstructionScript = SignatureFunction && (SignatureFunction->GetFName() == Schema->FN_UserConstructionScript);
+					UK2Node_FunctionEntry const* const Node = EntryPoints[0];
+					if (Node)
+					{
+						UFunction* const SignatureFunction = FindField<UFunction>(Node->SignatureClass, Node->SignatureName);
+						bNodeIsInConstructionScript = SignatureFunction && (SignatureFunction->GetFName() == Schema->FN_UserConstructionScript);
+					}
 				}
 			}
+
+			if ( bNodeIsInConstructionScript )
+			{
+				MessageLog.Warning(*LOCTEXT("FunctionUnsafeDuringConstruction", "Function '@@' is unsafe to call in a construction script.").ToString(), this);
+			}
 		}
 
-		if ( bNodeIsInConstructionScript )
+		// enforce WorldContext restrictions
+		else if (   Function->HasMetaData(FBlueprintMetadata::MD_WorldContext) 
+			     && !Function->HasMetaData(FBlueprintMetadata::MD_CallableWithoutWorldContext))
 		{
-			MessageLog.Warning(*LOCTEXT("FunctionUnsafeDuringConstruction", "Function '@@' is unsafe to call in a construction script.").ToString(), this);
+			UClass* ParentClass = GetBlueprint()->ParentClass;
+			if (!ParentClass->GetDefaultObject()->ImplementsGetWorld() && !ParentClass->HasMetaData(FBlueprintMetadata::MD_ShowWorldContextPin))
+			{
+				MessageLog.Warning(*LOCTEXT("FunctionUnsafeInContext", "Function '@@' is unsafe to call from blueprints of class '@@'.").ToString(), this, ParentClass);
+			}
 		}
 	}
 }
diff --git a/Engine/Source/Runtime/CoreUObject/Public/UObject/ObjectBase.h b/Engine/Source/Runtime/CoreUObject/Public/UObject/ObjectBase.h
--- a/Engine/Source/Runtime/CoreUObject/Public/UObject/ObjectBase.h
+++ b/Engine/Source/Runtime/CoreUObject/Public/UObject/ObjectBase.h
@@ -1032,6 +1032,9 @@ namespace UM
 		/// [FunctionMetadata] This function can only be called on 'this' in a blueprint. It cannot be called on another instance.
 		BlueprintProtected,
 
+		/// [FunctionMetadata] Used for BlueprintCallable functions that have a WorldContext pin to indicate that the function can be called even if the class does not implement the virtual function GetWorld().
+		CallableWithoutWorldContext,
+
 		/// [FunctionMetadata] Indicates that a BlueprintCallable function should use the Commutative Associative Binary node.
 		CommutativeAssociativeBinaryOperator,
 
