Keeping nodes from spawning unessecery menu actions (was causing a big memory spike for the new menu system).

diff --git a/Engine/Source/Editor/AnimGraph/Private/AnimGraphNode_Base.cpp b/Engine/Source/Editor/AnimGraph/Private/AnimGraphNode_Base.cpp
--- a/Engine/Source/Editor/AnimGraph/Private/AnimGraphNode_Base.cpp
+++ b/Engine/Source/Editor/AnimGraph/Private/AnimGraphNode_Base.cpp
@@ -230,9 +230,21 @@ void UAnimGraphNode_Base::GetNodeAttributes( TArray<TKeyValuePair<FString, FStri
 
 void UAnimGraphNode_Base::GetMenuActions(FBlueprintActionDatabaseRegistrar& ActionRegistrar) const
 {
-	UBlueprintNodeSpawner* NodeSpawner = UBlueprintNodeSpawner::Create(GetClass());
-	check(NodeSpawner != nullptr);
-	ActionRegistrar.AddBlueprintAction(NodeSpawner);
+	// actions get registered under specific object-keys; the idea is that 
+	// actions might have to be updated (or deleted) if their object-key is  
+	// mutated (or removed)... here we use the node's class (so if the node 
+	// type disappears, then the action should go with it)
+	UClass* ActionKey = GetClass();
+	// to keep from needlessly instantiating a UBlueprintNodeSpawner, first   
+	// check to make sure that the registrar is looking for actions of this type
+	// (could be regenerating actions for a specific asset, and therefore the 
+	// registrar would only accept actions corresponding to that asset)
+	if (ActionRegistrar.IsOpenForRegistration(ActionKey))
+	{
+		UBlueprintNodeSpawner* NodeSpawner = UBlueprintNodeSpawner::Create(GetClass());
+		check(NodeSpawner != nullptr);
+		ActionRegistrar.AddBlueprintAction(ActionKey, NodeSpawner);
+	}
 }
 
 FText UAnimGraphNode_Base::GetMenuCategory() const
diff --git a/Engine/Source/Editor/AnimGraph/Private/AnimGraphNode_BlendSpacePlayer.cpp b/Engine/Source/Editor/AnimGraph/Private/AnimGraphNode_BlendSpacePlayer.cpp
--- a/Engine/Source/Editor/AnimGraph/Private/AnimGraphNode_BlendSpacePlayer.cpp
+++ b/Engine/Source/Editor/AnimGraph/Private/AnimGraphNode_BlendSpacePlayer.cpp
@@ -84,6 +84,14 @@ void UAnimGraphNode_BlendSpacePlayer::GetMenuActions(FBlueprintActionDatabaseReg
 	for (TObjectIterator<UBlendSpaceBase> BlendSpaceIt; BlendSpaceIt; ++BlendSpaceIt)
 	{
 		UBlendSpaceBase* BlendSpace = *BlendSpaceIt;
+		// to keep from needlessly instantiating a UBlueprintNodeSpawner, first   
+		// check to make sure that the registrar is looking for actions of this type
+		// (could be regenerating actions for a specific asset, and therefore the 
+		// registrar would only accept actions corresponding to that asset)
+		if (!ActionRegistrar.IsOpenForRegistration(BlendSpace))
+		{
+			continue;
+		}
 
 		bool const bIsAimOffset = BlendSpace->IsA(UAimOffsetBlendSpace::StaticClass()) || 
 			BlendSpace->IsA(UAimOffsetBlendSpace1D::StaticClass());
@@ -91,8 +99,7 @@ void UAnimGraphNode_BlendSpacePlayer::GetMenuActions(FBlueprintActionDatabaseReg
 		{
 			UBlueprintNodeSpawner* NodeSpawner = UBlueprintNodeSpawner::Create(GetClass());
 			check(NodeSpawner != nullptr);
-			// @TODO: Need to file this one under the BlendSpace asset (so it can be cleared/updated when the asset is)
-			ActionRegistrar.AddBlueprintAction(NodeSpawner);
+			ActionRegistrar.AddBlueprintAction(BlendSpace, NodeSpawner);
 
 			TWeakObjectPtr<UBlendSpaceBase> BlendSpacePtr = BlendSpace;
 			NodeSpawner->CustomizeNodeDelegate = UBlueprintNodeSpawner::FCustomizeNodeDelegate::CreateStatic(PostSpawnSetupLambda, BlendSpacePtr);
diff --git a/Engine/Source/Editor/AnimGraph/Private/AnimGraphNode_RotationOffsetBlendSpace.cpp b/Engine/Source/Editor/AnimGraph/Private/AnimGraphNode_RotationOffsetBlendSpace.cpp
--- a/Engine/Source/Editor/AnimGraph/Private/AnimGraphNode_RotationOffsetBlendSpace.cpp
+++ b/Engine/Source/Editor/AnimGraph/Private/AnimGraphNode_RotationOffsetBlendSpace.cpp
@@ -73,6 +73,14 @@ void UAnimGraphNode_RotationOffsetBlendSpace::GetMenuActions(FBlueprintActionDat
 	for (TObjectIterator<UBlendSpaceBase> BlendSpaceIt; BlendSpaceIt; ++BlendSpaceIt)
 	{
 		UBlendSpaceBase* BlendSpace = *BlendSpaceIt;
+		// to keep from needlessly instantiating a UBlueprintNodeSpawner, first   
+		// check to make sure that the registrar is looking for actions of this type
+		// (could be regenerating actions for a specific asset, and therefore the 
+		// registrar would only accept actions corresponding to that asset)
+		if (!ActionRegistrar.IsOpenForRegistration(BlendSpace))
+		{
+			continue;
+		}
 
 		bool const bIsAimOffset = BlendSpace->IsA(UAimOffsetBlendSpace::StaticClass()) ||
 			BlendSpace->IsA(UAimOffsetBlendSpace1D::StaticClass());
@@ -80,8 +88,7 @@ void UAnimGraphNode_RotationOffsetBlendSpace::GetMenuActions(FBlueprintActionDat
 		{
 			UBlueprintNodeSpawner* NodeSpawner = UBlueprintNodeSpawner::Create(GetClass());
 			check(NodeSpawner != nullptr);
-			// @TODO: Need to file this one under the BlendSpace asset (so it can be cleared/updated when the asset is)
-			ActionRegistrar.AddBlueprintAction(NodeSpawner);
+			ActionRegistrar.AddBlueprintAction(BlendSpace, NodeSpawner);
 
 			TWeakObjectPtr<UBlendSpaceBase> BlendSpacePtr = BlendSpace;
 			NodeSpawner->CustomizeNodeDelegate = UBlueprintNodeSpawner::FCustomizeNodeDelegate::CreateStatic(PostSpawnSetupLambda, BlendSpacePtr);
diff --git a/Engine/Source/Editor/AnimGraph/Private/AnimGraphNode_SaveCachedPose.cpp b/Engine/Source/Editor/AnimGraph/Private/AnimGraphNode_SaveCachedPose.cpp
--- a/Engine/Source/Editor/AnimGraph/Private/AnimGraphNode_SaveCachedPose.cpp
+++ b/Engine/Source/Editor/AnimGraph/Private/AnimGraphNode_SaveCachedPose.cpp
@@ -102,10 +102,22 @@ void UAnimGraphNode_SaveCachedPose::GetMenuActions(FBlueprintActionDatabaseRegis
 		}
 	};
 
-	UBlueprintNodeSpawner* NodeSpawner = UBlueprintNodeSpawner::Create(GetClass());
-	NodeSpawner->CustomizeNodeDelegate = UBlueprintNodeSpawner::FCustomizeNodeDelegate::CreateStatic(PostSpawnSetupLambda);
+	// actions get registered under specific object-keys; the idea is that 
+	// actions might have to be updated (or deleted) if their object-key is  
+	// mutated (or removed)... here we use the node's class (so if the node 
+	// type disappears, then the action should go with it)
+	UClass* ActionKey = GetClass();
+	// to keep from needlessly instantiating a UBlueprintNodeSpawner, first   
+	// check to make sure that the registrar is looking for actions of this type
+	// (could be regenerating actions for a specific asset, and therefore the 
+	// registrar would only accept actions corresponding to that asset)
+	if (ActionRegistrar.IsOpenForRegistration(ActionKey))
+	{
+		UBlueprintNodeSpawner* NodeSpawner = UBlueprintNodeSpawner::Create(GetClass());
+		NodeSpawner->CustomizeNodeDelegate = UBlueprintNodeSpawner::FCustomizeNodeDelegate::CreateStatic(PostSpawnSetupLambda);
 
-	ActionRegistrar.AddBlueprintAction(NodeSpawner);
+		ActionRegistrar.AddBlueprintAction(ActionKey, NodeSpawner);
+	}
 }
 
 bool UAnimGraphNode_SaveCachedPose::IsCompatibleWithGraph(const UEdGraph* TargetGraph) const
diff --git a/Engine/Source/Editor/BlueprintGraph/Classes/K2Node.h b/Engine/Source/Editor/BlueprintGraph/Classes/K2Node.h
--- a/Engine/Source/Editor/BlueprintGraph/Classes/K2Node.h
+++ b/Engine/Source/Editor/BlueprintGraph/Classes/K2Node.h
@@ -298,7 +298,7 @@ class UK2Node : public UEdGraphNode
 	 *
 	 * @param  ActionListOut	The list to be populated with new spawners.
 	 */
-	virtual void GetMenuActions(FBlueprintActionDatabaseRegistrar& ActionRegisrar) const {}
+	virtual void GetMenuActions(FBlueprintActionDatabaseRegistrar& ActionRegistrar) const {}
 
 	/**
 	 * Override to provide a default category for specific node types to be 
diff --git a/Engine/Source/Editor/BlueprintGraph/Classes/K2Node_ConstructObjectFromClass.h b/Engine/Source/Editor/BlueprintGraph/Classes/K2Node_ConstructObjectFromClass.h
--- a/Engine/Source/Editor/BlueprintGraph/Classes/K2Node_ConstructObjectFromClass.h
+++ b/Engine/Source/Editor/BlueprintGraph/Classes/K2Node_ConstructObjectFromClass.h
@@ -26,7 +26,7 @@ class BLUEPRINTGRAPH_API UK2Node_ConstructObjectFromClass : public UK2Node
 	virtual bool IsNodeSafeToIgnore() const override { return true; }
 	virtual void ReallocatePinsDuringReconstruction(TArray<UEdGraphPin*>& OldPins) override;
 	virtual void GetNodeAttributes( TArray<TKeyValuePair<FString, FString>>& OutNodeAttributes ) const override;
-	virtual void GetMenuActions(FBlueprintActionDatabaseRegistrar& ActionRegisrar) const override;
+	virtual void GetMenuActions(FBlueprintActionDatabaseRegistrar& ActionRegistrar) const override;
 	virtual FText GetMenuCategory() const override;
 	virtual class FNodeHandlingFunctor* CreateNodeHandler(class FKismetCompilerContext& CompilerContext) const override;
 	// End UK2Node interface
diff --git a/Engine/Source/Editor/BlueprintGraph/Private/BlueprintActionDatabase.cpp b/Engine/Source/Editor/BlueprintGraph/Private/BlueprintActionDatabase.cpp
--- a/Engine/Source/Editor/BlueprintGraph/Private/BlueprintActionDatabase.cpp
+++ b/Engine/Source/Editor/BlueprintGraph/Private/BlueprintActionDatabase.cpp
@@ -846,7 +846,10 @@ void FBlueprintActionDatabase::RefreshClassActions(UClass* const Class)
 	else if (Class->IsChildOf<UEdGraphNode>())
 	{
 		FActionList& ClassActionList = ActionRegistry.FindOrAdd(Class);
-		ClassActionList.Empty();
+		if (!bIsInitializing)
+		{
+			ClassActionList.Empty();
+		}
 
 		FBlueprintActionDatabaseRegistrar Registrar(ActionRegistry, ActionPrimingQueue, Class);
 		if (!bIsInitializing)
diff --git a/Engine/Source/Editor/BlueprintGraph/Private/BlueprintActionDatabaseRegistrar.cpp b/Engine/Source/Editor/BlueprintGraph/Private/BlueprintActionDatabaseRegistrar.cpp
--- a/Engine/Source/Editor/BlueprintGraph/Private/BlueprintActionDatabaseRegistrar.cpp
+++ b/Engine/Source/Editor/BlueprintGraph/Private/BlueprintActionDatabaseRegistrar.cpp
@@ -69,7 +69,7 @@ FBlueprintActionDatabaseRegistrar::FBlueprintActionDatabaseRegistrar(FActionRegi
 }
 
 //------------------------------------------------------------------------------
-void FBlueprintActionDatabaseRegistrar::AddBlueprintAction(UBlueprintNodeSpawner* NodeSpawner)
+bool FBlueprintActionDatabaseRegistrar::AddBlueprintAction(UBlueprintNodeSpawner* NodeSpawner)
 {
 	UField const* ActionKey = GeneratingClass;
 	// if this spawner wraps some member function/property, then we want it 
@@ -79,49 +79,64 @@ void FBlueprintActionDatabaseRegistrar::AddBlueprintAction(UBlueprintNodeSpawner
 	{
 		ActionKey = MemberField;
 	}
-	AddActionToDatabase(ActionKey, NodeSpawner);
+	return AddActionToDatabase(ActionKey, NodeSpawner);
 }
 
 //------------------------------------------------------------------------------
-void FBlueprintActionDatabaseRegistrar::AddBlueprintAction(UClass const* ClassOwner, UBlueprintNodeSpawner* NodeSpawner)
+bool FBlueprintActionDatabaseRegistrar::AddBlueprintAction(UClass const* ClassOwner, UBlueprintNodeSpawner* NodeSpawner)
 {
+	// @TODO: assert that AddBlueprintAction(UBlueprintNodeSpawner* NodeSpawner) 
+	//        wouldn't come up with a different key (besides GeneratingClass)
+
 	// ResolveActionKey() is used on ClassOwner (in AddActionToDatabase), to 
 	// convert it into a proper key
-	AddActionToDatabase(ClassOwner, NodeSpawner);
+	return AddActionToDatabase(ClassOwner, NodeSpawner);
 }
 
 //------------------------------------------------------------------------------
-void FBlueprintActionDatabaseRegistrar::AddBlueprintAction(UEnum const* EnumOwner, UBlueprintNodeSpawner* NodeSpawner)
+bool FBlueprintActionDatabaseRegistrar::AddBlueprintAction(UEnum const* EnumOwner, UBlueprintNodeSpawner* NodeSpawner)
 {
+	// @TODO: assert that AddBlueprintAction(UBlueprintNodeSpawner* NodeSpawner) 
+	//        wouldn't come up with a different key (besides GeneratingClass)
+
 	// ResolveActionKey() is used on EnumOwner (in AddActionToDatabase), to 
 	// convert it into a proper key
-	AddActionToDatabase(EnumOwner, NodeSpawner);
+	return AddActionToDatabase(EnumOwner, NodeSpawner);
 }
 
 //------------------------------------------------------------------------------
-void FBlueprintActionDatabaseRegistrar::AddBlueprintAction(UScriptStruct const* StructOwner, UBlueprintNodeSpawner* NodeSpawner)
+bool FBlueprintActionDatabaseRegistrar::AddBlueprintAction(UScriptStruct const* StructOwner, UBlueprintNodeSpawner* NodeSpawner)
 {
+	// @TODO: assert that AddBlueprintAction(UBlueprintNodeSpawner* NodeSpawner) 
+	//        wouldn't come up with a different key (besides GeneratingClass)
+
 	// ResolveActionKey() is used on StructOwner (in AddActionToDatabase), to 
 	// convert it into a proper key
-	AddActionToDatabase(StructOwner, NodeSpawner);
+	return AddActionToDatabase(StructOwner, NodeSpawner);
 }
 
 //------------------------------------------------------------------------------
-void FBlueprintActionDatabaseRegistrar::AddBlueprintAction(UField const* FieldOwner, UBlueprintNodeSpawner* NodeSpawner)
+bool FBlueprintActionDatabaseRegistrar::AddBlueprintAction(UField const* FieldOwner, UBlueprintNodeSpawner* NodeSpawner)
 {
+	// @TODO: assert that AddBlueprintAction(UBlueprintNodeSpawner* NodeSpawner) 
+	//        wouldn't come up with a different key (besides GeneratingClass)
+
 	// ResolveActionKey() is used on FieldOwner (in AddActionToDatabase), to 
 	// convert it into a proper key
-	AddActionToDatabase(FieldOwner, NodeSpawner);
+	return AddActionToDatabase(FieldOwner, NodeSpawner);
 }
 
 //------------------------------------------------------------------------------
-void FBlueprintActionDatabaseRegistrar::AddBlueprintAction(UObject const* AssetOwner, UBlueprintNodeSpawner* NodeSpawner)
+bool FBlueprintActionDatabaseRegistrar::AddBlueprintAction(UObject const* AssetOwner, UBlueprintNodeSpawner* NodeSpawner)
 {
+	// @TODO: assert that AddBlueprintAction(UBlueprintNodeSpawner* NodeSpawner) 
+	//        wouldn't come up with a different key (besides GeneratingClass)
+
 	// cannot record an action under any ol' object (we want to associate them 
 	// with asset/class outers that are subject to change; so that we can 
 	// refresh/rebuild corresponding actions when that happens).
 	check(AssetOwner->IsAsset());
-	AddActionToDatabase(AssetOwner, NodeSpawner);
+	return AddActionToDatabase(AssetOwner, NodeSpawner);
 }
 
 //------------------------------------------------------------------------------
@@ -136,7 +151,7 @@ bool FBlueprintActionDatabaseRegistrar::IsOpenForRegistration(UObject const* Own
 }
 
 //------------------------------------------------------------------------------
-void FBlueprintActionDatabaseRegistrar::AddActionToDatabase(UObject const* ActionKey, UBlueprintNodeSpawner* NodeSpawner)
+bool FBlueprintActionDatabaseRegistrar::AddActionToDatabase(UObject const* ActionKey, UBlueprintNodeSpawner* NodeSpawner)
 {
 	ensureMsg(NodeSpawner->NodeClass == GeneratingClass, TEXT("We expect a nodes to add only spawners for its own type... Maybe a sub-class is adding nodes it shouldn't?"));
 	if (IsOpenForRegistration(ActionKey))
@@ -156,5 +171,7 @@ void FBlueprintActionDatabaseRegistrar::AddActionToDatabase(UObject const* Actio
 		}
 
 		ActionList.Add(NodeSpawner);
+		return true;
 	}
+	return false;
 }
\ No newline at end of file
diff --git a/Engine/Source/Editor/BlueprintGraph/Private/K2Node_AssignmentStatement.cpp b/Engine/Source/Editor/BlueprintGraph/Private/K2Node_AssignmentStatement.cpp
--- a/Engine/Source/Editor/BlueprintGraph/Private/K2Node_AssignmentStatement.cpp
+++ b/Engine/Source/Editor/BlueprintGraph/Private/K2Node_AssignmentStatement.cpp
@@ -177,10 +177,22 @@ FNodeHandlingFunctor* UK2Node_AssignmentStatement::CreateNodeHandler(FKismetComp
 
 void UK2Node_AssignmentStatement::GetMenuActions(FBlueprintActionDatabaseRegistrar& ActionRegistrar) const
 {
-	UBlueprintNodeSpawner* NodeSpawner = UBlueprintNodeSpawner::Create(GetClass());
-	check(NodeSpawner != nullptr);
+	// actions get registered under specific object-keys; the idea is that 
+	// actions might have to be updated (or deleted) if their object-key is  
+	// mutated (or removed)... here we use the node's class (so if the node 
+	// type disappears, then the action should go with it)
+	UClass* ActionKey = GetClass();
+	// to keep from needlessly instantiating a UBlueprintNodeSpawner, first   
+	// check to make sure that the registrar is looking for actions of this type
+	// (could be regenerating actions for a specific asset, and therefore the 
+	// registrar would only accept actions corresponding to that asset)
+	if (ActionRegistrar.IsOpenForRegistration(ActionKey))
+	{
+		UBlueprintNodeSpawner* NodeSpawner = UBlueprintNodeSpawner::Create(GetClass());
+		check(NodeSpawner != nullptr);
 
-	ActionRegistrar.AddBlueprintAction(NodeSpawner);
+		ActionRegistrar.AddBlueprintAction(ActionKey, NodeSpawner);
+	}
 }
 
 FText UK2Node_AssignmentStatement::GetMenuCategory() const
diff --git a/Engine/Source/Editor/BlueprintGraph/Private/K2Node_BaseAsyncTask.cpp b/Engine/Source/Editor/BlueprintGraph/Private/K2Node_BaseAsyncTask.cpp
--- a/Engine/Source/Editor/BlueprintGraph/Private/K2Node_BaseAsyncTask.cpp
+++ b/Engine/Source/Editor/BlueprintGraph/Private/K2Node_BaseAsyncTask.cpp
@@ -378,10 +378,22 @@ FText UK2Node_BaseAsyncTask::GetMenuCategory() const
 
 void UK2Node_BaseAsyncTask::GetMenuActions(FBlueprintActionDatabaseRegistrar& ActionRegistrar) const
 {
-	UBlueprintNodeSpawner* NodeSpawner = UBlueprintNodeSpawner::Create(GetClass());
-	check(NodeSpawner != nullptr);
-	
-	ActionRegistrar.AddBlueprintAction(NodeSpawner);
+	// actions get registered under specific object-keys; the idea is that 
+	// actions might have to be updated (or deleted) if their object-key is  
+	// mutated (or removed)... here we use the node's class (so if the node 
+	// type disappears, then the action should go with it)
+	UClass* ActionKey = GetClass();
+	// to keep from needlessly instantiating a UBlueprintNodeSpawner, first   
+	// check to make sure that the registrar is looking for actions of this type
+	// (could be regenerating actions for a specific asset, and therefore the 
+	// registrar would only accept actions corresponding to that asset)
+	if (ActionRegistrar.IsOpenForRegistration(ActionKey))
+	{
+		UBlueprintNodeSpawner* NodeSpawner = UBlueprintNodeSpawner::Create(GetClass());
+		check(NodeSpawner != nullptr);
+
+		ActionRegistrar.AddBlueprintAction(ActionKey, NodeSpawner);
+	}
 }
 
 UFunction* UK2Node_BaseAsyncTask::GetFactoryFunction() const
diff --git a/Engine/Source/Editor/BlueprintGraph/Private/K2Node_BreakStruct.cpp b/Engine/Source/Editor/BlueprintGraph/Private/K2Node_BreakStruct.cpp
--- a/Engine/Source/Editor/BlueprintGraph/Private/K2Node_BreakStruct.cpp
+++ b/Engine/Source/Editor/BlueprintGraph/Private/K2Node_BreakStruct.cpp
@@ -313,6 +313,15 @@ FNodeHandlingFunctor* UK2Node_BreakStruct::CreateNodeHandler(class FKismetCompil
 
 void UK2Node_BreakStruct::GetMenuActions(FBlueprintActionDatabaseRegistrar& ActionRegistrar) const
 {
+	auto CustomizeBreakNodeLambda = [](UEdGraphNode* NewNode, bool bIsTemplateNode, TWeakObjectPtr<UScriptStruct> StructPtr)
+	{
+		UK2Node_BreakStruct* BreakNode = CastChecked<UK2Node_BreakStruct>(NewNode);
+		if (StructPtr.IsValid())
+		{
+			BreakNode->StructType = StructPtr.Get();
+		}
+	};
+
 	for (TObjectIterator<UScriptStruct> StructIt; StructIt; ++StructIt)
 	{
 		UScriptStruct const* Struct = (*StructIt);
@@ -321,14 +330,14 @@ void UK2Node_BreakStruct::GetMenuActions(FBlueprintActionDatabaseRegistrar& Acti
 			continue;
 		}
 
-		auto CustomizeBreakNodeLambda = [](UEdGraphNode* NewNode, bool bIsTemplateNode, TWeakObjectPtr<UScriptStruct> StructPtr)
+		// to keep from needlessly instantiating a UBlueprintNodeSpawner, first   
+		// check to make sure that the registrar is looking for actions of this type
+		// (could be regenerating actions for a specific asset, and therefore the 
+		// registrar would only accept actions corresponding to that asset)
+		if (!ActionRegistrar.IsOpenForRegistration(Struct))
 		{
-			UK2Node_BreakStruct* BreakNode = CastChecked<UK2Node_BreakStruct>(NewNode);
-			if (StructPtr.IsValid())
-			{
-				BreakNode->StructType = StructPtr.Get();
-			}
-		};
+			continue;
+		}
 
 		UBlueprintNodeSpawner* NodeSpawner = UBlueprintNodeSpawner::Create(GetClass());
 		check(NodeSpawner != nullptr);
diff --git a/Engine/Source/Editor/BlueprintGraph/Private/K2Node_CastByteToEnum.cpp b/Engine/Source/Editor/BlueprintGraph/Private/K2Node_CastByteToEnum.cpp
--- a/Engine/Source/Editor/BlueprintGraph/Private/K2Node_CastByteToEnum.cpp
+++ b/Engine/Source/Editor/BlueprintGraph/Private/K2Node_CastByteToEnum.cpp
@@ -150,6 +150,16 @@ FNodeHandlingFunctor* UK2Node_CastByteToEnum::CreateNodeHandler(FKismetCompilerC
 
 void UK2Node_CastByteToEnum::GetMenuActions(FBlueprintActionDatabaseRegistrar& ActionRegistrar) const
 {
+	auto CustomizeEnumNodeLambda = [](UEdGraphNode* NewNode, bool bIsTemplateNode, TWeakObjectPtr<UEnum> EnumPtr)
+	{
+		UK2Node_CastByteToEnum* EnumNode = CastChecked<UK2Node_CastByteToEnum>(NewNode);
+		if (EnumPtr.IsValid())
+		{
+			EnumNode->Enum = EnumPtr.Get();
+		}
+		EnumNode->bSafe = true;
+	};
+
 	for (TObjectIterator<UEnum> EnumIt; EnumIt; ++EnumIt)
 	{
 		UEnum const* Enum = (*EnumIt);
@@ -158,15 +168,14 @@ void UK2Node_CastByteToEnum::GetMenuActions(FBlueprintActionDatabaseRegistrar& A
 			continue;
 		}
 
-		auto CustomizeEnumNodeLambda = [](UEdGraphNode* NewNode, bool bIsTemplateNode, TWeakObjectPtr<UEnum> EnumPtr)
+		// to keep from needlessly instantiating a UBlueprintNodeSpawner, first   
+		// check to make sure that the registrar is looking for actions of this type
+		// (could be regenerating actions for a specific asset, and therefore the 
+		// registrar would only accept actions corresponding to that asset)
+		if (!ActionRegistrar.IsOpenForRegistration(Enum))
 		{
-			UK2Node_CastByteToEnum* EnumNode = CastChecked<UK2Node_CastByteToEnum>(NewNode);
-			if (EnumPtr.IsValid())
-			{
-				EnumNode->Enum = EnumPtr.Get();
-			}
-			EnumNode->bSafe = true;
-		};
+			continue;
+		}
 
 		UBlueprintNodeSpawner* NodeSpawner = UBlueprintNodeSpawner::Create(GetClass());
 		check(NodeSpawner != nullptr);
diff --git a/Engine/Source/Editor/BlueprintGraph/Private/K2Node_ConstructObjectFromClass.cpp b/Engine/Source/Editor/BlueprintGraph/Private/K2Node_ConstructObjectFromClass.cpp
--- a/Engine/Source/Editor/BlueprintGraph/Private/K2Node_ConstructObjectFromClass.cpp
+++ b/Engine/Source/Editor/BlueprintGraph/Private/K2Node_ConstructObjectFromClass.cpp
@@ -293,12 +293,24 @@ void UK2Node_ConstructObjectFromClass::GetNodeAttributes( TArray<TKeyValuePair<F
 	OutNodeAttributes.Add( TKeyValuePair<FString, FString>( TEXT( "ObjectClass" ), ClassToSpawnStr ));
 }
 
-void UK2Node_ConstructObjectFromClass::GetMenuActions(FBlueprintActionDatabaseRegistrar& ActionRegisrar) const
+void UK2Node_ConstructObjectFromClass::GetMenuActions(FBlueprintActionDatabaseRegistrar& ActionRegistrar) const
 {
-	UBlueprintNodeSpawner* NodeSpawner = UBlueprintNodeSpawner::Create(GetClass());
-	check(NodeSpawner != nullptr);
+	// actions get registered under specific object-keys; the idea is that 
+	// actions might have to be updated (or deleted) if their object-key is  
+	// mutated (or removed)... here we use the node's class (so if the node 
+	// type disappears, then the action should go with it)
+	UClass* ActionKey = GetClass();
+	// to keep from needlessly instantiating a UBlueprintNodeSpawner, first   
+	// check to make sure that the registrar is looking for actions of this type
+	// (could be regenerating actions for a specific asset, and therefore the 
+	// registrar would only accept actions corresponding to that asset)
+	if (ActionRegistrar.IsOpenForRegistration(ActionKey))
+	{
+		UBlueprintNodeSpawner* NodeSpawner = UBlueprintNodeSpawner::Create(GetClass());
+		check(NodeSpawner != nullptr);
 
-	ActionRegisrar.AddBlueprintAction(NodeSpawner);
+		ActionRegistrar.AddBlueprintAction(ActionKey, NodeSpawner);
+	}
 }
 
 FText UK2Node_ConstructObjectFromClass::GetMenuCategory() const
diff --git a/Engine/Source/Editor/BlueprintGraph/Private/K2Node_CustomEvent.cpp b/Engine/Source/Editor/BlueprintGraph/Private/K2Node_CustomEvent.cpp
--- a/Engine/Source/Editor/BlueprintGraph/Private/K2Node_CustomEvent.cpp
+++ b/Engine/Source/Editor/BlueprintGraph/Private/K2Node_CustomEvent.cpp
@@ -237,24 +237,36 @@ void UK2Node_CustomEvent::ValidateNodeDuringCompilation(class FCompilerResultsLo
 
 void UK2Node_CustomEvent::GetMenuActions(FBlueprintActionDatabaseRegistrar& ActionRegistrar) const
 {
-	UBlueprintNodeSpawner* NodeSpawner = UBlueprintEventNodeSpawner::Create(GetClass(), FName());
-	check(NodeSpawner != nullptr);
-
-	auto SetupCustomEventNodeLambda = [](UEdGraphNode* NewNode, bool bIsTemplateNode)
+	// actions get registered under specific object-keys; the idea is that 
+	// actions might have to be updated (or deleted) if their object-key is  
+	// mutated (or removed)... here we use the node's class (so if the node 
+	// type disappears, then the action should go with it)
+	UClass* ActionKey = GetClass();
+	// to keep from needlessly instantiating a UBlueprintNodeSpawner, first   
+	// check to make sure that the registrar is looking for actions of this type
+	// (could be regenerating actions for a specific asset, and therefore the 
+	// registrar would only accept actions corresponding to that asset)
+	if (ActionRegistrar.IsOpenForRegistration(ActionKey))
 	{
-		UK2Node_CustomEvent* EventNode = CastChecked<UK2Node_CustomEvent>(NewNode);
-		UBlueprint* Blueprint = EventNode->GetBlueprint();
+		UBlueprintNodeSpawner* NodeSpawner = UBlueprintEventNodeSpawner::Create(GetClass(), FName());
+		check(NodeSpawner != nullptr);
 
-		// in GetNodeTitle(), we use an empty CustomFunctionName to identify a menu entry
-		if (!bIsTemplateNode)
+		auto SetupCustomEventNodeLambda = [](UEdGraphNode* NewNode, bool bIsTemplateNode)
 		{
-			EventNode->CustomFunctionName = FBlueprintEditorUtils::FindUniqueCustomEventName(Blueprint);
-		}
-		EventNode->bIsEditable = true;
-	};
+			UK2Node_CustomEvent* EventNode = CastChecked<UK2Node_CustomEvent>(NewNode);
+			UBlueprint* Blueprint = EventNode->GetBlueprint();
+
+			// in GetNodeTitle(), we use an empty CustomFunctionName to identify a menu entry
+			if (!bIsTemplateNode)
+			{
+				EventNode->CustomFunctionName = FBlueprintEditorUtils::FindUniqueCustomEventName(Blueprint);
+			}
+			EventNode->bIsEditable = true;
+		};
 
-	NodeSpawner->CustomizeNodeDelegate = UBlueprintNodeSpawner::FCustomizeNodeDelegate::CreateStatic(SetupCustomEventNodeLambda);
-	ActionRegistrar.AddBlueprintAction(NodeSpawner);
+		NodeSpawner->CustomizeNodeDelegate = UBlueprintNodeSpawner::FCustomizeNodeDelegate::CreateStatic(SetupCustomEventNodeLambda);
+		ActionRegistrar.AddBlueprintAction(ActionKey, NodeSpawner);
+	}
 }
 
 void UK2Node_CustomEvent::ReconstructNode()
diff --git a/Engine/Source/Editor/BlueprintGraph/Private/K2Node_DoOnceMultiInput.cpp b/Engine/Source/Editor/BlueprintGraph/Private/K2Node_DoOnceMultiInput.cpp
--- a/Engine/Source/Editor/BlueprintGraph/Private/K2Node_DoOnceMultiInput.cpp
+++ b/Engine/Source/Editor/BlueprintGraph/Private/K2Node_DoOnceMultiInput.cpp
@@ -378,10 +378,22 @@ void UK2Node_DoOnceMultiInput::ExpandNode(FKismetCompilerContext& CompilerContex
 
 void UK2Node_DoOnceMultiInput::GetMenuActions(FBlueprintActionDatabaseRegistrar& ActionRegistrar) const
 {
-	UBlueprintNodeSpawner* NodeSpawner = UBlueprintNodeSpawner::Create(GetClass());
-	check(NodeSpawner != nullptr);
-	
-	ActionRegistrar.AddBlueprintAction(NodeSpawner);
+	// actions get registered under specific object-keys; the idea is that 
+	// actions might have to be updated (or deleted) if their object-key is  
+	// mutated (or removed)... here we use the node's class (so if the node 
+	// type disappears, then the action should go with it)
+	UClass* ActionKey = GetClass();
+	// to keep from needlessly instantiating a UBlueprintNodeSpawner, first   
+	// check to make sure that the registrar is looking for actions of this type
+	// (could be regenerating actions for a specific asset, and therefore the 
+	// registrar would only accept actions corresponding to that asset)
+	if (ActionRegistrar.IsOpenForRegistration(ActionKey))
+	{
+		UBlueprintNodeSpawner* NodeSpawner = UBlueprintNodeSpawner::Create(GetClass());
+		check(NodeSpawner != nullptr);
+
+		ActionRegistrar.AddBlueprintAction(ActionKey, NodeSpawner);
+	}
 }
 
 FText UK2Node_DoOnceMultiInput::GetMenuCategory() const
diff --git a/Engine/Source/Editor/BlueprintGraph/Private/K2Node_EaseFunction.cpp b/Engine/Source/Editor/BlueprintGraph/Private/K2Node_EaseFunction.cpp
--- a/Engine/Source/Editor/BlueprintGraph/Private/K2Node_EaseFunction.cpp
+++ b/Engine/Source/Editor/BlueprintGraph/Private/K2Node_EaseFunction.cpp
@@ -237,10 +237,22 @@ bool UK2Node_EaseFunction::IsConnectionDisallowed(const UEdGraphPin* MyPin, cons
 
 void UK2Node_EaseFunction::GetMenuActions(FBlueprintActionDatabaseRegistrar& ActionRegistrar) const
 {
-	UBlueprintNodeSpawner* NodeSpawner = UBlueprintNodeSpawner::Create(GetClass());
-	check(NodeSpawner != nullptr);
+	// actions get registered under specific object-keys; the idea is that 
+	// actions might have to be updated (or deleted) if their object-key is  
+	// mutated (or removed)... here we use the node's class (so if the node 
+	// type disappears, then the action should go with it)
+	UClass* ActionKey = GetClass();
+	// to keep from needlessly instantiating a UBlueprintNodeSpawner, first   
+	// check to make sure that the registrar is looking for actions of this type
+	// (could be regenerating actions for a specific asset, and therefore the 
+	// registrar would only accept actions corresponding to that asset)
+	if (ActionRegistrar.IsOpenForRegistration(ActionKey))
+	{
+		UBlueprintNodeSpawner* NodeSpawner = UBlueprintNodeSpawner::Create(GetClass());
+		check(NodeSpawner != nullptr);
 
-	ActionRegistrar.AddBlueprintAction(NodeSpawner);
+		ActionRegistrar.AddBlueprintAction(ActionKey, NodeSpawner);
+	}
 }
 
 FText UK2Node_EaseFunction::GetMenuCategory() const
diff --git a/Engine/Source/Editor/BlueprintGraph/Private/K2Node_EnumEquality.cpp b/Engine/Source/Editor/BlueprintGraph/Private/K2Node_EnumEquality.cpp
--- a/Engine/Source/Editor/BlueprintGraph/Private/K2Node_EnumEquality.cpp
+++ b/Engine/Source/Editor/BlueprintGraph/Private/K2Node_EnumEquality.cpp
@@ -259,10 +259,22 @@ void UK2Node_EnumEquality::ExpandNode(class FKismetCompilerContext& CompilerCont
 
 void UK2Node_EnumEquality::GetMenuActions(FBlueprintActionDatabaseRegistrar& ActionRegistrar) const
 {
-	UBlueprintNodeSpawner* NodeSpawner = UBlueprintNodeSpawner::Create(GetClass());
-	check(NodeSpawner != nullptr);
-	
-	ActionRegistrar.AddBlueprintAction(NodeSpawner);
+	// actions get registered under specific object-keys; the idea is that 
+	// actions might have to be updated (or deleted) if their object-key is  
+	// mutated (or removed)... here we use the node's class (so if the node 
+	// type disappears, then the action should go with it)
+	UClass* ActionKey = GetClass();
+	// to keep from needlessly instantiating a UBlueprintNodeSpawner, first   
+	// check to make sure that the registrar is looking for actions of this type
+	// (could be regenerating actions for a specific asset, and therefore the 
+	// registrar would only accept actions corresponding to that asset)
+	if (ActionRegistrar.IsOpenForRegistration(ActionKey))
+	{
+		UBlueprintNodeSpawner* NodeSpawner = UBlueprintNodeSpawner::Create(GetClass());
+		check(NodeSpawner != nullptr);
+
+		ActionRegistrar.AddBlueprintAction(ActionKey, NodeSpawner);
+	}
 }
 
 FText UK2Node_EnumEquality::GetMenuCategory() const
diff --git a/Engine/Source/Editor/BlueprintGraph/Private/K2Node_EnumLiteral.cpp b/Engine/Source/Editor/BlueprintGraph/Private/K2Node_EnumLiteral.cpp
--- a/Engine/Source/Editor/BlueprintGraph/Private/K2Node_EnumLiteral.cpp
+++ b/Engine/Source/Editor/BlueprintGraph/Private/K2Node_EnumLiteral.cpp
@@ -103,6 +103,15 @@ FNodeHandlingFunctor* UK2Node_EnumLiteral::CreateNodeHandler(FKismetCompilerCont
 
 void UK2Node_EnumLiteral::GetMenuActions(FBlueprintActionDatabaseRegistrar& ActionRegistrar) const
 {
+	auto CustomizeEnumNodeLambda = [](UEdGraphNode* NewNode, bool bIsTemplateNode, TWeakObjectPtr<UEnum> EnumPtr)
+	{
+		UK2Node_EnumLiteral* EnumNode = CastChecked<UK2Node_EnumLiteral>(NewNode);
+		if (EnumPtr.IsValid())
+		{
+			EnumNode->Enum = EnumPtr.Get();
+		}
+	};
+
 	for (TObjectIterator<UEnum> EnumIt; EnumIt; ++EnumIt)
 	{
 		UEnum const* Enum = (*EnumIt);
@@ -111,14 +120,14 @@ void UK2Node_EnumLiteral::GetMenuActions(FBlueprintActionDatabaseRegistrar& Acti
 			continue;
 		}
 
-		auto CustomizeEnumNodeLambda = [](UEdGraphNode* NewNode, bool bIsTemplateNode, TWeakObjectPtr<UEnum> EnumPtr)
+		// to keep from needlessly instantiating a UBlueprintNodeSpawner, first   
+		// check to make sure that the registrar is looking for actions of this type
+		// (could be regenerating actions for a specific asset, and therefore the 
+		// registrar would only accept actions corresponding to that asset)
+		if (!ActionRegistrar.IsOpenForRegistration(Enum))
 		{
-			UK2Node_EnumLiteral* EnumNode = CastChecked<UK2Node_EnumLiteral>(NewNode);
-			if (EnumPtr.IsValid())
-			{
-				EnumNode->Enum = EnumPtr.Get();
-			}
-		};
+			continue;
+		}
 
 		UBlueprintNodeSpawner* NodeSpawner = UBlueprintNodeSpawner::Create(GetClass());
 		check(NodeSpawner != nullptr);
diff --git a/Engine/Source/Editor/BlueprintGraph/Private/K2Node_ExecutionSequence.cpp b/Engine/Source/Editor/BlueprintGraph/Private/K2Node_ExecutionSequence.cpp
--- a/Engine/Source/Editor/BlueprintGraph/Private/K2Node_ExecutionSequence.cpp
+++ b/Engine/Source/Editor/BlueprintGraph/Private/K2Node_ExecutionSequence.cpp
@@ -266,10 +266,22 @@ FNodeHandlingFunctor* UK2Node_ExecutionSequence::CreateNodeHandler(FKismetCompil
 
 void UK2Node_ExecutionSequence::GetMenuActions(FBlueprintActionDatabaseRegistrar& ActionRegistrar) const
 {
-	UBlueprintNodeSpawner* NodeSpawner = UBlueprintNodeSpawner::Create(GetClass());
-	check(NodeSpawner != nullptr);
+	// actions get registered under specific object-keys; the idea is that 
+	// actions might have to be updated (or deleted) if their object-key is  
+	// mutated (or removed)... here we use the node's class (so if the node 
+	// type disappears, then the action should go with it)
+	UClass* ActionKey = GetClass();
+	// to keep from needlessly instantiating a UBlueprintNodeSpawner, first   
+	// check to make sure that the registrar is looking for actions of this type
+	// (could be regenerating actions for a specific asset, and therefore the 
+	// registrar would only accept actions corresponding to that asset)
+	if (ActionRegistrar.IsOpenForRegistration(ActionKey))
+	{
+		UBlueprintNodeSpawner* NodeSpawner = UBlueprintNodeSpawner::Create(GetClass());
+		check(NodeSpawner != nullptr);
 
-	ActionRegistrar.AddBlueprintAction(NodeSpawner);
+		ActionRegistrar.AddBlueprintAction(ActionKey, NodeSpawner);
+	}
 }
 
 FText UK2Node_ExecutionSequence::GetMenuCategory() const
diff --git a/Engine/Source/Editor/BlueprintGraph/Private/K2Node_ForEachElementInEnum.cpp b/Engine/Source/Editor/BlueprintGraph/Private/K2Node_ForEachElementInEnum.cpp
--- a/Engine/Source/Editor/BlueprintGraph/Private/K2Node_ForEachElementInEnum.cpp
+++ b/Engine/Source/Editor/BlueprintGraph/Private/K2Node_ForEachElementInEnum.cpp
@@ -205,6 +205,15 @@ void UK2Node_ForEachElementInEnum::ExpandNode(class FKismetCompilerContext& Comp
 
 void UK2Node_ForEachElementInEnum::GetMenuActions(FBlueprintActionDatabaseRegistrar& ActionRegistrar) const
 {
+	auto CustomizeEnumNodeLambda = [](UEdGraphNode* NewNode, bool bIsTemplateNode, TWeakObjectPtr<UEnum> EnumPtr)
+	{
+		UK2Node_ForEachElementInEnum* EnumNode = CastChecked<UK2Node_ForEachElementInEnum>(NewNode);
+		if (EnumPtr.IsValid())
+		{
+			EnumNode->Enum = EnumPtr.Get();
+		}
+	};
+
 	for (TObjectIterator<UEnum> EnumIt; EnumIt; ++EnumIt)
 	{
 		UEnum const* Enum = (*EnumIt);
@@ -213,14 +222,14 @@ void UK2Node_ForEachElementInEnum::GetMenuActions(FBlueprintActionDatabaseRegist
 			continue;
 		}
 
-		auto CustomizeEnumNodeLambda = [](UEdGraphNode* NewNode, bool bIsTemplateNode, TWeakObjectPtr<UEnum> EnumPtr)
+		// to keep from needlessly instantiating a UBlueprintNodeSpawner, first   
+		// check to make sure that the registrar is looking for actions of this type
+		// (could be regenerating actions for a specific asset, and therefore the 
+		// registrar would only accept actions corresponding to that asset)
+		if (!ActionRegistrar.IsOpenForRegistration(Enum))
 		{
-			UK2Node_ForEachElementInEnum* EnumNode = CastChecked<UK2Node_ForEachElementInEnum>(NewNode);
-			if (EnumPtr.IsValid())
-			{
-				EnumNode->Enum = EnumPtr.Get();
-			}
-		};
+			continue;
+		}
 
 		UBlueprintNodeSpawner* NodeSpawner = UBlueprintNodeSpawner::Create(GetClass());
 		check(NodeSpawner != nullptr);
diff --git a/Engine/Source/Editor/BlueprintGraph/Private/K2Node_FormatText.cpp b/Engine/Source/Editor/BlueprintGraph/Private/K2Node_FormatText.cpp
--- a/Engine/Source/Editor/BlueprintGraph/Private/K2Node_FormatText.cpp
+++ b/Engine/Source/Editor/BlueprintGraph/Private/K2Node_FormatText.cpp
@@ -378,10 +378,22 @@ UEdGraphPin* UK2Node_FormatText::GetFormatPin() const
 
 void UK2Node_FormatText::GetMenuActions(FBlueprintActionDatabaseRegistrar& ActionRegistrar) const
 {
-	UBlueprintNodeSpawner* NodeSpawner = UBlueprintNodeSpawner::Create(GetClass());
-	check(NodeSpawner != nullptr);
+	// actions get registered under specific object-keys; the idea is that 
+	// actions might have to be updated (or deleted) if their object-key is  
+	// mutated (or removed)... here we use the node's class (so if the node 
+	// type disappears, then the action should go with it)
+	UClass* ActionKey = GetClass();
+	// to keep from needlessly instantiating a UBlueprintNodeSpawner, first   
+	// check to make sure that the registrar is looking for actions of this type
+	// (could be regenerating actions for a specific asset, and therefore the 
+	// registrar would only accept actions corresponding to that asset)
+	if (ActionRegistrar.IsOpenForRegistration(ActionKey))
+	{
+		UBlueprintNodeSpawner* NodeSpawner = UBlueprintNodeSpawner::Create(GetClass());
+		check(NodeSpawner != nullptr);
 
-	ActionRegistrar.AddBlueprintAction(NodeSpawner);
+		ActionRegistrar.AddBlueprintAction(ActionKey, NodeSpawner);
+	}
 }
 
 FText UK2Node_FormatText::GetMenuCategory() const
diff --git a/Engine/Source/Editor/BlueprintGraph/Private/K2Node_GetDataTableRow.cpp b/Engine/Source/Editor/BlueprintGraph/Private/K2Node_GetDataTableRow.cpp
--- a/Engine/Source/Editor/BlueprintGraph/Private/K2Node_GetDataTableRow.cpp
+++ b/Engine/Source/Editor/BlueprintGraph/Private/K2Node_GetDataTableRow.cpp
@@ -170,10 +170,22 @@ void UK2Node_GetDataTableRow::ReallocatePinsDuringReconstruction(TArray<UEdGraph
 
 void UK2Node_GetDataTableRow::GetMenuActions(FBlueprintActionDatabaseRegistrar& ActionRegistrar) const
 {
-	UBlueprintNodeSpawner* NodeSpawner = UBlueprintNodeSpawner::Create(GetClass());
-	check(NodeSpawner != nullptr);
+	// actions get registered under specific object-keys; the idea is that 
+	// actions might have to be updated (or deleted) if their object-key is  
+	// mutated (or removed)... here we use the node's class (so if the node 
+	// type disappears, then the action should go with it)
+	UClass* ActionKey = GetClass();
+	// to keep from needlessly instantiating a UBlueprintNodeSpawner, first   
+	// check to make sure that the registrar is looking for actions of this type
+	// (could be regenerating actions for a specific asset, and therefore the 
+	// registrar would only accept actions corresponding to that asset)
+	if (ActionRegistrar.IsOpenForRegistration(ActionKey))
+	{
+		UBlueprintNodeSpawner* NodeSpawner = UBlueprintNodeSpawner::Create(GetClass());
+		check(NodeSpawner != nullptr);
 
-	ActionRegistrar.AddBlueprintAction(NodeSpawner);
+		ActionRegistrar.AddBlueprintAction(ActionKey, NodeSpawner);
+	}
 }
 
 FText UK2Node_GetDataTableRow::GetMenuCategory() const
diff --git a/Engine/Source/Editor/BlueprintGraph/Private/K2Node_GetEnumeratorName.cpp b/Engine/Source/Editor/BlueprintGraph/Private/K2Node_GetEnumeratorName.cpp
--- a/Engine/Source/Editor/BlueprintGraph/Private/K2Node_GetEnumeratorName.cpp
+++ b/Engine/Source/Editor/BlueprintGraph/Private/K2Node_GetEnumeratorName.cpp
@@ -187,10 +187,22 @@ void UK2Node_GetEnumeratorName::ExpandNode(class FKismetCompilerContext& Compile
 
 void UK2Node_GetEnumeratorName::GetMenuActions(FBlueprintActionDatabaseRegistrar& ActionRegistrar) const
 {
-	UBlueprintNodeSpawner* NodeSpawner = UBlueprintNodeSpawner::Create(GetClass());
-	check(NodeSpawner != nullptr);
+	// actions get registered under specific object-keys; the idea is that 
+	// actions might have to be updated (or deleted) if their object-key is  
+	// mutated (or removed)... here we use the node's class (so if the node 
+	// type disappears, then the action should go with it)
+	UClass* ActionKey = GetClass();
+	// to keep from needlessly instantiating a UBlueprintNodeSpawner, first   
+	// check to make sure that the registrar is looking for actions of this type
+	// (could be regenerating actions for a specific asset, and therefore the 
+	// registrar would only accept actions corresponding to that asset)
+	if (ActionRegistrar.IsOpenForRegistration(ActionKey))
+	{
+		UBlueprintNodeSpawner* NodeSpawner = UBlueprintNodeSpawner::Create(GetClass());
+		check(NodeSpawner != nullptr);
 
-	ActionRegistrar.AddBlueprintAction(NodeSpawner);
+		ActionRegistrar.AddBlueprintAction(ActionKey, NodeSpawner);
+	}
 }
 
 FText UK2Node_GetEnumeratorName::GetMenuCategory() const
diff --git a/Engine/Source/Editor/BlueprintGraph/Private/K2Node_GetEnumeratorNameAsString.cpp b/Engine/Source/Editor/BlueprintGraph/Private/K2Node_GetEnumeratorNameAsString.cpp
--- a/Engine/Source/Editor/BlueprintGraph/Private/K2Node_GetEnumeratorNameAsString.cpp
+++ b/Engine/Source/Editor/BlueprintGraph/Private/K2Node_GetEnumeratorNameAsString.cpp
@@ -34,10 +34,22 @@ FText UK2Node_GetEnumeratorNameAsString::GetNodeTitle(ENodeTitleType::Type Title
 
 void UK2Node_GetEnumeratorNameAsString::GetMenuActions(FBlueprintActionDatabaseRegistrar& ActionRegistrar) const
 {
-	UBlueprintNodeSpawner* NodeSpawner = UBlueprintNodeSpawner::Create(GetClass());
-	check(NodeSpawner != nullptr);
+	// actions get registered under specific object-keys; the idea is that 
+	// actions might have to be updated (or deleted) if their object-key is  
+	// mutated (or removed)... here we use the node's class (so if the node 
+	// type disappears, then the action should go with it)
+	UClass* ActionKey = GetClass();
+	// to keep from needlessly instantiating a UBlueprintNodeSpawner, first   
+	// check to make sure that the registrar is looking for actions of this type
+	// (could be regenerating actions for a specific asset, and therefore the 
+	// registrar would only accept actions corresponding to that asset)
+	if (ActionRegistrar.IsOpenForRegistration(ActionKey))
+	{
+		UBlueprintNodeSpawner* NodeSpawner = UBlueprintNodeSpawner::Create(GetClass());
+		check(NodeSpawner != nullptr);
 
-	ActionRegistrar.AddBlueprintAction(NodeSpawner);
+		ActionRegistrar.AddBlueprintAction(ActionKey, NodeSpawner);
+	}
 }
 
 FName UK2Node_GetEnumeratorNameAsString::GetFunctionName() const
diff --git a/Engine/Source/Editor/BlueprintGraph/Private/K2Node_GetInputAxisKeyValue.cpp b/Engine/Source/Editor/BlueprintGraph/Private/K2Node_GetInputAxisKeyValue.cpp
--- a/Engine/Source/Editor/BlueprintGraph/Private/K2Node_GetInputAxisKeyValue.cpp
+++ b/Engine/Source/Editor/BlueprintGraph/Private/K2Node_GetInputAxisKeyValue.cpp
@@ -130,25 +130,39 @@ void UK2Node_GetInputAxisKeyValue::GetMenuActions(FBlueprintActionDatabaseRegist
 {
 	TArray<FKey> AllKeys;
 	EKeys::GetAllKeys(AllKeys);
-	
+
+	auto CustomizeInputNodeLambda = [](UEdGraphNode* NewNode, bool bIsTemplateNode, FKey Key)
+	{
+		UK2Node_GetInputAxisKeyValue* InputNode = CastChecked<UK2Node_GetInputAxisKeyValue>(NewNode);
+		InputNode->Initialize(Key);
+	};
+
 	for (FKey const Key : AllKeys)
 	{
 		if (!Key.IsBindableInBlueprints() || !Key.IsFloatAxis())
 		{
 			continue;
 		}
+
+		// actions get registered under specific object-keys; the idea is that 
+		// actions might have to be updated (or deleted) if their object-key is  
+		// mutated (or removed)... here we use the node's class (so if the node 
+		// type disappears, then the action should go with it)
+		UClass* ActionKey = GetClass();
+		// to keep from needlessly instantiating a UBlueprintNodeSpawner, first   
+		// check to make sure that the registrar is looking for actions of this type
+		// (could be regenerating actions for a specific asset, and therefore the 
+		// registrar would only accept actions corresponding to that asset)
+		if (!ActionRegistrar.IsOpenForRegistration(ActionKey))
+		{
+			continue;
+		}
 		
 		UBlueprintNodeSpawner* NodeSpawner = UBlueprintNodeSpawner::Create(GetClass());
 		check(NodeSpawner != nullptr);
 		
-		auto CustomizeInputNodeLambda = [](UEdGraphNode* NewNode, bool bIsTemplateNode, FKey Key)
-		{
-			UK2Node_GetInputAxisKeyValue* InputNode = CastChecked<UK2Node_GetInputAxisKeyValue>(NewNode);
-			InputNode->Initialize(Key);
-		};
-		
 		NodeSpawner->CustomizeNodeDelegate = UBlueprintNodeSpawner::FCustomizeNodeDelegate::CreateStatic(CustomizeInputNodeLambda, Key);
-		ActionRegistrar.AddBlueprintAction(NodeSpawner);
+		ActionRegistrar.AddBlueprintAction(ActionKey, NodeSpawner);
 	}
 }
 
diff --git a/Engine/Source/Editor/BlueprintGraph/Private/K2Node_GetInputAxisValue.cpp b/Engine/Source/Editor/BlueprintGraph/Private/K2Node_GetInputAxisValue.cpp
--- a/Engine/Source/Editor/BlueprintGraph/Private/K2Node_GetInputAxisValue.cpp
+++ b/Engine/Source/Editor/BlueprintGraph/Private/K2Node_GetInputAxisValue.cpp
@@ -105,19 +105,33 @@ void UK2Node_GetInputAxisValue::GetMenuActions(FBlueprintActionDatabaseRegistrar
 	TArray<FName> AxisNames;
 	GetDefault<UInputSettings>()->GetAxisNames(AxisNames);
 
+	auto CustomizeInputNodeLambda = [](UEdGraphNode* NewNode, bool bIsTemplateNode, FName AxisName)
+	{
+		UK2Node_GetInputAxisValue* InputNode = CastChecked<UK2Node_GetInputAxisValue>(NewNode);
+		InputNode->Initialize(AxisName);
+	};
+
 	for (FName const InputAxisName : AxisNames)
 	{
+		// actions get registered under specific object-keys; the idea is that 
+		// actions might have to be updated (or deleted) if their object-key is  
+		// mutated (or removed)... here we use the node's class (so if the node 
+		// type disappears, then the action should go with it)
+		UClass* ActionKey = GetClass();
+		// to keep from needlessly instantiating a UBlueprintNodeSpawner, first   
+		// check to make sure that the registrar is looking for actions of this type
+		// (could be regenerating actions for a specific asset, and therefore the 
+		// registrar would only accept actions corresponding to that asset)
+		if (!ActionRegistrar.IsOpenForRegistration(ActionKey))
+		{
+			continue;
+		}
+
 		UBlueprintNodeSpawner* NodeSpawner = UBlueprintNodeSpawner::Create(GetClass());
 		check(NodeSpawner != nullptr);
 
-		auto CustomizeInputNodeLambda = [](UEdGraphNode* NewNode, bool bIsTemplateNode, FName AxisName)
-		{
-			UK2Node_GetInputAxisValue* InputNode = CastChecked<UK2Node_GetInputAxisValue>(NewNode);
-			InputNode->Initialize(AxisName);
-		};
-
 		NodeSpawner->CustomizeNodeDelegate = UBlueprintNodeSpawner::FCustomizeNodeDelegate::CreateStatic(CustomizeInputNodeLambda, InputAxisName);
-		ActionRegistrar.AddBlueprintAction(NodeSpawner);
+		ActionRegistrar.AddBlueprintAction(ActionKey, NodeSpawner);
 	}
 }
 
diff --git a/Engine/Source/Editor/BlueprintGraph/Private/K2Node_GetInputVectorAxisValue.cpp b/Engine/Source/Editor/BlueprintGraph/Private/K2Node_GetInputVectorAxisValue.cpp
--- a/Engine/Source/Editor/BlueprintGraph/Private/K2Node_GetInputVectorAxisValue.cpp
+++ b/Engine/Source/Editor/BlueprintGraph/Private/K2Node_GetInputVectorAxisValue.cpp
@@ -67,6 +67,12 @@ void UK2Node_GetInputVectorAxisValue::GetMenuActions(FBlueprintActionDatabaseReg
 	TArray<FKey> AllKeys;
 	EKeys::GetAllKeys(AllKeys);
 
+	auto CustomizeInputNodeLambda = [](UEdGraphNode* NewNode, bool bIsTemplateNode, FKey Key)
+	{
+		UK2Node_GetInputVectorAxisValue* InputNode = CastChecked<UK2Node_GetInputVectorAxisValue>(NewNode);
+		InputNode->Initialize(Key);
+	};
+
 	for (FKey const Key : AllKeys)
 	{
 		if (!Key.IsBindableInBlueprints() || !Key.IsVectorAxis())
@@ -74,16 +80,24 @@ void UK2Node_GetInputVectorAxisValue::GetMenuActions(FBlueprintActionDatabaseReg
 			continue;
 		}
 
+		// actions get registered under specific object-keys; the idea is that 
+		// actions might have to be updated (or deleted) if their object-key is  
+		// mutated (or removed)... here we use the node's class (so if the node 
+		// type disappears, then the action should go with it)
+		UClass* ActionKey = GetClass();
+		// to keep from needlessly instantiating a UBlueprintNodeSpawner, first   
+		// check to make sure that the registrar is looking for actions of this type
+		// (could be regenerating actions for a specific asset, and therefore the 
+		// registrar would only accept actions corresponding to that asset)
+		if (!ActionRegistrar.IsOpenForRegistration(ActionKey))
+		{
+			continue;
+		}
+
 		UBlueprintNodeSpawner* NodeSpawner = UBlueprintNodeSpawner::Create(GetClass());
 		check(NodeSpawner != nullptr);
 
-		auto CustomizeInputNodeLambda = [](UEdGraphNode* NewNode, bool bIsTemplateNode, FKey Key)
-		{
-			UK2Node_GetInputVectorAxisValue* InputNode = CastChecked<UK2Node_GetInputVectorAxisValue>(NewNode);
-			InputNode->Initialize(Key);
-		};
-
 		NodeSpawner->CustomizeNodeDelegate = UBlueprintNodeSpawner::FCustomizeNodeDelegate::CreateStatic(CustomizeInputNodeLambda, Key);
-		ActionRegistrar.AddBlueprintAction(NodeSpawner);
+		ActionRegistrar.AddBlueprintAction(ActionKey, NodeSpawner);
 	}
 }
diff --git a/Engine/Source/Editor/BlueprintGraph/Private/K2Node_GetNumEnumEntries.cpp b/Engine/Source/Editor/BlueprintGraph/Private/K2Node_GetNumEnumEntries.cpp
--- a/Engine/Source/Editor/BlueprintGraph/Private/K2Node_GetNumEnumEntries.cpp
+++ b/Engine/Source/Editor/BlueprintGraph/Private/K2Node_GetNumEnumEntries.cpp
@@ -98,6 +98,15 @@ void UK2Node_GetNumEnumEntries::ExpandNode(class FKismetCompilerContext& Compile
 
 void UK2Node_GetNumEnumEntries::GetMenuActions(FBlueprintActionDatabaseRegistrar& ActionRegistrar) const
 {
+	auto CustomizeEnumNodeLambda = [](UEdGraphNode* NewNode, bool bIsTemplateNode, TWeakObjectPtr<UEnum> EnumPtr)
+	{
+		UK2Node_GetNumEnumEntries* EnumNode = CastChecked<UK2Node_GetNumEnumEntries>(NewNode);
+		if (EnumPtr.IsValid())
+		{
+			EnumNode->Enum = EnumPtr.Get();
+		}
+	};
+
 	for (TObjectIterator<UEnum> EnumIt; EnumIt; ++EnumIt)
 	{
 		UEnum const* Enum = (*EnumIt);
@@ -106,14 +115,14 @@ void UK2Node_GetNumEnumEntries::GetMenuActions(FBlueprintActionDatabaseRegistrar
 			continue;
 		}
 
-		auto CustomizeEnumNodeLambda = [](UEdGraphNode* NewNode, bool bIsTemplateNode, TWeakObjectPtr<UEnum> EnumPtr)
+		// to keep from needlessly instantiating a UBlueprintNodeSpawner, first   
+		// check to make sure that the registrar is looking for actions of this type
+		// (could be regenerating actions for a specific asset, and therefore the 
+		// registrar would only accept actions corresponding to that asset)
+		if (!ActionRegistrar.IsOpenForRegistration(Enum))
 		{
-			UK2Node_GetNumEnumEntries* EnumNode = CastChecked<UK2Node_GetNumEnumEntries>(NewNode);
-			if (EnumPtr.IsValid())
-			{
-				EnumNode->Enum = EnumPtr.Get();
-			}
-		};
+			continue;
+		}
 
 		UBlueprintNodeSpawner* NodeSpawner = UBlueprintNodeSpawner::Create(GetClass());
 		check(NodeSpawner != nullptr);
diff --git a/Engine/Source/Editor/BlueprintGraph/Private/K2Node_IfThenElse.cpp b/Engine/Source/Editor/BlueprintGraph/Private/K2Node_IfThenElse.cpp
--- a/Engine/Source/Editor/BlueprintGraph/Private/K2Node_IfThenElse.cpp
+++ b/Engine/Source/Editor/BlueprintGraph/Private/K2Node_IfThenElse.cpp
@@ -148,10 +148,22 @@ FNodeHandlingFunctor* UK2Node_IfThenElse::CreateNodeHandler(FKismetCompilerConte
 
 void UK2Node_IfThenElse::GetMenuActions(FBlueprintActionDatabaseRegistrar& ActionRegistrar) const
 {
-	UBlueprintNodeSpawner* NodeSpawner = UBlueprintNodeSpawner::Create(GetClass());
-	check(NodeSpawner != nullptr);
+	// actions get registered under specific object-keys; the idea is that 
+	// actions might have to be updated (or deleted) if their object-key is  
+	// mutated (or removed)... here we use the node's class (so if the node 
+	// type disappears, then the action should go with it)
+	UClass* ActionKey = GetClass();
+	// to keep from needlessly instantiating a UBlueprintNodeSpawner, first   
+	// check to make sure that the registrar is looking for actions of this type
+	// (could be regenerating actions for a specific asset, and therefore the 
+	// registrar would only accept actions corresponding to that asset)
+	if (ActionRegistrar.IsOpenForRegistration(ActionKey))
+	{
+		UBlueprintNodeSpawner* NodeSpawner = UBlueprintNodeSpawner::Create(GetClass());
+		check(NodeSpawner != nullptr);
 
-	ActionRegistrar.AddBlueprintAction(NodeSpawner);
+		ActionRegistrar.AddBlueprintAction(ActionKey, NodeSpawner);
+	}
 }
 
 FText UK2Node_IfThenElse::GetMenuCategory() const
diff --git a/Engine/Source/Editor/BlueprintGraph/Private/K2Node_InputAction.cpp b/Engine/Source/Editor/BlueprintGraph/Private/K2Node_InputAction.cpp
--- a/Engine/Source/Editor/BlueprintGraph/Private/K2Node_InputAction.cpp
+++ b/Engine/Source/Editor/BlueprintGraph/Private/K2Node_InputAction.cpp
@@ -149,19 +149,29 @@ void UK2Node_InputAction::GetMenuActions(FBlueprintActionDatabaseRegistrar& Acti
 	TArray<FName> ActionNames;
 	GetDefault<UInputSettings>()->GetActionNames(ActionNames);
 
+	auto CustomizeInputNodeLambda = [](UEdGraphNode* NewNode, bool bIsTemplateNode, FName ActionName)
+	{
+		UK2Node_InputAction* InputNode = CastChecked<UK2Node_InputAction>(NewNode);
+		InputNode->InputActionName = ActionName;
+	};
+
 	for (FName const InputActionName : ActionNames)
 	{
+		UClass* ActionKey = GetClass();
+		// to keep from needlessly instantiating a UBlueprintNodeSpawner, first   
+		// check to make sure that the registrar is looking for actions of this type
+		// (could be regenerating actions for a specific asset, and therefore the 
+		// registrar would only accept actions corresponding to that asset)
+		if (!ActionRegistrar.IsOpenForRegistration(ActionKey))
+		{
+			continue;
+		}
+
 		UBlueprintNodeSpawner* NodeSpawner = UBlueprintNodeSpawner::Create(GetClass());
 		check(NodeSpawner != nullptr);
 
-		auto CustomizeInputNodeLambda = [](UEdGraphNode* NewNode, bool bIsTemplateNode, FName ActionName)
-		{
-			UK2Node_InputAction* InputNode = CastChecked<UK2Node_InputAction>(NewNode);
-			InputNode->InputActionName = ActionName;
-		};
-
 		NodeSpawner->CustomizeNodeDelegate = UBlueprintNodeSpawner::FCustomizeNodeDelegate::CreateStatic(CustomizeInputNodeLambda, InputActionName);
-		ActionRegistrar.AddBlueprintAction(NodeSpawner);
+		ActionRegistrar.AddBlueprintAction(ActionKey, NodeSpawner);
 	}
 }
 
diff --git a/Engine/Source/Editor/BlueprintGraph/Private/K2Node_InputAxisEvent.cpp b/Engine/Source/Editor/BlueprintGraph/Private/K2Node_InputAxisEvent.cpp
--- a/Engine/Source/Editor/BlueprintGraph/Private/K2Node_InputAxisEvent.cpp
+++ b/Engine/Source/Editor/BlueprintGraph/Private/K2Node_InputAxisEvent.cpp
@@ -117,19 +117,33 @@ void UK2Node_InputAxisEvent::GetMenuActions(FBlueprintActionDatabaseRegistrar& A
 	TArray<FName> AxisNames;
 	GetDefault<UInputSettings>()->GetAxisNames(AxisNames);
 
+	auto CustomizeInputNodeLambda = [](UEdGraphNode* NewNode, bool bIsTemplateNode, FName AxisName)
+	{
+		UK2Node_InputAxisEvent* InputNode = CastChecked<UK2Node_InputAxisEvent>(NewNode);
+		InputNode->Initialize(AxisName);
+	};
+
 	for (FName const InputAxisName : AxisNames)
 	{
+		// actions get registered under specific object-keys; the idea is that 
+		// actions might have to be updated (or deleted) if their object-key is  
+		// mutated (or removed)... here we use the node's class (so if the node 
+		// type disappears, then the action should go with it)
+		UClass* ActionKey = GetClass();
+		// to keep from needlessly instantiating a UBlueprintNodeSpawner, first   
+		// check to make sure that the registrar is looking for actions of this type
+		// (could be regenerating actions for a specific asset, and therefore the 
+		// registrar would only accept actions corresponding to that asset)
+		if (!ActionRegistrar.IsOpenForRegistration(ActionKey))
+		{
+			continue;
+		}
+
 		UBlueprintNodeSpawner* NodeSpawner = UBlueprintNodeSpawner::Create(GetClass());
 		check(NodeSpawner != nullptr);
 
-		auto CustomizeInputNodeLambda = [](UEdGraphNode* NewNode, bool bIsTemplateNode, FName AxisName)
-		{
-			UK2Node_InputAxisEvent* InputNode = CastChecked<UK2Node_InputAxisEvent>(NewNode);
-			InputNode->Initialize(AxisName);
-		};
-
 		NodeSpawner->CustomizeNodeDelegate = UBlueprintNodeSpawner::FCustomizeNodeDelegate::CreateStatic(CustomizeInputNodeLambda, InputAxisName);
-		ActionRegistrar.AddBlueprintAction(NodeSpawner);
+		ActionRegistrar.AddBlueprintAction(ActionKey, NodeSpawner);
 	}
 }
 
diff --git a/Engine/Source/Editor/BlueprintGraph/Private/K2Node_InputAxisKeyEvent.cpp b/Engine/Source/Editor/BlueprintGraph/Private/K2Node_InputAxisKeyEvent.cpp
--- a/Engine/Source/Editor/BlueprintGraph/Private/K2Node_InputAxisKeyEvent.cpp
+++ b/Engine/Source/Editor/BlueprintGraph/Private/K2Node_InputAxisKeyEvent.cpp
@@ -121,6 +121,12 @@ void UK2Node_InputAxisKeyEvent::GetMenuActions(FBlueprintActionDatabaseRegistrar
 	TArray<FKey> AllKeys;
 	EKeys::GetAllKeys(AllKeys);
 
+	auto CustomizeInputNodeLambda = [](UEdGraphNode* NewNode, bool bIsTemplateNode, FKey Key)
+	{
+		UK2Node_InputAxisKeyEvent* InputNode = CastChecked<UK2Node_InputAxisKeyEvent>(NewNode);
+		InputNode->Initialize(Key);
+	};
+
 	for (FKey const Key : AllKeys)
 	{
 		if (!Key.IsBindableInBlueprints() || !Key.IsFloatAxis())
@@ -128,17 +134,25 @@ void UK2Node_InputAxisKeyEvent::GetMenuActions(FBlueprintActionDatabaseRegistrar
 			continue;
 		}
 
+		// actions get registered under specific object-keys; the idea is that 
+		// actions might have to be updated (or deleted) if their object-key is  
+		// mutated (or removed)... here we use the node's class (so if the node 
+		// type disappears, then the action should go with it)
+		UClass* ActionKey = GetClass();
+		// to keep from needlessly instantiating a UBlueprintNodeSpawner, first   
+		// check to make sure that the registrar is looking for actions of this type
+		// (could be regenerating actions for a specific asset, and therefore the 
+		// registrar would only accept actions corresponding to that asset)
+		if (!ActionRegistrar.IsOpenForRegistration(ActionKey))
+		{
+			continue;
+		}
+
 		UBlueprintNodeSpawner* NodeSpawner = UBlueprintNodeSpawner::Create(GetClass());
 		check(NodeSpawner != nullptr);
 
-		auto CustomizeInputNodeLambda = [](UEdGraphNode* NewNode, bool bIsTemplateNode, FKey Key)
-		{
-			UK2Node_InputAxisKeyEvent* InputNode = CastChecked<UK2Node_InputAxisKeyEvent>(NewNode);
-			InputNode->Initialize(Key);
-		};
-
 		NodeSpawner->CustomizeNodeDelegate = UBlueprintNodeSpawner::FCustomizeNodeDelegate::CreateStatic(CustomizeInputNodeLambda, Key);
-		ActionRegistrar.AddBlueprintAction(NodeSpawner);
+		ActionRegistrar.AddBlueprintAction(ActionKey, NodeSpawner);
 	}
 }
 
diff --git a/Engine/Source/Editor/BlueprintGraph/Private/K2Node_InputKey.cpp b/Engine/Source/Editor/BlueprintGraph/Private/K2Node_InputKey.cpp
--- a/Engine/Source/Editor/BlueprintGraph/Private/K2Node_InputKey.cpp
+++ b/Engine/Source/Editor/BlueprintGraph/Private/K2Node_InputKey.cpp
@@ -376,6 +376,12 @@ void UK2Node_InputKey::GetMenuActions(FBlueprintActionDatabaseRegistrar& ActionR
 	TArray<FKey> AllKeys;
 	EKeys::GetAllKeys(AllKeys);
 
+	auto CustomizeInputNodeLambda = [](UEdGraphNode* NewNode, bool bIsTemplateNode, FKey Key)
+	{
+		UK2Node_InputKey* InputNode = CastChecked<UK2Node_InputKey>(NewNode);
+		InputNode->InputKey = Key;
+	};
+
 	for (FKey const Key : AllKeys)
 	{
 		// these will be handled by UK2Node_GetInputAxisKeyValue and UK2Node_GetInputVectorAxisValue respectively
@@ -384,17 +390,25 @@ void UK2Node_InputKey::GetMenuActions(FBlueprintActionDatabaseRegistrar& ActionR
 			continue;
 		}
 
+		// actions get registered under specific object-keys; the idea is that 
+		// actions might have to be updated (or deleted) if their object-key is  
+		// mutated (or removed)... here we use the node's class (so if the node 
+		// type disappears, then the action should go with it)
+		UClass* ActionKey = GetClass();
+		// to keep from needlessly instantiating a UBlueprintNodeSpawner, first   
+		// check to make sure that the registrar is looking for actions of this type
+		// (could be regenerating actions for a specific asset, and therefore the 
+		// registrar would only accept actions corresponding to that asset)
+		if (!ActionRegistrar.IsOpenForRegistration(ActionKey))
+		{
+			continue;
+		}
+
 		UBlueprintNodeSpawner* NodeSpawner = UBlueprintNodeSpawner::Create(GetClass());
 		check(NodeSpawner != nullptr);
 
-		auto CustomizeInputNodeLambda = [](UEdGraphNode* NewNode, bool bIsTemplateNode, FKey Key)
-		{
-			UK2Node_InputKey* InputNode = CastChecked<UK2Node_InputKey>(NewNode);
-			InputNode->InputKey = Key;
-		};
-
 		NodeSpawner->CustomizeNodeDelegate = UBlueprintNodeSpawner::FCustomizeNodeDelegate::CreateStatic(CustomizeInputNodeLambda, Key);
-		ActionRegistrar.AddBlueprintAction(NodeSpawner);
+		ActionRegistrar.AddBlueprintAction(ActionKey, NodeSpawner);
 	}
 }
 
diff --git a/Engine/Source/Editor/BlueprintGraph/Private/K2Node_InputTouch.cpp b/Engine/Source/Editor/BlueprintGraph/Private/K2Node_InputTouch.cpp
--- a/Engine/Source/Editor/BlueprintGraph/Private/K2Node_InputTouch.cpp
+++ b/Engine/Source/Editor/BlueprintGraph/Private/K2Node_InputTouch.cpp
@@ -70,10 +70,22 @@ FText UK2Node_InputTouch::GetNodeTitle(ENodeTitleType::Type TitleType) const
 
 void UK2Node_InputTouch::GetMenuActions(FBlueprintActionDatabaseRegistrar& ActionRegistrar) const
 {
-	UBlueprintNodeSpawner* NodeSpawner = UBlueprintNodeSpawner::Create(GetClass());
-	check(NodeSpawner != nullptr);
+	// actions get registered under specific object-keys; the idea is that 
+	// actions might have to be updated (or deleted) if their object-key is  
+	// mutated (or removed)... here we use the node's class (so if the node 
+	// type disappears, then the action should go with it)
+	UClass* ActionKey = GetClass();
+	// to keep from needlessly instantiating a UBlueprintNodeSpawner, first   
+	// check to make sure that the registrar is looking for actions of this type
+	// (could be regenerating actions for a specific asset, and therefore the 
+	// registrar would only accept actions corresponding to that asset)
+	if (ActionRegistrar.IsOpenForRegistration(ActionKey))
+	{
+		UBlueprintNodeSpawner* NodeSpawner = UBlueprintNodeSpawner::Create(GetClass());
+		check(NodeSpawner != nullptr);
 
-	ActionRegistrar.AddBlueprintAction(NodeSpawner);
+		ActionRegistrar.AddBlueprintAction(ActionKey, NodeSpawner);
+	}
 }
 
 FText UK2Node_InputTouch::GetMenuCategory() const
diff --git a/Engine/Source/Editor/BlueprintGraph/Private/K2Node_InputVectorAxisEvent.cpp b/Engine/Source/Editor/BlueprintGraph/Private/K2Node_InputVectorAxisEvent.cpp
--- a/Engine/Source/Editor/BlueprintGraph/Private/K2Node_InputVectorAxisEvent.cpp
+++ b/Engine/Source/Editor/BlueprintGraph/Private/K2Node_InputVectorAxisEvent.cpp
@@ -55,6 +55,12 @@ void UK2Node_InputVectorAxisEvent::GetMenuActions(FBlueprintActionDatabaseRegist
 {
 	TArray<FKey> AllKeys;
 	EKeys::GetAllKeys(AllKeys);
+
+	auto CustomizeInputNodeLambda = [](UEdGraphNode* NewNode, bool bIsTemplateNode, FKey Key)
+	{
+		UK2Node_InputAxisKeyEvent* InputNode = CastChecked<UK2Node_InputAxisKeyEvent>(NewNode);
+		InputNode->Initialize(Key);
+	};
 	
 	for (FKey const Key : AllKeys)
 	{
@@ -62,17 +68,25 @@ void UK2Node_InputVectorAxisEvent::GetMenuActions(FBlueprintActionDatabaseRegist
 		{
 			continue;
 		}
+
+		// actions get registered under specific object-keys; the idea is that 
+		// actions might have to be updated (or deleted) if their object-key is  
+		// mutated (or removed)... here we use the node's class (so if the node 
+		// type disappears, then the action should go with it)
+		UClass* ActionKey = GetClass();
+		// to keep from needlessly instantiating a UBlueprintNodeSpawner, first   
+		// check to make sure that the registrar is looking for actions of this type
+		// (could be regenerating actions for a specific asset, and therefore the 
+		// registrar would only accept actions corresponding to that asset)
+		if (!ActionRegistrar.IsOpenForRegistration(ActionKey))
+		{
+			continue;
+		}
 		
 		UBlueprintNodeSpawner* NodeSpawner = UBlueprintNodeSpawner::Create(GetClass());
 		check(NodeSpawner != nullptr);
 		
-		auto CustomizeInputNodeLambda = [](UEdGraphNode* NewNode, bool bIsTemplateNode, FKey Key)
-		{
-			UK2Node_InputAxisKeyEvent* InputNode = CastChecked<UK2Node_InputAxisKeyEvent>(NewNode);
-			InputNode->Initialize(Key);
-		};
-		
 		NodeSpawner->CustomizeNodeDelegate = UBlueprintNodeSpawner::FCustomizeNodeDelegate::CreateStatic(CustomizeInputNodeLambda, Key);
-		ActionRegistrar.AddBlueprintAction(NodeSpawner);
+		ActionRegistrar.AddBlueprintAction(ActionKey, NodeSpawner);
 	}
 }
diff --git a/Engine/Source/Editor/BlueprintGraph/Private/K2Node_Knot.cpp b/Engine/Source/Editor/BlueprintGraph/Private/K2Node_Knot.cpp
--- a/Engine/Source/Editor/BlueprintGraph/Private/K2Node_Knot.cpp
+++ b/Engine/Source/Editor/BlueprintGraph/Private/K2Node_Knot.cpp
@@ -148,10 +148,22 @@ void UK2Node_Knot::PostReconstructNode()
 
 void UK2Node_Knot::GetMenuActions(FBlueprintActionDatabaseRegistrar& ActionRegistrar) const
 {
-	UBlueprintNodeSpawner* NodeSpawner = UBlueprintNodeSpawner::Create(GetClass());
-	check(NodeSpawner != nullptr);
+	// actions get registered under specific object-keys; the idea is that 
+	// actions might have to be updated (or deleted) if their object-key is  
+	// mutated (or removed)... here we use the node's class (so if the node 
+	// type disappears, then the action should go with it)
+	UClass* ActionKey = GetClass();
+	// to keep from needlessly instantiating a UBlueprintNodeSpawner, first   
+	// check to make sure that the registrar is looking for actions of this type
+	// (could be regenerating actions for a specific asset, and therefore the 
+	// registrar would only accept actions corresponding to that asset)
+	if (ActionRegistrar.IsOpenForRegistration(ActionKey))
+	{
+		UBlueprintNodeSpawner* NodeSpawner = UBlueprintNodeSpawner::Create(GetClass());
+		check(NodeSpawner != nullptr);
 
-	ActionRegistrar.AddBlueprintAction(NodeSpawner);
+		ActionRegistrar.AddBlueprintAction(ActionKey, NodeSpawner);
+	}
 }
 
 bool UK2Node_Knot::ShouldOverridePinNames() const
diff --git a/Engine/Source/Editor/BlueprintGraph/Private/K2Node_MakeArray.cpp b/Engine/Source/Editor/BlueprintGraph/Private/K2Node_MakeArray.cpp
--- a/Engine/Source/Editor/BlueprintGraph/Private/K2Node_MakeArray.cpp
+++ b/Engine/Source/Editor/BlueprintGraph/Private/K2Node_MakeArray.cpp
@@ -407,10 +407,26 @@ bool UK2Node_MakeArray::IsConnectionDisallowed(const UEdGraphPin* MyPin, const U
 
 void UK2Node_MakeArray::GetMenuActions(FBlueprintActionDatabaseRegistrar& ActionRegistrar) const
 {
-	UBlueprintNodeSpawner* NodeSpawner = UBlueprintNodeSpawner::Create(GetClass());
-	check(NodeSpawner != nullptr);
+	// actions get registered under specific object-keys; the idea is that 
+	// actions might have to be updated (or deleted) if their object-key is  
+	// mutated (or removed)... here we use the node's class (so if the node 
+	// type disappears, then the action should go with it)
+	// actions get registered under specific object-keys; the idea is that 
+	// actions might have to be updated (or deleted) if their object-key is  
+	// mutated (or removed)... here we use the node's class (so if the node 
+	// type disappears, then the action should go with it)
+	UClass* ActionKey = GetClass();
+	// to keep from needlessly instantiating a UBlueprintNodeSpawner, first   
+	// check to make sure that the registrar is looking for actions of this type
+	// (could be regenerating actions for a specific asset, and therefore the 
+	// registrar would only accept actions corresponding to that asset)
+	if (ActionRegistrar.IsOpenForRegistration(ActionKey))
+	{
+		UBlueprintNodeSpawner* NodeSpawner = UBlueprintNodeSpawner::Create(GetClass());
+		check(NodeSpawner != nullptr);
 
-	ActionRegistrar.AddBlueprintAction(NodeSpawner);
+		ActionRegistrar.AddBlueprintAction(ActionKey, NodeSpawner);
+	}
 }
 
 FText UK2Node_MakeArray::GetMenuCategory() const
diff --git a/Engine/Source/Editor/BlueprintGraph/Private/K2Node_MakeStruct.cpp b/Engine/Source/Editor/BlueprintGraph/Private/K2Node_MakeStruct.cpp
--- a/Engine/Source/Editor/BlueprintGraph/Private/K2Node_MakeStruct.cpp
+++ b/Engine/Source/Editor/BlueprintGraph/Private/K2Node_MakeStruct.cpp
@@ -268,6 +268,15 @@ UK2Node::ERedirectType UK2Node_MakeStruct::DoPinsMatchForReconstruction(const UE
 
 void UK2Node_MakeStruct::GetMenuActions(FBlueprintActionDatabaseRegistrar& ActionRegistrar) const
 {
+	auto CustomizeMakeNodeLambda = [](UEdGraphNode* NewNode, bool bIsTemplateNode, TWeakObjectPtr<UScriptStruct> StructPtr)
+	{
+		UK2Node_MakeStruct* MakeNode = CastChecked<UK2Node_MakeStruct>(NewNode);
+		if (StructPtr.IsValid())
+		{
+			MakeNode->StructType = StructPtr.Get();
+		}
+	};
+
 	for (TObjectIterator<UScriptStruct> StructIt; StructIt; ++StructIt)
 	{
 		UScriptStruct const* Struct = (*StructIt);
@@ -276,14 +285,14 @@ void UK2Node_MakeStruct::GetMenuActions(FBlueprintActionDatabaseRegistrar& Actio
 			continue;
 		}
 
-		auto CustomizeMakeNodeLambda = [](UEdGraphNode* NewNode, bool bIsTemplateNode, TWeakObjectPtr<UScriptStruct> StructPtr)
+		// to keep from needlessly instantiating a UBlueprintNodeSpawner, first   
+		// check to make sure that the registrar is looking for actions of this type
+		// (could be regenerating actions for a specific asset, and therefore the 
+		// registrar would only accept actions corresponding to that asset)
+		if (!ActionRegistrar.IsOpenForRegistration(Struct))
 		{
-			UK2Node_MakeStruct* MakeNode = CastChecked<UK2Node_MakeStruct>(NewNode);
-			if (StructPtr.IsValid())
-			{
-				MakeNode->StructType = StructPtr.Get();
-			}
-		};
+			continue;
+		}
 
 		UBlueprintNodeSpawner* NodeSpawner = UBlueprintNodeSpawner::Create(GetClass());
 		check(NodeSpawner != nullptr);
diff --git a/Engine/Source/Editor/BlueprintGraph/Private/K2Node_MathExpression.cpp b/Engine/Source/Editor/BlueprintGraph/Private/K2Node_MathExpression.cpp
--- a/Engine/Source/Editor/BlueprintGraph/Private/K2Node_MathExpression.cpp
+++ b/Engine/Source/Editor/BlueprintGraph/Private/K2Node_MathExpression.cpp
@@ -2304,10 +2304,22 @@ void UK2Node_MathExpression::GetMenuEntries(FGraphContextMenuBuilder& ContextMen
 //------------------------------------------------------------------------------
 void UK2Node_MathExpression::GetMenuActions(FBlueprintActionDatabaseRegistrar& ActionRegistrar) const
 {
-	UBlueprintNodeSpawner* NodeSpawner = UBlueprintNodeSpawner::Create(GetClass());
-	check(NodeSpawner != nullptr);
-	
-	ActionRegistrar.AddBlueprintAction(NodeSpawner);
+	// actions get registered under specific object-keys; the idea is that 
+	// actions might have to be updated (or deleted) if their object-key is  
+	// mutated (or removed)... here we use the node's class (so if the node 
+	// type disappears, then the action should go with it)
+	UClass* ActionKey = GetClass();
+	// to keep from needlessly instantiating a UBlueprintNodeSpawner, first   
+	// check to make sure that the registrar is looking for actions of this type
+	// (could be regenerating actions for a specific asset, and therefore the 
+	// registrar would only accept actions corresponding to that asset)
+	if (ActionRegistrar.IsOpenForRegistration(ActionKey))
+	{
+		UBlueprintNodeSpawner* NodeSpawner = UBlueprintNodeSpawner::Create(GetClass());
+		check(NodeSpawner != nullptr);
+
+		ActionRegistrar.AddBlueprintAction(ActionKey, NodeSpawner);
+	}
 }
 
 //------------------------------------------------------------------------------
diff --git a/Engine/Source/Editor/BlueprintGraph/Private/K2Node_MultiGate.cpp b/Engine/Source/Editor/BlueprintGraph/Private/K2Node_MultiGate.cpp
--- a/Engine/Source/Editor/BlueprintGraph/Private/K2Node_MultiGate.cpp
+++ b/Engine/Source/Editor/BlueprintGraph/Private/K2Node_MultiGate.cpp
@@ -659,10 +659,22 @@ void UK2Node_MultiGate::ExpandNode(class FKismetCompilerContext& CompilerContext
 
 void UK2Node_MultiGate::GetMenuActions(FBlueprintActionDatabaseRegistrar& ActionRegistrar) const
 {
-	UBlueprintNodeSpawner* NodeSpawner = UBlueprintNodeSpawner::Create(GetClass());
-	check(NodeSpawner != nullptr);
+	// actions get registered under specific object-keys; the idea is that 
+	// actions might have to be updated (or deleted) if their object-key is  
+	// mutated (or removed)... here we use the node's class (so if the node 
+	// type disappears, then the action should go with it)
+	UClass* ActionKey = GetClass();
+	// to keep from needlessly instantiating a UBlueprintNodeSpawner, first   
+	// check to make sure that the registrar is looking for actions of this type
+	// (could be regenerating actions for a specific asset, and therefore the 
+	// registrar would only accept actions corresponding to that asset)
+	if (ActionRegistrar.IsOpenForRegistration(ActionKey))
+	{
+		UBlueprintNodeSpawner* NodeSpawner = UBlueprintNodeSpawner::Create(GetClass());
+		check(NodeSpawner != nullptr);
 
-	ActionRegistrar.AddBlueprintAction(NodeSpawner);
+		ActionRegistrar.AddBlueprintAction(ActionKey, NodeSpawner);
+	}
 }
 
 #undef LOCTEXT_NAMESPACE
diff --git a/Engine/Source/Editor/BlueprintGraph/Private/K2Node_Select.cpp b/Engine/Source/Editor/BlueprintGraph/Private/K2Node_Select.cpp
--- a/Engine/Source/Editor/BlueprintGraph/Private/K2Node_Select.cpp
+++ b/Engine/Source/Editor/BlueprintGraph/Private/K2Node_Select.cpp
@@ -794,10 +794,22 @@ FNodeHandlingFunctor* UK2Node_Select::CreateNodeHandler(FKismetCompilerContext&
 
 void UK2Node_Select::GetMenuActions(FBlueprintActionDatabaseRegistrar& ActionRegistrar) const
 {
-	UBlueprintNodeSpawner* NodeSpawner = UBlueprintNodeSpawner::Create(GetClass());
-	check(NodeSpawner != nullptr);
+	// actions get registered under specific object-keys; the idea is that 
+	// actions might have to be updated (or deleted) if their object-key is  
+	// mutated (or removed)... here we use the node's class (so if the node 
+	// type disappears, then the action should go with it)
+	UClass* ActionKey = GetClass();
+	// to keep from needlessly instantiating a UBlueprintNodeSpawner, first   
+	// check to make sure that the registrar is looking for actions of this type
+	// (could be regenerating actions for a specific asset, and therefore the 
+	// registrar would only accept actions corresponding to that asset)
+	if (ActionRegistrar.IsOpenForRegistration(ActionKey))
+	{
+		UBlueprintNodeSpawner* NodeSpawner = UBlueprintNodeSpawner::Create(GetClass());
+		check(NodeSpawner != nullptr);
 
-	ActionRegistrar.AddBlueprintAction(NodeSpawner);
+		ActionRegistrar.AddBlueprintAction(ActionKey, NodeSpawner);
+	}
 }
 
 FText UK2Node_Select::GetMenuCategory() const
diff --git a/Engine/Source/Editor/BlueprintGraph/Private/K2Node_SelectEnum.cpp b/Engine/Source/Editor/BlueprintGraph/Private/K2Node_SelectEnum.cpp
--- a/Engine/Source/Editor/BlueprintGraph/Private/K2Node_SelectEnum.cpp
+++ b/Engine/Source/Editor/BlueprintGraph/Private/K2Node_SelectEnum.cpp
@@ -65,21 +65,36 @@ void UDEPRECATED_K2Node_SelectEnum::NotifyPinConnectionListChanged(UEdGraphPin*
 
 void UDEPRECATED_K2Node_SelectEnum::GetMenuActions(FBlueprintActionDatabaseRegistrar& ActionRegistrar) const
 {
+	auto CustomizeEnumNodeLambda = [](UEdGraphNode* NewNode, bool bIsTemplateNode, TWeakObjectPtr<UEnum> EnumPtr)
+	{
+		UDEPRECATED_K2Node_SelectEnum* EnumNode = CastChecked<UDEPRECATED_K2Node_SelectEnum>(NewNode);
+		if (EnumPtr.IsValid())
+		{
+			EnumNode->Enum = EnumPtr.Get();
+		}
+	};
+
 	for (TObjectIterator<UEnum> EnumIt; EnumIt; ++EnumIt)
 	{
 		UEnum const* Enum = (*EnumIt);
-		auto CustomizeEnumNodeLambda = [](UEdGraphNode* NewNode, bool bIsTemplateNode, TWeakObjectPtr<UEnum> EnumPtr)
+		if (!UEdGraphSchema_K2::IsAllowableBlueprintVariableType(Enum))
 		{
-			UDEPRECATED_K2Node_SelectEnum* EnumNode = CastChecked<UDEPRECATED_K2Node_SelectEnum>(NewNode);
-			if (EnumPtr.IsValid())
-			{
-				EnumNode->Enum = EnumPtr.Get();
-			}
-		};
+			continue;
+		}
+
+		// to keep from needlessly instantiating a UBlueprintNodeSpawner, first   
+		// check to make sure that the registrar is looking for actions of this type
+		// (could be regenerating actions for a specific asset, and therefore the 
+		// registrar would only accept actions corresponding to that asset)
+		if (!ActionRegistrar.IsOpenForRegistration(Enum))
+		{
+			continue;
+		}
+
 		
 		UBlueprintNodeSpawner* NodeSpawner = UBlueprintNodeSpawner::Create(GetClass());
 		check(NodeSpawner != nullptr);
-		ActionRegistrar.AddBlueprintAction(NodeSpawner);
+		ActionRegistrar.AddBlueprintAction(Enum, NodeSpawner);
 		
 		TWeakObjectPtr<UEnum> EnumPtr = Enum;
 		NodeSpawner->CustomizeNodeDelegate = UBlueprintNodeSpawner::FCustomizeNodeDelegate::CreateStatic(CustomizeEnumNodeLambda, EnumPtr);
diff --git a/Engine/Source/Editor/BlueprintGraph/Private/K2Node_Self.cpp b/Engine/Source/Editor/BlueprintGraph/Private/K2Node_Self.cpp
--- a/Engine/Source/Editor/BlueprintGraph/Private/K2Node_Self.cpp
+++ b/Engine/Source/Editor/BlueprintGraph/Private/K2Node_Self.cpp
@@ -84,10 +84,22 @@ void UK2Node_Self::ValidateNodeDuringCompilation(class FCompilerResultsLog& Mess
 
 void UK2Node_Self::GetMenuActions(FBlueprintActionDatabaseRegistrar& ActionRegistrar) const
 {
-	UBlueprintNodeSpawner* NodeSpawner = UBlueprintNodeSpawner::Create(GetClass());
-	check(NodeSpawner != nullptr);
+	// actions get registered under specific object-keys; the idea is that 
+	// actions might have to be updated (or deleted) if their object-key is  
+	// mutated (or removed)... here we use the node's class (so if the node 
+	// type disappears, then the action should go with it)
+	UClass* ActionKey = GetClass();
+	// to keep from needlessly instantiating a UBlueprintNodeSpawner, first   
+	// check to make sure that the registrar is looking for actions of this type
+	// (could be regenerating actions for a specific asset, and therefore the 
+	// registrar would only accept actions corresponding to that asset)
+	if (ActionRegistrar.IsOpenForRegistration(ActionKey))
+	{
+		UBlueprintNodeSpawner* NodeSpawner = UBlueprintNodeSpawner::Create(GetClass());
+		check(NodeSpawner != nullptr);
 
-	ActionRegistrar.AddBlueprintAction(NodeSpawner);
+		ActionRegistrar.AddBlueprintAction(ActionKey, NodeSpawner);
+	}
 }
 
 FText UK2Node_Self::GetMenuCategory() const
diff --git a/Engine/Source/Editor/BlueprintGraph/Private/K2Node_SpawnActorFromClass.cpp b/Engine/Source/Editor/BlueprintGraph/Private/K2Node_SpawnActorFromClass.cpp
--- a/Engine/Source/Editor/BlueprintGraph/Private/K2Node_SpawnActorFromClass.cpp
+++ b/Engine/Source/Editor/BlueprintGraph/Private/K2Node_SpawnActorFromClass.cpp
@@ -312,10 +312,22 @@ void UK2Node_SpawnActorFromClass::GetNodeAttributes( TArray<TKeyValuePair<FStrin
 
 void UK2Node_SpawnActorFromClass::GetMenuActions(FBlueprintActionDatabaseRegistrar& ActionRegistrar) const
 {
-	UBlueprintNodeSpawner* NodeSpawner = UBlueprintNodeSpawner::Create(GetClass());
-	check(NodeSpawner != nullptr);
+	// actions get registered under specific object-keys; the idea is that 
+	// actions might have to be updated (or deleted) if their object-key is  
+	// mutated (or removed)... here we use the node's class (so if the node 
+	// type disappears, then the action should go with it)
+	UClass* ActionKey = GetClass();
+	// to keep from needlessly instantiating a UBlueprintNodeSpawner, first   
+	// check to make sure that the registrar is looking for actions of this type
+	// (could be regenerating actions for a specific asset, and therefore the 
+	// registrar would only accept actions corresponding to that asset)
+	if (ActionRegistrar.IsOpenForRegistration(ActionKey))
+	{
+		UBlueprintNodeSpawner* NodeSpawner = UBlueprintNodeSpawner::Create(GetClass());
+		check(NodeSpawner != nullptr);
 
-	ActionRegistrar.AddBlueprintAction(NodeSpawner);
+		ActionRegistrar.AddBlueprintAction(ActionKey, NodeSpawner);
+	}
 }
 
 FText UK2Node_SpawnActorFromClass::GetMenuCategory() const
diff --git a/Engine/Source/Editor/BlueprintGraph/Private/K2Node_SwitchEnum.cpp b/Engine/Source/Editor/BlueprintGraph/Private/K2Node_SwitchEnum.cpp
--- a/Engine/Source/Editor/BlueprintGraph/Private/K2Node_SwitchEnum.cpp
+++ b/Engine/Source/Editor/BlueprintGraph/Private/K2Node_SwitchEnum.cpp
@@ -72,6 +72,15 @@ FText UK2Node_SwitchEnum::GetTooltipText() const
 
 void UK2Node_SwitchEnum::GetMenuActions(FBlueprintActionDatabaseRegistrar& ActionRegistrar) const
 {
+	auto CustomizeEnumNodeLambda = [](UEdGraphNode* NewNode, bool bIsTemplateNode, TWeakObjectPtr<UEnum> EnumPtr)
+	{
+		UK2Node_SwitchEnum* EnumNode = CastChecked<UK2Node_SwitchEnum>(NewNode);
+		if (EnumPtr.IsValid())
+		{
+			EnumNode->SetEnum(EnumPtr.Get());
+		}
+	};
+
 	for (TObjectIterator<UEnum> EnumIt; EnumIt; ++EnumIt)
 	{
 		UEnum const* Enum = (*EnumIt);
@@ -80,14 +89,14 @@ void UK2Node_SwitchEnum::GetMenuActions(FBlueprintActionDatabaseRegistrar& Actio
 			continue;
 		}
 
-		auto CustomizeEnumNodeLambda = [](UEdGraphNode* NewNode, bool bIsTemplateNode, TWeakObjectPtr<UEnum> EnumPtr)
+		// to keep from needlessly instantiating a UBlueprintNodeSpawner, first   
+		// check to make sure that the registrar is looking for actions of this type
+		// (could be regenerating actions for a specific asset, and therefore the 
+		// registrar would only accept actions corresponding to that asset)
+		if (!ActionRegistrar.IsOpenForRegistration(Enum))
 		{
-			UK2Node_SwitchEnum* EnumNode = CastChecked<UK2Node_SwitchEnum>(NewNode);
-			if (EnumPtr.IsValid())
-			{
-				EnumNode->SetEnum(EnumPtr.Get());
-			}
-		};
+			continue;
+		}
 
 		UBlueprintNodeSpawner* NodeSpawner = UBlueprintNodeSpawner::Create(GetClass());
 		check(NodeSpawner != nullptr);
diff --git a/Engine/Source/Editor/BlueprintGraph/Private/K2Node_SwitchInteger.cpp b/Engine/Source/Editor/BlueprintGraph/Private/K2Node_SwitchInteger.cpp
--- a/Engine/Source/Editor/BlueprintGraph/Private/K2Node_SwitchInteger.cpp
+++ b/Engine/Source/Editor/BlueprintGraph/Private/K2Node_SwitchInteger.cpp
@@ -26,10 +26,22 @@ FText UK2Node_SwitchInteger::GetTooltipText() const
 
 void UK2Node_SwitchInteger::GetMenuActions(FBlueprintActionDatabaseRegistrar& ActionRegistrar) const
 {
-	UBlueprintNodeSpawner* NodeSpawner = UBlueprintNodeSpawner::Create(GetClass());
-	check(NodeSpawner != nullptr);
+	// actions get registered under specific object-keys; the idea is that 
+	// actions might have to be updated (or deleted) if their object-key is  
+	// mutated (or removed)... here we use the node's class (so if the node 
+	// type disappears, then the action should go with it)
+	UClass* ActionKey = GetClass();
+	// to keep from needlessly instantiating a UBlueprintNodeSpawner, first   
+	// check to make sure that the registrar is looking for actions of this type
+	// (could be regenerating actions for a specific asset, and therefore the 
+	// registrar would only accept actions corresponding to that asset)
+	if (ActionRegistrar.IsOpenForRegistration(ActionKey))
+	{
+		UBlueprintNodeSpawner* NodeSpawner = UBlueprintNodeSpawner::Create(GetClass());
+		check(NodeSpawner != nullptr);
 
-	ActionRegistrar.AddBlueprintAction(NodeSpawner);
+		ActionRegistrar.AddBlueprintAction(ActionKey, NodeSpawner);
+	}
 }
 
 void UK2Node_SwitchInteger::CreateSelectionPin()
diff --git a/Engine/Source/Editor/BlueprintGraph/Private/K2Node_SwitchName.cpp b/Engine/Source/Editor/BlueprintGraph/Private/K2Node_SwitchName.cpp
--- a/Engine/Source/Editor/BlueprintGraph/Private/K2Node_SwitchName.cpp
+++ b/Engine/Source/Editor/BlueprintGraph/Private/K2Node_SwitchName.cpp
@@ -41,10 +41,22 @@ FText UK2Node_SwitchName::GetTooltipText() const
 
 void UK2Node_SwitchName::GetMenuActions(FBlueprintActionDatabaseRegistrar& ActionRegistrar) const
 {
-	UBlueprintNodeSpawner* NodeSpawner = UBlueprintNodeSpawner::Create(GetClass());
-	check(NodeSpawner != nullptr);
+	// actions get registered under specific object-keys; the idea is that 
+	// actions might have to be updated (or deleted) if their object-key is  
+	// mutated (or removed)... here we use the node's class (so if the node 
+	// type disappears, then the action should go with it)
+	UClass* ActionKey = GetClass();
+	// to keep from needlessly instantiating a UBlueprintNodeSpawner, first   
+	// check to make sure that the registrar is looking for actions of this type
+	// (could be regenerating actions for a specific asset, and therefore the 
+	// registrar would only accept actions corresponding to that asset)
+	if (ActionRegistrar.IsOpenForRegistration(ActionKey))
+	{
+		UBlueprintNodeSpawner* NodeSpawner = UBlueprintNodeSpawner::Create(GetClass());
+		check(NodeSpawner != nullptr);
 
-	ActionRegistrar.AddBlueprintAction(NodeSpawner);
+		ActionRegistrar.AddBlueprintAction(ActionKey, NodeSpawner);
+	}
 }
 
 void UK2Node_SwitchName::CreateSelectionPin()
diff --git a/Engine/Source/Editor/BlueprintGraph/Private/K2Node_SwitchString.cpp b/Engine/Source/Editor/BlueprintGraph/Private/K2Node_SwitchString.cpp
--- a/Engine/Source/Editor/BlueprintGraph/Private/K2Node_SwitchString.cpp
+++ b/Engine/Source/Editor/BlueprintGraph/Private/K2Node_SwitchString.cpp
@@ -53,10 +53,22 @@ FText UK2Node_SwitchString::GetTooltipText() const
 
 void UK2Node_SwitchString::GetMenuActions(FBlueprintActionDatabaseRegistrar& ActionRegistrar) const
 {
-	UBlueprintNodeSpawner* NodeSpawner = UBlueprintNodeSpawner::Create(GetClass());
-	check(NodeSpawner != nullptr);
+	// actions get registered under specific object-keys; the idea is that 
+	// actions might have to be updated (or deleted) if their object-key is  
+	// mutated (or removed)... here we use the node's class (so if the node 
+	// type disappears, then the action should go with it)
+	UClass* ActionKey = GetClass();
+	// to keep from needlessly instantiating a UBlueprintNodeSpawner, first   
+	// check to make sure that the registrar is looking for actions of this type
+	// (could be regenerating actions for a specific asset, and therefore the 
+	// registrar would only accept actions corresponding to that asset)
+	if (ActionRegistrar.IsOpenForRegistration(ActionKey))
+	{
+		UBlueprintNodeSpawner* NodeSpawner = UBlueprintNodeSpawner::Create(GetClass());
+		check(NodeSpawner != nullptr);
 
-	ActionRegistrar.AddBlueprintAction(NodeSpawner);
+		ActionRegistrar.AddBlueprintAction(ActionKey, NodeSpawner);
+	}
 }
 
 void UK2Node_SwitchString::CreateSelectionPin()
diff --git a/Engine/Source/Editor/BlueprintGraph/Private/K2Node_TemporaryVariable.cpp b/Engine/Source/Editor/BlueprintGraph/Private/K2Node_TemporaryVariable.cpp
--- a/Engine/Source/Editor/BlueprintGraph/Private/K2Node_TemporaryVariable.cpp
+++ b/Engine/Source/Editor/BlueprintGraph/Private/K2Node_TemporaryVariable.cpp
@@ -129,6 +129,20 @@ FNodeHandlingFunctor* UK2Node_TemporaryVariable::CreateNodeHandler(FKismetCompil
 
 void UK2Node_TemporaryVariable::GetMenuActions(FBlueprintActionDatabaseRegistrar& ActionRegistrar) const
 {
+	// actions get registered under specific object-keys; the idea is that 
+	// actions might have to be updated (or deleted) if their object-key is  
+	// mutated (or removed)... here we use the node's class (so if the node 
+	// type disappears, then the action should go with it)
+	UClass* ActionKey = GetClass();
+	// to keep from needlessly instantiating a UBlueprintNodeSpawner, first   
+	// check to make sure that the registrar is looking for actions of this type
+	// (could be regenerating actions for a specific asset, and therefore the 
+	// registrar would only accept actions corresponding to that asset)
+	if (!ActionRegistrar.IsOpenForRegistration(ActionKey))
+	{
+		return;
+	}
+
 	auto MakeTempVarNodeSpawner = [](FEdGraphPinType const& VarType, bool bIsPersistent) -> UBlueprintNodeSpawner*
 	{
 		UBlueprintNodeSpawner* NodeSpawner = UBlueprintNodeSpawner::Create(UK2Node_TemporaryVariable::StaticClass());
@@ -147,36 +161,37 @@ void UK2Node_TemporaryVariable::GetMenuActions(FBlueprintActionDatabaseRegistrar
 
 	UEdGraphSchema_K2 const* K2Schema = GetDefault<UEdGraphSchema_K2>();
 
-	ActionRegistrar.AddBlueprintAction(MakeTempVarNodeSpawner(FEdGraphPinType(K2Schema->PC_Int,     TEXT(""), nullptr, /*bIsArray =*/false, /*bIsReference =*/false), /*bIsPersistent =*/false));
-	ActionRegistrar.AddBlueprintAction(MakeTempVarNodeSpawner(FEdGraphPinType(K2Schema->PC_Int,     TEXT(""), nullptr, /*bIsArray =*/ true, /*bIsReference =*/false), /*bIsPersistent =*/false));
-	ActionRegistrar.AddBlueprintAction(MakeTempVarNodeSpawner(FEdGraphPinType(K2Schema->PC_Float,   TEXT(""), nullptr, /*bIsArray =*/false, /*bIsReference =*/false), /*bIsPersistent =*/false));
-	ActionRegistrar.AddBlueprintAction(MakeTempVarNodeSpawner(FEdGraphPinType(K2Schema->PC_Float,   TEXT(""), nullptr, /*bIsArray =*/ true, /*bIsReference =*/false), /*bIsPersistent =*/false));
-	ActionRegistrar.AddBlueprintAction(MakeTempVarNodeSpawner(FEdGraphPinType(K2Schema->PC_Boolean, TEXT(""), nullptr, /*bIsArray =*/false, /*bIsReference =*/false), /*bIsPersistent =*/false));
-	ActionRegistrar.AddBlueprintAction(MakeTempVarNodeSpawner(FEdGraphPinType(K2Schema->PC_Boolean, TEXT(""), nullptr, /*bIsArray =*/ true, /*bIsReference =*/false), /*bIsPersistent =*/false));
-	ActionRegistrar.AddBlueprintAction(MakeTempVarNodeSpawner(FEdGraphPinType(K2Schema->PC_String,  TEXT(""), nullptr, /*bIsArray =*/false, /*bIsReference =*/false), /*bIsPersistent =*/false));
-	ActionRegistrar.AddBlueprintAction(MakeTempVarNodeSpawner(FEdGraphPinType(K2Schema->PC_String,  TEXT(""), nullptr, /*bIsArray =*/ true, /*bIsReference =*/false), /*bIsPersistent =*/false));
-	ActionRegistrar.AddBlueprintAction(MakeTempVarNodeSpawner(FEdGraphPinType(K2Schema->PC_Text,    TEXT(""), nullptr, /*bIsArray =*/false, /*bIsReference =*/false), /*bIsPersistent =*/false));
-	ActionRegistrar.AddBlueprintAction(MakeTempVarNodeSpawner(FEdGraphPinType(K2Schema->PC_Text,    TEXT(""), nullptr, /*bIsArray =*/ true, /*bIsReference =*/false), /*bIsPersistent =*/false));
+	ActionRegistrar.AddBlueprintAction(ActionKey, MakeTempVarNodeSpawner(FEdGraphPinType(K2Schema->PC_Int, TEXT(""), nullptr, /*bIsArray =*/false, /*bIsReference =*/false), /*bIsPersistent =*/false));
+	ActionRegistrar.AddBlueprintAction(ActionKey, MakeTempVarNodeSpawner(FEdGraphPinType(K2Schema->PC_Int, TEXT(""), nullptr, /*bIsArray =*/ true, /*bIsReference =*/false), /*bIsPersistent =*/false));
+	ActionRegistrar.AddBlueprintAction(ActionKey, MakeTempVarNodeSpawner(FEdGraphPinType(K2Schema->PC_Float, TEXT(""), nullptr, /*bIsArray =*/false, /*bIsReference =*/false), /*bIsPersistent =*/false));
+	ActionRegistrar.AddBlueprintAction(ActionKey, MakeTempVarNodeSpawner(FEdGraphPinType(K2Schema->PC_Float, TEXT(""), nullptr, /*bIsArray =*/ true, /*bIsReference =*/false), /*bIsPersistent =*/false));
+	ActionRegistrar.AddBlueprintAction(ActionKey, MakeTempVarNodeSpawner(FEdGraphPinType(K2Schema->PC_Boolean, TEXT(""), nullptr, /*bIsArray =*/false, /*bIsReference =*/false), /*bIsPersistent =*/false));
+	ActionRegistrar.AddBlueprintAction(ActionKey, MakeTempVarNodeSpawner(FEdGraphPinType(K2Schema->PC_Boolean, TEXT(""), nullptr, /*bIsArray =*/ true, /*bIsReference =*/false), /*bIsPersistent =*/false));
+	ActionRegistrar.AddBlueprintAction(ActionKey, MakeTempVarNodeSpawner(FEdGraphPinType(K2Schema->PC_String, TEXT(""), nullptr, /*bIsArray =*/false, /*bIsReference =*/false), /*bIsPersistent =*/false));
+	ActionRegistrar.AddBlueprintAction(ActionKey, MakeTempVarNodeSpawner(FEdGraphPinType(K2Schema->PC_String, TEXT(""), nullptr, /*bIsArray =*/ true, /*bIsReference =*/false), /*bIsPersistent =*/false));
+	ActionRegistrar.AddBlueprintAction(ActionKey, MakeTempVarNodeSpawner(FEdGraphPinType(K2Schema->PC_Text, TEXT(""), nullptr, /*bIsArray =*/false, /*bIsReference =*/false), /*bIsPersistent =*/false));
+	ActionRegistrar.AddBlueprintAction(ActionKey, MakeTempVarNodeSpawner(FEdGraphPinType(K2Schema->PC_Text, TEXT(""), nullptr, /*bIsArray =*/ true, /*bIsReference =*/false), /*bIsPersistent =*/false));
 
 	UScriptStruct* VectorStruct  = FindObjectChecked<UScriptStruct>(UObject::StaticClass(), TEXT("Vector"));
-	ActionRegistrar.AddBlueprintAction(MakeTempVarNodeSpawner(FEdGraphPinType(K2Schema->PC_Struct, TEXT("Vector"), VectorStruct, /*bIsArray =*/false, /*bIsReference =*/false), /*bIsPersistent =*/false));
-	ActionRegistrar.AddBlueprintAction(MakeTempVarNodeSpawner(FEdGraphPinType(K2Schema->PC_Struct, TEXT("Vector"), VectorStruct, /*bIsArray =*/ true, /*bIsReference =*/false), /*bIsPersistent =*/false));
+	ActionRegistrar.AddBlueprintAction(ActionKey, MakeTempVarNodeSpawner(FEdGraphPinType(K2Schema->PC_Struct, TEXT("Vector"), VectorStruct, /*bIsArray =*/false, /*bIsReference =*/false), /*bIsPersistent =*/false));
+	ActionRegistrar.AddBlueprintAction(ActionKey, MakeTempVarNodeSpawner(FEdGraphPinType(K2Schema->PC_Struct, TEXT("Vector"), VectorStruct, /*bIsArray =*/ true, /*bIsReference =*/false), /*bIsPersistent =*/false));
 	
 	UScriptStruct* RotatorStruct = FindObjectChecked<UScriptStruct>(UObject::StaticClass(), TEXT("Rotator"));
-	ActionRegistrar.AddBlueprintAction(MakeTempVarNodeSpawner(FEdGraphPinType(K2Schema->PC_Struct, TEXT("Rotator"), RotatorStruct, /*bIsArray =*/false, /*bIsReference =*/false), /*bIsPersistent =*/false));
-	ActionRegistrar.AddBlueprintAction(MakeTempVarNodeSpawner(FEdGraphPinType(K2Schema->PC_Struct, TEXT("Rotator"), RotatorStruct, /*bIsArray =*/ true, /*bIsReference =*/false), /*bIsPersistent =*/false));
+	ActionRegistrar.AddBlueprintAction(ActionKey, MakeTempVarNodeSpawner(FEdGraphPinType(K2Schema->PC_Struct, TEXT("Rotator"), RotatorStruct, /*bIsArray =*/false, /*bIsReference =*/false), /*bIsPersistent =*/false));
+	ActionRegistrar.AddBlueprintAction(ActionKey, MakeTempVarNodeSpawner(FEdGraphPinType(K2Schema->PC_Struct, TEXT("Rotator"), RotatorStruct, /*bIsArray =*/ true, /*bIsReference =*/false), /*bIsPersistent =*/false));
 	
 	UScriptStruct* TransformStruct = FindObjectChecked<UScriptStruct>(UObject::StaticClass(), TEXT("Transform"));
-	ActionRegistrar.AddBlueprintAction(MakeTempVarNodeSpawner(FEdGraphPinType(K2Schema->PC_Struct, TEXT("Transform"), TransformStruct, /*bIsArray =*/false, /*bIsReference =*/false), /*bIsPersistent =*/false));
-	ActionRegistrar.AddBlueprintAction(MakeTempVarNodeSpawner(FEdGraphPinType(K2Schema->PC_Struct, TEXT("Transform"), TransformStruct, /*bIsArray =*/ true, /*bIsReference =*/false), /*bIsPersistent =*/false));
+	ActionRegistrar.AddBlueprintAction(ActionKey, MakeTempVarNodeSpawner(FEdGraphPinType(K2Schema->PC_Struct, TEXT("Transform"), TransformStruct, /*bIsArray =*/false, /*bIsReference =*/false), /*bIsPersistent =*/false));
+	ActionRegistrar.AddBlueprintAction(ActionKey, MakeTempVarNodeSpawner(FEdGraphPinType(K2Schema->PC_Struct, TEXT("Transform"), TransformStruct, /*bIsArray =*/ true, /*bIsReference =*/false), /*bIsPersistent =*/false));
 
 	UScriptStruct* BlendSampleStruct = FindObjectChecked<UScriptStruct>(ANY_PACKAGE, TEXT("BlendSampleData"));
-	ActionRegistrar.AddBlueprintAction(MakeTempVarNodeSpawner(FEdGraphPinType(K2Schema->PC_Struct, TEXT("BlendSampleData"), BlendSampleStruct, /*bIsArray =*/false, /*bIsReference =*/false), /*bIsPersistent =*/false));
-	ActionRegistrar.AddBlueprintAction(MakeTempVarNodeSpawner(FEdGraphPinType(K2Schema->PC_Struct, TEXT("BlendSampleData"), BlendSampleStruct, /*bIsArray =*/ true, /*bIsReference =*/false), /*bIsPersistent =*/false));
+	ActionRegistrar.AddBlueprintAction(ActionKey, MakeTempVarNodeSpawner(FEdGraphPinType(K2Schema->PC_Struct, TEXT("BlendSampleData"), BlendSampleStruct, /*bIsArray =*/false, /*bIsReference =*/false), /*bIsPersistent =*/false));
+	ActionRegistrar.AddBlueprintAction(ActionKey, MakeTempVarNodeSpawner(FEdGraphPinType(K2Schema->PC_Struct, TEXT("BlendSampleData"), BlendSampleStruct, /*bIsArray =*/ true, /*bIsReference =*/false), /*bIsPersistent =*/false));
 
 	// add persistent bool and int types (for macro graphs)
-	ActionRegistrar.AddBlueprintAction(MakeTempVarNodeSpawner(FEdGraphPinType(K2Schema->PC_Int,     TEXT(""), nullptr, /*bIsArray =*/false, /*bIsReference =*/false), /*bIsPersistent =*/true));
-	ActionRegistrar.AddBlueprintAction(MakeTempVarNodeSpawner(FEdGraphPinType(K2Schema->PC_Boolean, TEXT(""), nullptr, /*bIsArray =*/false, /*bIsReference =*/false), /*bIsPersistent =*/true));
+	// @TODO: these should probably be filtered (for macro graphs only)
+	ActionRegistrar.AddBlueprintAction(ActionKey, MakeTempVarNodeSpawner(FEdGraphPinType(K2Schema->PC_Int, TEXT(""), nullptr, /*bIsArray =*/false, /*bIsReference =*/false), /*bIsPersistent =*/true));
+	ActionRegistrar.AddBlueprintAction(ActionKey, MakeTempVarNodeSpawner(FEdGraphPinType(K2Schema->PC_Boolean, TEXT(""), nullptr, /*bIsArray =*/false, /*bIsReference =*/false), /*bIsPersistent =*/true));
 }
 
 FText UK2Node_TemporaryVariable::GetMenuCategory() const
diff --git a/Engine/Source/Editor/BlueprintGraph/Private/K2Node_Timeline.cpp b/Engine/Source/Editor/BlueprintGraph/Private/K2Node_Timeline.cpp
--- a/Engine/Source/Editor/BlueprintGraph/Private/K2Node_Timeline.cpp
+++ b/Engine/Source/Editor/BlueprintGraph/Private/K2Node_Timeline.cpp
@@ -571,28 +571,40 @@ void UK2Node_Timeline::GetNodeAttributes( TArray<TKeyValuePair<FString, FString>
 
 void UK2Node_Timeline::GetMenuActions(FBlueprintActionDatabaseRegistrar& ActionRegistrar) const
 {
-	UBlueprintNodeSpawner* NodeSpawner = UBlueprintNodeSpawner::Create(GetClass());
-	check(NodeSpawner != nullptr);
-
-	auto CustomizeTimelineNodeLambda = [](UEdGraphNode* NewNode, bool bIsTemplateNode)
+	// actions get registered under specific object-keys; the idea is that 
+	// actions might have to be updated (or deleted) if their object-key is  
+	// mutated (or removed)... here we use the node's class (so if the node 
+	// type disappears, then the action should go with it)
+	UClass* ActionKey = GetClass();
+	// to keep from needlessly instantiating a UBlueprintNodeSpawner, first   
+	// check to make sure that the registrar is looking for actions of this type
+	// (could be regenerating actions for a specific asset, and therefore the 
+	// registrar would only accept actions corresponding to that asset)
+	if (ActionRegistrar.IsOpenForRegistration(ActionKey))
 	{
-		UK2Node_Timeline* TimelineNode = CastChecked<UK2Node_Timeline>(NewNode);
+		UBlueprintNodeSpawner* NodeSpawner = UBlueprintNodeSpawner::Create(GetClass());
+		check(NodeSpawner != nullptr);
 
-		UBlueprint* Blueprint = TimelineNode->GetBlueprint();
-		if (Blueprint != nullptr)
+		auto CustomizeTimelineNodeLambda = [](UEdGraphNode* NewNode, bool bIsTemplateNode)
 		{
-			TimelineNode->TimelineName = FBlueprintEditorUtils::FindUniqueTimelineName(Blueprint);
-			if (!bIsTemplateNode && FBlueprintEditorUtils::AddNewTimeline(Blueprint, TimelineNode->TimelineName))
+			UK2Node_Timeline* TimelineNode = CastChecked<UK2Node_Timeline>(NewNode);
+
+			UBlueprint* Blueprint = TimelineNode->GetBlueprint();
+			if (Blueprint != nullptr)
 			{
-				// clear off any existing error message now that the timeline has been added
-				TimelineNode->ErrorMsg.Empty();
-				TimelineNode->bHasCompilerMessage = false;
+				TimelineNode->TimelineName = FBlueprintEditorUtils::FindUniqueTimelineName(Blueprint);
+				if (!bIsTemplateNode && FBlueprintEditorUtils::AddNewTimeline(Blueprint, TimelineNode->TimelineName))
+				{
+					// clear off any existing error message now that the timeline has been added
+					TimelineNode->ErrorMsg.Empty();
+					TimelineNode->bHasCompilerMessage = false;
+				}
 			}
-		}
-	};
+		};
 
-	NodeSpawner->CustomizeNodeDelegate = UBlueprintNodeSpawner::FCustomizeNodeDelegate::CreateStatic(CustomizeTimelineNodeLambda);
-	ActionRegistrar.AddBlueprintAction(NodeSpawner);
+		NodeSpawner->CustomizeNodeDelegate = UBlueprintNodeSpawner::FCustomizeNodeDelegate::CreateStatic(CustomizeTimelineNodeLambda);
+		ActionRegistrar.AddBlueprintAction(ActionKey, NodeSpawner);
+	}
 }
 
 #undef LOCTEXT_NAMESPACE
diff --git a/Engine/Source/Editor/BlueprintGraph/Private/K2Node_VariableSetRef.cpp b/Engine/Source/Editor/BlueprintGraph/Private/K2Node_VariableSetRef.cpp
--- a/Engine/Source/Editor/BlueprintGraph/Private/K2Node_VariableSetRef.cpp
+++ b/Engine/Source/Editor/BlueprintGraph/Private/K2Node_VariableSetRef.cpp
@@ -210,10 +210,22 @@ FNodeHandlingFunctor* UK2Node_VariableSetRef::CreateNodeHandler(FKismetCompilerC
 
 void UK2Node_VariableSetRef::GetMenuActions(FBlueprintActionDatabaseRegistrar& ActionRegistrar) const
 {
-	UBlueprintNodeSpawner* NodeSpawner = UBlueprintNodeSpawner::Create(GetClass());
-	check(NodeSpawner != nullptr);
+	// actions get registered under specific object-keys; the idea is that 
+	// actions might have to be updated (or deleted) if their object-key is  
+	// mutated (or removed)... here we use the node's class (so if the node 
+	// type disappears, then the action should go with it)
+	UClass* ActionKey = GetClass();
+	// to keep from needlessly instantiating a UBlueprintNodeSpawner, first   
+	// check to make sure that the registrar is looking for actions of this type
+	// (could be regenerating actions for a specific asset, and therefore the 
+	// registrar would only accept actions corresponding to that asset)
+	if (ActionRegistrar.IsOpenForRegistration(ActionKey))
+	{
+		UBlueprintNodeSpawner* NodeSpawner = UBlueprintNodeSpawner::Create(GetClass());
+		check(NodeSpawner != nullptr);
 
-	ActionRegistrar.AddBlueprintAction(NodeSpawner);
+		ActionRegistrar.AddBlueprintAction(ActionKey, NodeSpawner);
+	}
 }
 
 FText UK2Node_VariableSetRef::GetMenuCategory() const
diff --git a/Engine/Source/Editor/BlueprintGraph/Public/BlueprintActionDatabaseRegistrar.h b/Engine/Source/Editor/BlueprintGraph/Public/BlueprintActionDatabaseRegistrar.h
--- a/Engine/Source/Editor/BlueprintGraph/Public/BlueprintActionDatabaseRegistrar.h
+++ b/Engine/Source/Editor/BlueprintGraph/Public/BlueprintActionDatabaseRegistrar.h
@@ -19,7 +19,7 @@ public:
 	 * 
 	 * @param  NodeSpawner	The new node spawner that you wish to register.
 	 */
-	void AddBlueprintAction(UBlueprintNodeSpawner* NodeSpawner);
+	bool AddBlueprintAction(UBlueprintNodeSpawner* NodeSpawner);
 
 	/**
 	 * Each action should be recorded under a specific UField key; primarily to 
@@ -29,10 +29,10 @@ public:
 	 * @param  FieldOwner	A field object that the node is associated with (flags when this action should be updated).
 	 * @param  NodeSpawner	The new node spawner that you wish to register.
 	 */
-	void AddBlueprintAction(UClass const* ClassOwner, UBlueprintNodeSpawner* NodeSpawner);
-	void AddBlueprintAction(UEnum const* EnumOwner, UBlueprintNodeSpawner* NodeSpawner);
-	void AddBlueprintAction(UScriptStruct const* StructOwner, UBlueprintNodeSpawner* NodeSpawner);
-	void AddBlueprintAction(UField const* FieldOwner, UBlueprintNodeSpawner* NodeSpawner);
+	bool AddBlueprintAction(UClass const* ClassOwner, UBlueprintNodeSpawner* NodeSpawner);
+	bool AddBlueprintAction(UEnum const* EnumOwner, UBlueprintNodeSpawner* NodeSpawner);
+	bool AddBlueprintAction(UScriptStruct const* StructOwner, UBlueprintNodeSpawner* NodeSpawner);
+	bool AddBlueprintAction(UField const* FieldOwner, UBlueprintNodeSpawner* NodeSpawner);
 
 	/**
 	 * Special case for asset bound actions (we want to clean-up/refresh these 
@@ -41,7 +41,7 @@ public:
 	 * @param  AssetOwner	An asset object that the node is associated with (flags when this action should be updated).
 	 * @param  NodeSpawner	The new node spawner that you wish to register.
 	 */
-	void AddBlueprintAction(UObject const* AssetOwner, UBlueprintNodeSpawner* NodeSpawner);
+	bool AddBlueprintAction(UObject const* AssetOwner, UBlueprintNodeSpawner* NodeSpawner);
 
 	/**
 	 * Occasionally (when an asset is added/refreshed), this registrar will be 
@@ -69,7 +69,7 @@ private:
 	 * @param  ActionKey	The key you want the action listed under.
 	 * @param  NodeSpawner	The new node spawner that you wish to register.
 	 */
-	void AddActionToDatabase(UObject const* ActionKey, UBlueprintNodeSpawner* NodeSpawner);
+	bool AddActionToDatabase(UObject const* ActionKey, UBlueprintNodeSpawner* NodeSpawner);
 
 private:
 	/** The node type that this is currently passed to (acts as a fallback database key). */
diff --git a/Engine/Source/Editor/GameplayAbilitiesEditor/Private/K2Node_GameplayCueEvent.cpp b/Engine/Source/Editor/GameplayAbilitiesEditor/Private/K2Node_GameplayCueEvent.cpp
--- a/Engine/Source/Editor/GameplayAbilitiesEditor/Private/K2Node_GameplayCueEvent.cpp
+++ b/Engine/Source/Editor/GameplayAbilitiesEditor/Private/K2Node_GameplayCueEvent.cpp
@@ -86,6 +86,20 @@ void UK2Node_GameplayCueEvent::GetMenuEntries(FGraphContextMenuBuilder& Context)
 
 void UK2Node_GameplayCueEvent::GetMenuActions(FBlueprintActionDatabaseRegistrar& ActionRegistrar) const
 {
+	// actions get registered under specific object-keys; the idea is that 
+	// actions might have to be updated (or deleted) if their object-key is  
+	// mutated (or removed)... here we use the node's class (so if the node 
+	// type disappears, then the action should go with it)
+	UClass* ActionKey = GetClass();
+	// to keep from needlessly instantiating a UBlueprintNodeSpawner, first   
+	// check to make sure that the registrar is looking for actions of this type
+	// (could be regenerating actions for a specific asset, and therefore the 
+	// registrar would only accept actions corresponding to that asset)
+	if (!ActionRegistrar.IsOpenForRegistration(ActionKey))
+	{
+		return;
+	}
+
 	auto CustomizeCueNodeLambda = [](UEdGraphNode* NewNode, bool bIsTemplateNode, FName TagName)
 	{
 		UK2Node_GameplayCueEvent* EventNode = CastChecked<UK2Node_GameplayCueEvent>(NewNode);
@@ -104,7 +118,7 @@ void UK2Node_GameplayCueEvent::GetMenuActions(FBlueprintActionDatabaseRegistrar&
 		check(NodeSpawner != nullptr);
 		NodeSpawner->CustomizeNodeDelegate = PostSpawnDelegate;
 		
-		ActionRegistrar.AddBlueprintAction(NodeSpawner);
+		ActionRegistrar.AddBlueprintAction(ActionKey, NodeSpawner);
 	}
 }
 
diff --git a/Engine/Source/Editor/GameplayAbilitiesEditor/Private/K2Node_LatentAbilityCall.cpp b/Engine/Source/Editor/GameplayAbilitiesEditor/Private/K2Node_LatentAbilityCall.cpp
--- a/Engine/Source/Editor/GameplayAbilitiesEditor/Private/K2Node_LatentAbilityCall.cpp
+++ b/Engine/Source/Editor/GameplayAbilitiesEditor/Private/K2Node_LatentAbilityCall.cpp
@@ -101,6 +101,15 @@ void UK2Node_LatentAbilityCall::GetMenuActions(FBlueprintActionDatabaseRegistrar
 			{
 				continue;
 			}
+
+			// to keep from needlessly instantiating a UBlueprintNodeSpawner, first   
+			// check to make sure that the registrar is looking for actions of this type
+			// (could be regenerating actions for a specific asset, and therefore the 
+			// registrar would only accept actions corresponding to that asset)
+			if (!ActionRegistrar.IsOpenForRegistration(Function))
+			{
+				continue;
+			}
 			
 			UObjectProperty* ReturnProperty = Cast<UObjectProperty>(Function->GetReturnProperty());
 			// see if the function is a static factory method for online proxies
@@ -131,7 +140,7 @@ void UK2Node_LatentAbilityCall::GetMenuActions(FBlueprintActionDatabaseRegistrar
 				
 				// @TODO: since this can't be folded into FBlueprintActionDatabase, we
 				//        need a way to associate these spawners with a certain class
-				ActionRegistrar.AddBlueprintAction(NodeSpawner);
+				ActionRegistrar.AddBlueprintAction(Function, NodeSpawner);
 			}
 		}
 	}
diff --git a/Engine/Source/Editor/GameplayTagsEditor/Private/GameplayTagsK2Node_LiteralGameplayTag.cpp b/Engine/Source/Editor/GameplayTagsEditor/Private/GameplayTagsK2Node_LiteralGameplayTag.cpp
--- a/Engine/Source/Editor/GameplayTagsEditor/Private/GameplayTagsK2Node_LiteralGameplayTag.cpp
+++ b/Engine/Source/Editor/GameplayTagsEditor/Private/GameplayTagsK2Node_LiteralGameplayTag.cpp
@@ -149,10 +149,22 @@ void UGameplayTagsK2Node_LiteralGameplayTag::ExpandNode(class FKismetCompilerCon
 
 void UGameplayTagsK2Node_LiteralGameplayTag::GetMenuActions(FBlueprintActionDatabaseRegistrar& ActionRegistrar) const
 {
-	UBlueprintNodeSpawner* NodeSpawner = UBlueprintNodeSpawner::Create(GetClass());
-	check(NodeSpawner != nullptr);
-	
-	ActionRegistrar.AddBlueprintAction(NodeSpawner);
+	// actions get registered under specific object-keys; the idea is that 
+	// actions might have to be updated (or deleted) if their object-key is  
+	// mutated (or removed)... here we use the node's class (so if the node 
+	// type disappears, then the action should go with it)
+	UClass* ActionKey = GetClass();
+	// to keep from needlessly instantiating a UBlueprintNodeSpawner, first   
+	// check to make sure that the registrar is looking for actions of this type
+	// (could be regenerating actions for a specific asset, and therefore the 
+	// registrar would only accept actions corresponding to that asset)
+	if (ActionRegistrar.IsOpenForRegistration(ActionKey))
+	{
+		UBlueprintNodeSpawner* NodeSpawner = UBlueprintNodeSpawner::Create(GetClass());
+		check(NodeSpawner != nullptr);
+
+		ActionRegistrar.AddBlueprintAction(ActionKey, NodeSpawner);
+	}
 }
 
 FText UGameplayTagsK2Node_LiteralGameplayTag::GetMenuCategory() const
diff --git a/Engine/Source/Editor/OnlineBlueprintSupport/Private/K2Node_LatentOnlineCall.cpp b/Engine/Source/Editor/OnlineBlueprintSupport/Private/K2Node_LatentOnlineCall.cpp
--- a/Engine/Source/Editor/OnlineBlueprintSupport/Private/K2Node_LatentOnlineCall.cpp
+++ b/Engine/Source/Editor/OnlineBlueprintSupport/Private/K2Node_LatentOnlineCall.cpp
@@ -72,6 +72,15 @@ void UK2Node_LatentOnlineCall::GetMenuActions(FBlueprintActionDatabaseRegistrar&
 			{
 				continue;
 			}
+
+			// to keep from needlessly instantiating a UBlueprintNodeSpawner, first   
+			// check to make sure that the registrar is looking for actions of this type
+			// (could be regenerating actions for a specific asset, and therefore the 
+			// registrar would only accept actions corresponding to that asset)
+			if (!ActionRegistrar.IsOpenForRegistration(Function))
+			{
+				continue;
+			}
 			
 			UObjectProperty* ReturnProperty = Cast<UObjectProperty>(Function->GetReturnProperty());
 			// see if the function is a static factory method for online proxies
@@ -102,7 +111,7 @@ void UK2Node_LatentOnlineCall::GetMenuActions(FBlueprintActionDatabaseRegistrar&
 				
 				// @TODO: since this can't be folded into FBlueprintActionDatabase, we
 				//        need a way to associate these spawners with a certain class
-				ActionRegistrar.AddBlueprintAction(NodeSpawner);
+				ActionRegistrar.AddBlueprintAction(Function, NodeSpawner);
 			}
 		}
 	}
