#ue4 Fix for CPF_Edit pins being exposed that have an ArrayDim > 1, because static arrays are not supported in BP VM code

diff --git a/Engine/Source/Editor/BlueprintGraph/Private/K2Node_MakeStruct.cpp b/Engine/Source/Editor/BlueprintGraph/Private/K2Node_MakeStruct.cpp
--- a/Engine/Source/Editor/BlueprintGraph/Private/K2Node_MakeStruct.cpp
+++ b/Engine/Source/Editor/BlueprintGraph/Private/K2Node_MakeStruct.cpp
@@ -206,7 +206,7 @@ bool UK2Node_MakeStruct::CanBeExposed(const UProperty* Property)
 		const bool bConvertable = Schema->ConvertPropertyToPinType(Property, /*out*/ DumbGraphPinType);
 
 		//TODO: remove CPF_Edit in a future release
-		const bool bVisible = Property->HasAnyPropertyFlags(CPF_BlueprintVisible|CPF_Edit);
+		const bool bVisible = Property->HasAnyPropertyFlags(CPF_BlueprintVisible|CPF_Edit) && !(Property->ArrayDim > 1);
 		const bool bBlueprintReadOnly = Property->HasAllPropertyFlags(CPF_BlueprintReadOnly);
 		if(bVisible && bConvertable && !bBlueprintReadOnly)
 		{
