Selected actor references can now be placed using the new Blueprint menu system.

diff --git a/Engine/Source/Editor/BlueprintGraph/Classes/BlueprintBoundNodeSpawner.h b/Engine/Source/Editor/BlueprintGraph/Classes/BlueprintBoundNodeSpawner.h
new file mode 100644
--- /dev/null
+++ b/Engine/Source/Editor/BlueprintGraph/Classes/BlueprintBoundNodeSpawner.h
@@ -0,0 +1,40 @@
+// Copyright 1998-2014 Epic Games, Inc. All Rights Reserved.
+
+#pragma once
+
+#include "BlueprintNodeSpawner.h"
+#include "BlueprintBoundNodeSpawner.generated.h"
+
+/**
+ * Takes care of spawning various bound nodes. Acts as the 
+ * "action" portion of certain FBlueprintActionMenuItems. 
+ */
+UCLASS(Transient)
+class BLUEPRINTGRAPH_API UBlueprintBoundNodeSpawner : public UBlueprintNodeSpawner
+{
+	GENERATED_UCLASS_BODY()
+	DECLARE_DELEGATE_RetVal_OneParam(bool, FCanBindObjectDelegate, UObject const*);
+	DECLARE_DELEGATE_RetVal_TwoParams(bool, FOnBindObjectDelegate, UEdGraphNode*, UObject*);
+
+public:
+	/**
+	 * @return A newly allocated instance of this class.
+	 */
+	static UBlueprintBoundNodeSpawner* Create(TSubclassOf<UEdGraphNode> NodeClass, UObject* Outer = nullptr);
+
+	// FBlueprintNodeBinder interface
+	virtual bool IsBindingCompatible(UObject const* BindingCandidate) const override;
+	virtual bool BindToNode(UEdGraphNode* Node, UObject* Binding) const override;
+	// End FBlueprintNodeBinder interface
+
+public:
+	/**
+	 * A delegate to perform specialized node binding verification
+	 */
+	FCanBindObjectDelegate CanBindObjectDelegate;
+
+	/**
+	 * A delegate to perform specialized node setup during binding
+	 */
+	FOnBindObjectDelegate OnBindObjectDelegate;
+};
diff --git a/Engine/Source/Editor/BlueprintGraph/Classes/K2Node_Literal.h b/Engine/Source/Editor/BlueprintGraph/Classes/K2Node_Literal.h
--- a/Engine/Source/Editor/BlueprintGraph/Classes/K2Node_Literal.h
+++ b/Engine/Source/Editor/BlueprintGraph/Classes/K2Node_Literal.h
@@ -34,6 +34,7 @@ public:
 	virtual bool NodeCausesStructuralBlueprintChange() const override { return true; }
 	virtual void PostReconstructNode() override;
 	virtual class FNodeHandlingFunctor* CreateNodeHandler(class FKismetCompilerContext& CompilerContext) const override;
+	virtual void GetMenuActions(FBlueprintActionDatabaseRegistrar& ActionRegistrar) const override;
 	// End UK2Node interface
 
 	/** Accessor for the value pin of the node */
diff --git a/Engine/Source/Editor/BlueprintGraph/Private/BlueprintActionFilter.cpp b/Engine/Source/Editor/BlueprintGraph/Private/BlueprintActionFilter.cpp
--- a/Engine/Source/Editor/BlueprintGraph/Private/BlueprintActionFilter.cpp
+++ b/Engine/Source/Editor/BlueprintGraph/Private/BlueprintActionFilter.cpp
@@ -7,6 +7,7 @@
 #include "BlueprintVariableNodeSpawner.h"
 #include "BlueprintEventNodeSpawner.h"
 #include "BlueprintBoundEventNodeSpawner.h"
+#include "BlueprintBoundNodeSpawner.h"
 #include "EdGraphSchema_K2.h"		// for FBlueprintMetadata
 #include "BlueprintEditorUtils.h"	// for FindBlueprintForGraph()
 #include "ObjectEditorUtils.h"		// for IsFunctionHiddenFromClass()/IsVariableCategoryHiddenFromClass()
@@ -672,7 +673,7 @@ static bool BlueprintActionFilterImpl::IsFilteredNodeType(FBlueprintActionFilter
 //------------------------------------------------------------------------------
 static bool BlueprintActionFilterImpl::IsUnBoundBindingSpawner(FBlueprintActionFilter const& Filter, FBlueprintActionInfo& BlueprintAction)
 {
-	bool const bIsBindingSpecificSpawner = (Cast<UBlueprintBoundEventNodeSpawner>(BlueprintAction.NodeSpawner) != nullptr);
+	bool const bIsBindingSpecificSpawner = (Cast<UBlueprintBoundEventNodeSpawner>(BlueprintAction.NodeSpawner) != nullptr || Cast<UBlueprintBoundNodeSpawner>(BlueprintAction.NodeSpawner) != nullptr);
 
 	bool bIsFilteredOut = false;
 	if (bIsBindingSpecificSpawner)
@@ -759,7 +760,7 @@ static bool BlueprintActionFilterImpl::HasMatchingPin(FBlueprintActionInfo& Blue
 	bool bHasCompatiblePin = false;
 
 	UEdGraph* OuterGraph = Pin->GetOwningNode()->GetGraph();
-	if (UEdGraphNode* TemplateNode = BlueprintAction.NodeSpawner->GetTemplateNode(OuterGraph))
+	if (UEdGraphNode* TemplateNode = BlueprintAction.NodeSpawner->GetTemplateNode(OuterGraph, BlueprintAction.GetBindings()))
 	{
 		if (TemplateNode->Pins.Num() == 0)
 		{
diff --git a/Engine/Source/Editor/BlueprintGraph/Private/BlueprintBoundNodeSpawner.cpp b/Engine/Source/Editor/BlueprintGraph/Private/BlueprintBoundNodeSpawner.cpp
new file mode 100644
--- /dev/null
+++ b/Engine/Source/Editor/BlueprintGraph/Private/BlueprintBoundNodeSpawner.cpp
@@ -0,0 +1,52 @@
+// Copyright 1998-2014 Epic Games, Inc. All Rights Reserved.
+
+#include "BlueprintGraphPrivatePCH.h"
+#include "BlueprintBoundNodeSpawner.h"
+
+#define LOCTEXT_NAMESPACE "BlueprintBoundNodeSpawner"
+
+/*******************************************************************************
+ * UBlueprintBoundNodeSpawner
+ ******************************************************************************/
+
+//------------------------------------------------------------------------------
+UBlueprintBoundNodeSpawner* UBlueprintBoundNodeSpawner::Create(TSubclassOf<UEdGraphNode> NodeClass, UObject* Outer/* = nullptr*/)
+{
+	if (Outer == nullptr)
+	{
+		Outer = GetTransientPackage();
+	}
+
+	UBlueprintBoundNodeSpawner* NodeSpawner = NewObject<UBlueprintBoundNodeSpawner>(Outer);
+	NodeSpawner->NodeClass = NodeClass;
+
+	return NodeSpawner;
+}
+
+//------------------------------------------------------------------------------
+UBlueprintBoundNodeSpawner::UBlueprintBoundNodeSpawner(class FPostConstructInitializeProperties const& PCIP)
+	: Super(PCIP)
+{
+}
+
+//------------------------------------------------------------------------------
+bool UBlueprintBoundNodeSpawner::IsBindingCompatible(UObject const* BindingCandidate) const
+{
+	if(CanBindObjectDelegate.IsBound())
+	{
+		return CanBindObjectDelegate.Execute(BindingCandidate);
+	}
+	return false;
+}
+
+//------------------------------------------------------------------------------
+bool UBlueprintBoundNodeSpawner::BindToNode(UEdGraphNode* Node, UObject* Binding) const
+{
+	if(OnBindObjectDelegate.IsBound())
+	{
+		return OnBindObjectDelegate.Execute(Node, Binding);
+	}
+	return false;
+}
+
+#undef LOCTEXT_NAMESPACE
diff --git a/Engine/Source/Editor/BlueprintGraph/Private/BlueprintNodeSpawner.cpp b/Engine/Source/Editor/BlueprintGraph/Private/BlueprintNodeSpawner.cpp
--- a/Engine/Source/Editor/BlueprintGraph/Private/BlueprintNodeSpawner.cpp
+++ b/Engine/Source/Editor/BlueprintGraph/Private/BlueprintNodeSpawner.cpp
@@ -121,10 +121,17 @@ FName UBlueprintNodeSpawner::GetDefaultMenuIcon(FLinearColor& ColorOut) const
 }
 
 //------------------------------------------------------------------------------
-UEdGraphNode* UBlueprintNodeSpawner::GetTemplateNode(UEdGraph* Outer) const
-{	
+UEdGraphNode* UBlueprintNodeSpawner::GetTemplateNode(UEdGraph* Outer, FBindingSet const& Bindings) const 
+{       
 	UEdGraphNode* TemplateNode = BlueprintNodeSpawnerImpl::GetSharedTemplateCache()->GetNodeTemplate(this, Outer);
-	return TemplateNode;
+
+	if (TemplateNode && Bindings.Num() > 0) 
+	{ 
+		UEdGraphNode* BoundTemplateNode = DuplicateObject(TemplateNode, TemplateNode->GetOuter()); 
+		ApplyBindings(BoundTemplateNode, Bindings); 
+		return BoundTemplateNode; 
+	} 
+	return TemplateNode; 
 }
 
 //------------------------------------------------------------------------------
diff --git a/Engine/Source/Editor/BlueprintGraph/Private/K2Node_Literal.cpp b/Engine/Source/Editor/BlueprintGraph/Private/K2Node_Literal.cpp
--- a/Engine/Source/Editor/BlueprintGraph/Private/K2Node_Literal.cpp
+++ b/Engine/Source/Editor/BlueprintGraph/Private/K2Node_Literal.cpp
@@ -4,6 +4,7 @@
 #include "BlueprintGraphPrivatePCH.h"
 #include "KismetCompiler.h"
 #include "ClassIconFinder.h"
+#include "BlueprintBoundNodeSpawner.h"
 
 #define LOCTEXT_NAMESPACE "K2Node_Literal"
 
@@ -115,9 +116,14 @@ FText UK2Node_Literal::GetTooltipText() const
 
 FText UK2Node_Literal::GetNodeTitle(ENodeTitleType::Type TitleType) const
 {
+	AActor* Actor = Cast<AActor>(ObjectRef);
+	if(TitleType == ENodeTitleType::ListView)
+	{
+		return NSLOCTEXT("K2Node", "Literal_Title", "Get Selected Actor Reference(s)");
+	}
+
 	if( ObjectRef != NULL )
 	{
-		AActor* Actor = Cast<AActor>(ObjectRef);
 		if(Actor != NULL)
 		{
 			return FText::FromString(Actor->GetActorLabel());
@@ -147,6 +153,34 @@ FLinearColor UK2Node_Literal::GetNodeTitleColor() const
 	}
 }
 
+void UK2Node_Literal::GetMenuActions(FBlueprintActionDatabaseRegistrar& ActionRegistrar) const
+{
+	auto CanBindObjectLambda = [](UObject const* BindingObject)
+	{
+		if(AActor const* Actor = Cast<AActor>(BindingObject))
+		{
+			// Make sure the Actor has a world
+			if(Actor->GetWorld())
+			{
+				return true;
+			}
+		}
+		return false;
+	};
+
+	auto PostBindSetupLambda = [](UEdGraphNode* NewNode, UObject* BindObject)->bool
+	{
+ 		UK2Node_Literal* LiteralNode = CastChecked<UK2Node_Literal>(NewNode);
+ 		LiteralNode->SetObjectRef(BindObject);
+ 		return true;
+	};
+
+	UBlueprintBoundNodeSpawner* NodeSpawner = UBlueprintBoundNodeSpawner::Create(GetClass());
+	NodeSpawner->CanBindObjectDelegate = UBlueprintBoundNodeSpawner::FCanBindObjectDelegate::CreateStatic(CanBindObjectLambda);
+	NodeSpawner->OnBindObjectDelegate = UBlueprintBoundNodeSpawner::FOnBindObjectDelegate::CreateStatic(PostBindSetupLambda);
+	ActionRegistrar.AddBlueprintAction(NodeSpawner);
+}
+
 UK2Node::ERedirectType UK2Node_Literal::DoPinsMatchForReconstruction(const UEdGraphPin* NewPin, int32 NewPinIndex, const UEdGraphPin* OldPin, int32 OldPinIndex) const
 {
 	// This allows the value pin (the only pin) to stay connected through reconstruction, even if the name changes due to an actor in the renamed, etc
@@ -184,7 +218,7 @@ void UK2Node_Literal::SetObjectRef(UObject* NewValue)
 
 	if( ValuePin )
 	{
-		ValuePin->PinFriendlyName = GetNodeTitle(ENodeTitleType::ListView);
+		ValuePin->PinFriendlyName = GetNodeTitle(ENodeTitleType::FullTitle);
 		ValuePin->PinName = ValuePin->PinFriendlyName.BuildSourceString();
 	}
 }
diff --git a/Engine/Source/Editor/BlueprintGraph/Public/BlueprintNodeSpawner.h b/Engine/Source/Editor/BlueprintGraph/Public/BlueprintNodeSpawner.h
--- a/Engine/Source/Editor/BlueprintGraph/Public/BlueprintNodeSpawner.h
+++ b/Engine/Source/Editor/BlueprintGraph/Public/BlueprintNodeSpawner.h
@@ -157,9 +157,10 @@ public:
 	 * then we use TargetGraph as a model to create one that will work.
 	 *
 	 * @param  TargetGraph    Optional param that defines a compatible graph outer (used as an achetype if we don't have a compatible outer on hand).
+	 * @param  Bindings    	  Objects to bind to the template node
 	 * @return Should return a new/cached template-node (but could be null, or some pre-existing node... depends on the sub-class's Invoke() method).
 	 */
-	UEdGraphNode* GetTemplateNode(UEdGraph* TargetGraph = nullptr) const;
+	UEdGraphNode* GetTemplateNode(UEdGraph* TargetGraph = nullptr, FBindingSet const& Bindings = FBindingSet()) const;
 
 	// IBlueprintNodeBinder interface
 	virtual bool   IsBindingCompatible(UObject const* BindingCandidate) const override { return false; }
@@ -190,4 +191,4 @@ template<class NodeType>
 UBlueprintNodeSpawner* UBlueprintNodeSpawner::Create(UObject* Outer, FCustomizeNodeDelegate PostSpawnDelegate)
 {
 	return Create(NodeType::StaticClass(), Outer, PostSpawnDelegate);
-}
\ No newline at end of file
+}
diff --git a/Engine/Source/Editor/GraphEditor/Private/KismetNodes/SGraphNodeK2Var.cpp b/Engine/Source/Editor/GraphEditor/Private/KismetNodes/SGraphNodeK2Var.cpp
--- a/Engine/Source/Editor/GraphEditor/Private/KismetNodes/SGraphNodeK2Var.cpp
+++ b/Engine/Source/Editor/GraphEditor/Private/KismetNodes/SGraphNodeK2Var.cpp
@@ -101,7 +101,7 @@ void SGraphNodeK2Var::UpdateGraphNode()
 			SubTitleText = FText::Format(NSLOCTEXT("GraphEditor", "ActorRef", "from {0}"), LevelName);
 		}
 
-		TitleText = GraphNode->GetNodeTitle(ENodeTitleType::ListView);
+		TitleText = GraphNode->GetNodeTitle(ENodeTitleType::FullTitle);
 
 		TitleHAlign = HAlign_Left;
 		TitleMargin = FMargin(12.0f, VerticalTitleMargin, 32.0f, 2.0f);
diff --git a/Engine/Source/Editor/Kismet/Private/BlueprintActionMenuUtils.cpp b/Engine/Source/Editor/Kismet/Private/BlueprintActionMenuUtils.cpp
--- a/Engine/Source/Editor/Kismet/Private/BlueprintActionMenuUtils.cpp
+++ b/Engine/Source/Editor/Kismet/Private/BlueprintActionMenuUtils.cpp
@@ -120,6 +120,7 @@ void FBlueprintActionMenuUtils::MakeContextMenu(FBlueprintActionContext const& C
 	LevelActorsFilter.Context = Context;
 	LevelActorsFilter.PermittedNodeTypes.Add(UK2Node_CallFunction::StaticClass());
 	LevelActorsFilter.PermittedNodeTypes.Add(UK2Node_ActorBoundEvent::StaticClass());
+	LevelActorsFilter.PermittedNodeTypes.Add(UK2Node_Literal::StaticClass());
 	// only want bound actions for this menu section
 	LevelActorsFilter.AddRejectionTest(FBlueprintActionFilter::FRejectionTestDelegate::CreateStatic(IsUnBoundSpawner));
 
diff --git a/Engine/Source/Editor/Kismet/Private/BlueprintBoundMenuItem.cpp b/Engine/Source/Editor/Kismet/Private/BlueprintBoundMenuItem.cpp
--- a/Engine/Source/Editor/Kismet/Private/BlueprintBoundMenuItem.cpp
+++ b/Engine/Source/Editor/Kismet/Private/BlueprintBoundMenuItem.cpp
@@ -105,6 +105,7 @@ FBlueprintBoundMenuItem::FBlueprintBoundMenuItem(UBlueprintNodeSpawner const* Bo
 UEdGraphNode* FBlueprintBoundMenuItem::PerformAction(UEdGraph* ParentGraph, UEdGraphPin* FromPin, FVector2D const Location, bool bSelectNewNode/* = true*/)
 {
 	UEdGraphNode* NewNode = nullptr;
+	FVector2D CurrentLocation = Location;
 
 	IBlueprintNodeBinder::FBindingSet Bindings;
 	for (auto BoundObjIt = BoundObjects.CreateConstIterator(); BoundObjIt; )
@@ -120,7 +121,14 @@ UEdGraphNode* FBlueprintBoundMenuItem::PerformAction(UEdGraph* ParentGraph, UEdG
 		} while ( BoundObjIt && (BoundSpawner->CanBindMultipleObjects() || (Bindings.Num() == 0)) );
 
 		FBlueprintActionMenuItem BlueprintActionItem(BoundSpawner, Bindings);
-		NewNode = BlueprintActionItem.PerformAction(ParentGraph, FromPin, Location, bSelectNewNode);
+		NewNode = BlueprintActionItem.PerformAction(ParentGraph, FromPin, CurrentLocation, bSelectNewNode);
+
+		// Increase the node location a safe distance so follow-up nodes are not stacked
+		CurrentLocation.Y += UEdGraphSchema_K2::EstimateNodeHeight(NewNode);
+
+		// Clear the bindings array for the next pass
+		Bindings.Empty();
+
 		// @TODO: select ALL spawned nodes
 	}
 	
