#UE4 ease function node could crash with new context menu - just need to verify we have struct subobject..

diff --git a/Engine/Source/Editor/BlueprintGraph/Private/K2Node_EaseFunction.cpp b/Engine/Source/Editor/BlueprintGraph/Private/K2Node_EaseFunction.cpp
--- a/Engine/Source/Editor/BlueprintGraph/Private/K2Node_EaseFunction.cpp
+++ b/Engine/Source/Editor/BlueprintGraph/Private/K2Node_EaseFunction.cpp
@@ -222,6 +222,7 @@ bool UK2Node_EaseFunction::IsConnectionDisallowed(const UEdGraphPin* MyPin, cons
 	{
 		const bool bConnectionOk = (OtherPin->PinType.PinCategory == K2Schema->PC_Float ||
 			(OtherPin->PinType.PinCategory == K2Schema->PC_Struct &&
+			OtherPin->PinType.PinSubCategoryObject.IsValid() &&
 			(OtherPin->PinType.PinSubCategoryObject.Get()->GetName() == TEXT("Vector") ||
 			OtherPin->PinType.PinSubCategoryObject.Get()->GetName() == TEXT("Rotator") ||
 			OtherPin->PinType.PinSubCategoryObject.Get()->GetName() == TEXT("Transform"))));
