Fix non-unity due to FBlueprintActionDatabaseRegistrar

diff --git a/Engine/Source/Editor/BlueprintGraph/Private/K2Node_Literal.cpp b/Engine/Source/Editor/BlueprintGraph/Private/K2Node_Literal.cpp
--- a/Engine/Source/Editor/BlueprintGraph/Private/K2Node_Literal.cpp
+++ b/Engine/Source/Editor/BlueprintGraph/Private/K2Node_Literal.cpp
@@ -5,6 +5,7 @@
 #include "KismetCompiler.h"
 #include "ClassIconFinder.h"
 #include "BlueprintBoundNodeSpawner.h"
+#include "BlueprintActionDatabaseRegistrar.h"
 
 #define LOCTEXT_NAMESPACE "K2Node_Literal"
 
