Shrinking the template node database by removing the need for function/variable nodes to be cached.

diff --git a/Engine/Source/Editor/BlueprintGraph/Classes/BlueprintFunctionNodeSpawner.h b/Engine/Source/Editor/BlueprintGraph/Classes/BlueprintFunctionNodeSpawner.h
--- a/Engine/Source/Editor/BlueprintGraph/Classes/BlueprintFunctionNodeSpawner.h
+++ b/Engine/Source/Editor/BlueprintGraph/Classes/BlueprintFunctionNodeSpawner.h
@@ -42,11 +42,13 @@ public:
 	static UBlueprintFunctionNodeSpawner* Create(TSubclassOf<UK2Node_CallFunction> NodeClass, UFunction const* const Function, UObject* Outer = nullptr);
 
 	// UBlueprintNodeSpawner interface
+	virtual void Prime() override;
 	virtual UEdGraphNode* Invoke(UEdGraph* ParentGraph, FBindingSet const& Bindings, FVector2D const Location) const override;
 	virtual FText GetDefaultMenuName() const override;
 	virtual FText GetDefaultMenuCategory() const override;
 	virtual FText GetDefaultMenuTooltip() const override;
 	virtual FString GetDefaultSearchKeywords() const override;
+	virtual FName GetDefaultMenuIcon(FLinearColor& ColorOut) const;
 	// End UBlueprintNodeSpawner interface
 
 	// IBlueprintNodeBinder interface
diff --git a/Engine/Source/Editor/BlueprintGraph/Classes/BlueprintVariableNodeSpawner.h b/Engine/Source/Editor/BlueprintGraph/Classes/BlueprintVariableNodeSpawner.h
--- a/Engine/Source/Editor/BlueprintGraph/Classes/BlueprintVariableNodeSpawner.h
+++ b/Engine/Source/Editor/BlueprintGraph/Classes/BlueprintVariableNodeSpawner.h
@@ -46,9 +46,12 @@ public:
 	static UBlueprintVariableNodeSpawner* Create(TSubclassOf<UK2Node_Variable> NodeClass, UEdGraph* VarContext, FBPVariableDescription const& VarDesc, UObject* Outer = nullptr);
 
 	// UBlueprintNodeSpawner interface
+	virtual void Prime() override;
 	virtual UEdGraphNode* Invoke(UEdGraph* ParentGraph, FBindingSet const& Bindings, FVector2D const Location) const override;
 	virtual FText GetDefaultMenuName() const override;
 	virtual FText GetDefaultMenuCategory() const override;
+	virtual FText GetDefaultMenuTooltip() const override;
+	virtual FString GetDefaultSearchKeywords() const override;
 	virtual FName GetDefaultMenuIcon(FLinearColor& ColorOut) const override;
 	// End UBlueprintNodeSpawner interface
 	
@@ -108,7 +111,8 @@ private:
 	UPROPERTY()
 	FBPVariableDescription    LocalVarDesc;
 
-	/** Constructing FText strings can be costly, so we cache the default menu name/category */
+	/** Constructing FText strings can be costly, so we cache the default menu name/category/tooltip */
 	FNodeTextCache CachedMenuName;
 	FNodeTextCache CachedCategory;
+	FNodeTextCache CachedTooltip;
 };
\ No newline at end of file
diff --git a/Engine/Source/Editor/BlueprintGraph/Classes/K2Node_CallFunction.h b/Engine/Source/Editor/BlueprintGraph/Classes/K2Node_CallFunction.h
--- a/Engine/Source/Editor/BlueprintGraph/Classes/K2Node_CallFunction.h
+++ b/Engine/Source/Editor/BlueprintGraph/Classes/K2Node_CallFunction.h
@@ -167,6 +167,9 @@ public:
 	/** Used to determine the result of AllowMultipleSelfs() (without having a node instance) */
 	static bool CanFunctionSupportMultipleTargets(UFunction const* InFunction);
 
+	/** */
+	static FName GetPaletteIconForFunction(UFunction const* Function, FLinearColor& OutColor);
+
 private: 
 
 	/** Helper function to check if the SelfPin is wired correctly */
diff --git a/Engine/Source/Editor/BlueprintGraph/Classes/K2Node_Variable.h b/Engine/Source/Editor/BlueprintGraph/Classes/K2Node_Variable.h
--- a/Engine/Source/Editor/BlueprintGraph/Classes/K2Node_Variable.h
+++ b/Engine/Source/Editor/BlueprintGraph/Classes/K2Node_Variable.h
@@ -5,6 +5,9 @@
 #include "K2Node.h"
 #include "K2Node_Variable.generated.h"
 
+// Forward Declarations
+struct FBPVariableDescription;
+
 UENUM()
 namespace ESelfContextInfo
 {
@@ -138,5 +141,13 @@ public:
 	 * @return A icon representing the specified variable's type.
 	 */
 	static FName GetVarIconFromPinType(const FEdGraphPinType& InPinType, FLinearColor& IconColorOut);
+
+protected:
+	/**
+	 * 
+	 * 
+	 * @return 
+	 */
+	FBPVariableDescription const* GetBlueprintVarDescription() const;
 };
 
diff --git a/Engine/Source/Editor/BlueprintGraph/Classes/K2Node_VariableGet.h b/Engine/Source/Editor/BlueprintGraph/Classes/K2Node_VariableGet.h
--- a/Engine/Source/Editor/BlueprintGraph/Classes/K2Node_VariableGet.h
+++ b/Engine/Source/Editor/BlueprintGraph/Classes/K2Node_VariableGet.h
@@ -23,6 +23,10 @@ class BLUEPRINTGRAPH_API UK2Node_VariableGet : public UK2Node_Variable
 	virtual bool ShouldShowNodeProperties() const override { return true; }
 	virtual class FNodeHandlingFunctor* CreateNodeHandler(class FKismetCompilerContext& CompilerContext) const override;
 	// End K2Node interface
+
+	static FText GetPropertyTooltip(UProperty* VariableProperty);
+	static FText GetBlueprintVarTooltip(FBPVariableDescription const& VarDesc);
+
 private:
 	/** Constructing FText strings can be costly, so we cache the node's title/tooltip */
 	FNodeTextCache CachedTooltip;
diff --git a/Engine/Source/Editor/BlueprintGraph/Classes/K2Node_VariableSet.h b/Engine/Source/Editor/BlueprintGraph/Classes/K2Node_VariableSet.h
--- a/Engine/Source/Editor/BlueprintGraph/Classes/K2Node_VariableSet.h
+++ b/Engine/Source/Editor/BlueprintGraph/Classes/K2Node_VariableSet.h
@@ -28,6 +28,9 @@ class UK2Node_VariableSet : public UK2Node_Variable
 	BLUEPRINTGRAPH_API FName GetRepNotifyName() const;
 	BLUEPRINTGRAPH_API bool ShouldFlushDormancyOnSet() const;
 
+	static FText GetPropertyTooltip(UProperty* VariableProperty);
+	static FText GetBlueprintVarTooltip(FBPVariableDescription const& VarDesc);
+
 private:
 	/** Constructing FText strings can be costly, so we cache the node's title/tooltip */
 	FNodeTextCache CachedTooltip;
diff --git a/Engine/Source/Editor/BlueprintGraph/Private/BlueprintActionDatabase.cpp b/Engine/Source/Editor/BlueprintGraph/Private/BlueprintActionDatabase.cpp
--- a/Engine/Source/Editor/BlueprintGraph/Private/BlueprintActionDatabase.cpp
+++ b/Engine/Source/Editor/BlueprintGraph/Private/BlueprintActionDatabase.cpp
@@ -298,13 +298,6 @@ namespace BlueprintActionDatabaseImpl
 	 */
 	static void GetNodeSpecificActions(TSubclassOf<UEdGraphNode const> const NodeClass, FBlueprintActionDatabaseRegistrar& Registrar);
 
-	/**
-	 * 
-	 * 
-	 * @param  BlueprintAction	
-	 */
-	static void PrimeBlueprintAction(UBlueprintNodeSpawner* BlueprintAction);
-
 	/**
 	 * Callback to refresh the database when a blueprint has been altered 
 	 * (clears the database entries for the blueprint's classes and recreates 
@@ -582,31 +575,6 @@ static void BlueprintActionDatabaseImpl::GetNodeSpecificActions(TSubclassOf<UEdG
 	}
 }
 
-//------------------------------------------------------------------------------
-static void BlueprintActionDatabaseImpl::PrimeBlueprintAction(UBlueprintNodeSpawner* BlueprintAction)
-{
-	// @TODO: for UBlueprintFunctionNodeSpawner, we almost completely don't need
-	//        a template node (we can gather all the info from the UFunction)...
-	//        so we should save on template cache and not do this for function
-	//        node spawners (have to fill the holes first).
-	UEdGraphNode* TemplateNode = BlueprintAction->GetTemplateNode();
-	if (TemplateNode != nullptr)
-	{
-		// since we're priming incrementally, someone could have already
-		// requested this template, and allocated its pins (don't need to do 
-		// redundant work)
-		if (TemplateNode->Pins.Num() == 0)
-		{
-			// in certain scenarios we need pin information from the 
-			// spawner (to help filter by pin context)
-			TemplateNode->AllocateDefaultPins();
-		}
-
-		// cache any FText::Format() calls 
-		TemplateNode->GetNodeTitle(ENodeTitleType::ListView);
-	}
-}
-
 //------------------------------------------------------------------------------
 static void BlueprintActionDatabaseImpl::OnBlueprintChanged(UBlueprint* Blueprint)
 {
@@ -771,7 +739,7 @@ void FBlueprintActionDatabase::Tick(float DeltaTime)
 					FScopedDurationTimer ScopedTimer(TotalFrameDuration);
 
 					UBlueprintNodeSpawner* Action = (*ClassActionList)[ActionListIndex];
-					BlueprintActionDatabaseImpl::PrimeBlueprintAction(Action);
+					Action->Prime();
 				}
 
 				if (ActionListIndex >= ClassActionList->Num())
diff --git a/Engine/Source/Editor/BlueprintGraph/Private/BlueprintFunctionNodeSpawner.cpp b/Engine/Source/Editor/BlueprintGraph/Private/BlueprintFunctionNodeSpawner.cpp
--- a/Engine/Source/Editor/BlueprintGraph/Private/BlueprintFunctionNodeSpawner.cpp
+++ b/Engine/Source/Editor/BlueprintGraph/Private/BlueprintFunctionNodeSpawner.cpp
@@ -192,6 +192,16 @@ UBlueprintFunctionNodeSpawner::UBlueprintFunctionNodeSpawner(class FPostConstruc
 {
 }
 
+//------------------------------------------------------------------------------
+void UBlueprintFunctionNodeSpawner::Prime()
+{
+	// we expect that you don't need a node template to construct menu entries
+	// from this, so we choose not to pre-cache one here
+
+	// we don't have any expensive FText::Format() constructed strings to cache
+	// either
+}
+
 //------------------------------------------------------------------------------
 UEdGraphNode* UBlueprintFunctionNodeSpawner::Invoke(UEdGraph* ParentGraph, FBindingSet const& Bindings, FVector2D const Location) const
 {
@@ -255,6 +265,12 @@ FString UBlueprintFunctionNodeSpawner::GetDefaultSearchKeywords() const
 	return SearchKeywords.AppendChar(TEXT(' '));
 }
 
+//------------------------------------------------------------------------------
+FName UBlueprintFunctionNodeSpawner::GetDefaultMenuIcon(FLinearColor& ColorOut) const
+{
+	return UK2Node_CallFunction::GetPaletteIconForFunction(Function, ColorOut);
+}
+
 //------------------------------------------------------------------------------
 bool UBlueprintFunctionNodeSpawner::CanBindMultipleObjects() const
 {
diff --git a/Engine/Source/Editor/BlueprintGraph/Private/BlueprintNodeSpawner.cpp b/Engine/Source/Editor/BlueprintGraph/Private/BlueprintNodeSpawner.cpp
--- a/Engine/Source/Editor/BlueprintGraph/Private/BlueprintNodeSpawner.cpp
+++ b/Engine/Source/Editor/BlueprintGraph/Private/BlueprintNodeSpawner.cpp
@@ -68,6 +68,37 @@ UBlueprintNodeSpawner::~UBlueprintNodeSpawner()
 	}
 }
 
+//------------------------------------------------------------------------------
+void UBlueprintNodeSpawner::Prime()
+{
+	if (UEdGraphNode* CachedTemplateNode = GetTemplateNode())
+	{
+		// since we're priming incrementally, someone could have already
+		// requested this template, and allocated its pins (don't need to do 
+		// redundant work)
+		if (CachedTemplateNode->Pins.Num() == 0)
+		{
+			// in certain scenarios we need pin information from the 
+			// spawner (to help filter by pin context)
+			CachedTemplateNode->AllocateDefaultPins();
+		}
+
+		//
+		// let the node cache any FText::Format() operations that may be used...
+		if (UK2Node* K2NodeTemplate = Cast<UK2Node>(CachedTemplateNode))
+		{
+			K2NodeTemplate->GetMenuCategory();
+		}
+		CachedTemplateNode->GetTooltipText();
+		CachedTemplateNode->GetNodeTitle(ENodeTitleType::ListView);
+	}
+
+	// in case any of these cache FText::Format() operations
+	GetDefaultMenuName();
+	GetDefaultMenuCategory();
+	GetDefaultMenuTooltip();
+}
+
 //------------------------------------------------------------------------------
 UEdGraphNode* UBlueprintNodeSpawner::Invoke(UEdGraph* ParentGraph, FBindingSet const& Bindings, FVector2D const Location) const
 {
@@ -134,6 +165,12 @@ UEdGraphNode* UBlueprintNodeSpawner::GetTemplateNode(UEdGraph* Outer, FBindingSe
 	return TemplateNode; 
 }
 
+//------------------------------------------------------------------------------
+UEdGraphNode* UBlueprintNodeSpawner::GetTemplateNode(ENoInit) const
+{
+	return BlueprintNodeSpawnerImpl::GetSharedTemplateCache()->GetNodeTemplate(this, NoInit);
+}
+
 //------------------------------------------------------------------------------
 UEdGraphNode* UBlueprintNodeSpawner::Invoke(UEdGraph* ParentGraph, FBindingSet const& Bindings, FVector2D const Location, FCustomizeNodeDelegate PostSpawnDelegate) const
 {
diff --git a/Engine/Source/Editor/BlueprintGraph/Private/BlueprintNodeTemplateCache.cpp b/Engine/Source/Editor/BlueprintGraph/Private/BlueprintNodeTemplateCache.cpp
--- a/Engine/Source/Editor/BlueprintGraph/Private/BlueprintNodeTemplateCache.cpp
+++ b/Engine/Source/Editor/BlueprintGraph/Private/BlueprintNodeTemplateCache.cpp
@@ -212,6 +212,17 @@ UEdGraphNode* FBlueprintNodeTemplateCache::GetNodeTemplate(UBlueprintNodeSpawner
 	return TemplateNode;
 }
 
+//------------------------------------------------------------------------------
+UEdGraphNode* FBlueprintNodeTemplateCache::GetNodeTemplate(UBlueprintNodeSpawner const* NodeSpawner, ENoInit) const
+{
+	UEdGraphNode* TemplateNode = nullptr;
+	if (UEdGraphNode* const* FoundNode = NodeTemplateCache.Find(NodeSpawner))
+	{
+		return *FoundNode;
+	}
+	return nullptr;
+}
+
 //------------------------------------------------------------------------------
 void FBlueprintNodeTemplateCache::ClearCachedTemplate(UBlueprintNodeSpawner const* NodeSpawner)
 {
diff --git a/Engine/Source/Editor/BlueprintGraph/Private/BlueprintVariableNodeSpawner.cpp b/Engine/Source/Editor/BlueprintGraph/Private/BlueprintVariableNodeSpawner.cpp
--- a/Engine/Source/Editor/BlueprintGraph/Private/BlueprintVariableNodeSpawner.cpp
+++ b/Engine/Source/Editor/BlueprintGraph/Private/BlueprintVariableNodeSpawner.cpp
@@ -52,6 +52,19 @@ UBlueprintVariableNodeSpawner::UBlueprintVariableNodeSpawner(class FPostConstruc
 {
 }
 
+//------------------------------------------------------------------------------
+void UBlueprintVariableNodeSpawner::Prime()
+{
+	// we expect that you don't need a node template to construct menu entries
+	// from this, so we choose not to pre-cache one here
+
+	// all of these perform expensive FText::Format() operations and cache the 
+	// results...
+	GetDefaultMenuName();
+	GetDefaultMenuCategory();
+	GetDefaultMenuTooltip();
+}
+
 //------------------------------------------------------------------------------
 UEdGraphNode* UBlueprintVariableNodeSpawner::Invoke(UEdGraph* ParentGraph, FBindingSet const& Bindings, FVector2D const Location) const
 {
@@ -131,6 +144,47 @@ FText UBlueprintVariableNodeSpawner::GetDefaultMenuCategory() const
 	return CachedCategory;
 }
 
+//------------------------------------------------------------------------------
+FText UBlueprintVariableNodeSpawner::GetDefaultMenuTooltip() const
+{
+	if (CachedTooltip.IsOutOfDate())
+	{
+		if (NodeClass->IsChildOf<UK2Node_VariableSet>())
+		{
+			if (IsLocalVariable())
+			{
+				CachedTooltip = UK2Node_VariableSet::GetBlueprintVarTooltip(LocalVarDesc);
+			}
+			else
+			{
+				CachedTooltip = UK2Node_VariableSet::GetPropertyTooltip(MemberVariable.Get());
+			}
+		}
+		else if (NodeClass->IsChildOf<UK2Node_VariableGet>())
+		{
+			if (IsLocalVariable())
+			{
+				CachedTooltip = UK2Node_VariableGet::GetBlueprintVarTooltip(LocalVarDesc);
+			}
+			else
+			{
+				CachedTooltip = UK2Node_VariableGet::GetPropertyTooltip(MemberVariable.Get());
+			}
+		}
+	}	
+	return CachedTooltip;
+}
+
+//------------------------------------------------------------------------------
+FString UBlueprintVariableNodeSpawner::GetDefaultSearchKeywords() const
+{
+	// @TODO: maybe UPROPERTY() fields should have keyword metadata like functions
+	// 
+	// add at least one character, so that the menu item doesn't attempt to
+	// ping a node template 	
+	return TEXT(" ");
+}
+
 //------------------------------------------------------------------------------
 FName UBlueprintVariableNodeSpawner::GetDefaultMenuIcon(FLinearColor& ColorOut) const
 {
diff --git a/Engine/Source/Editor/BlueprintGraph/Private/K2Node_CallFunction.cpp b/Engine/Source/Editor/BlueprintGraph/Private/K2Node_CallFunction.cpp
--- a/Engine/Source/Editor/BlueprintGraph/Private/K2Node_CallFunction.cpp
+++ b/Engine/Source/Editor/BlueprintGraph/Private/K2Node_CallFunction.cpp
@@ -743,14 +743,44 @@ bool UK2Node_CallFunction::CanFunctionSupportMultipleTargets(UFunction const* Fu
 	return !bHasReturnParam && bIsImpure && !bIsLatent;
 }
 
-FLinearColor UK2Node_CallFunction::GetNodeTitleColor() const
+static FLinearColor GetPalletteIconColor(UFunction const* Function)
 {
-	if (IsNodePure())
+	bool const bIsPure = (Function != nullptr) && Function->HasAnyFunctionFlags(FUNC_BlueprintPure);
+	if (bIsPure)
 	{
 		return GetDefault<UGraphEditorSettings>()->PureFunctionCallNodeTitleColor;
 	}
+	return GetDefault<UGraphEditorSettings>()->FunctionCallNodeTitleColor;
+}
+
+FName UK2Node_CallFunction::GetPaletteIconForFunction(UFunction const* Function, FLinearColor& OutColor)
+{
+	static const FName NativeMakeFunc(TEXT("NativeMakeFunc"));
+	static const FName NativeBrakeFunc(TEXT("NativeBreakFunc"));
 
-	return Super::GetNodeTitleColor();
+	if (Function && Function->HasMetaData(NativeMakeFunc))
+	{
+		return TEXT("GraphEditor.MakeStruct_16x");
+	}
+	else if (Function && Function->HasMetaData(NativeBrakeFunc))
+	{
+		return TEXT("GraphEditor.BreakStruct_16x");
+	}
+	// Check to see if the function is calling an function that could be an event, display the event icon instead.
+	else if (Function && UEdGraphSchema_K2::FunctionCanBePlacedAsEvent(Function))
+	{
+		return TEXT("GraphEditor.Event_16x");
+	}
+	else
+	{
+		OutColor = GetPalletteIconColor(Function);
+		return TEXT("Kismet.AllClasses.FunctionIcon");
+	}
+}
+
+FLinearColor UK2Node_CallFunction::GetNodeTitleColor() const
+{
+	return GetPalletteIconColor(GetTargetFunction());
 }
 
 FText UK2Node_CallFunction::GetTooltipText() const
@@ -1465,28 +1495,7 @@ FName UK2Node_CallFunction::GetCornerIcon() const
 
 FName UK2Node_CallFunction::GetPaletteIcon(FLinearColor& OutColor) const
 {
-	static const FName NativeMakeFunc(TEXT("NativeMakeFunc"));
-	static const FName NativeBrakeFunc(TEXT("NativeBreakFunc"));
-
-	auto Function = GetTargetFunction();
-	if (Function && Function->HasMetaData(NativeMakeFunc))
-	{
-		return TEXT("GraphEditor.MakeStruct_16x");
-	}
-	else if (Function && Function->HasMetaData(NativeBrakeFunc))
-	{
-		return TEXT("GraphEditor.BreakStruct_16x");
-	}
-	// Check to see if the function is calling an function that could be an event, display the event icon instead.
-	else if (Function && UEdGraphSchema_K2::FunctionCanBePlacedAsEvent(Function))
-	{
-		return TEXT("GraphEditor.Event_16x");
-	}
-	else
-	{
-		OutColor = GetNodeTitleColor();
-		return TEXT("Kismet.AllClasses.FunctionIcon");
-	}
+	return GetPaletteIconForFunction(GetTargetFunction(), OutColor);
 }
 
 FText UK2Node_CallFunction::GetToolTipHeading() const
diff --git a/Engine/Source/Editor/BlueprintGraph/Private/K2Node_Variable.cpp b/Engine/Source/Editor/BlueprintGraph/Private/K2Node_Variable.cpp
--- a/Engine/Source/Editor/BlueprintGraph/Private/K2Node_Variable.cpp
+++ b/Engine/Source/Editor/BlueprintGraph/Private/K2Node_Variable.cpp
@@ -516,4 +516,28 @@ FString UK2Node_Variable::GetDocumentationExcerptName() const
 	return GetVarName().ToString();
 }
 
+FBPVariableDescription const* UK2Node_Variable::GetBlueprintVarDescription() const
+{
+	FName const& VarName = VariableReference.GetMemberName();
+	UStruct const* VariableScope = VariableReference.GetMemberScope(this);
+
+	bool const bIsLocalVariable = (VariableScope != nullptr);
+	if (bIsLocalVariable)
+	{
+		return FBlueprintEditorUtils::FindLocalVariable(GetBlueprint(), VariableScope, VarName);
+	}
+	else if (UProperty const* VarProperty = GetPropertyForVariable())
+	{
+		UClass const* SourceClass = VarProperty->GetOwnerClass();
+		UBlueprint const* SourceBlueprint = (SourceClass != nullptr) ? Cast<UBlueprint>(SourceClass->ClassGeneratedBy) : nullptr;
+
+		if (SourceBlueprint != nullptr)
+		{
+			int32 const VarIndex = FBlueprintEditorUtils::FindNewVariableIndex(SourceBlueprint, VarName);
+			return &SourceBlueprint->NewVariables[VarIndex];
+		}
+	}
+	return nullptr;
+}
+
 #undef LOCTEXT_NAMESPACE
diff --git a/Engine/Source/Editor/BlueprintGraph/Private/K2Node_VariableGet.cpp b/Engine/Source/Editor/BlueprintGraph/Private/K2Node_VariableGet.cpp
--- a/Engine/Source/Editor/BlueprintGraph/Private/K2Node_VariableGet.cpp
+++ b/Engine/Source/Editor/BlueprintGraph/Private/K2Node_VariableGet.cpp
@@ -52,6 +52,26 @@ public:
 	}
 };
 
+namespace K2Node_VariableGetImpl
+{
+	/**
+	 * Shared utility method for retrieving a UK2Node_VariableGet's bare tooltip.
+	 * 
+	 * @param  VarName	The name of the variable that the node represents.
+	 * @return A formatted text string, describing what the VariableGet node does.
+	 */
+	static FText GetBaseTooltip(FName VarName);
+}
+
+static FText K2Node_VariableGetImpl::GetBaseTooltip(FName VarName)
+{
+	FFormatNamedArguments Args;
+	Args.Add(TEXT("VarName"), FText::FromName(VarName));
+
+	return FText::Format(LOCTEXT("GetVariableTooltip", "Read the value of variable {VarName}"), Args);
+
+}
+
 UK2Node_VariableGet::UK2Node_VariableGet(const class FPostConstructInitializeProperties& PCIP)
 	: Super(PCIP)
 {
@@ -85,51 +105,87 @@ void UK2Node_VariableGet::ReallocatePinsDuringReconstruction(TArray<UEdGraphPin*
 	}
 }
 
-FText UK2Node_VariableGet::GetTooltipText() const
+FText UK2Node_VariableGet::GetPropertyTooltip(UProperty* VariableProperty)
 {
-	if (!CachedTooltip.IsOutOfDate())
+	FName VarName = NAME_None;
+	if (VariableProperty != nullptr)
 	{
-		return CachedTooltip;
-	}
+		VarName = VariableProperty->GetFName();
 
-	FFormatNamedArguments Args;
-	Args.Add( TEXT( "VarName" ), FText::FromString( GetVarNameString() ));
-	Args.Add( TEXT( "TextPartition" ), FText::GetEmpty());
-	Args.Add( TEXT( "MetaData" ), FText::GetEmpty());
+		UClass* SourceClass = VariableProperty->GetOwnerClass();
+		// discover if the variable property is a non blueprint user variable
+		bool const bIsNativeVariable = (SourceClass != nullptr) && (SourceClass->ClassGeneratedBy == nullptr);
+		FName const TooltipMetaKey(TEXT("tooltip"));
 
-	FName VarName = VariableReference.GetMemberName();
-	if (VarName != NAME_None)
-	{
-		FString BPMetaData;
-		FBlueprintEditorUtils::GetBlueprintVariableMetaData(GetBlueprint(), VarName, VariableReference.GetMemberScope(this), TEXT("tooltip"), BPMetaData);
+		FText SubTooltip;
+		if (bIsNativeVariable)
+		{
+			FText const PropertyTooltip = VariableProperty->GetToolTipText();
+			if (!PropertyTooltip.IsEmpty())
+			{
+				// See if the native property has a tooltip
+				SubTooltip = PropertyTooltip;
+				FString TooltipName = FString::Printf(TEXT("%s.%s"), *VarName.ToString(), *TooltipMetaKey.ToString());
+				FText::FindText(*VariableProperty->GetFullGroupName(true), *TooltipName, SubTooltip);
+			}
+		}
+		else if (UBlueprint* VarBlueprint = Cast<UBlueprint>(SourceClass->ClassGeneratedBy))
+		{
+			FString UserTooltipData;
+			if (FBlueprintEditorUtils::GetBlueprintVariableMetaData(VarBlueprint, VarName, /*InLocalVarScope =*/nullptr, TooltipMetaKey, UserTooltipData))
+			{
+				SubTooltip = FText::FromString(UserTooltipData);
+			}
+		}
 
-		if( !BPMetaData.IsEmpty() )
+		if (!SubTooltip.IsEmpty())
 		{
-			Args.Add( TEXT( "TextPartition" ), FText::FromString( "\n" ));
-			Args.Add( TEXT( "MetaData" ), FText::FromString( BPMetaData ));
+			FFormatNamedArguments Args;
+			Args.Add(TEXT("VarName"), FText::FromName(VarName));
+			Args.Add(TEXT("PropertyTooltip"), SubTooltip);
+
+			return FText::Format(LOCTEXT("GetVariableProperty_Tooltip", "Read the value of variable {VarName}\n{PropertyTooltip}"), Args);
 		}
 	}
-	if(  UProperty* Property = GetPropertyForVariable() )
+	return K2Node_VariableGetImpl::GetBaseTooltip(VarName);
+}
+
+FText UK2Node_VariableGet::GetBlueprintVarTooltip(FBPVariableDescription const& VarDesc)
+{
+	FName const TooltipMetaKey(TEXT("tooltip"));
+	int32 const MetaIndex = VarDesc.FindMetaDataEntryIndexForKey(TooltipMetaKey);
+	bool const bHasTooltipData = (MetaIndex != INDEX_NONE);
+
+	if (bHasTooltipData)
 	{
-		// discover if the variable property is a non blueprint user variable
-		UClass* SourceClass = Property->GetOwnerClass();
-		if( SourceClass && SourceClass->ClassGeneratedBy == NULL )
-		{
-			const FString MetaData = Property->GetToolTipText().ToString();
+		FString UserTooltipData = VarDesc.GetMetaData(TooltipMetaKey);
 
-			if( !MetaData.IsEmpty() )
-			{
-				// See if the property associated with this editor has a tooltip
-				FText PropertyMetaData = FText::FromString( *MetaData );
-				FString TooltipName = FString::Printf( TEXT("%s.tooltip"), *(Property->GetName()));
-				FText::FindText( *(Property->GetFullGroupName(true)), *TooltipName, PropertyMetaData );
-				Args.Add( TEXT( "TextPartition" ), FText::FromString( "\n" ));
-				Args.Add( TEXT( "MetaData" ), PropertyMetaData );
-			}
+		FFormatNamedArguments Args;
+		Args.Add(TEXT("VarName"), FText::FromName(VarDesc.VarName));
+		Args.Add(TEXT("UserTooltip"), FText::FromString(UserTooltipData));
+
+		return FText::Format(LOCTEXT("GetVariableProperty_Tooltip", "Read the value of variable {VarName}\n{UserTooltip}"), Args);
+	}
+	return K2Node_VariableGetImpl::GetBaseTooltip(VarDesc.VarName);
+}
+
+FText UK2Node_VariableGet::GetTooltipText() const
+{
+	if (!CachedTooltip.IsOutOfDate())
+	{
+		if (UProperty* Property = GetPropertyForVariable())
+		{
+			CachedTooltip = GetPropertyTooltip(Property);
+		}
+		else if (FBPVariableDescription const* VarDesc = GetBlueprintVarDescription())
+		{
+			CachedTooltip = GetBlueprintVarTooltip(*VarDesc);
+		}
+		else
+		{
+			CachedTooltip = K2Node_VariableGetImpl::GetBaseTooltip(GetVarName());
 		}
 	}
-	// FText::Format() is slow, so we cache this to save on performance
-	CachedTooltip = FText::Format(NSLOCTEXT("K2Node", "GetVariable_ToolTip", "Read the value of variable {VarName}{TextPartition}{MetaData}"), Args);
 	return CachedTooltip;
 }
 
diff --git a/Engine/Source/Editor/BlueprintGraph/Private/K2Node_VariableSet.cpp b/Engine/Source/Editor/BlueprintGraph/Private/K2Node_VariableSet.cpp
--- a/Engine/Source/Editor/BlueprintGraph/Private/K2Node_VariableSet.cpp
+++ b/Engine/Source/Editor/BlueprintGraph/Private/K2Node_VariableSet.cpp
@@ -5,6 +5,69 @@
 #include "KismetCompiler.h"
 #include "VariableSetHandler.h"
 
+#define LOCTEXT_NAMESPACE "K2Node_VariableSet"
+
+namespace K2Node_VariableSetImpl
+{
+	/**
+	 * Shared utility method for retrieving a UK2Node_VariableSet's bare tooltip.
+	 * 
+	 * @param  VarName	The name of the variable that the node represents.
+	 * @return A formatted text string, describing what the VariableSet node does.
+	 */
+	static FText GetBaseTooltip(FName VarName);
+
+	/**
+	 * Returns true if the specified variable RepNotify AND is defined in a 
+	 * blueprint. Most (all?) native rep notifies are intended to be client 
+	 * only. We are moving away from this paradigm in blueprints. So for now 
+	 * this is somewhat of a hold over to avoid nasty bugs where a K2 set node 
+	 * is calling a native function that the designer has no idea what it is 
+	 * doing.
+	 * 
+	 * @param  VariableProperty	The variable property you wish to check.
+	 * @return True if the specified variable RepNotify AND is defined in a blueprint.
+	 */
+	static bool PropertyHasLocalRepNotify(UProperty* VariableProperty);
+}
+
+static FText K2Node_VariableSetImpl::GetBaseTooltip(FName VarName)
+{
+	FFormatNamedArguments Args;
+	Args.Add(TEXT("VarName"), FText::FromName(VarName));
+
+	return FText::Format(LOCTEXT("SetVariableTooltip", "Set the value of variable {VarName}"), Args);
+
+}
+
+static bool K2Node_VariableSetImpl::PropertyHasLocalRepNotify(UProperty* VariableProperty)
+{
+	if (VariableProperty != nullptr)
+	{
+		// We check that the variable is 'defined in a blueprint' so as to avoid 
+		// natively defined RepNotifies being called unintentionally. Most(all?) 
+		// native rep notifies are intended to be client only. We are moving 
+		// away from this paradigm in blueprints. So for now this is somewhat of 
+		// a hold over to avoid nasty bugs where a K2 set node is calling a 
+		// native function that the designer has no idea what it is doing.
+		UBlueprintGeneratedClass* VariableSourceClass = Cast<UBlueprintGeneratedClass>(VariableProperty->GetOwnerClass());
+		bool const bIsBlueprintProperty = (VariableSourceClass != nullptr);
+
+		if (bIsBlueprintProperty && (VariableProperty->RepNotifyFunc != NAME_None))
+		{
+			// Find function (ok if its defined in native class)
+			UFunction* Function = VariableSourceClass->FindFunctionByName(VariableProperty->RepNotifyFunc);
+
+			// If valid repnotify func
+			if ((Function != nullptr) && (Function->NumParms == 0) && (Function->GetReturnProperty() == nullptr))
+			{
+				return true;
+			}
+		}
+	}
+	return false;
+}
+
 UK2Node_VariableSet::UK2Node_VariableSet(const class FPostConstructInitializeProperties& PCIP)
 	: Super(PCIP)
 {
@@ -48,12 +111,115 @@ void UK2Node_VariableSet::ReallocatePinsDuringReconstruction(TArray<UEdGraphPin*
 	}
 }
 
+
+
+FText UK2Node_VariableSet::GetPropertyTooltip(UProperty* VariableProperty)
+{
+	FText TextFormat;
+	FFormatNamedArguments Args;
+
+	bool const bHasLocalRepNotify = K2Node_VariableSetImpl::PropertyHasLocalRepNotify(VariableProperty);
+	if (bHasLocalRepNotify)
+	{
+		Args.Add(TEXT("ReplicationNotifyName"), FText::FromName(VariableProperty->RepNotifyFunc));
+		TextFormat = LOCTEXT("SetVariableWithRepNotify_Tooltip", "Set the value of variable {VarName} and call {ReplicationNotifyName}");
+	}
+
+	FName VarName = NAME_None;
+	if (VariableProperty != nullptr)
+	{
+		VarName = VariableProperty->GetFName();
+
+		UClass* SourceClass = VariableProperty->GetOwnerClass();
+		// discover if the variable property is a non blueprint user variable
+		bool const bIsNativeVariable = (SourceClass != nullptr) && (SourceClass->ClassGeneratedBy == nullptr);
+		FName const TooltipMetaKey(TEXT("tooltip"));
+
+		FText SubTooltip;
+		if (bIsNativeVariable)
+		{
+			FText const PropertyTooltip = VariableProperty->GetToolTipText();
+			if (!PropertyTooltip.IsEmpty())
+			{
+				// See if the native property has a tooltip
+				SubTooltip = PropertyTooltip;
+				FString TooltipName = FString::Printf(TEXT("%s.%s"), *VarName.ToString(), *TooltipMetaKey.ToString());
+				FText::FindText(*VariableProperty->GetFullGroupName(true), *TooltipName, SubTooltip);
+			}
+		}
+		else if (UBlueprint* VarBlueprint = Cast<UBlueprint>(SourceClass->ClassGeneratedBy))
+		{
+			FString UserTooltipData;
+			if (FBlueprintEditorUtils::GetBlueprintVariableMetaData(VarBlueprint, VarName, /*InLocalVarScope =*/nullptr, TooltipMetaKey, UserTooltipData))
+			{
+				SubTooltip = FText::FromString(UserTooltipData);
+			}
+		}
+
+		if (!SubTooltip.IsEmpty())
+		{
+			Args.Add(TEXT("PropertyTooltip"), SubTooltip);
+			if (bHasLocalRepNotify)
+			{
+				TextFormat = LOCTEXT("SetVariablePropertyWithRepNotify_Tooltip", "Set the value of variable {VarName} and call {ReplicationNotifyName}\n{PropertyTooltip}");
+			}
+			else
+			{
+				TextFormat = LOCTEXT("SetVariableProperty_Tooltip", "Set the value of variable {VarName}\n{PropertyTooltip}");
+			}
+		}
+	}
+
+	if (TextFormat.IsEmpty())
+	{
+		return K2Node_VariableSetImpl::GetBaseTooltip(VarName);
+	}
+	else
+	{
+		Args.Add(TEXT("VarName"), FText::FromName(VarName));
+		return FText::Format(TextFormat, Args);
+	}
+}
+
+FText UK2Node_VariableSet::GetBlueprintVarTooltip(FBPVariableDescription const& VarDesc)
+{
+	FName const TooltipMetaKey(TEXT("tooltip"));
+	int32 const MetaIndex = VarDesc.FindMetaDataEntryIndexForKey(TooltipMetaKey);
+	bool const bHasTooltipData = (MetaIndex != INDEX_NONE);
+
+	if (bHasTooltipData)
+	{
+		FString UserTooltipData = VarDesc.GetMetaData(TooltipMetaKey);
+
+		FFormatNamedArguments Args;
+		Args.Add(TEXT("VarName"), FText::FromName(VarDesc.VarName));
+		Args.Add(TEXT("UserTooltip"), FText::FromString(UserTooltipData));
+
+		return FText::Format(LOCTEXT("SetVariableProperty_Tooltip", "Set the value of variable {VarName}\n{UserTooltip}"), Args);
+	}
+	return K2Node_VariableSetImpl::GetBaseTooltip(VarDesc.VarName);
+}
+
 FText UK2Node_VariableSet::GetTooltipText() const
 {
-	if (!CachedTooltip.IsOutOfDate())
+	//if (!CachedTooltip.IsOutOfDate())
 	{
-		return CachedTooltip;
+		if (UProperty* Property = GetPropertyForVariable())
+		{
+			CachedTooltip = GetPropertyTooltip(Property);
+		}
+		else if (FBPVariableDescription const* VarDesc = GetBlueprintVarDescription())
+		{
+			CachedTooltip = GetBlueprintVarTooltip(*VarDesc);
+		}
+		else
+		{
+			CachedTooltip = K2Node_VariableSetImpl::GetBaseTooltip(GetVarName());
+		}
 	}
+	return CachedTooltip;
+
+	
 
 	FFormatNamedArguments Args;
 	Args.Add( TEXT( "VarName" ), FText::FromString( GetVarNameString() ));
@@ -139,38 +305,7 @@ FText UK2Node_VariableSet::GetNodeTitle(ENodeTitleType::Type TitleType) const
  */
 bool UK2Node_VariableSet::HasLocalRepNotify() const
 {
-	UProperty *Property = NULL;
-	{
-		// Find the property, but only look in blueprint classes for it
-		UBlueprintGeneratedClass *SearchClass = Cast<UBlueprintGeneratedClass>(GetVariableSourceClass());
-		while(Property == NULL && SearchClass != NULL)
-		{
-			for (TFieldIterator<UProperty> It(SearchClass, EFieldIteratorFlags::ExcludeSuper); It; ++It)
-			{
-				if (It->GetFName() == GetVarName())
-				{
-					Property = *It;
-					break;
-				}
-			}
-			SearchClass = Cast<UBlueprintGeneratedClass>(SearchClass->GetSuperStruct());
-		}
-	}
-
-	// If valid property with RepNotify
-	if (Property != NULL && Property->RepNotifyFunc != NAME_None)
-	{
-		// Find function (ok if its defined in native class)
-		UFunction* Function = GetVariableSourceClass()->FindFunctionByName(Property->RepNotifyFunc);
-
-		// If valid repnotify func
-		if( Function != NULL && Function->NumParms == 0 && Function->GetReturnProperty() == NULL )
-		{
-			return true;
-		}
-	}
-
-	return false;
+	return K2Node_VariableSetImpl::PropertyHasLocalRepNotify(GetPropertyForVariable());
 }
 
 bool UK2Node_VariableSet::ShouldFlushDormancyOnSet() const
@@ -200,3 +335,5 @@ FNodeHandlingFunctor* UK2Node_VariableSet::CreateNodeHandler(FKismetCompilerCont
 {
 	return new FKCHandler_VariableSet(CompilerContext);
 }
+
+#undef LOCTEXT_NAMESPACE
diff --git a/Engine/Source/Editor/BlueprintGraph/Public/BlueprintNodeSpawner.h b/Engine/Source/Editor/BlueprintGraph/Public/BlueprintNodeSpawner.h
--- a/Engine/Source/Editor/BlueprintGraph/Public/BlueprintNodeSpawner.h
+++ b/Engine/Source/Editor/BlueprintGraph/Public/BlueprintNodeSpawner.h
@@ -8,6 +8,7 @@
 // Forward declarations
 class UEdGraph;
 class UEdGraphNode;
+class FBlueprintNodeTemplateCache;
 
 /*******************************************************************************
 * UBlueprintNodeSpawner Declaration
@@ -71,6 +72,13 @@ public:
 	 */
 	FCustomizeNodeDelegate CustomizeNodeDelegate;
 
+	/**
+	 * Not required, but intended to passively help speed up menu building 
+	 * operations. Will cache a node-template (via GetTemplateNode), along with  
+	 * any expensive text strings, to avoid constructing them all on demand.
+	 */
+	virtual void Prime();
+
 	/**
 	 * Takes care of spawning a node for the specified graph. Looks to see if 
 	 * the supplied graph is transient, and if so, spawns a NOT fully formed
@@ -162,6 +170,14 @@ public:
 	 */
 	UEdGraphNode* GetTemplateNode(UEdGraph* TargetGraph = nullptr, FBindingSet const& Bindings = FBindingSet()) const;
 
+	/**
+	 * Retrieves a cached template for the node that this is set to spawn. Will
+	 * NOT spawn one if it is not already cached.
+	 *
+	 * @return The cached template-node (if one already exists for this spawner).
+	 */
+	UEdGraphNode* GetTemplateNode(ENoInit) const;
+
 	// IBlueprintNodeBinder interface
 	virtual bool   IsBindingCompatible(UObject const* BindingCandidate) const override { return false; }
 	virtual bool   CanBindMultipleObjects() const override { return false; }
diff --git a/Engine/Source/Editor/BlueprintGraph/Public/BlueprintNodeTemplateCache.h b/Engine/Source/Editor/BlueprintGraph/Public/BlueprintNodeTemplateCache.h
--- a/Engine/Source/Editor/BlueprintGraph/Public/BlueprintNodeTemplateCache.h
+++ b/Engine/Source/Editor/BlueprintGraph/Public/BlueprintNodeTemplateCache.h
@@ -35,6 +35,16 @@ public:
 	 */
 	UEdGraphNode* GetNodeTemplate(UBlueprintNodeSpawner const* NodeSpawner, UEdGraph* TargetGraph = nullptr);
 
+	/**
+	 * Retrieves a cached template associated with the supplied spawner. Does 
+	 * NOT attempt to allocate one if it doesn't exist.
+	 * 
+	 * @param  NodeSpawner	The spawner you want a template node for.
+	 * @param  ENoInit Signifies to use this function over the other (mutating) GetNodeTemplate().
+	 * @return Should return the cached template-node (if one exists, otherwise false).
+	 */
+	UEdGraphNode* GetNodeTemplate(UBlueprintNodeSpawner const* NodeSpawner, ENoInit) const;
+
 	/**
 	 * Wipes any nodes that were cached on behalf of the specified spawner 
 	 * (should be called when NodeSpawner is destroyed, in case 
