Adding metadata to nodes that helps track which ones were systematically added as "default" nodes.

diff --git a/Engine/Source/Editor/AnimGraph/Private/AnimationConduitGraphSchema.cpp b/Engine/Source/Editor/AnimGraph/Private/AnimationConduitGraphSchema.cpp
--- a/Engine/Source/Editor/AnimGraph/Private/AnimationConduitGraphSchema.cpp
+++ b/Engine/Source/Editor/AnimGraph/Private/AnimationConduitGraphSchema.cpp
@@ -18,6 +18,7 @@ void UAnimationConduitGraphSchema::CreateDefaultNodesForGraph(UEdGraph& Graph) c
 {
 	FGraphNodeCreator<UAnimGraphNode_TransitionResult> NodeCreator(Graph);
 	UAnimGraphNode_TransitionResult* ResultSinkNode = NodeCreator.CreateNode();
+	SetNodeMetaData(ResultSinkNode, FNodeMetadata::DefaultGraphNode);
 
 	UAnimationTransitionGraph* TypedGraph = CastChecked<UAnimationTransitionGraph>(&Graph);
 	TypedGraph->MyResultNode = ResultSinkNode;
diff --git a/Engine/Source/Editor/AnimGraph/Private/AnimationCustomTransitionSchema.cpp b/Engine/Source/Editor/AnimGraph/Private/AnimationCustomTransitionSchema.cpp
--- a/Engine/Source/Editor/AnimGraph/Private/AnimationCustomTransitionSchema.cpp
+++ b/Engine/Source/Editor/AnimGraph/Private/AnimationCustomTransitionSchema.cpp
@@ -23,6 +23,7 @@ void UAnimationCustomTransitionSchema::CreateDefaultNodesForGraph(UEdGraph& Grap
 	ResultSinkNode->NodePosX = 0;
 	ResultSinkNode->NodePosY = 0;
 	ResultNodeCreator.Finalize();
+	SetNodeMetaData(ResultSinkNode, FNodeMetadata::DefaultGraphNode);
 
 	UAnimationCustomTransitionGraph* TypedGraph = CastChecked<UAnimationCustomTransitionGraph>(&Graph);
 	TypedGraph->MyResultNode = ResultSinkNode;
@@ -35,6 +36,7 @@ void UAnimationCustomTransitionSchema::CreateDefaultNodesForGraph(UEdGraph& Grap
 		SourcePoseNode->NodePosX = -300;
 		SourcePoseNode->NodePosY = -150;
 		SourceNodeCreator.Finalize();
+		SetNodeMetaData(SourcePoseNode, FNodeMetadata::DefaultGraphNode);
 
 	}
 	{
@@ -44,6 +46,7 @@ void UAnimationCustomTransitionSchema::CreateDefaultNodesForGraph(UEdGraph& Grap
 		DestinationPoseNode->NodePosX = -300;
 		DestinationPoseNode->NodePosY = 150;
 		DestinationNodeCreator.Finalize();
+		SetNodeMetaData(DestinationPoseNode, FNodeMetadata::DefaultGraphNode);
 	}
 }
 
diff --git a/Engine/Source/Editor/AnimGraph/Private/AnimationGraphSchema.cpp b/Engine/Source/Editor/AnimGraph/Private/AnimationGraphSchema.cpp
--- a/Engine/Source/Editor/AnimGraph/Private/AnimationGraphSchema.cpp
+++ b/Engine/Source/Editor/AnimGraph/Private/AnimationGraphSchema.cpp
@@ -118,6 +118,7 @@ void UAnimationGraphSchema::CreateDefaultNodesForGraph(UEdGraph& Graph) const
 	FGraphNodeCreator<UAnimGraphNode_Root> NodeCreator(Graph);
 	UAnimGraphNode_Root* ResultSinkNode = NodeCreator.CreateNode();
 	NodeCreator.Finalize();
+	SetNodeMetaData(ResultSinkNode, FNodeMetadata::DefaultGraphNode);
 }
 
 void UAnimationGraphSchema::HandleGraphBeingDeleted(UEdGraph& GraphBeingRemoved) const
diff --git a/Engine/Source/Editor/AnimGraph/Private/AnimationStateGraphSchema.cpp b/Engine/Source/Editor/AnimGraph/Private/AnimationStateGraphSchema.cpp
--- a/Engine/Source/Editor/AnimGraph/Private/AnimationStateGraphSchema.cpp
+++ b/Engine/Source/Editor/AnimGraph/Private/AnimationStateGraphSchema.cpp
@@ -26,6 +26,7 @@ void UAnimationStateGraphSchema::CreateDefaultNodesForGraph(UEdGraph& Graph) con
 	FGraphNodeCreator<UAnimGraphNode_StateResult> NodeCreator(Graph);
 	UAnimGraphNode_StateResult* ResultSinkNode = NodeCreator.CreateNode();
 	NodeCreator.Finalize();
+	SetNodeMetaData(ResultSinkNode, FNodeMetadata::DefaultGraphNode);
 
 	UAnimationStateGraph* TypedGraph = CastChecked<UAnimationStateGraph>(&Graph);
 	TypedGraph->MyResultNode = ResultSinkNode;
diff --git a/Engine/Source/Editor/AnimGraph/Private/AnimationStateMachineSchema.cpp b/Engine/Source/Editor/AnimGraph/Private/AnimationStateMachineSchema.cpp
--- a/Engine/Source/Editor/AnimGraph/Private/AnimationStateMachineSchema.cpp
+++ b/Engine/Source/Editor/AnimGraph/Private/AnimationStateMachineSchema.cpp
@@ -120,6 +120,7 @@ void UAnimationStateMachineSchema::CreateDefaultNodesForGraph(UEdGraph& Graph) c
 	FGraphNodeCreator<UAnimStateEntryNode> NodeCreator(Graph);
 	UAnimStateEntryNode* EntryNode = NodeCreator.CreateNode();
 	NodeCreator.Finalize();
+	SetNodeMetaData(EntryNode, FNodeMetadata::DefaultGraphNode);
 
 	if (UAnimationStateMachineGraph* StateMachineGraph = CastChecked<UAnimationStateMachineGraph>(&Graph))
 	{
diff --git a/Engine/Source/Editor/AnimGraph/Private/AnimationTransitionSchema.cpp b/Engine/Source/Editor/AnimGraph/Private/AnimationTransitionSchema.cpp
--- a/Engine/Source/Editor/AnimGraph/Private/AnimationTransitionSchema.cpp
+++ b/Engine/Source/Editor/AnimGraph/Private/AnimationTransitionSchema.cpp
@@ -132,6 +132,7 @@ void UAnimationTransitionSchema::CreateDefaultNodesForGraph(UEdGraph& Graph) con
 	FGraphNodeCreator<UAnimGraphNode_TransitionResult> NodeCreator(Graph);
 	UAnimGraphNode_TransitionResult* ResultSinkNode = NodeCreator.CreateNode();
 	NodeCreator.Finalize();
+	SetNodeMetaData(ResultSinkNode, FNodeMetadata::DefaultGraphNode);
 
 	UAnimationTransitionGraph* TypedGraph = CastChecked<UAnimationTransitionGraph>(&Graph);
 	TypedGraph->MyResultNode = ResultSinkNode;
diff --git a/Engine/Source/Editor/BehaviorTreeEditor/Private/EdGraphSchema_BehaviorTree.cpp b/Engine/Source/Editor/BehaviorTreeEditor/Private/EdGraphSchema_BehaviorTree.cpp
--- a/Engine/Source/Editor/BehaviorTreeEditor/Private/EdGraphSchema_BehaviorTree.cpp
+++ b/Engine/Source/Editor/BehaviorTreeEditor/Private/EdGraphSchema_BehaviorTree.cpp
@@ -245,6 +245,7 @@ void UEdGraphSchema_BehaviorTree::CreateDefaultNodesForGraph(UEdGraph& Graph) co
 	FGraphNodeCreator<UBehaviorTreeGraphNode_Root> NodeCreator(Graph);
 	UBehaviorTreeGraphNode_Root* MyNode = NodeCreator.CreateNode();
 	NodeCreator.Finalize();
+	SetNodeMetaData(MyNode, FNodeMetadata::DefaultGraphNode);
 }
 
 FString UEdGraphSchema_BehaviorTree::GetShortTypeName(const UObject* Ob) const
diff --git a/Engine/Source/Editor/BehaviorTreeEditor/Private/EdGraphSchema_BehaviorTreeDecorator.cpp b/Engine/Source/Editor/BehaviorTreeEditor/Private/EdGraphSchema_BehaviorTreeDecorator.cpp
--- a/Engine/Source/Editor/BehaviorTreeEditor/Private/EdGraphSchema_BehaviorTreeDecorator.cpp
+++ b/Engine/Source/Editor/BehaviorTreeEditor/Private/EdGraphSchema_BehaviorTreeDecorator.cpp
@@ -118,7 +118,7 @@ void UEdGraphSchema_BehaviorTreeDecorator::CreateDefaultNodesForGraph(UEdGraph&
 	FGraphNodeCreator<UBehaviorTreeDecoratorGraphNode_Logic> NodeCreator(Graph);
 	UBehaviorTreeDecoratorGraphNode_Logic* MyNode = NodeCreator.CreateNode();
 	MyNode->LogicMode = EDecoratorLogicMode::Sink;
-
+	SetNodeMetaData(MyNode, FNodeMetadata::DefaultGraphNode);
 	NodeCreator.Finalize();
 }
 
diff --git a/Engine/Source/Editor/BlueprintGraph/Private/EdGraphSchema_K2.cpp b/Engine/Source/Editor/BlueprintGraph/Private/EdGraphSchema_K2.cpp
--- a/Engine/Source/Editor/BlueprintGraph/Private/EdGraphSchema_K2.cpp
+++ b/Engine/Source/Editor/BlueprintGraph/Private/EdGraphSchema_K2.cpp
@@ -2333,6 +2333,7 @@ void UEdGraphSchema_K2::CreateMacroGraphTerminators(UEdGraph& Graph, UClass* Cla
 		UK2Node_Tunnel* EntryNode = EntryNodeCreator.CreateNode();
 		EntryNode->bCanHaveOutputs = true;
 		EntryNodeCreator.Finalize();
+		SetNodeMetaData(EntryNode, FNodeMetadata::DefaultGraphNode);
 	}
 
 	{
@@ -2341,6 +2342,7 @@ void UEdGraphSchema_K2::CreateMacroGraphTerminators(UEdGraph& Graph, UClass* Cla
 		ExitNode->bCanHaveInputs = true;
 		ExitNode->NodePosX = 240;
 		ExitNodeCreator.Finalize();
+		SetNodeMetaData(ExitNode, FNodeMetadata::DefaultGraphNode);
 	}
 }
 		
@@ -2357,6 +2359,7 @@ void UEdGraphSchema_K2::CreateFunctionGraphTerminators(UEdGraph& Graph, UClass*
 	EntryNode->SignatureClass = Class;
 	EntryNode->SignatureName = GraphName;
 	FunctionEntryCreator.Finalize();
+	SetNodeMetaData(EntryNode, FNodeMetadata::DefaultGraphNode);
 
 	// See if we need to implement a return node
 	UFunction* InterfaceToImplement = FindField<UFunction>(Class, GraphName);
@@ -2382,6 +2385,7 @@ void UEdGraphSchema_K2::CreateFunctionGraphTerminators(UEdGraph& Graph, UClass*
 			ReturnNode->NodePosX = EntryNode->NodePosX + EntryNode->NodeWidth + 256;
 			ReturnNode->NodePosY = EntryNode->NodePosY;
 			NodeCreator.Finalize();
+			SetNodeMetaData(ReturnNode, FNodeMetadata::DefaultGraphNode);
 
 			// Auto-connect the pins for entry and exit, so that by default the signature is properly generated
 			UEdGraphPin* EntryNodeExec = FindExecutionPin(*EntryNode, EGPD_Output);
@@ -2404,6 +2408,7 @@ void UEdGraphSchema_K2::CreateFunctionGraphTerminators(UEdGraph& Graph, UFunctio
 	EntryNode->SignatureClass = NULL;
 	EntryNode->SignatureName = GraphName;
 	FunctionEntryCreator.Finalize();
+	SetNodeMetaData(EntryNode, FNodeMetadata::DefaultGraphNode);
 
 	// We don't have a signature class to base this on permanently, because it's not an override function.
 	// so we need to define the pins as user defined so that they are serialized.
@@ -2430,6 +2435,7 @@ void UEdGraphSchema_K2::CreateFunctionGraphTerminators(UEdGraph& Graph, UFunctio
 		ReturnNode->NodePosX = EntryNode->NodePosX + EntryNode->NodeWidth + 256;
 		ReturnNode->NodePosY = EntryNode->NodePosY;
 		NodeCreator.Finalize();
+		SetNodeMetaData(ReturnNode, FNodeMetadata::DefaultGraphNode);
 
 		ReturnNode->CreateUserDefinedPinsForFunctionEntryExit(FunctionSignature, /*bIsFunctionEntry=*/ false);
 
diff --git a/Engine/Source/Editor/EnvironmentQueryEditor/Private/EdGraphSchema_EnvironmentQuery.cpp b/Engine/Source/Editor/EnvironmentQueryEditor/Private/EdGraphSchema_EnvironmentQuery.cpp
--- a/Engine/Source/Editor/EnvironmentQueryEditor/Private/EdGraphSchema_EnvironmentQuery.cpp
+++ b/Engine/Source/Editor/EnvironmentQueryEditor/Private/EdGraphSchema_EnvironmentQuery.cpp
@@ -141,6 +141,7 @@ void UEdGraphSchema_EnvironmentQuery::CreateDefaultNodesForGraph(UEdGraph& Graph
 	FGraphNodeCreator<UEnvironmentQueryGraphNode_Root> NodeCreator(Graph);
 	UEnvironmentQueryGraphNode_Root* MyNode = NodeCreator.CreateNode();
 	NodeCreator.Finalize();
+	SetNodeMetaData(MyNode, FNodeMetadata::DefaultGraphNode);
 }
 
 FString GetNodeDescriptionHelper(UObject* Ob)
diff --git a/Engine/Source/Editor/UnrealEd/Private/Kismet2/Kismet2.cpp b/Engine/Source/Editor/UnrealEd/Private/Kismet2/Kismet2.cpp
--- a/Engine/Source/Editor/UnrealEd/Private/Kismet2/Kismet2.cpp
+++ b/Engine/Source/Editor/UnrealEd/Private/Kismet2/Kismet2.cpp
@@ -287,6 +287,7 @@ UBlueprint* FKismetEditorUtilities::CreateBlueprint(UClass* ParentClass, UObject
 			EventNode->CreateNewGuid();
 			EventNode->PostPlacedNewNode();
 			EventNode->AllocateDefaultPins();
+			UEdGraphSchema_K2::SetNodeMetaData(EventNode, FNodeMetadata::DefaultGraphNode);
 
 			Graph->AddNode(EventNode);
 
@@ -299,6 +300,7 @@ UBlueprint* FKismetEditorUtilities::CreateBlueprint(UClass* ParentClass, UObject
 			GetOwnerNode->AllocateDefaultPins();
 			GetOwnerNode->NodePosX = EventNode->NodePosX;
 			GetOwnerNode->NodePosY = EventNode->NodePosY + EventNode->NodeHeight + 100;
+			UEdGraphSchema_K2::SetNodeMetaData(GetOwnerNode, FNodeMetadata::DefaultGraphNode);
 
 			Graph->AddNode(GetOwnerNode);
 		}
diff --git a/Engine/Source/Editor/UnrealEd/Private/SoundCueGraphSchema.cpp b/Engine/Source/Editor/UnrealEd/Private/SoundCueGraphSchema.cpp
--- a/Engine/Source/Editor/UnrealEd/Private/SoundCueGraphSchema.cpp
+++ b/Engine/Source/Editor/UnrealEd/Private/SoundCueGraphSchema.cpp
@@ -289,6 +289,7 @@ void USoundCueGraphSchema::CreateDefaultNodesForGraph(UEdGraph& Graph) const
 	USoundCueGraphNode_Root* ResultRootNode = NodeCreator.CreateNode();
 	ResultRootNode->NodePosY = RootNodeHeightOffset;
 	NodeCreator.Finalize();
+	SetNodeMetaData(ResultRootNode, FNodeMetadata::DefaultGraphNode);
 }
 
 const FPinConnectionResponse USoundCueGraphSchema::CanCreateConnection(const UEdGraphPin* PinA, const UEdGraphPin* PinB) const
diff --git a/Engine/Source/Runtime/Engine/Classes/EdGraph/EdGraphNode.h b/Engine/Source/Runtime/Engine/Classes/EdGraph/EdGraphNode.h
--- a/Engine/Source/Runtime/Engine/Classes/EdGraph/EdGraphNode.h
+++ b/Engine/Source/Runtime/Engine/Classes/EdGraph/EdGraphNode.h
@@ -52,6 +52,15 @@ namespace ENodeAdvancedPins
 	};
 }
 
+/** Holds metadata keys, so as to discourage text duplication throughout the engine */
+struct ENGINE_API FNodeMetadata
+{
+	/** Identifies nodes that are defaultly added to populate new graphs (helps determine if a graph has any user placed nodes) */
+	static const FName DefaultGraphNode;
+private: 
+	FNodeMetadata() {}
+};
+
 // This is the context for a GetContextMenuActions call into a specific node
 struct FGraphNodeContextMenuBuilder
 {
diff --git a/Engine/Source/Runtime/Engine/Classes/EdGraph/EdGraphSchema.h b/Engine/Source/Runtime/Engine/Classes/EdGraph/EdGraphSchema.h
--- a/Engine/Source/Runtime/Engine/Classes/EdGraph/EdGraphSchema.h
+++ b/Engine/Source/Runtime/Engine/Classes/EdGraph/EdGraphSchema.h
@@ -395,8 +395,6 @@ public:
 };
 
 
-
-
 UCLASS(abstract)
 class ENGINE_API UEdGraphSchema : public UObject
 {
@@ -597,6 +595,9 @@ class ENGINE_API UEdGraphSchema : public UObject
 	 */
 	virtual void BreakNodeLinks(UEdGraphNode& TargetNode) const;
 
+	/** */
+	static bool SetNodeMetaData(UEdGraphNode* Node, FName const& KeyValue);
+
 	/**
 	 * Breaks all links from/to a single pin
 	 *
diff --git a/Engine/Source/Runtime/Engine/Private/EdGraph/EdGraphNode.cpp b/Engine/Source/Runtime/Engine/Private/EdGraph/EdGraphNode.cpp
--- a/Engine/Source/Runtime/Engine/Private/EdGraph/EdGraphNode.cpp
+++ b/Engine/Source/Runtime/Engine/Private/EdGraph/EdGraphNode.cpp
@@ -13,6 +13,8 @@
 
 #define LOCTEXT_NAMESPACE "EdGraph"
 
+FName const FNodeMetadata::DefaultGraphNode(TEXT("DefaultGraphNode"));
+
 /////////////////////////////////////////////////////
 // FGraphNodeContextMenuBuilder
 
diff --git a/Engine/Source/Runtime/Engine/Private/EdGraph/EdGraphSchema.cpp b/Engine/Source/Runtime/Engine/Private/EdGraph/EdGraphSchema.cpp
--- a/Engine/Source/Runtime/Engine/Private/EdGraph/EdGraphSchema.cpp
+++ b/Engine/Source/Runtime/Engine/Private/EdGraph/EdGraphSchema.cpp
@@ -366,6 +366,19 @@ void UEdGraphSchema::BreakNodeLinks(UEdGraphNode& TargetNode) const
 #endif	//#if WITH_EDITOR
 }
 
+bool UEdGraphSchema::SetNodeMetaData(UEdGraphNode* Node, FName const& KeyValue)
+{
+	if (UPackage* Package = Node->GetOutermost())
+	{
+		if (UMetaData* MetaData = Package->GetMetaData())
+		{
+			MetaData->SetValue(Node, KeyValue, TEXT("true"));
+			return true;
+		}
+	}
+	return false;
+}
+
 void UEdGraphSchema::BreakPinLinks(UEdGraphPin& TargetPin, bool bSendsNodeNotifcation) const
 {
 #if WITH_EDITOR
