Capping the blueprint cache in case it grows out of control for certain projects.

diff --git a/Engine/Source/Editor/BlueprintGraph/Private/BlueprintEditorSettings.cpp b/Engine/Source/Editor/BlueprintGraph/Private/BlueprintEditorSettings.cpp
new file mode 100644
--- /dev/null
+++ b/Engine/Source/Editor/BlueprintGraph/Private/BlueprintEditorSettings.cpp
@@ -0,0 +1,11 @@
+// Copyright 1998-2014 Epic Games, Inc. All Rights Reserved.
+
+#include "BlueprintGraphPrivatePCH.h"
+#include "BlueprintEditorSettings.h"
+
+UBlueprintEditorSettings::UBlueprintEditorSettings(const class FPostConstructInitializeProperties& PCIP)
+	: Super(PCIP)
+{
+	bShowGraphInstructionText = true;
+	NodeTemplateCacheCapMB    = 20.f;
+}
diff --git a/Engine/Source/Editor/BlueprintGraph/Private/BlueprintNodeTemplateCache.cpp b/Engine/Source/Editor/BlueprintGraph/Private/BlueprintNodeTemplateCache.cpp
--- a/Engine/Source/Editor/BlueprintGraph/Private/BlueprintNodeTemplateCache.cpp
+++ b/Engine/Source/Editor/BlueprintGraph/Private/BlueprintNodeTemplateCache.cpp
@@ -6,6 +6,9 @@
 #include "EdGraph/EdGraph.h"
 #include "KismetEditorUtilities.h"	// for CreateBlueprint()
 #include "BlueprintNodeSpawner.h"   // for NodeClass/Invoke()
+#include "BlueprintEditorSettings.h"
+
+DEFINE_LOG_CATEGORY_STATIC(LogBlueprintNodeCache, Log, All);
 
 /*******************************************************************************
  * Static FBlueprintNodeTemplateCache Helpers
@@ -51,6 +54,21 @@ namespace BlueprintNodeTemplateCacheImpl
 	 * @return 
 	 */
 	static UEdGraph* AddGraph(UBlueprint* BlueprintOuter, TSubclassOf<UEdGraphSchema> SchemaClass);
+
+	/**
+	 * 
+	 * 
+	 * @param  Object	
+	 * @return 
+	 */
+	static int32 ApproximateMemFootprint(UObject* Object);
+
+	/**
+	 * 
+	 * 
+	 * @return 
+	 */
+	static int32 GetCacheCapSize();
 }
 
 //------------------------------------------------------------------------------
@@ -118,19 +136,47 @@ static UEdGraph* BlueprintNodeTemplateCacheImpl::AddGraph(UBlueprint* BlueprintO
 	return NewGraph;
 }
 
+//------------------------------------------------------------------------------
+static int32 BlueprintNodeTemplateCacheImpl::ApproximateMemFootprint(UObject* Object)
+{
+	TArray<UObject*> ChildObjs;
+	GetObjectsWithOuter(Object, ChildObjs);
+
+	int32 ApproimateDataSize = sizeof(*Object);
+	for (UObject* ChildObj : ChildObjs)
+	{
+		// @TODO: doesn't account for any internal allocated memory (for member TArrays, etc.)
+		ApproimateDataSize += sizeof(*ChildObj);
+	}
+	return ApproimateDataSize;
+}
+
+//------------------------------------------------------------------------------
+static int32 BlueprintNodeTemplateCacheImpl::GetCacheCapSize()
+{
+	UBlueprintEditorSettings const* BpSettings = GetDefault<UBlueprintEditorSettings>();
+	// have to convert from MB to bytes
+	return (BpSettings->NodeTemplateCacheCapMB * 1024.f * 1024.f);
+}
+
 //------------------------------------------------------------------------------
 FBlueprintNodeTemplateCache::FBlueprintNodeTemplateCache()
+	: ApproximateAllocationTotal(0)
 {
 	using namespace BlueprintNodeTemplateCacheImpl; // for MakeCompatibleBlueprint()
 
-	TemplateOuters.Add(MakeCompatibleBlueprint(UBlueprint::StaticClass(), AActor::StaticClass(), UBlueprintGeneratedClass::StaticClass()));
-	TemplateOuters.Add(MakeCompatibleBlueprint(UAnimBlueprint::StaticClass(), UAnimInstance::StaticClass(), UAnimBlueprintGeneratedClass::StaticClass()));
+	UBlueprint* StandardBlueprint = MakeCompatibleBlueprint(UBlueprint::StaticClass(), AActor::StaticClass(), UBlueprintGeneratedClass::StaticClass());
+	AddBlueprintOuter(StandardBlueprint);
+
+	UBlueprint* AnimBlueprint = MakeCompatibleBlueprint(UAnimBlueprint::StaticClass(), UAnimInstance::StaticClass(), UAnimBlueprintGeneratedClass::StaticClass());
+	AddBlueprintOuter(AnimBlueprint);
 }
 
 //------------------------------------------------------------------------------
 UEdGraphNode* FBlueprintNodeTemplateCache::GetNodeTemplate(UBlueprintNodeSpawner const* NodeSpawner, UEdGraph* TargetGraph)
 {
 	using namespace BlueprintNodeTemplateCacheImpl;
+	static bool LoggedOnce = false;
 
 	UEdGraphNode* TemplateNode = nullptr;
 	if (UEdGraphNode** FoundNode = NodeTemplateCache.Find(NodeSpawner))
@@ -189,8 +235,13 @@ UEdGraphNode* FBlueprintNodeTemplateCache::GetNodeTemplate(UBlueprintNodeSpawner
 					GeneratedClassType = TargetBlueprint->GeneratedClass->GetClass();
 				}
 
+				// @TODO: if we're out of space, we shouldn't be 
 				CompatibleBlueprint = MakeCompatibleBlueprint(BlueprintClass, TargetBlueprint->ParentClass, GeneratedClassType);
-				TemplateOuters.Add(CompatibleBlueprint);
+				if (!AddBlueprintOuter(CompatibleBlueprint) && !LoggedOnce)
+				{
+					UE_LOG(LogBlueprintNodeCache, Display, TEXT("The blueprint template-node cache is full. As a result, you may experience slower than usual menu interactions. To avoid this, increase the cache's cap in the blueprint editor prefences."));
+					LoggedOnce = true;
+				}
 
 				// this graph may come default with a compatible graph
 				CompatibleOuter = FindCompatibleGraph(CompatibleBlueprint, NodeCDO);
@@ -199,13 +250,18 @@ UEdGraphNode* FBlueprintNodeTemplateCache::GetNodeTemplate(UBlueprintNodeSpawner
 			if (CompatibleOuter == nullptr)
 			{
 				CompatibleOuter = AddGraph(CompatibleBlueprint, TargetGraph->Schema);
+				ApproximateAllocationTotal += ApproximateMemFootprint(CompatibleOuter);
 			}
 		}
 
 		if (CompatibleOuter != nullptr)
 		{
 			TemplateNode = NodeSpawner->Invoke(CompatibleOuter, IBlueprintNodeBinder::FBindingSet(), FVector2D::ZeroVector);
-			NodeTemplateCache.Add(NodeSpawner, TemplateNode);
+			if (!CacheTemplateNode(NodeSpawner, TemplateNode) && !LoggedOnce)
+			{
+				UE_LOG(LogBlueprintNodeCache, Display, TEXT("The blueprint template-node cache is full. As a result, you may experience slower than usual menu interactions. To avoid this, increase the cache's cap in the blueprint editor prefences."));
+				LoggedOnce = true;
+			}
 		}
 	}
 
@@ -229,6 +285,16 @@ void FBlueprintNodeTemplateCache::ClearCachedTemplate(UBlueprintNodeSpawner cons
 	NodeTemplateCache.Remove(NodeSpawner); // GC should take care of the rest
 }
 
+//------------------------------------------------------------------------------
+int32 FBlueprintNodeTemplateCache::EstimateAllocatedSize() const
+{
+	int32 TotalEstimatedSize = NodeTemplateCache.GetAllocatedSize();
+	TotalEstimatedSize += TemplateOuters.GetAllocatedSize();
+	TotalEstimatedSize += ApproximateAllocationTotal;
+
+	return TotalEstimatedSize;
+}
+
 //------------------------------------------------------------------------------
 void FBlueprintNodeTemplateCache::AddReferencedObjects(FReferenceCollector& Collector)
 {
@@ -238,3 +304,46 @@ void FBlueprintNodeTemplateCache::AddReferencedObjects(FReferenceCollector& Coll
 	}
 	Collector.AddReferencedObjects(TemplateOuters);
 }
+
+//------------------------------------------------------------------------------
+bool FBlueprintNodeTemplateCache::AddBlueprintOuter(UBlueprint* Blueprint)
+{
+	int32 const ApproxBlueprintSize = BlueprintNodeTemplateCacheImpl::ApproximateMemFootprint(Blueprint);
+	ApproximateAllocationTotal += ApproxBlueprintSize;
+	int32 const NewEstimatedSize = EstimateAllocatedSize();
+
+	if (NewEstimatedSize > BlueprintNodeTemplateCacheImpl::GetCacheCapSize())
+	{
+		ApproximateAllocationTotal -= ApproxBlueprintSize;
+		return false;
+	}
+	else
+	{
+		TemplateOuters.Add(Blueprint);
+		return true;
+	}
+}
+
+//------------------------------------------------------------------------------
+bool FBlueprintNodeTemplateCache::CacheTemplateNode(UBlueprintNodeSpawner const* NodeSpawner, UEdGraphNode* NewNode)
+{
+	if (NewNode == nullptr)
+	{
+		return true;
+	}
+
+	int32 const ApproxNodeSize = BlueprintNodeTemplateCacheImpl::ApproximateMemFootprint(NewNode);
+	ApproximateAllocationTotal += ApproxNodeSize;
+	int32 const NewEstimatedSize = EstimateAllocatedSize();
+
+	if (NewEstimatedSize > BlueprintNodeTemplateCacheImpl::GetCacheCapSize())
+	{
+		ApproximateAllocationTotal -= ApproxNodeSize;
+		return false;
+	}
+	else
+	{
+		NodeTemplateCache.Add(NodeSpawner, NewNode);
+		return true;
+	}
+}
diff --git a/Engine/Source/Editor/BlueprintGraph/Public/BlueprintEditorSettings.h b/Engine/Source/Editor/BlueprintGraph/Public/BlueprintEditorSettings.h
new file mode 100644
--- /dev/null
+++ b/Engine/Source/Editor/BlueprintGraph/Public/BlueprintEditorSettings.h
@@ -0,0 +1,25 @@
+// Copyright 1998-2014 Epic Games, Inc. All Rights Reserved.
+
+#pragma once
+
+#include "BlueprintEditorSettings.generated.h"
+
+UCLASS(config=EditorUserSettings)
+class BLUEPRINTGRAPH_API UBlueprintEditorSettings
+	:	public UObject
+{
+	GENERATED_UCLASS_BODY()
+
+// UX Settings
+public:
+	/** Determines if lightweight tutorial text shows up at the top of empty blueprint graphs */
+	UPROPERTY(EditAnywhere, config, Category="User Experience")
+	bool bShowGraphInstructionText;
+
+// Perf Settings
+public:
+	/** The node template cache is used to speed up blueprint menuing. This determines the peak data size for that cache. */
+	UPROPERTY(EditAnywhere, config, Category="Performance", DisplayName="Node-Template Cache Cap (MB)", meta=(ClampMin="0.004", UIMin="0"))
+	float NodeTemplateCacheCapMB;
+};
+ 
\ No newline at end of file
diff --git a/Engine/Source/Editor/BlueprintGraph/Public/BlueprintNodeTemplateCache.h b/Engine/Source/Editor/BlueprintGraph/Public/BlueprintNodeTemplateCache.h
--- a/Engine/Source/Editor/BlueprintGraph/Public/BlueprintNodeTemplateCache.h
+++ b/Engine/Source/Editor/BlueprintGraph/Public/BlueprintNodeTemplateCache.h
@@ -54,11 +54,35 @@ public:
 	 */
 	void ClearCachedTemplate(UBlueprintNodeSpawner const* NodeSpawner);
 
+	/**
+	 * 
+	 * 
+	 * @return 
+	 */
+	int32 EstimateAllocatedSize() const;
+
 	// FGCObject interface
 	virtual void AddReferencedObjects(FReferenceCollector& Collector) override;
 	// End FGCObject interface
 
 private:
+	/**
+	 * 
+	 * 
+	 * @param  Blueprint	
+	 * @return 
+	 */
+	bool AddBlueprintOuter(UBlueprint* Blueprint);
+
+	/**
+	 * 
+	 * 
+	 * @param  NodeSpawner	
+	 * @param  NewNode	
+	 * @return 
+	 */
+	bool CacheTemplateNode(UBlueprintNodeSpawner const* NodeSpawner, UEdGraphNode* NewNode);
+
 	/** 
 	 * Unfortunately, we cannot nest template-nodes in the transient package. 
 	 * Certain nodes operate on the assumption that they have a UEdGraph outer, 
@@ -73,4 +97,7 @@ private:
 
 	/** */
 	TMap<UBlueprintNodeSpawner const*, UEdGraphNode*> NodeTemplateCache;
+
+	/** */
+	int32 ApproximateAllocationTotal;
 };
diff --git a/Engine/Source/Editor/EditorSettingsViewer/EditorSettingsViewer.Build.cs b/Engine/Source/Editor/EditorSettingsViewer/EditorSettingsViewer.Build.cs
--- a/Engine/Source/Editor/EditorSettingsViewer/EditorSettingsViewer.Build.cs
+++ b/Engine/Source/Editor/EditorSettingsViewer/EditorSettingsViewer.Build.cs
@@ -24,7 +24,7 @@ namespace UnrealBuildTool.Rules
 					"SlateCore",
 					"UnrealEd",
                     "InternationalizationSettings",
-					"Kismet"
+					"BlueprintGraph"
 				}
 			);
 
diff --git a/Engine/Source/Editor/Kismet/Private/BlueprintEditorSettings.cpp b/Engine/Source/Editor/Kismet/Private/BlueprintEditorSettings.cpp
deleted file mode 100644
diff --git a/Engine/Source/Editor/Kismet/Public/BlueprintEditorSettings.h b/Engine/Source/Editor/Kismet/Public/BlueprintEditorSettings.h
deleted file mode 100644
