Temporary variables are now only allowed in macro graphs.

diff --git a/Engine/Source/Editor/BlueprintGraph/Private/K2Node_TemporaryVariable.cpp b/Engine/Source/Editor/BlueprintGraph/Private/K2Node_TemporaryVariable.cpp
--- a/Engine/Source/Editor/BlueprintGraph/Private/K2Node_TemporaryVariable.cpp
+++ b/Engine/Source/Editor/BlueprintGraph/Private/K2Node_TemporaryVariable.cpp
@@ -104,10 +104,10 @@ bool UK2Node_TemporaryVariable::IsCompatibleWithGraph(UEdGraph const* TargetGrap
 	bool bIsCompatible = Super::IsCompatibleWithGraph(TargetGraph);
 	if (bIsCompatible)
 	{
-		EGraphType const GraphType = TargetGraph->GetSchema()->GetGraphType(TargetGraph);
-		bIsCompatible = (GraphType != GT_Ubergraph) && (GraphType != GT_Animation);
+		bIsCompatible = false;
 
-		if (bIsCompatible && (GraphType != GT_Macro))
+		EGraphType const GraphType = TargetGraph->GetSchema()->GetGraphType(TargetGraph);
+		if (GraphType == GT_Macro)
 		{
 			bIsCompatible = !bIsPersistent;
 		}
