Implemented IsConnectionDisallowed for CateByteToEnum/SwitchEnum. Reduces spam when filtering by context is turned on with new blueprint context menu..

diff --git a/Engine/Source/Editor/BlueprintGraph/Classes/K2Node_CastByteToEnum.h b/Engine/Source/Editor/BlueprintGraph/Classes/K2Node_CastByteToEnum.h
--- a/Engine/Source/Editor/BlueprintGraph/Classes/K2Node_CastByteToEnum.h
+++ b/Engine/Source/Editor/BlueprintGraph/Classes/K2Node_CastByteToEnum.h
@@ -32,6 +32,7 @@ class UK2Node_CastByteToEnum : public UK2Node
 	virtual bool IsNodePure() const override { return true; }
 	FNodeHandlingFunctor* CreateNodeHandler(FKismetCompilerContext& CompilerContext) const override;
 	virtual void ExpandNode(class FKismetCompilerContext& CompilerContext, UEdGraph* SourceGraph) override;
+	virtual bool IsConnectionDisallowed(const UEdGraphPin* MyPin, const UEdGraphPin* OtherPin, FString& OutReason) const override;
 	virtual void GetMenuActions(FBlueprintActionDatabaseRegistrar& ActionRegistrar) const override;
 	virtual FText GetMenuCategory() const override;
 	// End UK2Node interface
diff --git a/Engine/Source/Editor/BlueprintGraph/Classes/K2Node_SwitchEnum.h b/Engine/Source/Editor/BlueprintGraph/Classes/K2Node_SwitchEnum.h
--- a/Engine/Source/Editor/BlueprintGraph/Classes/K2Node_SwitchEnum.h
+++ b/Engine/Source/Editor/BlueprintGraph/Classes/K2Node_SwitchEnum.h
@@ -36,6 +36,7 @@ class UK2Node_SwitchEnum : public UK2Node_Switch, public INodeDependingOnEnumInt
 	// End of UEdGraphNode interface
 
 	// UK2Node interface
+	virtual bool IsConnectionDisallowed(const UEdGraphPin* MyPin, const UEdGraphPin* OtherPin, FString& OutReason) const override;
 	virtual void GetMenuActions(FBlueprintActionDatabaseRegistrar& ActionRegistrar) const override;
 	// End of UK2Node interface
 
diff --git a/Engine/Source/Editor/BlueprintGraph/Private/K2Node_CastByteToEnum.cpp b/Engine/Source/Editor/BlueprintGraph/Private/K2Node_CastByteToEnum.cpp
--- a/Engine/Source/Editor/BlueprintGraph/Private/K2Node_CastByteToEnum.cpp
+++ b/Engine/Source/Editor/BlueprintGraph/Private/K2Node_CastByteToEnum.cpp
@@ -148,6 +148,19 @@ FNodeHandlingFunctor* UK2Node_CastByteToEnum::CreateNodeHandler(FKismetCompilerC
 	return new FNodeHandlingFunctor(CompilerContext);;
 }
 
+bool UK2Node_CastByteToEnum::IsConnectionDisallowed(const UEdGraphPin* MyPin, const UEdGraphPin* OtherPin, FString& OutReason) const
+{
+	const UEnum* SubCategoryObject = Cast<UEnum>(OtherPin->PinType.PinSubCategoryObject.Get());
+	if (SubCategoryObject)
+	{
+		if (SubCategoryObject != Enum)
+		{
+			return true;
+		}
+	}
+	return false;
+}
+
 void UK2Node_CastByteToEnum::GetMenuActions(FBlueprintActionDatabaseRegistrar& ActionRegistrar) const
 {
 	auto CustomizeEnumNodeLambda = [](UEdGraphNode* NewNode, bool bIsTemplateNode, TWeakObjectPtr<UEnum> EnumPtr)
diff --git a/Engine/Source/Editor/BlueprintGraph/Private/K2Node_SwitchEnum.cpp b/Engine/Source/Editor/BlueprintGraph/Private/K2Node_SwitchEnum.cpp
--- a/Engine/Source/Editor/BlueprintGraph/Private/K2Node_SwitchEnum.cpp
+++ b/Engine/Source/Editor/BlueprintGraph/Private/K2Node_SwitchEnum.cpp
@@ -70,6 +70,19 @@ FText UK2Node_SwitchEnum::GetTooltipText() const
 	return NSLOCTEXT("K2Node", "SwitchEnum_ToolTip", "Selects an output that matches the input value");
 }
 
+bool UK2Node_SwitchEnum::IsConnectionDisallowed(const UEdGraphPin* MyPin, const UEdGraphPin* OtherPin, FString& OutReason) const
+{
+	const UEnum* SubCategoryObject = Cast<UEnum>( OtherPin->PinType.PinSubCategoryObject.Get() );
+	if( SubCategoryObject )
+	{
+		if( SubCategoryObject != Enum )
+		{
+			return true;
+		}
+	}
+	return false;
+}
+
 void UK2Node_SwitchEnum::GetMenuActions(FBlueprintActionDatabaseRegistrar& ActionRegistrar) const
 {
 	auto CustomizeEnumNodeLambda = [](UEdGraphNode* NewNode, bool bIsTemplateNode, TWeakObjectPtr<UEnum> EnumPtr)
