Make Switch Enum nodes have a pin of the actual enum type, not a generic byte

diff --git a/Engine/Source/Editor/BlueprintGraph/Private/K2Node_SwitchEnum.cpp b/Engine/Source/Editor/BlueprintGraph/Private/K2Node_SwitchEnum.cpp
--- a/Engine/Source/Editor/BlueprintGraph/Private/K2Node_SwitchEnum.cpp
+++ b/Engine/Source/Editor/BlueprintGraph/Private/K2Node_SwitchEnum.cpp
@@ -127,7 +127,8 @@ void UK2Node_SwitchEnum::GetMenuActions(FBlueprintActionDatabaseRegistrar& Actio
 void UK2Node_SwitchEnum::CreateSelectionPin()
 {
 	const UEdGraphSchema_K2* K2Schema = GetDefault<UEdGraphSchema_K2>();
-	CreatePin(EGPD_Input, K2Schema->PC_Byte, TEXT(""), NULL, false, false, TEXT("Selection"));
+	UEdGraphPin* Pin = CreatePin(EGPD_Input, K2Schema->PC_Byte, TEXT(""), Enum, false, false, TEXT("Selection"));
+	K2Schema->SetPinDefaultValueBasedOnType(Pin);
 }
 
 FString UK2Node_SwitchEnum::GetUniquePinName()
