Properly handling when the node cache size is set to zero.

diff --git a/Engine/Source/Editor/BlueprintGraph/Private/BlueprintComponentNodeSpawner.cpp b/Engine/Source/Editor/BlueprintGraph/Private/BlueprintComponentNodeSpawner.cpp
--- a/Engine/Source/Editor/BlueprintGraph/Private/BlueprintComponentNodeSpawner.cpp
+++ b/Engine/Source/Editor/BlueprintGraph/Private/BlueprintComponentNodeSpawner.cpp
@@ -4,6 +4,7 @@
 #include "BlueprintComponentNodeSpawner.h"
 #include "K2Node_AddComponent.h"
 #include "ClassIconFinder.h" // for FindIconNameForClass()
+#include "BlueprintNodeTemplateCache.h" // for IsTemplateOuter()
 
 #define LOCTEXT_NAMESPACE "BlueprintComponenetNodeSpawner"
 
@@ -55,7 +56,7 @@ UEdGraphNode* UBlueprintComponentNodeSpawner::Invoke(UEdGraph* ParentGraph, FBin
 	
 	UK2Node_AddComponent* NewNode = CastChecked<UK2Node_AddComponent>(Super::Invoke(ParentGraph, Bindings, Location, PostSpawnDelegate));
 
-	bool const bIsTemplateNode = (ParentGraph->GetOutermost() == GetTransientPackage());
+	bool const bIsTemplateNode = FBlueprintNodeTemplateCache::IsTemplateOuter(ParentGraph);
 	if (!bIsTemplateNode)
 	{
 		UBlueprint* Blueprint = NewNode->GetBlueprint();
diff --git a/Engine/Source/Editor/BlueprintGraph/Private/BlueprintFunctionNodeSpawner.cpp b/Engine/Source/Editor/BlueprintGraph/Private/BlueprintFunctionNodeSpawner.cpp
--- a/Engine/Source/Editor/BlueprintGraph/Private/BlueprintFunctionNodeSpawner.cpp
+++ b/Engine/Source/Editor/BlueprintGraph/Private/BlueprintFunctionNodeSpawner.cpp
@@ -5,6 +5,7 @@
 #include "EdGraphSchema_K2.h"	// for FBlueprintMetadata
 #include "GameFramework/Actor.h"
 #include "BlueprintVariableNodeSpawner.h"
+#include "BlueprintNodeTemplateCache.h" // for IsTemplateOuter()
 
 #define LOCTEXT_NAMESPACE "BlueprintFunctionNodeSpawner"
 
@@ -50,7 +51,7 @@ static bool BlueprintFunctionNodeSpawnerImpl::BindFunctionNode(UK2Node_CallFunct
 {
 	bool bSuccessfulBinding = false;
 
-	bool const bIsTemplateNode = (NewNode->GetOutermost() == GetTransientPackage());
+	bool const bIsTemplateNode = FBlueprintNodeTemplateCache::IsTemplateOuter(NewNode->GetGraph());
 	if (!bIsTemplateNode)
 	{
 		if (UProperty const* BoundProperty = Cast<UProperty>(BoundObject))
diff --git a/Engine/Source/Editor/BlueprintGraph/Private/BlueprintNodeSpawner.cpp b/Engine/Source/Editor/BlueprintGraph/Private/BlueprintNodeSpawner.cpp
--- a/Engine/Source/Editor/BlueprintGraph/Private/BlueprintNodeSpawner.cpp
+++ b/Engine/Source/Editor/BlueprintGraph/Private/BlueprintNodeSpawner.cpp
@@ -186,7 +186,7 @@ UEdGraphNode* UBlueprintNodeSpawner::Invoke(UEdGraph* ParentGraph, FBindingSet c
 		NewNode->NodePosX = Location.X;
 		NewNode->NodePosY = Location.Y;
 
-		bool const bIsTemplateNode = (ParentGraph->GetOutermost() == GetTransientPackage());
+		bool const bIsTemplateNode = FBlueprintNodeTemplateCache::IsTemplateOuter(ParentGraph);
 		PostSpawnDelegate.ExecuteIfBound(NewNode, bIsTemplateNode);
 
 		if (!bIsTemplateNode)
diff --git a/Engine/Source/Editor/BlueprintGraph/Private/BlueprintNodeTemplateCache.cpp b/Engine/Source/Editor/BlueprintGraph/Private/BlueprintNodeTemplateCache.cpp
--- a/Engine/Source/Editor/BlueprintGraph/Private/BlueprintNodeTemplateCache.cpp
+++ b/Engine/Source/Editor/BlueprintGraph/Private/BlueprintNodeTemplateCache.cpp
@@ -16,59 +16,90 @@ DEFINE_LOG_CATEGORY_STATIC(LogBlueprintNodeCache, Log, All);
 
 namespace BlueprintNodeTemplateCacheImpl
 {
+	/** */
+	static int32 ActiveMemFootprint   = 0;
+	/** Used to track the average blueprint size (so that we can try to predict when a blueprint would fail to be cached) */
+	static int32 MadeBlueprintCount   = 0;
+	static int32 AverageBlueprintSize = 0;
+	
+	/** Metadata tag used to identify graphs created by this system. */
+	static const FName TemplateGraphMetaTag(TEXT("NodeTemplateCache_Graph"));
+
 	/**
+	 * Checks to see if this node is compatible with the given graph (to know if
+	 * a node template can be spawned within it).
 	 * 
-	 * 
-	 * @param  NodeObj	
-	 * @param  Graph	
-	 * @return 
+	 * @param  NodeObj	The CDO of the node you want to spawn.
+	 * @param  Graph	The graph you want to check compatibility for.
+	 * @return True if the node and graph are compatible (for templating purposes).
 	 */
 	static bool IsCompatible(UEdGraphNode* NodeObj, UEdGraph* Graph);
 
 	/**
+	 * Looks through a blueprint for compatible graphs (ones that the specified
+	 * node can spawn into).
 	 * 
-	 * 
-	 * @param  BlueprintOuter	
-	 * @param  NodeObj	
-	 * @return 
+	 * @param  BlueprintOuter	The blueprint to search through.
+	 * @param  NodeObj			The CDO of the node you want to spawn.
+	 * @param  IsCompatibleFunc	An option callback function to further filter out incompatible nodes.
+	 * @return The first compatible graph found (null if no graph was found).
 	 */
-	static UEdGraph* FindCompatibleGraph(UBlueprint* BlueprintOuter, UEdGraphNode* NodeObj);
+	static UEdGraph* FindCompatibleGraph(UBlueprint* BlueprintOuter, UEdGraphNode* NodeObj, bool(*IsCompatibleFunc)(UEdGraph*) = nullptr);
 
 	/**
+	 * Creates a transient, temporary blueprint. Intended to be used as a 
+	 * template-node's outer (grandparent).
 	 * 
-	 * 
-	 * @param  BlueprintClass	
-	 * @param  ParentClass	
-	 * @param  GeneratedClassType	
-	 * @return 
+	 * @param  BlueprintClass		The class of blueprint to make.
+	 * @param  ParentClass			The class type of blueprint to make (actor, object, etc.).
+	 * @param  GeneratedClassType	The type of class that the blueprint should generate.
+	 * @return A newly spawned (transient) blueprint.
 	 */
 	static UBlueprint* MakeCompatibleBlueprint(TSubclassOf<UBlueprint> BlueprintClass, UClass* ParentClass, TSubclassOf<UBlueprintGeneratedClass> GeneratedClassType);
 
 	/**
+	 * Creates a new transient graph, for template node use (meant to be used as
+	 * a template node's outer).
 	 * 
-	 * 
-	 * @param  GraphClass	
-	 * @param  SchemaClass	
-	 * @param  BlueprintOuter	
-	 * @param  GraphType	
-	 * @return 
+	 * @param  BlueprintOuter	The blueprint to nest the new graph under.
+	 * @param  SchemaClass		The schema to assign the new graph.
+	 * @return A newly created graph.
 	 */
 	static UEdGraph* AddGraph(UBlueprint* BlueprintOuter, TSubclassOf<UEdGraphSchema> SchemaClass);
 
 	/**
+	 * Adds metadata to the supplied graph, flagging it as a graph belonging to 
+	 * the BlueprintNodeTemplateCache (so we can easily identify it later on).
 	 * 
-	 * 
-	 * @param  Object	
-	 * @return 
+	 * @param  NewGraph	The graph to flag.
 	 */
-	static int32 ApproximateMemFootprint(UObject* Object);
+	static void MarkGraphForTemplateUse(UEdGraph* NewGraph);
 
 	/**
+	 * Determines if the specified graph is one that was allocated by 
+	 * BlueprintNodeTemplateCache (to house template nodes).
 	 * 
+	 * @param  ParentGraph	The graph you want checked.
+	 * @return True if this graph belongs to a BlueprintNodeTemplateCache, false if not.
+	 */
+	static bool IsTemplateOuter(UEdGraph* ParentGraph);
+
+	/**
+	 * Converts the cache memory cap into bytes (form user settings).
 	 * 
-	 * @return 
+	 * @return The user defined cache cap (in bytes).
 	 */
 	static int32 GetCacheCapSize();
+
+	/**
+	 * Totals the size of the specified object, along with any other objects 
+	 * that have it in their outer chain. Does not account for any allocated 
+	 * memory that belongs to the object(s).
+	 * 
+	 * @param  Object	The object you want an estimated byte size for.
+	 * @return An estimated size (in bytes)... currently does not account for any allocated memory that the object may be responsible for.
+	 */
+	static int32 ApproximateMemFootprint(UObject const* Object);	
 }
 
 //------------------------------------------------------------------------------
@@ -78,7 +109,7 @@ static bool BlueprintNodeTemplateCacheImpl::IsCompatible(UEdGraphNode* NodeObj,
 }
 
 //------------------------------------------------------------------------------
-static UEdGraph* BlueprintNodeTemplateCacheImpl::FindCompatibleGraph(UBlueprint* BlueprintOuter, UEdGraphNode* NodeObj)
+static UEdGraph* BlueprintNodeTemplateCacheImpl::FindCompatibleGraph(UBlueprint* BlueprintOuter, UEdGraphNode* NodeObj, bool(*IsCompatibleFunc)(UEdGraph*))
 {
 	UEdGraph* FoundGraph = nullptr;
 
@@ -88,7 +119,14 @@ static UEdGraph* BlueprintNodeTemplateCacheImpl::FindCompatibleGraph(UBlueprint*
 	for (UObject* Child : BlueprintChildObjs)
 	{
 		UEdGraph* ChildGraph = Cast<UEdGraph>(Child);
-		if ((ChildGraph != nullptr) && IsCompatible(NodeObj, ChildGraph))
+		bool bIsCompatible = (ChildGraph != nullptr) && IsCompatible(NodeObj, ChildGraph);
+
+		if (bIsCompatible && (IsCompatibleFunc != nullptr))
+		{
+			bIsCompatible = IsCompatibleFunc(ChildGraph);
+		}
+
+		if (bIsCompatible)
 		{
 			FoundGraph = ChildGraph;
 			break;
@@ -121,6 +159,12 @@ static UBlueprint* BlueprintNodeTemplateCacheImpl::MakeCompatibleBlueprint(TSubc
 	UBlueprint* NewBlueprint = FKismetEditorUtilities::CreateBlueprint(ParentClass, BlueprintOuter, BlueprintName, BlueprintType, BlueprintClass, GeneratedClassType);
 	NewBlueprint->SetFlags(RF_Transient);
 
+	float const AproxBlueprintSize = ApproximateMemFootprint(NewBlueprint);
+	// track the average blueprint size, so that we can attempt to predict 
+	// whether a  blueprint will fail to be cached (when the cache is near full)
+	AverageBlueprintSize = AverageBlueprintSize * ((float)MadeBlueprintCount / (++MadeBlueprintCount)) + 
+		(AproxBlueprintSize / MadeBlueprintCount) + 0.5f;
+
 	return NewBlueprint;
 }
 
@@ -133,22 +177,28 @@ static UEdGraph* BlueprintNodeTemplateCacheImpl::AddGraph(UBlueprint* BlueprintO
 	UEdGraph* NewGraph = ConstructObject<UEdGraph>(GraphClass, BlueprintOuter, GraphName, RF_Transient);
 	NewGraph->Schema = SchemaClass;
 
+	MarkGraphForTemplateUse(NewGraph);
 	return NewGraph;
 }
 
 //------------------------------------------------------------------------------
-static int32 BlueprintNodeTemplateCacheImpl::ApproximateMemFootprint(UObject* Object)
+static void BlueprintNodeTemplateCacheImpl::MarkGraphForTemplateUse(UEdGraph* NewGraph)
 {
-	TArray<UObject*> ChildObjs;
-	GetObjectsWithOuter(Object, ChildObjs);
+	UPackage*  TemplatePackage = NewGraph->GetOutermost();
+	UMetaData* PackageMetadata = TemplatePackage->GetMetaData();
+	PackageMetadata->SetValue(NewGraph, TemplateGraphMetaTag, TEXT("true"));
+}
 
-	int32 ApproimateDataSize = sizeof(*Object);
-	for (UObject* ChildObj : ChildObjs)
+//------------------------------------------------------------------------------
+bool BlueprintNodeTemplateCacheImpl::IsTemplateOuter(UEdGraph* ParentGraph)
+{
+	if (ParentGraph->HasAnyFlags(RF_Transient))
 	{
-		// @TODO: doesn't account for any internal allocated memory (for member TArrays, etc.)
-		ApproimateDataSize += sizeof(*ChildObj);
+		UPackage* GraphPackage = ParentGraph->GetOutermost();
+		UMetaData* PackageMetadata = GraphPackage->GetMetaData();
+		return PackageMetadata->HasValue(ParentGraph, TemplateGraphMetaTag);
 	}
-	return ApproimateDataSize;
+	return false;
 }
 
 //------------------------------------------------------------------------------
@@ -156,27 +206,63 @@ static int32 BlueprintNodeTemplateCacheImpl::GetCacheCapSize()
 {
 	UBlueprintEditorSettings const* BpSettings = GetDefault<UBlueprintEditorSettings>();
 	// have to convert from MB to bytes
-	return (BpSettings->NodeTemplateCacheCapMB * 1024.f * 1024.f);
+	return (BpSettings->NodeTemplateCacheCapMB * 1024.f * 1024.f) + 0.5f;
+}
+
+//------------------------------------------------------------------------------
+static int32 BlueprintNodeTemplateCacheImpl::ApproximateMemFootprint(UObject const* Object)
+{
+	TArray<UObject*> ChildObjs;
+	GetObjectsWithOuter(Object, ChildObjs, /*bIncludeNestedObjects =*/true);
+
+	int32 ApproimateDataSize = sizeof(*Object);
+	for (UObject* ChildObj : ChildObjs)
+	{
+		// @TODO: doesn't account for any internal allocated memory (for member TArrays, etc.)
+		ApproimateDataSize += sizeof(*ChildObj);
+	}
+	return ApproimateDataSize;
 }
 
+/*******************************************************************************
+ * FBlueprintNodeTemplateCache
+ ******************************************************************************/
+
 //------------------------------------------------------------------------------
 FBlueprintNodeTemplateCache::FBlueprintNodeTemplateCache()
-	: ApproximateAllocationTotal(0)
+	: ApproximateObjectMem(0)
 {
 	using namespace BlueprintNodeTemplateCacheImpl; // for MakeCompatibleBlueprint()
 
 	UBlueprint* StandardBlueprint = MakeCompatibleBlueprint(UBlueprint::StaticClass(), AActor::StaticClass(), UBlueprintGeneratedClass::StaticClass());
-	AddBlueprintOuter(StandardBlueprint);
+	CacheBlueprintOuter(StandardBlueprint);
 
 	UBlueprint* AnimBlueprint = MakeCompatibleBlueprint(UAnimBlueprint::StaticClass(), UAnimInstance::StaticClass(), UAnimBlueprintGeneratedClass::StaticClass());
-	AddBlueprintOuter(AnimBlueprint);
+	CacheBlueprintOuter(AnimBlueprint);
 }
 
 //------------------------------------------------------------------------------
 UEdGraphNode* FBlueprintNodeTemplateCache::GetNodeTemplate(UBlueprintNodeSpawner const* NodeSpawner, UEdGraph* TargetGraph)
 {
 	using namespace BlueprintNodeTemplateCacheImpl;
-	static bool LoggedOnce = false;
+
+	bool bIsOverMemCap = false;
+	auto LogCacheFullMsg = [&bIsOverMemCap]()
+	{
+		if (!bIsOverMemCap)
+		{
+			static int32 LoggedCapSize = -1;
+			int32 const CurrentCacheSize = GetCacheCapSize();
+			// log only once for each cap size change
+			if (LoggedCapSize != CurrentCacheSize)
+			{
+				UE_LOG(LogBlueprintNodeCache, Display, TEXT("The blueprint template-node cache is full. As a result, you may experience interactions which are slower than normal. To avoid this, increase the cache's cap in the blueprint editor prefences."));
+				LoggedCapSize = CurrentCacheSize;
+			}			
+			bIsOverMemCap = true;
+		}
+	};
+	
 
 	UEdGraphNode* TemplateNode = nullptr;
 	if (UEdGraphNode** FoundNode = NodeTemplateCache.Find(NodeSpawner))
@@ -202,7 +288,14 @@ UEdGraphNode* FBlueprintNodeTemplateCache::GetNodeTemplate(UBlueprintNodeSpawner
 			check(TargetBlueprint != nullptr);
 			BlueprintClass = TargetBlueprint->GetClass();
 
-			TargetGraph = FindCompatibleGraph(TargetBlueprint, NodeCDO);
+			// check used to help identify user interactable graphs (as opposed to 
+			// intermediate/transient graphs).
+			auto IsCompatibleUserGraph = [](UEdGraph* Graph)->bool
+			{
+				return !Graph->HasAnyFlags(RF_Transient);
+			};
+
+			TargetGraph = FindCompatibleGraph(TargetBlueprint, NodeCDO, IsCompatibleUserGraph);
 			check(TargetGraph != nullptr);
 		}		
 
@@ -212,6 +305,11 @@ UEdGraphNode* FBlueprintNodeTemplateCache::GetNodeTemplate(UBlueprintNodeSpawner
 		for (UBlueprint* Blueprint : TemplateOuters)
 		{
 			CompatibleOuter = FindCompatibleGraph(Blueprint, NodeCDO);
+			if (CompatibleOuter != nullptr)
+			{
+				MarkGraphForTemplateUse(CompatibleOuter);
+			}
+
 			if (CompatibleOuter != nullptr)
 			{
 				CompatibleBlueprint = Blueprint;
@@ -223,44 +321,84 @@ UEdGraphNode* FBlueprintNodeTemplateCache::GetNodeTemplate(UBlueprintNodeSpawner
 			}
 		}
 
+		// reset ActiveMemFootprint, so calls to CacheBlueprintOuter()/CacheTemplateNode()
+		// use the most up-to-date value (users could have since modified the
+		// nodes, so they could have grown in size... like with AllocateDefaultPins)
+		//
+		// @TODO: GetEstimateCacheSize() is (most likely) inaccurate, seeing as
+		//        external systems mutate template-nodes (such as calling
+		//        AllocateDefaultPins), and this returns a size estimate from
+		//        when the node was first spawned (it is too slow to recalculate
+		//        the size of the object hierarchy here)
+		ActiveMemFootprint = GetEstimateCacheSize();
+
+		int32 const CacheCapSize = GetCacheCapSize();
+		if (ActiveMemFootprint > CacheCapSize)
+		{
+			LogCacheFullMsg();
+			// @TODO: evict nodes until we're back under the cap (in case the cap
+			//        was changed at runtime, or external user modified node sizes)
+		}
+
 		// if a TargetGraph was supplied, and we couldn't find a suitable outer
 		// for this template-node, then attempt to emulate that graph
 		if (bHasTargetGraph)
 		{
 			if (CompatibleBlueprint == nullptr)
 			{
-				TSubclassOf<UBlueprintGeneratedClass> GeneratedClassType = UBlueprintGeneratedClass::StaticClass();
-				if (TargetBlueprint->GeneratedClass != nullptr)
+				// if the cache is near full, attempt to predict if this  
+				// impending cache will fail (if so, we don't want to waste the 
+				// cycles on allocating a temp blueprint)
+				if (!bIsOverMemCap && ((AverageBlueprintSize == 0) ||
+					(ActiveMemFootprint + AverageBlueprintSize <= CacheCapSize)))
 				{
-					GeneratedClassType = TargetBlueprint->GeneratedClass->GetClass();
+					TSubclassOf<UBlueprintGeneratedClass> GeneratedClassType = UBlueprintGeneratedClass::StaticClass();
+					if (TargetBlueprint->GeneratedClass != nullptr)
+					{
+						GeneratedClassType = TargetBlueprint->GeneratedClass->GetClass();
+					}
+
+					CompatibleBlueprint = MakeCompatibleBlueprint(BlueprintClass, TargetBlueprint->ParentClass, GeneratedClassType);
+					if (!CacheBlueprintOuter(CompatibleBlueprint))
+					{
+						LogCacheFullMsg();
+					}
+
+					// this graph may come default with a compatible graph
+					CompatibleOuter = FindCompatibleGraph(CompatibleBlueprint, NodeCDO);
+					if (CompatibleOuter != nullptr)
+					{
+						MarkGraphForTemplateUse(CompatibleOuter);
+					}
 				}
-
-				// @TODO: if we're out of space, we shouldn't be 
-				CompatibleBlueprint = MakeCompatibleBlueprint(BlueprintClass, TargetBlueprint->ParentClass, GeneratedClassType);
-				if (!AddBlueprintOuter(CompatibleBlueprint) && !LoggedOnce)
+				else
 				{
-					UE_LOG(LogBlueprintNodeCache, Display, TEXT("The blueprint template-node cache is full. As a result, you may experience slower than usual menu interactions. To avoid this, increase the cache's cap in the blueprint editor prefences."));
-					LoggedOnce = true;
-				}
+ 					CompatibleBlueprint = TargetBlueprint;
+ 					CompatibleOuter = FindCompatibleGraph(TargetBlueprint, NodeCDO, &BlueprintNodeTemplateCacheImpl::IsTemplateOuter);
 
-				// this graph may come default with a compatible graph
-				CompatibleOuter = FindCompatibleGraph(CompatibleBlueprint, NodeCDO);
+					LogCacheFullMsg();
+				}
 			}
 
 			if (CompatibleOuter == nullptr)
 			{
 				CompatibleOuter = AddGraph(CompatibleBlueprint, TargetGraph->Schema);
-				ApproximateAllocationTotal += ApproximateMemFootprint(CompatibleOuter);
+
+				if (CompatibleBlueprint != TargetBlueprint)
+				{
+					int32 const ApproxGraphSize = ApproximateMemFootprint(CompatibleOuter);
+					ActiveMemFootprint   += ApproxGraphSize;
+					ApproximateObjectMem += ApproxGraphSize;
+				}				
 			}
 		}
 
 		if (CompatibleOuter != nullptr)
 		{
 			TemplateNode = NodeSpawner->Invoke(CompatibleOuter, IBlueprintNodeBinder::FBindingSet(), FVector2D::ZeroVector);
-			if (!CacheTemplateNode(NodeSpawner, TemplateNode) && !LoggedOnce)
+			if (!bIsOverMemCap && !CacheTemplateNode(NodeSpawner, TemplateNode))
 			{
-				UE_LOG(LogBlueprintNodeCache, Display, TEXT("The blueprint template-node cache is full. As a result, you may experience slower than usual menu interactions. To avoid this, increase the cache's cap in the blueprint editor prefences."));
-				LoggedOnce = true;
+				LogCacheFullMsg();
 			}
 		}
 	}
@@ -282,19 +420,40 @@ UEdGraphNode* FBlueprintNodeTemplateCache::GetNodeTemplate(UBlueprintNodeSpawner
 //------------------------------------------------------------------------------
 void FBlueprintNodeTemplateCache::ClearCachedTemplate(UBlueprintNodeSpawner const* NodeSpawner)
 {
-	NodeTemplateCache.Remove(NodeSpawner); // GC should take care of the rest
+	NodeTemplateCache.Remove(NodeSpawner);
+	// GC should take care of the rest
 }
 
 //------------------------------------------------------------------------------
-int32 FBlueprintNodeTemplateCache::EstimateAllocatedSize() const
+int32 FBlueprintNodeTemplateCache::GetEstimateCacheSize() const
 {
-	int32 TotalEstimatedSize = NodeTemplateCache.GetAllocatedSize();
+	int32 TotalEstimatedSize = ApproximateObjectMem;
 	TotalEstimatedSize += TemplateOuters.GetAllocatedSize();
-	TotalEstimatedSize += ApproximateAllocationTotal;
+	TotalEstimatedSize += NodeTemplateCache.GetAllocatedSize();
+	TotalEstimatedSize += sizeof(*this);
 
 	return TotalEstimatedSize;
 }
 
+//------------------------------------------------------------------------------
+int32 FBlueprintNodeTemplateCache::RecalculateCacheSize()
+{
+	ApproximateObjectMem = 0;
+	for (UBlueprint* Blueprint : TemplateOuters)
+	{
+		// if we didn't run garbage collection at the top, then this could also
+		// account for nodes that were never stored (because the cache was too full)
+		ApproximateObjectMem += BlueprintNodeTemplateCacheImpl::ApproximateMemFootprint(Blueprint);
+	}
+	return ApproximateObjectMem;
+}
+
+//------------------------------------------------------------------------------
+bool FBlueprintNodeTemplateCache::IsTemplateOuter(UEdGraph* ParentGraph)
+{
+	return BlueprintNodeTemplateCacheImpl::IsTemplateOuter(ParentGraph);
+}
+
 //------------------------------------------------------------------------------
 void FBlueprintNodeTemplateCache::AddReferencedObjects(FReferenceCollector& Collector)
 {
@@ -306,19 +465,18 @@ void FBlueprintNodeTemplateCache::AddReferencedObjects(FReferenceCollector& Coll
 }
 
 //------------------------------------------------------------------------------
-bool FBlueprintNodeTemplateCache::AddBlueprintOuter(UBlueprint* Blueprint)
+bool FBlueprintNodeTemplateCache::CacheBlueprintOuter(UBlueprint* Blueprint)
 {
-	int32 const ApproxBlueprintSize = BlueprintNodeTemplateCacheImpl::ApproximateMemFootprint(Blueprint);
-	ApproximateAllocationTotal += ApproxBlueprintSize;
-	int32 const NewEstimatedSize = EstimateAllocatedSize();
+	using namespace BlueprintNodeTemplateCacheImpl;
+	int32 const ApproxBlueprintSize = ApproximateMemFootprint(Blueprint);
 
-	if (NewEstimatedSize > BlueprintNodeTemplateCacheImpl::GetCacheCapSize())
+	if (ActiveMemFootprint + ApproxBlueprintSize > GetCacheCapSize())
 	{
-		ApproximateAllocationTotal -= ApproxBlueprintSize;
 		return false;
 	}
 	else
 	{
+		ApproximateObjectMem += ApproxBlueprintSize;
 		TemplateOuters.Add(Blueprint);
 		return true;
 	}
@@ -327,22 +485,21 @@ bool FBlueprintNodeTemplateCache::AddBlueprintOuter(UBlueprint* Blueprint)
 //------------------------------------------------------------------------------
 bool FBlueprintNodeTemplateCache::CacheTemplateNode(UBlueprintNodeSpawner const* NodeSpawner, UEdGraphNode* NewNode)
 {
+	using namespace BlueprintNodeTemplateCacheImpl;
+
 	if (NewNode == nullptr)
 	{
 		return true;
 	}
 
 	int32 const ApproxNodeSize = BlueprintNodeTemplateCacheImpl::ApproximateMemFootprint(NewNode);
-	ApproximateAllocationTotal += ApproxNodeSize;
-	int32 const NewEstimatedSize = EstimateAllocatedSize();
-
-	if (NewEstimatedSize > BlueprintNodeTemplateCacheImpl::GetCacheCapSize())
+	if (ActiveMemFootprint + ApproxNodeSize > GetCacheCapSize())
 	{
-		ApproximateAllocationTotal -= ApproxNodeSize;
 		return false;
 	}
 	else
 	{
+		ApproximateObjectMem += ApproxNodeSize;
 		NodeTemplateCache.Add(NodeSpawner, NewNode);
 		return true;
 	}
diff --git a/Engine/Source/Editor/BlueprintGraph/Public/BlueprintEditorSettings.h b/Engine/Source/Editor/BlueprintGraph/Public/BlueprintEditorSettings.h
--- a/Engine/Source/Editor/BlueprintGraph/Public/BlueprintEditorSettings.h
+++ b/Engine/Source/Editor/BlueprintGraph/Public/BlueprintEditorSettings.h
@@ -19,7 +19,7 @@ public:
 // Perf Settings
 public:
 	/** The node template cache is used to speed up blueprint menuing. This determines the peak data size for that cache. */
-	UPROPERTY(EditAnywhere, config, Category="Performance", DisplayName="Node-Template Cache Cap (MB)", meta=(ClampMin="0.004", UIMin="0"))
+	UPROPERTY(EditAnywhere, config, Category="Performance", DisplayName="Node-Template Cache Cap (MB)", meta=(ClampMin="0", UIMin="0"))
 	float NodeTemplateCacheCapMB;
 };
  
\ No newline at end of file
diff --git a/Engine/Source/Editor/BlueprintGraph/Public/BlueprintNodeTemplateCache.h b/Engine/Source/Editor/BlueprintGraph/Public/BlueprintNodeTemplateCache.h
--- a/Engine/Source/Editor/BlueprintGraph/Public/BlueprintNodeTemplateCache.h
+++ b/Engine/Source/Editor/BlueprintGraph/Public/BlueprintNodeTemplateCache.h
@@ -55,11 +55,31 @@ public:
 	void ClearCachedTemplate(UBlueprintNodeSpawner const* NodeSpawner);
 
 	/**
+	 * Utility method to help external systems identify if a graph they have 
+	 * belongs here, to the FBlueprintNodeTemplateCache system.
 	 * 
+	 * @param  ParentGraph	The graph you want to check.
+	 * @return True if the graph was allocated by a FBlueprintNodeTemplateCache (to house template nodes).
+	 */
+	static bool IsTemplateOuter(UEdGraph* ParentGraph);
+
+	/**
+	 * Approximates the current memory footprint of the entire cache 
+	 * (instantiated UObject sizes + allocated container space).
+	 * 
+	 * @return The approximated total (in bytes) that this cache has allocated.
+	 */
+	int32 GetEstimateCacheSize() const;
+
+	/**
+	 * External systems can make changes that alter the memory footprint of the
+	 * cache (like calling AllocateDefaultPins), and since we don't recalculate
+	 * the cache's size every frame sometimes we need to update the internal 
+	 * estimate.
 	 * 
-	 * @return 
+	 * @return The new approximated total (in bytes) that this cache has allocated.
 	 */
-	int32 EstimateAllocatedSize() const;
+	int32 RecalculateCacheSize();
 
 	// FGCObject interface
 	virtual void AddReferencedObjects(FReferenceCollector& Collector) override;
@@ -67,22 +87,27 @@ public:
 
 private:
 	/**
+	 * Caches the supplied blueprint, so that it may be reused as an outer for
+	 * template nodes (certain nodes types assume they'll have a graph outer, 
+	 * with a blueprint outer beyond that, so we cannot just spawn nodes into 
+	 * the transient package).
 	 * 
-	 * 
-	 * @param  Blueprint	
-	 * @return 
+	 * @param  Blueprint	The blueprint you wish to store for reuse.
+	 * @return True if the blueprint was successfully cached (the cache could be full, and therefore this could fail).
 	 */
-	bool AddBlueprintOuter(UBlueprint* Blueprint);
+	bool CacheBlueprintOuter(UBlueprint* Blueprint);
 
 	/**
+	 * Attempts to cache the supplied node, and associates it with the specified 
+	 * spawner (so that we can remove it later if it is no longer needed).
 	 * 
-	 * 
-	 * @param  NodeSpawner	
-	 * @param  NewNode	
-	 * @return 
+	 * @param  NodeSpawner	Acts as the key for the given template node.
+	 * @param  NewNode		The template node you want stored.
+	 * @return True if the node was successfully cached (the cache could be full, and therefore this could fail).
 	 */
 	bool CacheTemplateNode(UBlueprintNodeSpawner const* NodeSpawner, UEdGraphNode* NewNode);
 
+private:
 	/** 
 	 * Unfortunately, we cannot nest template-nodes in the transient package. 
 	 * Certain nodes operate on the assumption that they have a UEdGraph outer, 
@@ -98,6 +123,11 @@ private:
 	/** */
 	TMap<UBlueprintNodeSpawner const*, UEdGraphNode*> NodeTemplateCache;
 
-	/** */
-	int32 ApproximateAllocationTotal;
+
+	/** 
+	 * It can be costly to tally back up the estimated cache size every time an
+	 * entry is added, so we keep this approximate tally of memory allocated for
+	 * UObjects (owned by this system).
+	 */
+	int32 ApproximateObjectMem;
 };
diff --git a/Engine/Source/Editor/UnrealEd/Private/Commandlets/DumpBlueprintsInfoCommandlet.cpp b/Engine/Source/Editor/UnrealEd/Private/Commandlets/DumpBlueprintsInfoCommandlet.cpp
--- a/Engine/Source/Editor/UnrealEd/Private/Commandlets/DumpBlueprintsInfoCommandlet.cpp
+++ b/Engine/Source/Editor/UnrealEd/Private/Commandlets/DumpBlueprintsInfoCommandlet.cpp
@@ -1239,7 +1239,7 @@ static bool DumpBlueprintInfoUtils::DumpActionDatabaseInfo(uint32 Indent, FArchi
 		EstimatedSystemSize += EstimatedCacheSize;
 
 		FString const NodeCacheSizeStr = BuildByteSizeString(EstimatedCacheSize);
-		FString const AvgNodeSizeStr = BuildByteSizeString(EstimatedCacheSize / TemplateCount);
+		FString const AvgNodeSizeStr = BuildByteSizeString((TemplateCount > 0) ? EstimatedCacheSize / TemplateCount : 0);
 		FString NodeCacheStats = FString::Printf(TEXT("%s\"Template-Cache Stats\" : {"), *IndentedNewline);
 		NodeCacheStats += FString::Printf(TEXT("%s\"TotalPrimingDuration\" : %.3f seconds,"), *DblIndentedNewline, TotalPrimingTime);
 		NodeCacheStats += FString::Printf(TEXT("%s\"CachedNodeCount\"      : %d,"), *DblIndentedNewline, TemplateCount);
