Filtering class dynamic casts using the same logic as dynamic casts. #codereview Mike.Beach

diff --git a/Engine/Source/Editor/BlueprintGraph/Private/BlueprintActionFilter.cpp b/Engine/Source/Editor/BlueprintGraph/Private/BlueprintActionFilter.cpp
--- a/Engine/Source/Editor/BlueprintGraph/Private/BlueprintActionFilter.cpp
+++ b/Engine/Source/Editor/BlueprintGraph/Private/BlueprintActionFilter.cpp
@@ -983,7 +983,7 @@ static bool BlueprintActionFilterImpl::IsNotSubClassCast(FBlueprintActionFilter
 {
 	bool bIsFilteredOut = false;
 
-	if (BlueprintAction.GetNodeClass() == UK2Node_DynamicCast::StaticClass())
+	if (BlueprintAction.GetNodeClass()->IsChildOf( UK2Node_DynamicCast::StaticClass() ))
 	{
 		for(UEdGraphPin const* ContextPin : Filter.Context.Pins)
 		{
