Moving the option to enable/disable the new menuing system into experimental editor preferences.

diff --git a/Engine/Source/Editor/BlueprintGraph/Classes/EdGraphSchema_K2.h b/Engine/Source/Editor/BlueprintGraph/Classes/EdGraphSchema_K2.h
--- a/Engine/Source/Editor/BlueprintGraph/Classes/EdGraphSchema_K2.h
+++ b/Engine/Source/Editor/BlueprintGraph/Classes/EdGraphSchema_K2.h
@@ -284,9 +284,9 @@ class BLUEPRINTGRAPH_API UEdGraphSchema_K2 : public UEdGraphSchema
 	// Somewhat hacky mechanism to prevent tooltips created for pins from including the display name and type when generating BP API documentation
 	static bool bGeneratingDocumentation;
 
-	/** Whether this schema should use the old (legacy) menu building, or the new (experimental) menu building */
+	/** DEPRECATED - Use UEditorExperimentalSettings::bUseRefactoredBlueprintMenuingSystem instead */
 	UPROPERTY(GlobalConfig)
-	bool bUseLegacyActionMenus;
+	bool bUseLegacyActionMenus; 
 
 	UPROPERTY(globalconfig)
 	TArray<FBlueprintCallableFunctionRedirect> EditoronlyBPFunctionRedirects;
diff --git a/Engine/Source/Editor/BlueprintGraph/Private/BlueprintEditorSettings.cpp b/Engine/Source/Editor/BlueprintGraph/Private/BlueprintEditorSettings.cpp
--- a/Engine/Source/Editor/BlueprintGraph/Private/BlueprintEditorSettings.cpp
+++ b/Engine/Source/Editor/BlueprintGraph/Private/BlueprintEditorSettings.cpp
@@ -2,10 +2,15 @@
 
 #include "BlueprintGraphPrivatePCH.h"
 #include "BlueprintEditorSettings.h"
+#include "Editor/UnrealEd/Classes/Settings/EditorExperimentalSettings.h"
 
 UBlueprintEditorSettings::UBlueprintEditorSettings(const class FPostConstructInitializeProperties& PCIP)
 	: Super(PCIP)
 {
 	bShowGraphInstructionText = true;
 	NodeTemplateCacheCapMB    = 20.f;
+
+	// settings that were moved out of experimental...
+	UEditorExperimentalSettings const* ExperimentalSettings = GetDefault<UEditorExperimentalSettings>();
+	bDrawMidpointArrowsInBlueprints = ExperimentalSettings->bDrawMidpointArrowsInBlueprints;
 }
diff --git a/Engine/Source/Editor/BlueprintGraph/Public/BlueprintEditorSettings.h b/Engine/Source/Editor/BlueprintGraph/Public/BlueprintEditorSettings.h
--- a/Engine/Source/Editor/BlueprintGraph/Public/BlueprintEditorSettings.h
+++ b/Engine/Source/Editor/BlueprintGraph/Public/BlueprintEditorSettings.h
@@ -2,14 +2,22 @@
 
 #pragma once
 
+#include "EdGraph/EdGraphPin.h" // for EBlueprintPinStyleType
 #include "BlueprintEditorSettings.generated.h"
 
+
 UCLASS(config=EditorUserSettings)
 class BLUEPRINTGRAPH_API UBlueprintEditorSettings
 	:	public UObject
 {
 	GENERATED_UCLASS_BODY()
 
+// Style Settings
+public:
+	/** Should arrows indicating data/execution flow be drawn halfway along wires? */
+	UPROPERTY(EditAnywhere, config, Category="Visual Style", meta=(DisplayName="Draw midpoint arrows in Blueprints"))
+	bool bDrawMidpointArrowsInBlueprints;
+
 // UX Settings
 public:
 	/** Determines if lightweight tutorial text shows up at the top of empty blueprint graphs */
@@ -19,7 +27,7 @@ public:
 // Perf Settings
 public:
 	/** The node template cache is used to speed up blueprint menuing. This determines the peak data size for that cache. */
-	UPROPERTY(EditAnywhere, config, Category="Performance", DisplayName="Node-Template Cache Cap (MB)", meta=(ClampMin="0", UIMin="0"))
+	UPROPERTY(EditAnywhere, AdvancedDisplay, config, Category="Performance", DisplayName="Node-Template Cache Cap (MB)", meta=(ClampMin="0", UIMin="0"))
 	float NodeTemplateCacheCapMB;
 };
  
\ No newline at end of file
diff --git a/Engine/Source/Editor/GraphEditor/Private/ConnectionDrawingPolicy.cpp b/Engine/Source/Editor/GraphEditor/Private/ConnectionDrawingPolicy.cpp
--- a/Engine/Source/Editor/GraphEditor/Private/ConnectionDrawingPolicy.cpp
+++ b/Engine/Source/Editor/GraphEditor/Private/ConnectionDrawingPolicy.cpp
@@ -14,6 +14,7 @@
 #include "AnimGraphNode_Base.h"
 #include "Sound/SoundNode.h"
 #include "K2Node_Knot.h"
+#include "BlueprintEditorSettings.h"
 
 DEFINE_LOG_CATEGORY_STATIC(LogConnectionDrawingPolicy, Log, All);
 
@@ -537,7 +538,7 @@ FKismetConnectionDrawingPolicy::FKismetConnectionDrawingPolicy(int32 InBackLayer
 	ArrowRadius = FVector2D::ZeroVector;
 
 	// But we do want to draw midpoint arrowheads
-	if (GetDefault<UEditorExperimentalSettings>()->bDrawMidpointArrowsInBlueprints)
+	if (GetDefault<UBlueprintEditorSettings>()->bDrawMidpointArrowsInBlueprints)
 	{
 		MidpointImage = FEditorStyle::GetBrush( TEXT("Graph.Arrow") );
 		MidpointRadius = MidpointImage->ImageSize * ZoomFactor * 0.5f;
diff --git a/Engine/Source/Editor/GraphEditor/Public/GraphEditorSettings.h b/Engine/Source/Editor/GraphEditor/Public/GraphEditorSettings.h
--- a/Engine/Source/Editor/GraphEditor/Public/GraphEditorSettings.h
+++ b/Engine/Source/Editor/GraphEditor/Public/GraphEditorSettings.h
@@ -21,7 +21,7 @@ class GRAPHEDITOR_API UGraphEditorSettings
 public:
 
 	/** The visual styling to use for graph editor pins (in Blueprints, materials, etc...) */
-	UPROPERTY(config, EditAnywhere, Category=Blueprint)
+	UPROPERTY(config, EditAnywhere, Category=GeneralStyle)
 	TEnumAsByte<EBlueprintPinStyleType> DataPinStyle;
 
 public:
diff --git a/Engine/Source/Editor/Kismet/Private/BlueprintActionMenuBuilder.cpp b/Engine/Source/Editor/Kismet/Private/BlueprintActionMenuBuilder.cpp
--- a/Engine/Source/Editor/Kismet/Private/BlueprintActionMenuBuilder.cpp
+++ b/Engine/Source/Editor/Kismet/Private/BlueprintActionMenuBuilder.cpp
@@ -661,7 +661,8 @@ void FBlueprintActionMenuBuilder::RebuildActionList()
 		MenuSection->Empty();
 	}
 	
-	if (!GetDefault<UEdGraphSchema_K2>()->bUseLegacyActionMenus)
+	UEditorExperimentalSettings const* ExperimentalSettings = GetDefault<UEditorExperimentalSettings>();
+	if (ExperimentalSettings->bUseRefactoredBlueprintMenuingSystem)
 	{
 		FBlueprintActionDatabase::FActionRegistry const& ActionDatabase = FBlueprintActionDatabase::Get().GetAllActions();
 		for (auto const& ActionEntry : ActionDatabase)
diff --git a/Engine/Source/Editor/Kismet/Private/BlueprintActionMenuUtils.cpp b/Engine/Source/Editor/Kismet/Private/BlueprintActionMenuUtils.cpp
--- a/Engine/Source/Editor/Kismet/Private/BlueprintActionMenuUtils.cpp
+++ b/Engine/Source/Editor/Kismet/Private/BlueprintActionMenuUtils.cpp
@@ -216,8 +216,8 @@ void FBlueprintActionMenuUtils::MakeContextMenu(FBlueprintActionContext const& C
 
 	MenuOut.RebuildActionList();
 
-	UEdGraphSchema_K2 const* const K2Schema = GetDefault<UEdGraphSchema_K2>();
-	if (!K2Schema->bUseLegacyActionMenus)
+	UEditorExperimentalSettings const* ExperimentalSettings =  GetDefault<UEditorExperimentalSettings>();
+	if (ExperimentalSettings->bUseRefactoredBlueprintMenuingSystem)
 	{
 		for (UEdGraph const* Graph : Context.Graphs)
 		{
diff --git a/Engine/Source/Editor/UnrealEd/Classes/Settings/EditorExperimentalSettings.h b/Engine/Source/Editor/UnrealEd/Classes/Settings/EditorExperimentalSettings.h
--- a/Engine/Source/Editor/UnrealEd/Classes/Settings/EditorExperimentalSettings.h
+++ b/Engine/Source/Editor/UnrealEd/Classes/Settings/EditorExperimentalSettings.h
@@ -53,9 +53,13 @@ public:
 	bool bBreakOnExceptions;
 
 	/** Should arrows indicating data/execution flow be drawn halfway along wires? */
-	UPROPERTY(EditAnywhere, config, Category=Blueprints, meta=(DisplayName="Draw midpoint arrows in Blueprints"))
+	UPROPERTY(/*EditAnywhere - deprecated (moved into UBlueprintEditorSettings), */config/*, Category=Blueprints, meta=(DisplayName="Draw midpoint arrows in Blueprints")*/)
 	bool bDrawMidpointArrowsInBlueprints;
 
+	/** Determines if the Blueprint editor should use the new (experimental) menu system, or the old (legacy) system */
+	UPROPERTY(EditAnywhere, config, Category=Blueprints, meta=(DisplayName="Use new Blueprint menuing system"))
+	bool bUseRefactoredBlueprintMenuingSystem;
+
 	/** Whether to show Audio Streaming options for SoundWaves (disabling will not stop all audio streaming) */
 	UPROPERTY(EditAnywhere, config, Category=Audio)
 	bool bShowAudioStreamingOptions;
diff --git a/Engine/Source/Editor/UnrealEd/Private/Settings/SettingsClasses.cpp b/Engine/Source/Editor/UnrealEd/Private/Settings/SettingsClasses.cpp
--- a/Engine/Source/Editor/UnrealEd/Private/Settings/SettingsClasses.cpp
+++ b/Engine/Source/Editor/UnrealEd/Private/Settings/SettingsClasses.cpp
@@ -2,6 +2,7 @@
 
 #include "UnrealEd.h"
 #include "ISourceControlModule.h"
+#include "EdGraphSchema_K2.h" // for bUseLegacyActionMenus
 
 
 /* UContentBrowserSettings interface
@@ -51,7 +52,10 @@ UDestructableMeshEditorSettings::UDestructableMeshEditorSettings( const class FP
 
 UEditorExperimentalSettings::UEditorExperimentalSettings( const class FPostConstructInitializeProperties& PCIP )
 	: Super(PCIP)
-{ }
+{ 
+	// setting was originally loaded from .ini into UEdGraphSchema_K2
+	bUseRefactoredBlueprintMenuingSystem = !GetDefault<UEdGraphSchema_K2>()->bUseLegacyActionMenus;
+}
 
 
 void UEditorExperimentalSettings::PostEditChangeProperty( struct FPropertyChangedEvent& PropertyChangedEvent )
