Fixing compiler error on mac.

diff --git a/Engine/Source/Editor/BlueprintGraph/Private/BlueprintNodeTemplateCache.cpp b/Engine/Source/Editor/BlueprintGraph/Private/BlueprintNodeTemplateCache.cpp
--- a/Engine/Source/Editor/BlueprintGraph/Private/BlueprintNodeTemplateCache.cpp
+++ b/Engine/Source/Editor/BlueprintGraph/Private/BlueprintNodeTemplateCache.cpp
@@ -159,10 +159,12 @@ static UBlueprint* BlueprintNodeTemplateCacheImpl::MakeCompatibleBlueprint(TSubc
 	UBlueprint* NewBlueprint = FKismetEditorUtilities::CreateBlueprint(ParentClass, BlueprintOuter, BlueprintName, BlueprintType, BlueprintClass, GeneratedClassType);
 	NewBlueprint->SetFlags(RF_Transient);
 
+	++MadeBlueprintCount;
+
 	float const AproxBlueprintSize = ApproximateMemFootprint(NewBlueprint);
 	// track the average blueprint size, so that we can attempt to predict 
 	// whether a  blueprint will fail to be cached (when the cache is near full)
-	AverageBlueprintSize = AverageBlueprintSize * ((float)MadeBlueprintCount / (++MadeBlueprintCount)) + 
+	AverageBlueprintSize = AverageBlueprintSize * ((float)(MadeBlueprintCount-1) / MadeBlueprintCount) + 
 		(AproxBlueprintSize / MadeBlueprintCount) + 0.5f;
 
 	return NewBlueprint;
