Filtering out class types from BP context menu when we have an actor as context

diff --git a/Engine/Source/Editor/BlueprintGraph/Private/BlueprintActionFilter.cpp b/Engine/Source/Editor/BlueprintGraph/Private/BlueprintActionFilter.cpp
--- a/Engine/Source/Editor/BlueprintGraph/Private/BlueprintActionFilter.cpp
+++ b/Engine/Source/Editor/BlueprintGraph/Private/BlueprintActionFilter.cpp
@@ -814,7 +814,8 @@ static bool BlueprintActionFilterImpl::IsPinCompatibleWithTargetSelf(UEdGraphPin
 			UBlueprint const* Blueprint = FBlueprintEditorUtils::FindBlueprintForNodeChecked(Pin->GetOwningNode());
 			PinObjClass = GetBlueprintClass(Blueprint);
 		}
-		else if (PinType.PinSubCategoryObject.IsValid()) // if PC_Object/PC_Interface
+		else if ( (PinType.PinCategory == UEdGraphSchema_K2::PC_Object || PinType.PinCategory == UEdGraphSchema_K2::PC_Interface)
+					&& PinType.PinSubCategoryObject.IsValid())
 		{
 			PinObjClass = Cast<UClass>(PinType.PinSubCategoryObject.Get());
 		}
