Fix compilation issues with Switch Enum Node.
Modify FKismetCompiledFunctionContext::ValidatePin to take a PinType instead of just the Category and SubCategory strings

diff --git a/Engine/Source/Editor/BlueprintGraph/Classes/K2Node_Switch.h b/Engine/Source/Editor/BlueprintGraph/Classes/K2Node_Switch.h
--- a/Engine/Source/Editor/BlueprintGraph/Classes/K2Node_Switch.h
+++ b/Engine/Source/Editor/BlueprintGraph/Classes/K2Node_Switch.h
@@ -45,7 +45,7 @@ class UK2Node_Switch : public UK2Node
 	virtual FString GetUniquePinName() { return FString(); }
 
 	/** Gets the pin type from the schema for the subclass */
-	virtual const FString& GetPinType(const UEdGraphSchema_K2* Schema) const { check(false); return Schema->PC_Int; }
+	virtual FEdGraphPinType GetPinType() const { check(false); return FEdGraphPinType(); }
 
 	/**
 	 * Adds a new execution pin to a switch node
diff --git a/Engine/Source/Editor/BlueprintGraph/Classes/K2Node_SwitchEnum.h b/Engine/Source/Editor/BlueprintGraph/Classes/K2Node_SwitchEnum.h
--- a/Engine/Source/Editor/BlueprintGraph/Classes/K2Node_SwitchEnum.h
+++ b/Engine/Source/Editor/BlueprintGraph/Classes/K2Node_SwitchEnum.h
@@ -42,7 +42,7 @@ class UK2Node_SwitchEnum : public UK2Node_Switch, public INodeDependingOnEnumInt
 
 	// UK2Node_Switch Interface
 	virtual FString GetUniquePinName() override;
-	virtual const FString& GetPinType(const UEdGraphSchema_K2* Schema) const override { return Schema->PC_Byte; }
+	virtual FEdGraphPinType GetPinType() const override;
 	virtual void AddPinToSwitchNode() override;
 	virtual void RemovePinFromSwitchNode(UEdGraphPin* TargetPin) override;
 	virtual ERedirectType DoPinsMatchForReconstruction(const UEdGraphPin* NewPin, int32 NewPinIndex, const UEdGraphPin* OldPin, int32 OldPinIndex) const override;
diff --git a/Engine/Source/Editor/BlueprintGraph/Classes/K2Node_SwitchInteger.h b/Engine/Source/Editor/BlueprintGraph/Classes/K2Node_SwitchInteger.h
--- a/Engine/Source/Editor/BlueprintGraph/Classes/K2Node_SwitchInteger.h
+++ b/Engine/Source/Editor/BlueprintGraph/Classes/K2Node_SwitchInteger.h
@@ -33,7 +33,7 @@ class UK2Node_SwitchInteger : public UK2Node_Switch
 	virtual FString GetUniquePinName() override;
 	virtual FString GetPinNameGivenIndex(int32 Index) override;
 	virtual void ReallocatePinsDuringReconstruction(TArray<UEdGraphPin*>& OldPins) override;
-	virtual const FString& GetPinType(const UEdGraphSchema_K2* Schema) const override { return Schema->PC_Int; }
+	virtual FEdGraphPinType GetPinType() const override;
 	// End of UK2Node_Switch Interface
 
 protected:
diff --git a/Engine/Source/Editor/BlueprintGraph/Classes/K2Node_SwitchName.h b/Engine/Source/Editor/BlueprintGraph/Classes/K2Node_SwitchName.h
--- a/Engine/Source/Editor/BlueprintGraph/Classes/K2Node_SwitchName.h
+++ b/Engine/Source/Editor/BlueprintGraph/Classes/K2Node_SwitchName.h
@@ -30,7 +30,7 @@ class UK2Node_SwitchName : public UK2Node_Switch
 	// UK2Node_Switch Interface
 	BLUEPRINTGRAPH_API virtual void AddPinToSwitchNode() override;
 	virtual FString GetUniquePinName() override;
-	virtual const FString& GetPinType(const UEdGraphSchema_K2* Schema) const override { return Schema->PC_Name; }
+	virtual FEdGraphPinType GetPinType() const override;
 	// End of UK2Node_Switch Interface
 
 	virtual FString GetPinNameGivenIndex(int32 Index) override;
diff --git a/Engine/Source/Editor/BlueprintGraph/Classes/K2Node_SwitchString.h b/Engine/Source/Editor/BlueprintGraph/Classes/K2Node_SwitchString.h
--- a/Engine/Source/Editor/BlueprintGraph/Classes/K2Node_SwitchString.h
+++ b/Engine/Source/Editor/BlueprintGraph/Classes/K2Node_SwitchString.h
@@ -33,7 +33,7 @@ class UK2Node_SwitchString : public UK2Node_Switch
 	// UK2Node_Switch Interface
 	BLUEPRINTGRAPH_API virtual void AddPinToSwitchNode() override;
 	virtual FString GetUniquePinName() override;
-	virtual const FString& GetPinType(const UEdGraphSchema_K2* Schema) const override { return Schema->PC_String; }
+	virtual FEdGraphPinType GetPinType() const override;
 	// End of UK2Node_Switch Interface
 
 	virtual FString GetPinNameGivenIndex(int32 Index) override;
diff --git a/Engine/Source/Editor/BlueprintGraph/Private/K2Node_ExecutionSequence.cpp b/Engine/Source/Editor/BlueprintGraph/Private/K2Node_ExecutionSequence.cpp
--- a/Engine/Source/Editor/BlueprintGraph/Private/K2Node_ExecutionSequence.cpp
+++ b/Engine/Source/Editor/BlueprintGraph/Private/K2Node_ExecutionSequence.cpp
@@ -24,8 +24,11 @@ public:
 	virtual void Compile(FKismetFunctionContext& Context, UEdGraphNode* Node) override
 	{
 		// Make sure that the input pin is connected and valid for this block
-		UEdGraphPin* ExecTriggeringPin = Context.FindRequiredPinByName(Node, CompilerContext.GetSchema()->PN_Execute, EGPD_Input);
-		if ((ExecTriggeringPin == NULL) || !Context.ValidatePinType(ExecTriggeringPin, CompilerContext.GetSchema()->PC_Exec))
+		FEdGraphPinType ExpectedPinType;
+		ExpectedPinType.PinCategory = UEdGraphSchema_K2::PC_Exec;
+
+		UEdGraphPin* ExecTriggeringPin = Context.FindRequiredPinByName(Node, UEdGraphSchema_K2::PN_Execute, EGPD_Input);
+		if ((ExecTriggeringPin == NULL) || !Context.ValidatePinType(ExecTriggeringPin, ExpectedPinType))
 		{
 			CompilerContext.MessageLog.Error(*FString::Printf(*LOCTEXT("NoValidExecutionPinForExecSeq_Error", "@@ must have a valid execution pin @@").ToString()), Node, ExecTriggeringPin);
 			return;
diff --git a/Engine/Source/Editor/BlueprintGraph/Private/K2Node_IfThenElse.cpp b/Engine/Source/Editor/BlueprintGraph/Private/K2Node_IfThenElse.cpp
--- a/Engine/Source/Editor/BlueprintGraph/Private/K2Node_IfThenElse.cpp
+++ b/Engine/Source/Editor/BlueprintGraph/Private/K2Node_IfThenElse.cpp
@@ -23,8 +23,15 @@ public:
 	virtual void Compile(FKismetFunctionContext& Context, UEdGraphNode* Node) override
 	{
 		// For imperative nodes, make sure the exec function was actually triggered and not just included due to an output data dependency
-		UEdGraphPin* ExecTriggeringPin = Context.FindRequiredPinByName(Node, CompilerContext.GetSchema()->PN_Execute, EGPD_Input);
-		if ((ExecTriggeringPin == NULL) || !Context.ValidatePinType(ExecTriggeringPin, CompilerContext.GetSchema()->PC_Exec))
+		FEdGraphPinType ExpectedExecPinType;
+		ExpectedExecPinType.PinCategory = UEdGraphSchema_K2::PC_Exec;
+
+		FEdGraphPinType ExpectedBoolPinType;
+		ExpectedBoolPinType.PinCategory = UEdGraphSchema_K2::PC_Boolean;
+
+
+		UEdGraphPin* ExecTriggeringPin = Context.FindRequiredPinByName(Node, UEdGraphSchema_K2::PN_Execute, EGPD_Input);
+		if ((ExecTriggeringPin == NULL) || !Context.ValidatePinType(ExecTriggeringPin, ExpectedExecPinType))
 		{
 			CompilerContext.MessageLog.Error(*FString::Printf(*LOCTEXT("NoValidExecutionPinForBranch_Error", "@@ must have a valid execution pin @@").ToString()), Node, ExecTriggeringPin);
 			return;
@@ -39,9 +46,9 @@ public:
 		UEdGraphPin* CondPin = Context.FindRequiredPinByName(Node, CompilerContext.GetSchema()->PN_Condition, EGPD_Input);
 		UEdGraphPin* ThenPin = Context.FindRequiredPinByName(Node, CompilerContext.GetSchema()->PN_Then, EGPD_Output);
 		UEdGraphPin* ElsePin = Context.FindRequiredPinByName(Node, CompilerContext.GetSchema()->PN_Else, EGPD_Output);
-		if (Context.ValidatePinType(ThenPin, CompilerContext.GetSchema()->PC_Exec) &&
-			Context.ValidatePinType(ElsePin, CompilerContext.GetSchema()->PC_Exec) &&
-			Context.ValidatePinType(CondPin, CompilerContext.GetSchema()->PC_Boolean))
+		if (Context.ValidatePinType(ThenPin, ExpectedExecPinType) &&
+			Context.ValidatePinType(ElsePin, ExpectedExecPinType) &&
+			Context.ValidatePinType(CondPin, ExpectedBoolPinType))
 		{
 			UEdGraphPin* PinToTry = FEdGraphUtilities::GetNetFromPin(CondPin);
 			FBPTerminal** CondTerm = Context.NetMap.Find(PinToTry);
diff --git a/Engine/Source/Editor/BlueprintGraph/Private/K2Node_Switch.cpp b/Engine/Source/Editor/BlueprintGraph/Private/K2Node_Switch.cpp
--- a/Engine/Source/Editor/BlueprintGraph/Private/K2Node_Switch.cpp
+++ b/Engine/Source/Editor/BlueprintGraph/Private/K2Node_Switch.cpp
@@ -25,9 +25,8 @@ protected:
 	TMap<UEdGraphNode*, FBPTerminal*> BoolTermMap;
 
 public:
-	FKCHandler_Switch(FKismetCompilerContext& InCompilerContext, const FString& InConnectionPinType)
+	FKCHandler_Switch(FKismetCompilerContext& InCompilerContext)
 		: FNodeHandlingFunctor(InCompilerContext)
-		, ConnectionPinType(InConnectionPinType)
 	{
 	}
 
@@ -51,9 +50,12 @@ public:
 	{
 		UK2Node_Switch* SwitchNode = CastChecked<UK2Node_Switch>(Node);
 
+		FEdGraphPinType ExpectedExecPinType;
+		ExpectedExecPinType.PinCategory = UEdGraphSchema_K2::PC_Exec;
+
 		// Make sure that the input pin is connected and valid for this block
-		UEdGraphPin* ExecTriggeringPin = Context.FindRequiredPinByName(SwitchNode, CompilerContext.GetSchema()->PN_Execute, EGPD_Input);
-		if ((ExecTriggeringPin == NULL) || !Context.ValidatePinType(ExecTriggeringPin, CompilerContext.GetSchema()->PC_Exec))
+		UEdGraphPin* ExecTriggeringPin = Context.FindRequiredPinByName(SwitchNode, UEdGraphSchema_K2::PN_Execute, EGPD_Input);
+		if ((ExecTriggeringPin == NULL) || !Context.ValidatePinType(ExecTriggeringPin, ExpectedExecPinType))
 		{
 			CompilerContext.MessageLog.Error(*FString::Printf(*LOCTEXT("NoValidExecutionPinForSwitch_Error", "@@ must have a valid execution pin @@").ToString()), SwitchNode, ExecTriggeringPin);
 			return;
@@ -61,7 +63,7 @@ public:
 
 		// Make sure that the selection pin is connected and valid for this block
 		UEdGraphPin* SelectionPin = SwitchNode->GetSelectionPin();
-		if ((SelectionPin == NULL) || !Context.ValidatePinType(SelectionPin, ConnectionPinType))
+		if ((SelectionPin == NULL) || !Context.ValidatePinType(SelectionPin, SwitchNode->GetPinType()))
 		{
 			CompilerContext.MessageLog.Error(*FString::Printf(*LOCTEXT("NoValidSelectionPinForSwitch_Error", "@@ must have a valid execution pin @@").ToString()), SwitchNode, SelectionPin);
 			return;
@@ -137,7 +139,7 @@ public:
 	}
 
 private:
-	FString ConnectionPinType;
+	FEdGraphPinType ExpectedSelectionPinType;
 };
 
 UK2Node_Switch::UK2Node_Switch(const class FPostConstructInitializeProperties& PCIP)
@@ -316,7 +318,7 @@ UEdGraphPin* UK2Node_Switch::GetDefaultPin()
 
 FNodeHandlingFunctor* UK2Node_Switch::CreateNodeHandler(FKismetCompilerContext& CompilerContext) const
 {
-	return new FKCHandler_Switch(CompilerContext, GetPinType(CompilerContext.GetSchema()));
+	return new FKCHandler_Switch(CompilerContext);
 }
 
 FText UK2Node_Switch::GetMenuCategory() const
diff --git a/Engine/Source/Editor/BlueprintGraph/Private/K2Node_SwitchEnum.cpp b/Engine/Source/Editor/BlueprintGraph/Private/K2Node_SwitchEnum.cpp
--- a/Engine/Source/Editor/BlueprintGraph/Private/K2Node_SwitchEnum.cpp
+++ b/Engine/Source/Editor/BlueprintGraph/Private/K2Node_SwitchEnum.cpp
@@ -131,6 +131,14 @@ void UK2Node_SwitchEnum::CreateSelectionPin()
 	K2Schema->SetPinDefaultValueBasedOnType(Pin);
 }
 
+FEdGraphPinType UK2Node_SwitchEnum::GetPinType() const
+{ 
+	FEdGraphPinType EnumPinType;
+	EnumPinType.PinCategory = UEdGraphSchema_K2::PC_Byte;
+	EnumPinType.PinSubCategoryObject = Enum;
+	return EnumPinType; 
+}
+
 FString UK2Node_SwitchEnum::GetUniquePinName()
 {
 	FString NewPinName;
diff --git a/Engine/Source/Editor/BlueprintGraph/Private/K2Node_SwitchInteger.cpp b/Engine/Source/Editor/BlueprintGraph/Private/K2Node_SwitchInteger.cpp
--- a/Engine/Source/Editor/BlueprintGraph/Private/K2Node_SwitchInteger.cpp
+++ b/Engine/Source/Editor/BlueprintGraph/Private/K2Node_SwitchInteger.cpp
@@ -51,6 +51,13 @@ void UK2Node_SwitchInteger::CreateSelectionPin()
 	K2Schema->SetPinDefaultValueBasedOnType(Pin);
 }
 
+FEdGraphPinType UK2Node_SwitchInteger::GetPinType() const 
+{ 
+	FEdGraphPinType PinType;
+	PinType.PinSubCategory = UEdGraphSchema_K2::PC_Int;
+	return PinType;
+}
+
 FString UK2Node_SwitchInteger::GetPinNameGivenIndex(int32 Index)
 {
 	const UEdGraphSchema_K2* K2Schema = GetDefault<UEdGraphSchema_K2>();
diff --git a/Engine/Source/Editor/BlueprintGraph/Private/K2Node_SwitchName.cpp b/Engine/Source/Editor/BlueprintGraph/Private/K2Node_SwitchName.cpp
--- a/Engine/Source/Editor/BlueprintGraph/Private/K2Node_SwitchName.cpp
+++ b/Engine/Source/Editor/BlueprintGraph/Private/K2Node_SwitchName.cpp
@@ -66,6 +66,13 @@ void UK2Node_SwitchName::CreateSelectionPin()
 	K2Schema->SetPinDefaultValueBasedOnType(Pin);
 }
 
+FEdGraphPinType UK2Node_SwitchName::GetPinType() const 
+{ 
+	FEdGraphPinType PinType;
+	PinType.PinSubCategory = UEdGraphSchema_K2::PC_Name;
+	return PinType;
+}
+
 FString UK2Node_SwitchName::GetPinNameGivenIndex(int32 Index)
 {
 	check(Index);
diff --git a/Engine/Source/Editor/BlueprintGraph/Private/K2Node_SwitchString.cpp b/Engine/Source/Editor/BlueprintGraph/Private/K2Node_SwitchString.cpp
--- a/Engine/Source/Editor/BlueprintGraph/Private/K2Node_SwitchString.cpp
+++ b/Engine/Source/Editor/BlueprintGraph/Private/K2Node_SwitchString.cpp
@@ -78,6 +78,13 @@ void UK2Node_SwitchString::CreateSelectionPin()
 	K2Schema->SetPinDefaultValueBasedOnType(Pin);
 }
 
+FEdGraphPinType UK2Node_SwitchString::GetPinType() const 
+{ 
+	FEdGraphPinType PinType;
+	PinType.PinSubCategory = UEdGraphSchema_K2::PC_String;
+	return PinType;
+}
+
 FString UK2Node_SwitchString::GetPinNameGivenIndex(int32 Index)
 {
 	check(Index);
diff --git a/Engine/Source/Editor/KismetCompiler/Public/KismetCompiledFunctionContext.h b/Engine/Source/Editor/KismetCompiler/Public/KismetCompiledFunctionContext.h
--- a/Engine/Source/Editor/KismetCompiler/Public/KismetCompiledFunctionContext.h
+++ b/Engine/Source/Editor/KismetCompiler/Public/KismetCompiledFunctionContext.h
@@ -508,7 +508,7 @@ public:
 	}
 
 	/** Checks to see if a pin is of the requested type */
-	bool ValidatePinType(const UEdGraphPin* Pin, const FString& Category, const FString& SubCategory = TEXT(""))
+	bool ValidatePinType(const UEdGraphPin* Pin, const FEdGraphPinType& TestType)
 	{
 		if (Pin == NULL)
 		{
@@ -517,11 +517,6 @@ public:
 		}
 		else
 		{
-			FEdGraphPinType TestType;
-			TestType.PinCategory = Category;
-			TestType.PinSubCategory = SubCategory;
-
-			check(Pin);
 			if (Pin->PinType == TestType)
 			{
 				return true;
