Adding a custom event is not listed as "Add Custom Event..." instead of "Event".

diff --git a/Engine/Source/Editor/BlueprintGraph/Private/BlueprintEventNodeSpawner.cpp b/Engine/Source/Editor/BlueprintGraph/Private/BlueprintEventNodeSpawner.cpp
--- a/Engine/Source/Editor/BlueprintGraph/Private/BlueprintEventNodeSpawner.cpp
+++ b/Engine/Source/Editor/BlueprintGraph/Private/BlueprintEventNodeSpawner.cpp
@@ -151,8 +151,15 @@ FText UBlueprintEventNodeSpawner::GetDefaultMenuName() const
 		{
 			EventName = FText::FromName(CustomEventName);
 		}
+		else
+		{
+			CachedMenuName = LOCTEXT("AddCustomEvent", "Add Custom Event...");
+		}
 
-		CachedMenuName = FText::Format(LOCTEXT("EventWithSignatureName", "Event {0}"), EventName);
+		if(CachedMenuName.IsOutOfDate())
+		{
+			CachedMenuName = FText::Format(LOCTEXT("EventWithSignatureName", "Event {0}"), EventName);
+		}
 	}
 	return CachedMenuName;
 }
