Set PinCategory, not PinSubCategory for SwitchInteger/Name/String

diff --git a/Engine/Source/Editor/BlueprintGraph/Private/K2Node_SwitchInteger.cpp b/Engine/Source/Editor/BlueprintGraph/Private/K2Node_SwitchInteger.cpp
--- a/Engine/Source/Editor/BlueprintGraph/Private/K2Node_SwitchInteger.cpp
+++ b/Engine/Source/Editor/BlueprintGraph/Private/K2Node_SwitchInteger.cpp
@@ -54,7 +54,7 @@ void UK2Node_SwitchInteger::CreateSelectionPin()
 FEdGraphPinType UK2Node_SwitchInteger::GetPinType() const 
 { 
 	FEdGraphPinType PinType;
-	PinType.PinSubCategory = UEdGraphSchema_K2::PC_Int;
+	PinType.PinCategory = UEdGraphSchema_K2::PC_Int;
 	return PinType;
 }
 
diff --git a/Engine/Source/Editor/BlueprintGraph/Private/K2Node_SwitchName.cpp b/Engine/Source/Editor/BlueprintGraph/Private/K2Node_SwitchName.cpp
--- a/Engine/Source/Editor/BlueprintGraph/Private/K2Node_SwitchName.cpp
+++ b/Engine/Source/Editor/BlueprintGraph/Private/K2Node_SwitchName.cpp
@@ -69,7 +69,7 @@ void UK2Node_SwitchName::CreateSelectionPin()
 FEdGraphPinType UK2Node_SwitchName::GetPinType() const 
 { 
 	FEdGraphPinType PinType;
-	PinType.PinSubCategory = UEdGraphSchema_K2::PC_Name;
+	PinType.PinCategory = UEdGraphSchema_K2::PC_Name;
 	return PinType;
 }
 
diff --git a/Engine/Source/Editor/BlueprintGraph/Private/K2Node_SwitchString.cpp b/Engine/Source/Editor/BlueprintGraph/Private/K2Node_SwitchString.cpp
--- a/Engine/Source/Editor/BlueprintGraph/Private/K2Node_SwitchString.cpp
+++ b/Engine/Source/Editor/BlueprintGraph/Private/K2Node_SwitchString.cpp
@@ -81,7 +81,7 @@ void UK2Node_SwitchString::CreateSelectionPin()
 FEdGraphPinType UK2Node_SwitchString::GetPinType() const 
 { 
 	FEdGraphPinType PinType;
-	PinType.PinSubCategory = UEdGraphSchema_K2::PC_String;
+	PinType.PinCategory = UEdGraphSchema_K2::PC_String;
 	return PinType;
 }
 
