Blueprint - Preventing crashes when the target function fails to be located when checking for WorldContext metadata.

diff --git a/Engine/Source/Editor/BlueprintGraph/Private/EdGraphSchema_K2.cpp b/Engine/Source/Editor/BlueprintGraph/Private/EdGraphSchema_K2.cpp
--- a/Engine/Source/Editor/BlueprintGraph/Private/EdGraphSchema_K2.cpp
+++ b/Engine/Source/Editor/BlueprintGraph/Private/EdGraphSchema_K2.cpp
@@ -3286,8 +3286,11 @@ bool UEdGraphSchema_K2::ShouldShowAssetPickerForPin(UEdGraphPin* Pin) const
 			{
 				if (UK2Node_CallFunction* CallFunctionNode = Cast<UK2Node_CallFunction>(Pin->GetOwningNode()))
 				{
-					const UEdGraphPin* WorldContextPin = CallFunctionNode->FindPin(CallFunctionNode->GetTargetFunction()->GetMetaData(FBlueprintMetadata::MD_WorldContext));
-					bShow = (WorldContextPin != Pin);
+					if ( UFunction* FunctionRef = CallFunctionNode->GetTargetFunction() )
+					{
+						const UEdGraphPin* WorldContextPin = CallFunctionNode->FindPin(FunctionRef->GetMetaData(FBlueprintMetadata::MD_WorldContext));
+						bShow = ( WorldContextPin != Pin );
+					}
 				}
 			}
 		}
