Fix for filtering the vast majority of variables out of default actor context.... we were treating every variable as public by default

diff --git a/Engine/Source/Editor/BlueprintGraph/Private/BlueprintActionFilter.cpp b/Engine/Source/Editor/BlueprintGraph/Private/BlueprintActionFilter.cpp
--- a/Engine/Source/Editor/BlueprintGraph/Private/BlueprintActionFilter.cpp
+++ b/Engine/Source/Editor/BlueprintGraph/Private/BlueprintActionFilter.cpp
@@ -391,22 +391,45 @@ static bool BlueprintActionFilterImpl::IsFieldInaccessible(FBlueprintActionFilte
 		bool bIsProtected = ClassField->HasMetaData(FBlueprintMetadata::MD_Protected);
 		bool bIsPrivate   = ClassField->HasMetaData(FBlueprintMetadata::MD_Private);
 
-		UClass* FieldOwner = ClassField->GetOwnerClass();
-		for (UBlueprint const* Blueprint : FilterContext.Blueprints)
+		bool bIsPublic = !bIsPrivate && !bIsProtected;
+		if (UProperty const* Property = Cast<UProperty>(ClassField))
 		{
-			UClass* BpClass = Blueprint->GeneratedClass;
-			check(BpClass != nullptr);
-			
-			// private functions are only accessible from the class they belong to
-			if (bIsPrivate && (FieldOwner != BpClass))
+			// CPF_DisableEditOnInstance corresponds to the eyeball/editable checkbox available 
+			// in the blueprint editor. It is poorly named. When CPF_DisableEditOnInstance is 
+			// *not* set, the user has tagged the variable as 'editable' and it is treated
+			// as public:
+			bIsPublic = !Property->HasAnyPropertyFlags(CPF_DisableEditOnInstance);
+
+			// If the variable is not public, and is not tagged as private, it should be protected
+			// by default. This branch handles variables declared in blueprints:
+			if( !bIsPublic )
 			{
-				bIsFilteredOut = true;
-				break;
+				if( !bIsPrivate )
+				{
+					bIsProtected = true; 
+				}
 			}
-			else if (bIsProtected && !BpClass->IsChildOf(FieldOwner))
+		}
+
+		if( !bIsPublic )
+		{
+			UClass* FieldOwner = ClassField->GetOwnerClass();
+			for (UBlueprint const* Blueprint : FilterContext.Blueprints)
 			{
-				bIsFilteredOut = true;
-				break;
+				UClass* BpClass = Blueprint->SkeletonGeneratedClass;
+				check(BpClass != nullptr);
+			
+				// private functions are only accessible from the class they belong to
+				if (bIsPrivate && (FieldOwner != BpClass))
+				{
+					bIsFilteredOut = true;
+					break;
+				}
+				else if (bIsProtected && !BpClass->IsChildOf(FieldOwner))
+				{
+					bIsFilteredOut = true;
+					break;
+				}
 			}
 		}
 	}
