Subcategory object can be NULL for some structs when we run the blueprint dump commandlet

diff --git a/Engine/Source/Editor/BlueprintGraph/Private/K2ActionMenuBuilder.cpp b/Engine/Source/Editor/BlueprintGraph/Private/K2ActionMenuBuilder.cpp
--- a/Engine/Source/Editor/BlueprintGraph/Private/K2ActionMenuBuilder.cpp
+++ b/Engine/Source/Editor/BlueprintGraph/Private/K2ActionMenuBuilder.cpp
@@ -1341,29 +1341,33 @@ void FK2ActionMenuBuilder::GetPinAllowedNodeTypes(FBlueprintGraphActionListBuild
 			if ((K2Schema->PC_Struct == FromPin.PinType.PinCategory) && !FromPin.PinType.bIsArray)
 			{
 				UScriptStruct* PinStruct = Cast<UScriptStruct>(FromPin.PinType.PinSubCategoryObject.Get());
-				const bool bCanBeMade = UK2Node_MakeStruct::CanBeMade(PinStruct);
-				if ((FromPin.Direction == EGPD_Output) && UK2Node_BreakStruct::CanBeBroken(PinStruct))
+				if( PinStruct )
 				{
-					UK2Node_BreakStruct* NodeTemplate = ContextMenuBuilder.CreateTemplateNode<UK2Node_BreakStruct>();
-					NodeTemplate->StructType = PinStruct;
-					TSharedPtr<FEdGraphSchemaAction_K2NewNode> Action = FK2ActionMenuBuilder::AddNewNodeAction(ContextMenuBuilder, FString(), NodeTemplate->GetNodeTitle(ENodeTitleType::ListView), NodeTemplate->GetTooltipText().ToString(), 0, NodeTemplate->GetKeywords());
-					Action->NodeTemplate = NodeTemplate;
-				}
-				else if ((FromPin.Direction == EGPD_Input) && bCanBeMade)
-				{
-					UK2Node_MakeStruct* NodeTemplate = ContextMenuBuilder.CreateTemplateNode<UK2Node_MakeStruct>();
-					NodeTemplate->StructType = PinStruct;
-					TSharedPtr<FEdGraphSchemaAction_K2NewNode> Action = FK2ActionMenuBuilder::AddNewNodeAction(ContextMenuBuilder, FString(), NodeTemplate->GetNodeTitle(ENodeTitleType::ListView), NodeTemplate->GetTooltipText().ToString(), 0, NodeTemplate->GetKeywords());
-					Action->NodeTemplate = NodeTemplate;
-				}
+					const bool bCanBeMade = UK2Node_MakeStruct::CanBeMade(PinStruct);
+					if ((FromPin.Direction == EGPD_Output) && UK2Node_BreakStruct::CanBeBroken(PinStruct))
+					{
+						UK2Node_BreakStruct* NodeTemplate = ContextMenuBuilder.CreateTemplateNode<UK2Node_BreakStruct>();
+						NodeTemplate->StructType = PinStruct;
+						TSharedPtr<FEdGraphSchemaAction_K2NewNode> Action = FK2ActionMenuBuilder::AddNewNodeAction(ContextMenuBuilder, FString(), NodeTemplate->GetNodeTitle(ENodeTitleType::ListView), NodeTemplate->GetTooltipText().ToString(), 0, NodeTemplate->GetKeywords());
+						Action->NodeTemplate = NodeTemplate;
+					}
+					else if ((FromPin.Direction == EGPD_Input) && bCanBeMade)
+					{
+						UK2Node_MakeStruct* NodeTemplate = ContextMenuBuilder.CreateTemplateNode<UK2Node_MakeStruct>();
+						NodeTemplate->StructType = PinStruct;
+						TSharedPtr<FEdGraphSchemaAction_K2NewNode> Action = FK2ActionMenuBuilder::AddNewNodeAction(ContextMenuBuilder, FString(), NodeTemplate->GetNodeTitle(ENodeTitleType::ListView), NodeTemplate->GetTooltipText().ToString(), 0, NodeTemplate->GetKeywords());
+						Action->NodeTemplate = NodeTemplate;
+					}
 
-				if ((FromPin.Direction == EGPD_Output) && bCanBeMade)
-				{
-					UK2Node_SetFieldsInStruct* NodeTemplate = ContextMenuBuilder.CreateTemplateNode<UK2Node_SetFieldsInStruct>();
-					NodeTemplate->StructType = PinStruct;
-					TSharedPtr<FEdGraphSchemaAction_K2NewNode> Action = FK2ActionMenuBuilder::AddNewNodeAction(ContextMenuBuilder, FString(), NodeTemplate->GetNodeTitle(ENodeTitleType::ListView), NodeTemplate->GetTooltipText().ToString(), 0, NodeTemplate->GetKeywords());
-					Action->NodeTemplate = NodeTemplate;
+					if ((FromPin.Direction == EGPD_Output) && bCanBeMade)
+					{
+						UK2Node_SetFieldsInStruct* NodeTemplate = ContextMenuBuilder.CreateTemplateNode<UK2Node_SetFieldsInStruct>();
+						NodeTemplate->StructType = PinStruct;
+						TSharedPtr<FEdGraphSchemaAction_K2NewNode> Action = FK2ActionMenuBuilder::AddNewNodeAction(ContextMenuBuilder, FString(), NodeTemplate->GetNodeTitle(ENodeTitleType::ListView), NodeTemplate->GetTooltipText().ToString(), 0, NodeTemplate->GetKeywords());
+						Action->NodeTemplate = NodeTemplate;
+					}
 				}
+
 			}
 
 			if ((FromPin.Direction == EGPD_Input) && (FromPin.PinType.PinCategory == K2Schema->PC_Delegate))
@@ -1396,6 +1400,7 @@ void FK2ActionMenuBuilder::GetPinAllowedNodeTypes(FBlueprintGraphActionListBuild
 		}
 		else if (FromPin.PinType.PinCategory == K2Schema->PC_Float ||
 			(FromPin.PinType.PinCategory == K2Schema->PC_Struct &&
+			FromPin.PinType.PinSubCategoryObject.IsValid() && 
 			(FromPin.PinType.PinSubCategoryObject.Get()->GetName() == TEXT("Vector") ||
 			FromPin.PinType.PinSubCategoryObject.Get()->GetName() == TEXT("Rotator") ||
 			FromPin.PinType.PinSubCategoryObject.Get()->GetName() == TEXT("Transform"))))
