Fix compilation with enum switch nodes

diff --git a/Engine/Source/Editor/BlueprintGraph/Private/K2Node_Switch.cpp b/Engine/Source/Editor/BlueprintGraph/Private/K2Node_Switch.cpp
--- a/Engine/Source/Editor/BlueprintGraph/Private/K2Node_Switch.cpp
+++ b/Engine/Source/Editor/BlueprintGraph/Private/K2Node_Switch.cpp
@@ -88,13 +88,6 @@ public:
 			UFunction* FunctionPtr = FindField<UFunction>(FuncClass, *FuncPin->PinName);
 			check(FunctionPtr);
 
-			// Find the enum object for the switch node if it's an enum switch
-			UEnum* SelectionEnum = NULL;
-			if (UK2Node_SwitchEnum* SwitchNodeEnum = Cast<UK2Node_SwitchEnum>(SwitchNode))
-			{
-				SelectionEnum = SwitchNodeEnum->Enum;
-			}
-
 			// Run thru all the output pins except for the default label
 			for (auto PinIt = SwitchNode->Pins.CreateIterator(); PinIt; ++PinIt)
 			{
@@ -104,7 +97,7 @@ public:
 				{
 					// Create a term for the switch case value
 					FBPTerminal* CaseValueTerm = new (Context.Literals) FBPTerminal();
-					CaseValueTerm->Name = (SelectionEnum != NULL) ? FString::FromInt(SelectionEnum->FindEnumIndex(*(Pin->PinName))) : Pin->PinName;
+					CaseValueTerm->Name = Pin->PinName;
 					CaseValueTerm->Type = SelectionPin->PinType;
 					CaseValueTerm->Source = Pin;
 					CaseValueTerm->bIsLiteral = true;
