Guarding (and not blocking) failed pin type lookups... now with more informative error text (to help catch the problematic category).
(TTP #346394)

diff --git a/Engine/Source/Editor/BlueprintGraph/Private/EdGraphSchema_K2.cpp b/Engine/Source/Editor/BlueprintGraph/Private/EdGraphSchema_K2.cpp
--- a/Engine/Source/Editor/BlueprintGraph/Private/EdGraphSchema_K2.cpp
+++ b/Engine/Source/Editor/BlueprintGraph/Private/EdGraphSchema_K2.cpp
@@ -2651,7 +2651,15 @@ FText UEdGraphSchema_K2::GetCategoryText(const FString& Category, const bool bFo
 		}
 	}
 
-	return CategoryDescriptions.FindChecked(Category);
+	if (FText const* TypeDesc = CategoryDescriptions.Find(Category))
+	{
+		return *TypeDesc;
+	}
+	else
+	{
+		ensureMsgf(false, TEXT("Unhandled blueprint type ('%s'), failed to find type description."), *Category);
+		return FText::FromString(Category);
+	}
 }
 
 FText UEdGraphSchema_K2::TypeToText(const FEdGraphPinType& Type)
