ttp346772 LIVE: MAIN: BP: Renaming a variable inside of a struct will disconnect the pin in any blueprint it is used in

diff --git a/Engine/Source/Editor/BlueprintGraph/Private/K2Node_BreakStruct.cpp b/Engine/Source/Editor/BlueprintGraph/Private/K2Node_BreakStruct.cpp
--- a/Engine/Source/Editor/BlueprintGraph/Private/K2Node_BreakStruct.cpp
+++ b/Engine/Source/Editor/BlueprintGraph/Private/K2Node_BreakStruct.cpp
@@ -280,16 +280,20 @@ FLinearColor UK2Node_BreakStruct::GetNodeTitleColor() const
 UK2Node::ERedirectType UK2Node_BreakStruct::DoPinsMatchForReconstruction(const UEdGraphPin* NewPin, int32 NewPinIndex, const UEdGraphPin* OldPin, int32 OldPinIndex)  const
 {
 	ERedirectType Result = UK2Node::DoPinsMatchForReconstruction(NewPin, NewPinIndex, OldPin, OldPinIndex);
-	if ((ERedirectType_None == Result) && NewPin && OldPin)
+	if ((ERedirectType_None == Result) && DoRenamedPinsMatch(NewPin, OldPin, true))
 	{
-		if ((EGPD_Input == NewPin->Direction) && (EGPD_Input == OldPin->Direction))
+		Result = ERedirectType_Custom;
+	}
+	else if ((ERedirectType_None == Result) && NewPin && OldPin)
 	{
-		const UEdGraphSchema_K2* K2Schema = GetDefault<UEdGraphSchema_K2>();
-		if (K2Schema->ArePinTypesCompatible( NewPin->PinType, OldPin->PinType))
+		if ((EGPD_Input == NewPin->Direction) && (EGPD_Input == OldPin->Direction))
 		{
-			Result = ERedirectType_Custom;
+			const UEdGraphSchema_K2* K2Schema = GetDefault<UEdGraphSchema_K2>();
+			if (K2Schema->ArePinTypesCompatible( NewPin->PinType, OldPin->PinType))
+			{
+				Result = ERedirectType_Custom;
+			}
 		}
-	}
 		else if ((EGPD_Output == NewPin->Direction) && (EGPD_Output == OldPin->Direction))
 		{
 			TMap<FName, FName>* StructRedirects = UStruct::TaggedPropertyRedirects.Find(StructType->GetFName());
diff --git a/Engine/Source/Editor/BlueprintGraph/Private/K2Node_MakeStruct.cpp b/Engine/Source/Editor/BlueprintGraph/Private/K2Node_MakeStruct.cpp
--- a/Engine/Source/Editor/BlueprintGraph/Private/K2Node_MakeStruct.cpp
+++ b/Engine/Source/Editor/BlueprintGraph/Private/K2Node_MakeStruct.cpp
@@ -239,7 +239,11 @@ FNodeHandlingFunctor* UK2Node_MakeStruct::CreateNodeHandler(FKismetCompilerConte
 UK2Node::ERedirectType UK2Node_MakeStruct::DoPinsMatchForReconstruction(const UEdGraphPin* NewPin, int32 NewPinIndex, const UEdGraphPin* OldPin, int32 OldPinIndex)  const
 {
 	ERedirectType Result = UK2Node::DoPinsMatchForReconstruction(NewPin, NewPinIndex, OldPin, OldPinIndex);
-	if ((ERedirectType_None == Result) && NewPin && OldPin)
+	if ((ERedirectType_None == Result) && DoRenamedPinsMatch(NewPin, OldPin, false))
+	{
+		Result = ERedirectType_Custom;
+	}
+	else if ((ERedirectType_None == Result) && NewPin && OldPin)
 	{
 		const UEdGraphSchema_K2* K2Schema = GetDefault<UEdGraphSchema_K2>();
 		if ((EGPD_Output == NewPin->Direction) && (EGPD_Output == OldPin->Direction))
