Fixing it so that make/break struct nodes show up in the palette (for the new BP menu system).

diff --git a/Engine/Source/Editor/BlueprintGraph/Private/BlueprintActionFilter.cpp b/Engine/Source/Editor/BlueprintGraph/Private/BlueprintActionFilter.cpp
--- a/Engine/Source/Editor/BlueprintGraph/Private/BlueprintActionFilter.cpp
+++ b/Engine/Source/Editor/BlueprintGraph/Private/BlueprintActionFilter.cpp
@@ -172,12 +172,13 @@ namespace BlueprintActionFilterImpl
 	 * Rejection test that checks to see if the node-spawner would produce a 
 	 * node type that isn't white-listed. 
 	 * 
-	 * @param  Filter			Holds the class context for this test.
-	 * @param  BlueprintAction	The action you wish to query.
-	 * @param  bMatchExplicitly When true, the node's class has to perfectly match a white-listed type (sub-classes don't count).
+	 * @param  Filter				Holds the class context for this test.
+	 * @param  BlueprintAction		The action you wish to query.
+	 * @param  bPermitChildClasses	When true, the node's class doesn't have to perfectly match a class to be passed (it can be a sub-class).
+	 * @param  bRejectChildClasses	When true, the node's class doesn't have to perfectly match a class to be rejected (it can be a sub-class).
 	 * @return True if the action would produce a non-whitelisted node.
 	 */
-	static bool IsFilteredNodeType(FBlueprintActionFilter const& Filter, FBlueprintActionInfo& BlueprintAction, bool const bMatchExplicitly = false);
+	static bool IsFilteredNodeType(FBlueprintActionFilter const& Filter, FBlueprintActionInfo& BlueprintAction, bool const bPermitChildClasses, bool const bRejectChildClasses);
 
 	/**
 	 * Rejection test that checks to see if the node-spawner is tied to a 
@@ -662,7 +663,7 @@ static bool BlueprintActionFilterImpl::IsIncompatibleWithGraphType(FBlueprintAct
 }
 
 //------------------------------------------------------------------------------
-static bool BlueprintActionFilterImpl::IsFilteredNodeType(FBlueprintActionFilter const& Filter, FBlueprintActionInfo& BlueprintAction, bool const bExcludeChildClasses)
+static bool BlueprintActionFilterImpl::IsFilteredNodeType(FBlueprintActionFilter const& Filter, FBlueprintActionInfo& BlueprintAction, bool const bPermitChildClasses, bool const bRejectChildClasses)
 {
 	bool bIsFilteredOut = (Filter.PermittedNodeTypes.Num() > 0);
 
@@ -671,8 +672,8 @@ static bool BlueprintActionFilterImpl::IsFilteredNodeType(FBlueprintActionFilter
 	{
 		for (TSubclassOf<UEdGraphNode> AllowedClass : Filter.PermittedNodeTypes)
 		{
-			if ((NodeClass->IsChildOf(AllowedClass) && !bExcludeChildClasses) ||
-				(AllowedClass == NodeClass))
+			if ((AllowedClass == NodeClass) || 
+				(bPermitChildClasses && NodeClass->IsChildOf(AllowedClass)))
 			{
 				bIsFilteredOut = false;
 				break;
@@ -682,8 +683,8 @@ static bool BlueprintActionFilterImpl::IsFilteredNodeType(FBlueprintActionFilter
 		for (int32 ClassIndx = 0; !bIsFilteredOut && (ClassIndx < Filter.RejectedNodeTypes.Num()); ++ClassIndx)
 		{
 			TSubclassOf<UEdGraphNode> ExcludedClass = Filter.RejectedNodeTypes[ClassIndx];
-			if ((bExcludeChildClasses && (ExcludedClass == NodeClass)) ||
-				NodeClass->IsChildOf(ExcludedClass))
+			if ((ExcludedClass == NodeClass) ||
+				(bRejectChildClasses && NodeClass->IsChildOf(ExcludedClass)))
 			{
 				bIsFilteredOut = true;
 				break;
@@ -1205,7 +1206,7 @@ FBlueprintActionFilter::FBlueprintActionFilter(uint32 Flags/*= 0x00*/)
 	AddRejectionTest(FRejectionTestDelegate::CreateStatic(IsIncompatibleImpureNode));
 	
 	AddRejectionTest(FRejectionTestDelegate::CreateStatic(IsFieldCategoryHidden));
-	if (Flags & BPFILTER_RejectPersistentNonTargetFields)
+	if (Flags & BPFILTER_RejectGlobalFields)
 	{
 		AddRejectionTest(FRejectionTestDelegate::CreateStatic(IsPersistentNonTargetField));
 	}
@@ -1223,8 +1224,8 @@ FBlueprintActionFilter::FBlueprintActionFilter(uint32 Flags/*= 0x00*/)
 		AddRejectionTest(FRejectionTestDelegate::CreateStatic(IsDeprecated));
 	}
 
-	AddRejectionTest(FRejectionTestDelegate::CreateStatic(IsFilteredNodeType, (Flags & BPFILTER_RejectPermittedNodeSubClasses) != 0));
-	AddRejectionTest(FRejectionTestDelegate::CreateStatic(IsNonTargetMemeber, (Flags & BPFILTER_RejectPersistentNonTargetFields) == 0));
+	AddRejectionTest(FRejectionTestDelegate::CreateStatic(IsFilteredNodeType, !(Flags & BPFILTER_RejectPermittedSubClasses), !(Flags & BPFILTER_PermitRejectionSubClasses)));
+	AddRejectionTest(FRejectionTestDelegate::CreateStatic(IsNonTargetMemeber, !(Flags & BPFILTER_RejectGlobalFields)));
 	AddRejectionTest(FRejectionTestDelegate::CreateStatic(IsUnBoundBindingSpawner));
 	AddRejectionTest(FRejectionTestDelegate::CreateStatic(IsOutOfScopeLocalVariable));
 
diff --git a/Engine/Source/Editor/BlueprintGraph/Public/BlueprintActionFilter.h b/Engine/Source/Editor/BlueprintGraph/Public/BlueprintActionFilter.h
--- a/Engine/Source/Editor/BlueprintGraph/Public/BlueprintActionFilter.h
+++ b/Engine/Source/Editor/BlueprintGraph/Public/BlueprintActionFilter.h
@@ -166,7 +166,7 @@ public:
 	enum EFlags // Flags, which configure certain rejection tests.
 	{
 		/** Deprecated class actions will not be filtered out*/
-		BPFILTER_PermitDeprecated				= 0x00000001,
+		BPFILTER_PermitDeprecated			= (1<<0),
 
 		/** 
 		 * Rejects actions associated with global/static ("persistent") fields.
@@ -174,14 +174,21 @@ public:
 		 * rejected by this (if TargetClasses is left empty, then all 
 		 * "persistent" fields are rejected).
 		 */
-		BPFILTER_RejectPersistentNonTargetFields= 0x00000002,
+		BPFILTER_RejectGlobalFields			= (1<<1),
 
 		/**
 		 * Makes PermittedNodeType tests more aggressive by rejecting node 
 		 * sub-classes, (actions would have to explicitly match a class listed
 		 * in PermittedNodeTypes). 
 		 */
-		BPFILTER_RejectPermittedNodeSubClasses	= 0x00000004,
+		BPFILTER_RejectPermittedSubClasses	= (1<<2),
+
+		/**
+		 * Makes RejectedNodeType tests less aggressive by permitting node 
+		 * sub-classes, (actions would have to explicitly match a class listed
+		 * in RejectedNodeTypes). 
+		 */
+		BPFILTER_PermitRejectionSubClasses	= (1<<3),
 	};
 	FBlueprintActionFilter(uint32 const Flags = 0x00);
 	
diff --git a/Engine/Source/Editor/Kismet/Private/BlueprintActionMenuUtils.cpp b/Engine/Source/Editor/Kismet/Private/BlueprintActionMenuUtils.cpp
--- a/Engine/Source/Editor/Kismet/Private/BlueprintActionMenuUtils.cpp
+++ b/Engine/Source/Editor/Kismet/Private/BlueprintActionMenuUtils.cpp
@@ -6,7 +6,8 @@
 #include "BlueprintActionFilter.h"
 #include "BlueprintNodeSpawner.h"
 #include "KismetEditorUtilities.h"	// for CanPasteNodes()
-#include "K2Node_Variable.h"
+#include "K2Node_VariableGet.h"
+#include "K2Node_VariableSet.h"
 #include "K2Node_CallFunction.h"
 #include "K2Node_ActorBoundEvent.h"
 #include "K2Node_ComponentBoundEvent.h"
@@ -100,7 +101,7 @@ void FBlueprintActionMenuUtils::MakePaletteMenu(FBlueprintActionContext const& C
 	if (FilterClass != nullptr)
 	{
 		// make sure we exclude global and static library actions
-		FilterFlags = FBlueprintActionFilter::BPFILTER_RejectPersistentNonTargetFields;
+		FilterFlags |= FBlueprintActionFilter::BPFILTER_RejectGlobalFields;
 	}
 	
 	FBlueprintActionFilter MenuFilter(FilterFlags);
@@ -109,7 +110,8 @@ void FBlueprintActionMenuUtils::MakePaletteMenu(FBlueprintActionContext const& C
 	// self member variables can be accessed through the MyBluprint panel (even
 	// inherited ones)... external variables can be accessed through the context
 	// menu (don't want to clutter the palette, I guess?)
-	MenuFilter.RejectedNodeTypes.Add(UK2Node_Variable::StaticClass());
+	MenuFilter.RejectedNodeTypes.Add(UK2Node_VariableGet::StaticClass());
+	MenuFilter.RejectedNodeTypes.Add(UK2Node_VariableSet::StaticClass());
 	
 	if (FilterClass != nullptr)
 	{
