typo fix.. memeber => member

diff --git a/Engine/Source/Editor/BlueprintGraph/Private/BlueprintActionFilter.cpp b/Engine/Source/Editor/BlueprintGraph/Private/BlueprintActionFilter.cpp
--- a/Engine/Source/Editor/BlueprintGraph/Private/BlueprintActionFilter.cpp
+++ b/Engine/Source/Editor/BlueprintGraph/Private/BlueprintActionFilter.cpp
@@ -156,7 +156,7 @@ namespace BlueprintActionFilterImpl
 	 * @param  bPermitNonTargetGlobals	Determines if this test should pass for external global/static fields.
 	 * @return True if the action is associated with a non-whitelisted class member.
 	 */
-	static bool IsNonTargetMemeber(FBlueprintActionFilter const& Filter, FBlueprintActionInfo& BlueprintAction, bool const bPermitNonTargetGlobals);
+	static bool IsNonTargetMember(FBlueprintActionFilter const& Filter, FBlueprintActionInfo& BlueprintAction, bool const bPermitNonTargetGlobals);
 
 	/**
 	 * Rejection test that checks to see if the node-spawner is associated with 
@@ -544,7 +544,7 @@ static bool BlueprintActionFilterImpl::IsPersistentNonTargetField(FBlueprintActi
 }
 
 //------------------------------------------------------------------------------
-static bool BlueprintActionFilterImpl::IsNonTargetMemeber(FBlueprintActionFilter const& Filter, FBlueprintActionInfo& BlueprintAction, bool bPermitNonTargetGlobals)
+static bool BlueprintActionFilterImpl::IsNonTargetMember(FBlueprintActionFilter const& Filter, FBlueprintActionInfo& BlueprintAction, bool bPermitNonTargetGlobals)
 {
 	bool bIsFilteredOut = false;
 	if (UField const* ClassField = BlueprintAction.GetAssociatedMemberField())
@@ -1225,7 +1225,7 @@ FBlueprintActionFilter::FBlueprintActionFilter(uint32 Flags/*= 0x00*/)
 	}
 
 	AddRejectionTest(FRejectionTestDelegate::CreateStatic(IsFilteredNodeType, !(Flags & BPFILTER_RejectPermittedSubClasses), !(Flags & BPFILTER_PermitRejectionSubClasses)));
-	AddRejectionTest(FRejectionTestDelegate::CreateStatic(IsNonTargetMemeber, !(Flags & BPFILTER_RejectGlobalFields)));
+	AddRejectionTest(FRejectionTestDelegate::CreateStatic(IsNonTargetMember, !(Flags & BPFILTER_RejectGlobalFields)));
 	AddRejectionTest(FRejectionTestDelegate::CreateStatic(IsUnBoundBindingSpawner));
 	AddRejectionTest(FRejectionTestDelegate::CreateStatic(IsOutOfScopeLocalVariable));
 
diff --git a/Engine/Source/Editor/BlueprintGraph/Private/K2Node_CreateDelegate.cpp b/Engine/Source/Editor/BlueprintGraph/Private/K2Node_CreateDelegate.cpp
--- a/Engine/Source/Editor/BlueprintGraph/Private/K2Node_CreateDelegate.cpp
+++ b/Engine/Source/Editor/BlueprintGraph/Private/K2Node_CreateDelegate.cpp
@@ -90,9 +90,9 @@ bool UK2Node_CreateDelegate::IsValid(FString* OutMsg, bool bDontUseSkeletalClass
 		return false;
 	}
 
-	FMemberReference MemeberReference;
-	MemeberReference.SetDirect(SelectedFunctionName, SelectedFunctionGuid, ScopeClass, false);
-	const UFunction* FoundFunction = MemeberReference.ResolveMember<UFunction>((UClass*) NULL);
+	FMemberReference MemberReference;
+	MemberReference.SetDirect(SelectedFunctionName, SelectedFunctionGuid, ScopeClass, false);
+	const UFunction* FoundFunction = MemberReference.ResolveMember<UFunction>((UClass*) NULL);
 	if (!FoundFunction)
 	{
 		if (OutMsg)
