Enum/Struct spawners can now be favorited in the new menu system.

diff --git a/Engine/Source/Editor/BlueprintGraph/Classes/BlueprintFieldNodeSpawner.h b/Engine/Source/Editor/BlueprintGraph/Classes/BlueprintFieldNodeSpawner.h
--- a/Engine/Source/Editor/BlueprintGraph/Classes/BlueprintFieldNodeSpawner.h
+++ b/Engine/Source/Editor/BlueprintGraph/Classes/BlueprintFieldNodeSpawner.h
@@ -37,7 +37,7 @@ public:
 	virtual UEdGraphNode* Invoke(UEdGraph* ParentGraph, FBindingSet const& Bindings, FVector2D const Location) const override;
 	// End UBlueprintNodeSpawner interface
 
-	/** */
+	/** Callback to define how the field should be applied to new nodes */
 	FSetNodeFieldDelegate SetNodeFieldDelegate;
 
 	/**
diff --git a/Engine/Source/Editor/BlueprintGraph/Private/BlueprintActionDatabase.cpp b/Engine/Source/Editor/BlueprintGraph/Private/BlueprintActionDatabase.cpp
--- a/Engine/Source/Editor/BlueprintGraph/Private/BlueprintActionDatabase.cpp
+++ b/Engine/Source/Editor/BlueprintGraph/Private/BlueprintActionDatabase.cpp
@@ -125,6 +125,7 @@ static UBlueprintNodeSpawner* FBlueprintNodeSpawnerFactory::MakeMacroNodeSpawner
 	};
 
 	TWeakObjectPtr<UEdGraph> GraphPtr = MacroGraph;
+	// @TODO: Needs access to the MacroGraph to be favoritable...
 	NodeSpawner->CustomizeNodeDelegate = UBlueprintNodeSpawner::FCustomizeNodeDelegate::CreateStatic(CustomizeMacroNodeLambda, GraphPtr);
 
 	return NodeSpawner;
@@ -136,21 +137,15 @@ static UBlueprintNodeSpawner* FBlueprintNodeSpawnerFactory::MakeMessageNodeSpawn
 	check(InterfaceFunction != nullptr);
 	check(FKismetEditorUtilities::IsClassABlueprintInterface(CastChecked<UClass>(InterfaceFunction->GetOuter())));
 
-	UBlueprintNodeSpawner* NodeSpawner = UBlueprintNodeSpawner::Create(UK2Node_Message::StaticClass());
+	UBlueprintFunctionNodeSpawner* NodeSpawner = UBlueprintFunctionNodeSpawner::Create(UK2Node_Message::StaticClass(), InterfaceFunction);
 	check(NodeSpawner != nullptr);
 
-	auto CustomizeMessageNodeLambda = [](UEdGraphNode* NewNode, bool bIsTemplateNode, TWeakObjectPtr<UFunction> FunctionPtr)
+	auto SetNodeFunctionLambda = [](UEdGraphNode* NewNode, UField const* FuncField)
 	{
 		UK2Node_Message* MessageNode = CastChecked<UK2Node_Message>(NewNode);
-		if (FunctionPtr.IsValid())
-		{
-			UFunction* Function = FunctionPtr.Get();
-			MessageNode->FunctionReference.SetExternalMember(Function->GetFName(), Cast<UClass>(Function->GetOuter()));
-		}
+		MessageNode->FunctionReference.SetExternalMember(FuncField->GetFName(), FuncField->GetOwnerClass());
 	};
-
-	TWeakObjectPtr<UFunction> FunctionPtr = InterfaceFunction;
-	NodeSpawner->CustomizeNodeDelegate = UBlueprintNodeSpawner::FCustomizeNodeDelegate::CreateStatic(CustomizeMessageNodeLambda, FunctionPtr);
+	NodeSpawner->SetNodeFieldDelegate = UBlueprintFunctionNodeSpawner::FSetNodeFieldDelegate::CreateStatic(SetNodeFunctionLambda);
 
 	return NodeSpawner;
 }
@@ -495,6 +490,7 @@ static void BlueprintActionDatabaseImpl::AddClassCastActions(UClass* const Class
 		};
 
 		UBlueprintNodeSpawner* CastObjNodeSpawner = UBlueprintNodeSpawner::Create(UK2Node_DynamicCast::StaticClass());
+		// @TODO: Needs access to the Class to be favoritable...
 		CastObjNodeSpawner->CustomizeNodeDelegate = UBlueprintNodeSpawner::FCustomizeNodeDelegate::CreateStatic(CustomizeCastNodeLambda, Class);
 		ActionListOut.Add(CastObjNodeSpawner);
 
diff --git a/Engine/Source/Editor/BlueprintGraph/Private/K2Node_BreakStruct.cpp b/Engine/Source/Editor/BlueprintGraph/Private/K2Node_BreakStruct.cpp
--- a/Engine/Source/Editor/BlueprintGraph/Private/K2Node_BreakStruct.cpp
+++ b/Engine/Source/Editor/BlueprintGraph/Private/K2Node_BreakStruct.cpp
@@ -2,7 +2,7 @@
 
 #include "BlueprintGraphPrivatePCH.h"
 #include "KismetCompiler.h"
-#include "BlueprintNodeSpawner.h"
+#include "BlueprintFieldNodeSpawner.h"
 #include "EditorCategoryUtils.h"
 #include "BlueprintActionDatabaseRegistrar.h"
 
@@ -317,13 +317,10 @@ FNodeHandlingFunctor* UK2Node_BreakStruct::CreateNodeHandler(class FKismetCompil
 
 void UK2Node_BreakStruct::GetMenuActions(FBlueprintActionDatabaseRegistrar& ActionRegistrar) const
 {
-	auto CustomizeBreakNodeLambda = [](UEdGraphNode* NewNode, bool bIsTemplateNode, TWeakObjectPtr<UScriptStruct> StructPtr)
+	auto SetNodeStructLambda = [](UEdGraphNode* NewNode, UField const* /*StructField*/, TWeakObjectPtr<UScriptStruct> NonConstStructPtr)
 	{
 		UK2Node_BreakStruct* BreakNode = CastChecked<UK2Node_BreakStruct>(NewNode);
-		if (StructPtr.IsValid())
-		{
-			BreakNode->StructType = StructPtr.Get();
-		}
+		BreakNode->StructType = NonConstStructPtr.Get();
 	};
 
 	for (TObjectIterator<UScriptStruct> StructIt; StructIt; ++StructIt)
@@ -334,7 +331,7 @@ void UK2Node_BreakStruct::GetMenuActions(FBlueprintActionDatabaseRegistrar& Acti
 			continue;
 		}
 
-		// to keep from needlessly instantiating a UBlueprintNodeSpawner, first   
+		// to keep from needlessly instantiating a UBlueprintNodeSpawners, first   
 		// check to make sure that the registrar is looking for actions of this type
 		// (could be regenerating actions for a specific asset, and therefore the 
 		// registrar would only accept actions corresponding to that asset)
@@ -343,15 +340,15 @@ void UK2Node_BreakStruct::GetMenuActions(FBlueprintActionDatabaseRegistrar& Acti
 			continue;
 		}
 
-		UBlueprintNodeSpawner* NodeSpawner = UBlueprintNodeSpawner::Create(GetClass());
+		UBlueprintFieldNodeSpawner* NodeSpawner = UBlueprintFieldNodeSpawner::Create(GetClass(), Struct);
 		check(NodeSpawner != nullptr);
+		TWeakObjectPtr<UScriptStruct> NonConstStructPtr = Struct;
+		NodeSpawner->SetNodeFieldDelegate = UBlueprintFieldNodeSpawner::FSetNodeFieldDelegate::CreateStatic(SetNodeStructLambda, NonConstStructPtr);
+
 		// this struct could belong to a class, or is a user defined struct 
 		// (asset), that's why we want to make sure to register it along with 
 		// the action (so the action knows to refresh when the class/asset is).
 		ActionRegistrar.AddBlueprintAction(Struct, NodeSpawner);
-
-		TWeakObjectPtr<UScriptStruct> StructPtr = Struct;
-		NodeSpawner->CustomizeNodeDelegate = UBlueprintNodeSpawner::FCustomizeNodeDelegate::CreateStatic(CustomizeBreakNodeLambda, StructPtr);
 	}	
 }
 
diff --git a/Engine/Source/Editor/BlueprintGraph/Private/K2Node_CastByteToEnum.cpp b/Engine/Source/Editor/BlueprintGraph/Private/K2Node_CastByteToEnum.cpp
--- a/Engine/Source/Editor/BlueprintGraph/Private/K2Node_CastByteToEnum.cpp
+++ b/Engine/Source/Editor/BlueprintGraph/Private/K2Node_CastByteToEnum.cpp
@@ -5,7 +5,7 @@
 #include "KismetCompiler.h"
 #include "Kismet/KismetNodeHelperLibrary.h"
 #include "K2Node_CastByteToEnum.h"
-#include "BlueprintNodeSpawner.h"
+#include "BlueprintFieldNodeSpawner.h"
 #include "EditorCategoryUtils.h"
 #include "BlueprintActionDatabaseRegistrar.h"
 
@@ -163,13 +163,10 @@ bool UK2Node_CastByteToEnum::IsConnectionDisallowed(const UEdGraphPin* MyPin, co
 
 void UK2Node_CastByteToEnum::GetMenuActions(FBlueprintActionDatabaseRegistrar& ActionRegistrar) const
 {
-	auto CustomizeEnumNodeLambda = [](UEdGraphNode* NewNode, bool bIsTemplateNode, TWeakObjectPtr<UEnum> EnumPtr)
+	auto SetNodeEnumLambda = [](UEdGraphNode* NewNode, UField const* /*EnumField*/, TWeakObjectPtr<UEnum> NonConstEnumPtr)
 	{
 		UK2Node_CastByteToEnum* EnumNode = CastChecked<UK2Node_CastByteToEnum>(NewNode);
-		if (EnumPtr.IsValid())
-		{
-			EnumNode->Enum = EnumPtr.Get();
-		}
+		EnumNode->Enum  = NonConstEnumPtr.Get();
 		EnumNode->bSafe = true;
 	};
 
@@ -181,7 +178,7 @@ void UK2Node_CastByteToEnum::GetMenuActions(FBlueprintActionDatabaseRegistrar& A
 			continue;
 		}
 
-		// to keep from needlessly instantiating a UBlueprintNodeSpawner, first   
+		// to keep from needlessly instantiating a UBlueprintNodeSpawners, first   
 		// check to make sure that the registrar is looking for actions of this type
 		// (could be regenerating actions for a specific asset, and therefore the 
 		// registrar would only accept actions corresponding to that asset)
@@ -190,15 +187,15 @@ void UK2Node_CastByteToEnum::GetMenuActions(FBlueprintActionDatabaseRegistrar& A
 			continue;
 		}
 
-		UBlueprintNodeSpawner* NodeSpawner = UBlueprintNodeSpawner::Create(GetClass());
+		UBlueprintFieldNodeSpawner* NodeSpawner = UBlueprintFieldNodeSpawner::Create(GetClass(), Enum);
 		check(NodeSpawner != nullptr);
+		TWeakObjectPtr<UEnum> NonConstEnumPtr = Enum;
+		NodeSpawner->SetNodeFieldDelegate = UBlueprintFieldNodeSpawner::FSetNodeFieldDelegate::CreateStatic(SetNodeEnumLambda, NonConstEnumPtr);
+
 		// this enum could belong to a class, or is a user defined enum (asset), 
 		// that's why we want to make sure to register it along with the action 
 		// (so the action can be refreshed when the class/asset is).
 		ActionRegistrar.AddBlueprintAction(Enum, NodeSpawner);
-
-		TWeakObjectPtr<UEnum> EnumPtr = Enum;
-		NodeSpawner->CustomizeNodeDelegate = UBlueprintNodeSpawner::FCustomizeNodeDelegate::CreateStatic(CustomizeEnumNodeLambda, EnumPtr);
 	}
 }
 
diff --git a/Engine/Source/Editor/BlueprintGraph/Private/K2Node_EnumLiteral.cpp b/Engine/Source/Editor/BlueprintGraph/Private/K2Node_EnumLiteral.cpp
--- a/Engine/Source/Editor/BlueprintGraph/Private/K2Node_EnumLiteral.cpp
+++ b/Engine/Source/Editor/BlueprintGraph/Private/K2Node_EnumLiteral.cpp
@@ -4,7 +4,7 @@
 #include "BlueprintGraphPrivatePCH.h"
 #include "KismetCompilerMisc.h"
 #include "K2Node_EnumLiteral.h"
-#include "BlueprintNodeSpawner.h"
+#include "BlueprintFieldNodeSpawner.h"
 #include "EditorCategoryUtils.h"
 #include "BlueprintActionDatabaseRegistrar.h"
 
@@ -103,13 +103,10 @@ FNodeHandlingFunctor* UK2Node_EnumLiteral::CreateNodeHandler(FKismetCompilerCont
 
 void UK2Node_EnumLiteral::GetMenuActions(FBlueprintActionDatabaseRegistrar& ActionRegistrar) const
 {
-	auto CustomizeEnumNodeLambda = [](UEdGraphNode* NewNode, bool bIsTemplateNode, TWeakObjectPtr<UEnum> EnumPtr)
+	auto SetNodeEnumLambda = [](UEdGraphNode* NewNode, UField const* /*EnumField*/, TWeakObjectPtr<UEnum> NonConstEnumPtr)
 	{
 		UK2Node_EnumLiteral* EnumNode = CastChecked<UK2Node_EnumLiteral>(NewNode);
-		if (EnumPtr.IsValid())
-		{
-			EnumNode->Enum = EnumPtr.Get();
-		}
+		EnumNode->Enum = NonConstEnumPtr.Get();
 	};
 
 	for (TObjectIterator<UEnum> EnumIt; EnumIt; ++EnumIt)
@@ -120,7 +117,7 @@ void UK2Node_EnumLiteral::GetMenuActions(FBlueprintActionDatabaseRegistrar& Acti
 			continue;
 		}
 
-		// to keep from needlessly instantiating a UBlueprintNodeSpawner, first   
+		// to keep from needlessly instantiating a UBlueprintNodeSpawners, first   
 		// check to make sure that the registrar is looking for actions of this type
 		// (could be regenerating actions for a specific asset, and therefore the 
 		// registrar would only accept actions corresponding to that asset)
@@ -129,15 +126,15 @@ void UK2Node_EnumLiteral::GetMenuActions(FBlueprintActionDatabaseRegistrar& Acti
 			continue;
 		}
 
-		UBlueprintNodeSpawner* NodeSpawner = UBlueprintNodeSpawner::Create(GetClass());
+		UBlueprintFieldNodeSpawner* NodeSpawner = UBlueprintFieldNodeSpawner::Create(GetClass(), Enum);
 		check(NodeSpawner != nullptr);
+		TWeakObjectPtr<UEnum> NonConstEnumPtr = Enum;
+		NodeSpawner->SetNodeFieldDelegate = UBlueprintFieldNodeSpawner::FSetNodeFieldDelegate::CreateStatic(SetNodeEnumLambda, NonConstEnumPtr);
+
 		// this enum could belong to a class, or is a user defined enum (asset), 
 		// that's why we want to make sure to register it along with the action 
 		// (so the action can be refreshed when the class/asset is).
 		ActionRegistrar.AddBlueprintAction(Enum, NodeSpawner);
-
-		TWeakObjectPtr<UEnum> EnumPtr = Enum;
-		NodeSpawner->CustomizeNodeDelegate = UBlueprintNodeSpawner::FCustomizeNodeDelegate::CreateStatic(CustomizeEnumNodeLambda, EnumPtr);
 	}
 }
 
diff --git a/Engine/Source/Editor/BlueprintGraph/Private/K2Node_ForEachElementInEnum.cpp b/Engine/Source/Editor/BlueprintGraph/Private/K2Node_ForEachElementInEnum.cpp
--- a/Engine/Source/Editor/BlueprintGraph/Private/K2Node_ForEachElementInEnum.cpp
+++ b/Engine/Source/Editor/BlueprintGraph/Private/K2Node_ForEachElementInEnum.cpp
@@ -8,7 +8,7 @@
 #include "K2Node_CastByteToEnum.h"
 #include "K2Node_ForEachElementInEnum.h"
 #include "K2Node_GetNumEnumEntries.h"
-#include "BlueprintNodeSpawner.h"
+#include "BlueprintFieldNodeSpawner.h"
 #include "EditorCategoryUtils.h"
 #include "BlueprintActionDatabaseRegistrar.h"
 
@@ -205,13 +205,10 @@ void UK2Node_ForEachElementInEnum::ExpandNode(class FKismetCompilerContext& Comp
 
 void UK2Node_ForEachElementInEnum::GetMenuActions(FBlueprintActionDatabaseRegistrar& ActionRegistrar) const
 {
-	auto CustomizeEnumNodeLambda = [](UEdGraphNode* NewNode, bool bIsTemplateNode, TWeakObjectPtr<UEnum> EnumPtr)
+	auto SetNodeEnumLambda = [](UEdGraphNode* NewNode, UField const* /*EnumField*/, TWeakObjectPtr<UEnum> NonConstEnumPtr)
 	{
 		UK2Node_ForEachElementInEnum* EnumNode = CastChecked<UK2Node_ForEachElementInEnum>(NewNode);
-		if (EnumPtr.IsValid())
-		{
-			EnumNode->Enum = EnumPtr.Get();
-		}
+		EnumNode->Enum = NonConstEnumPtr.Get();
 	};
 
 	for (TObjectIterator<UEnum> EnumIt; EnumIt; ++EnumIt)
@@ -222,7 +219,7 @@ void UK2Node_ForEachElementInEnum::GetMenuActions(FBlueprintActionDatabaseRegist
 			continue;
 		}
 
-		// to keep from needlessly instantiating a UBlueprintNodeSpawner, first   
+		// to keep from needlessly instantiating a UBlueprintNodeSpawners, first   
 		// check to make sure that the registrar is looking for actions of this type
 		// (could be regenerating actions for a specific asset, and therefore the 
 		// registrar would only accept actions corresponding to that asset)
@@ -231,15 +228,15 @@ void UK2Node_ForEachElementInEnum::GetMenuActions(FBlueprintActionDatabaseRegist
 			continue;
 		}
 
-		UBlueprintNodeSpawner* NodeSpawner = UBlueprintNodeSpawner::Create(GetClass());
+		UBlueprintFieldNodeSpawner* NodeSpawner = UBlueprintFieldNodeSpawner::Create(GetClass(), Enum);
 		check(NodeSpawner != nullptr);
+		TWeakObjectPtr<UEnum> NonConstEnumPtr = Enum;
+		NodeSpawner->SetNodeFieldDelegate = UBlueprintFieldNodeSpawner::FSetNodeFieldDelegate::CreateStatic(SetNodeEnumLambda, NonConstEnumPtr);
+
 		// this enum could belong to a class, or is a user defined enum (asset), 
 		// that's why we want to make sure to register it along with the action 
 		// (so the action can be refreshed when the class/asset is).
 		ActionRegistrar.AddBlueprintAction(Enum, NodeSpawner);
-
-		TWeakObjectPtr<UEnum> EnumPtr = Enum;
-		NodeSpawner->CustomizeNodeDelegate = UBlueprintNodeSpawner::FCustomizeNodeDelegate::CreateStatic(CustomizeEnumNodeLambda, EnumPtr);
 	}
 }
 
diff --git a/Engine/Source/Editor/BlueprintGraph/Private/K2Node_GetNumEnumEntries.cpp b/Engine/Source/Editor/BlueprintGraph/Private/K2Node_GetNumEnumEntries.cpp
--- a/Engine/Source/Editor/BlueprintGraph/Private/K2Node_GetNumEnumEntries.cpp
+++ b/Engine/Source/Editor/BlueprintGraph/Private/K2Node_GetNumEnumEntries.cpp
@@ -5,7 +5,7 @@
 #include "KismetCompiler.h"
 #include "../../../Runtime/Engine/Classes/Kismet/KismetSystemLibrary.h"
 #include "K2Node_GetNumEnumEntries.h"
-#include "BlueprintNodeSpawner.h"
+#include "BlueprintFieldNodeSpawner.h"
 #include "EditorCategoryUtils.h"
 #include "BlueprintActionDatabaseRegistrar.h"
 
@@ -98,13 +98,10 @@ void UK2Node_GetNumEnumEntries::ExpandNode(class FKismetCompilerContext& Compile
 
 void UK2Node_GetNumEnumEntries::GetMenuActions(FBlueprintActionDatabaseRegistrar& ActionRegistrar) const
 {
-	auto CustomizeEnumNodeLambda = [](UEdGraphNode* NewNode, bool bIsTemplateNode, TWeakObjectPtr<UEnum> EnumPtr)
+	auto SetNodeEnumLambda = [](UEdGraphNode* NewNode, UField const* /*EnumField*/, TWeakObjectPtr<UEnum> NonConstEnumPtr)
 	{
 		UK2Node_GetNumEnumEntries* EnumNode = CastChecked<UK2Node_GetNumEnumEntries>(NewNode);
-		if (EnumPtr.IsValid())
-		{
-			EnumNode->Enum = EnumPtr.Get();
-		}
+		EnumNode->Enum = NonConstEnumPtr.Get();
 	};
 
 	for (TObjectIterator<UEnum> EnumIt; EnumIt; ++EnumIt)
@@ -115,7 +112,7 @@ void UK2Node_GetNumEnumEntries::GetMenuActions(FBlueprintActionDatabaseRegistrar
 			continue;
 		}
 
-		// to keep from needlessly instantiating a UBlueprintNodeSpawner, first   
+		// to keep from needlessly instantiating a UBlueprintNodeSpawners, first   
 		// check to make sure that the registrar is looking for actions of this type
 		// (could be regenerating actions for a specific asset, and therefore the 
 		// registrar would only accept actions corresponding to that asset)
@@ -124,15 +121,15 @@ void UK2Node_GetNumEnumEntries::GetMenuActions(FBlueprintActionDatabaseRegistrar
 			continue;
 		}
 
-		UBlueprintNodeSpawner* NodeSpawner = UBlueprintNodeSpawner::Create(GetClass());
+		UBlueprintFieldNodeSpawner* NodeSpawner = UBlueprintFieldNodeSpawner::Create(GetClass(), Enum);
 		check(NodeSpawner != nullptr);
+		TWeakObjectPtr<UEnum> NonConstEnumPtr = Enum;
+		NodeSpawner->SetNodeFieldDelegate = UBlueprintFieldNodeSpawner::FSetNodeFieldDelegate::CreateStatic(SetNodeEnumLambda, NonConstEnumPtr);
+
 		// this enum could belong to a class, or is a user defined enum (asset), 
 		// that's why we want to make sure to register it along with the action 
 		// (so the action can be refreshed when the class/asset is).
 		ActionRegistrar.AddBlueprintAction(Enum, NodeSpawner);
-
-		TWeakObjectPtr<UEnum> EnumPtr = Enum;
-		NodeSpawner->CustomizeNodeDelegate = UBlueprintNodeSpawner::FCustomizeNodeDelegate::CreateStatic(CustomizeEnumNodeLambda, EnumPtr);
 	}
 }
 
diff --git a/Engine/Source/Editor/BlueprintGraph/Private/K2Node_MakeStruct.cpp b/Engine/Source/Editor/BlueprintGraph/Private/K2Node_MakeStruct.cpp
--- a/Engine/Source/Editor/BlueprintGraph/Private/K2Node_MakeStruct.cpp
+++ b/Engine/Source/Editor/BlueprintGraph/Private/K2Node_MakeStruct.cpp
@@ -3,7 +3,7 @@
 #include "BlueprintGraphPrivatePCH.h"
 #include "KismetCompiler.h"
 #include "MakeStructHandler.h"
-#include "BlueprintNodeSpawner.h"
+#include "BlueprintFieldNodeSpawner.h"
 #include "EditorCategoryUtils.h"
 #include "BlueprintActionDatabaseRegistrar.h"
 
@@ -272,13 +272,10 @@ UK2Node::ERedirectType UK2Node_MakeStruct::DoPinsMatchForReconstruction(const UE
 
 void UK2Node_MakeStruct::GetMenuActions(FBlueprintActionDatabaseRegistrar& ActionRegistrar) const
 {
-	auto CustomizeMakeNodeLambda = [](UEdGraphNode* NewNode, bool bIsTemplateNode, TWeakObjectPtr<UScriptStruct> StructPtr)
+	auto SetNodeStructLambda = [](UEdGraphNode* NewNode, UField const* /*StructField*/, TWeakObjectPtr<UScriptStruct> NonConstStructPtr)
 	{
 		UK2Node_MakeStruct* MakeNode = CastChecked<UK2Node_MakeStruct>(NewNode);
-		if (StructPtr.IsValid())
-		{
-			MakeNode->StructType = StructPtr.Get();
-		}
+		MakeNode->StructType = NonConstStructPtr.Get();
 	};
 
 	for (TObjectIterator<UScriptStruct> StructIt; StructIt; ++StructIt)
@@ -289,7 +286,7 @@ void UK2Node_MakeStruct::GetMenuActions(FBlueprintActionDatabaseRegistrar& Actio
 			continue;
 		}
 
-		// to keep from needlessly instantiating a UBlueprintNodeSpawner, first   
+		// to keep from needlessly instantiating a UBlueprintNodeSpawners, first   
 		// check to make sure that the registrar is looking for actions of this type
 		// (could be regenerating actions for a specific asset, and therefore the 
 		// registrar would only accept actions corresponding to that asset)
@@ -298,15 +295,15 @@ void UK2Node_MakeStruct::GetMenuActions(FBlueprintActionDatabaseRegistrar& Actio
 			continue;
 		}
 
-		UBlueprintNodeSpawner* NodeSpawner = UBlueprintNodeSpawner::Create(GetClass());
+		UBlueprintFieldNodeSpawner* NodeSpawner = UBlueprintFieldNodeSpawner::Create(GetClass(), Struct);
 		check(NodeSpawner != nullptr);
+		TWeakObjectPtr<UScriptStruct> NonConstStructPtr = Struct;
+		NodeSpawner->SetNodeFieldDelegate = UBlueprintFieldNodeSpawner::FSetNodeFieldDelegate::CreateStatic(SetNodeStructLambda, NonConstStructPtr);
+
 		// this struct could belong to a class, or is a user defined struct 
 		// (asset), that's why we want to make sure to register it along with 
 		// the action (so the action knows to refresh when the class/asset is).
 		ActionRegistrar.AddBlueprintAction(Struct, NodeSpawner);
-
-		TWeakObjectPtr<UScriptStruct> StructPtr = Struct;
-		NodeSpawner->CustomizeNodeDelegate = UBlueprintNodeSpawner::FCustomizeNodeDelegate::CreateStatic(CustomizeMakeNodeLambda, StructPtr);
 	}
 }
 
diff --git a/Engine/Source/Editor/BlueprintGraph/Private/K2Node_SelectEnum.cpp b/Engine/Source/Editor/BlueprintGraph/Private/K2Node_SelectEnum.cpp
--- a/Engine/Source/Editor/BlueprintGraph/Private/K2Node_SelectEnum.cpp
+++ b/Engine/Source/Editor/BlueprintGraph/Private/K2Node_SelectEnum.cpp
@@ -3,7 +3,7 @@
 
 #include "BlueprintGraphPrivatePCH.h"
 #include "K2Node_SelectEnum.h"
-#include "BlueprintNodeSpawner.h"
+#include "BlueprintFieldNodeSpawner.h"
 #include "BlueprintActionDatabaseRegistrar.h"
 
 UDEPRECATED_K2Node_SelectEnum::UDEPRECATED_K2Node_SelectEnum(const class FPostConstructInitializeProperties& PCIP)
@@ -65,13 +65,10 @@ void UDEPRECATED_K2Node_SelectEnum::NotifyPinConnectionListChanged(UEdGraphPin*
 
 void UDEPRECATED_K2Node_SelectEnum::GetMenuActions(FBlueprintActionDatabaseRegistrar& ActionRegistrar) const
 {
-	auto CustomizeEnumNodeLambda = [](UEdGraphNode* NewNode, bool bIsTemplateNode, TWeakObjectPtr<UEnum> EnumPtr)
+	auto SetNodeEnumLambda = [](UEdGraphNode* NewNode, UField const* /*EnumField*/, TWeakObjectPtr<UEnum> NonConstEnumPtr)
 	{
 		UDEPRECATED_K2Node_SelectEnum* EnumNode = CastChecked<UDEPRECATED_K2Node_SelectEnum>(NewNode);
-		if (EnumPtr.IsValid())
-		{
-			EnumNode->Enum = EnumPtr.Get();
-		}
+		EnumNode->Enum = NonConstEnumPtr.Get();
 	};
 
 	for (TObjectIterator<UEnum> EnumIt; EnumIt; ++EnumIt)
@@ -82,7 +79,7 @@ void UDEPRECATED_K2Node_SelectEnum::GetMenuActions(FBlueprintActionDatabaseRegis
 			continue;
 		}
 
-		// to keep from needlessly instantiating a UBlueprintNodeSpawner, first   
+		// to keep from needlessly instantiating a UBlueprintNodeSpawners, first   
 		// check to make sure that the registrar is looking for actions of this type
 		// (could be regenerating actions for a specific asset, and therefore the 
 		// registrar would only accept actions corresponding to that asset)
@@ -90,14 +87,16 @@ void UDEPRECATED_K2Node_SelectEnum::GetMenuActions(FBlueprintActionDatabaseRegis
 		{
 			continue;
 		}
-
 		
-		UBlueprintNodeSpawner* NodeSpawner = UBlueprintNodeSpawner::Create(GetClass());
+		UBlueprintFieldNodeSpawner* NodeSpawner = UBlueprintFieldNodeSpawner::Create(GetClass(), Enum);
 		check(NodeSpawner != nullptr);
+		TWeakObjectPtr<UEnum> NonConstEnumPtr = Enum;
+		NodeSpawner->SetNodeFieldDelegate = UBlueprintFieldNodeSpawner::FSetNodeFieldDelegate::CreateStatic(SetNodeEnumLambda, NonConstEnumPtr);
+
+		// this enum could belong to a class, or is a user defined enum (asset), 
+		// that's why we want to make sure to register it along with the action 
+		// (so the action can be refreshed when the class/asset is).
 		ActionRegistrar.AddBlueprintAction(Enum, NodeSpawner);
-		
-		TWeakObjectPtr<UEnum> EnumPtr = Enum;
-		NodeSpawner->CustomizeNodeDelegate = UBlueprintNodeSpawner::FCustomizeNodeDelegate::CreateStatic(CustomizeEnumNodeLambda, EnumPtr);
 	}
 }
 
diff --git a/Engine/Source/Editor/BlueprintGraph/Private/K2Node_SwitchEnum.cpp b/Engine/Source/Editor/BlueprintGraph/Private/K2Node_SwitchEnum.cpp
--- a/Engine/Source/Editor/BlueprintGraph/Private/K2Node_SwitchEnum.cpp
+++ b/Engine/Source/Editor/BlueprintGraph/Private/K2Node_SwitchEnum.cpp
@@ -3,7 +3,7 @@
 #include "BlueprintGraphPrivatePCH.h"
 #include "Kismet/KismetMathLibrary.h"
 #include "K2Node_SwitchEnum.h"
-#include "BlueprintNodeSpawner.h"
+#include "BlueprintFieldNodeSpawner.h"
 #include "BlueprintActionDatabaseRegistrar.h"
 #include "BlueprintActionDatabaseRegistrar.h"
 
@@ -85,13 +85,10 @@ bool UK2Node_SwitchEnum::IsConnectionDisallowed(const UEdGraphPin* MyPin, const
 
 void UK2Node_SwitchEnum::GetMenuActions(FBlueprintActionDatabaseRegistrar& ActionRegistrar) const
 {
-	auto CustomizeEnumNodeLambda = [](UEdGraphNode* NewNode, bool bIsTemplateNode, TWeakObjectPtr<UEnum> EnumPtr)
+	auto SetNodeEnumLambda = [](UEdGraphNode* NewNode, UField const* /*EnumField*/, TWeakObjectPtr<UEnum> NonConstEnumPtr)
 	{
 		UK2Node_SwitchEnum* EnumNode = CastChecked<UK2Node_SwitchEnum>(NewNode);
-		if (EnumPtr.IsValid())
-		{
-			EnumNode->SetEnum(EnumPtr.Get());
-		}
+		EnumNode->Enum = NonConstEnumPtr.Get();
 	};
 
 	for (TObjectIterator<UEnum> EnumIt; EnumIt; ++EnumIt)
@@ -102,7 +99,7 @@ void UK2Node_SwitchEnum::GetMenuActions(FBlueprintActionDatabaseRegistrar& Actio
 			continue;
 		}
 
-		// to keep from needlessly instantiating a UBlueprintNodeSpawner, first   
+		// to keep from needlessly instantiating a UBlueprintNodeSpawners, first   
 		// check to make sure that the registrar is looking for actions of this type
 		// (could be regenerating actions for a specific asset, and therefore the 
 		// registrar would only accept actions corresponding to that asset)
@@ -111,16 +108,15 @@ void UK2Node_SwitchEnum::GetMenuActions(FBlueprintActionDatabaseRegistrar& Actio
 			continue;
 		}
 
-		UBlueprintNodeSpawner* NodeSpawner = UBlueprintNodeSpawner::Create(GetClass());
+		UBlueprintFieldNodeSpawner* NodeSpawner = UBlueprintFieldNodeSpawner::Create(GetClass(), Enum);
 		check(NodeSpawner != nullptr);
+		TWeakObjectPtr<UEnum> NonConstEnumPtr = Enum;
+		NodeSpawner->SetNodeFieldDelegate = UBlueprintFieldNodeSpawner::FSetNodeFieldDelegate::CreateStatic(SetNodeEnumLambda, NonConstEnumPtr);
 
 		// this enum could belong to a class, or is a user defined enum (asset), 
 		// that's why we want to make sure to register it along with the action 
 		// (so the action can be refreshed when the class/asset is).
 		ActionRegistrar.AddBlueprintAction(Enum, NodeSpawner);
-
-		TWeakObjectPtr<UEnum> EnumPtr = Enum;
-		NodeSpawner->CustomizeNodeDelegate = UBlueprintNodeSpawner::FCustomizeNodeDelegate::CreateStatic(CustomizeEnumNodeLambda, EnumPtr);
 	}
 }
 
diff --git a/Engine/Source/Editor/BlueprintGraph/Private/K2Node_TemporaryVariable.cpp b/Engine/Source/Editor/BlueprintGraph/Private/K2Node_TemporaryVariable.cpp
--- a/Engine/Source/Editor/BlueprintGraph/Private/K2Node_TemporaryVariable.cpp
+++ b/Engine/Source/Editor/BlueprintGraph/Private/K2Node_TemporaryVariable.cpp
@@ -155,6 +155,7 @@ void UK2Node_TemporaryVariable::GetMenuActions(FBlueprintActionDatabaseRegistrar
 			TempVarNode->bIsPersistent = bIsPersistent;
 		};
 
+		// @TODO: Need some kind of unique spawner so that this can generate a unique key
 		NodeSpawner->CustomizeNodeDelegate = UBlueprintNodeSpawner::FCustomizeNodeDelegate::CreateStatic(PostSpawnLambda, VarType, bIsPersistent);
 		return NodeSpawner;
 	};
