Level Blueprints refresh their actions in the database and remove themselves when the World is destroyed.

diff --git a/Engine/Source/Editor/BlueprintGraph/Private/BlueprintActionDatabase.cpp b/Engine/Source/Editor/BlueprintGraph/Private/BlueprintActionDatabase.cpp
--- a/Engine/Source/Editor/BlueprintGraph/Private/BlueprintActionDatabase.cpp
+++ b/Engine/Source/Editor/BlueprintGraph/Private/BlueprintActionDatabase.cpp
@@ -335,6 +335,21 @@ namespace BlueprintActionDatabaseImpl
 	 */
 	static void OnAssetRemoved(FAssetData const& AssetInfo);
 
+	/**
+	 * Callback to clear all levels from the database when a world is destroyed 
+	 * 
+	 * @param  DestroyedWorld	The world that was destroyed
+	 */
+	static void OnWorldDestroyed(UWorld* DestroyedWorld);
+
+	/**
+	 * Returns TRUE if the Object is valid for the database
+	 *
+	 * @param Object		Object to check for validity
+	 * @return				TRUE if the Blueprint is valid for the database
+	 */
+	static bool IsObjectValidForDatabase(UObject const* Object);
+
 	/** 
 	 * Assets that we cleared from the database (to remove references, and make 
 	 * way for a delete), but in-case the class wasn't deleted we need them 
@@ -574,7 +589,7 @@ static void BlueprintActionDatabaseImpl::GetNodeSpecificActions(TSubclassOf<UEdG
 //------------------------------------------------------------------------------
 static void BlueprintActionDatabaseImpl::OnBlueprintChanged(UBlueprint* Blueprint)
 {
-	if (Blueprint->IsAsset())
+	if (IsObjectValidForDatabase(Blueprint))
 	{
 		FBlueprintActionDatabase& ActionDatabase = FBlueprintActionDatabase::Get();
 		ActionDatabase.RefreshAssetActions(Blueprint);
@@ -660,6 +675,34 @@ static void BlueprintActionDatabaseImpl::OnAssetRemoved(FAssetData const& AssetI
 	}
 }
 
+//------------------------------------------------------------------------------
+static void BlueprintActionDatabaseImpl::OnWorldDestroyed(UWorld* DestroyedWorld)
+{
+	for(auto It = DestroyedWorld->GetLevelIterator() ; It ; ++It)
+	{
+		if(const ULevelScriptBlueprint* LevelScript = (*It)->GetLevelScriptBlueprint(true))
+		{
+			FBlueprintActionDatabase::Get().ClearAssetActions((UObject*)LevelScript);
+		}
+	}
+}
+
+//------------------------------------------------------------------------------
+static bool BlueprintActionDatabaseImpl::IsObjectValidForDatabase(UObject const* Object)
+{
+	bool bReturn = false;
+	if(Object->IsAsset())
+	{
+		bReturn = true;
+	}
+	else if(UBlueprint const* Blueprint = Cast<UBlueprint>(Object))
+	{
+		// Level scripts are sometimes not assets because they have not been saved yet, but they are still valid for the database.
+		bReturn = FBlueprintEditorUtils::IsLevelScriptBlueprint(Blueprint);
+	}
+	return bReturn;
+}
+
 /*******************************************************************************
  * FBlueprintActionDatabase
  ******************************************************************************/
@@ -691,6 +734,8 @@ FBlueprintActionDatabase::FBlueprintActionDatabase()
 	AssetRegistry.OnAssetRemoved().AddStatic(&BlueprintActionDatabaseImpl::OnAssetRemoved);
 
 	FEditorDelegates::OnAssetsPreDelete.AddStatic(&BlueprintActionDatabaseImpl::OnAssetsPendingDelete);
+
+	GEngine->OnWorldDestroyed().AddStatic(&BlueprintActionDatabaseImpl::OnWorldDestroyed);
 }
 
 //------------------------------------------------------------------------------
@@ -791,7 +836,7 @@ void FBlueprintActionDatabase::RefreshClassActions(UClass* const Class)
 	else if (bIsBlueprintClass)
 	{
 		UBlueprint const* Blueprint = Cast<UBlueprint>(Class->ClassGeneratedBy);
-		if ((Blueprint != nullptr) && Blueprint->IsAsset())
+		if ((Blueprint != nullptr) && BlueprintActionDatabaseImpl::IsObjectValidForDatabase(Blueprint))
 		{
 			// to prevent us from hitting this twice on init (once for the skel 
 			// class, again for the generated class)
@@ -876,7 +921,7 @@ void FBlueprintActionDatabase::RefreshClassActions(UClass* const Class)
 void FBlueprintActionDatabase::RefreshAssetActions(UObject const* const AssetObject)
 {
 	using namespace BlueprintActionDatabaseImpl;
-	check(AssetObject->IsAsset());
+	check(BlueprintActionDatabaseImpl::IsObjectValidForDatabase(AssetObject));
 
 	FActionList& AssetActionList = ActionRegistry.FindOrAdd(AssetObject);
 	AssetActionList.Empty();
@@ -909,7 +954,7 @@ void FBlueprintActionDatabase::RefreshAssetActions(UObject const* const AssetObj
 //------------------------------------------------------------------------------
 void FBlueprintActionDatabase::ClearAssetActions(UObject const* const AssetObject)
 {
-	check(AssetObject->IsAsset());
+	check(BlueprintActionDatabaseImpl::IsObjectValidForDatabase(AssetObject));
 	ActionRegistry.Remove(AssetObject);
 }
 
