(WIP) Fixing it so input nodes can be favorited in the new menu system (setting it up so that nodes can provide their own node signature; so we don't have to create a spawner for every kind).

diff --git a/Engine/Source/Editor/AnimGraph/Classes/AnimGraphNode_BlendSpacePlayer.h b/Engine/Source/Editor/AnimGraph/Classes/AnimGraphNode_BlendSpacePlayer.h
--- a/Engine/Source/Editor/AnimGraph/Classes/AnimGraphNode_BlendSpacePlayer.h
+++ b/Engine/Source/Editor/AnimGraph/Classes/AnimGraphNode_BlendSpacePlayer.h
@@ -21,7 +21,6 @@ class UAnimGraphNode_BlendSpacePlayer : public UAnimGraphNode_BlendSpaceBase
 	// UEdGraphNode interface
 	virtual FText GetTooltipText() const override;
 	virtual FText GetNodeTitle(ENodeTitleType::Type TitleType) const override;
-	virtual void GetMenuActions(FBlueprintActionDatabaseRegistrar& ActionRegistrar) const override;
 	// End of UEdGraphNode interface
 
 	// UAnimGraphNode_Base interface
@@ -38,6 +37,8 @@ class UAnimGraphNode_BlendSpacePlayer : public UAnimGraphNode_BlendSpaceBase
 
 	// UK2Node interface
 	virtual void GetContextMenuActions(const FGraphNodeContextMenuBuilder& Context) const override;
+	virtual void GetMenuActions(FBlueprintActionDatabaseRegistrar& ActionRegistrar) const override;
+	virtual FBlueprintNodeSignature GetSignature() const override;
 	// End of UK2Node interface
 
 private:
diff --git a/Engine/Source/Editor/AnimGraph/Classes/AnimGraphNode_RotationOffsetBlendSpace.h b/Engine/Source/Editor/AnimGraph/Classes/AnimGraphNode_RotationOffsetBlendSpace.h
--- a/Engine/Source/Editor/AnimGraph/Classes/AnimGraphNode_RotationOffsetBlendSpace.h
+++ b/Engine/Source/Editor/AnimGraph/Classes/AnimGraphNode_RotationOffsetBlendSpace.h
@@ -28,6 +28,7 @@ class UAnimGraphNode_RotationOffsetBlendSpace: public UAnimGraphNode_BlendSpaceB
 	// UK2Node interface
 	virtual void GetContextMenuActions(const FGraphNodeContextMenuBuilder& Context) const override;
 	virtual void GetMenuActions(FBlueprintActionDatabaseRegistrar& ActionRegistrar) const override;
+	virtual FBlueprintNodeSignature GetSignature() const override;
 	// End of UK2Node interface
 
 private:
diff --git a/Engine/Source/Editor/AnimGraph/Private/AnimGraphNode_BlendSpacePlayer.cpp b/Engine/Source/Editor/AnimGraph/Private/AnimGraphNode_BlendSpacePlayer.cpp
--- a/Engine/Source/Editor/AnimGraph/Private/AnimGraphNode_BlendSpacePlayer.cpp
+++ b/Engine/Source/Editor/AnimGraph/Private/AnimGraphNode_BlendSpacePlayer.cpp
@@ -73,40 +73,6 @@ void UAnimGraphNode_BlendSpacePlayer::GetMenuEntries(FGraphContextMenuBuilder& C
 	GetBlendSpaceEntries(bWantAimOffsets, ContextMenuBuilder);
 }
 
-void UAnimGraphNode_BlendSpacePlayer::GetMenuActions(FBlueprintActionDatabaseRegistrar& ActionRegistrar) const
-{
-	auto PostSpawnSetupLambda = [](UEdGraphNode* NewNode, bool /*bIsTemplateNode*/, TWeakObjectPtr<UBlendSpaceBase> BlendSpace)
-	{
-		UAnimGraphNode_BlendSpacePlayer* BlendSpaceNode = CastChecked<UAnimGraphNode_BlendSpacePlayer>(NewNode);
-		BlendSpaceNode->Node.BlendSpace = BlendSpace.Get();		
-	};
-
-	for (TObjectIterator<UBlendSpaceBase> BlendSpaceIt; BlendSpaceIt; ++BlendSpaceIt)
-	{
-		UBlendSpaceBase* BlendSpace = *BlendSpaceIt;
-		// to keep from needlessly instantiating a UBlueprintNodeSpawner, first   
-		// check to make sure that the registrar is looking for actions of this type
-		// (could be regenerating actions for a specific asset, and therefore the 
-		// registrar would only accept actions corresponding to that asset)
-		if (!ActionRegistrar.IsOpenForRegistration(BlendSpace))
-		{
-			continue;
-		}
-
-		bool const bIsAimOffset = BlendSpace->IsA(UAimOffsetBlendSpace::StaticClass()) || 
-			BlendSpace->IsA(UAimOffsetBlendSpace1D::StaticClass());
-		if (!bIsAimOffset)
-		{
-			UBlueprintNodeSpawner* NodeSpawner = UBlueprintNodeSpawner::Create(GetClass());
-			check(NodeSpawner != nullptr);
-			ActionRegistrar.AddBlueprintAction(BlendSpace, NodeSpawner);
-
-			TWeakObjectPtr<UBlendSpaceBase> BlendSpacePtr = BlendSpace;
-			NodeSpawner->CustomizeNodeDelegate = UBlueprintNodeSpawner::FCustomizeNodeDelegate::CreateStatic(PostSpawnSetupLambda, BlendSpacePtr);
-		}
-	}
-}
-
 void UAnimGraphNode_BlendSpacePlayer::ValidateAnimNodeDuringCompilation(class USkeleton* ForSkeleton, class FCompilerResultsLog& MessageLog)
 {
 	if (Node.BlendSpace == NULL)
@@ -145,6 +111,48 @@ void UAnimGraphNode_BlendSpacePlayer::GetContextMenuActions(const FGraphNodeCont
 	}
 }
 
+void UAnimGraphNode_BlendSpacePlayer::GetMenuActions(FBlueprintActionDatabaseRegistrar& ActionRegistrar) const
+{
+	auto PostSpawnSetupLambda = [](UEdGraphNode* NewNode, bool /*bIsTemplateNode*/, TWeakObjectPtr<UBlendSpaceBase> BlendSpace)
+	{
+		UAnimGraphNode_BlendSpacePlayer* BlendSpaceNode = CastChecked<UAnimGraphNode_BlendSpacePlayer>(NewNode);
+		BlendSpaceNode->Node.BlendSpace = BlendSpace.Get();
+	};
+
+	for (TObjectIterator<UBlendSpaceBase> BlendSpaceIt; BlendSpaceIt; ++BlendSpaceIt)
+	{
+		UBlendSpaceBase* BlendSpace = *BlendSpaceIt;
+		// to keep from needlessly instantiating a UBlueprintNodeSpawner, first   
+		// check to make sure that the registrar is looking for actions of this type
+		// (could be regenerating actions for a specific asset, and therefore the 
+		// registrar would only accept actions corresponding to that asset)
+		if (!ActionRegistrar.IsOpenForRegistration(BlendSpace))
+		{
+			continue;
+		}
+
+		bool const bIsAimOffset = BlendSpace->IsA(UAimOffsetBlendSpace::StaticClass()) ||
+			BlendSpace->IsA(UAimOffsetBlendSpace1D::StaticClass());
+		if (!bIsAimOffset)
+		{
+			UBlueprintNodeSpawner* NodeSpawner = UBlueprintNodeSpawner::Create(GetClass());
+			check(NodeSpawner != nullptr);
+			ActionRegistrar.AddBlueprintAction(BlendSpace, NodeSpawner);
+
+			TWeakObjectPtr<UBlendSpaceBase> BlendSpacePtr = BlendSpace;
+			NodeSpawner->CustomizeNodeDelegate = UBlueprintNodeSpawner::FCustomizeNodeDelegate::CreateStatic(PostSpawnSetupLambda, BlendSpacePtr);
+		}
+	}
+}
+
+FBlueprintNodeSignature UAnimGraphNode_BlendSpacePlayer::GetSignature() const
+{
+	FBlueprintNodeSignature NodeSignature = Super::GetSignature();
+	NodeSignature.AddSubObject(Node.BlendSpace);
+
+	return NodeSignature;
+}
+
 void UAnimGraphNode_BlendSpacePlayer::GetAllAnimationSequencesReferred(TArray<UAnimationAsset*>& ComplexAnims, TArray<UAnimSequence*>& AnimationSequences) const
 {
 	if(Node.BlendSpace)
diff --git a/Engine/Source/Editor/AnimGraph/Private/AnimGraphNode_RotationOffsetBlendSpace.cpp b/Engine/Source/Editor/AnimGraph/Private/AnimGraphNode_RotationOffsetBlendSpace.cpp
--- a/Engine/Source/Editor/AnimGraph/Private/AnimGraphNode_RotationOffsetBlendSpace.cpp
+++ b/Engine/Source/Editor/AnimGraph/Private/AnimGraphNode_RotationOffsetBlendSpace.cpp
@@ -96,6 +96,14 @@ void UAnimGraphNode_RotationOffsetBlendSpace::GetMenuActions(FBlueprintActionDat
 	}
 }
 
+FBlueprintNodeSignature UAnimGraphNode_RotationOffsetBlendSpace::GetSignature() const
+{
+	FBlueprintNodeSignature NodeSignature = Super::GetSignature();
+	NodeSignature.AddSubObject(Node.BlendSpace);
+
+	return NodeSignature;
+}
+
 void UAnimGraphNode_RotationOffsetBlendSpace::ValidateAnimNodeDuringCompilation(class USkeleton* ForSkeleton, class FCompilerResultsLog& MessageLog)
 {
 	if (Node.BlendSpace == NULL)
diff --git a/Engine/Source/Editor/BlueprintGraph/Classes/BlueprintBoundEventNodeSpawner.h b/Engine/Source/Editor/BlueprintGraph/Classes/BlueprintBoundEventNodeSpawner.h
--- a/Engine/Source/Editor/BlueprintGraph/Classes/BlueprintBoundEventNodeSpawner.h
+++ b/Engine/Source/Editor/BlueprintGraph/Classes/BlueprintBoundEventNodeSpawner.h
@@ -32,7 +32,7 @@ public:
 	static UBlueprintBoundEventNodeSpawner* Create(TSubclassOf<UK2Node_Event> NodeClass, UMulticastDelegateProperty* EventDelegate, UObject* Outer = nullptr);
 
 	// UBlueprintNodeSpawner interface
-	virtual FBlueprintNodeSpawnerSignature GetSpawnerSignature() const override;
+	virtual FBlueprintNodeSignature GetSpawnerSignature() const override;
 	virtual UEdGraphNode* Invoke(UEdGraph* ParentGraph, FBindingSet const& Bindings, FVector2D const Location) const override;
 	virtual FText GetDefaultMenuName() const override;
 	virtual FText GetDefaultMenuCategory() const override;
diff --git a/Engine/Source/Editor/BlueprintGraph/Classes/BlueprintBoundNodeSpawner.h b/Engine/Source/Editor/BlueprintGraph/Classes/BlueprintBoundNodeSpawner.h
--- a/Engine/Source/Editor/BlueprintGraph/Classes/BlueprintBoundNodeSpawner.h
+++ b/Engine/Source/Editor/BlueprintGraph/Classes/BlueprintBoundNodeSpawner.h
@@ -23,7 +23,7 @@ public:
 	static UBlueprintBoundNodeSpawner* Create(TSubclassOf<UEdGraphNode> NodeClass, UObject* Outer = nullptr);
 
 	// UBlueprintNodeSpawner interface
-	virtual FBlueprintNodeSpawnerSignature GetSpawnerSignature() const override;
+	virtual FBlueprintNodeSignature GetSpawnerSignature() const override;
 	// End UBlueprintNodeSpawner interface
 
 	// FBlueprintNodeBinder interface
diff --git a/Engine/Source/Editor/BlueprintGraph/Classes/BlueprintComponentNodeSpawner.h b/Engine/Source/Editor/BlueprintGraph/Classes/BlueprintComponentNodeSpawner.h
--- a/Engine/Source/Editor/BlueprintGraph/Classes/BlueprintComponentNodeSpawner.h
+++ b/Engine/Source/Editor/BlueprintGraph/Classes/BlueprintComponentNodeSpawner.h
@@ -32,7 +32,7 @@ public:
 	static UBlueprintComponentNodeSpawner* Create(TSubclassOf<UActorComponent> const ComponentClass, UObject* Outer = nullptr);
 
 	// UBlueprintNodeSpawner interface
-	virtual FBlueprintNodeSpawnerSignature GetSpawnerSignature() const override;
+	virtual FBlueprintNodeSignature GetSpawnerSignature() const override;
 	virtual UEdGraphNode* Invoke(UEdGraph* ParentGraph, FBindingSet const& Bindings, FVector2D const Location) const override;
 	virtual FText GetDefaultMenuName() const override;
 	virtual FText GetDefaultMenuCategory() const override;
diff --git a/Engine/Source/Editor/BlueprintGraph/Classes/BlueprintEventNodeSpawner.h b/Engine/Source/Editor/BlueprintGraph/Classes/BlueprintEventNodeSpawner.h
--- a/Engine/Source/Editor/BlueprintGraph/Classes/BlueprintEventNodeSpawner.h
+++ b/Engine/Source/Editor/BlueprintGraph/Classes/BlueprintEventNodeSpawner.h
@@ -45,7 +45,7 @@ public:
 	static UBlueprintEventNodeSpawner* Create(TSubclassOf<UK2Node_Event> NodeClass, FName CustomEventName, UObject* Outer = nullptr);
 
 	// UBlueprintNodeSpawner interface
-	virtual FBlueprintNodeSpawnerSignature GetSpawnerSignature() const override;
+	virtual FBlueprintNodeSignature GetSpawnerSignature() const override;
 	virtual UEdGraphNode* Invoke(UEdGraph* ParentGraph, FBindingSet const& Bindings, FVector2D const Location) const override;
 	virtual FText GetDefaultMenuName() const override;
 	virtual FText GetDefaultMenuCategory() const override;
diff --git a/Engine/Source/Editor/BlueprintGraph/Classes/BlueprintFieldNodeSpawner.h b/Engine/Source/Editor/BlueprintGraph/Classes/BlueprintFieldNodeSpawner.h
--- a/Engine/Source/Editor/BlueprintGraph/Classes/BlueprintFieldNodeSpawner.h
+++ b/Engine/Source/Editor/BlueprintGraph/Classes/BlueprintFieldNodeSpawner.h
@@ -33,7 +33,7 @@ public:
 	static UBlueprintFieldNodeSpawner* Create(TSubclassOf<UK2Node> NodeClass, UField const* const Field, UObject* Outer = nullptr);
 
 	// UBlueprintNodeSpawner interface
-	virtual FBlueprintNodeSpawnerSignature GetSpawnerSignature() const override;
+	virtual FBlueprintNodeSignature GetSpawnerSignature() const override;
 	virtual UEdGraphNode* Invoke(UEdGraph* ParentGraph, FBindingSet const& Bindings, FVector2D const Location) const override;
 	// End UBlueprintNodeSpawner interface
 
diff --git a/Engine/Source/Editor/BlueprintGraph/Classes/BlueprintVariableNodeSpawner.h b/Engine/Source/Editor/BlueprintGraph/Classes/BlueprintVariableNodeSpawner.h
--- a/Engine/Source/Editor/BlueprintGraph/Classes/BlueprintVariableNodeSpawner.h
+++ b/Engine/Source/Editor/BlueprintGraph/Classes/BlueprintVariableNodeSpawner.h
@@ -47,7 +47,7 @@ public:
 
 	// UBlueprintNodeSpawner interface
 	virtual void Prime() override;
-	virtual FBlueprintNodeSpawnerSignature GetSpawnerSignature() const override;
+	virtual FBlueprintNodeSignature GetSpawnerSignature() const override;
 	virtual UEdGraphNode* Invoke(UEdGraph* ParentGraph, FBindingSet const& Bindings, FVector2D const Location) const override;
 	virtual FText GetDefaultMenuName() const override;
 	virtual FText GetDefaultMenuCategory() const override;
diff --git a/Engine/Source/Editor/BlueprintGraph/Classes/K2Node.h b/Engine/Source/Editor/BlueprintGraph/Classes/K2Node.h
--- a/Engine/Source/Editor/BlueprintGraph/Classes/K2Node.h
+++ b/Engine/Source/Editor/BlueprintGraph/Classes/K2Node.h
@@ -2,6 +2,7 @@
 
 #pragma once
 #include "EdGraph/EdGraphNode.h"
+#include "BlueprintNodeSignature.h"
 #include "K2Node.generated.h"
 
 class UBlueprintNodeSpawner;
@@ -313,6 +314,19 @@ class UK2Node : public UEdGraphNode
 	 */
 	virtual FText GetMenuCategory() const { return FText::GetEmpty(); }
 
+	/**
+	 * Retrieves a unique identifier for this node type. Built from the node's 
+	 * class, as well as other signature items (like functions for CallFunction
+	 * nodes).
+	 *
+	 * NOTE: This is not the same as a node identification GUID, two node 
+	 *       instances can have the same signature (if both call the same 
+	 *       function, etc.).
+	 * 
+	 * @return A signature struct, discerning this node from others.
+	 */
+	virtual FBlueprintNodeSignature GetSignature() const { return FBlueprintNodeSignature(GetClass()); }
+
 	BLUEPRINTGRAPH_API enum EBaseNodeRefreshPriority
 	{
 		Low_UsesDependentWildcard = 100,
diff --git a/Engine/Source/Editor/BlueprintGraph/Classes/K2Node_DynamicCast.h b/Engine/Source/Editor/BlueprintGraph/Classes/K2Node_DynamicCast.h
--- a/Engine/Source/Editor/BlueprintGraph/Classes/K2Node_DynamicCast.h
+++ b/Engine/Source/Editor/BlueprintGraph/Classes/K2Node_DynamicCast.h
@@ -27,6 +27,7 @@ class UK2Node_DynamicCast : public UK2Node
 	virtual class FNodeHandlingFunctor* CreateNodeHandler(class FKismetCompilerContext& CompilerContext) const override;
 	virtual bool HasExternalBlueprintDependencies(TArray<class UStruct*>* OptionalOutput) const override;
 	virtual FText GetMenuCategory() const override;
+	virtual FBlueprintNodeSignature GetSignature() const override;
 	// End of UK2Node interface
 
 	/** Get the 'valid cast' exec pin */
diff --git a/Engine/Source/Editor/BlueprintGraph/Classes/K2Node_GetInputAxisKeyValue.h b/Engine/Source/Editor/BlueprintGraph/Classes/K2Node_GetInputAxisKeyValue.h
--- a/Engine/Source/Editor/BlueprintGraph/Classes/K2Node_GetInputAxisKeyValue.h
+++ b/Engine/Source/Editor/BlueprintGraph/Classes/K2Node_GetInputAxisKeyValue.h
@@ -36,6 +36,7 @@ class UK2Node_GetInputAxisKeyValue : public UK2Node_CallFunction
 	virtual void RegisterDynamicBinding(UDynamicBlueprintBinding* BindingObject) const override;
 	virtual void GetMenuActions(FBlueprintActionDatabaseRegistrar& ActionRegistrar) const override;
 	virtual FText GetMenuCategory() const override;
+	virtual FBlueprintNodeSignature GetSignature() const override;
 	// End UK2Node interface
 	
 	void Initialize(const FKey AxisKey);
diff --git a/Engine/Source/Editor/BlueprintGraph/Classes/K2Node_GetInputAxisValue.h b/Engine/Source/Editor/BlueprintGraph/Classes/K2Node_GetInputAxisValue.h
--- a/Engine/Source/Editor/BlueprintGraph/Classes/K2Node_GetInputAxisValue.h
+++ b/Engine/Source/Editor/BlueprintGraph/Classes/K2Node_GetInputAxisValue.h
@@ -35,6 +35,7 @@ class UK2Node_GetInputAxisValue : public UK2Node_CallFunction
 	virtual void RegisterDynamicBinding(UDynamicBlueprintBinding* BindingObject) const override;
 	virtual void GetMenuActions(FBlueprintActionDatabaseRegistrar& ActionRegistrar) const override;
 	virtual FText GetMenuCategory() const override;
+	virtual FBlueprintNodeSignature GetSignature() const override;
 	// End UK2Node interface
 	
 	void Initialize(const FName AxisName);
diff --git a/Engine/Source/Editor/BlueprintGraph/Classes/K2Node_InputAction.h b/Engine/Source/Editor/BlueprintGraph/Classes/K2Node_InputAction.h
--- a/Engine/Source/Editor/BlueprintGraph/Classes/K2Node_InputAction.h
+++ b/Engine/Source/Editor/BlueprintGraph/Classes/K2Node_InputAction.h
@@ -45,6 +45,7 @@ class UK2Node_InputAction : public UK2Node
 	virtual void ExpandNode(FKismetCompilerContext& CompilerContext, UEdGraph* SourceGraph) override;
 	virtual void GetMenuActions(FBlueprintActionDatabaseRegistrar& ActionRegistrar) const override;
 	virtual FText GetMenuCategory() const override;
+	virtual FBlueprintNodeSignature GetSignature() const override;
 	// End UK2Node interface
 
 	/** Get the 'pressed' input pin */
diff --git a/Engine/Source/Editor/BlueprintGraph/Classes/K2Node_InputAxisEvent.h b/Engine/Source/Editor/BlueprintGraph/Classes/K2Node_InputAxisEvent.h
--- a/Engine/Source/Editor/BlueprintGraph/Classes/K2Node_InputAxisEvent.h
+++ b/Engine/Source/Editor/BlueprintGraph/Classes/K2Node_InputAxisEvent.h
@@ -42,6 +42,7 @@ class UK2Node_InputAxisEvent : public UK2Node_Event
 	virtual void RegisterDynamicBinding(UDynamicBlueprintBinding* BindingObject) const override;
 	virtual void GetMenuActions(FBlueprintActionDatabaseRegistrar& ActionRegistrar) const override;
 	virtual FText GetMenuCategory() const override;
+	virtual FBlueprintNodeSignature GetSignature() const override;
 	// End UK2Node interface
 
 	void Initialize(const FName AxisName);
diff --git a/Engine/Source/Editor/BlueprintGraph/Classes/K2Node_InputAxisKeyEvent.h b/Engine/Source/Editor/BlueprintGraph/Classes/K2Node_InputAxisKeyEvent.h
--- a/Engine/Source/Editor/BlueprintGraph/Classes/K2Node_InputAxisKeyEvent.h
+++ b/Engine/Source/Editor/BlueprintGraph/Classes/K2Node_InputAxisKeyEvent.h
@@ -38,6 +38,7 @@ class UK2Node_InputAxisKeyEvent : public UK2Node_Event
 	virtual void RegisterDynamicBinding(UDynamicBlueprintBinding* BindingObject) const override;
 	virtual void GetMenuActions(FBlueprintActionDatabaseRegistrar& ActionRegistrar) const override;
 	virtual FText GetMenuCategory() const override;
+	virtual FBlueprintNodeSignature GetSignature() const override;
 	// End UK2Node interface
 
 	void Initialize(const FKey AxisKey);
diff --git a/Engine/Source/Editor/BlueprintGraph/Classes/K2Node_InputKey.h b/Engine/Source/Editor/BlueprintGraph/Classes/K2Node_InputKey.h
--- a/Engine/Source/Editor/BlueprintGraph/Classes/K2Node_InputKey.h
+++ b/Engine/Source/Editor/BlueprintGraph/Classes/K2Node_InputKey.h
@@ -53,6 +53,7 @@ class UK2Node_InputKey : public UK2Node
 	virtual void ExpandNode(FKismetCompilerContext& CompilerContext, UEdGraph* SourceGraph) override;
 	virtual void GetMenuActions(FBlueprintActionDatabaseRegistrar& ActionRegistrar) const override;
 	virtual FText GetMenuCategory() const override;
+	virtual FBlueprintNodeSignature GetSignature() const override;
 	// End UK2Node interface
 
 	// Begin UEdGraphNode interface.
diff --git a/Engine/Source/Editor/BlueprintGraph/Classes/K2Node_MacroInstance.h b/Engine/Source/Editor/BlueprintGraph/Classes/K2Node_MacroInstance.h
--- a/Engine/Source/Editor/BlueprintGraph/Classes/K2Node_MacroInstance.h
+++ b/Engine/Source/Editor/BlueprintGraph/Classes/K2Node_MacroInstance.h
@@ -58,6 +58,7 @@ public:
 	virtual void GetNodeAttributes( TArray<TKeyValuePair<FString, FString>>& OutNodeAttributes ) const override;
 	virtual FText GetMenuCategory() const override;
 	virtual int32 GetNodeRefreshPriority() const override { return EBaseNodeRefreshPriority::Low_UsesDependentWildcard; }
+	virtual FBlueprintNodeSignature GetSignature() const override;
 	// End UK2Node interface
 
 	void SetMacroGraph(UEdGraph* Graph) { MacroGraphReference.SetGraph(Graph); }
diff --git a/Engine/Source/Editor/BlueprintGraph/Classes/K2Node_TemporaryVariable.h b/Engine/Source/Editor/BlueprintGraph/Classes/K2Node_TemporaryVariable.h
--- a/Engine/Source/Editor/BlueprintGraph/Classes/K2Node_TemporaryVariable.h
+++ b/Engine/Source/Editor/BlueprintGraph/Classes/K2Node_TemporaryVariable.h
@@ -33,6 +33,7 @@ class UK2Node_TemporaryVariable : public UK2Node
 	virtual class FNodeHandlingFunctor* CreateNodeHandler(class FKismetCompilerContext& CompilerContext) const override;
 	virtual void GetMenuActions(FBlueprintActionDatabaseRegistrar& ActionRegistrar) const override;
 	virtual FText GetMenuCategory() const override;
+	virtual FBlueprintNodeSignature GetSignature() const override;
 	// End UK2Node interface.
 
 private:
diff --git a/Engine/Source/Editor/BlueprintGraph/Private/BlueprintActionDatabase.cpp b/Engine/Source/Editor/BlueprintGraph/Private/BlueprintActionDatabase.cpp
--- a/Engine/Source/Editor/BlueprintGraph/Private/BlueprintActionDatabase.cpp
+++ b/Engine/Source/Editor/BlueprintGraph/Private/BlueprintActionDatabase.cpp
@@ -125,7 +125,6 @@ static UBlueprintNodeSpawner* FBlueprintNodeSpawnerFactory::MakeMacroNodeSpawner
 	};
 
 	TWeakObjectPtr<UEdGraph> GraphPtr = MacroGraph;
-	// @TODO: Needs access to the MacroGraph to be favoritable...
 	NodeSpawner->CustomizeNodeDelegate = UBlueprintNodeSpawner::FCustomizeNodeDelegate::CreateStatic(CustomizeMacroNodeLambda, GraphPtr);
 
 	return NodeSpawner;
@@ -505,7 +504,6 @@ static void BlueprintActionDatabaseImpl::AddClassCastActions(UClass* const Class
 		};
 
 		UBlueprintNodeSpawner* CastObjNodeSpawner = UBlueprintNodeSpawner::Create(UK2Node_DynamicCast::StaticClass());
-		// @TODO: Needs access to the Class to be favoritable...
 		CastObjNodeSpawner->CustomizeNodeDelegate = UBlueprintNodeSpawner::FCustomizeNodeDelegate::CreateStatic(CustomizeCastNodeLambda, Class);
 		ActionListOut.Add(CastObjNodeSpawner);
 
diff --git a/Engine/Source/Editor/BlueprintGraph/Private/BlueprintActionFilter.cpp b/Engine/Source/Editor/BlueprintGraph/Private/BlueprintActionFilter.cpp
--- a/Engine/Source/Editor/BlueprintGraph/Private/BlueprintActionFilter.cpp
+++ b/Engine/Source/Editor/BlueprintGraph/Private/BlueprintActionFilter.cpp
@@ -1048,7 +1048,6 @@ namespace EBlueprintActionInfoFlags
 		CachedField    = (1 << 1),
 		CachedProperty = (1 << 2),
 		CachedFunction = (1 << 3),
-		CachedGuid     = (1 << 4),
 	};
 }
 
@@ -1119,17 +1118,6 @@ UClass const* FBlueprintActionInfo::GetNodeClass()
 	return NodeClass;
 }
 
-//------------------------------------------------------------------------------
-FGuid FBlueprintActionInfo::GetActionGuid()
-{
-	if ((CacheFlags & EBlueprintActionInfoFlags::CachedGuid) == 0)
-	{
-		CachedActionId = NodeSpawner->GetSpawnerSignature().AsGuid();
-		CacheFlags |= EBlueprintActionInfoFlags::CachedGuid;
-	}
-	return CachedActionId;
-}
-
 //------------------------------------------------------------------------------
 UField const* FBlueprintActionInfo::GetAssociatedMemberField()
 {
diff --git a/Engine/Source/Editor/BlueprintGraph/Private/BlueprintBoundEventNodeSpawner.cpp b/Engine/Source/Editor/BlueprintGraph/Private/BlueprintBoundEventNodeSpawner.cpp
--- a/Engine/Source/Editor/BlueprintGraph/Private/BlueprintBoundEventNodeSpawner.cpp
+++ b/Engine/Source/Editor/BlueprintGraph/Private/BlueprintBoundEventNodeSpawner.cpp
@@ -35,12 +35,12 @@ UBlueprintBoundEventNodeSpawner::UBlueprintBoundEventNodeSpawner(class FPostCons
 }
 
 //------------------------------------------------------------------------------
-FBlueprintNodeSpawnerSignature UBlueprintBoundEventNodeSpawner::GetSpawnerSignature() const
+FBlueprintNodeSignature UBlueprintBoundEventNodeSpawner::GetSpawnerSignature() const
 {
 	// explicit actions for binding (like this) cannot be reconstructed form a 
 	// signature (since this spawner does not own whatever it will be binding 
 	// to), therefore we return an empty (invalid) signature
-	return FBlueprintNodeSpawnerSignature();
+	return FBlueprintNodeSignature();
 }
 
 //------------------------------------------------------------------------------
diff --git a/Engine/Source/Editor/BlueprintGraph/Private/BlueprintBoundNodeSpawner.cpp b/Engine/Source/Editor/BlueprintGraph/Private/BlueprintBoundNodeSpawner.cpp
--- a/Engine/Source/Editor/BlueprintGraph/Private/BlueprintBoundNodeSpawner.cpp
+++ b/Engine/Source/Editor/BlueprintGraph/Private/BlueprintBoundNodeSpawner.cpp
@@ -30,12 +30,12 @@ UBlueprintBoundNodeSpawner::UBlueprintBoundNodeSpawner(class FPostConstructIniti
 }
 
 //------------------------------------------------------------------------------
-FBlueprintNodeSpawnerSignature UBlueprintBoundNodeSpawner::GetSpawnerSignature() const
+FBlueprintNodeSignature UBlueprintBoundNodeSpawner::GetSpawnerSignature() const
 {
 	// explicit actions for binding (like this) cannot be reconstructed form a 
 	// signature (since this spawner does not own whatever it will be binding 
 	// to), therefore we return an empty (invalid) signature
-	return FBlueprintNodeSpawnerSignature();
+	return FBlueprintNodeSignature();
 }
 
 //------------------------------------------------------------------------------
diff --git a/Engine/Source/Editor/BlueprintGraph/Private/BlueprintComponentNodeSpawner.cpp b/Engine/Source/Editor/BlueprintGraph/Private/BlueprintComponentNodeSpawner.cpp
--- a/Engine/Source/Editor/BlueprintGraph/Private/BlueprintComponentNodeSpawner.cpp
+++ b/Engine/Source/Editor/BlueprintGraph/Private/BlueprintComponentNodeSpawner.cpp
@@ -36,9 +36,9 @@ UBlueprintComponentNodeSpawner::UBlueprintComponentNodeSpawner(class FPostConstr
 }
 
 //------------------------------------------------------------------------------
-FBlueprintNodeSpawnerSignature UBlueprintComponentNodeSpawner::GetSpawnerSignature() const
+FBlueprintNodeSignature UBlueprintComponentNodeSpawner::GetSpawnerSignature() const
 {
-	FBlueprintNodeSpawnerSignature SpawnerSignature = Super::GetSpawnerSignature();
+	FBlueprintNodeSignature SpawnerSignature(NodeClass);
 	SpawnerSignature.AddSubObject(ComponentClass);
 	return SpawnerSignature;
 }
diff --git a/Engine/Source/Editor/BlueprintGraph/Private/BlueprintEventNodeSpawner.cpp b/Engine/Source/Editor/BlueprintGraph/Private/BlueprintEventNodeSpawner.cpp
--- a/Engine/Source/Editor/BlueprintGraph/Private/BlueprintEventNodeSpawner.cpp
+++ b/Engine/Source/Editor/BlueprintGraph/Private/BlueprintEventNodeSpawner.cpp
@@ -89,13 +89,13 @@ UBlueprintEventNodeSpawner::UBlueprintEventNodeSpawner(class FPostConstructIniti
 }
 
 //------------------------------------------------------------------------------
-FBlueprintNodeSpawnerSignature UBlueprintEventNodeSpawner::GetSpawnerSignature() const
+FBlueprintNodeSignature UBlueprintEventNodeSpawner::GetSpawnerSignature() const
 {
-	FBlueprintNodeSpawnerSignature SpawnerSignature = Super::GetSpawnerSignature();
+	FBlueprintNodeSignature SpawnerSignature(NodeClass);
 	if (IsForCustomEvent() && !CustomEventName.IsNone())
 	{
 		static const FName CustomSignatureKey(TEXT("CustomEvent"));
-		SpawnerSignature.AddKeyValue(CustomSignatureKey, CustomEventName.ToString());
+		SpawnerSignature.AddNamedValue(CustomSignatureKey, CustomEventName.ToString());
 	}
 	else
 	{
diff --git a/Engine/Source/Editor/BlueprintGraph/Private/BlueprintFieldNodeSpawner.cpp b/Engine/Source/Editor/BlueprintGraph/Private/BlueprintFieldNodeSpawner.cpp
--- a/Engine/Source/Editor/BlueprintGraph/Private/BlueprintFieldNodeSpawner.cpp
+++ b/Engine/Source/Editor/BlueprintGraph/Private/BlueprintFieldNodeSpawner.cpp
@@ -27,9 +27,9 @@ UBlueprintFieldNodeSpawner::UBlueprintFieldNodeSpawner(class FPostConstructIniti
 }
 
 //------------------------------------------------------------------------------
-FBlueprintNodeSpawnerSignature UBlueprintFieldNodeSpawner::GetSpawnerSignature() const
+FBlueprintNodeSignature UBlueprintFieldNodeSpawner::GetSpawnerSignature() const
 {
-	FBlueprintNodeSpawnerSignature SpawnerSignature = Super::GetSpawnerSignature();
+	FBlueprintNodeSignature SpawnerSignature(NodeClass);
 	SpawnerSignature.AddSubObject(Field);
 
 	return SpawnerSignature;
diff --git a/Engine/Source/Editor/BlueprintGraph/Private/BlueprintNodeSignature.cpp b/Engine/Source/Editor/BlueprintGraph/Private/BlueprintNodeSignature.cpp
new file mode 100644
--- /dev/null
+++ b/Engine/Source/Editor/BlueprintGraph/Private/BlueprintNodeSignature.cpp
@@ -0,0 +1,199 @@
+// Copyright 1998-2014 Epic Games, Inc. All Rights Reserved.
+
+#include "BlueprintGraphPrivatePCH.h"
+#include "BlueprintNodeSignature.h"
+
+/*******************************************************************************
+ * Static FBlueprintNodeSignature Helpers
+ ******************************************************************************/
+
+namespace BlueprintNodeSignatureImpl
+{
+	static FString const SignatureOpeningStr("(");
+	static FString const SignatureElementDelim(",");
+	static FString const SignatureClosingStr(")");
+	static FString const SignatureKeyDelim("=");
+
+	static FName const LegacySubObjSignatureKey(TEXT("FieldName"));
+
+	/**
+	 * Iterates through the given KeyMap and looks for name collisions. Will 
+	 * append an incrementing value until a collision is not detected.
+	 * 
+	 * @param  BaseName	The base name you want some permutation of.
+	 * @param  KeyMap	The map you want the name key for.
+	 * @return A unique name key for the supplied map.
+	 */
+	static FName FindUniqueKeyName(FName BaseName, TMap<FName, FString>& KeyMap);
+}
+
+//------------------------------------------------------------------------------
+static FName BlueprintNodeSignatureImpl::FindUniqueKeyName(FName BaseName, TMap<FName, FString>& KeyMap)
+{
+	FName SingatureKey = BaseName;
+
+	int32 FNameIndex = 0;
+	while (KeyMap.Find(SingatureKey) != nullptr)
+	{
+		SingatureKey = FName(BaseName, ++FNameIndex);
+	}
+	return SingatureKey;
+}
+
+/*******************************************************************************
+ * FBlueprintNodeSignature
+ ******************************************************************************/
+
+//------------------------------------------------------------------------------
+FBlueprintNodeSignature::FBlueprintNodeSignature(FString const& UserString)
+{
+	using namespace BlueprintNodeSignatureImpl;
+
+	FString SanitizedSignature = UserString;
+	SanitizedSignature.RemoveFromStart(SignatureOpeningStr);
+	SanitizedSignature.RemoveFromEnd(SignatureClosingStr);
+
+	TArray<FString> SignatureElements;
+	SanitizedSignature.ParseIntoArray(&SignatureElements, *SignatureElementDelim, 1);
+
+	for (FString& SignatureElement : SignatureElements)
+	{
+		SignatureElement.Trim();
+
+		FString SignatureKey, SignatureValue;
+		SignatureElement.Split(SignatureKeyDelim, &SignatureKey, &SignatureValue);
+		// @TODO: look for UObject redirects with SignatureValue
+
+		SignatureValue.RemoveFromStart(TEXT("\""));
+		SignatureValue.RemoveFromEnd(TEXT("\""));
+		AddNamedValue(FName(*SignatureKey), SignatureValue);
+	}
+}
+
+//------------------------------------------------------------------------------
+FBlueprintNodeSignature::FBlueprintNodeSignature(TSubclassOf<UEdGraphNode> NodeClass)
+{
+	SetNodeClass(NodeClass);
+}
+
+//------------------------------------------------------------------------------
+void FBlueprintNodeSignature::SetNodeClass(TSubclassOf<UEdGraphNode> NodeClass)
+{
+	static const FName NodeClassSignatureKey(TEXT("NodeName"));
+
+	if (NodeClass != nullptr)
+	{
+		AddNamedValue(NodeClassSignatureKey, NodeClass->GetPathName());
+	}
+	else
+	{
+		SignatureSet.Remove(NodeClassSignatureKey);
+		MarkDirty();
+	}
+}
+
+//------------------------------------------------------------------------------
+void FBlueprintNodeSignature::AddSubObject(UObject const* SignatureObj)
+{
+	// not ideal for generic "objects", but we have to keep in line with the 
+	// old favorites system (for backwards compatibility)
+	using namespace BlueprintNodeSignatureImpl;
+	FName SubObjectSignatureKey = FindUniqueKeyName(LegacySubObjSignatureKey, SignatureSet);
+
+	AddNamedValue(SubObjectSignatureKey, SignatureObj->GetPathName());
+}
+
+//------------------------------------------------------------------------------
+void FBlueprintNodeSignature::AddKeyValue(FString const& KeyValue)
+{
+	// not ideal for some arbitrary value, but we have to keep in line with the 
+	// old favorites system (for backwards compatibility)
+	using namespace BlueprintNodeSignatureImpl;
+	FName SignatureKey = FindUniqueKeyName(LegacySubObjSignatureKey, SignatureSet);
+
+	AddNamedValue(SignatureKey, KeyValue);
+}
+
+//------------------------------------------------------------------------------
+void FBlueprintNodeSignature::AddNamedValue(FName SignatureKey, FString const& Value)
+{
+	SignatureSet.Add(SignatureKey, Value);
+	MarkDirty();
+}
+
+//------------------------------------------------------------------------------
+bool FBlueprintNodeSignature::IsValid() const
+{
+	return (SignatureSet.Num() > 0);
+}
+
+//------------------------------------------------------------------------------
+FString const& FBlueprintNodeSignature::ToString() const
+{
+	using namespace BlueprintNodeSignatureImpl;
+
+	if (CachedSignatureString.IsEmpty() && IsValid())
+	{
+		TArray<FString> SignatureElements;
+		for (auto& SignatureElement : SignatureSet)
+		{
+			FString ObjSignature = SignatureElement.Key.ToString() + SignatureKeyDelim + TEXT("\"") + SignatureElement.Value + TEXT("\"");
+			SignatureElements.Add(ObjSignature);
+		}
+		SignatureElements.Sort();
+
+		CachedSignatureString = SignatureOpeningStr;
+		for (FString& SignatureElement : SignatureElements)
+		{
+			CachedSignatureString += SignatureElement + SignatureElementDelim;
+		}
+		CachedSignatureString.RemoveFromEnd(SignatureElementDelim);
+		CachedSignatureString += SignatureClosingStr;
+	}
+	
+	return CachedSignatureString;
+}
+
+//------------------------------------------------------------------------------
+FGuid const& FBlueprintNodeSignature::AsGuid() const
+{
+	static const int32 BytesPerMd5Hash = 16;
+	static const int32 BytesPerGuidVal =  4;
+	static const int32 MembersPerGuid  =  4;
+	static const int32 BitsPerByte     =  8;
+
+	if (!CachedSignatureGuid.IsValid() && IsValid())
+	{
+		FString const& SignatureString = ToString();
+
+		FMD5  Md5Gen;
+		uint8 HashedBytes[BytesPerMd5Hash];
+		Md5Gen.Update((unsigned char*)TCHAR_TO_ANSI(*SignatureString), SignatureString.Len());
+		Md5Gen.Final(HashedBytes);
+		//FString HashedStr = FMD5::HashAnsiString(*SignatureString); // to check against
+
+		for (int32 GuidIndex = 0; GuidIndex < MembersPerGuid; ++GuidIndex)
+		{
+			int32 const Md5UpperIndex = (GuidIndex + 1) * BytesPerGuidVal - 1;
+			int32 const Md5LowerIndex = Md5UpperIndex - BytesPerGuidVal;
+
+			int32 GuidValue = 0;
+
+			int32 BitOffset = 1;
+			for (int32 Md5ByteIndex = Md5UpperIndex; Md5ByteIndex >= Md5LowerIndex; --Md5ByteIndex)
+			{
+				GuidValue += HashedBytes[Md5ByteIndex] * BitOffset;
+				BitOffset <<= BitsPerByte;
+			}
+			CachedSignatureGuid[GuidIndex] = GuidValue;
+		}
+	}
+	return CachedSignatureGuid;
+}
+
+//------------------------------------------------------------------------------
+void FBlueprintNodeSignature::MarkDirty()
+{
+	CachedSignatureGuid.Invalidate();
+	CachedSignatureString.Empty();
+}
diff --git a/Engine/Source/Editor/BlueprintGraph/Private/BlueprintNodeSpawner.cpp b/Engine/Source/Editor/BlueprintGraph/Private/BlueprintNodeSpawner.cpp
--- a/Engine/Source/Editor/BlueprintGraph/Private/BlueprintNodeSpawner.cpp
+++ b/Engine/Source/Editor/BlueprintGraph/Private/BlueprintNodeSpawner.cpp
@@ -100,9 +100,19 @@ void UBlueprintNodeSpawner::Prime()
 }
 
 //------------------------------------------------------------------------------
-FBlueprintNodeSpawnerSignature UBlueprintNodeSpawner::GetSpawnerSignature() const
+FBlueprintNodeSignature UBlueprintNodeSpawner::GetSpawnerSignature() const
 {
-	return FBlueprintNodeSpawnerSignature(NodeClass);
+	FBlueprintNodeSignature SpawnerSignature;
+	if (UK2Node const* NodeTemplate = Cast<UK2Node>(GetTemplateNode()))
+	{
+		SpawnerSignature = NodeTemplate->GetSignature();
+	}
+
+	if (!SpawnerSignature.IsValid())
+	{
+		SpawnerSignature.SetNodeClass(NodeClass);
+	}
+	return SpawnerSignature;
 }
 
 //------------------------------------------------------------------------------
diff --git a/Engine/Source/Editor/BlueprintGraph/Private/BlueprintNodeSpawnerSignature.cpp b/Engine/Source/Editor/BlueprintGraph/Private/BlueprintNodeSpawnerSignature.cpp
deleted file mode 100644
diff --git a/Engine/Source/Editor/BlueprintGraph/Private/BlueprintVariableNodeSpawner.cpp b/Engine/Source/Editor/BlueprintGraph/Private/BlueprintVariableNodeSpawner.cpp
--- a/Engine/Source/Editor/BlueprintGraph/Private/BlueprintVariableNodeSpawner.cpp
+++ b/Engine/Source/Editor/BlueprintGraph/Private/BlueprintVariableNodeSpawner.cpp
@@ -84,14 +84,14 @@ void UBlueprintVariableNodeSpawner::Prime()
 }
 
 //------------------------------------------------------------------------------
-FBlueprintNodeSpawnerSignature UBlueprintVariableNodeSpawner::GetSpawnerSignature() const
+FBlueprintNodeSignature UBlueprintVariableNodeSpawner::GetSpawnerSignature() const
 {
-	FBlueprintNodeSpawnerSignature SpawnerSignature = Super::GetSpawnerSignature();
+	FBlueprintNodeSignature SpawnerSignature(NodeClass);
 	if (IsLocalVariable())
 	{
 		SpawnerSignature.AddSubObject(LocalVarOuter);
 		static const FName LocalVarSignatureKey(TEXT("LocalVarName"));
-		SpawnerSignature.AddKeyValue(LocalVarSignatureKey, LocalVarDesc.VarName.ToString());
+		SpawnerSignature.AddNamedValue(LocalVarSignatureKey, LocalVarDesc.VarName.ToString());
 	}
 	
 	return SpawnerSignature;
diff --git a/Engine/Source/Editor/BlueprintGraph/Private/K2Node_DynamicCast.cpp b/Engine/Source/Editor/BlueprintGraph/Private/K2Node_DynamicCast.cpp
--- a/Engine/Source/Editor/BlueprintGraph/Private/K2Node_DynamicCast.cpp
+++ b/Engine/Source/Editor/BlueprintGraph/Private/K2Node_DynamicCast.cpp
@@ -175,4 +175,12 @@ FText UK2Node_DynamicCast::GetMenuCategory() const
 	return CachedCategory;
 }
 
+FBlueprintNodeSignature UK2Node_DynamicCast::GetSignature() const
+{
+	FBlueprintNodeSignature NodeSignature = Super::GetSignature();
+	NodeSignature.AddSubObject(TargetType);
+
+	return NodeSignature;
+}
+
 #undef LOCTEXT_NAMESPACE
diff --git a/Engine/Source/Editor/BlueprintGraph/Private/K2Node_GetInputAxisKeyValue.cpp b/Engine/Source/Editor/BlueprintGraph/Private/K2Node_GetInputAxisKeyValue.cpp
--- a/Engine/Source/Editor/BlueprintGraph/Private/K2Node_GetInputAxisKeyValue.cpp
+++ b/Engine/Source/Editor/BlueprintGraph/Private/K2Node_GetInputAxisKeyValue.cpp
@@ -137,6 +137,12 @@ void UK2Node_GetInputAxisKeyValue::GetMenuActions(FBlueprintActionDatabaseRegist
 		InputNode->Initialize(Key);
 	};
 
+	// actions get registered under specific object-keys; the idea is that 
+	// actions might have to be updated (or deleted) if their object-key is  
+	// mutated (or removed)... here we use the node's class (so if the node 
+	// type disappears, then the action should go with it)
+	UClass* ActionKey = GetClass();
+
 	for (FKey const Key : AllKeys)
 	{
 		if (!Key.IsBindableInBlueprints() || !Key.IsFloatAxis())
@@ -144,11 +150,6 @@ void UK2Node_GetInputAxisKeyValue::GetMenuActions(FBlueprintActionDatabaseRegist
 			continue;
 		}
 
-		// actions get registered under specific object-keys; the idea is that 
-		// actions might have to be updated (or deleted) if their object-key is  
-		// mutated (or removed)... here we use the node's class (so if the node 
-		// type disappears, then the action should go with it)
-		UClass* ActionKey = GetClass();
 		// to keep from needlessly instantiating a UBlueprintNodeSpawner, first   
 		// check to make sure that the registrar is looking for actions of this type
 		// (could be regenerating actions for a specific asset, and therefore the 
@@ -204,4 +205,12 @@ FText UK2Node_GetInputAxisKeyValue::GetMenuCategory() const
 	return CachedCategories[CategoryIndex];
 }
 
+FBlueprintNodeSignature UK2Node_GetInputAxisKeyValue::GetSignature() const
+{
+	FBlueprintNodeSignature NodeSignature = Super::GetSignature();
+	NodeSignature.AddKeyValue(InputAxisKey.ToString());
+
+	return NodeSignature;
+}
+
 #undef LOCTEXT_NAMESPACE
\ No newline at end of file
diff --git a/Engine/Source/Editor/BlueprintGraph/Private/K2Node_GetInputAxisValue.cpp b/Engine/Source/Editor/BlueprintGraph/Private/K2Node_GetInputAxisValue.cpp
--- a/Engine/Source/Editor/BlueprintGraph/Private/K2Node_GetInputAxisValue.cpp
+++ b/Engine/Source/Editor/BlueprintGraph/Private/K2Node_GetInputAxisValue.cpp
@@ -111,13 +111,14 @@ void UK2Node_GetInputAxisValue::GetMenuActions(FBlueprintActionDatabaseRegistrar
 		InputNode->Initialize(AxisName);
 	};
 
+	// actions get registered under specific object-keys; the idea is that 
+	// actions might have to be updated (or deleted) if their object-key is  
+	// mutated (or removed)... here we use the node's class (so if the node 
+	// type disappears, then the action should go with it)
+	UClass* ActionKey = GetClass();
+
 	for (FName const InputAxisName : AxisNames)
 	{
-		// actions get registered under specific object-keys; the idea is that 
-		// actions might have to be updated (or deleted) if their object-key is  
-		// mutated (or removed)... here we use the node's class (so if the node 
-		// type disappears, then the action should go with it)
-		UClass* ActionKey = GetClass();
 		// to keep from needlessly instantiating a UBlueprintNodeSpawner, first   
 		// check to make sure that the registrar is looking for actions of this type
 		// (could be regenerating actions for a specific asset, and therefore the 
@@ -146,4 +147,12 @@ FText UK2Node_GetInputAxisValue::GetMenuCategory() const
 	return CachedCategory;
 }
 
+FBlueprintNodeSignature UK2Node_GetInputAxisValue::GetSignature() const
+{
+	FBlueprintNodeSignature NodeSignature = Super::GetSignature();
+	NodeSignature.AddKeyValue(InputAxisName.ToString());
+
+	return NodeSignature;
+}
+
 #undef LOCTEXT_NAMESPACE
diff --git a/Engine/Source/Editor/BlueprintGraph/Private/K2Node_InputAction.cpp b/Engine/Source/Editor/BlueprintGraph/Private/K2Node_InputAction.cpp
--- a/Engine/Source/Editor/BlueprintGraph/Private/K2Node_InputAction.cpp
+++ b/Engine/Source/Editor/BlueprintGraph/Private/K2Node_InputAction.cpp
@@ -155,9 +155,14 @@ void UK2Node_InputAction::GetMenuActions(FBlueprintActionDatabaseRegistrar& Acti
 		InputNode->InputActionName = ActionName;
 	};
 
+	// actions get registered under specific object-keys; the idea is that 
+	// actions might have to be updated (or deleted) if their object-key is  
+	// mutated (or removed)... here we use the node's class (so if the node 
+	// type disappears, then the action should go with it)
+	UClass* ActionKey = GetClass();
+
 	for (FName const InputActionName : ActionNames)
 	{
-		UClass* ActionKey = GetClass();
 		// to keep from needlessly instantiating a UBlueprintNodeSpawner, first   
 		// check to make sure that the registrar is looking for actions of this type
 		// (could be regenerating actions for a specific asset, and therefore the 
@@ -186,4 +191,12 @@ FText UK2Node_InputAction::GetMenuCategory() const
 	return CachedCategory;
 }
 
+FBlueprintNodeSignature UK2Node_InputAction::GetSignature() const
+{
+	FBlueprintNodeSignature NodeSignature = Super::GetSignature();
+	NodeSignature.AddKeyValue(InputActionName.ToString());
+
+	return NodeSignature;
+}
+
 #undef LOCTEXT_NAMESPACE
diff --git a/Engine/Source/Editor/BlueprintGraph/Private/K2Node_InputAxisEvent.cpp b/Engine/Source/Editor/BlueprintGraph/Private/K2Node_InputAxisEvent.cpp
--- a/Engine/Source/Editor/BlueprintGraph/Private/K2Node_InputAxisEvent.cpp
+++ b/Engine/Source/Editor/BlueprintGraph/Private/K2Node_InputAxisEvent.cpp
@@ -123,13 +123,14 @@ void UK2Node_InputAxisEvent::GetMenuActions(FBlueprintActionDatabaseRegistrar& A
 		InputNode->Initialize(AxisName);
 	};
 
+	// actions get registered under specific object-keys; the idea is that 
+	// actions might have to be updated (or deleted) if their object-key is  
+	// mutated (or removed)... here we use the node's class (so if the node 
+	// type disappears, then the action should go with it)
+	UClass* ActionKey = GetClass();
+
 	for (FName const InputAxisName : AxisNames)
 	{
-		// actions get registered under specific object-keys; the idea is that 
-		// actions might have to be updated (or deleted) if their object-key is  
-		// mutated (or removed)... here we use the node's class (so if the node 
-		// type disappears, then the action should go with it)
-		UClass* ActionKey = GetClass();
 		// to keep from needlessly instantiating a UBlueprintNodeSpawner, first   
 		// check to make sure that the registrar is looking for actions of this type
 		// (could be regenerating actions for a specific asset, and therefore the 
@@ -158,4 +159,12 @@ FText UK2Node_InputAxisEvent::GetMenuCategory() const
 	return CachedCategory;
 }
 
+FBlueprintNodeSignature UK2Node_InputAxisEvent::GetSignature() const
+{
+	FBlueprintNodeSignature NodeSignature = Super::GetSignature();
+	NodeSignature.AddKeyValue(InputAxisName.ToString());
+
+	return NodeSignature;
+}
+
 #undef LOCTEXT_NAMESPACE
diff --git a/Engine/Source/Editor/BlueprintGraph/Private/K2Node_InputAxisKeyEvent.cpp b/Engine/Source/Editor/BlueprintGraph/Private/K2Node_InputAxisKeyEvent.cpp
--- a/Engine/Source/Editor/BlueprintGraph/Private/K2Node_InputAxisKeyEvent.cpp
+++ b/Engine/Source/Editor/BlueprintGraph/Private/K2Node_InputAxisKeyEvent.cpp
@@ -127,6 +127,12 @@ void UK2Node_InputAxisKeyEvent::GetMenuActions(FBlueprintActionDatabaseRegistrar
 		InputNode->Initialize(Key);
 	};
 
+	// actions get registered under specific object-keys; the idea is that 
+	// actions might have to be updated (or deleted) if their object-key is  
+	// mutated (or removed)... here we use the node's class (so if the node 
+	// type disappears, then the action should go with it)
+	UClass* ActionKey = GetClass();
+
 	for (FKey const Key : AllKeys)
 	{
 		if (!Key.IsBindableInBlueprints() || !Key.IsFloatAxis())
@@ -134,11 +140,6 @@ void UK2Node_InputAxisKeyEvent::GetMenuActions(FBlueprintActionDatabaseRegistrar
 			continue;
 		}
 
-		// actions get registered under specific object-keys; the idea is that 
-		// actions might have to be updated (or deleted) if their object-key is  
-		// mutated (or removed)... here we use the node's class (so if the node 
-		// type disappears, then the action should go with it)
-		UClass* ActionKey = GetClass();
 		// to keep from needlessly instantiating a UBlueprintNodeSpawner, first   
 		// check to make sure that the registrar is looking for actions of this type
 		// (could be regenerating actions for a specific asset, and therefore the 
@@ -194,4 +195,12 @@ FText UK2Node_InputAxisKeyEvent::GetMenuCategory() const
 	return CachedCategories[CategoryIndex];
 }
 
+FBlueprintNodeSignature UK2Node_InputAxisKeyEvent::GetSignature() const
+{
+	FBlueprintNodeSignature NodeSignature = Super::GetSignature();
+	NodeSignature.AddKeyValue(AxisKey.ToString());
+
+	return NodeSignature;
+}
+
 #undef LOCTEXT_NAMESPACE
diff --git a/Engine/Source/Editor/BlueprintGraph/Private/K2Node_InputKey.cpp b/Engine/Source/Editor/BlueprintGraph/Private/K2Node_InputKey.cpp
--- a/Engine/Source/Editor/BlueprintGraph/Private/K2Node_InputKey.cpp
+++ b/Engine/Source/Editor/BlueprintGraph/Private/K2Node_InputKey.cpp
@@ -382,6 +382,12 @@ void UK2Node_InputKey::GetMenuActions(FBlueprintActionDatabaseRegistrar& ActionR
 		InputNode->InputKey = Key;
 	};
 
+	// actions get registered under specific object-keys; the idea is that 
+	// actions might have to be updated (or deleted) if their object-key is  
+	// mutated (or removed)... here we use the node's class (so if the node 
+	// type disappears, then the action should go with it)
+	UClass* ActionKey = GetClass();
+
 	for (FKey const Key : AllKeys)
 	{
 		// these will be handled by UK2Node_GetInputAxisKeyValue and UK2Node_GetInputVectorAxisValue respectively
@@ -390,11 +396,6 @@ void UK2Node_InputKey::GetMenuActions(FBlueprintActionDatabaseRegistrar& ActionR
 			continue;
 		}
 
-		// actions get registered under specific object-keys; the idea is that 
-		// actions might have to be updated (or deleted) if their object-key is  
-		// mutated (or removed)... here we use the node's class (so if the node 
-		// type disappears, then the action should go with it)
-		UClass* ActionKey = GetClass();
 		// to keep from needlessly instantiating a UBlueprintNodeSpawner, first   
 		// check to make sure that the registrar is looking for actions of this type
 		// (could be regenerating actions for a specific asset, and therefore the 
@@ -450,4 +451,12 @@ FText UK2Node_InputKey::GetMenuCategory() const
 	return CachedCategories[CategoryIndex];
 }
 
+FBlueprintNodeSignature UK2Node_InputKey::GetSignature() const
+{
+	FBlueprintNodeSignature NodeSignature = Super::GetSignature();
+	NodeSignature.AddKeyValue(InputKey.ToString());
+
+	return NodeSignature;
+}
+
 #undef LOCTEXT_NAMESPACE
diff --git a/Engine/Source/Editor/BlueprintGraph/Private/K2Node_MacroInstance.cpp b/Engine/Source/Editor/BlueprintGraph/Private/K2Node_MacroInstance.cpp
--- a/Engine/Source/Editor/BlueprintGraph/Private/K2Node_MacroInstance.cpp
+++ b/Engine/Source/Editor/BlueprintGraph/Private/K2Node_MacroInstance.cpp
@@ -444,5 +444,12 @@ FText UK2Node_MacroInstance::GetMenuCategory() const
 	return MenuCategory;
 }
 
+FBlueprintNodeSignature UK2Node_MacroInstance::GetSignature() const
+{
+	FBlueprintNodeSignature NodeSignature = Super::GetSignature();
+	NodeSignature.AddSubObject(GetMacroGraph());
+
+	return NodeSignature;
+}
 
 #undef LOCTEXT_NAMESPACE
diff --git a/Engine/Source/Editor/BlueprintGraph/Private/K2Node_TemporaryVariable.cpp b/Engine/Source/Editor/BlueprintGraph/Private/K2Node_TemporaryVariable.cpp
--- a/Engine/Source/Editor/BlueprintGraph/Private/K2Node_TemporaryVariable.cpp
+++ b/Engine/Source/Editor/BlueprintGraph/Private/K2Node_TemporaryVariable.cpp
@@ -155,7 +155,6 @@ void UK2Node_TemporaryVariable::GetMenuActions(FBlueprintActionDatabaseRegistrar
 			TempVarNode->bIsPersistent = bIsPersistent;
 		};
 
-		// @TODO: Need some kind of unique spawner so that this can generate a unique key
 		NodeSpawner->CustomizeNodeDelegate = UBlueprintNodeSpawner::FCustomizeNodeDelegate::CreateStatic(PostSpawnLambda, VarType, bIsPersistent);
 		return NodeSpawner;
 	};
@@ -190,7 +189,6 @@ void UK2Node_TemporaryVariable::GetMenuActions(FBlueprintActionDatabaseRegistrar
 	ActionRegistrar.AddBlueprintAction(ActionKey, MakeTempVarNodeSpawner(FEdGraphPinType(K2Schema->PC_Struct, TEXT("BlendSampleData"), BlendSampleStruct, /*bIsArray =*/ true, /*bIsReference =*/false), /*bIsPersistent =*/false));
 
 	// add persistent bool and int types (for macro graphs)
-	// @TODO: these should probably be filtered (for macro graphs only)
 	ActionRegistrar.AddBlueprintAction(ActionKey, MakeTempVarNodeSpawner(FEdGraphPinType(K2Schema->PC_Int, TEXT(""), nullptr, /*bIsArray =*/false, /*bIsReference =*/false), /*bIsPersistent =*/true));
 	ActionRegistrar.AddBlueprintAction(ActionKey, MakeTempVarNodeSpawner(FEdGraphPinType(K2Schema->PC_Boolean, TEXT(""), nullptr, /*bIsArray =*/false, /*bIsReference =*/false), /*bIsPersistent =*/true));
 }
@@ -199,3 +197,20 @@ FText UK2Node_TemporaryVariable::GetMenuCategory() const
 {
 	return FEditorCategoryUtils::GetCommonCategory(FCommonEditorCategory::Macro);
 }
+
+FBlueprintNodeSignature UK2Node_TemporaryVariable::GetSignature() const
+{
+	FBlueprintNodeSignature NodeSignature = Super::GetSignature();
+
+	FString TypeString;
+	if (bIsPersistent)
+	{
+		TypeString = TEXT("Persistent ");
+	}
+	TypeString += UEdGraphSchema_K2::TypeToText(VariableType).ToString();
+
+	static const FName VarTypeSignatureKey(TEXT("VarType"));
+	NodeSignature.AddNamedValue(VarTypeSignatureKey, TypeString);
+
+	return NodeSignature;
+}
diff --git a/Engine/Source/Editor/BlueprintGraph/Public/BlueprintActionFilter.h b/Engine/Source/Editor/BlueprintGraph/Public/BlueprintActionFilter.h
--- a/Engine/Source/Editor/BlueprintGraph/Public/BlueprintActionFilter.h
+++ b/Engine/Source/Editor/BlueprintGraph/Public/BlueprintActionFilter.h
@@ -93,13 +93,6 @@ struct BLUEPRINTGRAPH_API FBlueprintActionInfo
 	 */
 	UClass const* GetNodeClass();
 
-	/**
-	 * 
-	 * 
-	 * @return 
-	 */
-	FGuid GetActionGuid();
-
 	/**
 	 * Certain actions are associated with specific member fields (a member 
 	 * function call, a variable get/set, etc.) This retrieves that member field 
@@ -144,9 +137,6 @@ private:
 	UProperty const* CachedActionProperty;
 	UFunction const* CachedActionFunction;
 
-	/** */
-	FGuid CachedActionId;
-
 	/** */
 	IBlueprintNodeBinder::FBindingSet Bindings;
 };
diff --git a/Engine/Source/Editor/BlueprintGraph/Public/BlueprintNodeSignature.h b/Engine/Source/Editor/BlueprintGraph/Public/BlueprintNodeSignature.h
new file mode 100644
--- /dev/null
+++ b/Engine/Source/Editor/BlueprintGraph/Public/BlueprintNodeSignature.h
@@ -0,0 +1,81 @@
+// Copyright 1998-2014 Epic Games, Inc. All Rights Reserved.
+
+#pragma once
+
+struct BLUEPRINTGRAPH_API FBlueprintNodeSignature
+{
+public:
+	FBlueprintNodeSignature() {}
+	FBlueprintNodeSignature(FString const& UserString);
+	FBlueprintNodeSignature(TSubclassOf<UEdGraphNode> NodeClass);
+
+	/**
+	 * 
+	 * 
+	 * @param  NodeClass	
+	 * @return 
+	 */
+	void SetNodeClass(TSubclassOf<UEdGraphNode> NodeClass);
+
+	/**
+	 * 
+	 * 
+	 * @param  SignatureObj	
+	 * @return 
+	 */
+	void AddSubObject(UObject const* SignatureObj);
+
+	/**
+	 * 
+	 * 
+	 * @param  Value	
+	 * @return 
+	 */
+	void AddKeyValue(FString const& KeyValue);
+
+	/**
+	 * 
+	 * 
+	 * @param  SignatureKey	
+	 * @param  Value	
+	 * @return 
+	 */
+	void AddNamedValue(FName SignatureKey, FString const& Value);
+
+	/**
+	 * 
+	 * 
+	 * @return 
+	 */
+	bool IsValid() const;
+
+	/**
+	 * 
+	 * 
+	 * @return 
+	 */
+	FString const& ToString() const;
+
+	/**
+	*
+	*
+	* @return
+	*/
+	FGuid const& AsGuid() const;
+
+private:
+	/**
+	 * 
+	 * 
+	 * @return 
+	 */
+	void MarkDirty();
+
+private:
+	typedef TMap<FName, FString> FSignatureSet;
+	FSignatureSet SignatureSet;
+
+	mutable FGuid   CachedSignatureGuid;
+	mutable FString CachedSignatureString;
+};
+
diff --git a/Engine/Source/Editor/BlueprintGraph/Public/BlueprintNodeSpawner.h b/Engine/Source/Editor/BlueprintGraph/Public/BlueprintNodeSpawner.h
--- a/Engine/Source/Editor/BlueprintGraph/Public/BlueprintNodeSpawner.h
+++ b/Engine/Source/Editor/BlueprintGraph/Public/BlueprintNodeSpawner.h
@@ -3,7 +3,7 @@
 #pragma once
 
 #include "BlueprintNodeBinder.h"
-#include "BlueprintNodeSpawnerSignature.h"
+#include "BlueprintNodeSignature.h"
 #include "BlueprintNodeSpawner.generated.h"
 
 // Forward declarations
@@ -88,7 +88,7 @@ public:
 	 * 
 	 * @return A set of object-paths/names that distinguish this spawner from others.
 	 */
-	virtual FBlueprintNodeSpawnerSignature GetSpawnerSignature() const;
+	virtual FBlueprintNodeSignature GetSpawnerSignature() const;
 
 	/**
 	 * Takes care of spawning a node for the specified graph. Looks to see if 
diff --git a/Engine/Source/Editor/BlueprintGraph/Public/BlueprintNodeSpawnerSignature.h b/Engine/Source/Editor/BlueprintGraph/Public/BlueprintNodeSpawnerSignature.h
deleted file mode 100644
diff --git a/Engine/Source/Editor/Kismet/Classes/BlueprintPaletteFavorites.h b/Engine/Source/Editor/Kismet/Classes/BlueprintPaletteFavorites.h
--- a/Engine/Source/Editor/Kismet/Classes/BlueprintPaletteFavorites.h
+++ b/Engine/Source/Editor/Kismet/Classes/BlueprintPaletteFavorites.h
@@ -2,7 +2,7 @@
 
 #pragma once
 
-#include "BlueprintNodeSpawnerSignature.h"
+#include "BlueprintNodeSignature.h"
 #include "BlueprintPaletteFavorites.generated.h"
 
 // Forward Declarations
@@ -69,22 +69,6 @@ public:
 	 */
 	bool operator==(TSharedPtr<FEdGraphSchemaAction> PaletteAction) const;
 
-	/**
-	 * 
-	 * 
-	 * @param  BlueprintAction	
-	 * @return 
-	 */
-	bool operator==(FBlueprintActionInfo& BlueprintAction) const;
-
-	/**
-	 * 
-	 * 
-	 * @param  ActionGuid	
-	 * @return 
-	 */
-	bool operator==(FGuid const& ActionGuid) const;
-
 	/**
 	 * We want to be able to specify some of these in .ini files, so let we have
 	 * to have a readable string representation for them.
@@ -94,7 +78,7 @@ public:
 	FString const& ToString() const;
 
 private:
-	FBlueprintNodeSpawnerSignature ActionSignature;
+	FBlueprintNodeSignature ActionSignature;
 };
 
 /*******************************************************************************
diff --git a/Engine/Source/Editor/Kismet/Private/BlueprintPaletteFavorites.cpp b/Engine/Source/Editor/Kismet/Private/BlueprintPaletteFavorites.cpp
--- a/Engine/Source/Editor/Kismet/Private/BlueprintPaletteFavorites.cpp
+++ b/Engine/Source/Editor/Kismet/Private/BlueprintPaletteFavorites.cpp
@@ -22,17 +22,17 @@ namespace BlueprintPaletteFavoritesImpl
 	/**
 	 * Before we refactored the blueprint menu system, signatures were manually
 	 * constructed based off node type, by combining a series of objects and 
-	 * names. Here we construct a new FBlueprintNodeSpawnerSignature from the 
+	 * names. Here we construct a new FBlueprintNodeSignature from the 
 	 * old code (so as to mirror functionality).
 	 * 
 	 * @param  PaletteAction	The action you want a signature for.
 	 * @return A signature object, distinguishing the palette action from others (could also be invalid).
 	 */
-	static FBlueprintNodeSpawnerSignature ConstructLegacySignature(TSharedPtr<FEdGraphSchemaAction> PaletteAction);
+	static FBlueprintNodeSignature ConstructLegacySignature(TSharedPtr<FEdGraphSchemaAction> PaletteAction);
 }
 
 //------------------------------------------------------------------------------
-static FBlueprintNodeSpawnerSignature BlueprintPaletteFavoritesImpl::ConstructLegacySignature(TSharedPtr<FEdGraphSchemaAction> InPaletteAction)
+static FBlueprintNodeSignature BlueprintPaletteFavoritesImpl::ConstructLegacySignature(TSharedPtr<FEdGraphSchemaAction> InPaletteAction)
 {
 	TSubclassOf<UEdGraphNode> SignatureNodeClass;
 	UObject const* SignatureSubObject = nullptr;
@@ -116,7 +116,7 @@ static FBlueprintNodeSpawnerSignature BlueprintPaletteFavoritesImpl::ConstructLe
 		}
 	}
 
-	FBlueprintNodeSpawnerSignature LegacySignatureSet;
+	FBlueprintNodeSignature LegacySignatureSet;
 	if (SignatureNodeClass != nullptr)
 	{
 		LegacySignatureSet.SetNodeClass(SignatureNodeClass);
@@ -126,8 +126,7 @@ static FBlueprintNodeSpawnerSignature BlueprintPaletteFavoritesImpl::ConstructLe
 		}
 		else if (SignatureSubObjName != NAME_None)
 		{
-			static FName const SubObjectSignatureKey("FieldName");
-			LegacySignatureSet.AddKeyValue(SubObjectSignatureKey, SignatureSubObjName.ToString());
+			LegacySignatureSet.AddKeyValue(SignatureSubObjName.ToString());
 		}
 	}
 
@@ -147,6 +146,8 @@ FFavoritedBlueprintPaletteItem::FFavoritedBlueprintPaletteItem(FString const& Se
 //------------------------------------------------------------------------------
 FFavoritedBlueprintPaletteItem::FFavoritedBlueprintPaletteItem(TSharedPtr<FEdGraphSchemaAction> InPaletteAction)
 {
+	using namespace BlueprintPaletteFavoritesImpl;
+
 	if (InPaletteAction.IsValid())
 	{
 		if (InPaletteAction->GetTypeId() == FBlueprintActionMenuItem::StaticGetTypeId())
@@ -160,12 +161,12 @@ FFavoritedBlueprintPaletteItem::FFavoritedBlueprintPaletteItem(TSharedPtr<FEdGra
 			ActionSignature = CollectionMenuItem->GetSampleAction()->GetSpawnerSignature();
 
 			// drag-n-drop menu items represent a collection of actions on the 
-			// same field (they spawn a submenu for the user to pick from)... so
+			// same field (they spawn a sub-menu for the user to pick from), so
 			// they don't have a single node class
 			ActionSignature.SetNodeClass(nullptr);
 
 			static const FName CollectionSignatureKey(TEXT("ActionCollection"));
-			ActionSignature.AddKeyValue(CollectionSignatureKey, TEXT("true"));
+			ActionSignature.AddNamedValue(CollectionSignatureKey, TEXT("true"));
 		}
 		else if (InPaletteAction->GetTypeId() == FBlueprintBoundMenuItem::StaticGetTypeId())
 		{
@@ -207,18 +208,6 @@ bool FFavoritedBlueprintPaletteItem::operator==(TSharedPtr<FEdGraphSchemaAction>
 	return (*this == FFavoritedBlueprintPaletteItem(PaletteAction));
 }
 
-//------------------------------------------------------------------------------
-bool FFavoritedBlueprintPaletteItem::operator==(FBlueprintActionInfo& BlueprintAction) const
-{
-	return (BlueprintAction.GetActionGuid() == ActionSignature.AsGuid());
-}
-
-//------------------------------------------------------------------------------
-bool FFavoritedBlueprintPaletteItem::operator==(FGuid const& ActionGuid) const
-{
-	return (ActionGuid == ActionSignature.AsGuid());
-}
-
 //------------------------------------------------------------------------------
 FString const& FFavoritedBlueprintPaletteItem::ToString() const
 {
@@ -314,17 +303,7 @@ bool UBlueprintPaletteFavorites::IsFavorited(TSharedPtr<FEdGraphSchemaAction> Pa
 //------------------------------------------------------------------------------
 bool UBlueprintPaletteFavorites::IsFavorited(FBlueprintActionInfo& BlueprintAction) const
 {
-	bool bIsFavorited = false;
-
-	for (FFavoritedBlueprintPaletteItem const& Favorite : CurrentFavorites)
-	{
-		if (Favorite == BlueprintAction)
-		{
-			bIsFavorited = true;
-			break;
-		}
-	}
-	return bIsFavorited;
+	return IsFavorited(BlueprintAction.NodeSpawner);
 }
 
 //------------------------------------------------------------------------------
@@ -332,10 +311,10 @@ bool UBlueprintPaletteFavorites::IsFavorited(UBlueprintNodeSpawner const* Bluepr
 {
 	bool bIsFavorited = false;
 
-	FGuid ActionGuid = BlueprintAction->GetSpawnerSignature().AsGuid();
+	FFavoritedBlueprintPaletteItem ActionAsFavorite(BlueprintAction);
 	for (FFavoritedBlueprintPaletteItem const& Favorite : CurrentFavorites)
 	{
-		if (Favorite == ActionGuid)
+		if (Favorite == ActionAsFavorite)
 		{
 			bIsFavorited = true;
 			break;
