#ue4 - Strip lines with "[REDACTED_HANDLE]" from function tooltips.

diff --git a/Engine/Source/Editor/BlueprintGraph/Private/K2Node_CallFunction.cpp b/Engine/Source/Editor/BlueprintGraph/Private/K2Node_CallFunction.cpp
--- a/Engine/Source/Editor/BlueprintGraph/Private/K2Node_CallFunction.cpp
+++ b/Engine/Source/Editor/BlueprintGraph/Private/K2Node_CallFunction.cpp
@@ -852,6 +852,7 @@ FString UK2Node_CallFunction::GetDefaultTooltipForFunction(const UFunction* Func
 		// Strip off the @param nastiness
 		Tooltip.Split(TEXT("@param"), &Tooltip, nullptr, ESearchCase::IgnoreCase, ESearchDir::FromStart);
 		Tooltip.Split(TEXT("@return"), &Tooltip, nullptr, ESearchCase::IgnoreCase, ESearchDir::FromStart);
+		Tooltip.Split(TEXT("@see"), &Tooltip, nullptr, ESearchCase::IgnoreCase, ESearchDir::FromStart);
 		Tooltip.Trim();
 		while (Tooltip.RemoveFromEnd(TEXT("\n"))) ;
 
