Macros with latent functions no longer appear in right-click context menu for function graphs (using new menu system).

diff --git a/Engine/Source/Editor/BlueprintGraph/Classes/K2Node.h b/Engine/Source/Editor/BlueprintGraph/Classes/K2Node.h
--- a/Engine/Source/Editor/BlueprintGraph/Classes/K2Node.h
+++ b/Engine/Source/Editor/BlueprintGraph/Classes/K2Node.h
@@ -187,6 +187,11 @@ class UK2Node : public UEdGraphNode
 	/** Return tooltip text that explains the result of an active breakpoint on this node */
 	BLUEPRINTGRAPH_API virtual FText GetActiveBreakpointToolTipText() const;
 
+	/**
+	 * Determine if the node of this type should be filtered in the actions menu
+	 */
+	virtual bool IsActionFilteredOut(class FBlueprintActionFilter const& Filter) { return false; }
+
 	/** Should draw as a bead with no location of it's own */
 	virtual bool ShouldDrawAsBead() const { return false; }
 
diff --git a/Engine/Source/Editor/BlueprintGraph/Classes/K2Node_MacroInstance.h b/Engine/Source/Editor/BlueprintGraph/Classes/K2Node_MacroInstance.h
--- a/Engine/Source/Editor/BlueprintGraph/Classes/K2Node_MacroInstance.h
+++ b/Engine/Source/Editor/BlueprintGraph/Classes/K2Node_MacroInstance.h
@@ -59,6 +59,7 @@ public:
 	virtual FText GetMenuCategory() const override;
 	virtual int32 GetNodeRefreshPriority() const override { return EBaseNodeRefreshPriority::Low_UsesDependentWildcard; }
 	virtual FBlueprintNodeSignature GetSignature() const override;
+	virtual bool IsActionFilteredOut(class FBlueprintActionFilter const& Filter) override;
 	// End UK2Node interface
 
 	void SetMacroGraph(UEdGraph* Graph) { MacroGraphReference.SetGraph(Graph); }
diff --git a/Engine/Source/Editor/BlueprintGraph/Private/BlueprintActionFilter.cpp b/Engine/Source/Editor/BlueprintGraph/Private/BlueprintActionFilter.cpp
--- a/Engine/Source/Editor/BlueprintGraph/Private/BlueprintActionFilter.cpp
+++ b/Engine/Source/Editor/BlueprintGraph/Private/BlueprintActionFilter.cpp
@@ -264,6 +264,16 @@ namespace BlueprintActionFilterImpl
 	 * @return 
 	 */
 	static bool IsNotSubClassCast(FBlueprintActionFilter const& Filter, FBlueprintActionInfo& BlueprintAction);
+
+	/**
+	 * Rejection test that checks if the NodeTemplate (if any)
+	 * decides to filter itself out of the action list
+	 * 
+	 * @param  Filter			Holds the graph context for this test.
+	 * @param  BlueprintAction	The action you wish to query.
+	 * @return
+	 */
+	static bool IsNodeTemplateSelfFiltered(FBlueprintActionFilter const& Filter, FBlueprintActionInfo& BlueprintAction);
 };
 
 
@@ -1035,6 +1045,17 @@ static bool BlueprintActionFilterImpl::IsNotSubClassCast(FBlueprintActionFilter
 	return bIsFilteredOut;
 }
 
+static bool BlueprintActionFilterImpl::IsNodeTemplateSelfFiltered(FBlueprintActionFilter const& Filter, FBlueprintActionInfo& BlueprintAction)
+{
+	bool bIsFilteredOut = false;
+
+	if(UK2Node* NodeTemplate = Cast<UK2Node>(BlueprintAction.NodeSpawner->GetTemplateNode()))
+	{
+		bIsFilteredOut = NodeTemplate->IsActionFilteredOut(Filter);
+	}
+	return bIsFilteredOut;
+}
+
 /*******************************************************************************
  * FBlueprintActionInfo
  ******************************************************************************/
@@ -1188,6 +1209,7 @@ FBlueprintActionFilter::FBlueprintActionFilter(uint32 Flags/*= 0x00*/)
 	//
 	// this test in-particular spawns a template-node and then calls 
 	// AllocateDefaultPins() which is costly, so it should be very last!
+	AddRejectionTest(FRejectionTestDelegate::CreateStatic(IsNodeTemplateSelfFiltered));
 	AddRejectionTest(FRejectionTestDelegate::CreateStatic(IsMissingMatchingPinParam));
 	AddRejectionTest(FRejectionTestDelegate::CreateStatic(IsMissmatchedPropertyType));
 	AddRejectionTest(FRejectionTestDelegate::CreateStatic(IsFunctionMissingPinParam));
diff --git a/Engine/Source/Editor/BlueprintGraph/Private/K2Node_MacroInstance.cpp b/Engine/Source/Editor/BlueprintGraph/Private/K2Node_MacroInstance.cpp
--- a/Engine/Source/Editor/BlueprintGraph/Private/K2Node_MacroInstance.cpp
+++ b/Engine/Source/Editor/BlueprintGraph/Private/K2Node_MacroInstance.cpp
@@ -3,6 +3,7 @@
 #include "BlueprintGraphPrivatePCH.h"
 #include "Slate.h"
 #include "EditorCategoryUtils.h"
+#include "BlueprintActionFilter.h"
 
 #define LOCTEXT_NAMESPACE "K2Node_MacroInstance"
 
@@ -22,6 +23,23 @@ void UK2Node_MacroInstance::Serialize(FArchive& Ar)
 	}
 }
 
+bool UK2Node_MacroInstance::IsActionFilteredOut(FBlueprintActionFilter const& Filter)
+{
+	bool bIsFilteredOut = false;
+	FBlueprintActionContext const& FilterContext = Filter.Context;
+
+	for (UEdGraph* Graph : FilterContext.Graphs)
+	{
+		// Macro Instances are not allowed in it's own graph, nor in Function graphs if the macro has latent functions in it
+		if (Graph == GetMacroGraph() || (Graph->GetSchema()->GetGraphType(Graph) == GT_Function && FBlueprintEditorUtils::CheckIfGraphHasLatentFunctions(GetMacroGraph())) )
+		{
+			bIsFilteredOut = true;
+			break;
+		}
+	}
+	return bIsFilteredOut;
+}
+
 void UK2Node_MacroInstance::PostPasteNode()
 {
 	const UBlueprint* InstanceOwner = GetBlueprint();
