Blueprints - Fix lost connections to Select nodes which can occur when the underlying Enum has a redirect.

- Note: This issue was introduced in CL# 2286152.
- Added UK2Node_Select::DoPinsMatchForReconstruction() override to ensure a match for old pins that have a redirect to new Enum values.

diff --git a/Engine/Source/Editor/BlueprintGraph/Classes/K2Node_Select.h b/Engine/Source/Editor/BlueprintGraph/Classes/K2Node_Select.h
--- a/Engine/Source/Editor/BlueprintGraph/Classes/K2Node_Select.h
+++ b/Engine/Source/Editor/BlueprintGraph/Classes/K2Node_Select.h
@@ -47,6 +47,7 @@ class UK2Node_Select : public UK2Node, public INodeDependingOnEnumInterface
 	// End UEdGraphNode interface
 
 	// Begin UK2Node interface
+	virtual ERedirectType DoPinsMatchForReconstruction(const UEdGraphPin* NewPin, int32 NewPinIndex, const UEdGraphPin* OldPin, int32 OldPinIndex) const override;
 	virtual void ReallocatePinsDuringReconstruction(TArray<UEdGraphPin*>& OldPins) override;
 	virtual void PostReconstructNode() override;
 	virtual bool IsNodeSafeToIgnore() const override { return true; }
diff --git a/Engine/Source/Editor/BlueprintGraph/Private/K2Node_Select.cpp b/Engine/Source/Editor/BlueprintGraph/Private/K2Node_Select.cpp
--- a/Engine/Source/Editor/BlueprintGraph/Private/K2Node_Select.cpp
+++ b/Engine/Source/Editor/BlueprintGraph/Private/K2Node_Select.cpp
@@ -271,6 +271,29 @@ FText UK2Node_Select::GetNodeTitle(ENodeTitleType::Type TitleType) const
 	return LOCTEXT("Select", "Select");
 }
 
+UK2Node::ERedirectType UK2Node_Select::DoPinsMatchForReconstruction(const UEdGraphPin* NewPin, int32 NewPinIndex, const UEdGraphPin* OldPin, int32 OldPinIndex) const
+{
+	// Check to see if the new pin name matches the old pin name (case insensitive - since the base uses Stricmp() to compare pin names, we also ignore case here).
+	if(Enum != nullptr && NewPinIndex < NumOptionPins && !NewPin->PinName.Equals(OldPin->PinName, ESearchCase::IgnoreCase))
+	{
+		// The names don't match, so check for an enum redirect from the old pin name.
+		int32 EnumIndex = UEnum::FindEnumRedirects(Enum, FName(*OldPin->PinName));
+		if(EnumIndex != INDEX_NONE)
+		{
+			// Found a redirect. Attempt to match it to the new pin name.
+			FString NewPinName = Enum->GetEnumName(EnumIndex);
+			if(NewPinName.Equals(NewPin->PinName, ESearchCase::IgnoreCase))
+			{
+				// The redirect is a match, so we can reconstruct this pin using the old pin's state.
+				return UK2Node::ERedirectType_Name;
+			}
+		}
+	}
+
+	// Fall back to base class functionality for all other cases.
+	return Super::DoPinsMatchForReconstruction(NewPin, NewPinIndex, OldPin, OldPinIndex);
+}
+
 void UK2Node_Select::ReallocatePinsDuringReconstruction(TArray<UEdGraphPin*>& OldPins)
 {
 	Super::ReallocatePinsDuringReconstruction(OldPins);
