Better error message when proxyclass fails to load

diff --git a/Engine/Source/Editor/BlueprintGraph/Private/K2Node_BaseAsyncTask.cpp b/Engine/Source/Editor/BlueprintGraph/Private/K2Node_BaseAsyncTask.cpp
--- a/Engine/Source/Editor/BlueprintGraph/Private/K2Node_BaseAsyncTask.cpp
+++ b/Engine/Source/Editor/BlueprintGraph/Private/K2Node_BaseAsyncTask.cpp
@@ -398,7 +398,12 @@ void UK2Node_BaseAsyncTask::GetMenuActions(FBlueprintActionDatabaseRegistrar& Ac
 
 UFunction* UK2Node_BaseAsyncTask::GetFactoryFunction() const
 {
-	check(ProxyFactoryClass);
+	if (ProxyFactoryClass == nullptr)
+	{
+		UE_LOG(LogBlueprint, Fatal, TEXT("ProxyFactoryClass null in %s. Was a class deleted or saved on a non promoted build?"), *GetFullName());
+		return nullptr;
+	}
+	
 	UFunction* FactoryFunction = ProxyFactoryClass->FindFunctionByName(ProxyFactoryFunctionName);
 	check(FactoryFunction);
 	return FactoryFunction;
