#ue4 - Strip "[REDACTED_HANDLE]" from the end of all tooltips for enums. Make enums use UEnum::GetToolTipText() rather than reading metadata directly.

- Also converted related functions to use static strings.

diff --git a/Engine/Source/Editor/BlueprintGraph/Private/K2Node_CallFunction.cpp b/Engine/Source/Editor/BlueprintGraph/Private/K2Node_CallFunction.cpp
--- a/Engine/Source/Editor/BlueprintGraph/Private/K2Node_CallFunction.cpp
+++ b/Engine/Source/Editor/BlueprintGraph/Private/K2Node_CallFunction.cpp
@@ -849,12 +849,16 @@ FString UK2Node_CallFunction::GetDefaultTooltipForFunction(const UFunction* Func
 
 	if (!Tooltip.IsEmpty())
 	{
-		// Strip off the @param nastiness
-		Tooltip.Split(TEXT("@param"), &Tooltip, nullptr, ESearchCase::IgnoreCase, ESearchDir::FromStart);
-		Tooltip.Split(TEXT("@return"), &Tooltip, nullptr, ESearchCase::IgnoreCase, ESearchDir::FromStart);
-		Tooltip.Split(TEXT("@see"), &Tooltip, nullptr, ESearchCase::IgnoreCase, ESearchDir::FromStart);
+		// Strip off the doxygen nastiness
+		static const FString DoxygenParam(TEXT("@param"));
+		static const FString DoxygenReturn(TEXT("@return"));
+		static const FString DoxygenSee(TEXT("@see"));
+
+		Tooltip.Split(DoxygenParam, &Tooltip, nullptr, ESearchCase::IgnoreCase, ESearchDir::FromStart);
+		Tooltip.Split(DoxygenReturn, &Tooltip, nullptr, ESearchCase::IgnoreCase, ESearchDir::FromStart);
+		Tooltip.Split(DoxygenSee, &Tooltip, nullptr, ESearchCase::IgnoreCase, ESearchDir::FromStart);
 		Tooltip.Trim();
-		while (Tooltip.RemoveFromEnd(TEXT("\n"))) ;
+		Tooltip.TrimTrailing();
 
 		return Tooltip;
 	}
diff --git a/Engine/Source/Editor/PropertyEditor/Private/PropertyHandleImpl.cpp b/Engine/Source/Editor/PropertyEditor/Private/PropertyHandleImpl.cpp
--- a/Engine/Source/Editor/PropertyEditor/Private/PropertyHandleImpl.cpp
+++ b/Engine/Source/Editor/PropertyEditor/Private/PropertyHandleImpl.cpp
@@ -1912,7 +1912,7 @@ bool FPropertyHandleBase::GeneratePossibleValues(TArray< TSharedPtr<FString> >&
 				TSharedPtr< FString > EnumStr( new FString( EnumValueName ) );
 				OutOptionStrings.Add( EnumStr );
 
-				FString EnumValueToolTip = bIsRestricted ? RestrictionTooltip.ToString() : Enum->GetMetaData(TEXT("ToolTip"), EnumIndex);
+				FString EnumValueToolTip = bIsRestricted ? RestrictionTooltip.ToString() : Enum->GetToolTipText(EnumIndex).ToString();
 				TSharedPtr< FString > ToolTipStr( new FString( EnumValueToolTip ) );
 				OutToolTips.Add( ToolTipStr );
 			}
diff --git a/Engine/Source/Editor/PropertyEditor/Private/UserInterface/PropertyEditor/SPropertyEditorCombo.cpp b/Engine/Source/Editor/PropertyEditor/Private/UserInterface/PropertyEditor/SPropertyEditorCombo.cpp
--- a/Engine/Source/Editor/PropertyEditor/Private/UserInterface/PropertyEditor/SPropertyEditorCombo.cpp
+++ b/Engine/Source/Editor/PropertyEditor/Private/UserInterface/PropertyEditor/SPropertyEditorCombo.cpp
@@ -152,7 +152,7 @@ void SPropertyEditorCombo::SendToObjects( const FString& NewValue )
 
 		Value = Enum->GetEnumName(Index);
 
-		ToolTipValue = Enum->GetMetaData( TEXT("ToolTip"), Index );
+		ToolTipValue = Enum->GetToolTipText(Index).ToString();
 		FString ToolTipText = Property->GetToolTipText().ToString();
 		if (ToolTipValue.Len() > 0)
 		{
diff --git a/Engine/Source/Runtime/CoreUObject/Private/UObject/Class.cpp b/Engine/Source/Runtime/CoreUObject/Private/UObject/Class.cpp
--- a/Engine/Source/Runtime/CoreUObject/Private/UObject/Class.cpp
+++ b/Engine/Source/Runtime/CoreUObject/Private/UObject/Class.cpp
@@ -245,7 +245,8 @@ FText UField::GetToolTipText() const
 		}
 		else
 		{
-			if (NativeToolTip.Split(TEXT("@see"), &NativeToolTip, nullptr, ESearchCase::IgnoreCase, ESearchDir::FromStart))
+			static const FString DoxygenSee(TEXT("@see"));
+			if (NativeToolTip.Split(DoxygenSee, &NativeToolTip, nullptr, ESearchCase::IgnoreCase, ESearchDir::FromStart))
 			{
 				NativeToolTip.TrimTrailing();
 			}
diff --git a/Engine/Source/Runtime/CoreUObject/Private/UObject/Enum.cpp b/Engine/Source/Runtime/CoreUObject/Private/UObject/Enum.cpp
--- a/Engine/Source/Runtime/CoreUObject/Private/UObject/Enum.cpp
+++ b/Engine/Source/Runtime/CoreUObject/Private/UObject/Enum.cpp
@@ -321,15 +321,21 @@ FText UEnum::GetDisplayNameText(int32 NameIndex) const
 FText UEnum::GetToolTipText(int32 NameIndex) const
 {
 	FText LocalizedToolTip;
-	const FString NativeToolTip = GetMetaData( TEXT("ToolTip"), NameIndex );
+	FString NativeToolTip = GetMetaData( TEXT("ToolTip"), NameIndex );
 
-	FString Namespace = TEXT("UObjectToolTips");
+	static const FString Namespace = TEXT("UObjectToolTips");
 	FString Key =	NameIndex == INDEX_NONE
 		?			GetFullGroupName(true) + TEXT(".") + GetName()
 		:			GetFullGroupName(true) + TEXT(".") + GetName() + TEXT(".") + GetEnumName(NameIndex);
 		
 	if ( !(FText::FindText( Namespace, Key, /*OUT*/LocalizedToolTip )) || *FTextInspector::GetSourceString(LocalizedToolTip) != NativeToolTip)
 	{
+		static const FString DoxygenSee(TEXT("@see"));
+		if (NativeToolTip.Split(DoxygenSee, &NativeToolTip, nullptr, ESearchCase::IgnoreCase, ESearchDir::FromStart))
+		{
+			NativeToolTip.TrimTrailing();
+		}
+
 		LocalizedToolTip = FText::FromString(NativeToolTip);
 	}
 
