Menu item to create references in the level blueprint to things in scene outliner now gives user more information, more closely matches existing functionality

diff --git a/Engine/Source/Editor/BlueprintGraph/Private/K2Node_Literal.cpp b/Engine/Source/Editor/BlueprintGraph/Private/K2Node_Literal.cpp
--- a/Engine/Source/Editor/BlueprintGraph/Private/K2Node_Literal.cpp
+++ b/Engine/Source/Editor/BlueprintGraph/Private/K2Node_Literal.cpp
@@ -118,11 +118,6 @@ FText UK2Node_Literal::GetTooltipText() const
 FText UK2Node_Literal::GetNodeTitle(ENodeTitleType::Type TitleType) const
 {
 	AActor* Actor = Cast<AActor>(ObjectRef);
-	if(TitleType == ENodeTitleType::MenuTitle)
-	{
-		return NSLOCTEXT("K2Node", "Literal_Title", "Get Selected Actor Reference(s)");
-	}
-
 	if( ObjectRef != NULL )
 	{
 		if(Actor != NULL)
@@ -176,9 +171,26 @@ void UK2Node_Literal::GetMenuActions(FBlueprintActionDatabaseRegistrar& ActionRe
  		return true;
 	};
 
+	auto MenuDescriptionLambda = []( const IBlueprintNodeBinder::FBindingSet& BindingContext ) -> FText
+	{
+		if (BindingContext.Num() == 1)
+		{
+			return FText::Format(NSLOCTEXT("K2Node", "LiteralTitle", "Create a Reference to {0}"), FText::FromString((*(BindingContext.CreateConstIterator()))->GetName()));
+		}
+		else if (BindingContext.Num() > 1)
+		{
+			return FText::Format(NSLOCTEXT("K2Node", "FallbackLiteralTitle", "Create References to {0} selected Actors"), FText::AsNumber(BindingContext.Num()));
+		}
+		else
+		{
+			return NSLOCTEXT("K2Node", "FallbackLiteralTitle", "Error: No Actors in Context");
+		}
+	};
+
 	UBlueprintBoundNodeSpawner* NodeSpawner = UBlueprintBoundNodeSpawner::Create(GetClass());
 	NodeSpawner->CanBindObjectDelegate = UBlueprintBoundNodeSpawner::FCanBindObjectDelegate::CreateStatic(CanBindObjectLambda);
 	NodeSpawner->OnBindObjectDelegate = UBlueprintBoundNodeSpawner::FOnBindObjectDelegate::CreateStatic(PostBindSetupLambda);
+	NodeSpawner->OnGenerateMenuDescriptionDelegate = UBlueprintBoundNodeSpawner::FOnGenerateMenuDescriptionDelegate::CreateStatic(MenuDescriptionLambda);
 	ActionRegistrar.AddBlueprintAction(NodeSpawner);
 }
 
