Guarding against registering the OnBlueprintChanged callback twice (crash in the new blueprint action database).

diff --git a/Engine/Source/Editor/BlueprintGraph/Private/BlueprintActionDatabase.cpp b/Engine/Source/Editor/BlueprintGraph/Private/BlueprintActionDatabase.cpp
--- a/Engine/Source/Editor/BlueprintGraph/Private/BlueprintActionDatabase.cpp
+++ b/Engine/Source/Editor/BlueprintGraph/Private/BlueprintActionDatabase.cpp
@@ -12,7 +12,6 @@
 #include "BlueprintBoundEventNodeSpawner.h"
 #include "BlueprintVariableNodeSpawner.h"
 #include "Stats.h"					// for RETURN_QUICK_DECLARE_CYCLE_STAT()
-#include "ScopedTimers.h"			// for FScopedDurationTimer
 #include "CallbackDevice.h"			// for FCoreDelegates::OnAssetLoaded
 #include "ModuleManager.h"
 #include "AssetRegistryModule.h"	// for OnAssetAdded()/OnAssetRemoved()
@@ -503,11 +502,11 @@ static void BlueprintActionDatabaseImpl::AddClassCastActions(UClass* const Class
 			CastNode->TargetType = TargetType;
 		};
 
-		UBlueprintNodeSpawner* CastObjNodeSpawner = UBlueprintNodeSpawner::Create(UK2Node_DynamicCast::StaticClass());
+		UBlueprintNodeSpawner* CastObjNodeSpawner = UBlueprintNodeSpawner::Create<UK2Node_DynamicCast>();
 		CastObjNodeSpawner->CustomizeNodeDelegate = UBlueprintNodeSpawner::FCustomizeNodeDelegate::CreateStatic(CustomizeCastNodeLambda, Class);
 		ActionListOut.Add(CastObjNodeSpawner);
 
-		UBlueprintNodeSpawner* CastClassNodeSpawner = UBlueprintNodeSpawner::Create(UK2Node_ClassDynamicCast::StaticClass());
+		UBlueprintNodeSpawner* CastClassNodeSpawner = UBlueprintNodeSpawner::Create<UK2Node_ClassDynamicCast>();
 		CastClassNodeSpawner->CustomizeNodeDelegate = CastObjNodeSpawner->CustomizeNodeDelegate;
 		ActionListOut.Add(CastClassNodeSpawner);
 	}
@@ -600,7 +599,6 @@ static void BlueprintActionDatabaseImpl::OnAssetLoaded(UObject* NewObject)
 	if (UBlueprint* NewBlueprint = Cast<UBlueprint>(NewObject))
 	{
 		OnBlueprintChanged(NewBlueprint);
-		NewBlueprint->OnChanged().AddStatic(&BlueprintActionDatabaseImpl::OnBlueprintChanged);
 	}
 	else
 	{
@@ -617,20 +615,7 @@ static void BlueprintActionDatabaseImpl::OnAssetAdded(FAssetData const& NewAsset
 		UObject* AssetObject = NewAssetInfo.GetAsset();
 		if (UBlueprint* NewBlueprint = Cast<UBlueprint>(AssetObject))
 		{
-			FBlueprintActionDatabase& ActionDatabase = FBlueprintActionDatabase::Get();
-			// this callback can be triggered twice from the asset registry (once 
-			// when it is initially created, and again when it is permanently 
-			// saved)
-			bool const bNewlyCreated = (ActionDatabase.GetAllActions().Find(NewBlueprint) == nullptr);
-
-			OnBlueprintChanged(NewBlueprint);
-			// have to be careful not to register this callback twice for the 
-			// blueprint (since this function can be triggered multiple times 
-			// for the same asset)
-			if (bNewlyCreated)
-			{
-				NewBlueprint->OnChanged().AddStatic(&BlueprintActionDatabaseImpl::OnBlueprintChanged);
-			}			
+			OnBlueprintChanged(NewBlueprint);			
 		}
 		else 
 		{
@@ -664,10 +649,6 @@ static void BlueprintActionDatabaseImpl::OnAssetRemoved(FAssetData const& AssetI
 	if (AssetInfo.IsAssetLoaded())
 	{
 		UObject* AssetObject = AssetInfo.GetAsset();
-		if (UBlueprint* BlueprintAsset = Cast<UBlueprint>(AssetObject))
-		{
-			BlueprintAsset->OnChanged().RemoveStatic(&BlueprintActionDatabaseImpl::OnBlueprintChanged);
-		}
 		// the delete went through, so we don't need to track these for re-add
 		PendingDelete.Remove(AssetObject);
 	}
@@ -721,10 +702,6 @@ FBlueprintActionDatabase& FBlueprintActionDatabase::Get()
 FBlueprintActionDatabase::FBlueprintActionDatabase()
 {
 	RefreshAll();
-	for (TObjectIterator<UBlueprint> BpIt; BpIt; ++BpIt)
-	{
-		BpIt->OnChanged().AddStatic(&BlueprintActionDatabaseImpl::OnBlueprintChanged);
-	}
 	FCoreDelegates::OnAssetLoaded.AddStatic(&BlueprintActionDatabaseImpl::OnAssetLoaded);
 
 	IAssetRegistry& AssetRegistry = FModuleManager::LoadModuleChecked<FAssetRegistryModule>(TEXT("AssetRegistry")).Get();
@@ -752,7 +729,7 @@ void FBlueprintActionDatabase::Tick(float DeltaTime)
 
 	// entries that were removed from the database, in preparation for a delete
 	// (but the user ended up not deleting the object)
-	for (UObject const* AssetObj : BlueprintActionDatabaseImpl::PendingDelete)
+	for (UObject* AssetObj : BlueprintActionDatabaseImpl::PendingDelete)
 	{
 		if (IsValid(AssetObj))
 		{
@@ -761,8 +738,13 @@ void FBlueprintActionDatabase::Tick(float DeltaTime)
 	}
 	BlueprintActionDatabaseImpl::PendingDelete.Empty();
 
-	double TotalFrameDuration = 0.0;
-	while ((ActionPrimingQueue.Num() > 0) && (TotalFrameDuration < DurationThreshold))
+	
+	// priming every database entry at once would cause a hitch, so we spread it 
+	// out over several frames
+	static int32 const PrimingMaxPerFrame = 16;
+	int32 PrimedCount = 0;
+
+	while ((ActionPrimingQueue.Num() > 0) && (PrimedCount < PrimingMaxPerFrame))
 	{
 		auto ActionIndex = ActionPrimingQueue.CreateIterator();	
 			
@@ -773,12 +755,11 @@ void FBlueprintActionDatabase::Tick(float DeltaTime)
 			if (FActionList* ClassActionList = ActionRegistry.Find(ActionsKey.Get()))
 			{
 				int32& ActionListIndex = ActionIndex.Value();
-				for (; (ActionListIndex < ClassActionList->Num()) && (TotalFrameDuration < DurationThreshold); ++ActionListIndex)
+				for (; (ActionListIndex < ClassActionList->Num()) && (PrimedCount < PrimingMaxPerFrame); ++ActionListIndex)
 				{
-					FScopedDurationTimer ScopedTimer(TotalFrameDuration);
-
 					UBlueprintNodeSpawner* Action = (*ClassActionList)[ActionListIndex];
 					Action->Prime();
+					++PrimedCount;
 				}
 
 				if (ActionListIndex >= ClassActionList->Num())
@@ -833,7 +814,7 @@ void FBlueprintActionDatabase::RefreshClassActions(UClass* const Class)
 	}
 	else if (bIsBlueprintClass)
 	{
-		UBlueprint const* Blueprint = Cast<UBlueprint>(Class->ClassGeneratedBy);
+		UBlueprint* Blueprint = Cast<UBlueprint>(Class->ClassGeneratedBy);
 		if ((Blueprint != nullptr) && BlueprintActionDatabaseImpl::IsObjectValidForDatabase(Blueprint))
 		{
 			// to prevent us from hitting this twice on init (once for the skel 
@@ -916,21 +897,30 @@ void FBlueprintActionDatabase::RefreshClassActions(UClass* const Class)
 }
 
 //------------------------------------------------------------------------------
-void FBlueprintActionDatabase::RefreshAssetActions(UObject const* const AssetObject)
+void FBlueprintActionDatabase::RefreshAssetActions(UObject* const AssetObject)
 {
 	using namespace BlueprintActionDatabaseImpl;
 	check(BlueprintActionDatabaseImpl::IsObjectValidForDatabase(AssetObject));
 
+	bool const bHadExistingEntry = ActionRegistry.Contains(AssetObject);
 	FActionList& AssetActionList = ActionRegistry.FindOrAdd(AssetObject);
 	AssetActionList.Empty();
 
-	if (UBlueprint const* Blueprint = Cast<UBlueprint>(AssetObject))
+	UBlueprint* BlueprintAsset = Cast<UBlueprint>(AssetObject);
+	if (BlueprintAsset != nullptr)
 	{
-		AddBlueprintGraphActions(Blueprint, AssetActionList);
-		if (UClass* SkeletonClass = Blueprint->SkeletonGeneratedClass)
+		AddBlueprintGraphActions(BlueprintAsset, AssetActionList);
+		if (UClass* SkeletonClass = BlueprintAsset->SkeletonGeneratedClass)
 		{
 			GetClassMemberActions(SkeletonClass, AssetActionList);
 		}
+
+		// have to be careful not to register this callback twice for the 
+		// blueprint
+		if (!bHadExistingEntry)
+		{
+			BlueprintAsset->OnChanged().AddStatic(&BlueprintActionDatabaseImpl::OnBlueprintChanged);
+		}
 	}
 
 	FBlueprintActionDatabaseRegistrar Registrar(ActionRegistry, ActionPrimingQueue);
@@ -943,17 +933,26 @@ void FBlueprintActionDatabase::RefreshAssetActions(UObject const* const AssetObj
 		// queue these assets for priming
 		ActionPrimingQueue.Add(AssetObject, 0);
 	}
-	else
+	// we don't want to clear entries for blueprints, mainly because we 
+	// use the presence of an entry to know if we've set the blueprint's 
+	// OnChanged(), but also because most blueprints will have actions at some 
+	// later point
+	else if (BlueprintAsset == nullptr)
 	{
 		ClearAssetActions(AssetObject);
 	}
 }
 
 //------------------------------------------------------------------------------
-void FBlueprintActionDatabase::ClearAssetActions(UObject const* const AssetObject)
+void FBlueprintActionDatabase::ClearAssetActions(UObject* const AssetObject)
 {
 	check(BlueprintActionDatabaseImpl::IsObjectValidForDatabase(AssetObject));
 	ActionRegistry.Remove(AssetObject);
+
+	if (UBlueprint* BlueprintAsset = Cast<UBlueprint>(AssetObject))
+	{
+		BlueprintAsset->OnChanged().RemoveStatic(&BlueprintActionDatabaseImpl::OnBlueprintChanged);
+	}
 }
 
 //------------------------------------------------------------------------------
diff --git a/Engine/Source/Editor/BlueprintGraph/Public/BlueprintActionDatabase.h b/Engine/Source/Editor/BlueprintGraph/Public/BlueprintActionDatabase.h
--- a/Engine/Source/Editor/BlueprintGraph/Public/BlueprintActionDatabase.h
+++ b/Engine/Source/Editor/BlueprintGraph/Public/BlueprintActionDatabase.h
@@ -81,7 +81,7 @@ public:
 	 * 
 	 * @param  AssetObject	The asset entry you want rebuilt.
 	 */
-	void RefreshAssetActions(UObject const* const AssetObject);
+	void RefreshAssetActions(UObject* const AssetObject);
 
 	/**
 	 * Finds the database entry for the specified class and wipes it. The entry 
@@ -89,7 +89,7 @@ public:
 	 * 
 	 * @param  AssetObject	
 	 */
-	void ClearAssetActions(UObject const* const AssetObject);
+	void ClearAssetActions(UObject* const AssetObject);
 	
 	/**
 	 * 
