Fix spelling of signature

diff --git a/Engine/Source/Editor/BlueprintGraph/Classes/EdGraphSchema_K2.h b/Engine/Source/Editor/BlueprintGraph/Classes/EdGraphSchema_K2.h
--- a/Engine/Source/Editor/BlueprintGraph/Classes/EdGraphSchema_K2.h
+++ b/Engine/Source/Editor/BlueprintGraph/Classes/EdGraphSchema_K2.h
@@ -653,8 +653,8 @@ public:
 	/** Can this function be overridden by kismet (either placed as event or new function graph created) */
 	static bool CanKismetOverrideFunction(const UFunction* Function);
 
-	/** returns friendly signiture name if possible or Removes any mangling to get the unmangled signature name of the function */
-	static FString GetFriendlySignitureName(const UFunction* Function);
+	/** returns friendly signature name if possible or Removes any mangling to get the unmangled signature name of the function */
+	static FString GetFriendlySignatureName(const UFunction* Function);
 
 	static bool IsAllowableBlueprintVariableType(const class UEnum* InEnum);
 	static bool IsAllowableBlueprintVariableType(const class UClass* InClass);
diff --git a/Engine/Source/Editor/BlueprintGraph/Private/BlueprintEventNodeSpawner.cpp b/Engine/Source/Editor/BlueprintGraph/Private/BlueprintEventNodeSpawner.cpp
--- a/Engine/Source/Editor/BlueprintGraph/Private/BlueprintEventNodeSpawner.cpp
+++ b/Engine/Source/Editor/BlueprintGraph/Private/BlueprintEventNodeSpawner.cpp
@@ -2,7 +2,7 @@
 
 #include "BlueprintGraphPrivatePCH.h"
 #include "BlueprintEventNodeSpawner.h"
-#include "EdGraphSchema_K2.h" // for GetFriendlySignitureName()
+#include "EdGraphSchema_K2.h" // for GetFriendlySignatureName()
 
 #define LOCTEXT_NAMESPACE "BlueprintEventNodeSpawner"
 
@@ -161,7 +161,7 @@ FText UBlueprintEventNodeSpawner::GetDefaultMenuName(FBindingSet const& Bindings
 		FText EventName;
 		if (EventFunc != nullptr)
 		{
-			EventName = FText::FromString(UEdGraphSchema_K2::GetFriendlySignitureName(EventFunc));
+			EventName = FText::FromString(UEdGraphSchema_K2::GetFriendlySignatureName(EventFunc));
 		}
 		else if (!CustomEventName.IsNone())
 		{
diff --git a/Engine/Source/Editor/BlueprintGraph/Private/EdGraphSchema_K2.cpp b/Engine/Source/Editor/BlueprintGraph/Private/EdGraphSchema_K2.cpp
--- a/Engine/Source/Editor/BlueprintGraph/Private/EdGraphSchema_K2.cpp
+++ b/Engine/Source/Editor/BlueprintGraph/Private/EdGraphSchema_K2.cpp
@@ -509,7 +509,7 @@ bool UEdGraphSchema_K2::FunctionCanBeUsedInDelegate(const UFunction* InFunction)
 	return FNoOutputParametersHelper::Check(InFunction);
 }
 
-FString UEdGraphSchema_K2::GetFriendlySignitureName(const UFunction* Function)
+FString UEdGraphSchema_K2::GetFriendlySignatureName(const UFunction* Function)
 {
 	return UK2Node_CallFunction::GetUserFacingFunctionName( Function );
 }
diff --git a/Engine/Source/Editor/BlueprintGraph/Private/K2ActionMenuBuilder.cpp b/Engine/Source/Editor/BlueprintGraph/Private/K2ActionMenuBuilder.cpp
--- a/Engine/Source/Editor/BlueprintGraph/Private/K2ActionMenuBuilder.cpp
+++ b/Engine/Source/Editor/BlueprintGraph/Private/K2ActionMenuBuilder.cpp
@@ -564,7 +564,7 @@ void FK2ActionMenuBuilder::AddEventForFunction(FGraphActionListBuilderBase& Acti
 	if (Function != NULL)
 	{
 		check(!Function->HasAnyFunctionFlags(FUNC_Delegate));
-		const FString SigName   = UEdGraphSchema_K2::GetFriendlySignitureName(Function);
+		const FString SigName   = UEdGraphSchema_K2::GetFriendlySignatureName(Function);
 		const FString EventName = FText::Format( LOCTEXT("EventWithSignatureName", "Event {0}"), FText::FromString(SigName)).ToString();
 
 		UK2Node_Event* EventNodeTemplate = ActionMenuBuilder.CreateTemplateNode<UK2Node_Event>();
diff --git a/Engine/Source/Editor/BlueprintGraph/Private/K2Node_CustomEvent.cpp b/Engine/Source/Editor/BlueprintGraph/Private/K2Node_CustomEvent.cpp
--- a/Engine/Source/Editor/BlueprintGraph/Private/K2Node_CustomEvent.cpp
+++ b/Engine/Source/Editor/BlueprintGraph/Private/K2Node_CustomEvent.cpp
@@ -226,7 +226,7 @@ void UK2Node_CustomEvent::ValidateNodeDuringCompilation(class FCompilerResultsLo
 				MessageLog.Error(*FString::Printf(*LOCTEXT("NonCustomEventOverride", "@@ name conflicts with a '%s' function").ToString(), *FuncOwner->GetName()), this);
 			}
 			// else, we assume the user was attempting to override the parent's custom-event
-			// the signitures could still be off, but FKismetCompilerContext::PrecompileFunction() should catch that
+			// the signatures could still be off, but FKismetCompilerContext::PrecompileFunction() should catch that
 		}		
 	}
 }
diff --git a/Engine/Source/Editor/BlueprintGraph/Private/K2Node_Event.cpp b/Engine/Source/Editor/BlueprintGraph/Private/K2Node_Event.cpp
--- a/Engine/Source/Editor/BlueprintGraph/Private/K2Node_Event.cpp
+++ b/Engine/Source/Editor/BlueprintGraph/Private/K2Node_Event.cpp
@@ -57,7 +57,7 @@ FText UK2Node_Event::GetNodeTitle(ENodeTitleType::Type TitleType) const
 
 		if (UFunction* Function = FindField<UFunction>(EventSignatureClass, EventSignatureName))
 		{
-			FunctionName = UEdGraphSchema_K2::GetFriendlySignitureName(Function);
+			FunctionName = UEdGraphSchema_K2::GetFriendlySignatureName(Function);
 		}
 
 		FFormatNamedArguments Args;
