Make Event nodes strip [REDACTED_HANDLE] from tooltip
Apply [REDACTED_HANDLE] from function comment to Event node pins

diff --git a/Engine/Source/Editor/BlueprintGraph/Classes/K2Node_CallFunction.h b/Engine/Source/Editor/BlueprintGraph/Classes/K2Node_CallFunction.h
--- a/Engine/Source/Editor/BlueprintGraph/Classes/K2Node_CallFunction.h
+++ b/Engine/Source/Editor/BlueprintGraph/Classes/K2Node_CallFunction.h
@@ -135,6 +135,8 @@ public:
 	/** Gets the user-facing name for the function */
 	static FString GetUserFacingFunctionName(const UFunction* Function);
 
+	/** Set up a pins tooltip from a function's tooltip */
+	static void GeneratePinTooltipFromFunction(UEdGraphPin& Pin, const UFunction* Function);
 	/** Gets the non-specific tooltip for the function */
 	static FString GetDefaultTooltipForFunction(const UFunction* Function);
 	/** Get default category for this function in action menu */
@@ -182,9 +184,8 @@ private:
 	 * Creates hover text for the specified pin.
 	 * 
 	 * @param   Pin				The pin you want hover text for (should belong to this node)
-	 * @param   HoverTextOut	This will get filled out with the generated text
 	 */
-	void GeneratePinHoverText(const UEdGraphPin& Pin, FString& HoverTextOut) const;
+	void GeneratePinTooltip(UEdGraphPin& Pin) const;
 
 protected:
 	/** Helper function to ensure function is called in our context */
diff --git a/Engine/Source/Editor/BlueprintGraph/Private/K2Node.cpp b/Engine/Source/Editor/BlueprintGraph/Private/K2Node.cpp
--- a/Engine/Source/Editor/BlueprintGraph/Private/K2Node.cpp
+++ b/Engine/Source/Editor/BlueprintGraph/Private/K2Node.cpp
@@ -69,6 +69,8 @@ bool UK2Node::CreatePinsForFunctionEntryExit(const UFunction* Function, bool bFo
 			const bool bPinGood = K2Schema->ConvertPropertyToPinType(Param, /*out*/ Pin->PinType);
 			K2Schema->SetPinDefaultValueBasedOnType(Pin);
 			
+			UK2Node_CallFunction::GeneratePinTooltipFromFunction(*Pin, Function);
+
 			bAllPinsGood = bAllPinsGood && bPinGood;
 		}
 	}
diff --git a/Engine/Source/Editor/BlueprintGraph/Private/K2Node_CallFunction.cpp b/Engine/Source/Editor/BlueprintGraph/Private/K2Node_CallFunction.cpp
--- a/Engine/Source/Editor/BlueprintGraph/Private/K2Node_CallFunction.cpp
+++ b/Engine/Source/Editor/BlueprintGraph/Private/K2Node_CallFunction.cpp
@@ -150,7 +150,7 @@
 	if (UFunction* Function = GetTargetFunction())
 	{
 		RPCString = UK2Node_Event::GetLocalizedNetString(Function->FunctionFlags, true);
-		FunctionName = UK2Node_CallFunction::GetUserFacingFunctionName(Function);
+		FunctionName = GetUserFacingFunctionName(Function);
 		ContextString = GetFunctionContextString();
 	}
 	else
@@ -388,7 +388,7 @@
 	}
 }
 
-void UK2Node_CallFunction::GeneratePinHoverText(const UEdGraphPin& Pin, FString& HoverTextOut) const
+void UK2Node_CallFunction::GeneratePinTooltip(UEdGraphPin& Pin) const
 {
 	ensure(Pin.GetOwningNode() == this);
 
@@ -398,119 +398,20 @@
 
 	if (K2Schema == NULL)
 	{
-		Schema->ConstructBasicPinTooltip(Pin, FText::GetEmpty(), HoverTextOut);
+		Schema->ConstructBasicPinTooltip(Pin, FText::GetEmpty(), Pin.PinToolTip);
 		return;
 	}
 	
-	// figure what tag we should be parsing for (is this a return-val pin, or a parameter?)
-	FString ParamName;
-	FString TagStr = TEXT("@param");
-	if (Pin.PinName == K2Schema->PN_ReturnValue)
-	{
-		TagStr = TEXT("@return");
-	}
-	else
-	{
-		ParamName = Pin.PinName.ToLower();
-	}
-
 	// get the class function object associated with this node
 	UFunction* Function = GetTargetFunction();
 	if (Function == NULL)
 	{
-		Schema->ConstructBasicPinTooltip(Pin, FText::GetEmpty(), HoverTextOut);
+		Schema->ConstructBasicPinTooltip(Pin, FText::GetEmpty(), Pin.PinToolTip);
 		return;
 	}
+	
 
-	// grab the the function's comment block for us to parse
-	FString FunctionToolTipText = Function->GetToolTipText().ToString();
-	
-	int32 CurStrPos = INDEX_NONE;
-	int32 FullToolTipLen = FunctionToolTipText.Len();
-	// parse the full function tooltip text, looking for tag lines
-	do 
-	{
-		CurStrPos = FunctionToolTipText.Find(TagStr, ESearchCase::IgnoreCase, ESearchDir::FromStart, CurStrPos);
-		if (CurStrPos == INDEX_NONE) // if the tag wasn't found
-		{
-			break;
-		}
-
-		// advance past the tag
-		CurStrPos += TagStr.Len();
-
-		// advance past whitespace
-		while(CurStrPos < FullToolTipLen && FChar::IsWhitespace(FunctionToolTipText[CurStrPos]))
-		{
-			++CurStrPos;
-		}
-
-		// if this is a parameter pin
-		if (!ParamName.IsEmpty())
-		{
-			FString TagParamName;
-
-			// copy the parameter name
-			while (CurStrPos < FullToolTipLen && !FChar::IsWhitespace(FunctionToolTipText[CurStrPos]))
-			{
-				TagParamName.AppendChar(FunctionToolTipText[CurStrPos++]);
-			}
-
-			// if this @param tag doesn't match the param we're looking for
-			if (TagParamName != ParamName)
-			{
-				continue;
-			}
-		}
-
-		// advance past whitespace (get to the meat of the comment)
-		while(CurStrPos < FullToolTipLen && FChar::IsWhitespace(FunctionToolTipText[CurStrPos]))
-		{
-			++CurStrPos;
-		}
-
-
-		FString ParamDesc;
-		// collect the param/return-val description
-		while (CurStrPos < FullToolTipLen && FunctionToolTipText[CurStrPos] != TEXT('@'))
-		{
-			// advance past newline
-			while(CurStrPos < FullToolTipLen && FChar::IsLinebreak(FunctionToolTipText[CurStrPos]))
-			{
-				++CurStrPos;
-
-				// advance past whitespace at the start of a new line
-				while(CurStrPos < FullToolTipLen && FChar::IsWhitespace(FunctionToolTipText[CurStrPos]))
-				{
-					++CurStrPos;
-				}
-
-				// replace the newline with a single space
-				if(!FChar::IsLinebreak(FunctionToolTipText[CurStrPos]))
-				{
-					ParamDesc.AppendChar(TEXT(' '));
-				}
-			}
-
-			if (FunctionToolTipText[CurStrPos] != TEXT('@'))
-			{
-				ParamDesc.AppendChar(FunctionToolTipText[CurStrPos++]);
-			}
-		}
-
-		// trim any trailing whitespace from the descriptive text
-		ParamDesc.TrimTrailing();
-
-		// if we came up with a valid description for the param/return-val
-		if (!ParamDesc.IsEmpty())
-		{
-			HoverTextOut += ParamDesc;
-			break; // we found a match, so there's no need to continue
-		}
-
-	} while (CurStrPos < FullToolTipLen);
-
-	Schema->ConstructBasicPinTooltip(Pin, FText::FromString(HoverTextOut), HoverTextOut);
+	GeneratePinTooltipFromFunction(Pin, Function);
 }
 
 bool UK2Node_CallFunction::CreatePinsForFunctionCall(const UFunction* Function)
@@ -533,7 +434,7 @@
 	SelfPin->PinFriendlyName =  LOCTEXT("Target", "Target");
 
 	// fill out the self-pin's default tool-tip
-	GeneratePinHoverText(*SelfPin, SelfPin->PinToolTip);
+	GeneratePinTooltip(*SelfPin);
 
 	const bool bIsProtectedFunc = Function->GetBoolMetaData(FBlueprintMetadata::MD_Protected);
 	const bool bIsStaticFunc = Function->HasAllFunctionFlags(FUNC_Static);
@@ -586,7 +487,7 @@
 			K2Schema->SetPinDefaultValue(Pin, Function, Param);
 
 			// setup the default tool-tip text for this pin
-			GeneratePinHoverText(*Pin, Pin->PinToolTip);
+			GeneratePinTooltip(*Pin);
 			
 			if (PinsToHide.Contains(Pin->PinName))
 			{
@@ -827,6 +728,111 @@
 	return CachedTooltip;
 }
 
+void UK2Node_CallFunction::GeneratePinTooltipFromFunction(UEdGraphPin& Pin, const UFunction* Function)
+{
+	// figure what tag we should be parsing for (is this a return-val pin, or a parameter?)
+	FString ParamName;
+	FString TagStr = TEXT("@param");
+	if (Pin.PinName == UEdGraphSchema_K2::PN_ReturnValue)
+	{
+		TagStr = TEXT("@return");
+	}
+	else
+	{
+		ParamName = Pin.PinName.ToLower();
+	}
+
+	// grab the the function's comment block for us to parse
+	FString FunctionToolTipText = Function->GetToolTipText().ToString();
+	
+	int32 CurStrPos = INDEX_NONE;
+	int32 FullToolTipLen = FunctionToolTipText.Len();
+	// parse the full function tooltip text, looking for tag lines
+	do 
+	{
+		CurStrPos = FunctionToolTipText.Find(TagStr, ESearchCase::IgnoreCase, ESearchDir::FromStart, CurStrPos);
+		if (CurStrPos == INDEX_NONE) // if the tag wasn't found
+		{
+			break;
+		}
+
+		// advance past the tag
+		CurStrPos += TagStr.Len();
+
+		// advance past whitespace
+		while(CurStrPos < FullToolTipLen && FChar::IsWhitespace(FunctionToolTipText[CurStrPos]))
+		{
+			++CurStrPos;
+		}
+
+		// if this is a parameter pin
+		if (!ParamName.IsEmpty())
+		{
+			FString TagParamName;
+
+			// copy the parameter name
+			while (CurStrPos < FullToolTipLen && !FChar::IsWhitespace(FunctionToolTipText[CurStrPos]))
+			{
+				TagParamName.AppendChar(FunctionToolTipText[CurStrPos++]);
+			}
+
+			// if this @param tag doesn't match the param we're looking for
+			if (TagParamName != ParamName)
+			{
+				continue;
+			}
+		}
+
+		// advance past whitespace (get to the meat of the comment)
+		while(CurStrPos < FullToolTipLen && FChar::IsWhitespace(FunctionToolTipText[CurStrPos]))
+		{
+			++CurStrPos;
+		}
+
+
+		FString ParamDesc;
+		// collect the param/return-val description
+		while (CurStrPos < FullToolTipLen && FunctionToolTipText[CurStrPos] != TEXT('@'))
+		{
+			// advance past newline
+			while(CurStrPos < FullToolTipLen && FChar::IsLinebreak(FunctionToolTipText[CurStrPos]))
+			{
+				++CurStrPos;
+
+				// advance past whitespace at the start of a new line
+				while(CurStrPos < FullToolTipLen && FChar::IsWhitespace(FunctionToolTipText[CurStrPos]))
+				{
+					++CurStrPos;
+				}
+
+				// replace the newline with a single space
+				if(!FChar::IsLinebreak(FunctionToolTipText[CurStrPos]))
+				{
+					ParamDesc.AppendChar(TEXT(' '));
+				}
+			}
+
+			if (FunctionToolTipText[CurStrPos] != TEXT('@'))
+			{
+				ParamDesc.AppendChar(FunctionToolTipText[CurStrPos++]);
+			}
+		}
+
+		// trim any trailing whitespace from the descriptive text
+		ParamDesc.TrimTrailing();
+
+		// if we came up with a valid description for the param/return-val
+		if (!ParamDesc.IsEmpty())
+		{
+			Pin.PinToolTip += ParamDesc;
+			break; // we found a match, so there's no need to continue
+		}
+
+	} while (CurStrPos < FullToolTipLen);
+
+	GetDefault<UEdGraphSchema_K2>()->ConstructBasicPinTooltip(Pin, FText::FromString(Pin.PinToolTip), Pin.PinToolTip);
+}
+
 FString UK2Node_CallFunction::GetUserFacingFunctionName(const UFunction* Function)
 {
 	FString FunctionName = Function->GetMetaData(TEXT("FriendlyName"));
diff --git a/Engine/Source/Editor/BlueprintGraph/Private/K2Node_Event.cpp b/Engine/Source/Editor/BlueprintGraph/Private/K2Node_Event.cpp
--- a/Engine/Source/Editor/BlueprintGraph/Private/K2Node_Event.cpp
+++ b/Engine/Source/Editor/BlueprintGraph/Private/K2Node_Event.cpp
@@ -94,7 +94,7 @@ FText UK2Node_Event::GetTooltipText() const
 	UFunction* Function = FindField<UFunction>(EventSignatureClass, EventSignatureName);
 	if (CachedTooltip.IsOutOfDate() && (Function != nullptr))
 	{
-		CachedTooltip = Function->GetToolTipText();
+		CachedTooltip = FText::FromString(UK2Node_CallFunction::GetDefaultTooltipForFunction(Function));
 
 		if (bOverrideFunction || (CustomFunctionName == NAME_None))
 		{
