Strip '-' after [REDACTED_HANDLE] for pin tooltips

diff --git a/Engine/Source/Editor/BlueprintGraph/Private/K2Node_CallFunction.cpp b/Engine/Source/Editor/BlueprintGraph/Private/K2Node_CallFunction.cpp
--- a/Engine/Source/Editor/BlueprintGraph/Private/K2Node_CallFunction.cpp
+++ b/Engine/Source/Editor/BlueprintGraph/Private/K2Node_CallFunction.cpp
@@ -409,7 +409,7 @@ void UK2Node_CallFunction::GeneratePinTooltip(UEdGraphPin& Pin) const
 		Schema->ConstructBasicPinTooltip(Pin, FText::GetEmpty(), Pin.PinToolTip);
 		return;
 	}
-	
+
 
 	GeneratePinTooltipFromFunction(Pin, Function);
 }
@@ -784,7 +784,8 @@ void UK2Node_CallFunction::GeneratePinTooltipFromFunction(UEdGraphPin& Pin, cons
 		}
 
 		// advance past whitespace (get to the meat of the comment)
-		while(CurStrPos < FullToolTipLen && FChar::IsWhitespace(FunctionToolTipText[CurStrPos]))
+		// since many doxygen style @param use the format "@param <param name> - <comment>" we also strip - if it is before we get to any other non-whitespace
+		while(CurStrPos < FullToolTipLen && (FChar::IsWhitespace(FunctionToolTipText[CurStrPos]) || FunctionToolTipText[CurStrPos] == '-'))
 		{
 			++CurStrPos;
 		}
