"Set By-Ref Value" node will no longer appear in the new context menu system unless dragging off of an output reference pin (bringing it to parity with old system).

diff --git a/Engine/Source/Editor/BlueprintGraph/Classes/K2Node_VariableSetRef.h b/Engine/Source/Editor/BlueprintGraph/Classes/K2Node_VariableSetRef.h
--- a/Engine/Source/Editor/BlueprintGraph/Classes/K2Node_VariableSetRef.h
+++ b/Engine/Source/Editor/BlueprintGraph/Classes/K2Node_VariableSetRef.h
@@ -24,6 +24,7 @@ class UK2Node_VariableSetRef : public UK2Node
 	virtual void GetMenuActions(FBlueprintActionDatabaseRegistrar& ActionRegistrar) const override;
 	virtual FText GetMenuCategory() const override;
 	virtual int32 GetNodeRefreshPriority() const override { return EBaseNodeRefreshPriority::Low_UsesDependentWildcard; }
+	virtual bool IsActionFilteredOut(class FBlueprintActionFilter const& Filter) override;
 	// End UK2Node interface
 
 	/** 
diff --git a/Engine/Source/Editor/BlueprintGraph/Private/K2Node_VariableSetRef.cpp b/Engine/Source/Editor/BlueprintGraph/Private/K2Node_VariableSetRef.cpp
--- a/Engine/Source/Editor/BlueprintGraph/Private/K2Node_VariableSetRef.cpp
+++ b/Engine/Source/Editor/BlueprintGraph/Private/K2Node_VariableSetRef.cpp
@@ -8,6 +8,7 @@
 #include "BlueprintNodeSpawner.h"
 #include "EditorCategoryUtils.h"
 #include "BlueprintActionDatabaseRegistrar.h"
+#include "BlueprintActionFilter.h"
 
 static FString TargetVarPinName(TEXT("Target"));
 static FString VarValuePinName(TEXT("Value"));
@@ -145,6 +146,23 @@ FText UK2Node_VariableSetRef::GetNodeTitle(ENodeTitleType::Type TitleType) const
 	return CachedNodeTitle;
 }
 
+bool UK2Node_VariableSetRef::IsActionFilteredOut(class FBlueprintActionFilter const& Filter)
+{
+	// Default to filtering this node out unless dragging off of a reference output pin
+	bool bIsFilteredOut = true;
+	FBlueprintActionContext const& FilterContext = Filter.Context;
+
+	for (UEdGraphPin* Pin : FilterContext.Pins)
+	{
+		if(Pin->Direction == EGPD_Output && Pin->PinType.bIsReference == true)
+		{
+			bIsFilteredOut = false;
+			break;
+		}
+	}
+	return bIsFilteredOut;
+}
+
 void UK2Node_VariableSetRef::NotifyPinConnectionListChanged(UEdGraphPin* Pin)
 {
 	Super::NotifyPinConnectionListChanged(Pin);
