Blueprints: Changed logic for hiding the target pin on a pure node to take into account the compatibility of the BP and the function

diff --git a/Engine/Source/Editor/BlueprintGraph/Private/K2Node_CallFunction.cpp b/Engine/Source/Editor/BlueprintGraph/Private/K2Node_CallFunction.cpp
--- a/Engine/Source/Editor/BlueprintGraph/Private/K2Node_CallFunction.cpp
+++ b/Engine/Source/Editor/BlueprintGraph/Private/K2Node_CallFunction.cpp
@@ -439,17 +439,27 @@ bool UK2Node_CallFunction::CreatePinsForFunctionCall(const UFunction* Function)
 	const bool bIsProtectedFunc = Function->GetBoolMetaData(FBlueprintMetadata::MD_Protected);
 	const bool bIsStaticFunc = Function->HasAllFunctionFlags(FUNC_Static);
 
-	//@TODO: Can't strictly speaking always hide the self pin for pure and static functions; it'll still be needed if the function belongs to a class not in the blueprint's class hierarchy
-	SelfPin->bHidden = ((bIsPureFunc && !bIsConstFunc) || bIsProtectedFunc || bIsStaticFunc);
-
 	UBlueprint* BP = GetBlueprint();
 	ensure(BP);
-	if (bIsStaticFunc && BP)
+	if (BP != nullptr)
 	{
-		// Wire up the self to the CDO of the class if it's not us
-		if (!BP->SkeletonGeneratedClass->IsChildOf(FunctionOwnerClass))
+		const bool bIsFunctionCompatibleWithSelf = BP->SkeletonGeneratedClass->IsChildOf(FunctionOwnerClass);
+
+		if (bIsStaticFunc)
+		{
+			// For static methods, wire up the self to the CDO of the class if it's not us
+			if (!bIsFunctionCompatibleWithSelf)
+			{
+				SelfPin->DefaultObject = FunctionOwnerClass->GetDefaultObject();
+			}
+
+			// Purity doesn't matter with a static function, we can always hide the self pin since we know how to call the method
+			SelfPin->bHidden = true;
+		}
+		else
 		{
-			SelfPin->DefaultObject = FunctionOwnerClass->GetDefaultObject();
+			// Hide the self pin if the function is compatible with the blueprint class and pure (the !bIsConstFunc portion should be going away soon too hopefully)
+			SelfPin->bHidden = bIsFunctionCompatibleWithSelf && (bIsPureFunc && !bIsConstFunc);
 		}
 	}
 
