CallOnMember function refs now forced to be non-self calls (because they were referencing the skeleton class in the new menu system).

diff --git a/Engine/Source/Editor/BlueprintGraph/Private/BlueprintFunctionNodeSpawner.cpp b/Engine/Source/Editor/BlueprintGraph/Private/BlueprintFunctionNodeSpawner.cpp
--- a/Engine/Source/Editor/BlueprintGraph/Private/BlueprintFunctionNodeSpawner.cpp
+++ b/Engine/Source/Editor/BlueprintGraph/Private/BlueprintFunctionNodeSpawner.cpp
@@ -58,7 +58,10 @@ static bool BlueprintFunctionNodeSpawnerImpl::BindFunctionNode(UK2Node_CallFunct
 		{
 			if (UK2Node_CallFunctionOnMember* CallOnMemberNode = Cast<UK2Node_CallFunctionOnMember>(NewNode))
 			{
-				CallOnMemberNode->MemberVariableToCallOn.SetFromField<UProperty>(BoundProperty, NewNode);
+				// force bIsConsideredSelfContext to false, else the target 
+				// could end up being the skeleton class (and functionally, 
+				// there is no difference)
+				CallOnMemberNode->MemberVariableToCallOn.SetFromField<UProperty>(BoundProperty, /*bIsConsideredSelfContext =*/false);
 				bSuccessfulBinding = true;
 				CallOnMemberNode->ReconstructNode();
 			}
diff --git a/Engine/Source/Editor/BlueprintGraph/Private/BlueprintNodeSpawner.cpp b/Engine/Source/Editor/BlueprintGraph/Private/BlueprintNodeSpawner.cpp
--- a/Engine/Source/Editor/BlueprintGraph/Private/BlueprintNodeSpawner.cpp
+++ b/Engine/Source/Editor/BlueprintGraph/Private/BlueprintNodeSpawner.cpp
@@ -187,12 +187,6 @@ UEdGraphNode* UBlueprintNodeSpawner::GetTemplateNode(ENoInit) const
 {
 	return BlueprintNodeSpawnerImpl::GetSharedTemplateCache()->GetNodeTemplate(this, NoInit);
 }
-// 
-// //------------------------------------------------------------------------------
-// UEdGraphNode* UBlueprintNodeSpawner::Invoke(UEdGraph* ParentGraph, FBindingSet const& Bindings, FVector2D const Location, FCustomizeNodeDelegate PostSpawnDelegate) const
-// {
-// 	return SpawnNode(NodeClass, ParentGraph, Bindings, Location, PostSpawnDelegate);
-// }
 
 //------------------------------------------------------------------------------
 UEdGraphNode* UBlueprintNodeSpawner::SpawnEdGraphNode(TSubclassOf<UEdGraphNode> InNodeClass, UEdGraph* ParentGraph, FBindingSet const& Bindings, FVector2D const Location, FCustomizeNodeDelegate PostSpawnDelegate) const
