Making sure to match CallOnMember nodes with the right target class (not the function's owner, but the owner of the property... for the new blueprint menu system).

diff --git a/Engine/Source/Editor/BlueprintGraph/Private/BlueprintActionFilter.cpp b/Engine/Source/Editor/BlueprintGraph/Private/BlueprintActionFilter.cpp
--- a/Engine/Source/Editor/BlueprintGraph/Private/BlueprintActionFilter.cpp
+++ b/Engine/Source/Editor/BlueprintGraph/Private/BlueprintActionFilter.cpp
@@ -28,12 +28,26 @@
 namespace BlueprintActionFilterImpl
 {
 	/**
+	 * Blueprints have several classes associated with them (the skeleton, for
+	 * UI reflection, and the full generated class). This retrieves the 
+	 * "authoritative" one for comparisons and queries.
 	 * 
+	 * @param  Blueprint	The blueprint you want a class for.
+	 * @return The authoritative class for the specified blueprint.
+	 */
+	static UClass* GetAuthoritativeBlueprintClass(UBlueprint const* const Blueprint);
+
+	/**
+	 * Since objects (like a blueprint) can have multiple classes (skeleton, 
+	 * generated, etc.), this utility method provides a way to compare two 
+	 * classes and ensure both are using the same class type.
 	 * 
-	 * @param  Blueprint	
-	 * @return 
+	 * @param  ClassToTest		The class you're querying for.
+	 * @param  TypeToCheckFor	The type you want to test for (this is asking "is ClassToTest a TypeToCheckFor class?")
+	 * @param  bNeedsExactMatch	Optional parameter. If false, will check IsChildOf() and ImplementsInterface(), otherwise it just checks for equivalence.
+	 * @return True if the ClassToTest is the same type as (or inherited from) TypeToCheckFor.
 	 */
-	static UClass* GetBlueprintClass(UBlueprint const* const Blueprint);	
+	static bool IsClassOfType(UClass const* ClassToTest, UClass const* TypeToCheckFor, bool const bNeedsExactMatch = false);
 	
 	/**
 	 * Utility method to check and see if the specified node-spawner would 
@@ -144,7 +158,7 @@ namespace BlueprintActionFilterImpl
 	 * @param  BlueprintAction	The action you wish to query.
 	 * @return 
 	 */
-	static bool IsPersistentNonTargetField(FBlueprintActionFilter const& Filter, FBlueprintActionInfo& BlueprintAction);
+	static bool IsRejectedGlobalField(FBlueprintActionFilter const& Filter, FBlueprintActionInfo& BlueprintAction);
 
 	/**
 	 * Rejection test that checks to see if the node-spawner is associated with 
@@ -222,11 +236,10 @@ namespace BlueprintActionFilterImpl
 	 * 
 	 * 
 	 * @param  Pin	
-	 * @param  MemberField	
 	 * @param  MemberNodeSpawner
 	 * @return 
 	 */
-	static bool IsPinCompatibleWithTargetSelf(UEdGraphPin const* Pin, UField const* MemberField, FBlueprintActionInfo& MemberNodeSpawner);
+	static bool IsPinCompatibleWithTargetSelf(UEdGraphPin const* Pin, FBlueprintActionInfo& MemberNodeSpawner);
 
 	/**
 	 * 
@@ -274,13 +287,47 @@ namespace BlueprintActionFilterImpl
 	 * @return
 	 */
 	static bool IsNodeTemplateSelfFiltered(FBlueprintActionFilter const& Filter, FBlueprintActionInfo& BlueprintAction);
+
 };
 
+//------------------------------------------------------------------------------
+UClass* BlueprintActionFilterImpl::GetAuthoritativeBlueprintClass(UBlueprint const* const Blueprint)
+{
+	UClass* BpClass = (Blueprint->SkeletonGeneratedClass != nullptr) ? Blueprint->SkeletonGeneratedClass :
+		Blueprint->GeneratedClass;
+
+	if (BpClass == nullptr)
+	{
+		BpClass = Blueprint->ParentClass;
+	}
+
+	UClass* AuthoritativeClass = BpClass;
+	if (BpClass != nullptr)
+	{
+		AuthoritativeClass = BpClass->GetAuthoritativeClass();
+	}
+	return AuthoritativeClass;
+}
 
 //------------------------------------------------------------------------------
-UClass* BlueprintActionFilterImpl::GetBlueprintClass(UBlueprint const* const Blueprint)
+static bool BlueprintActionFilterImpl::IsClassOfType(UClass const* ClassToTest, UClass const* TypeToCheckFor, bool const bNeedsExactMatch)
 {
-	return (Blueprint->SkeletonGeneratedClass != nullptr) ? Blueprint->SkeletonGeneratedClass : Blueprint->ParentClass;
+	// @TODO: get rid of these nasty const_casts ... ew! (spent a while 
+	//        following the const rabbit down the hole, and eventually gave up)
+	//        should make a const version of GetAuthoritativeClass(), somehow
+	//        without duplicating functionality
+	UClass const* AuthoritativeTestClass = const_cast<UClass*>(ClassToTest)->GetAuthoritativeClass();
+	UClass const* AuthoritativeToTestFor = const_cast<UClass*>(TypeToCheckFor)->GetAuthoritativeClass();
+
+	if (bNeedsExactMatch)
+	{
+		return (AuthoritativeTestClass == AuthoritativeToTestFor);
+	}
+	else
+	{
+		return AuthoritativeTestClass->IsChildOf(AuthoritativeToTestFor) ||
+			AuthoritativeTestClass->ImplementsInterface(AuthoritativeToTestFor);
+	}
 }
 
 //------------------------------------------------------------------------------
@@ -364,13 +411,13 @@ static bool BlueprintActionFilterImpl::IsEventUnimplementable(FBlueprintActionFi
 			UClass* FuncOwner = EventFunc->GetOuterUClass();
 			for (UBlueprint const* Blueprint : FilterContext.Blueprints)
 			{
-				UClass* BpClass = Blueprint->GeneratedClass;
+				UClass const* BpClass = GetAuthoritativeBlueprintClass(Blueprint);
 				check(BpClass != nullptr);
 
 				// if this function belongs directly to this blueprint, then it is
 				// already implemented here (this action however is valid for sub-
 				// classes, as they can override the event's functionality)
-				if (FuncOwner == BpClass)
+				if (IsClassOfType(BpClass, FuncOwner, /*bNeedsExactMatch =*/true))
 				{
 					bIsFilteredOut = true;
 					break;
@@ -379,7 +426,7 @@ static bool BlueprintActionFilterImpl::IsEventUnimplementable(FBlueprintActionFi
 				// you can only implement events that you inherit; so if this
 				// blueprint is not a subclass of the event's owner, then we're not
 				// allowed to implement it
-				if (!BpClass->IsChildOf(FuncOwner))
+				if (!IsClassOfType(BpClass, FuncOwner))
 				{
 					bIsFilteredOut = true;
 					break;
@@ -397,19 +444,25 @@ static bool BlueprintActionFilterImpl::IsFieldInaccessible(FBlueprintActionFilte
 	bool bIsFilteredOut = false;
 	FBlueprintActionContext const& FilterContext = Filter.Context;
 	
-	if (UField const* ClassField = BlueprintAction.GetAssociatedMemberField())
+	UField const* Field = BlueprintAction.GetAssociatedMemberField();
+	bool const bIsMemberAction = (Field != nullptr) && (Field->GetOwnerClass() != nullptr);
+
+	if (bIsMemberAction)
 	{
-		bool bIsProtected = ClassField->HasMetaData(FBlueprintMetadata::MD_Protected);
-		bool bIsPrivate   = ClassField->HasMetaData(FBlueprintMetadata::MD_Private);
+		bool bIsProtected = Field->HasMetaData(FBlueprintMetadata::MD_Protected);
+		bool bIsPrivate = Field->HasMetaData(FBlueprintMetadata::MD_Private);
 
 		bool bIsPublic = !bIsPrivate && !bIsProtected;
-		if (UProperty const* Property = Cast<UProperty>(ClassField))
+		if (UProperty const* Property = Cast<UProperty>(Field))
 		{
-			// CPF_DisableEditOnInstance corresponds to the eyeball/editable checkbox available 
-			// in the blueprint editor. It is poorly named. When CPF_DisableEditOnInstance is 
-			// *not* set, the user has tagged the variable as 'editable' and it is treated
-			// as public:
-			bIsPublic = !Property->HasAnyPropertyFlags(CPF_DisableEditOnInstance);
+			// CPF_DisableEditOnInstance corresponds to the eyeball/editable 
+			// checkbox available in the blueprint editor. It is poorly named. 
+			// When CPF_DisableEditOnInstance is *not* set, the user has tagged 
+			// the variable as 'editable' and it is treated as public (delegate 
+			// properties don't have the visibility check box, so they default
+			// to public):
+			bIsPublic = !Property->HasAnyPropertyFlags(CPF_DisableEditOnInstance) || 
+				(Cast<UMulticastDelegateProperty>(Property) != nullptr);
 
 			// If the variable is not public, and is not tagged as private, it should be protected
 			// by default. This branch handles variables declared in blueprints:
@@ -424,19 +477,19 @@ static bool BlueprintActionFilterImpl::IsFieldInaccessible(FBlueprintActionFilte
 
 		if( !bIsPublic )
 		{
-			UClass* FieldOwner = ClassField->GetOwnerClass();
+			UClass const* ActionOwner = BlueprintAction.GetOwnerClass();
 			for (UBlueprint const* Blueprint : FilterContext.Blueprints)
 			{
-				UClass* BpClass = Blueprint->SkeletonGeneratedClass;
+				UClass const* BpClass = GetAuthoritativeBlueprintClass(Blueprint);
 				check(BpClass != nullptr);
 			
 				// private functions are only accessible from the class they belong to
-				if (bIsPrivate && (FieldOwner != BpClass))
+				if (bIsPrivate && !IsClassOfType(BpClass, ActionOwner, /*bNeedsExactMatch =*/true))
 				{
 					bIsFilteredOut = true;
 					break;
 				}
-				else if (bIsProtected && !BpClass->IsChildOf(FieldOwner))
+				else if (bIsProtected && !IsClassOfType(BpClass, ActionOwner))
 				{
 					bIsFilteredOut = true;
 					break;
@@ -463,7 +516,7 @@ static bool BlueprintActionFilterImpl::IsRestrictedClassMember(FBlueprintActionF
 			{
 				bool bIsClassListed = false;
 				
-				UClass* BpClass = GetBlueprintClass(Blueprint);
+				UClass const* BpClass = GetAuthoritativeBlueprintClass(Blueprint);
 				// walk the class inheritance chain to see if this class is one
 				// of the allowed
 				while (!bIsClassListed && (BpClass != nullptr))
@@ -528,19 +581,19 @@ static bool BlueprintActionFilterImpl::IsDeprecated(FBlueprintActionFilter const
 }
 
 //------------------------------------------------------------------------------
-static bool BlueprintActionFilterImpl::IsPersistentNonTargetField(FBlueprintActionFilter const& Filter, FBlueprintActionInfo& BlueprintAction)
+static bool BlueprintActionFilterImpl::IsRejectedGlobalField(FBlueprintActionFilter const& Filter, FBlueprintActionInfo& BlueprintAction)
 {
 	bool bIsFilteredOut = false;
-	if (UField const* ClassField = BlueprintAction.GetAssociatedMemberField())
+	if (UField const* Field = BlueprintAction.GetAssociatedMemberField())
 	{
-		bIsFilteredOut = IsGloballyAccessible(ClassField);
+		bIsFilteredOut = IsGloballyAccessible(Field);
 		
-		UClass* FieldClass = ClassField->GetOwnerClass();
+		UClass* FieldClass = Field->GetOwnerClass();
 		if (bIsFilteredOut && (FieldClass != nullptr))
 		{
 			for (UClass const* Class : Filter.TargetClasses)
 			{
-				bool const bIsInternalMemberField = Class->IsChildOf(FieldClass);
+				bool const bIsInternalMemberField = IsClassOfType(Class, FieldClass);
 				if (bIsInternalMemberField)
 				{
 					bIsFilteredOut = false;
@@ -557,20 +610,24 @@ static bool BlueprintActionFilterImpl::IsPersistentNonTargetField(FBlueprintActi
 static bool BlueprintActionFilterImpl::IsNonTargetMember(FBlueprintActionFilter const& Filter, FBlueprintActionInfo& BlueprintAction, bool bPermitNonTargetGlobals)
 {
 	bool bIsFilteredOut = false;
-	if (UField const* ClassField = BlueprintAction.GetAssociatedMemberField())
+
+	UField const* ClassField = BlueprintAction.GetAssociatedMemberField();
+	bool const bIsMemberAction = (ClassField != nullptr) && (ClassField->GetOwnerClass() != nullptr);
+
+	if (bIsMemberAction)
 	{
-		UClass* FieldClass = Cast<UClass>(ClassField->GetOuter());
-		bool const bIsInterfaceMethod = FieldClass->IsChildOf<UInterface>();
+		UClass const* ActionClass = BlueprintAction.GetOwnerClass();
+		checkSlow(ActionClass != nullptr);
 
 		// global (and static library) fields can stay (unless explicitly
 		// excluded... save that for a separate test)
 		bool const bSkip = bPermitNonTargetGlobals && IsGloballyAccessible(ClassField);
-		if ((FieldClass != nullptr) && !bSkip)
+		if (!bSkip)
 		{
 			for (UClass const* Class : Filter.TargetClasses)
 			{
-				bool const bIsInternalFunc = Class->IsChildOf(FieldClass) || (bIsInterfaceMethod && Class->ImplementsInterface(FieldClass));
-				if (!bIsInternalFunc)
+				bool const bIsTargetOwnedField = IsClassOfType(Class, ActionClass);
+				if (!bIsTargetOwnedField)
 				{
 					bIsFilteredOut = true;
 					break;
@@ -588,43 +645,41 @@ static bool BlueprintActionFilterImpl::IsFieldCategoryHidden(FBlueprintActionFil
 	bool bIsFilteredOut = false;
 	FBlueprintActionContext const& FilterContext = Filter.Context;
 
-	if (UField const* AssociatedField = BlueprintAction.GetAssociatedMemberField())
-	{
-		DECLARE_DELEGATE_RetVal_OneParam(bool, FIsFieldHiddenDelegate, UClass*);
-		FIsFieldHiddenDelegate IsFieldHiddenDelegate;
+	DECLARE_DELEGATE_RetVal_OneParam(bool, FIsFieldHiddenDelegate, UClass*);
+	FIsFieldHiddenDelegate IsFieldHiddenDelegate;
 
-		if (UFunction const* Function = BlueprintAction.GetAssociatedFunction())
+	if (UFunction const* Function = BlueprintAction.GetAssociatedFunction())
+	{
+		auto IsFunctionHiddenLambda = [](UClass* Class, UFunction const* Function)->bool
 		{
-			auto IsFunctionHiddenLambda = [](UClass* Class, UFunction const* Function)->bool
-			{
-				return FObjectEditorUtils::IsFunctionHiddenFromClass(Function, Class);
-			};
-			IsFieldHiddenDelegate = FIsFieldHiddenDelegate::CreateStatic(IsFunctionHiddenLambda, Function);
-		}
-		else if (UProperty const* Property = BlueprintAction.GetAssociatedProperty())
+			return FObjectEditorUtils::IsFunctionHiddenFromClass(Function, Class->GetAuthoritativeClass());
+		};
+		IsFieldHiddenDelegate = FIsFieldHiddenDelegate::CreateStatic(IsFunctionHiddenLambda, Function);
+	}
+	else if (UProperty const* Property = BlueprintAction.GetAssociatedProperty())
+	{
+		auto IsPropertyHiddenLambda = [](UClass* Class, UProperty const* Property)->bool
 		{
-			auto IsPropertyHiddenLambda = [](UClass* Class, UProperty const* Property)->bool
-			{
-				return FObjectEditorUtils::IsVariableCategoryHiddenFromClass(Property, Class);
-			};
-			IsFieldHiddenDelegate = FIsFieldHiddenDelegate::CreateStatic(IsPropertyHiddenLambda, Property);
-		}
+			return FObjectEditorUtils::IsVariableCategoryHiddenFromClass(Property, Class->GetAuthoritativeClass());
+		};
+		IsFieldHiddenDelegate = FIsFieldHiddenDelegate::CreateStatic(IsPropertyHiddenLambda, Property);
+	}
 
-		if (IsFieldHiddenDelegate.IsBound())
+	if (IsFieldHiddenDelegate.IsBound())
+	{
+		for (UBlueprint* Blueprint : FilterContext.Blueprints)
 		{
-			for (UBlueprint* Blueprint : FilterContext.Blueprints)
+			UClass* BpClass = GetAuthoritativeBlueprintClass(Blueprint);
+			if (IsFieldHiddenDelegate.Execute(BpClass))
 			{
-				if (IsFieldHiddenDelegate.Execute(Blueprint->GeneratedClass))
-				{
-					bIsFilteredOut = true;
-					break;
-				}
+				bIsFilteredOut = true;
+				break;
 			}
+		}
 
-			for (int32 ClassIndex = 0; !bIsFilteredOut && (ClassIndex < Filter.TargetClasses.Num()); ++ClassIndex)
-			{
-				bIsFilteredOut = IsFieldHiddenDelegate.Execute(Filter.TargetClasses[ClassIndex]);
-			}
+		for (int32 ClassIndex = 0; !bIsFilteredOut && (ClassIndex < Filter.TargetClasses.Num()); ++ClassIndex)
+		{
+			bIsFilteredOut = IsFieldHiddenDelegate.Execute(Filter.TargetClasses[ClassIndex]);
 		}
 	}
 
@@ -805,7 +860,7 @@ static bool BlueprintActionFilterImpl::HasMatchingPin(FBlueprintActionInfo& Blue
 		check(Blueprint != nullptr);
 		UEdGraphSchema_K2 const* Schema = CastChecked<UEdGraphSchema_K2>(OuterGraph->GetSchema());
 
-		UClass* CallingContext = GetBlueprintClass(Blueprint);
+		UClass const* CallingContext = GetAuthoritativeBlueprintClass(Blueprint);
 		UK2Node* K2TemplateNode = Cast<UK2Node>(TemplateNode);
 		
 		for (int32 PinIndex = 0; !bHasCompatiblePin && (PinIndex < TemplateNode->Pins.Num()); ++PinIndex)
@@ -830,13 +885,12 @@ static bool BlueprintActionFilterImpl::HasMatchingPin(FBlueprintActionInfo& Blue
 }
 
 //------------------------------------------------------------------------------
-static bool BlueprintActionFilterImpl::IsPinCompatibleWithTargetSelf(UEdGraphPin const* Pin, UField const* MemberField, FBlueprintActionInfo& BlueprintAction)
+static bool BlueprintActionFilterImpl::IsPinCompatibleWithTargetSelf(UEdGraphPin const* Pin, FBlueprintActionInfo& BlueprintAction)
 {
 	bool bIsCompatible = false;
-	checkSlow(BlueprintAction.GetAssociatedMemberField() == MemberField);
+	UClass const* TargetClass = BlueprintAction.GetOwnerClass();
 
-	UClass* OwnerClass = Cast<UClass>(MemberField->GetOuter());
-	if ((Pin->Direction == EGPD_Output) && (OwnerClass != nullptr))
+	if ((Pin->Direction == EGPD_Output) && (TargetClass != nullptr))
 	{
 		FEdGraphPinType const& PinType = Pin->PinType;
 		UEdGraphSchema const* const PinSchema = Pin->GetSchema();
@@ -846,22 +900,22 @@ static bool BlueprintActionFilterImpl::IsPinCompatibleWithTargetSelf(UEdGraphPin
 		if (PinSchema->IsSelfPin(*Pin))
 		{
 			UBlueprint const* Blueprint = FBlueprintEditorUtils::FindBlueprintForNodeChecked(Pin->GetOwningNode());
-			PinObjClass = GetBlueprintClass(Blueprint);
+			PinObjClass = GetAuthoritativeBlueprintClass(Blueprint);
 		}
-		else if ( (PinType.PinCategory == UEdGraphSchema_K2::PC_Object || PinType.PinCategory == UEdGraphSchema_K2::PC_Interface)
-					&& PinType.PinSubCategoryObject.IsValid())
+		else if (PinType.PinSubCategoryObject.IsValid() && (PinType.PinCategory == UEdGraphSchema_K2::PC_Object || 
+			PinType.PinCategory == UEdGraphSchema_K2::PC_Interface))
 		{
 			PinObjClass = Cast<UClass>(PinType.PinSubCategoryObject.Get());
 		}
 		
 		if (PinObjClass != nullptr)
 		{
-			if (PinObjClass->IsChildOf(OwnerClass) || PinObjClass->ImplementsInterface(OwnerClass))
+			if (IsClassOfType(PinObjClass, TargetClass))
 			{
 				bIsCompatible = true;
 				if (PinType.bIsArray)
 				{
-					if (UFunction const* Function = Cast<UFunction>(MemberField))
+					if (UFunction const* Function = BlueprintAction.GetAssociatedFunction())
 					{
 						bIsCompatible = UK2Node_CallFunction::CanFunctionSupportMultipleTargets(Function);
 					}
@@ -875,6 +929,27 @@ static bool BlueprintActionFilterImpl::IsPinCompatibleWithTargetSelf(UEdGraphPin
 					}
 				}
 			}
+			else if (!PinType.bIsArray && (BlueprintAction.GetNodeClass() == UK2Node_CallFunction::StaticClass()))
+			{
+				// if this is a bound CallFunction action, then we make the 
+				// assumption that it will be turned into a UK2Node_CallFunctionOnMember
+				// node, therefore transforming the target from the function's
+				// outer to the owner of the bound property
+				IBlueprintNodeBinder::FBindingSet const& Bindings = BlueprintAction.GetBindings();
+				// UK2Node_CallFunction is only transformed into a UK2Node_CallFunctionOnMember
+				// when there is only one binding being applied ...
+				if (Bindings.Num() == 1) 
+				{
+					const UObject* Binding = Bindings.CreateConstIterator()->Get();
+					const UProperty* MemberBinding = Cast<UProperty>(Binding);
+
+					if ((MemberBinding != nullptr) && BlueprintAction.NodeSpawner->IsBindingCompatible(MemberBinding))
+					{
+						const UClass* BindingOwner = MemberBinding->GetOwnerClass();
+						bIsCompatible = IsClassOfType(PinObjClass, BindingOwner);
+					}
+				}
+			}
 		}
 	}
 
@@ -916,7 +991,7 @@ static bool BlueprintActionFilterImpl::IsFunctionMissingPinParam(FBlueprintActio
 				else
 				{
 					// need to take "Target" self pins into consideration for objects
-					bIsFilteredOut = bIsEventSpawner || !IsPinCompatibleWithTargetSelf(ContextPin, AssociatedFunc, BlueprintAction);
+					bIsFilteredOut = bIsEventSpawner || !IsPinCompatibleWithTargetSelf(ContextPin, BlueprintAction);
 				}
 			}
 			
@@ -947,7 +1022,7 @@ static bool BlueprintActionFilterImpl::IsMissmatchedPropertyType(FBlueprintActio
 				UEdGraphSchema_K2 const* K2Schema = CastChecked<UEdGraphSchema_K2 const>(ContextPin->GetSchema());
 
 				// have to account for "self" context pin
-				if (IsPinCompatibleWithTargetSelf(ContextPin, Property, BlueprintAction))
+				if (IsPinCompatibleWithTargetSelf(ContextPin, BlueprintAction))
 				{
 					continue;
 				}
@@ -996,9 +1071,11 @@ static bool BlueprintActionFilterImpl::IsMissingMatchingPinParam(FBlueprintActio
 {
 	bool bIsFilteredOut = false;
 
-	UField const* AssociatedField = BlueprintAction.GetAssociatedMemberField();
-	// we have a separate tests for field nodes (IsFunctionMissingPinParam/IsMissmatchedPropertyType)
-	if (AssociatedField == nullptr)
+	// we have a separate pin tests for function/property nodes (IsFunctionMissingPinParam/IsMissmatchedPropertyType)
+	bool const bTestPinCompatibility = (BlueprintAction.GetAssociatedProperty() == nullptr) &&
+		(BlueprintAction.GetAssociatedFunction() == nullptr);
+
+	if (bTestPinCompatibility)
 	{
 		for (UEdGraphPin const* ContextPin : Filter.Context.Pins)
 		{
@@ -1018,7 +1095,7 @@ static bool BlueprintActionFilterImpl::IsNotSubClassCast(FBlueprintActionFilter
 {
 	bool bIsFilteredOut = false;
 
-	if (BlueprintAction.GetNodeClass()->IsChildOf( UK2Node_DynamicCast::StaticClass() ))
+	if (BlueprintAction.GetNodeClass()->IsChildOf<UK2Node_DynamicCast>())
 	{
 		for(UEdGraphPin const* ContextPin : Filter.Context.Pins)
 		{
@@ -1034,7 +1111,7 @@ static bool BlueprintActionFilterImpl::IsNotSubClassCast(FBlueprintActionFilter
 				checkSlow(CastClass != nullptr);
 
 				UClass const* ContextPinClass = Cast<UClass>(ContextPin->PinType.PinSubCategoryObject.Get());
-				if ((ContextPinClass == CastClass) || !CastClass->IsChildOf(ContextPinClass))
+				if ((ContextPinClass == CastClass) || !IsClassOfType(CastClass, ContextPinClass))
 				{
 					bIsFilteredOut = true;
 					break;
@@ -1045,6 +1122,7 @@ static bool BlueprintActionFilterImpl::IsNotSubClassCast(FBlueprintActionFilter
 	return bIsFilteredOut;
 }
 
+//------------------------------------------------------------------------------
 static bool BlueprintActionFilterImpl::IsNodeTemplateSelfFiltered(FBlueprintActionFilter const& Filter, FBlueprintActionInfo& BlueprintAction)
 {
 	bool bIsFilteredOut = false;
@@ -1218,7 +1296,7 @@ FBlueprintActionFilter::FBlueprintActionFilter(uint32 Flags/*= 0x00*/)
 	AddRejectionTest(FRejectionTestDelegate::CreateStatic(IsFieldCategoryHidden));
 	if (Flags & BPFILTER_RejectGlobalFields)
 	{
-		AddRejectionTest(FRejectionTestDelegate::CreateStatic(IsPersistentNonTargetField));
+		AddRejectionTest(FRejectionTestDelegate::CreateStatic(IsRejectedGlobalField));
 	}
 	
 	AddRejectionTest(FRejectionTestDelegate::CreateStatic(IsFieldInaccessible));
