Fixes BlueprintNodeTemplateCacheImpl::IsTemplateOuter to check if the outer is Transactional instead of Transient, resolves issues with nodes being marked as non-template nodes during favorites menu generation.

diff --git a/Engine/Source/Editor/BlueprintGraph/Private/BlueprintNodeTemplateCache.cpp b/Engine/Source/Editor/BlueprintGraph/Private/BlueprintNodeTemplateCache.cpp
--- a/Engine/Source/Editor/BlueprintGraph/Private/BlueprintNodeTemplateCache.cpp
+++ b/Engine/Source/Editor/BlueprintGraph/Private/BlueprintNodeTemplateCache.cpp
@@ -194,7 +194,7 @@ static void BlueprintNodeTemplateCacheImpl::MarkGraphForTemplateUse(UEdGraph* Ne
 //------------------------------------------------------------------------------
 bool BlueprintNodeTemplateCacheImpl::IsTemplateOuter(UEdGraph* ParentGraph)
 {
-	if (ParentGraph->HasAnyFlags(RF_Transient))
+	if (ParentGraph->HasAnyFlags(RF_Transactional))
 	{
 		UPackage* GraphPackage = ParentGraph->GetOutermost();
 		UMetaData* PackageMetadata = GraphPackage->GetMetaData();
