Fixing cook error, gaurding against null property ref (when expanding CallOnMember nodes).

diff --git a/Engine/Source/Editor/BlueprintGraph/Private/K2Node_CallFunctionOnMember.cpp b/Engine/Source/Editor/BlueprintGraph/Private/K2Node_CallFunctionOnMember.cpp
--- a/Engine/Source/Editor/BlueprintGraph/Private/K2Node_CallFunctionOnMember.cpp
+++ b/Engine/Source/Editor/BlueprintGraph/Private/K2Node_CallFunctionOnMember.cpp
@@ -71,8 +71,10 @@ void UK2Node_CallFunctionOnMember::ExpandNode(class FKismetCompilerContext& Comp
 				GetVarNode->VariableReference.SetSelfMember(MemberVariableToCallOn.GetMemberName());
 				GetVarNode->AllocateDefaultPins();
 
-				UEdGraphPin* ValuePin = GetVarNode->GetValuePin();
-				ValuePin->MakeLinkTo(CallFuncSelfPin);
+				if (UEdGraphPin* ValuePin = GetVarNode->GetValuePin())
+				{
+					ValuePin->MakeLinkTo(CallFuncSelfPin);
+				}
 			}
 			else
 			{
