New blueprint context menu now conditionally displays anim notifications (events raised from animations)

diff --git a/Engine/Source/Editor/BlueprintGraph/Private/BlueprintActionDatabase.cpp b/Engine/Source/Editor/BlueprintGraph/Private/BlueprintActionDatabase.cpp
--- a/Engine/Source/Editor/BlueprintGraph/Private/BlueprintActionDatabase.cpp
+++ b/Engine/Source/Editor/BlueprintGraph/Private/BlueprintActionDatabase.cpp
@@ -270,6 +270,14 @@ namespace BlueprintActionDatabaseImpl
 	 */
 	static void AddComponentClassActions(TSubclassOf<UActorComponent> const ComponentClass, FActionList& ActionListOut);
 
+	/**
+	 * Adds custom actions to operate on the provided skeleton. Used primarily
+	 * to find AnimNotify event vocabulary
+	 *
+	 * @Param Skeleton	The skeleton that may have anim notifies defined on it.
+	 */
+	static void AddSkeletonActions( const USkeleton& Skeleton, FActionList& ActionListOut);
+
 	/**
 	 * If the associated class is a blueprint generated class, then this will
 	 * loop over the blueprint's graphs and create any node-spawners associated
@@ -524,6 +532,19 @@ static void BlueprintActionDatabaseImpl::AddComponentClassActions(TSubclassOf<UA
 	}
 }
 
+//------------------------------------------------------------------------------
+static void BlueprintActionDatabaseImpl::AddSkeletonActions(const USkeleton& Skeleton, FActionList& ActionListOut)
+{
+	for (int32 I = 0; I < Skeleton.AnimationNotifies.Num(); ++I)
+	{
+		FName NotifyName = Skeleton.AnimationNotifies[I];
+		FString Label = NotifyName.ToString();
+
+		FString SignatureName = FString::Printf(TEXT("AnimNotify_%s"), *Label);
+		ActionListOut.Add(UBlueprintEventNodeSpawner::Create(UK2Node_Event::StaticClass(), FName(*SignatureName)));
+	}
+}
+
 //------------------------------------------------------------------------------
 static void BlueprintActionDatabaseImpl::AddBlueprintGraphActions(UBlueprint const* const Blueprint, FActionList& ActionListOut)
 {
@@ -552,7 +573,7 @@ static void BlueprintActionDatabaseImpl::AddBlueprintGraphActions(UBlueprint con
 				ActionListOut.Add(SetVarSpawner);
 			}
 		}
-	}	
+	}
 }
 
 //------------------------------------------------------------------------------
@@ -796,6 +817,12 @@ void FBlueprintActionDatabase::RefreshAll()
 		UClass* const Class = (*ClassIt);
 		RefreshClassActions(Class);
 	}
+	// this handles creating entries for skeletons that were loaded before the database was alive:
+	for( TObjectIterator<USkeleton> SkeletonIt; SkeletonIt; ++SkeletonIt )
+	{
+		FActionList& ClassActionList = ActionRegistry.FindOrAdd(*SkeletonIt);
+		BlueprintActionDatabaseImpl::AddSkeletonActions(**SkeletonIt, ClassActionList);
+	}
 }
 
 //------------------------------------------------------------------------------
@@ -906,6 +933,11 @@ void FBlueprintActionDatabase::RefreshAssetActions(UObject* const AssetObject)
 	FActionList& AssetActionList = ActionRegistry.FindOrAdd(AssetObject);
 	AssetActionList.Empty();
 
+	if(const USkeleton* Skeleton = Cast<USkeleton>(AssetObject))
+	{
+		AddSkeletonActions(*Skeleton, AssetActionList);
+	}
+
 	UBlueprint* BlueprintAsset = Cast<UBlueprint>(AssetObject);
 	if (BlueprintAsset != nullptr)
 	{
diff --git a/Engine/Source/Editor/BlueprintGraph/Private/BlueprintActionFilter.cpp b/Engine/Source/Editor/BlueprintGraph/Private/BlueprintActionFilter.cpp
--- a/Engine/Source/Editor/BlueprintGraph/Private/BlueprintActionFilter.cpp
+++ b/Engine/Source/Editor/BlueprintGraph/Private/BlueprintActionFilter.cpp
@@ -288,6 +288,15 @@ namespace BlueprintActionFilterImpl
 	 */
 	static bool IsNodeTemplateSelfFiltered(FBlueprintActionFilter const& Filter, FBlueprintActionInfo& BlueprintAction);
 
+	/**
+	 * Rejection test that checks if the skeleton associated with the current blueprint
+	 * will raise any of the available anim notification events
+	 * 
+	 * @param  Filter			Holds the graph context for this test.
+	 * @param  BlueprintAction	The action you wish to query.
+	 * @return true if the action is an animnotification that is incompatible with the current skeleton
+	 */
+	static bool IsIncompatibleAnimNotification( FBlueprintActionFilter const& Filter, FBlueprintActionInfo& BlueprintAction); 
 };
 
 //------------------------------------------------------------------------------
@@ -1134,6 +1143,41 @@ static bool BlueprintActionFilterImpl::IsNodeTemplateSelfFiltered(FBlueprintActi
 	return bIsFilteredOut;
 }
 
+static bool BlueprintActionFilterImpl::IsIncompatibleAnimNotification(FBlueprintActionFilter const& Filter, FBlueprintActionInfo& BlueprintAction)
+{
+	bool bIsFilteredOut = false;
+
+	if ( BlueprintAction.GetNodeClass()->IsChildOf<UK2Node_Event>() )
+	{
+		if( const USkeleton* SkeletonOwningEvent = Cast<USkeleton>(BlueprintAction.GetActionOwner()) )
+		{
+			// The event is owned by a skeleton. Only show if it the current anim blueprint is targetting
+			// that skeleton:
+			FBlueprintActionContext const& FilterContext = Filter.Context;
+			bool bFoundInAllBlueprints = true;
+
+			for (const UBlueprint* Blueprint : FilterContext.Blueprints)
+			{
+				bool bFoundInCurrentBlueprint = false;
+				if (const UAnimBlueprint* AnimBlueprint = Cast<UAnimBlueprint>(Blueprint))
+				{
+					if( AnimBlueprint->TargetSkeleton == SkeletonOwningEvent )
+					{
+						bFoundInCurrentBlueprint = true;
+						break;
+					}
+				}
+				bFoundInAllBlueprints &= bFoundInCurrentBlueprint;
+			}
+			// If all of the blueprints selected aren't animblueprints targetting this skeleton then we need to 
+			// filter it out:
+			bIsFilteredOut = !bFoundInAllBlueprints; 
+		}
+	}
+
+	return bIsFilteredOut;
+}
+
 /*******************************************************************************
  * FBlueprintActionInfo
  ******************************************************************************/
@@ -1287,6 +1331,7 @@ FBlueprintActionFilter::FBlueprintActionFilter(uint32 Flags/*= 0x00*/)
 	//
 	// this test in-particular spawns a template-node and then calls 
 	// AllocateDefaultPins() which is costly, so it should be very last!
+	AddRejectionTest(FRejectionTestDelegate::CreateStatic(IsIncompatibleAnimNotification));
 	AddRejectionTest(FRejectionTestDelegate::CreateStatic(IsNodeTemplateSelfFiltered));
 	AddRejectionTest(FRejectionTestDelegate::CreateStatic(IsMissingMatchingPinParam));
 	AddRejectionTest(FRejectionTestDelegate::CreateStatic(IsMissmatchedPropertyType));
