More anim notification generation - now if there's an anim notification in your state machine it will appear in the new bp context menu

diff --git a/Engine/Source/Editor/BlueprintGraph/Private/BlueprintActionDatabase.cpp b/Engine/Source/Editor/BlueprintGraph/Private/BlueprintActionDatabase.cpp
--- a/Engine/Source/Editor/BlueprintGraph/Private/BlueprintActionDatabase.cpp
+++ b/Engine/Source/Editor/BlueprintGraph/Private/BlueprintActionDatabase.cpp
@@ -287,6 +287,16 @@ namespace BlueprintActionDatabaseImpl
 	 * @param  ActionListOut	The list you want populated with new spawners.
 	 */
 	static void AddBlueprintGraphActions(UBlueprint const* const Blueprint, FActionList& ActionListOut);
+	
+	/**
+	 * If the associated class is an anim blueprint generated class, then this
+	 * will loop over AnimNotification events in the anim blueprint generated
+	 * class and create node spawners for those events.
+	 *
+	 * @param  Blueprint		The blueprint which you want graph associated node-spawners for.
+	 * @param  ActionListOut	The list you want populated with new spawners.
+	 */
+	static void AddAnimBlueprintGraphActions( UAnimBlueprint const* AnimBlueprint, FActionList& ActionListOut );
 
 	/**
 	 * Emulates UEdGraphSchema::GetGraphContextActions(). If the supplied class  
@@ -576,6 +586,25 @@ static void BlueprintActionDatabaseImpl::AddBlueprintGraphActions(UBlueprint con
 	}
 }
 
+//------------------------------------------------------------------------------
+static void BlueprintActionDatabaseImpl::AddAnimBlueprintGraphActions(UAnimBlueprint const* AnimBlueprint, FActionList& ActionListOut)
+{
+	if (UAnimBlueprintGeneratedClass* GeneratedClass = AnimBlueprint->GetAnimBlueprintGeneratedClass())
+	{
+		for (int32 NotifyIdx = 0; NotifyIdx < GeneratedClass->AnimNotifies.Num(); NotifyIdx++)
+		{
+			FName NotifyName = GeneratedClass->AnimNotifies[NotifyIdx].NotifyName;
+			if (NotifyName != NAME_None)
+			{
+				FString Label = NotifyName.ToString();
+
+				FString SignatureName = FString::Printf(TEXT("AnimNotify_%s"), *Label);
+				ActionListOut.Add(UBlueprintEventNodeSpawner::Create(UK2Node_Event::StaticClass(), FName(*SignatureName)));
+			}
+		}
+	}
+}
+
 //------------------------------------------------------------------------------
 static void BlueprintActionDatabaseImpl::GetNodeSpecificActions(TSubclassOf<UEdGraphNode const> const NodeClass, FBlueprintActionDatabaseRegistrar& Registrar)
 {
@@ -947,6 +976,11 @@ void FBlueprintActionDatabase::RefreshAssetActions(UObject* const AssetObject)
 			GetClassMemberActions(SkeletonClass, AssetActionList);
 		}
 
+		if( const UAnimBlueprint* AnimBlueprint = Cast<UAnimBlueprint>(BlueprintAsset) )
+		{
+			AddAnimBlueprintGraphActions( AnimBlueprint, AssetActionList );
+		}
+
 		// have to be careful not to register this callback twice for the 
 		// blueprint
 		if (!bHadExistingEntry)
