ttp 347436 LIVE: BP: An empty Sequence node will cause an Infinite Loop error

diff --git a/Engine/Source/Editor/BlueprintGraph/Private/K2Node_ExecutionSequence.cpp b/Engine/Source/Editor/BlueprintGraph/Private/K2Node_ExecutionSequence.cpp
--- a/Engine/Source/Editor/BlueprintGraph/Private/K2Node_ExecutionSequence.cpp
+++ b/Engine/Source/Editor/BlueprintGraph/Private/K2Node_ExecutionSequence.cpp
@@ -120,6 +120,12 @@ public:
 				Context.GotoFixupRequestMap.Add(&NextExecutionState, OutputPins[0]);
 			}
 		}
+		else
+		{
+			FBlueprintCompiledStatement& NextExecutionState = Context.AppendStatementForNode(Node);
+			NextExecutionState.Type = KCST_EndOfThread;
+			CompilerContext.MessageLog.Warning(*LOCTEXT("NoValidOutput_Warning", "@@ has no valid output").ToString(), Node);
+		}
 	}
 };
 
