Latent nodes no longer available in new context menu in functions.

diff --git a/Engine/Source/Editor/BlueprintGraph/Private/BlueprintActionFilter.cpp b/Engine/Source/Editor/BlueprintGraph/Private/BlueprintActionFilter.cpp
--- a/Engine/Source/Editor/BlueprintGraph/Private/BlueprintActionFilter.cpp
+++ b/Engine/Source/Editor/BlueprintGraph/Private/BlueprintActionFilter.cpp
@@ -140,6 +140,16 @@ namespace BlueprintActionFilterImpl
 	 */
 	static bool IsIncompatibleImpureNode(FBlueprintActionFilter const& Filter, FBlueprintActionInfo& BlueprintAction);
 	
+	/**
+	 * Rejection test that checks to see if the supplied node-spawner would 
+	 * produce a latent node, incompatible with the specified graph.
+	 * 
+	 * @param  Filter			Holds the graph context for this test.
+	 * @param  BlueprintAction	The action you wish to query.
+	 * @return True if the action would spawn a node (and the graph wouldn't allow it).
+	 */
+	static bool IsIncompatibleLatentNode(FBlueprintActionFilter const& Filter, FBlueprintActionInfo& BlueprintAction);
+
 	/**
 	 * Rejection test that checks to see if the supplied node-spawner would 
 	 * produce a node incompatible with the specified graph.
@@ -149,7 +159,7 @@ namespace BlueprintActionFilterImpl
 	 * @return True if the action would spawn a node (and the graph wouldn't allow it).
 	 */
 	static bool IsIncompatibleWithGraphType(FBlueprintActionFilter const& Filter, FBlueprintActionInfo& BlueprintAction);
-	
+
 	/**
 	 * Rejection test that checks to see if the node-spawner has any associated
 	 * "non-target" fields that are global/static.
@@ -713,6 +723,27 @@ static bool BlueprintActionFilterImpl::IsIncompatibleImpureNode(FBlueprintAction
 	return bIsFilteredOut;
 }
 
+//------------------------------------------------------------------------------
+static bool BlueprintActionFilterImpl::IsIncompatibleLatentNode(FBlueprintActionFilter const& Filter, FBlueprintActionInfo& BlueprintAction)
+{
+	bool bAllowLatentNodes = true;
+	FBlueprintActionContext const& FilterContext = Filter.Context;
+
+	for (UEdGraph* Graph : FilterContext.Graphs)
+	{
+		if (UEdGraphSchema_K2* K2Schema = Graph->Schema->GetDefaultObject<UEdGraphSchema_K2>())
+		{
+			if(K2Schema->GetGraphType(Graph) == GT_Function)
+			{
+				bAllowLatentNodes = false;
+			}
+		}
+	}
+
+	bool bIsFilteredOut = !bAllowLatentNodes && IsLatent(BlueprintAction);
+	return bIsFilteredOut;
+}
+
 //------------------------------------------------------------------------------
 static bool BlueprintActionFilterImpl::IsIncompatibleWithGraphType(FBlueprintActionFilter const& Filter, FBlueprintActionInfo& BlueprintAction)
 {
@@ -1336,6 +1367,7 @@ FBlueprintActionFilter::FBlueprintActionFilter(uint32 Flags/*= 0x00*/)
 	AddRejectionTest(FRejectionTestDelegate::CreateStatic(IsMissingMatchingPinParam));
 	AddRejectionTest(FRejectionTestDelegate::CreateStatic(IsMissmatchedPropertyType));
 	AddRejectionTest(FRejectionTestDelegate::CreateStatic(IsFunctionMissingPinParam));
+	AddRejectionTest(FRejectionTestDelegate::CreateStatic(IsIncompatibleLatentNode));
 	AddRejectionTest(FRejectionTestDelegate::CreateStatic(IsIncompatibleImpureNode));
 	
 	AddRejectionTest(FRejectionTestDelegate::CreateStatic(IsFieldCategoryHidden));
