FMemberReferences now default to the authoritative class (unless set through the SetDirect() function).

This is to support references being set for the skeleton class (in the new BP menu system).

diff --git a/Engine/Source/Editor/BlueprintGraph/Classes/K2Node.h b/Engine/Source/Editor/BlueprintGraph/Classes/K2Node.h
--- a/Engine/Source/Editor/BlueprintGraph/Classes/K2Node.h
+++ b/Engine/Source/Editor/BlueprintGraph/Classes/K2Node.h
@@ -465,6 +465,11 @@ public:
 		bSelfContext = bIsConsideredSelfContext;
 		bWasDeprecated = false;
 
+		if (MemberParentClass != nullptr)
+		{
+			MemberParentClass = MemberParentClass->GetAuthoritativeClass();
+		}
+
 		MemberGuid.Invalidate();
 		if (InField->GetOwnerClass())
 		{
@@ -645,11 +650,15 @@ public:
 				MemberGuid.Invalidate();
 				UBlueprint::GetGuidFromClassByFieldName<TFieldType>(TargetScope, MemberName, MemberGuid);
 
-				// Re-evaluate self-ness against the redirect if we were given a valid SelfScope
-				if(MemberParentClass != NULL && SelfScope != NULL)
+				if (MemberParentClass != nullptr)
 				{
-					SetGivenSelfScope(MemberName, MemberGuid, MemberParentClass, SelfScope);
-				}
+					MemberParentClass = MemberParentClass->GetAuthoritativeClass();
+					// Re-evaluate self-ness against the redirect if we were given a valid SelfScope
+					if (SelfScope != NULL)
+					{
+						SetGivenSelfScope(MemberName, MemberGuid, MemberParentClass, SelfScope);
+					}
+				}	
 			}
 			else if(TargetScope != NULL)
 			{
diff --git a/Engine/Source/Editor/BlueprintGraph/Private/DelegateNodeHandlers.cpp b/Engine/Source/Editor/BlueprintGraph/Private/DelegateNodeHandlers.cpp
--- a/Engine/Source/Editor/BlueprintGraph/Private/DelegateNodeHandlers.cpp
+++ b/Engine/Source/Editor/BlueprintGraph/Private/DelegateNodeHandlers.cpp
@@ -59,7 +59,12 @@ struct FKCHandlerDelegateHelper
 
 		if (!BoundProperty)
 		{
-			MessageLog.Error(*LOCTEXT("NoDelegateProperty", "Event Dispatcher has no property @@").ToString(), DelegateNode);
+			FString const OwnerName = PropertyOwnerClass->GetName();
+			FString const PropName  = DelegateNode->GetPropertyName().ToString();
+
+			FText const ErrorFormat = LOCTEXT("DelegateNotFound", "Could not find an event-dispatcher named \"%s\" in '%s'.\nMake sure '%s' has been compiled for @@");
+			MessageLog.Error(*FString::Printf(*ErrorFormat.ToString(), *PropName, *OwnerName, *OwnerName), DelegateNode);
+
 			return NULL;
 		}
 
diff --git a/Engine/Source/Editor/BlueprintGraph/Private/K2Node.cpp b/Engine/Source/Editor/BlueprintGraph/Private/K2Node.cpp
--- a/Engine/Source/Editor/BlueprintGraph/Private/K2Node.cpp
+++ b/Engine/Source/Editor/BlueprintGraph/Private/K2Node.cpp
@@ -1235,7 +1235,7 @@ void UK2Node::GetPinHoverText(const UEdGraphPin& Pin, FString& HoverTextOut) con
 void FMemberReference::SetExternalMember(FName InMemberName, TSubclassOf<class UObject> InMemberParentClass)
 {
 	MemberName = InMemberName;
-	MemberParentClass = InMemberParentClass;
+	MemberParentClass = (InMemberParentClass != nullptr) ? InMemberParentClass->GetAuthoritativeClass() : nullptr;
 	MemberScope.Empty();
 	bSelfContext = false;
 	bWasDeprecated = false;
@@ -1264,7 +1264,7 @@ void FMemberReference::SetGivenSelfScope(const FName InMemberName, const FGuid I
 {
 	MemberName = InMemberName;
 	MemberGuid = InMemberGuid;
-	MemberParentClass = InMemberParentClass;
+	MemberParentClass = (InMemberParentClass != nullptr) ? InMemberParentClass->GetAuthoritativeClass() : nullptr;
 	MemberScope.Empty();
 	bSelfContext = (SelfScope->IsChildOf(InMemberParentClass)) || (SelfScope->ClassGeneratedBy == InMemberParentClass->ClassGeneratedBy);
 	bWasDeprecated = false;
diff --git a/Engine/Source/Editor/BlueprintGraph/Private/K2Node_CallFunction.cpp b/Engine/Source/Editor/BlueprintGraph/Private/K2Node_CallFunction.cpp
--- a/Engine/Source/Editor/BlueprintGraph/Private/K2Node_CallFunction.cpp
+++ b/Engine/Source/Editor/BlueprintGraph/Private/K2Node_CallFunction.cpp
@@ -197,6 +197,24 @@ void UK2Node_CallFunction::AllocateDefaultPins()
 		}
 	}
 
+	if (Function == nullptr)
+	{
+		if (!FunctionReference.IsSelfContext())
+		{
+			UClass* FunctionClass = FunctionReference.GetMemberParentClass(MyBlueprint->GeneratedClass);
+			if (UBlueprintGeneratedClass* BpClassOwner = Cast<UBlueprintGeneratedClass>(FunctionClass))
+			{
+				// this function could currently only be a part of some skeleton 
+				// class (the blueprint has not be compiled with it yet), so let's 
+				// check the skeleton class as well, see if we can pull pin data 
+				// from there...
+				UBlueprint* FunctionBlueprint = CastChecked<UBlueprint>(BpClassOwner->ClassGeneratedBy);
+				Function = FindField<UFunction>(FunctionBlueprint->SkeletonGeneratedClass, FunctionReference.GetMemberName());
+			}
+		}
+	}
+
+
 	if (Function == NULL)
 	{
 		// The function no longer exists in the stored scope
@@ -281,9 +299,15 @@ UEdGraphPin* UK2Node_CallFunction::CreateSelfPin(const UFunction* Function)
 
 	// Create the self pin
 	UClass* FunctionClass = CastChecked<UClass>(FirstDeclaredFunction->GetOuter());
+	// we don't want blueprint-function target pins to be formed from the
+	// skeleton class (otherwise, they could be incompatible with other pins
+	// that represent the same type)... this here could lead to a compiler 
+	// warning (the GeneratedClass could not have the function yet), but in
+	// that, the user would be reminded to compile the other blueprint
+	FunctionClass = FunctionClass->GetAuthoritativeClass();
 
 	UEdGraphPin* SelfPin = NULL;
-	if ((FunctionClass == GetBlueprint()->GeneratedClass) || (FunctionClass == GetBlueprint()->SkeletonGeneratedClass))
+	if (FunctionClass == GetBlueprint()->GeneratedClass)
 	{
 		// This means the function is defined within the blueprint, so the pin should be a true "self" pin
 		SelfPin = CreatePin(EGPD_Input, K2Schema->PC_Object, K2Schema->PSC_Self, NULL, false, false, K2Schema->PN_Self);
@@ -1172,7 +1196,12 @@ void UK2Node_CallFunction::ValidateNodeDuringCompilation(class FCompilerResultsL
 	UFunction *Function = GetTargetFunction();
 	if (Function == NULL)
 	{
-		MessageLog.Warning(*FString::Printf(*LOCTEXT("FunctionNotFound", "Unable to find function with name '%s' for @@").ToString(), *FunctionReference.GetMemberName().ToString()), this);
+		UClass* FuncOwnerClass = FunctionReference.GetMemberParentClass(GetBlueprint()->GeneratedClass);
+		FString const OwnerName = FuncOwnerClass->GetName();
+		FString const FunctName = FunctionReference.GetMemberName().ToString();
+
+		FText const WarningFormat = LOCTEXT("FunctionNotFound", "Could not find a function named \"%s\" in '%s'.\nMake sure '%s' has been compiled for @@");
+		MessageLog.Warning(*FString::Printf(*WarningFormat.ToString(), *FunctName, *OwnerName, *OwnerName), this);
 	}
 	else if (Function->HasMetaData(FBlueprintMetadata::MD_ExpandEnumAsExecs) && bWantsEnumToExecExpansion == false)
 	{
diff --git a/Engine/Source/Editor/BlueprintGraph/Private/K2Node_CallParentFunction.cpp b/Engine/Source/Editor/BlueprintGraph/Private/K2Node_CallParentFunction.cpp
--- a/Engine/Source/Editor/BlueprintGraph/Private/K2Node_CallParentFunction.cpp
+++ b/Engine/Source/Editor/BlueprintGraph/Private/K2Node_CallParentFunction.cpp
@@ -49,13 +49,16 @@ void UK2Node_CallParentFunction::SetFromFunction(const UFunction* Function)
 		bIsPureFunc = Function->HasAnyFunctionFlags(FUNC_BlueprintPure);
 		bIsConstFunc = Function->HasAnyFunctionFlags(FUNC_Const);
 
+		UClass* OwnerClass = Function->GetOwnerClass();
+
 		FGuid FunctionGuid;
-		if (Function->GetOwnerClass())
+		if (OwnerClass != nullptr)
 		{
-			UBlueprint::GetGuidFromClassByFieldName<UFunction>(Function->GetOwnerClass(), Function->GetFName(), FunctionGuid);
+			OwnerClass = OwnerClass->GetAuthoritativeClass();
+			UBlueprint::GetGuidFromClassByFieldName<UFunction>(OwnerClass, Function->GetFName(), FunctionGuid);
 		}
 
-		FunctionReference.SetDirect(Function->GetFName(), FunctionGuid, Function->GetOwnerClass(), false);
+		FunctionReference.SetDirect(Function->GetFName(), FunctionGuid, OwnerClass, /*bIsConsideredSelfContext =*/false);
 	}
 }
 
diff --git a/Engine/Source/Editor/BlueprintGraph/Private/K2Node_MCDelegate.cpp b/Engine/Source/Editor/BlueprintGraph/Private/K2Node_MCDelegate.cpp
--- a/Engine/Source/Editor/BlueprintGraph/Private/K2Node_MCDelegate.cpp
+++ b/Engine/Source/Editor/BlueprintGraph/Private/K2Node_MCDelegate.cpp
@@ -63,11 +63,14 @@ void UK2Node_BaseMCDelegate::AllocateDefaultPins()
 	CreatePin(EGPD_Input, K2Schema->PC_Exec, TEXT(""), NULL, false, false, K2Schema->PN_Execute);
 	CreatePin(EGPD_Output, K2Schema->PC_Exec, TEXT(""), NULL, false, false, K2Schema->PN_Then);
 
-	const auto Property = DelegateReference.ResolveMember<UMulticastDelegateProperty>(this);
-	const auto PropertyOwnerClass = (Property ? Property->GetOwnerClass() : NULL);
+	UClass* PropertyOwnerClass = DelegateReference.GetMemberParentClass(this);
+	if (PropertyOwnerClass != nullptr)
+	{
+		PropertyOwnerClass = PropertyOwnerClass->GetAuthoritativeClass();
+	}
 	const auto Blueprint = GetBlueprint();
 	
-	const bool bUseSelf = Blueprint && (PropertyOwnerClass == Blueprint->GeneratedClass || PropertyOwnerClass == Blueprint->SkeletonGeneratedClass);
+	const bool bUseSelf = Blueprint && (PropertyOwnerClass == Blueprint->GeneratedClass);
 
 	UEdGraphPin* SelfPin = NULL;
 	if (bUseSelf)
@@ -87,22 +90,28 @@ void UK2Node_BaseMCDelegate::AllocateDefaultPins()
 
 UFunction* UK2Node_BaseMCDelegate::GetDelegateSignature(bool bForceNotFromSkelClass) const
 {
-	FMemberReference ReferenceToUse;
-
-	if(!bForceNotFromSkelClass)
+	UClass* OwnerClass = DelegateReference.GetMemberParentClass(this);
+	if (bForceNotFromSkelClass)
 	{
-		ReferenceToUse = DelegateReference;
+		OwnerClass = (OwnerClass != nullptr) ? OwnerClass->GetAuthoritativeClass() : nullptr;
 	}
-	else
+	else if (UBlueprintGeneratedClass* BpClassOwner = Cast<UBlueprintGeneratedClass>(OwnerClass))
 	{
-		UClass* OwnerClass = DelegateReference.GetMemberParentClass(this);
-		FGuid DelegateGuid;
-		if (OwnerClass)
-		{
-			UBlueprint::GetGuidFromClassByFieldName<UFunction>(OwnerClass, DelegateReference.GetMemberName(), DelegateGuid);
-		}
-		ReferenceToUse.SetDirect(DelegateReference.GetMemberName(), DelegateGuid, OwnerClass ? OwnerClass->GetAuthoritativeClass() : NULL, false);
+		UBlueprint* DelegateBlueprint = CastChecked<UBlueprint>(BpClassOwner->ClassGeneratedBy);
+		// favor the skeleton class, because the generated class may not 
+		// have the delegate yet (hasn't been compiled with it), or it could 
+		// be out of date
+		OwnerClass = DelegateBlueprint->SkeletonGeneratedClass;
+	}
+
+	FMemberReference ReferenceToUse;
+	FGuid DelegateGuid;
+
+	if (OwnerClass != nullptr)
+	{
+		UBlueprint::GetGuidFromClassByFieldName<UFunction>(OwnerClass, DelegateReference.GetMemberName(), DelegateGuid);
 	}
+	ReferenceToUse.SetDirect(DelegateReference.GetMemberName(), DelegateGuid, OwnerClass, /*bIsConsideredSelfContext =*/false);
 
 	UMulticastDelegateProperty* DelegateProperty = ReferenceToUse.ResolveMember<UMulticastDelegateProperty>(this);
 	return (DelegateProperty != NULL) ? DelegateProperty->SignatureFunction : NULL;
diff --git a/Engine/Source/Editor/BlueprintGraph/Private/K2Node_Variable.cpp b/Engine/Source/Editor/BlueprintGraph/Private/K2Node_Variable.cpp
--- a/Engine/Source/Editor/BlueprintGraph/Private/K2Node_Variable.cpp
+++ b/Engine/Source/Editor/BlueprintGraph/Private/K2Node_Variable.cpp
@@ -60,6 +60,23 @@ bool UK2Node_Variable::CreatePinForVariable(EEdGraphPinDirection Direction)
 	const UEdGraphSchema_K2* K2Schema = GetDefault<UEdGraphSchema_K2>();
 
 	UProperty* VariableProperty = GetPropertyForVariable();
+	if (VariableProperty == nullptr)
+	{
+		if (!VariableReference.IsSelfContext())
+		{
+			UClass* VariableClass = VariableReference.GetMemberParentClass(GetBlueprint()->GeneratedClass);
+			if (UBlueprintGeneratedClass* BpClassOwner = Cast<UBlueprintGeneratedClass>(VariableClass))
+			{
+				// this variable could currently only be a part of some skeleton 
+				// class (the blueprint has not be compiled with it yet), so let's 
+				// check the skeleton class as well, see if we can pull pin data 
+				// from there...
+				UBlueprint* VariableBlueprint = CastChecked<UBlueprint>(BpClassOwner->ClassGeneratedBy);
+				VariableProperty = FindField<UProperty>(VariableBlueprint->SkeletonGeneratedClass, VariableReference.GetMemberName());
+			}
+		}
+	}
+
 	if (VariableProperty != NULL)
 	{
 		// Create the pin
@@ -298,7 +315,12 @@ void UK2Node_Variable::ValidateNodeDuringCompilation(class FCompilerResultsLog&
 	{
 		if (!VariableReference.IsDeprecated())
 		{
-			MessageLog.Warning(*FString::Printf(*LOCTEXT("VariableNotFound", "Unable to find variable with name '%s' for @@").ToString(), *VariableReference.GetMemberName().ToString()), this);
+			UClass* VarOwnerClass = VariableReference.GetMemberParentClass(GetBlueprint()->GeneratedClass);
+			FString const OwnerName = VarOwnerClass->GetName();
+			FString const VarName = VariableReference.GetMemberName().ToString();
+
+			FText const WarningFormat = LOCTEXT("VariableNotFound", "Could not find a variable named \"%s\" in '%s'.\nMake sure '%s' has been compiled for @@");
+			MessageLog.Warning(*FString::Printf(*WarningFormat.ToString(), *VarName, *OwnerName, *OwnerName), this);
 		}
 		else
 		{
