Reverting some perf optimizations that were resulting in stale node text (have to figure out when to clear the cached text).
TTP #346788

diff --git a/Engine/Source/Developer/MovieSceneTools/Private/K2Node_PlayMovieScene.cpp b/Engine/Source/Developer/MovieSceneTools/Private/K2Node_PlayMovieScene.cpp
--- a/Engine/Source/Developer/MovieSceneTools/Private/K2Node_PlayMovieScene.cpp
+++ b/Engine/Source/Developer/MovieSceneTools/Private/K2Node_PlayMovieScene.cpp
@@ -82,7 +82,10 @@ FText UK2Node_PlayMovieScene::GetNodeTitle(ENodeTitleType::Type TitleType) const
 	{
 		return NSLOCTEXT("PlayMovieSceneNode", "NodeTitleWithNoMovieScene", "Play Movie Scene (No Asset)");
 	}
-	else if (CachedNodeTitle.IsOutOfDate())
+	// @TODO: don't know enough about this node type to comfortably assert that
+	//        the MovieScene won't change after the node has spawned... until
+	//        then, we'll leave this optimization off
+	else //if (CachedNodeTitle.IsOutOfDate())
 	{
 		FFormatNamedArguments Args;
 		Args.Add(TEXT("SceneName"), FText::FromString(MovieScene->GetName()));
diff --git a/Engine/Source/Editor/AnimGraph/Private/AnimGraphNode_BlendListByEnum.cpp b/Engine/Source/Editor/AnimGraph/Private/AnimGraphNode_BlendListByEnum.cpp
--- a/Engine/Source/Editor/AnimGraph/Private/AnimGraphNode_BlendListByEnum.cpp
+++ b/Engine/Source/Editor/AnimGraph/Private/AnimGraphNode_BlendListByEnum.cpp
@@ -36,7 +36,10 @@ FText UAnimGraphNode_BlendListByEnum::GetNodeTitle(ENodeTitleType::Type TitleTyp
 	{
 		return LOCTEXT("AnimGraphNode_BlendListByEnum_TitleError", "ERROR: Blend Poses (by missing enum)");
 	}
-	else if (CachedNodeTitle.IsOutOfDate())
+	// @TODO: don't know enough about this node type to comfortably assert that
+	//        the BoundEnum won't change after the node has spawned... until
+	//        then, we'll leave this optimization off
+	else //if (CachedNodeTitle.IsOutOfDate())
 	{
 		FFormatNamedArguments Args;
 		Args.Add(TEXT("EnumName"), FText::FromString(BoundEnum->GetName()));
diff --git a/Engine/Source/Editor/AnimGraph/Private/AnimGraphNode_BlendSpaceEvaluator.cpp b/Engine/Source/Editor/AnimGraph/Private/AnimGraphNode_BlendSpaceEvaluator.cpp
--- a/Engine/Source/Editor/AnimGraph/Private/AnimGraphNode_BlendSpaceEvaluator.cpp
+++ b/Engine/Source/Editor/AnimGraph/Private/AnimGraphNode_BlendSpaceEvaluator.cpp
@@ -34,7 +34,9 @@ FText UAnimGraphNode_BlendSpaceEvaluator::GetNodeTitle(ENodeTitleType::Type Titl
 			return LOCTEXT("BlendSpaceEvaluator_NONE_Title", "(None)\nBlendspace Evaluator");
 		}
 	}
-	else if (!CachedNodeTitles.IsTitleCached(TitleType))
+	// @TODO: the bone can be altered in the property editor, so we have to 
+	//        choose to mark this dirty when that happens for this to properly work
+	else //if (!CachedNodeTitles.IsTitleCached(TitleType))
 	{
 		const FText BlendSpaceName = FText::FromString(Node.BlendSpace->GetName());
 
diff --git a/Engine/Source/Editor/AnimGraph/Private/AnimGraphNode_BlendSpacePlayer.cpp b/Engine/Source/Editor/AnimGraph/Private/AnimGraphNode_BlendSpacePlayer.cpp
--- a/Engine/Source/Editor/AnimGraph/Private/AnimGraphNode_BlendSpacePlayer.cpp
+++ b/Engine/Source/Editor/AnimGraph/Private/AnimGraphNode_BlendSpacePlayer.cpp
@@ -36,7 +36,9 @@ FText UAnimGraphNode_BlendSpacePlayer::GetNodeTitle(ENodeTitleType::Type TitleTy
 			return LOCTEXT("BlendspacePlayer_NONE_Title", "(None)\nBlendspace Player");
 		}
 	}
-	else if (!CachedNodeTitles.IsTitleCached(TitleType))
+	// @TODO: the bone can be altered in the property editor, so we have to 
+	//        choose to mark this dirty when that happens for this to properly work
+	else //if (!CachedNodeTitles.IsTitleCached(TitleType))
 	{
 		const FText BlendSpaceName = FText::FromString(Node.BlendSpace->GetName());
 
diff --git a/Engine/Source/Editor/AnimGraph/Private/AnimGraphNode_BoneDrivenController.cpp b/Engine/Source/Editor/AnimGraph/Private/AnimGraphNode_BoneDrivenController.cpp
--- a/Engine/Source/Editor/AnimGraph/Private/AnimGraphNode_BoneDrivenController.cpp
+++ b/Engine/Source/Editor/AnimGraph/Private/AnimGraphNode_BoneDrivenController.cpp
@@ -22,7 +22,9 @@ FText UAnimGraphNode_BoneDrivenController::GetNodeTitle(ENodeTitleType::Type Tit
 	{
 		return GetControllerDescription();
 	}
-	else if (!CachedNodeTitles.IsTitleCached(TitleType))
+	// @TODO: the bone can be altered in the property editor, so we have to 
+	//        choose to mark this dirty when that happens for this to properly work
+	else //if (!CachedNodeTitles.IsTitleCached(TitleType))
 	{
 		FFormatNamedArguments Args;
 		Args.Add(TEXT("ControllerDesc"), GetControllerDescription());
diff --git a/Engine/Source/Editor/AnimGraph/Private/AnimGraphNode_CopyBone.cpp b/Engine/Source/Editor/AnimGraph/Private/AnimGraphNode_CopyBone.cpp
--- a/Engine/Source/Editor/AnimGraph/Private/AnimGraphNode_CopyBone.cpp
+++ b/Engine/Source/Editor/AnimGraph/Private/AnimGraphNode_CopyBone.cpp
@@ -29,7 +29,9 @@ FText UAnimGraphNode_CopyBone::GetNodeTitle(ENodeTitleType::Type TitleType) cons
 	{
 		return GetControllerDescription();
 	}
-	else if (!CachedNodeTitles.IsTitleCached(TitleType))
+	// @TODO: the bone can be altered in the property editor, so we have to 
+	//        choose to mark this dirty when that happens for this to properly work
+	else //if (!CachedNodeTitles.IsTitleCached(TitleType))
 	{
 		FFormatNamedArguments Args;
 		Args.Add(TEXT("ControllerDescription"), GetControllerDescription());
diff --git a/Engine/Source/Editor/AnimGraph/Private/AnimGraphNode_LookAt.cpp b/Engine/Source/Editor/AnimGraph/Private/AnimGraphNode_LookAt.cpp
--- a/Engine/Source/Editor/AnimGraph/Private/AnimGraphNode_LookAt.cpp
+++ b/Engine/Source/Editor/AnimGraph/Private/AnimGraphNode_LookAt.cpp
@@ -34,7 +34,9 @@ FText UAnimGraphNode_LookAt::GetNodeTitle(ENodeTitleType::Type TitleType) const
 	{
 		return GetControllerDescription();
 	}
-	else if (!CachedNodeTitles.IsTitleCached(TitleType))
+	// @TODO: the bone can be altered in the property editor, so we have to 
+	//        choose to mark this dirty when that happens for this to properly work
+	else //if (!CachedNodeTitles.IsTitleCached(TitleType))
 	{
 		FFormatNamedArguments Args;
 		Args.Add(TEXT("ControllerDescription"), GetControllerDescription());
diff --git a/Engine/Source/Editor/AnimGraph/Private/AnimGraphNode_ModifyBone.cpp b/Engine/Source/Editor/AnimGraph/Private/AnimGraphNode_ModifyBone.cpp
--- a/Engine/Source/Editor/AnimGraph/Private/AnimGraphNode_ModifyBone.cpp
+++ b/Engine/Source/Editor/AnimGraph/Private/AnimGraphNode_ModifyBone.cpp
@@ -34,7 +34,9 @@ FText UAnimGraphNode_ModifyBone::GetNodeTitle(ENodeTitleType::Type TitleType) co
 	{
 		return GetControllerDescription();
 	}
-	else if (!CachedNodeTitles.IsTitleCached(TitleType))
+	// @TODO: the bone can be altered in the property editor, so we have to 
+	//        choose to mark this dirty when that happens for this to properly work
+	else //if (!CachedNodeTitles.IsTitleCached(TitleType))
 	{
 		FFormatNamedArguments Args;
 		Args.Add(TEXT("ControllerDescription"), GetControllerDescription());
diff --git a/Engine/Source/Editor/AnimGraph/Private/AnimGraphNode_RotationMultiplier.cpp b/Engine/Source/Editor/AnimGraph/Private/AnimGraphNode_RotationMultiplier.cpp
--- a/Engine/Source/Editor/AnimGraph/Private/AnimGraphNode_RotationMultiplier.cpp
+++ b/Engine/Source/Editor/AnimGraph/Private/AnimGraphNode_RotationMultiplier.cpp
@@ -28,7 +28,9 @@ FText UAnimGraphNode_RotationMultiplier::GetNodeTitle(ENodeTitleType::Type Title
 	{
 		return GetControllerDescription();
 	}
-	else if (!CachedNodeTitles.IsTitleCached(TitleType))
+	// @TODO: the bone can be altered in the property editor, so we have to 
+	//        choose to mark this dirty when that happens for this to properly work
+	else //if (!CachedNodeTitles.IsTitleCached(TitleType))
 	{
 		FFormatNamedArguments Args;
 		Args.Add(TEXT("ControllerDescription"), GetControllerDescription());
diff --git a/Engine/Source/Editor/AnimGraph/Private/AnimGraphNode_RotationOffsetBlendSpace.cpp b/Engine/Source/Editor/AnimGraph/Private/AnimGraphNode_RotationOffsetBlendSpace.cpp
--- a/Engine/Source/Editor/AnimGraph/Private/AnimGraphNode_RotationOffsetBlendSpace.cpp
+++ b/Engine/Source/Editor/AnimGraph/Private/AnimGraphNode_RotationOffsetBlendSpace.cpp
@@ -36,7 +36,9 @@ FText UAnimGraphNode_RotationOffsetBlendSpace::GetNodeTitle(ENodeTitleType::Type
 			return LOCTEXT("RotationOffsetBlend_NONE_Title", "(None)\nAimOffset");
 		}
 	}
-	else if (!CachedNodeTitles.IsTitleCached(TitleType))
+	// @TODO: the bone can be altered in the property editor, so we have to 
+	//        choose to mark this dirty when that happens for this to properly work
+	else //if (!CachedNodeTitles.IsTitleCached(TitleType))
 	{
 		const FText BlendSpaceName = FText::FromString(Node.BlendSpace->GetName());
 
diff --git a/Engine/Source/Editor/AnimGraph/Private/AnimGraphNode_SaveCachedPose.cpp b/Engine/Source/Editor/AnimGraph/Private/AnimGraphNode_SaveCachedPose.cpp
--- a/Engine/Source/Editor/AnimGraph/Private/AnimGraphNode_SaveCachedPose.cpp
+++ b/Engine/Source/Editor/AnimGraph/Private/AnimGraphNode_SaveCachedPose.cpp
@@ -55,7 +55,10 @@ FText UAnimGraphNode_SaveCachedPose::GetNodeTitle(ENodeTitleType::Type TitleType
 	{
 		return LOCTEXT("NewSaveCachedPose", "New Save cached pose...");
 	}
-	else if (CachedNodeTitle.IsOutOfDate())
+	// @TODO: don't know enough about this node type to comfortably assert that
+	//        the CacheName won't change after the node has spawned... until
+	//        then, we'll leave this optimization off
+	else //if (CachedNodeTitle.IsOutOfDate())
 	{
 		FFormatNamedArguments Args;
 		Args.Add(TEXT("NodeTitle"), FText::FromString(CacheName));
diff --git a/Engine/Source/Editor/AnimGraph/Private/AnimGraphNode_SequenceEvaluator.cpp b/Engine/Source/Editor/AnimGraph/Private/AnimGraphNode_SequenceEvaluator.cpp
--- a/Engine/Source/Editor/AnimGraph/Private/AnimGraphNode_SequenceEvaluator.cpp
+++ b/Engine/Source/Editor/AnimGraph/Private/AnimGraphNode_SequenceEvaluator.cpp
@@ -48,7 +48,10 @@ FText UAnimGraphNode_SequenceEvaluator::GetNodeTitle(ENodeTitleType::Type TitleT
 	{
 		return LOCTEXT("EvaluateSequence_TitleNONE", "Evaluate (None)");
 	}
-	else if (CachedNodeTitle.IsOutOfDate())
+	// @TODO: don't know enough about this node type to comfortably assert that
+	//        the CacheName won't change after the node has spawned... until
+	//        then, we'll leave this optimization off
+	else //if (CachedNodeTitle.IsOutOfDate())
 	{
 		const FText SequenceName = FText::FromString(Node.Sequence->GetName());
 
diff --git a/Engine/Source/Editor/AnimGraph/Private/AnimGraphNode_SequencePlayer.cpp b/Engine/Source/Editor/AnimGraph/Private/AnimGraphNode_SequencePlayer.cpp
--- a/Engine/Source/Editor/AnimGraph/Private/AnimGraphNode_SequencePlayer.cpp
+++ b/Engine/Source/Editor/AnimGraph/Private/AnimGraphNode_SequencePlayer.cpp
@@ -81,7 +81,9 @@ FText UAnimGraphNode_SequencePlayer::GetNodeTitle(ENodeTitleType::Type TitleType
 	{
 		return LOCTEXT("SequenceNullTitle", "Play (None)");
 	}
-	else if (!CachedNodeTitles.IsTitleCached(TitleType))
+	// @TODO: the bone can be altered in the property editor, so we have to 
+	//        choose to mark this dirty when that happens for this to properly work
+	else //if (!CachedNodeTitles.IsTitleCached(TitleType))
 	{
 		if(SyncGroup.GroupName == NAME_None)
 		{
diff --git a/Engine/Source/Editor/AnimGraph/Private/AnimGraphNode_Slot.cpp b/Engine/Source/Editor/AnimGraph/Private/AnimGraphNode_Slot.cpp
--- a/Engine/Source/Editor/AnimGraph/Private/AnimGraphNode_Slot.cpp
+++ b/Engine/Source/Editor/AnimGraph/Private/AnimGraphNode_Slot.cpp
@@ -39,7 +39,9 @@ FText UAnimGraphNode_Slot::GetNodeTitle(ENodeTitleType::Type TitleType) const
 			return LOCTEXT("SlotNodeTitle_NoName", "(No slot name)\nSlot");
 		}
 	}
-	else if (!CachedNodeTitles.IsTitleCached(TitleType))
+	// @TODO: the bone can be altered in the property editor, so we have to 
+	//        choose to mark this dirty when that happens for this to properly work
+	else //if (!CachedNodeTitles.IsTitleCached(TitleType))
 	{
 		FFormatNamedArguments Args;
 		Args.Add(TEXT("SlotName"), FText::FromName(Node.SlotName));
diff --git a/Engine/Source/Editor/AnimGraph/Private/AnimGraphNode_SpringBone.cpp b/Engine/Source/Editor/AnimGraph/Private/AnimGraphNode_SpringBone.cpp
--- a/Engine/Source/Editor/AnimGraph/Private/AnimGraphNode_SpringBone.cpp
+++ b/Engine/Source/Editor/AnimGraph/Private/AnimGraphNode_SpringBone.cpp
@@ -29,7 +29,9 @@ FText UAnimGraphNode_SpringBone::GetNodeTitle(ENodeTitleType::Type TitleType) co
 	{
 		return GetControllerDescription();
 	}
-	else if(!CachedNodeTitles.IsTitleCached(TitleType))
+	// @TODO: the bone can be altered in the property editor, so we have to 
+	//        choose to mark this dirty when that happens for this to properly work
+	else //if(!CachedNodeTitles.IsTitleCached(TitleType))
 	{
 		FFormatNamedArguments Args;
 		Args.Add(TEXT("ControllerDescription"), GetControllerDescription());
diff --git a/Engine/Source/Editor/AnimGraph/Private/AnimGraphNode_StateMachineBase.cpp b/Engine/Source/Editor/AnimGraph/Private/AnimGraphNode_StateMachineBase.cpp
--- a/Engine/Source/Editor/AnimGraph/Private/AnimGraphNode_StateMachineBase.cpp
+++ b/Engine/Source/Editor/AnimGraph/Private/AnimGraphNode_StateMachineBase.cpp
@@ -76,7 +76,10 @@ FText UAnimGraphNode_StateMachineBase::GetNodeTitle(ENodeTitleType::Type TitleTy
 	}
 	else if (TitleType == ENodeTitleType::FullTitle)
 	{
-		if (CachedFullTitle.IsOutOfDate())
+		// @TODO: don't know enough about this node type to comfortably assert that
+		//        the EditorStateMachineGraph won't change after the node has 
+		//        spawned... until then, we'll leave this optimization off
+		//if (CachedFullTitle.IsOutOfDate())
 		{
 			FFormatNamedArguments Args;
 			Args.Add(TEXT("Title"), FText::FromName(EditorStateMachineGraph->GetFName()));
diff --git a/Engine/Source/Editor/AnimGraph/Private/AnimGraphNode_Trail.cpp b/Engine/Source/Editor/AnimGraph/Private/AnimGraphNode_Trail.cpp
--- a/Engine/Source/Editor/AnimGraph/Private/AnimGraphNode_Trail.cpp
+++ b/Engine/Source/Editor/AnimGraph/Private/AnimGraphNode_Trail.cpp
@@ -29,7 +29,9 @@ FText UAnimGraphNode_Trail::GetNodeTitle(ENodeTitleType::Type TitleType) const
 	{
 		return GetControllerDescription();
 	}
-	else if (!CachedNodeTitles.IsTitleCached(TitleType))
+	// @TODO: the bone can be altered in the property editor, so we have to 
+	//        choose to mark this dirty when that happens for this to properly work
+	else //if (!CachedNodeTitles.IsTitleCached(TitleType))
 	{
 		FFormatNamedArguments Args;
 		Args.Add(TEXT("ControllerDescription"), GetControllerDescription());
diff --git a/Engine/Source/Editor/AnimGraph/Private/AnimGraphNode_TwoBoneIK.cpp b/Engine/Source/Editor/AnimGraph/Private/AnimGraphNode_TwoBoneIK.cpp
--- a/Engine/Source/Editor/AnimGraph/Private/AnimGraphNode_TwoBoneIK.cpp
+++ b/Engine/Source/Editor/AnimGraph/Private/AnimGraphNode_TwoBoneIK.cpp
@@ -52,7 +52,9 @@ FText UAnimGraphNode_TwoBoneIK::GetNodeTitle(ENodeTitleType::Type TitleType) con
 	{
 		return GetControllerDescription();
 	}
-	else if (!CachedNodeTitles.IsTitleCached(TitleType))
+	// @TODO: the bone can be altered in the property editor, so we have to 
+	//        choose to mark this dirty when that happens for this to properly work
+	else //if (!CachedNodeTitles.IsTitleCached(TitleType))
 	{
 		FFormatNamedArguments Args;
 		Args.Add(TEXT("ControllerDescription"), GetControllerDescription());
diff --git a/Engine/Source/Editor/AnimGraph/Private/AnimGraphNode_UseCachedPose.cpp b/Engine/Source/Editor/AnimGraph/Private/AnimGraphNode_UseCachedPose.cpp
--- a/Engine/Source/Editor/AnimGraph/Private/AnimGraphNode_UseCachedPose.cpp
+++ b/Engine/Source/Editor/AnimGraph/Private/AnimGraphNode_UseCachedPose.cpp
@@ -25,7 +25,10 @@ FText UAnimGraphNode_UseCachedPose::GetTooltipText() const
 
 FText UAnimGraphNode_UseCachedPose::GetNodeTitle(ENodeTitleType::Type TitleType) const
 {
-	if (CachedNodeTitle.IsOutOfDate())
+	// @TODO: don't know enough about this node type to comfortably assert that
+	//        the NameOfCache won't change after the node has spawned... until
+	//        then, we'll leave this optimization off
+	//if (CachedNodeTitle.IsOutOfDate())
 	{
 		FFormatNamedArguments Args;
 		Args.Add(TEXT("CachePoseName"), FText::FromString(NameOfCache));
diff --git a/Engine/Source/Editor/BlueprintGraph/Private/K2Node_SpawnActor.cpp b/Engine/Source/Editor/BlueprintGraph/Private/K2Node_SpawnActor.cpp
--- a/Engine/Source/Editor/BlueprintGraph/Private/K2Node_SpawnActor.cpp
+++ b/Engine/Source/Editor/BlueprintGraph/Private/K2Node_SpawnActor.cpp
@@ -142,6 +142,8 @@ void UK2Node_SpawnActor::PinDefaultValueChanged(UEdGraphPin* ChangedPin)
 			}
 		}
 
+		CachedNodeTitle.MarkDirty();
+
 		UClass* UseSpawnClass = GetClassToSpawn();
 		if(UseSpawnClass != NULL)
 		{
diff --git a/Engine/Source/Editor/BlueprintGraph/Private/K2Node_SpawnActorFromClass.cpp b/Engine/Source/Editor/BlueprintGraph/Private/K2Node_SpawnActorFromClass.cpp
--- a/Engine/Source/Editor/BlueprintGraph/Private/K2Node_SpawnActorFromClass.cpp
+++ b/Engine/Source/Editor/BlueprintGraph/Private/K2Node_SpawnActorFromClass.cpp
@@ -173,6 +173,8 @@ void UK2Node_SpawnActorFromClass::PinDefaultValueChanged(UEdGraphPin* ChangedPin
 			}
 		}
 
+		CachedNodeTitle.MarkDirty();
+
 		UClass* UseSpawnClass = GetClassToSpawn();
 		if (UseSpawnClass != NULL)
 		{
diff --git a/Engine/Source/Editor/BlueprintGraph/Private/K2Node_VariableGet.cpp b/Engine/Source/Editor/BlueprintGraph/Private/K2Node_VariableGet.cpp
--- a/Engine/Source/Editor/BlueprintGraph/Private/K2Node_VariableGet.cpp
+++ b/Engine/Source/Editor/BlueprintGraph/Private/K2Node_VariableGet.cpp
@@ -171,7 +171,10 @@ FText UK2Node_VariableGet::GetBlueprintVarTooltip(FBPVariableDescription const&
 
 FText UK2Node_VariableGet::GetTooltipText() const
 {
-	if (!CachedTooltip.IsOutOfDate())
+	// @TODO: The variable name mutates as the user makes changes to the 
+	//        underlying property, so until we can catch all those cases, we're
+	//        going to leave this optimization off
+	//if (CachedTooltip.IsOutOfDate())
 	{
 		if (UProperty* Property = GetPropertyForVariable())
 		{
@@ -209,7 +212,10 @@ FText UK2Node_VariableGet::GetNodeTitle(ENodeTitleType::Type TitleType) const
 	{
 		return LOCTEXT("Get", "Get");
 	}
-	else if (CachedNodeTitle.IsOutOfDate())
+	// @TODO: The variable name mutates as the user makes changes to the 
+	//        underlying property, so until we can catch all those cases, we're
+	//        going to leave this optimization off
+	else //if (CachedNodeTitle.IsOutOfDate())
 	{
 		FFormatNamedArguments Args;
 		Args.Add(TEXT("PinName"), FText::FromString(OutputPinName));
diff --git a/Engine/Source/Editor/BlueprintGraph/Private/K2Node_VariableSet.cpp b/Engine/Source/Editor/BlueprintGraph/Private/K2Node_VariableSet.cpp
--- a/Engine/Source/Editor/BlueprintGraph/Private/K2Node_VariableSet.cpp
+++ b/Engine/Source/Editor/BlueprintGraph/Private/K2Node_VariableSet.cpp
@@ -202,7 +202,10 @@ FText UK2Node_VariableSet::GetBlueprintVarTooltip(FBPVariableDescription const&
 
 FText UK2Node_VariableSet::GetTooltipText() const
 {
-	//if (!CachedTooltip.IsOutOfDate())
+	// @TODO: The variable name mutates as the user makes changes to the 
+	//        underlying property, so until we can catch all those cases, we're
+	//        going to leave this optimization off
+	//if (CachedTooltip.IsOutOfDate())
 	{
 		if (UProperty* Property = GetPropertyForVariable())
 		{
@@ -279,7 +282,10 @@ FText UK2Node_VariableSet::GetNodeTitle(ENodeTitleType::Type TitleType) const
 	{
 		return HasLocalRepNotify() ? NSLOCTEXT("K2Node", "SetWithNotify", "Set with Notify") : NSLOCTEXT("K2Node", "Set", "Set");
 	}
-	else if (CachedNodeTitle.IsOutOfDate())
+	// @TODO: The variable name mutates as the user makes changes to the 
+	//        underlying property, so until we can catch all those cases, we're
+	//        going to leave this optimization off
+	else //if (CachedNodeTitle.IsOutOfDate())
 	{
 		FFormatNamedArguments Args;
 		Args.Add(TEXT("PinName"), FText::FromString(InputPinName));
diff --git a/Engine/Source/Editor/BlueprintGraph/Private/K2Node_VariableSetRef.cpp b/Engine/Source/Editor/BlueprintGraph/Private/K2Node_VariableSetRef.cpp
--- a/Engine/Source/Editor/BlueprintGraph/Private/K2Node_VariableSetRef.cpp
+++ b/Engine/Source/Editor/BlueprintGraph/Private/K2Node_VariableSetRef.cpp
@@ -120,6 +120,7 @@ void UK2Node_VariableSetRef::ReallocatePinsDuringReconstruction(TArray<UEdGraphP
   		UEdGraphPin* NewTargetPin = GetTargetPin();
   		CoerceTypeFromPin(OldTargetPin);
   	}
+	CachedNodeTitle.MarkDirty();
 }
 
 FText UK2Node_VariableSetRef::GetTooltipText() const
