Fixing crash when dragging off delegate pins (in the new BP menu system).

diff --git a/Engine/Source/Editor/BlueprintGraph/Private/BlueprintActionFilter.cpp b/Engine/Source/Editor/BlueprintGraph/Private/BlueprintActionFilter.cpp
--- a/Engine/Source/Editor/BlueprintGraph/Private/BlueprintActionFilter.cpp
+++ b/Engine/Source/Editor/BlueprintGraph/Private/BlueprintActionFilter.cpp
@@ -1151,6 +1151,13 @@ static bool BlueprintActionFilterImpl::IsNotSubClassCast(FBlueprintActionFilter
 				checkSlow(CastClass != nullptr);
 
 				UClass const* ContextPinClass = Cast<UClass>(ContextPin->PinType.PinSubCategoryObject.Get());
+				// could be a delegate (or some other) pin, with a non-class
+				// PinSubCategoryObject
+				if (ContextPinClass == nullptr)
+				{
+					continue;
+				}
+
 				if ((ContextPinClass == CastClass) || !IsClassOfType(CastClass, ContextPinClass))
 				{
 					bIsFilteredOut = true;
