AddComponent nodes can now be bound to selected assets (in the new menu system).

diff --git a/Engine/Source/Editor/BlueprintGraph/Classes/BlueprintComponentNodeSpawner.h b/Engine/Source/Editor/BlueprintGraph/Classes/BlueprintComponentNodeSpawner.h
--- a/Engine/Source/Editor/BlueprintGraph/Classes/BlueprintComponentNodeSpawner.h
+++ b/Engine/Source/Editor/BlueprintGraph/Classes/BlueprintComponentNodeSpawner.h
@@ -35,10 +35,17 @@ public:
 	virtual FBlueprintNodeSignature GetSpawnerSignature() const override;
 	virtual UEdGraphNode* Invoke(UEdGraph* ParentGraph, FBindingSet const& Bindings, FVector2D const Location) const override;
 	virtual FText GetDefaultMenuName(FBindingSet const& Bindings) const override;
+	virtual FText GetDefaultMenuTooltip() const;
 	virtual FText GetDefaultMenuCategory() const override;
+	virtual FString GetDefaultSearchKeywords() const override;
 	virtual FName GetDefaultMenuIcon(FLinearColor& ColorOut) const override;
 	// End UBlueprintNodeSpawner interface
 	
+	// IBlueprintNodeBinder interface
+	virtual bool IsBindingCompatible(UObject const* BindingCandidate) const override;
+	virtual bool BindToNode(UEdGraphNode* Node, UObject* Binding) const override;
+	// End IBlueprintNodeBinder interface
+
 	/**
 	 * Retrieves the component class that this configures spawned nodes with.
 	 *
diff --git a/Engine/Source/Editor/BlueprintGraph/Private/BlueprintComponentNodeSpawner.cpp b/Engine/Source/Editor/BlueprintGraph/Private/BlueprintComponentNodeSpawner.cpp
--- a/Engine/Source/Editor/BlueprintGraph/Private/BlueprintComponentNodeSpawner.cpp
+++ b/Engine/Source/Editor/BlueprintGraph/Private/BlueprintComponentNodeSpawner.cpp
@@ -5,6 +5,7 @@
 #include "K2Node_AddComponent.h"
 #include "ClassIconFinder.h" // for FindIconNameForClass()
 #include "BlueprintNodeTemplateCache.h" // for IsTemplateOuter()
+#include "ComponentAssetBroker.h" // for GetComponentsForAsset()/AssignAssetToComponent()
 
 #define LOCTEXT_NAMESPACE "BlueprintComponenetNodeSpawner"
 
@@ -51,18 +52,26 @@ UEdGraphNode* UBlueprintComponentNodeSpawner::Invoke(UEdGraph* ParentGraph, FBin
 	check(ComponentClass != nullptr);
 	
 	auto PostSpawnLambda = [](UEdGraphNode* NewNode, bool bIsTemplateNode, FCustomizeNodeDelegate UserDelegate)
-	{
-		UserDelegate.ExecuteIfBound(NewNode, bIsTemplateNode);
-		
+	{		
 		UK2Node_AddComponent* AddCompNode = CastChecked<UK2Node_AddComponent>(NewNode);
 		UBlueprint* Blueprint = AddCompNode->GetBlueprint();
 		
 		UFunction* AddComponentFunc = FindFieldChecked<UFunction>(AActor::StaticClass(), UK2Node_AddComponent::GetAddComponentFunctionName());
 		AddCompNode->FunctionReference.SetFromField<UFunction>(AddComponentFunc, FBlueprintEditorUtils::IsActorBased(Blueprint));
+
+		UserDelegate.ExecuteIfBound(NewNode, bIsTemplateNode);
 	};
+
 	FCustomizeNodeDelegate PostSpawnDelegate = FCustomizeNodeDelegate::CreateStatic(PostSpawnLambda, CustomizeNodeDelegate);
-	
-	UK2Node_AddComponent* NewNode = Super::SpawnNode<UK2Node_AddComponent>(NodeClass, ParentGraph, Bindings, Location, PostSpawnDelegate);
+	// let SpawnNode() allocate default pins (so we can modify them)
+	UK2Node_AddComponent* NewNode = Super::SpawnNode<UK2Node_AddComponent>(NodeClass, ParentGraph, FBindingSet(), Location, PostSpawnDelegate);
+
+	// set the return type to be the type of the template
+	UEdGraphPin* ReturnPin = NewNode->GetReturnValuePin();
+	if (ReturnPin != nullptr)
+	{
+		ReturnPin->PinType.PinSubCategoryObject = *ComponentClass;
+	}
 
 	bool const bIsTemplateNode = FBlueprintNodeTemplateCache::IsTemplateOuter(ParentGraph);
 	if (!bIsTemplateNode)
@@ -79,23 +88,11 @@ UEdGraphNode* UBlueprintComponentNodeSpawner::Invoke(UEdGraph* ParentGraph, FBin
 		{
 			TemplateNamePin->DefaultValue = ComponentTemplate->GetName();
 		}
-
-		// set the return type to be the type of the template
-		UEdGraphPin* ReturnPin = NewNode->GetReturnValuePin();
-		if (ReturnPin != nullptr)
-		{
-			ReturnPin->PinType.PinSubCategoryObject = *ComponentClass;
-		}
-
-		// @TODO:
-// 		if (ComponentAsset != NULL)
-// 		{
-// 			FComponentAssetBrokerage::AssignAssetToComponent(ComponentTemplate, ComponentAsset);
-// 		}
-
-		NewNode->ReconstructNode();
 	}
 
+	// apply bindings, after we've setup the template pin
+	ApplyBindings(NewNode, Bindings);
+
 	return NewNode;
 }
 
@@ -108,7 +105,29 @@ FText UBlueprintComponentNodeSpawner::GetDefaultMenuName(FBindingSet const& Bind
 		// FText::Format() is slow, so we cache this to save on performance
 		CachedMenuName = FText::Format(LOCTEXT("AddComponentMenuName", "Add {0}"), FText::FromName(ComponentClass->GetFName()));
 	}
-	return CachedMenuName;
+	FText MenuName = CachedMenuName.GetCachedText();
+
+	if (Bindings.Num() > 0)
+	{
+		FText AssetName;
+		if (UObject* AssetBinding = Bindings.CreateConstIterator()->Get())
+		{
+			AssetName = FText::FromName(AssetBinding->GetFName());
+		}
+
+		FText const ComponentTypeName = FText::FromName(ComponentClass->GetFName());
+		MenuName = FText::Format(LOCTEXT("AddBoundComponentMenuName", "Add {0} (as {1})"), AssetName, ComponentTypeName);
+	}
+	return MenuName;
+}
+
+//------------------------------------------------------------------------------
+FText UBlueprintComponentNodeSpawner::GetDefaultMenuTooltip() const
+{
+	FText const ComponentTypeName = FText::FromName(ComponentClass->GetFName());
+	// @TODO: consider caching this, and provide a separate tooltip for when the
+	//        node would be bound ("Spawn a {0} using {1}").
+	return FText::Format(LOCTEXT("AddComponentTooltip", "Spawn a {0}"), ComponentTypeName);
 }
 
 //------------------------------------------------------------------------------
@@ -138,6 +157,15 @@ FText UBlueprintComponentNodeSpawner::GetDefaultMenuCategory() const
 	return CachedCategory;
 }
 
+//------------------------------------------------------------------------------
+FString UBlueprintComponentNodeSpawner::GetDefaultSearchKeywords() const
+{
+	FString SearchKeywords = ComponentClass->GetMetaData(FBlueprintMetadata::MD_FunctionKeywords);
+	// add at least one character, so that the menu item doesn't attempt to
+	// ping a template node
+	return SearchKeywords.AppendChar(TEXT(' '));
+}
+
 //------------------------------------------------------------------------------
 FName UBlueprintComponentNodeSpawner::GetDefaultMenuIcon(FLinearColor& ColorOut) const
 {
@@ -145,6 +173,33 @@ FName UBlueprintComponentNodeSpawner::GetDefaultMenuIcon(FLinearColor& ColorOut)
 	return FClassIconFinder::FindIconNameForClass(ComponentClass);
 }
 
+//------------------------------------------------------------------------------
+bool UBlueprintComponentNodeSpawner::IsBindingCompatible(UObject const* BindingCandidate) const
+{
+	bool bCanBindWith = false;
+	if (BindingCandidate->IsAsset())
+	{
+		TArray< TSubclassOf<UActorComponent> > ComponentClasses = FComponentAssetBrokerage::GetComponentsForAsset(BindingCandidate);
+		bCanBindWith = ComponentClasses.Contains(ComponentClass);
+	}
+	return bCanBindWith;
+}
+
+//------------------------------------------------------------------------------
+bool UBlueprintComponentNodeSpawner::BindToNode(UEdGraphNode* Node, UObject* Binding) const
+{
+	bool bSuccessfulBinding = false;
+	UK2Node_AddComponent* AddCompNode = CastChecked<UK2Node_AddComponent>(Node);
+
+	UActorComponent* ComponentTemplate = AddCompNode->GetTemplateFromNode();
+	if (ComponentTemplate != nullptr)
+	{
+		bSuccessfulBinding = FComponentAssetBrokerage::AssignAssetToComponent(ComponentTemplate, Binding);
+		AddCompNode->ReconstructNode();
+	}
+	return bSuccessfulBinding;
+}
+
 //------------------------------------------------------------------------------
 TSubclassOf<UActorComponent> UBlueprintComponentNodeSpawner::GetComponentClass() const
 {
diff --git a/Engine/Source/Editor/GraphEditor/Private/GraphActionNode.cpp b/Engine/Source/Editor/GraphEditor/Private/GraphActionNode.cpp
--- a/Engine/Source/Editor/GraphEditor/Private/GraphActionNode.cpp
+++ b/Engine/Source/Editor/GraphEditor/Private/GraphActionNode.cpp
@@ -472,8 +472,8 @@ TSharedPtr<FGraphActionNode> FGraphActionNode::FindMatchingParent(FString const&
 	{
 		if (ChildNode->IsCategoryNode())
 		{
-			// @TODO: should we be matching grouping as well?
 			if ((NodeToAdd->SectionID == ChildNode->SectionID) &&
+				(NodeToAdd->Grouping  == ChildNode->Grouping)  &&
 				(ParentName == ChildNode->DisplayText.ToString()))
 			{
 				FoundCategoryNode = ChildNode;
diff --git a/Engine/Source/Editor/Kismet/Private/BlueprintActionMenuUtils.cpp b/Engine/Source/Editor/Kismet/Private/BlueprintActionMenuUtils.cpp
--- a/Engine/Source/Editor/Kismet/Private/BlueprintActionMenuUtils.cpp
+++ b/Engine/Source/Editor/Kismet/Private/BlueprintActionMenuUtils.cpp
@@ -25,9 +25,10 @@
 
 namespace BlueprintActionMenuUtilsImpl
 {
-	static int32 const MainMenuSectionGroup   = 000;
-	static int32 const ComponentsSectionGroup = 100;
 	static int32 const LevelActorSectionGroup = 101;
+	static int32 const ComponentsSectionGroup = 100;
+	static int32 const BoundAddComponentGroup = 002;
+	static int32 const MainMenuSectionGroup   = 000;
 
 	/**
 	 * Additional filter rejection test, for menu sections that only contain 
@@ -49,8 +50,22 @@ namespace BlueprintActionMenuUtilsImpl
 	 */
 	static bool IsNonFavoritedAction(FBlueprintActionFilter const& Filter, FBlueprintActionInfo& BlueprintAction);
 
+	/**
+	 * 
+	 * 
+	 * @param  Filter	
+	 * @param  BlueprintAction	
+	 * @return 
+	 */
 	static bool IsPureNonConstAction(FBlueprintActionFilter const& Filter, FBlueprintActionInfo& BlueprintAction);
 
+	/**
+	 * 
+	 * 
+	 * @param  Filter	
+	 * @param  BlueprintAction	
+	 * @return 
+	 */
 	static bool IsUnexposedMemberAction(FBlueprintActionFilter const& Filter, FBlueprintActionInfo& BlueprintAction);
 
 	/**
@@ -406,6 +421,20 @@ void FBlueprintActionMenuUtils::MakeContextMenu(FBlueprintActionContext const& C
 		}
 	}
 
+	FBlueprintActionFilter AddComponentFilter;
+	AddComponentFilter.Context = MainMenuFilter.Context;
+	AddComponentFilter.PermittedNodeTypes.Add(UK2Node_AddComponent::StaticClass());
+	AddComponentFilter.AddRejectionTest(FBlueprintActionFilter::FRejectionTestDelegate::CreateStatic(IsUnBoundSpawner));
+
+	for (FSelectionIterator SelectionIt(*GEditor->GetSelectedObjects()); SelectionIt; ++SelectionIt)
+	{
+		UObject* PerspectiveAsset = *SelectionIt;
+		if (PerspectiveAsset->IsAsset())
+		{
+			AddComponentFilter.Context.SelectedObjects.Add(PerspectiveAsset);
+		}
+	}
+
 	//--------------------------------------
 	// Defining Menu Sections
 	//--------------------------------------	
@@ -418,7 +447,7 @@ void FBlueprintActionMenuUtils::MakeContextMenu(FBlueprintActionContext const& C
 	}
 	// for legacy purposes, we have to add the main menu section first (when 
 	// reconstructing the legacy menu, we pull the first menu system)
-	MenuOut.AddMenuSection(MainMenuFilter);
+	MenuOut.AddMenuSection(MainMenuFilter, FText::GetEmpty(), MainMenuSectionGroup);
 
 	bool const bAddComponentsSection = bIsContextSensitive && bCanHaveActorComponents && (ComponentsFilter.Context.SelectedObjects.Num() > 0);
 	// add the components section to the menu (if we don't have any components
@@ -437,7 +466,8 @@ void FBlueprintActionMenuUtils::MakeContextMenu(FBlueprintActionContext const& C
 
 	if (bIsContextSensitive)
 	{
-		MenuOut.AddMenuSection(CallOnMemberFilter);
+		MenuOut.AddMenuSection(CallOnMemberFilter, FText::GetEmpty(), MainMenuSectionGroup);
+		MenuOut.AddMenuSection(AddComponentFilter, FText::GetEmpty(), BoundAddComponentGroup);
 	}	
 
 	//--------------------------------------
diff --git a/Engine/Source/Editor/UnrealEd/Private/ComponentAssetBroker.cpp b/Engine/Source/Editor/UnrealEd/Private/ComponentAssetBroker.cpp
--- a/Engine/Source/Editor/UnrealEd/Private/ComponentAssetBroker.cpp
+++ b/Engine/Source/Editor/UnrealEd/Private/ComponentAssetBroker.cpp
@@ -243,7 +243,7 @@ bool FComponentAssetBrokerage::bShutSystemDown = false;
 // FComponentAssetBrokerageage
 
 /** Find set of components that support this asset */
-FComponentClassList FComponentAssetBrokerage::GetComponentsForAsset(UObject* InAsset)
+FComponentClassList FComponentAssetBrokerage::GetComponentsForAsset(const UObject* InAsset)
 {
 	InitializeMap();
 	FComponentClassList OutClasses;
diff --git a/Engine/Source/Editor/UnrealEd/Public/ComponentAssetBroker.h b/Engine/Source/Editor/UnrealEd/Public/ComponentAssetBroker.h
--- a/Engine/Source/Editor/UnrealEd/Public/ComponentAssetBroker.h
+++ b/Engine/Source/Editor/UnrealEd/Public/ComponentAssetBroker.h
@@ -42,7 +42,7 @@ class UNREALED_API FComponentAssetBrokerage
 {
 public:
 	/** Find set of components that support this asset */
-	static TArray< TSubclassOf<UActorComponent> > GetComponentsForAsset(UObject* InAsset);
+	static TArray< TSubclassOf<UActorComponent> > GetComponentsForAsset(const UObject* InAsset);
 
 	/** Get the primary component for this asset*/
 	static TSubclassOf<UActorComponent> GetPrimaryComponentForAsset(UClass* InAssetClass);
