Giving bound event actions an icon for the new blueprint menu system.

diff --git a/Engine/Source/Editor/BlueprintGraph/Classes/BlueprintBoundEventNodeSpawner.h b/Engine/Source/Editor/BlueprintGraph/Classes/BlueprintBoundEventNodeSpawner.h
--- a/Engine/Source/Editor/BlueprintGraph/Classes/BlueprintBoundEventNodeSpawner.h
+++ b/Engine/Source/Editor/BlueprintGraph/Classes/BlueprintBoundEventNodeSpawner.h
@@ -36,6 +36,7 @@ public:
 	virtual UEdGraphNode* Invoke(UEdGraph* ParentGraph, FBindingSet const& Bindings, FVector2D const Location) const override;
 	virtual FText GetDefaultMenuName(FBindingSet const& Bindings) const override;
 	virtual FText GetDefaultMenuCategory() const override;
+	virtual FName GetDefaultMenuIcon(FLinearColor& ColorOut) const;
 	// End UBlueprintNodeSpawner interface
 
 	// UBlueprintEventNodeSpawner interface
diff --git a/Engine/Source/Editor/BlueprintGraph/Classes/BlueprintEventNodeSpawner.h b/Engine/Source/Editor/BlueprintGraph/Classes/BlueprintEventNodeSpawner.h
--- a/Engine/Source/Editor/BlueprintGraph/Classes/BlueprintEventNodeSpawner.h
+++ b/Engine/Source/Editor/BlueprintGraph/Classes/BlueprintEventNodeSpawner.h
@@ -50,6 +50,7 @@ public:
 	virtual FText GetDefaultMenuName(FBindingSet const& Bindings) const override;
 	virtual FText GetDefaultMenuCategory() const override;
 	virtual FString GetDefaultSearchKeywords() const override;
+	virtual FName GetDefaultMenuIcon(FLinearColor& ColorOut) const;
 	// End UBlueprintNodeSpawner interface
 
 	/**
diff --git a/Engine/Source/Editor/BlueprintGraph/Private/BlueprintBoundEventNodeSpawner.cpp b/Engine/Source/Editor/BlueprintGraph/Private/BlueprintBoundEventNodeSpawner.cpp
--- a/Engine/Source/Editor/BlueprintGraph/Private/BlueprintBoundEventNodeSpawner.cpp
+++ b/Engine/Source/Editor/BlueprintGraph/Private/BlueprintBoundEventNodeSpawner.cpp
@@ -77,6 +77,13 @@ FText UBlueprintBoundEventNodeSpawner::GetDefaultMenuCategory() const
 	return DelegateCategory;
 }
 
+//------------------------------------------------------------------------------
+FName UBlueprintBoundEventNodeSpawner::GetDefaultMenuIcon(FLinearColor& ColorOut) const
+{
+	ColorOut = FLinearColor::White;
+	return TEXT("GraphEditor.Event_16x");
+}
+
 //------------------------------------------------------------------------------
 UK2Node_Event const* UBlueprintBoundEventNodeSpawner::FindPreExistingEvent(UBlueprint* Blueprint, FBindingSet const& Bindings) const
 {
diff --git a/Engine/Source/Editor/BlueprintGraph/Private/BlueprintEventNodeSpawner.cpp b/Engine/Source/Editor/BlueprintGraph/Private/BlueprintEventNodeSpawner.cpp
--- a/Engine/Source/Editor/BlueprintGraph/Private/BlueprintEventNodeSpawner.cpp
+++ b/Engine/Source/Editor/BlueprintGraph/Private/BlueprintEventNodeSpawner.cpp
@@ -200,6 +200,13 @@ FString UBlueprintEventNodeSpawner::GetDefaultSearchKeywords() const
 	return Keywords;
 }
 
+//------------------------------------------------------------------------------
+FName UBlueprintEventNodeSpawner::GetDefaultMenuIcon(FLinearColor& ColorOut) const
+{
+	ColorOut = FLinearColor::White;
+	return (IsForCustomEvent() && CustomEventName.IsNone()) ? TEXT("GraphEditor.CustomEvent_16x") : TEXT("GraphEditor.Event_16x");
+}
+
 //------------------------------------------------------------------------------
 UFunction const* UBlueprintEventNodeSpawner::GetEventFunction() const
 {
