Guarding against null function outer(when attempting to log an warning during node validation).

diff --git a/Engine/Source/Editor/BlueprintGraph/Private/K2Node_CallFunction.cpp b/Engine/Source/Editor/BlueprintGraph/Private/K2Node_CallFunction.cpp
--- a/Engine/Source/Editor/BlueprintGraph/Private/K2Node_CallFunction.cpp
+++ b/Engine/Source/Editor/BlueprintGraph/Private/K2Node_CallFunction.cpp
@@ -1196,8 +1196,17 @@ void UK2Node_CallFunction::ValidateNodeDuringCompilation(class FCompilerResultsL
 	UFunction *Function = GetTargetFunction();
 	if (Function == NULL)
 	{
-		UClass* FuncOwnerClass = FunctionReference.GetMemberParentClass(GetBlueprint()->GeneratedClass);
-		FString const OwnerName = FuncOwnerClass->GetName();
+		FString OwnerName;
+
+		UBlueprint* Blueprint = GetBlueprint();
+		if (Blueprint != nullptr)
+		{
+			OwnerName = Blueprint->GetName();
+			if (UClass* FuncOwnerClass = FunctionReference.GetMemberParentClass(Blueprint->GeneratedClass))
+			{
+				OwnerName = FuncOwnerClass->GetName();
+			}
+		}
 		FString const FunctName = FunctionReference.GetMemberName().ToString();
 
 		FText const WarningFormat = LOCTEXT("FunctionNotFound", "Could not find a function named \"%s\" in '%s'.\nMake sure '%s' has been compiled for @@");
diff --git a/Engine/Source/Editor/BlueprintGraph/Private/K2Node_Variable.cpp b/Engine/Source/Editor/BlueprintGraph/Private/K2Node_Variable.cpp
--- a/Engine/Source/Editor/BlueprintGraph/Private/K2Node_Variable.cpp
+++ b/Engine/Source/Editor/BlueprintGraph/Private/K2Node_Variable.cpp
@@ -315,8 +315,17 @@ void UK2Node_Variable::ValidateNodeDuringCompilation(class FCompilerResultsLog&
 	{
 		if (!VariableReference.IsDeprecated())
 		{
-			UClass* VarOwnerClass = VariableReference.GetMemberParentClass(GetBlueprint()->GeneratedClass);
-			FString const OwnerName = VarOwnerClass->GetName();
+			FString OwnerName;
+
+			UBlueprint* Blueprint = GetBlueprint();
+			if (Blueprint != nullptr)
+			{
+				OwnerName = Blueprint->GetName();
+				if (UClass* VarOwnerClass = VariableReference.GetMemberParentClass(Blueprint->GeneratedClass))
+				{
+					OwnerName = VarOwnerClass->GetName();
+				}
+			}
 			FString const VarName = VariableReference.GetMemberName().ToString();
 
 			FText const WarningFormat = LOCTEXT("VariableNotFound", "Could not find a variable named \"%s\" in '%s'.\nMake sure '%s' has been compiled for @@");
