Cleaning up some lingering TODO comments that have already been addressed.

diff --git a/Engine/Source/Editor/BlueprintGraph/Private/BlueprintActionDatabase.cpp b/Engine/Source/Editor/BlueprintGraph/Private/BlueprintActionDatabase.cpp
--- a/Engine/Source/Editor/BlueprintGraph/Private/BlueprintActionDatabase.cpp
+++ b/Engine/Source/Editor/BlueprintGraph/Private/BlueprintActionDatabase.cpp
@@ -936,20 +936,6 @@ void FBlueprintActionDatabase::RefreshClassActions(UClass* const Class)
 			ActionRegistry.Remove(Class);
 		}
 	}
-
-	// @TODO: account for all K2ActionMenuBuilder methods...
-	//   GetLiteralsFromActorSelection() - UK2Node_Literal
-	//   GetAnimNotifyMenuItems()
-	//   GetMatineeControllers() - UK2Node_MatineeController
-	//   GetEventDispatcherNodesForClass()
-	//   GetBoundEventsFromActorSelection() - handle with filter
-	//   GetFunctionCallsOnSelectedActors() - handle with filter
-	//   GetAddComponentActionsUsingSelectedAssets() 
-	//   GetFunctionCallsOnSelectedComponents() - handle with filter
-	//   GetBoundEventsFromComponentSelection() - handle with filter 
-	//   GetPinAllowedMacros()
-	//   GetFuncNodesWithPinType()
-	//   GetVariableGettersSettersForClass()
 }
 
 //------------------------------------------------------------------------------
diff --git a/Engine/Source/Editor/BlueprintGraph/Private/BlueprintActionFilter.cpp b/Engine/Source/Editor/BlueprintGraph/Private/BlueprintActionFilter.cpp
--- a/Engine/Source/Editor/BlueprintGraph/Private/BlueprintActionFilter.cpp
+++ b/Engine/Source/Editor/BlueprintGraph/Private/BlueprintActionFilter.cpp
@@ -365,7 +365,8 @@ static bool BlueprintActionFilterImpl::IsImpure(FBlueprintActionInfo& BlueprintA
 
 		// TODO: why are some of these "impure"?... we shouldn't have hardcoded 
 		//       node types here (game modules cannot add their node types here,
-		//       so we should find another way of identifying "pure" node types)
+		//       so we should find another way of identifying "pure" node types
+		//       ... maybe look for an exec pin?)
 		bIsImpure = (NodeClass == UK2Node_IfThenElse::StaticClass()) ||
 			(NodeClass == UK2Node_MultiGate::StaticClass()) ||
 			(NodeClass == UK2Node_MakeArray::StaticClass()) ||
@@ -1400,27 +1401,6 @@ FBlueprintActionFilter::FBlueprintActionFilter(uint32 Flags/*= 0x00*/)
 	AddRejectionTest(FRejectionTestDelegate::CreateStatic(IsNonTargetMember, !(Flags & BPFILTER_RejectGlobalFields)));
 	AddRejectionTest(FRejectionTestDelegate::CreateStatic(IsUnBoundBindingSpawner));
 	AddRejectionTest(FRejectionTestDelegate::CreateStatic(IsOutOfScopeLocalVariable));
-
-	// @TODO: account for all K2ActionMenuBuilder checks...
-	//    NodeCDO->CanCreateUnderSpecifiedSchema(this)
-	//    FBlueprintEditorUtils::DoesSupportTimelines()
-	//    FKismetEditorUtilities::IsClassABlueprintInterface()
-	//    FK2ActionMenuBuilder::GetPaletteActions()
-	//       is allowed in level-script?
-	//       is allowed in non-actor?
-	//       is allowed in actor?
-	//    FK2ActionMenuBuilder::GetAllActionsForClass()
-	//       show variables?
-	//       show protected?
-	//       show delegates?
-	//    FK2ActionMenuBuilder::GetFuncNodesForClass()
-	//       show inherited?
-	//    FK2ActionMenuBuilder::GetContextAllowedNodeTypes()
-	//       impure funcs?
-	//       construction script?
-	//    UEdGraphSchema_K2::CanFunctionBeUsedInClass()
-
-	// @TODO: account for implemented interface methods
 }
 
 //------------------------------------------------------------------------------
diff --git a/Engine/Source/Editor/Kismet/Private/BlueprintActionMenuBuilder.cpp b/Engine/Source/Editor/Kismet/Private/BlueprintActionMenuBuilder.cpp
--- a/Engine/Source/Editor/Kismet/Private/BlueprintActionMenuBuilder.cpp
+++ b/Engine/Source/Editor/Kismet/Private/BlueprintActionMenuBuilder.cpp
@@ -789,23 +789,6 @@ void FBlueprintActionMenuBuilder::RebuildActionList()
 	{
 		AppendLegacyItems(*MenuSections[0], BlueprintEditorPtr, *this);
 	}
-
-	// @TODO: account for all K2ActionMenuBuilder action types...
-	// - FEdGraphSchemaAction_K2AddTimeline
-	// - FEdGraphSchemaAction_K2ViewNode
-	// - FEdGraphSchemaAction_K2AddCustomEvent
-	//   FEdGraphSchemaAction_EventFromFunction
-	// - FEdGraphSchemaAction_K2Var
-	// - FEdGraphSchemaAction_K2Delegate
-	// - FEdGraphSchemaAction_K2AssignDelegate
-	// - FEdGraphSchemaAction_K2AddComment
-	// - FEdGraphSchemaAction_K2PasteHere
-	// - FEdGraphSchemaAction_K2NewNode
-	// - FEdGraphSchemaAction_Dummy
-	//   FEdGraphSchemaAction_K2AddCallOnActor
-	//   FEdGraphSchemaAction_K2AddCallOnVariable
-	// - FEdGraphSchemaAction_K2AddComponent
-	// - FEdGraphSchemaAction_K2AddComment
 }
 
 #undef LOCTEXT_NAMESPACE
