Removing some "not in use yet" comments, in prep for fully enabling the new BP menu system.

diff --git a/Engine/Source/Editor/BlueprintGraph/Classes/K2Node.h b/Engine/Source/Editor/BlueprintGraph/Classes/K2Node.h
--- a/Engine/Source/Editor/BlueprintGraph/Classes/K2Node.h
+++ b/Engine/Source/Editor/BlueprintGraph/Classes/K2Node.h
@@ -303,8 +303,6 @@ class UK2Node : public UEdGraphNode
 	 * extensible way for new nodes, and game module nodes to add themselves to
 	 * context menus.
 	 *
-	 * @TODO: Not operational yet (stubbed in to prep for blueprint action menu refactor).
-	 *
 	 * @param  ActionListOut	The list to be populated with new spawners.
 	 */
 	virtual void GetMenuActions(FBlueprintActionDatabaseRegistrar& ActionRegistrar) const {}
@@ -313,8 +311,6 @@ class UK2Node : public UEdGraphNode
 	 * Override to provide a default category for specific node types to be 
 	 * listed under.
 	 *
-	 * @TODO: Not operational yet (stubbed in to prep for blueprint action menu refactor).
-	 *
 	 * @return A localized category string (or an empty string if you want this node listed at the menu's root).
 	 */
 	virtual FText GetMenuCategory() const { return FText::GetEmpty(); }
