Favoring skeleton property/function signatures for variable/function node pins.

diff --git a/Engine/Source/Editor/BlueprintGraph/Private/K2Node_CallFunction.cpp b/Engine/Source/Editor/BlueprintGraph/Private/K2Node_CallFunction.cpp
--- a/Engine/Source/Editor/BlueprintGraph/Private/K2Node_CallFunction.cpp
+++ b/Engine/Source/Editor/BlueprintGraph/Private/K2Node_CallFunction.cpp
@@ -179,7 +179,26 @@ FText UK2Node_CallFunction::GetNodeTitle(ENodeTitleType::Type TitleType) const
 void UK2Node_CallFunction::AllocateDefaultPins()
 {
 	UBlueprint* MyBlueprint = GetBlueprint();
+	
 	UFunction* Function = GetTargetFunction();
+	// favor the skeleton function if possible (in case the signature has 
+	// changed, and not yet compiled).
+	if (!FunctionReference.IsSelfContext())
+	{
+		UClass* FunctionClass = FunctionReference.GetMemberParentClass(MyBlueprint->GeneratedClass);
+		if (UBlueprintGeneratedClass* BpClassOwner = Cast<UBlueprintGeneratedClass>(FunctionClass))
+		{
+			// this function could currently only be a part of some skeleton 
+			// class (the blueprint has not be compiled with it yet), so let's 
+			// check the skeleton class as well, see if we can pull pin data 
+			// from there...
+			UBlueprint* FunctionBlueprint = CastChecked<UBlueprint>(BpClassOwner->ClassGeneratedBy);
+			if (UFunction* SkelFunction = FindField<UFunction>(FunctionBlueprint->SkeletonGeneratedClass, FunctionReference.GetMemberName()))
+			{
+				Function = SkelFunction;
+			}
+		}
+	}
 
 	// First try remap table
 	if (Function == NULL)
@@ -197,24 +216,6 @@ void UK2Node_CallFunction::AllocateDefaultPins()
 		}
 	}
 
-	if (Function == nullptr)
-	{
-		if (!FunctionReference.IsSelfContext())
-		{
-			UClass* FunctionClass = FunctionReference.GetMemberParentClass(MyBlueprint->GeneratedClass);
-			if (UBlueprintGeneratedClass* BpClassOwner = Cast<UBlueprintGeneratedClass>(FunctionClass))
-			{
-				// this function could currently only be a part of some skeleton 
-				// class (the blueprint has not be compiled with it yet), so let's 
-				// check the skeleton class as well, see if we can pull pin data 
-				// from there...
-				UBlueprint* FunctionBlueprint = CastChecked<UBlueprint>(BpClassOwner->ClassGeneratedBy);
-				Function = FindField<UFunction>(FunctionBlueprint->SkeletonGeneratedClass, FunctionReference.GetMemberName());
-			}
-		}
-	}
-
-
 	if (Function == NULL)
 	{
 		// The function no longer exists in the stored scope
diff --git a/Engine/Source/Editor/BlueprintGraph/Private/K2Node_Variable.cpp b/Engine/Source/Editor/BlueprintGraph/Private/K2Node_Variable.cpp
--- a/Engine/Source/Editor/BlueprintGraph/Private/K2Node_Variable.cpp
+++ b/Engine/Source/Editor/BlueprintGraph/Private/K2Node_Variable.cpp
@@ -60,19 +60,21 @@ bool UK2Node_Variable::CreatePinForVariable(EEdGraphPinDirection Direction)
 	const UEdGraphSchema_K2* K2Schema = GetDefault<UEdGraphSchema_K2>();
 
 	UProperty* VariableProperty = GetPropertyForVariable();
-	if (VariableProperty == nullptr)
+	// favor the skeleton property if possible (in case the property type has 
+	// been changed, and not yet compiled).
+	if (!VariableReference.IsSelfContext())
 	{
-		if (!VariableReference.IsSelfContext())
+		UClass* VariableClass = VariableReference.GetMemberParentClass(GetBlueprint()->GeneratedClass);
+		if (UBlueprintGeneratedClass* BpClassOwner = Cast<UBlueprintGeneratedClass>(VariableClass))
 		{
-			UClass* VariableClass = VariableReference.GetMemberParentClass(GetBlueprint()->GeneratedClass);
-			if (UBlueprintGeneratedClass* BpClassOwner = Cast<UBlueprintGeneratedClass>(VariableClass))
+			// this variable could currently only be a part of some skeleton 
+			// class (the blueprint has not be compiled with it yet), so let's 
+			// check the skeleton class as well, see if we can pull pin data 
+			// from there...
+			UBlueprint* VariableBlueprint = CastChecked<UBlueprint>(BpClassOwner->ClassGeneratedBy);
+			if (UProperty* SkelProperty = FindField<UProperty>(VariableBlueprint->SkeletonGeneratedClass, VariableReference.GetMemberName()))
 			{
-				// this variable could currently only be a part of some skeleton 
-				// class (the blueprint has not be compiled with it yet), so let's 
-				// check the skeleton class as well, see if we can pull pin data 
-				// from there...
-				UBlueprint* VariableBlueprint = CastChecked<UBlueprint>(BpClassOwner->ClassGeneratedBy);
-				VariableProperty = FindField<UProperty>(VariableBlueprint->SkeletonGeneratedClass, VariableReference.GetMemberName());
+				VariableProperty = SkelProperty;
 			}
 		}
 	}
