Fixing up initial sizing of documentation/comment nodes (in the new BP menu system).

diff --git a/Engine/Source/Editor/BlueprintGraph/Private/BlueprintActionDatabase.cpp b/Engine/Source/Editor/BlueprintGraph/Private/BlueprintActionDatabase.cpp
--- a/Engine/Source/Editor/BlueprintGraph/Private/BlueprintActionDatabase.cpp
+++ b/Engine/Source/Editor/BlueprintGraph/Private/BlueprintActionDatabase.cpp
@@ -164,11 +164,13 @@ static UBlueprintNodeSpawner* FBlueprintNodeSpawnerFactory::MakeDocumentationNod
 		UBlueprint* Blueprint = FBlueprintEditorUtils::FindBlueprintForGraph(OuterGraph);
 		check(Blueprint != nullptr);
 
-		float const OldNodePosX = NewNode->NodePosX;
-		float const OldNodePosY = NewNode->NodePosY;
-
+		float const OldNodePosX   = NewNode->NodePosX;
+		float const OldNodePosY   = NewNode->NodePosY;
+		float const OldHalfHeight = NewNode->NodeHeight / 2.f;
+		float const OldHalfWidth  = NewNode->NodeWidth  / 2.f;
+		
 		static const float DocNodePadding = 50.0f;
-		FSlateRect Bounds(OldNodePosX - DocNodePadding, OldNodePosY - DocNodePadding, OldNodePosX + DocNodePadding, OldNodePosY + DocNodePadding);
+		FSlateRect Bounds(OldNodePosX - OldHalfWidth, OldNodePosY - OldHalfHeight, OldNodePosX + OldHalfWidth, OldNodePosY + OldHalfHeight);
 		FKismetEditorUtilities::GetBoundsForSelectedNodes(Blueprint, Bounds, DocNodePadding);
 		DocNode->SetBounds(Bounds);
 	};
