Fixed category for anim notify events, montage branching point events in new blueprint context menu. Submitting a helper function that will be used by the blueprint merge tool.

diff --git a/Engine/Source/Editor/BlueprintGraph/Classes/BlueprintEventNodeSpawner.h b/Engine/Source/Editor/BlueprintGraph/Classes/BlueprintEventNodeSpawner.h
--- a/Engine/Source/Editor/BlueprintGraph/Classes/BlueprintEventNodeSpawner.h
+++ b/Engine/Source/Editor/BlueprintGraph/Classes/BlueprintEventNodeSpawner.h
@@ -33,7 +33,7 @@ public:
 	static UBlueprintEventNodeSpawner* Create(UFunction const* const EventFunc, UObject* Outer = nullptr);
 
 	/**
-	 * Creates a new UBlueprintEventNodeSpawner for custom events. The 
+	 * Creates a new UBlueprintEventNodeSpawner for custom events. The
 	 * CustomEventName can be left blank if the node will pick one itself on
 	 * instantiation.
 	 *
@@ -44,6 +44,18 @@ public:
 	 */
 	static UBlueprintEventNodeSpawner* Create(TSubclassOf<UK2Node_Event> NodeClass, FName CustomEventName, UObject* Outer = nullptr);
 
+	/**
+	  * Creates a new UBlueprintEventNodeSpawner for custom events. The
+	  * CustomEventName can be left blank if the node will pick one itself on
+	  * instantiation.
+	  *
+	  * @param CustomEventName			The name you want assigned to the event.
+	  * @param CustomizeNodeDelegate	Delegate that will be executed once the node spawner has spawned a node
+	  * @param Outer					Optional outer for the new spawner (if left null, the transient package will be used).
+	  * @return A newly allocated instance of this class.
+	  */
+	static UBlueprintEventNodeSpawner* CreateWithDelegate(FName CustomEventName, FCustomizeNodeDelegate CustomizeNodeDelegate, FText CustomCategory = FText(), UObject* Outer = nullptr);
+
 	// UBlueprintNodeSpawner interface
 	virtual FBlueprintNodeSignature GetSpawnerSignature() const override;
 	virtual UEdGraphNode* Invoke(UEdGraph* ParentGraph, FBindingSet const& Bindings, FVector2D const Location) const override;
@@ -80,6 +92,10 @@ protected:
 	FNodeTextCache CachedMenuName;
 
 private:
+	/** Custom category to place nodes in: */
+	UPROPERTY()
+	FText CustomCategory;
+
 	/** The function to configure new nodes with. */
 	UPROPERTY()
     UFunction const* EventFunc;
diff --git a/Engine/Source/Editor/BlueprintGraph/Private/BlueprintActionDatabase.cpp b/Engine/Source/Editor/BlueprintGraph/Private/BlueprintActionDatabase.cpp
--- a/Engine/Source/Editor/BlueprintGraph/Private/BlueprintActionDatabase.cpp
+++ b/Engine/Source/Editor/BlueprintGraph/Private/BlueprintActionDatabase.cpp
@@ -3,6 +3,7 @@
 #include "BlueprintGraphPrivatePCH.h"
 #include "BlueprintActionDatabase.h"
 #include "EdGraphSchema_K2.h"       // for CanUserKismetCallFunction()
+#include "EditorCategoryUtils.h"
 #include "KismetEditorUtilities.h"	// for IsClassABlueprintSkeleton(), IsClassABlueprintInterface(), GetBoundsForSelectedNodes(), etc.
 #include "BlueprintNodeSpawner.h"
 #include "BlueprintFunctionNodeSpawner.h"
@@ -103,6 +104,15 @@ namespace FBlueprintNodeSpawnerFactory
 	 * @return 
 	 */
 	static UBlueprintNodeSpawner* MakeActorBoundEventSpawner(UMulticastDelegateProperty* DelegateProperty);
+
+	/**
+	 * Constructs UK2Node_Event spawner that is owned by UAnimInstance. Used for Anim Notificatios and montage 
+	 * branching points.
+	 *
+	 * @param EventName
+	 * @return A new node-spawner, set up to spawn UK2Node_Event
+	 */
+	static UBlueprintNodeSpawner* MakeAnimOwnedEventSpawner(FName SignatureName, FText CustomCategory);
 };
 
 //------------------------------------------------------------------------------
@@ -200,6 +210,18 @@ static UBlueprintNodeSpawner* FBlueprintNodeSpawnerFactory::MakeActorBoundEventS
 	return UBlueprintBoundEventNodeSpawner::Create(UK2Node_ActorBoundEvent::StaticClass(), DelegateProperty);
 }
 
+//------------------------------------------------------------------------------
+static UBlueprintNodeSpawner* FBlueprintNodeSpawnerFactory::MakeAnimOwnedEventSpawner(FName SignatureName, FText CustomCategory)
+{
+	auto PostSpawnSetupLambda = [](UEdGraphNode* NewNode, bool /*bIsTemplateNode*/)
+	{
+		UK2Node_Event* ActorRefNode = CastChecked<UK2Node_Event>(NewNode);
+		ActorRefNode->EventSignatureClass = UAnimInstance::StaticClass();
+	};
+
+	return UBlueprintEventNodeSpawner::CreateWithDelegate(SignatureName, UBlueprintNodeSpawner::FCustomizeNodeDelegate::CreateStatic( PostSpawnSetupLambda ), CustomCategory );
+}
+
 /*******************************************************************************
  * Static FBlueprintActionDatabase Helpers
  ******************************************************************************/
@@ -553,7 +575,32 @@ static void BlueprintActionDatabaseImpl::AddSkeletonActions(const USkeleton& Ske
 		FString Label = NotifyName.ToString();
 
 		FString SignatureName = FString::Printf(TEXT("AnimNotify_%s"), *Label);
-		ActionListOut.Add(UBlueprintEventNodeSpawner::Create(UK2Node_Event::StaticClass(), FName(*SignatureName)));
+		ActionListOut.Add(FBlueprintNodeSpawnerFactory::MakeAnimOwnedEventSpawner(FName(*SignatureName), FEditorCategoryUtils::GetCommonCategory(FCommonEditorCategory::AnimNotify)));
+	}
+
+	// @todo anim: fix this to be same as notifies, save same list in the 
+	// add montage menus
+	// find all montages that uses current target skeleton
+	TArray<FName> BranchingPointHandlers;
+
+	for (FObjectIterator Iter(UAnimMontage::StaticClass()); Iter; ++Iter)
+	{
+		UAnimMontage * Montage = CastChecked<UAnimMontage>(*Iter);
+		if (Montage && Montage->GetSkeleton() == &Skeleton)
+		{
+			// now add event handler if exists
+			for (int32 I = 0; I < Montage->BranchingPoints.Num(); ++I)
+			{
+				BranchingPointHandlers.AddUnique(Montage->BranchingPoints[I].EventName);
+			}
+		}
+	}
+
+	for( auto NotifyName: BranchingPointHandlers )
+	{
+		FString Label = NotifyName.ToString();
+		FString SignatureName = FString::Printf(TEXT("MontageBranchingPoint_%s"), *Label);
+		ActionListOut.Add(FBlueprintNodeSpawnerFactory::MakeAnimOwnedEventSpawner(FName(*SignatureName),  FEditorCategoryUtils::GetCommonCategory(FCommonEditorCategory::BranchPoint)));
 	}
 }
 
@@ -599,9 +646,8 @@ static void BlueprintActionDatabaseImpl::AddAnimBlueprintGraphActions(UAnimBluep
 			if (NotifyName != NAME_None)
 			{
 				FString Label = NotifyName.ToString();
-
 				FString SignatureName = FString::Printf(TEXT("AnimNotify_%s"), *Label);
-				ActionListOut.Add(UBlueprintEventNodeSpawner::Create(UK2Node_Event::StaticClass(), FName(*SignatureName)));
+				ActionListOut.Add(FBlueprintNodeSpawnerFactory::MakeAnimOwnedEventSpawner(FName(*SignatureName), FEditorCategoryUtils::GetCommonCategory(FCommonEditorCategory::AnimNotify)));
 			}
 		}
 	}
diff --git a/Engine/Source/Editor/BlueprintGraph/Private/BlueprintEventNodeSpawner.cpp b/Engine/Source/Editor/BlueprintGraph/Private/BlueprintEventNodeSpawner.cpp
--- a/Engine/Source/Editor/BlueprintGraph/Private/BlueprintEventNodeSpawner.cpp
+++ b/Engine/Source/Editor/BlueprintGraph/Private/BlueprintEventNodeSpawner.cpp
@@ -81,6 +81,22 @@ UBlueprintEventNodeSpawner* UBlueprintEventNodeSpawner::Create(TSubclassOf<UK2No
 	return NodeSpawner;
 }
 
+//------------------------------------------------------------------------------
+UBlueprintEventNodeSpawner* UBlueprintEventNodeSpawner::CreateWithDelegate(FName CustomEventName, FCustomizeNodeDelegate CustomizeNodeDelegate, FText CustomCategory /*= FText()*/, UObject* Outer /*= nullptr*/)
+{
+	if (Outer == nullptr)
+	{
+		Outer = GetTransientPackage();
+	}
+
+	UBlueprintEventNodeSpawner* NodeSpawner = NewObject<UBlueprintEventNodeSpawner>(Outer);
+	NodeSpawner->NodeClass		= UK2Node_Event::StaticClass();
+	NodeSpawner->CustomCategory = CustomCategory;
+	NodeSpawner->CustomEventName = CustomEventName;
+	NodeSpawner->CustomizeNodeDelegate = CustomizeNodeDelegate;
+	return NodeSpawner;
+}
+
 //------------------------------------------------------------------------------
 UBlueprintEventNodeSpawner::UBlueprintEventNodeSpawner(class FPostConstructInitializeProperties const& PCIP)
 	: Super(PCIP)
@@ -186,7 +202,7 @@ FText UBlueprintEventNodeSpawner::GetDefaultMenuCategory() const
 	// certain events can have specialized categories (like input events), so 
 	// we choose to leave this in the node's hands (returns an empty text string
 	// so that the menu builder polls the node instead).
-	return UBlueprintNodeSpawner::GetDefaultMenuCategory();
+	return CustomCategory;
 }
 
 //------------------------------------------------------------------------------
diff --git a/Engine/Source/Editor/UnrealEd/Private/EditorCategoryUtils.cpp b/Engine/Source/Editor/UnrealEd/Private/EditorCategoryUtils.cpp
--- a/Engine/Source/Editor/UnrealEd/Private/EditorCategoryUtils.cpp
+++ b/Engine/Source/Editor/UnrealEd/Private/EditorCategoryUtils.cpp
@@ -84,6 +84,8 @@ FEditorCategoryUtilsImpl::FCategoryInfoMap& FEditorCategoryUtilsImpl::GetCategor
 		RegisterCategoryKey("Delegates", LOCTEXT("DelegatesCategory", "Event Dispatchers"));
 		RegisterCategoryKey("Variables", LOCTEXT("VariablesCategory", "Variables"));
 		RegisterCategoryKey("UserInterface", LOCTEXT("UserInterfaceCategory", "User Interface"));
+		RegisterCategoryKey("AnimNotify", LOCTEXT("AnimNotifyCategory", "Add AnimNotify Event"));
+		RegisterCategoryKey("BranchPoint", LOCTEXT("BranchPointCategory", "Add Montage Branching Point Event"));
 
 		// Utilities sub categories
 		RegisterCategoryKey("FlowControl", BuildCategoryString(FCommonEditorCategory::Utilities, LOCTEXT("FlowControlCategory", "Flow Control")));
@@ -195,6 +197,8 @@ FText const& FEditorCategoryUtils::GetCommonCategory(const FCommonEditorCategory
 		CommonCategoryKeys.Add(FCommonEditorCategory::Delegates, "Delegates");
 		CommonCategoryKeys.Add(FCommonEditorCategory::Variables, "Variables");
 		CommonCategoryKeys.Add(FCommonEditorCategory::UserInterface, "UserInterface");
+		CommonCategoryKeys.Add(FCommonEditorCategory::AnimNotify, "AnimNotify");
+		CommonCategoryKeys.Add(FCommonEditorCategory::BranchPoint, "BranchPoint");
 
 		CommonCategoryKeys.Add(FCommonEditorCategory::FlowControl, "FlowControl");
 		CommonCategoryKeys.Add(FCommonEditorCategory::Transformation, "Transformation");
@@ -393,4 +397,18 @@ void FEditorCategoryUtils::GetCategoryTooltipInfo(const FString& Category, FText
 	}
 }
 
+//------------------------------------------------------------------------------
+TSet<FString> FEditorCategoryUtils::GetHiddenCategories(UClass const* Class)
+{
+	TArray<FString> ClassHiddenCategories;
+	GetClassHideCategories(Class, ClassHiddenCategories);
+	TArray<FString> ClassForceVisibleCategories;
+	GetClassShowCategories(Class, ClassForceVisibleCategories);
+
+	const TSet<FString> ClassHiddenCategoriesSet(ClassHiddenCategories);
+	const TSet<FString> ClassForceVisibleCategoriesSet(ClassForceVisibleCategories);
+
+	return ClassHiddenCategoriesSet.Difference(ClassForceVisibleCategoriesSet);
+}
+
 #undef LOCTEXT_NAMESPACE
diff --git a/Engine/Source/Editor/UnrealEd/Public/EditorCategoryUtils.h b/Engine/Source/Editor/UnrealEd/Public/EditorCategoryUtils.h
--- a/Engine/Source/Editor/UnrealEd/Public/EditorCategoryUtils.h
+++ b/Engine/Source/Editor/UnrealEd/Public/EditorCategoryUtils.h
@@ -34,6 +34,8 @@ struct FCommonEditorCategory
 		Utilities,
 		FlowControl,
 		UserInterface,
+		AnimNotify,
+		BranchPoint,
 
 		// Type library categories:
 		String,
@@ -171,4 +173,15 @@ namespace FEditorCategoryUtils
 	 * @param  DocExcerpt	[OUT] A category path that you want to check.
 	 */
 	UNREALED_API void GetCategoryTooltipInfo(const FString& Category, FText& Tooltip, FString& DocLink, FString& DocExcerpt);
+
+	/**
+	 * Returns the set of categories that should be hidden, categories that are both
+	 * explicitly hidden and explicitly shown will not be included in this list (current
+	 * behavior is that such categories should be shown). This occurs when you show
+	 * a class that your parent has hidden.
+	 * 
+	 * @param  Class		The class you want to query.
+	 * @return  The set of categories that should be hidden.
+	 */
+	UNREALED_API TSet<FString> GetHiddenCategories(UClass const* Class);
 };
\ No newline at end of file
