Fixing up how we filter hidden categories. Ignoring blueprint's class, now only using the context's class.

diff --git a/Engine/Source/Editor/BlueprintGraph/Private/BlueprintActionFilter.cpp b/Engine/Source/Editor/BlueprintGraph/Private/BlueprintActionFilter.cpp
--- a/Engine/Source/Editor/BlueprintGraph/Private/BlueprintActionFilter.cpp
+++ b/Engine/Source/Editor/BlueprintGraph/Private/BlueprintActionFilter.cpp
@@ -532,19 +532,19 @@ static bool BlueprintActionFilterImpl::IsRestrictedClassMember(FBlueprintActionF
 		if (ActionClass->HasMetaData(FBlueprintMetadata::MD_RestrictedToClasses))
 		{
 			FString const& ClassRestrictions = ActionClass->GetMetaData(FBlueprintMetadata::MD_RestrictedToClasses);
-			for (UBlueprint const* Blueprint : FilterContext.Blueprints)
+			for (UClass const* TargetClass : Filter.TargetClasses)
 			{
 				bool bIsClassListed = false;
 				
-				UClass const* BpClass = GetAuthoritativeBlueprintClass(Blueprint);
+				UClass const* QueryClass = TargetClass;
 				// walk the class inheritance chain to see if this class is one
 				// of the allowed
-				while (!bIsClassListed && (BpClass != nullptr))
+				while (!bIsClassListed && (QueryClass != nullptr))
 				{
-					FString const ClassName = BpClass->GetName();
+					FString const ClassName = QueryClass->GetName();
 					bIsClassListed = (ClassName == ClassRestrictions) || !!FCString::StrfindDelim(*ClassRestrictions, *ClassName, TEXT(" "));
 					
-					BpClass = BpClass->GetSuperClass();
+					QueryClass = QueryClass->GetSuperClass();
 				}
 				
 				// if the blueprint class wasn't listed as one of the few
@@ -663,7 +663,6 @@ static bool BlueprintActionFilterImpl::IsNonTargetMember(FBlueprintActionFilter
 static bool BlueprintActionFilterImpl::IsFieldCategoryHidden(FBlueprintActionFilter const& Filter, FBlueprintActionInfo& BlueprintAction)
 {
 	bool bIsFilteredOut = false;
-	FBlueprintActionContext const& FilterContext = Filter.Context;
 
 	DECLARE_DELEGATE_RetVal_OneParam(bool, FIsFieldHiddenDelegate, UClass*);
 	FIsFieldHiddenDelegate IsFieldHiddenDelegate;
@@ -687,20 +686,14 @@ static bool BlueprintActionFilterImpl::IsFieldCategoryHidden(FBlueprintActionFil
 
 	if (IsFieldHiddenDelegate.IsBound())
 	{
-		for (UBlueprint* Blueprint : FilterContext.Blueprints)
+		for (UClass* TargetClass : Filter.TargetClasses)
 		{
-			UClass* BpClass = GetAuthoritativeBlueprintClass(Blueprint);
-			if (IsFieldHiddenDelegate.Execute(BpClass))
+			if (IsFieldHiddenDelegate.Execute(TargetClass))
 			{
 				bIsFilteredOut = true;
 				break;
 			}
 		}
-
-		for (int32 ClassIndex = 0; !bIsFilteredOut && (ClassIndex < Filter.TargetClasses.Num()); ++ClassIndex)
-		{
-			bIsFilteredOut = IsFieldHiddenDelegate.Execute(Filter.TargetClasses[ClassIndex]);
-		}
 	}
 
 	return bIsFilteredOut;
