Blend Enum and SequencePlayer nodes now appear in action list in AnimGraphs.

BlueprintActionDatabase has some handling for actions referencing unloaded assets.

diff --git a/Engine/Source/Editor/AnimGraph/Classes/AnimGraphNode_BlendListByEnum.h b/Engine/Source/Editor/AnimGraph/Classes/AnimGraphNode_BlendListByEnum.h
--- a/Engine/Source/Editor/AnimGraph/Classes/AnimGraphNode_BlendListByEnum.h
+++ b/Engine/Source/Editor/AnimGraph/Classes/AnimGraphNode_BlendListByEnum.h
@@ -35,6 +35,7 @@ public:
 	// UAnimGraphNode_Base interface
 	virtual FString GetNodeCategory() const override;
 	virtual void GetMenuEntries(FGraphContextMenuBuilder& ContextMenuBuilder) const override;
+	virtual void GetMenuActions(FBlueprintActionDatabaseRegistrar& ActionRegistrar) const override;
 	virtual void CustomizePinData(UEdGraphPin* Pin, FName SourcePropertyName, int32 ArrayIndex) const override;
 	virtual void ValidateAnimNodeDuringCompilation(class USkeleton* ForSkeleton, class FCompilerResultsLog& MessageLog) override;
 	virtual void BakeDataDuringCompilation(class FCompilerResultsLog& MessageLog) override;
diff --git a/Engine/Source/Editor/AnimGraph/Classes/AnimGraphNode_SequencePlayer.h b/Engine/Source/Editor/AnimGraph/Classes/AnimGraphNode_SequencePlayer.h
--- a/Engine/Source/Editor/AnimGraph/Classes/AnimGraphNode_SequencePlayer.h
+++ b/Engine/Source/Editor/AnimGraph/Classes/AnimGraphNode_SequencePlayer.h
@@ -22,6 +22,7 @@ class UAnimGraphNode_SequencePlayer : public UAnimGraphNode_Base
 	virtual FLinearColor GetNodeTitleColor() const override;
 	virtual FText GetTooltipText() const override;
 	virtual FText GetNodeTitle(ENodeTitleType::Type TitleType) const override;
+	virtual FText GetMenuCategory() const override;
 	// End of UEdGraphNode interface
 
 	// UAnimGraphNode_Base interface
@@ -35,6 +36,8 @@ class UAnimGraphNode_SequencePlayer : public UAnimGraphNode_Base
 	virtual UScriptStruct* GetTimePropertyStruct() const override;
 	virtual void GetAllAnimationSequencesReferred(TArray<UAnimationAsset*>& ComplexAnims, TArray<UAnimSequence*>& AnimationSequences) const override;
 	virtual void ReplaceReferredAnimations(const TMap<UAnimationAsset*, UAnimationAsset*>& ComplexAnimsMap, const TMap<UAnimSequence*, UAnimSequence*>& AnimSequenceMap) override;
+	virtual void GetMenuActions(FBlueprintActionDatabaseRegistrar& ActionRegistrar) const override;
+	virtual bool IsActionFilteredOut(class FBlueprintActionFilter const& Filter) override;
 	// End of UAnimGraphNode_Base interface
 
 	// UK2Node interface
@@ -46,4 +49,7 @@ private:
 
 	/** Constructing FText strings can be costly, so we cache the node's title */
 	FNodeTitleTextTable CachedNodeTitles;
+
+	/** Used for filtering in the Blueprint context menu when the sequence asset this node uses is unloaded */
+	FString UnloadedSkeletonName;
 };
diff --git a/Engine/Source/Editor/AnimGraph/Private/AnimGraphNode_BlendListByEnum.cpp b/Engine/Source/Editor/AnimGraph/Private/AnimGraphNode_BlendListByEnum.cpp
--- a/Engine/Source/Editor/AnimGraph/Private/AnimGraphNode_BlendListByEnum.cpp
+++ b/Engine/Source/Editor/AnimGraph/Private/AnimGraphNode_BlendListByEnum.cpp
@@ -8,6 +8,8 @@
 #include "K2ActionMenuBuilder.h" // for FK2ActionMenuBuilder::AddNewNodeAction()
 #include "AnimGraphNode_BlendListByEnum.h"
 #include "AnimationGraphSchema.h"
+#include "BlueprintActionDatabaseRegistrar.h"
+#include "BlueprintNodeSpawner.h"
 
 #define LOCTEXT_NAMESPACE "BlendListByEnum"
 
@@ -76,6 +78,41 @@ void UAnimGraphNode_BlendListByEnum::GetMenuEntries(FGraphContextMenuBuilder& Co
 	}
 }
 
+void UAnimGraphNode_BlendListByEnum::GetMenuActions(FBlueprintActionDatabaseRegistrar& ActionRegistrar) const
+{
+	auto CustomizeBlendListEnumNodeLambda = [](UEdGraphNode* NewNode, bool bIsTemplateNode, UEnum* EnumValue)
+	{
+		UAnimGraphNode_BlendListByEnum* BlendListEnumNode = CastChecked<UAnimGraphNode_BlendListByEnum>(NewNode);
+		BlendListEnumNode->BoundEnum = EnumValue;
+	};
+
+	// add all blendlist enum entries
+	for (TObjectIterator<UEnum> EnumIt; EnumIt; ++EnumIt)
+	{
+		UEnum* CurrentEnum = *EnumIt;
+
+		// to keep from needlessly instantiating a UBlueprintNodeSpawner, first   
+		// check to make sure that the registrar is looking for actions of this type
+		// (could be regenerating actions for a specific asset, and therefore the 
+		// registrar would only accept actions corresponding to that asset)
+		if (!ActionRegistrar.IsOpenForRegistration(CurrentEnum))
+		{
+			continue;
+		}
+
+		const bool bIsBlueprintType = UEdGraphSchema_K2::IsAllowableBlueprintVariableType(CurrentEnum);
+		if (bIsBlueprintType)
+		{
+			UBlueprintNodeSpawner* NodeSpawner = UBlueprintNodeSpawner::Create(GetClass());
+			check(NodeSpawner != nullptr);
+
+			NodeSpawner->CustomizeNodeDelegate = UBlueprintNodeSpawner::FCustomizeNodeDelegate::CreateStatic(CustomizeBlendListEnumNodeLambda, CurrentEnum);
+
+			ActionRegistrar.AddBlueprintAction(CurrentEnum, NodeSpawner);
+		}
+	}
+}
+
 void UAnimGraphNode_BlendListByEnum::GetContextMenuActions(const FGraphNodeContextMenuBuilder& Context) const
 {
 	if (!Context.bIsDebugging && (BoundEnum != NULL))
diff --git a/Engine/Source/Editor/AnimGraph/Private/AnimGraphNode_SequencePlayer.cpp b/Engine/Source/Editor/AnimGraph/Private/AnimGraphNode_SequencePlayer.cpp
--- a/Engine/Source/Editor/AnimGraph/Private/AnimGraphNode_SequencePlayer.cpp
+++ b/Engine/Source/Editor/AnimGraph/Private/AnimGraphNode_SequencePlayer.cpp
@@ -7,6 +7,11 @@
 #include "GraphEditorActions.h"
 #include "AssetRegistryModule.h"
 #include "AnimationGraphSchema.h"
+#include "BlueprintActionDatabaseRegistrar.h"
+#include "BlueprintActionFilter.h"
+#include "EditorCategoryUtils.h"
+#include "BlueprintUnloadedAssetNodeSpawner.h"
+
 #define LOCTEXT_NAMESPACE "A3Nodes"
 
 /////////////////////////////////////////////////////
@@ -170,6 +175,121 @@ void UAnimGraphNode_SequencePlayer::GetMenuEntries(FGraphContextMenuBuilder& Con
 	}
 }
 
+FText UAnimGraphNode_SequencePlayer::GetMenuCategory() const
+{
+	return FEditorCategoryUtils::GetCommonCategory(FCommonEditorCategory::Animation);
+}
+
+void UAnimGraphNode_SequencePlayer::GetMenuActions(FBlueprintActionDatabaseRegistrar& ActionRegistrar) const
+{
+	auto CustomizeSequencePlayerNodeLambda = [](UEdGraphNode* NewNode, bool bIsTemplateNode, const FAssetData SequenceAsset)
+	{
+		UAnimGraphNode_SequencePlayer* SequencePlayerNode = CastChecked<UAnimGraphNode_SequencePlayer>(NewNode);
+		if(!bIsTemplateNode || SequenceAsset.IsAssetLoaded())
+		{
+			UAnimSequence* Sequence = Cast<UAnimSequence>(SequenceAsset.GetAsset());
+			check(Sequence);
+
+			SequencePlayerNode->Node.Sequence = Sequence;
+			SequencePlayerNode->UnloadedSkeletonName.Empty();
+		}
+		else
+		{
+			if( const FString* SkeletonTag = SequenceAsset.TagsAndValues.Find(TEXT("Skeleton")) )
+			{
+				SequencePlayerNode->UnloadedSkeletonName = *SkeletonTag;
+			}
+		}
+	};
+
+	// Load the asset registry module
+	FAssetRegistryModule& AssetRegistryModule = FModuleManager::LoadModuleChecked<FAssetRegistryModule>(TEXT("AssetRegistry"));
+
+	FARFilter Filter;
+	Filter.ClassNames.Add(UAnimSequence::StaticClass()->GetFName());
+	Filter.bRecursiveClasses = true;
+
+	// Find matching assets and add an entry for each one
+	TArray<FAssetData> SequenceList;
+	AssetRegistryModule.Get().GetAssets(Filter, /*out*/ SequenceList);
+
+	for (auto AssetIt = SequenceList.CreateConstIterator(); AssetIt; ++AssetIt)
+	{
+		const FAssetData& Asset = *AssetIt;
+
+		UBlueprintNodeSpawner* NodeSpawner = nullptr;
+
+		if(!Asset.IsAssetLoaded())
+		{
+			// to keep from needlessly instantiating a UBlueprintNodeSpawner, first   
+			// check to make sure that the registrar is looking for actions of this type
+			// (could be regenerating actions for a specific asset, and therefore the 
+			// registrar would only accept actions corresponding to that asset)
+			if (!ActionRegistrar.IsOpenForRegistration(GetClass()))
+			{
+				continue;
+			}
+
+			NodeSpawner = UBlueprintUnloadedAssetNodeSpawner::Create(GetClass(), Asset, UAnimGraphNode_SequencePlayer::GetTitleGivenAssetInfo(FText::FromName(Asset.AssetName), false));
+		}
+		else
+		{
+			// to keep from needlessly instantiating a UBlueprintNodeSpawner, first   
+			// check to make sure that the registrar is looking for actions of this type
+			// (could be regenerating actions for a specific asset, and therefore the 
+			// registrar would only accept actions corresponding to that asset)
+			if (!ActionRegistrar.IsOpenForRegistration(Asset.GetAsset()))
+			{
+				continue;
+			}
+
+			NodeSpawner = UBlueprintNodeSpawner::Create(GetClass());
+		}
+
+		check(NodeSpawner != nullptr);
+		NodeSpawner->CustomizeNodeDelegate = UBlueprintUnloadedAssetNodeSpawner::FCustomizeNodeDelegate::CreateStatic(CustomizeSequencePlayerNodeLambda, Asset);
+		ActionRegistrar.AddBlueprintAction(Asset, NodeSpawner);
+	}
+}
+
+bool UAnimGraphNode_SequencePlayer::IsActionFilteredOut(class FBlueprintActionFilter const& Filter)
+{
+	bool bIsFilteredOut = false;
+	FBlueprintActionContext const& FilterContext = Filter.Context;
+
+	for (UBlueprint* Blueprint : FilterContext.Blueprints)
+	{
+		if (UAnimBlueprint* AnimBlueprint = Cast<UAnimBlueprint>(Blueprint))
+		{
+			if(Node.Sequence)
+			{
+				if(Node.Sequence->GetSkeleton() != AnimBlueprint->TargetSkeleton)
+				{
+					// Sequence does not use the same skeleton as the Blueprint, cannot use
+					bIsFilteredOut = true;
+					break;
+				}
+			}
+			else 
+			{
+				FAssetData SkeletonData(AnimBlueprint->TargetSkeleton);
+				if(UnloadedSkeletonName != SkeletonData.GetExportTextName())
+				{
+					bIsFilteredOut = true;
+					break;
+				}
+			}
+		}
+		else
+		{
+			// Not an animation Blueprint, cannot use
+			bIsFilteredOut = true;
+			break;
+		}
+	}
+	return bIsFilteredOut;
+}
+
 void UAnimGraphNode_SequencePlayer::ValidateAnimNodeDuringCompilation(class USkeleton* ForSkeleton, class FCompilerResultsLog& MessageLog)
 {
 	if (Node.Sequence == NULL)
diff --git a/Engine/Source/Editor/BlueprintGraph/Classes/BlueprintUnloadedAssetNodeSpawner.h b/Engine/Source/Editor/BlueprintGraph/Classes/BlueprintUnloadedAssetNodeSpawner.h
new file mode 100644
--- /dev/null
+++ b/Engine/Source/Editor/BlueprintGraph/Classes/BlueprintUnloadedAssetNodeSpawner.h
@@ -0,0 +1,32 @@
+// Copyright 1998-2014 Epic Games, Inc. All Rights Reserved.
+
+#pragma once
+
+#include "BlueprintNodeSpawner.h"
+#include "BlueprintUnloadedAssetNodeSpawner.generated.h"
+
+/**
+ * Handles spawning nodes that use unloaded assets
+ */
+UCLASS(Transient)
+class BLUEPRINTGRAPH_API UBlueprintUnloadedAssetNodeSpawner : public UBlueprintNodeSpawner
+{
+	GENERATED_UCLASS_BODY()
+
+public:
+	/**
+	 * @return A newly allocated instance of this class.
+	 */
+	static UBlueprintUnloadedAssetNodeSpawner* Create(TSubclassOf<UEdGraphNode> NodeClass, const FAssetData& AssetData, FText ActionMenuDescription, UObject* Outer = nullptr);
+
+	// UBlueprintNodeSpawner interface
+	virtual FText GetDefaultMenuName(FBindingSet const& Bindings) const override;
+	// End UBlueprintNodeSpawner interface
+
+public:
+	/** The data from the Asset Registry detailing the unloaded asset */
+	FAssetData AssetData;
+
+	/** The asset's action menu description */
+	FText ActionMenuDescription;
+};
diff --git a/Engine/Source/Editor/BlueprintGraph/Private/BlueprintActionDatabase.cpp b/Engine/Source/Editor/BlueprintGraph/Private/BlueprintActionDatabase.cpp
--- a/Engine/Source/Editor/BlueprintGraph/Private/BlueprintActionDatabase.cpp
+++ b/Engine/Source/Editor/BlueprintGraph/Private/BlueprintActionDatabase.cpp
@@ -375,6 +375,14 @@ namespace BlueprintActionDatabaseImpl
 	 */
 	static void OnAssetRemoved(FAssetData const& AssetInfo);
 
+	/**
+	 * Callback to refresh the database when an object has been renamed (to clear 
+	 * any related classes that were stored in the database under the old name) 
+	 * 
+	 * @param  AssetInfo	Data regarding the freshly removed asset.
+	 */
+	static void OnAssetRenamed(FAssetData const& AssetInfo, const FString& InOldName);
+
 	/**
 	 * Callback to clear all levels from the database when a world is destroyed 
 	 * 
@@ -750,6 +758,21 @@ static void BlueprintActionDatabaseImpl::OnAssetRemoved(FAssetData const& AssetI
 		// the delete went through, so we don't need to track these for re-add
 		PendingDelete.Remove(AssetObject);
 	}
+	else
+	{
+		ActionDatabase.ClearUnloadedAssetActions(AssetInfo.ObjectPath);
+	}
+}
+
+//------------------------------------------------------------------------------
+static void BlueprintActionDatabaseImpl::OnAssetRenamed(FAssetData const& AssetInfo, const FString& InOldName)
+{
+	FBlueprintActionDatabase& ActionDatabase = FBlueprintActionDatabase::Get();
+
+	if (!AssetInfo.IsAssetLoaded())
+	{
+		ActionDatabase.MoveUnloadedAssetActions(*InOldName, AssetInfo.ObjectPath);
+	}
 }
 
 //------------------------------------------------------------------------------
@@ -805,6 +828,7 @@ FBlueprintActionDatabase::FBlueprintActionDatabase()
 	IAssetRegistry& AssetRegistry = FModuleManager::LoadModuleChecked<FAssetRegistryModule>(TEXT("AssetRegistry")).Get();
 	AssetRegistry.OnAssetAdded().AddStatic(&BlueprintActionDatabaseImpl::OnAssetAdded);
 	AssetRegistry.OnAssetRemoved().AddStatic(&BlueprintActionDatabaseImpl::OnAssetRemoved);
+	AssetRegistry.OnAssetRenamed().AddStatic(&BlueprintActionDatabaseImpl::OnAssetRenamed);
 
 	FEditorDelegates::OnAssetsPreDelete.AddStatic(&BlueprintActionDatabaseImpl::OnAssetsPendingDelete);
 
@@ -943,7 +967,7 @@ void FBlueprintActionDatabase::RefreshClassActions(UClass* const Class)
 			ClassActionList.Empty();
 		}
 
-		FBlueprintActionDatabaseRegistrar Registrar(ActionRegistry, ActionPrimingQueue, Class);
+		FBlueprintActionDatabaseRegistrar Registrar(ActionRegistry, UnloadedActionRegistry, ActionPrimingQueue, Class);
 		if (!bIsInitializing)
 		{
 			// if this a call to RefreshClassActions() from somewhere other than 
@@ -967,7 +991,7 @@ void FBlueprintActionDatabase::RefreshClassActions(UClass* const Class)
 			// if we're only refreshing this class (and not init'ing the whole 
 			// database), then we have to reach out to individual nodes in case 
 			// they'd add entries for this as well
-			FBlueprintActionDatabaseRegistrar Registrar(ActionRegistry, ActionPrimingQueue);
+			FBlueprintActionDatabaseRegistrar Registrar(ActionRegistry, UnloadedActionRegistry, ActionPrimingQueue);
 			Registrar.ActionKeyFilter = Class; // only want actions for this class
 
 			RegisterAllNodeActions(Registrar);
@@ -1023,11 +1047,14 @@ void FBlueprintActionDatabase::RefreshAssetActions(UObject* const AssetObject)
 		}
 	}
 
-	FBlueprintActionDatabaseRegistrar Registrar(ActionRegistry, ActionPrimingQueue);
+	FBlueprintActionDatabaseRegistrar Registrar(ActionRegistry, UnloadedActionRegistry, ActionPrimingQueue);
 	Registrar.ActionKeyFilter = AssetObject; // make sure actions only associated with this asset get added
 	// nodes may have actions they wish to add actions for this asset
 	RegisterAllNodeActions(Registrar);
 
+	// Will clear up any unloaded asset actions associated with this object, if any
+	ClearUnloadedAssetActions(*AssetObject->GetPathName());
+
 	if (AssetActionList.Num() > 0)
 	{
 		// queue these assets for priming
@@ -1055,6 +1082,41 @@ void FBlueprintActionDatabase::ClearAssetActions(UObject* const AssetObject)
 	}
 }
 
+//------------------------------------------------------------------------------
+void FBlueprintActionDatabase::ClearUnloadedAssetActions(FName ObjectPath)
+{
+	// Check if the asset can be found in the unloaded action registry, if it can, we need to remove it
+	if(TArray<UBlueprintNodeSpawner*>* UnloadedActionList = UnloadedActionRegistry.Find(ObjectPath))
+	{
+		for(UBlueprintNodeSpawner* NodeSpawner : *UnloadedActionList)
+		{
+			FActionList* ActionList = ActionRegistry.Find(NodeSpawner->NodeClass);
+
+			// Remove the NodeSpawner from the main registry, it will be replaced with the loaded version of the action
+			ActionList->Remove(NodeSpawner);
+		}
+
+		// Remove the asset's path from the unloaded registry, it is no longer needed
+		UnloadedActionRegistry.Remove(ObjectPath);
+	}
+}
+
+//------------------------------------------------------------------------------
+void FBlueprintActionDatabase::MoveUnloadedAssetActions(FName SourceObjectPath, FName TargetObjectPath)
+{
+	// Check if the asset can be found in the unloaded action registry, if it can, we need to remove it and re-add under the new name
+	if(TArray<UBlueprintNodeSpawner*>* UnloadedActionList = UnloadedActionRegistry.Find(SourceObjectPath))
+	{
+		check(!UnloadedActionRegistry.Find(TargetObjectPath));
+
+		// Add the entire array to the database under the new path
+		UnloadedActionRegistry.Add(TargetObjectPath, *UnloadedActionList);
+
+		// Remove the old asset's path from the unloaded registry, it is no longer needed
+		UnloadedActionRegistry.Remove(SourceObjectPath);
+	}
+}
+
 //------------------------------------------------------------------------------
 FBlueprintActionDatabase::FActionRegistry const& FBlueprintActionDatabase::GetAllActions()
 {
diff --git a/Engine/Source/Editor/BlueprintGraph/Private/BlueprintActionDatabaseRegistrar.cpp b/Engine/Source/Editor/BlueprintGraph/Private/BlueprintActionDatabaseRegistrar.cpp
--- a/Engine/Source/Editor/BlueprintGraph/Private/BlueprintActionDatabaseRegistrar.cpp
+++ b/Engine/Source/Editor/BlueprintGraph/Private/BlueprintActionDatabaseRegistrar.cpp
@@ -60,9 +60,10 @@ static UObject const* BlueprintActionDatabaseRegistrarImpl::ResolveActionKey(UOb
  ******************************************************************************/
 
 //------------------------------------------------------------------------------
-FBlueprintActionDatabaseRegistrar::FBlueprintActionDatabaseRegistrar(FActionRegistry& Database, FPrimingQueue& PrimingQueue, TSubclassOf<UEdGraphNode> DefaultKey)
+FBlueprintActionDatabaseRegistrar::FBlueprintActionDatabaseRegistrar(FActionRegistry& Database, FUnloadedActionRegistry& UnloadedDatabase, FPrimingQueue& PrimingQueue, TSubclassOf<UEdGraphNode> DefaultKey)
 	: GeneratingClass(DefaultKey)
 	, ActionDatabase(Database)
+	, UnloadedActionDatabase(UnloadedDatabase)
 	, ActionKeyFilter(nullptr)
 	, ActionPrimingQueue(PrimingQueue)
 {
@@ -139,6 +140,29 @@ bool FBlueprintActionDatabaseRegistrar::AddBlueprintAction(UObject const* AssetO
 	return AddActionToDatabase(AssetOwner, NodeSpawner);
 }
 
+//------------------------------------------------------------------------------
+bool FBlueprintActionDatabaseRegistrar::AddBlueprintAction(FAssetData const& AssetDataOwner, UBlueprintNodeSpawner* NodeSpawner)
+{
+	bool bReturnResult = false;
+
+	// @TODO: assert that AddBlueprintAction(UBlueprintNodeSpawner* NodeSpawner) 
+	//        wouldn't come up with a different key (besides GeneratingClass)
+	if(AssetDataOwner.IsAssetLoaded())
+	{
+		bReturnResult = AddBlueprintAction(AssetDataOwner.GetAsset(), NodeSpawner);
+	}
+	else
+	{
+		bReturnResult = AddBlueprintAction(NodeSpawner->NodeClass, NodeSpawner);
+		if(bReturnResult)
+		{
+			TArray<UBlueprintNodeSpawner*>& ActionList = UnloadedActionDatabase.FindOrAdd(AssetDataOwner.ObjectPath);
+			ActionList.Add(NodeSpawner);
+		}
+	}
+	return bReturnResult;
+}
+
 //------------------------------------------------------------------------------
 bool FBlueprintActionDatabaseRegistrar::IsOpenForRegistration(UObject const* OwnerKey)
 {
diff --git a/Engine/Source/Editor/BlueprintGraph/Private/BlueprintUnloadedAssetNodeSpawner.cpp b/Engine/Source/Editor/BlueprintGraph/Private/BlueprintUnloadedAssetNodeSpawner.cpp
new file mode 100644
--- /dev/null
+++ b/Engine/Source/Editor/BlueprintGraph/Private/BlueprintUnloadedAssetNodeSpawner.cpp
@@ -0,0 +1,40 @@
+// Copyright 1998-2014 Epic Games, Inc. All Rights Reserved.
+
+#include "BlueprintGraphPrivatePCH.h"
+#include "BlueprintUnloadedAssetNodeSpawner.h"
+
+#define LOCTEXT_NAMESPACE "BlueprintUnloadedAssetNodeSpawner"
+
+/*******************************************************************************
+ * BlueprintUnloadedAssetNodeSpawner
+ ******************************************************************************/
+
+//------------------------------------------------------------------------------
+UBlueprintUnloadedAssetNodeSpawner* UBlueprintUnloadedAssetNodeSpawner::Create(TSubclassOf<UEdGraphNode> NodeClass, const FAssetData& AssetData, FText ActionMenuDescription, UObject* Outer/* = nullptr*/)
+{
+	if (Outer == nullptr)
+	{
+		Outer = GetTransientPackage();
+	}
+
+	UBlueprintUnloadedAssetNodeSpawner* NodeSpawner = NewObject<UBlueprintUnloadedAssetNodeSpawner>(Outer);
+	NodeSpawner->NodeClass = NodeClass;
+	NodeSpawner->AssetData = AssetData;
+	NodeSpawner->ActionMenuDescription = ActionMenuDescription;
+
+	return NodeSpawner;
+}
+
+//------------------------------------------------------------------------------
+UBlueprintUnloadedAssetNodeSpawner::UBlueprintUnloadedAssetNodeSpawner(class FPostConstructInitializeProperties const& PCIP)
+	: Super(PCIP)
+{
+}
+
+//------------------------------------------------------------------------------
+FText UBlueprintUnloadedAssetNodeSpawner::GetDefaultMenuName(FBindingSet const& Bindings) const
+{
+	return ActionMenuDescription;
+}
+
+#undef LOCTEXT_NAMESPACE
diff --git a/Engine/Source/Editor/BlueprintGraph/Public/BlueprintActionDatabase.h b/Engine/Source/Editor/BlueprintGraph/Public/BlueprintActionDatabase.h
--- a/Engine/Source/Editor/BlueprintGraph/Public/BlueprintActionDatabase.h
+++ b/Engine/Source/Editor/BlueprintGraph/Public/BlueprintActionDatabase.h
@@ -24,10 +24,11 @@ class FBlueprintActionDatabaseRegistrar;
 class BLUEPRINTGRAPH_API FBlueprintActionDatabase : public FGCObject, public FTickableEditorObject
 {
 public:
-	typedef TMap<TWeakObjectPtr<UObject>, int32> FPrimingQueue;
-	typedef TArray<UBlueprintNodeSpawner*>       FActionList;
-	typedef TMap<UObject const*, FActionList>    FActionRegistry;
-	
+	typedef TMap<TWeakObjectPtr<UObject>, int32>			FPrimingQueue;
+	typedef TArray<UBlueprintNodeSpawner*>					FActionList;
+	typedef TMap<UObject const*, FActionList>				FActionRegistry;
+	typedef TMap<FName, TArray<UBlueprintNodeSpawner*>>		FUnloadedActionRegistry;
+
 public:
 	/**
 	 * Getter to access the database singleton. Will populate the database first 
@@ -91,6 +92,22 @@ public:
 	 */
 	void ClearAssetActions(UObject* const AssetObject);
 	
+	/**
+	 * Finds the database entry for the specified unloaded asset and wipes it.
+	 * The entry won't be rebuilt, unless RefreshAssetActions() is explicitly called after.
+	 *
+	 * @param ObjectPath	Object's path to lookup into the database
+	 */
+	void ClearUnloadedAssetActions(FName ObjectPath);
+
+	/**
+	 * Moves the unloaded asset actions from one location to another
+	 *
+	 * @param SourceObjectPath	The object path that the data can currently be found under
+	 * @param TargetObjectPath	The object path that the data should be moved to
+	 */
+	void MoveUnloadedAssetActions(FName SourceObjectPath, FName TargetObjectPath);
+
 	/**
 	 * 
 	 * @return An estimated memory footprint of this database (in bytes).
@@ -118,6 +135,13 @@ private:
 	 */
 	FActionRegistry ActionRegistry;
 
+	/** 
+	 * A map of associated object paths for each node-class that is associated
+	 * with it. This is used for unloaded assets that will need to be replaced
+	 * after the asset is loaded with the final (and more complete) nodespawner.
+	 */
+	FUnloadedActionRegistry UnloadedActionRegistry;
+
 	/** 
 	 * References newly allocated actions that need to be "primed". Priming is 
 	 * something we do on Tick() aimed at speeding up performance (like pre-
diff --git a/Engine/Source/Editor/BlueprintGraph/Public/BlueprintActionDatabaseRegistrar.h b/Engine/Source/Editor/BlueprintGraph/Public/BlueprintActionDatabaseRegistrar.h
--- a/Engine/Source/Editor/BlueprintGraph/Public/BlueprintActionDatabaseRegistrar.h
+++ b/Engine/Source/Editor/BlueprintGraph/Public/BlueprintActionDatabaseRegistrar.h
@@ -33,6 +33,7 @@ public:
 	bool AddBlueprintAction(UEnum const* EnumOwner, UBlueprintNodeSpawner* NodeSpawner);
 	bool AddBlueprintAction(UScriptStruct const* StructOwner, UBlueprintNodeSpawner* NodeSpawner);
 	bool AddBlueprintAction(UField const* FieldOwner, UBlueprintNodeSpawner* NodeSpawner);
+	bool AddBlueprintAction(FAssetData const& AssetDataOwner, UBlueprintNodeSpawner* NodeSpawner);
 
 	/**
 	 * Special case for asset bound actions (we want to clean-up/refresh these 
@@ -56,12 +57,13 @@ public:
 	bool IsOpenForRegistration(UObject const* OwnerKey);
 
 private:
-	typedef FBlueprintActionDatabase::FActionRegistry FActionRegistry;
-	typedef FBlueprintActionDatabase::FPrimingQueue   FPrimingQueue;
+	typedef FBlueprintActionDatabase::FActionRegistry			FActionRegistry;
+	typedef FBlueprintActionDatabase::FUnloadedActionRegistry	FUnloadedActionRegistry;
+	typedef FBlueprintActionDatabase::FPrimingQueue				FPrimingQueue;
 
 	/** Only FBlueprintActionDatabase can spawn and distribute this. */
 	friend class FBlueprintActionDatabase;
-	FBlueprintActionDatabaseRegistrar(FActionRegistry& Database, FPrimingQueue& PrimingQueue, TSubclassOf<UEdGraphNode> DefaultKey = nullptr);
+	FBlueprintActionDatabaseRegistrar(FActionRegistry& Database, FUnloadedActionRegistry& UnloadedDatabase, FPrimingQueue& PrimingQueue, TSubclassOf<UEdGraphNode> DefaultKey = nullptr);
 
 	/**
 	 * Internal method that actually adds the action to the database.
@@ -78,6 +80,9 @@ private:
 	/** A reference to FBlueprintActionDatabase's internal map */
 	FActionRegistry& ActionDatabase;
 
+	/** A reference to FBlueprintActionDatabase's internal map */
+	FUnloadedActionRegistry& UnloadedActionDatabase;
+
 	/** 
 	 * When an asset is added/updated we want to poll nodes again for 
 	 * new/refreshed actions, but only for that specific asset. This serves as
