Local Graph Varaibles. An early version of function FEventGraphUtils::PinRepresentsSharedTerminal, that tells if a pin represetns a terminal, that is truly shared. If not, it will be added to local varaibles.

Creating Terminals was slightly refactored.

diff --git a/Engine/Config/BaseEngine.ini b/Engine/Config/BaseEngine.ini
--- a/Engine/Config/BaseEngine.ini
+++ b/Engine/Config/BaseEngine.ini
@@ -1630,6 +1630,7 @@ CompileDisplaysBinaryBackend=false
 CompileDisplaysAnimBlueprintBackend=false
 bTurnOffEditorConstructionScript=false
 PrintStringDuration=2.0
+bUseLocalGraphVariables=false
 
 [/Script/Engine.Blueprint]
 bRecompileOnLoad=true
diff --git a/Engine/Source/Editor/BlueprintGraph/Private/CallFunctionHandler.cpp b/Engine/Source/Editor/BlueprintGraph/Private/CallFunctionHandler.cpp
--- a/Engine/Source/Editor/BlueprintGraph/Private/CallFunctionHandler.cpp
+++ b/Engine/Source/Editor/BlueprintGraph/Private/CallFunctionHandler.cpp
@@ -225,7 +225,7 @@ void FKCHandler_CallFunction::CreateFunctionCallStatement(FKismetFunctionContext
 								{
 									UClass* InterfaceClass = CastChecked<UClass>(PinMatch->PinType.PinSubCategoryObject.Get());
 
-									FBPTerminal* ClassTerm = new (Context.IsEventGraph() ? Context.EventGraphLocals : Context.Locals) FBPTerminal();
+									FBPTerminal* ClassTerm = Context.CreateLocalTerminal(ETerminalSpecification::TS_Literal);
 									ClassTerm->Name       = InterfaceClass->GetName();
 									ClassTerm->bIsLiteral = true;
 									ClassTerm->Source     = Node;
@@ -568,10 +568,9 @@ void FKCHandler_CallFunction::RegisterNets(FKismetFunctionContext& Context, UEdG
 					if (CompilerContext.UbergraphContext != NULL)
 					{
 						check(CompilerContext.UbergraphContext->NetNameMap);
-						FBPTerminal* Term = new (Context.EventGraphLocals) FBPTerminal();
+						FBPTerminal* Term = Context.CreateLocalTerminal(ETerminalSpecification::TS_ForcedShared);
 						Term->CopyFromPin(Pin, *CompilerContext.UbergraphContext->NetNameMap->MakeValidName(Pin));
 						Term->Name += TEXT("_RefProperty");
-						Term->bIsLocal = true;
 						if (!Pin->PinType.bIsArray)
 						{
 							Term->PropertyDefault = Pin->DefaultObject ? Pin->DefaultObject->GetPathName() : Pin->DefaultValue;
@@ -600,7 +599,7 @@ void FKCHandler_CallFunction::RegisterNets(FKismetFunctionContext& Context, UEdG
 		// from object to interface
 		if ((Pin->PinType.PinCategory == K2Schema->PC_Interface) && (Pin->LinkedTo[0]->PinType.PinCategory == K2Schema->PC_Object))
 		{
-			FBPTerminal* InterfaceTerm = new (Context.IsEventGraph() ? Context.EventGraphLocals : Context.Locals) FBPTerminal();
+			FBPTerminal* InterfaceTerm = Context.CreateLocalTerminal();
 			InterfaceTerm->CopyFromPin(Pin, Context.NetNameMap->MakeValidName(Pin) + TEXT("_CastInput"));
 			InterfaceTerm->Source = Node;
 
@@ -614,12 +613,7 @@ void FKCHandler_CallFunction::RegisterNets(FKismetFunctionContext& Context, UEdG
 void FKCHandler_CallFunction::RegisterNet(FKismetFunctionContext& Context, UEdGraphPin* Net)
 {
 	// This net is an output from a function call
-	FBPTerminal* Term = new (Context.IsEventGraph() ? Context.EventGraphLocals : Context.Locals) FBPTerminal();
-
-	FString NetName = Context.NetNameMap->MakeValidName(Net);
-
-	Term->CopyFromPin(Net, NetName);
-
+	FBPTerminal* Term = Context.CreateLocalTerminalFromPinAutoChooseScope(Net, Context.NetNameMap->MakeValidName(Net));
 	Context.NetMap.Add(Net, Term);
 }
 
diff --git a/Engine/Source/Editor/BlueprintGraph/Private/DelegateNodeHandlers.cpp b/Engine/Source/Editor/BlueprintGraph/Private/DelegateNodeHandlers.cpp
--- a/Engine/Source/Editor/BlueprintGraph/Private/DelegateNodeHandlers.cpp
+++ b/Engine/Source/Editor/BlueprintGraph/Private/DelegateNodeHandlers.cpp
@@ -162,8 +162,7 @@ void FKCHandler_AddRemoveDelegate::RegisterNets(FKismetFunctionContext& Context,
 	FBPTerminal* Term = FoundTerm ? *FoundTerm : NULL;
 	if(NULL == Term)
 	{
-		Term = new (Context.IsEventGraph() ? Context.EventGraphLocals : Context.Locals) FBPTerminal();
-		Term->CopyFromPin(Net, Context.NetNameMap->MakeValidName(Net));
+		Term = Context.CreateLocalTerminalFromPinAutoChooseScope(Net, Context.NetNameMap->MakeValidName(Net));
 		Context.NetMap.Add(Net, Term);
 	}
 }
@@ -243,8 +242,7 @@ void FKCHandler_CreateDelegate::RegisterNets(FKismetFunctionContext& Context, UE
 		FBPTerminal* InputObjTerm = FoundTerm ? *FoundTerm : NULL;
 		if(NULL == InputObjTerm)
 		{
-			InputObjTerm = new (Context.IsEventGraph() ? Context.EventGraphLocals : Context.Locals) FBPTerminal();
-			InputObjTerm->CopyFromPin(Net, Context.NetNameMap->MakeValidName(Net));
+			InputObjTerm = Context.CreateLocalTerminalFromPinAutoChooseScope(Net, Context.NetNameMap->MakeValidName(Net));
 			Context.NetMap.Add(Net, InputObjTerm);
 		}
 	}
@@ -263,8 +261,7 @@ void FKCHandler_CreateDelegate::RegisterNets(FKismetFunctionContext& Context, UE
 		FBPTerminal* OutDelegateTerm = FoundTerm ? *FoundTerm : NULL;
 		if(NULL == OutDelegateTerm)
 		{
-			OutDelegateTerm = new (Context.IsEventGraph() ? Context.EventGraphLocals : Context.Locals) FBPTerminal();
-			OutDelegateTerm->CopyFromPin(Net, Context.NetNameMap->MakeValidName(Net));
+			OutDelegateTerm = Context.CreateLocalTerminalFromPinAutoChooseScope(Net, Context.NetNameMap->MakeValidName(Net));
 			if (NULL == FMemberReference::ResolveSimpleMemberReference<UFunction>(OutDelegateTerm->Type.PinSubCategoryMemberReference))
 			{
 				FMemberReference::FillSimpleMemberReference<UFunction>(DelegateNode->GetDelegateSignature(), OutDelegateTerm->Type.PinSubCategoryMemberReference);
@@ -298,7 +295,7 @@ void FKCHandler_CreateDelegate::Compile(FKismetFunctionContext& Context, UEdGrap
 	}
 
 	{
-		FBPTerminal* DelegateNameTerm = new (Context.IsEventGraph() ? Context.EventGraphLocals : Context.Locals) FBPTerminal();
+		FBPTerminal* DelegateNameTerm = Context.CreateLocalTerminal(ETerminalSpecification::TS_Literal);
 		DelegateNameTerm->Type.PinCategory = CompilerContext.GetSchema()->PC_Name;
 		DelegateNameTerm->Name = DelegateNode->GetFunctionName().ToString();
 		DelegateNameTerm->bIsLiteral = true;
diff --git a/Engine/Source/Editor/BlueprintGraph/Private/DynamicCastHandler.cpp b/Engine/Source/Editor/BlueprintGraph/Private/DynamicCastHandler.cpp
--- a/Engine/Source/Editor/BlueprintGraph/Private/DynamicCastHandler.cpp
+++ b/Engine/Source/Editor/BlueprintGraph/Private/DynamicCastHandler.cpp
@@ -13,18 +13,16 @@ void FKCHandler_DynamicCast::RegisterNets(FKismetFunctionContext& Context, UEdGr
 	FNodeHandlingFunctor::RegisterNets(Context, Node);
 
 	// Create a term to determine if the cast was successful or not
-	FBPTerminal* BoolTerm = new (Context.IsEventGraph() ? Context.EventGraphLocals : Context.Locals) FBPTerminal();
+	FBPTerminal* BoolTerm = Context.CreateLocalTerminal();
 	BoolTerm->Type.PinCategory = CompilerContext.GetSchema()->PC_Boolean;
 	BoolTerm->Source = Node;
 	BoolTerm->Name = Context.NetNameMap->MakeValidName(Node) + TEXT("_CastSuccess");
-	BoolTerm->bIsLocal = true;
 	BoolTermMap.Add(Node, BoolTerm);
 }
 
 void FKCHandler_DynamicCast::RegisterNet(FKismetFunctionContext& Context, UEdGraphPin* Net)
 {
-	FBPTerminal* Term = new (Context.IsEventGraph() ? Context.EventGraphLocals : Context.Locals) FBPTerminal();
-	Term->CopyFromPin(Net, Context.NetNameMap->MakeValidName(Net));
+	FBPTerminal* Term = Context.CreateLocalTerminalFromPinAutoChooseScope(Net, Context.NetNameMap->MakeValidName(Net));
 	Context.NetMap.Add(Net, Term);
 }
 
@@ -85,7 +83,7 @@ void FKCHandler_DynamicCast::Compile(FKismetFunctionContext& Context, UEdGraphNo
 	}
 
 	// Create a literal term from the class specified in the node
-	FBPTerminal* ClassTerm = new (Context.IsEventGraph() ? Context.EventGraphLocals : Context.Locals) FBPTerminal();
+	FBPTerminal* ClassTerm = Context.CreateLocalTerminal(ETerminalSpecification::TS_Literal);
 	ClassTerm->Name = DynamicCastNode->TargetType->GetName();
 	ClassTerm->bIsLiteral = true;
 	ClassTerm->Source = Node;
diff --git a/Engine/Source/Editor/BlueprintGraph/Private/EventEntryHandler.cpp b/Engine/Source/Editor/BlueprintGraph/Private/EventEntryHandler.cpp
--- a/Engine/Source/Editor/BlueprintGraph/Private/EventEntryHandler.cpp
+++ b/Engine/Source/Editor/BlueprintGraph/Private/EventEntryHandler.cpp
@@ -10,7 +10,7 @@ void FKCHandler_EventEntry::RegisterNet(FKismetFunctionContext& Context, UEdGrap
 	if(!bIsDelegateOutput)
 	{
 		// This net is an event parameter; push to be a private class member variable
-		FBPTerminal* Term = new (Context.EventGraphLocals) FBPTerminal();
+		FBPTerminal* Term = Context.CreateLocalTerminal(ETerminalSpecification::TS_ForcedShared);
 		Term->CopyFromPin(Net, Context.NetNameMap->MakeValidName(Net));
 
 		Context.NetMap.Add(Net, Term);
diff --git a/Engine/Source/Editor/BlueprintGraph/Private/K2Node_AddComponent.cpp b/Engine/Source/Editor/BlueprintGraph/Private/K2Node_AddComponent.cpp
--- a/Engine/Source/Editor/BlueprintGraph/Private/K2Node_AddComponent.cpp
+++ b/Engine/Source/Editor/BlueprintGraph/Private/K2Node_AddComponent.cpp
@@ -37,9 +37,8 @@ public:
 			{
 				check(CompilerContext.UbergraphContext->NetNameMap);
 				// If there's no term, add a dummy identity transform here.  Term is defined in the event graph's scope, because a default value needs to be set
-				FBPTerminal* TransformTerm = new (Context.EventGraphLocals) FBPTerminal();
+				FBPTerminal* TransformTerm = Context.CreateLocalTerminal(ETerminalSpecification::TS_ForcedShared);
 				TransformTerm->CopyFromPin(TransformPin, FString::Printf(TEXT("%s_AddComponentDefaultTransform"), *CompilerContext.UbergraphContext->NetNameMap->MakeValidName(TransformPin)));
-				TransformTerm->bIsLocal = true;
 				TransformTerm->bIsConst = true;
 				TransformTerm->PropertyDefault = TEXT("");
 
diff --git a/Engine/Source/Editor/BlueprintGraph/Private/K2Node_BreakStruct.cpp b/Engine/Source/Editor/BlueprintGraph/Private/K2Node_BreakStruct.cpp
--- a/Engine/Source/Editor/BlueprintGraph/Private/K2Node_BreakStruct.cpp
+++ b/Engine/Source/Editor/BlueprintGraph/Private/K2Node_BreakStruct.cpp
@@ -59,8 +59,7 @@ public:
 			// standard register net
 			else
 			{
-				Term = new (Context.IsEventGraph() ? Context.EventGraphLocals : Context.Locals) FBPTerminal();
-				Term->CopyFromPin(Net, Context.NetNameMap->MakeValidName(Net));
+				Term = Context.CreateLocalTerminalFromPinAutoChooseScope(Net, Context.NetNameMap->MakeValidName(Net));
 			}
 			Context.NetMap.Add(Net, Term);
 		}
@@ -79,8 +78,7 @@ public:
 
 		if (BoundProperty != NULL)
 		{
-			FBPTerminal* Term = new (Context.IsEventGraph() ? Context.EventGraphLocals : Context.Locals) FBPTerminal();
-			Term->CopyFromPin(Net, Net->PinName);
+			FBPTerminal* Term = Context.CreateLocalTerminalFromPinAutoChooseScope(Net, Net->PinName);
 			Term->AssociatedVarProperty = BoundProperty;
 			Context.NetMap.Add(Net, Term);
 			Term->Context = ContextTerm;
diff --git a/Engine/Source/Editor/BlueprintGraph/Private/K2Node_CastByteToEnum.cpp b/Engine/Source/Editor/BlueprintGraph/Private/K2Node_CastByteToEnum.cpp
--- a/Engine/Source/Editor/BlueprintGraph/Private/K2Node_CastByteToEnum.cpp
+++ b/Engine/Source/Editor/BlueprintGraph/Private/K2Node_CastByteToEnum.cpp
@@ -120,8 +120,7 @@ public:
 		UEdGraphPin* Net = FEdGraphUtilities::GetNetFromPin(InPin);
 		if (Context.NetMap.Find(Net) == NULL)
 		{
-			FBPTerminal* Term = new (Context.IsEventGraph() ? Context.EventGraphLocals : Context.Locals) FBPTerminal();
-			Term->CopyFromPin(Net, Context.NetNameMap->MakeValidName(Net));
+			FBPTerminal* Term = Context.CreateLocalTerminalFromPinAutoChooseScope(Net, Context.NetNameMap->MakeValidName(Net));
 			Context.NetMap.Add(Net, Term);
 		}
 
diff --git a/Engine/Source/Editor/BlueprintGraph/Private/K2Node_DelegateSet.cpp b/Engine/Source/Editor/BlueprintGraph/Private/K2Node_DelegateSet.cpp
--- a/Engine/Source/Editor/BlueprintGraph/Private/K2Node_DelegateSet.cpp
+++ b/Engine/Source/Editor/BlueprintGraph/Private/K2Node_DelegateSet.cpp
@@ -28,12 +28,11 @@ public:
 			CompilerContext.MessageLog.Warning(*FString(*LOCTEXT("DeprecatedDelegateSet_Warning", "DelegateSet node @@ is Deprecated. It should be replaced by an EventCaller Bind node").ToString()), DelegateNode);
 
 			// Create a term to store the locally created delegate that we'll use to add to the MC delegate
-			FBPTerminal* DelegateTerm = new (Context.IsEventGraph() ? Context.EventGraphLocals : Context.Locals) FBPTerminal();
+			FBPTerminal* DelegateTerm = Context.CreateLocalTerminal();
 			DelegateTerm->Type.PinCategory = CompilerContext.GetSchema()->PC_Delegate;
 			FMemberReference::FillSimpleMemberReference<UFunction>(DelegateNode->GetDelegateSignature(), DelegateTerm->Type.PinSubCategoryMemberReference);
 			DelegateTerm->Source = Node;
 			DelegateTerm->Name = Context.NetNameMap->MakeValidName(Node) + TEXT("_TempBindingDelegate");
-			DelegateTerm->bIsLocal = true;
 			LocalDelegateMap.Add(Node, DelegateTerm);
 
 			// The only net we need to register for this node is the delegate's target (self) pin, since the others are expanded to their own event node
@@ -100,7 +99,7 @@ public:
 		FBPTerminal** pDelegateOwnerTerm = DelegateOwnerPin ? Context.NetMap.Find(DelegateOwnerPin) : NULL;
 
 		// Create a delegate name term
-		FBPTerminal* DelegateNameTerm = new (Context.IsEventGraph() ? Context.EventGraphLocals : Context.Locals) FBPTerminal();
+		FBPTerminal* DelegateNameTerm = Context.CreateLocalTerminal(ETerminalSpecification::TS_Literal);
 		DelegateNameTerm->Type.PinCategory = CompilerContext.GetSchema()->PC_Name;
 		DelegateNameTerm->Name = DelegateNode->GetDelegateTargetEntryPointName().ToString();
 		DelegateNameTerm->bIsLiteral = true;
diff --git a/Engine/Source/Editor/BlueprintGraph/Private/K2Node_EnumLiteral.cpp b/Engine/Source/Editor/BlueprintGraph/Private/K2Node_EnumLiteral.cpp
--- a/Engine/Source/Editor/BlueprintGraph/Private/K2Node_EnumLiteral.cpp
+++ b/Engine/Source/Editor/BlueprintGraph/Private/K2Node_EnumLiteral.cpp
@@ -80,8 +80,7 @@ public:
 		UEdGraphPin* Net = FEdGraphUtilities::GetNetFromPin(InPin);
 		if (Context.NetMap.Find(Net) == NULL)
 		{
-			FBPTerminal* Term = new (Context.IsEventGraph() ? Context.EventGraphLocals : Context.Locals) FBPTerminal();
-			Term->CopyFromPin(Net, Context.NetNameMap->MakeValidName(Net));
+			FBPTerminal* Term = Context.CreateLocalTerminalFromPinAutoChooseScope(Net, Context.NetNameMap->MakeValidName(Net));
 			Context.NetMap.Add(Net, Term);
 		}
 
diff --git a/Engine/Source/Editor/BlueprintGraph/Private/K2Node_MakeArray.cpp b/Engine/Source/Editor/BlueprintGraph/Private/K2Node_MakeArray.cpp
--- a/Engine/Source/Editor/BlueprintGraph/Private/K2Node_MakeArray.cpp
+++ b/Engine/Source/Editor/BlueprintGraph/Private/K2Node_MakeArray.cpp
@@ -34,9 +34,7 @@ public:
 		FNodeHandlingFunctor::RegisterNets(Context, Node);
 
 		// Create a local term to drop the array into
- 		FBPTerminal* Term = new (Context.IsEventGraph() ? Context.EventGraphLocals : Context.Locals) FBPTerminal();
- 		Term->CopyFromPin(OutputPin, Context.NetNameMap->MakeValidName(OutputPin));
-		Term->bIsLocal = true;
+		FBPTerminal* Term = Context.CreateLocalTerminalFromPinAutoChooseScope(OutputPin, Context.NetNameMap->MakeValidName(OutputPin));
 		Term->bPassedByReference = false;
 		Term->Source = Node;
  		Context.NetMap.Add(OutputPin, Term);
diff --git a/Engine/Source/Editor/BlueprintGraph/Private/K2Node_MultiGate.cpp b/Engine/Source/Editor/BlueprintGraph/Private/K2Node_MultiGate.cpp
--- a/Engine/Source/Editor/BlueprintGraph/Private/K2Node_MultiGate.cpp
+++ b/Engine/Source/Editor/BlueprintGraph/Private/K2Node_MultiGate.cpp
@@ -42,7 +42,7 @@ public:
 		const FString BaseNetName = Context.NetNameMap->MakeValidName(Node);
 
 		// Create a term to store a bool that determines if we're in the first execution of the node or not
-		FBPTerminal* FirstRunTerm = new (Context.IsEventGraph() ? Context.EventGraphLocals : Context.Locals) FBPTerminal();
+		FBPTerminal* FirstRunTerm = Context.CreateLocalTerminal();
 		FirstRunTerm->Type.PinCategory = CompilerContext.GetSchema()->PC_Boolean;
 		FirstRunTerm->Source = Node;
 		FirstRunTerm->Name = BaseNetName + TEXT("_FirstRun");
@@ -52,7 +52,7 @@ public:
 		// If there is already a data node from expansion phase
 		if (!GateNode || !GateNode->DataNode)
 		{
-			FBPTerminal* DataTerm = new (Context.IsEventGraph() ? Context.EventGraphLocals : Context.Locals) FBPTerminal();
+			FBPTerminal* DataTerm = Context.CreateLocalTerminal();
 			DataTerm->Type.PinCategory = CompilerContext.GetSchema()->PC_Int;
 			DataTerm->Source = Node;
 			DataTerm->Name = BaseNetName + TEXT("_Data");
@@ -62,7 +62,7 @@ public:
 		// Create a local scratch bool for run-time if there isn't already one
 		if (!GenericBoolTerm)
 		{
-			GenericBoolTerm = new (Context.IsEventGraph() ? Context.EventGraphLocals : Context.Locals) FBPTerminal();
+			GenericBoolTerm = Context.CreateLocalTerminal();
 			GenericBoolTerm->Type.PinCategory = CompilerContext.GetSchema()->PC_Boolean;
 			GenericBoolTerm->Source = Node;
 			GenericBoolTerm->Name = BaseNetName + TEXT("_ScratchBool");
@@ -71,7 +71,7 @@ public:
 		// Create a local scratch int for run-time index tracking if there isn't already one
 		if (!IndexTerm)
 		{
-			IndexTerm = new (Context.IsEventGraph() ? Context.EventGraphLocals : Context.Locals) FBPTerminal();
+			IndexTerm = Context.CreateLocalTerminal();
 			IndexTerm->Type.PinCategory = CompilerContext.GetSchema()->PC_Int;
 			IndexTerm->Source = Node;
 			IndexTerm->Name = BaseNetName + TEXT("_ScratchIndex");
@@ -162,26 +162,23 @@ public:
 		FBPTerminal* FirstRunBoolTerm = FirstRunTermMap.FindRef(GateNode);
 
 		// Create a literal pin that represents a -1 value
-		FBPTerminal* InvalidIndexTerm = new (Context.IsEventGraph() ? Context.EventGraphLocals : Context.Locals) FBPTerminal();
-		InvalidIndexTerm->bIsLocal = true;
+		FBPTerminal* InvalidIndexTerm = Context.CreateLocalTerminal(ETerminalSpecification::TS_Literal);
 		InvalidIndexTerm->bIsLiteral = true;
 		InvalidIndexTerm->Type.PinCategory = CompilerContext.GetSchema()->PC_Int;
 		InvalidIndexTerm->Name = TEXT("-1");
 
 		// Create a literal pin that represents a true value
-		FBPTerminal* TrueBoolTerm = new (Context.IsEventGraph() ? Context.EventGraphLocals : Context.Locals) FBPTerminal();
+		FBPTerminal* TrueBoolTerm = Context.CreateLocalTerminal(ETerminalSpecification::TS_Literal);
 		TrueBoolTerm->Type.PinCategory = CompilerContext.GetSchema()->PC_Boolean;
 		TrueBoolTerm->bIsLiteral = true;
-		TrueBoolTerm->bIsLocal = true;
 		TrueBoolTerm->Name = TEXT("true");
 
 		// Get the out pins and create a literal describing how many logical outs there are
 		TArray<UEdGraphPin*> OutPins;
 		GateNode->GetOutPins(OutPins);
-		FBPTerminal* NumOutsTerm = new (Context.IsEventGraph() ? Context.EventGraphLocals : Context.Locals) FBPTerminal();
+		FBPTerminal* NumOutsTerm = Context.CreateLocalTerminal(ETerminalSpecification::TS_Literal);
 		NumOutsTerm->Type.PinCategory = CompilerContext.GetSchema()->PC_Int;
 		NumOutsTerm->bIsLiteral = true;
-		NumOutsTerm->bIsLocal = true;
 		NumOutsTerm->Name = FString::Printf(TEXT("%d"), OutPins.Num());
 
 		///////////////////////////////////////////////////
@@ -398,8 +395,7 @@ public:
 			IndexEqualityStatement.FunctionContext = NULL;
 			IndexEqualityStatement.bIsParentContext = false;
 			// LiteralIndexTerm will be the right side of the == statemnt
-			FBPTerminal* LiteralIndexTerm = new (Context.IsEventGraph() ? Context.EventGraphLocals : Context.Locals) FBPTerminal();
-			LiteralIndexTerm->bIsLocal = true;
+			FBPTerminal* LiteralIndexTerm = Context.CreateLocalTerminal(ETerminalSpecification::TS_Literal);
 			LiteralIndexTerm->bIsLiteral = true;
 			LiteralIndexTerm->Type.PinCategory = CompilerContext.GetSchema()->PC_Int;
 			LiteralIndexTerm->Name = FString::Printf(TEXT("%d"), OutIdx);
@@ -446,8 +442,7 @@ public:
 		PrintStatement.FunctionContext = NULL;
 		PrintStatement.bIsParentContext = false;
 		// Create a local int for use in the equality call function below (LiteralTerm = the right hand side of the EqualEqual_IntInt or NotEqual_BoolBool statement)
-		FBPTerminal* LiteralStringTerm = new (Context.IsEventGraph() ? Context.EventGraphLocals : Context.Locals) FBPTerminal();
-		LiteralStringTerm->bIsLocal = true;
+		FBPTerminal* LiteralStringTerm = Context.CreateLocalTerminal(ETerminalSpecification::TS_Literal);
 		LiteralStringTerm->bIsLiteral = true;
 		LiteralStringTerm->Type.PinCategory = CompilerContext.GetSchema()->PC_String;
 		LiteralStringTerm->Name = FString::Printf(*LOCTEXT("MultiGateNode IndexWarning", "MultiGate Node failed! Out of bounds indexing of the out pins. There are only %d outs available.").ToString(), OutPins.Num());
diff --git a/Engine/Source/Editor/BlueprintGraph/Private/K2Node_Select.cpp b/Engine/Source/Editor/BlueprintGraph/Private/K2Node_Select.cpp
--- a/Engine/Source/Editor/BlueprintGraph/Private/K2Node_Select.cpp
+++ b/Engine/Source/Editor/BlueprintGraph/Private/K2Node_Select.cpp
@@ -31,16 +31,15 @@ public:
 		// Create the net for the return value manually as it's a special case Output Direction pin
 		UK2Node_Select* SelectNode = Cast<UK2Node_Select>(Node);
 		UEdGraphPin* ReturnPin = SelectNode->GetReturnValuePin();
-		FBPTerminal* Term = new (Context.IsEventGraph() ? Context.EventGraphLocals : Context.Locals) FBPTerminal();
-		Term->CopyFromPin(ReturnPin, Context.NetNameMap->MakeValidName(ReturnPin));
+
+		FBPTerminal* Term = Context.CreateLocalTerminalFromPinAutoChooseScope(ReturnPin, Context.NetNameMap->MakeValidName(ReturnPin));
 		Context.NetMap.Add(SelectNode->GetReturnValuePin(), Term);
 
 		// Create a term to determine if the compare was successful or not
-		FBPTerminal* BoolTerm = new (Context.IsEventGraph() ? Context.EventGraphLocals : Context.Locals) FBPTerminal();
+		FBPTerminal* BoolTerm = Context.CreateLocalTerminal();
 		BoolTerm->Type.PinCategory = CompilerContext.GetSchema()->PC_Boolean;
 		BoolTerm->Source = Node;
 		BoolTerm->Name = Context.NetNameMap->MakeValidName(Node) + TEXT("_CmpSuccess");
-		BoolTerm->bIsLocal = true;
 		BoolTermMap.Add(Node, BoolTerm);
 	}
 
@@ -91,8 +90,7 @@ public:
 				// The condition passed into the Select node
 				Statement.RHS.Add(*ConditionTerm);
 				// Create a local int for use in the equality call function below (LiteralTerm = the right hand side of the EqualEqual_IntInt or NotEqual_BoolBool statement)
-				FBPTerminal* LiteralTerm = new (Context.IsEventGraph() ? Context.EventGraphLocals : Context.Locals) FBPTerminal();
-				LiteralTerm->bIsLocal = true;
+				FBPTerminal* LiteralTerm = Context.CreateLocalTerminal(ETerminalSpecification::TS_Literal);
 				LiteralTerm->bIsLiteral = true;
 				LiteralTerm->Type.PinCategory = CompilerContext.GetSchema()->PC_Int;
 				LiteralTerm->Name = FString::Printf(TEXT("%d"), OptionIdx);
@@ -144,8 +142,7 @@ public:
 					PrintStatement.bIsParentContext = false;
 
 					// Create a local int for use in the equality call function below (LiteralTerm = the right hand side of the EqualEqual_IntInt or NotEqual_BoolBool statement)
-					FBPTerminal* LiteralStringTerm = new (Context.IsEventGraph() ? Context.EventGraphLocals : Context.Locals) FBPTerminal();
-					LiteralStringTerm->bIsLocal = true;
+					FBPTerminal* LiteralStringTerm = Context.CreateLocalTerminal(ETerminalSpecification::TS_Literal);
 					LiteralStringTerm->bIsLiteral = true;
 					LiteralStringTerm->Type.PinCategory = CompilerContext.GetSchema()->PC_String;
 
diff --git a/Engine/Source/Editor/BlueprintGraph/Private/K2Node_Switch.cpp b/Engine/Source/Editor/BlueprintGraph/Private/K2Node_Switch.cpp
--- a/Engine/Source/Editor/BlueprintGraph/Private/K2Node_Switch.cpp
+++ b/Engine/Source/Editor/BlueprintGraph/Private/K2Node_Switch.cpp
@@ -38,11 +38,10 @@ public:
 
 		// Create a term to determine if the compare was successful or not
 		//@TODO: Ideally we just create one ever, not one per switch
-		FBPTerminal* BoolTerm = new (Context.IsEventGraph() ? Context.EventGraphLocals : Context.Locals) FBPTerminal();
+		FBPTerminal* BoolTerm = Context.CreateLocalTerminal();
 		BoolTerm->Type.PinCategory = CompilerContext.GetSchema()->PC_Boolean;
 		BoolTerm->Source = Node;
 		BoolTerm->Name = Context.NetNameMap->MakeValidName(Node) + TEXT("_CmpSuccess");
-		BoolTerm->bIsLocal = true;
 		BoolTermMap.Add(Node, BoolTerm);
 	}
 
diff --git a/Engine/Source/Editor/BlueprintGraph/Private/K2Node_TemporaryVariable.cpp b/Engine/Source/Editor/BlueprintGraph/Private/K2Node_TemporaryVariable.cpp
--- a/Engine/Source/Editor/BlueprintGraph/Private/K2Node_TemporaryVariable.cpp
+++ b/Engine/Source/Editor/BlueprintGraph/Private/K2Node_TemporaryVariable.cpp
@@ -18,7 +18,7 @@ public:
 	virtual void RegisterNet(FKismetFunctionContext& Context, UEdGraphPin* Net) override
 	{
 		// This net is an anonymous temporary variable
-		FBPTerminal* Term = new (Context.IsEventGraph() ? Context.EventGraphLocals : Context.Locals) FBPTerminal();
+		FBPTerminal* Term = Context.CreateLocalTerminal(Context.IsEventGraph() ? ETerminalSpecification::TS_ForcedShared : ETerminalSpecification::TS_Unspecified);
 
 		FString NetName = Context.NetNameMap->MakeValidName(Net);
 
diff --git a/Engine/Source/Editor/BlueprintGraph/Private/MakeStructHandler.cpp b/Engine/Source/Editor/BlueprintGraph/Private/MakeStructHandler.cpp
--- a/Engine/Source/Editor/BlueprintGraph/Private/MakeStructHandler.cpp
+++ b/Engine/Source/Editor/BlueprintGraph/Private/MakeStructHandler.cpp
@@ -71,8 +71,7 @@ void FKCHandler_MakeStruct::RegisterNets(FKismetFunctionContext& Context, UEdGra
 
 void FKCHandler_MakeStruct::RegisterNet(FKismetFunctionContext& Context, UEdGraphPin* Net)
 {
-	FBPTerminal* Term = new (Context.IsEventGraph() ? Context.EventGraphLocals : Context.Locals) FBPTerminal();
-	Term->CopyFromPin(Net, Context.NetNameMap->MakeValidName(Net));
+	FBPTerminal* Term = Context.CreateLocalTerminalFromPinAutoChooseScope(Net, Context.NetNameMap->MakeValidName(Net));
 	Context.NetMap.Add(Net, Term);
 }
 
@@ -103,7 +102,7 @@ void FKCHandler_MakeStruct::Compile(FKismetFunctionContext& Context, UEdGraphNod
 			UProperty* BoundProperty = FindField<UProperty>(Node->StructType, *(Pin->PinName));
 			check(NULL != BoundProperty);
 
-			FBPTerminal* DstTerm = new (Context.IsEventGraph() ? Context.EventGraphLocals : Context.Locals) FBPTerminal();
+			FBPTerminal* DstTerm = Context.CreateLocalTerminal();
 			DstTerm->CopyFromPin(Pin, Context.NetNameMap->MakeValidName(Pin));
 			DstTerm->AssociatedVarProperty = BoundProperty;
 			DstTerm->Context = OutputStructTerm;
diff --git a/Engine/Source/Editor/KismetCompiler/Private/KismetCompiler.cpp b/Engine/Source/Editor/KismetCompiler/Private/KismetCompiler.cpp
--- a/Engine/Source/Editor/KismetCompiler/Private/KismetCompiler.cpp
+++ b/Engine/Source/Editor/KismetCompiler/Private/KismetCompiler.cpp
@@ -44,6 +44,7 @@ DEFINE_STAT(EKismetCompilerStats_PostcompileFunction);
 DEFINE_STAT(EKismetCompilerStats_FinalizationWork);
 DEFINE_STAT(EKismetCompilerStats_CodeGenerationTime);
 DEFINE_STAT(EKismetCompilerStats_UpdateBlueprintGeneratedClass);
+DEFINE_STAT(EKismetCompilerStats_ChooseTerminalScope);
 		
 //////////////////////////////////////////////////////////////////////////
 // FKismetCompilerContext
@@ -587,6 +588,11 @@ void FKismetCompilerContext::CreatePropertiesFromList(UStruct* Scope, UField**&
 			MessageLog.Warning(*FString::Printf(*LOCTEXT("AssociatedVarProperty_Error", "AssociatedVarProperty property overriden %s from @@ type (%s)").ToString(), *Term.Name, *UEdGraphSchema_K2::TypeToText(Term.Type).ToString()), Term.Source);
 		}
 
+		if (Term.bIsLiteral)
+		{
+			MessageLog.Error(*FString::Printf(*LOCTEXT("PropertyForLiteral_Error", "Cannot create property for a literal: %s from @@ type (%s)").ToString(), *Term.Name, *UEdGraphSchema_K2::TypeToText(Term.Type).ToString()), Term.Source);
+		}
+
 		UProperty* NewProperty = FKismetCompilerUtilities::CreatePropertyOnScope(Scope, FName(*Term.Name), Term.Type, NewClass, PropertyFlags, Schema, MessageLog);
 		if (NewProperty != NULL)
 		{
diff --git a/Engine/Source/Editor/KismetCompiler/Private/KismetCompilerMisc.cpp b/Engine/Source/Editor/KismetCompiler/Private/KismetCompilerMisc.cpp
--- a/Engine/Source/Editor/KismetCompiler/Private/KismetCompilerMisc.cpp
+++ b/Engine/Source/Editor/KismetCompiler/Private/KismetCompilerMisc.cpp
@@ -1039,6 +1039,170 @@ int32 FKismetFunctionContext::GetContextUniqueID()
 	return UUIDCounter++;
 }
 
+struct FEventGraphUtils
+{
+	static bool IsEntryPointNode(const UK2Node* Node)
+	{
+		bool bResult = false;
+		if (Node)
+		{
+			bResult |= Node->IsA<UK2Node_Event>();
+			bResult |= Node->IsA<UK2Node_Timeline>();
+			if (auto CallNode = Cast<const UK2Node_CallFunction>(Node))
+			{
+				bResult |= CallNode->IsLatentFunction();
+			}
+		}
+		return bResult;
+	}
+
+	static void FindEventsCallingTheNodeRecursive(const UK2Node* Node, TSet<const UK2Node*>& Results, TSet<const UK2Node*>& CheckedNodes, const UK2Node* StopOn)
+	{
+		if (!Node)
+		{
+			return;
+		}
+
+		bool bAlreadyTraversed = false;
+		CheckedNodes.Add(Node, &bAlreadyTraversed);
+		if (bAlreadyTraversed)
+		{
+			return;
+		}
+
+		if (Node == StopOn)
+		{
+			return;
+		}
+
+		if (IsEntryPointNode(Node))
+		{
+			Results.Add(Node);
+			return;
+		}
+
+		const UEdGraphSchema_K2* Schema = CastChecked<const UEdGraphSchema_K2>(Node->GetSchema());
+		const bool bIsPure = Node->IsNodePure();
+		for (auto Pin : Node->Pins)
+		{
+			const bool bProperPure		= bIsPure	&& Pin && (Pin->Direction == EEdGraphPinDirection::EGPD_Output);
+			const bool bProperNotPure	= !bIsPure	&& Pin && (Pin->Direction == EEdGraphPinDirection::EGPD_Input) && Schema->IsExecPin(*Pin);
+			if (bProperPure || bProperNotPure)
+			{
+				for (auto Link : Pin->LinkedTo)
+				{
+					auto LinkOwner = Link ? Link->GetOwningNodeUnchecked() : NULL;
+					auto NodeToCheck = LinkOwner ? CastChecked<const UK2Node>(LinkOwner) : NULL;
+					FindEventsCallingTheNodeRecursive(NodeToCheck, Results, CheckedNodes, StopOn);
+				}
+			}
+		}
+	}
+
+	static TSet<const UK2Node*> FindExecutionNodes(const UK2Node* Node, const UK2Node* StopOn)
+	{
+		TSet<const UK2Node*> Results;
+		TSet<const UK2Node*> CheckedNodes;
+		FindEventsCallingTheNodeRecursive(Node, Results, CheckedNodes, StopOn);
+		return Results;
+	}
+
+	static bool PinRepresentsSharedTerminal(const UEdGraphPin& Net)
+	{
+		// TODO: Strange cases..
+		if ((Net.Direction != EEdGraphPinDirection::EGPD_Output)
+			|| Net.PinType.bIsArray
+			|| Net.PinType.bIsReference
+			|| Net.PinType.bIsConst
+			|| Net.SubPins.Num())
+		{
+			return true;
+		}
+
+		// NOT CONNECTED, so it doesn't have to be shared
+		if (!Net.LinkedTo.Num())
+		{
+			return false;
+		}
+
+		auto OwnerNode = Cast<const UK2Node>(Net.GetOwningNodeUnchecked());
+		ensure(OwnerNode);
+		// Terminals from pure nodes will be recreated anyway, so they can be always local
+		if (OwnerNode && OwnerNode->IsNodePure())
+		{
+			return false;
+		}
+
+		// 
+		if (IsEntryPointNode(OwnerNode))
+		{
+			return true;
+		}
+
+		// 
+		auto SourceEntryPoints = FEventGraphUtils::FindExecutionNodes(OwnerNode, NULL);
+		if (1 != SourceEntryPoints.Num())
+		{
+			return true;
+		}
+
+		//
+		for (auto Link : Net.LinkedTo)
+		{
+			auto LinkOwnerNode = Cast<const UK2Node>(Link->GetOwningNodeUnchecked());
+			ensure(LinkOwnerNode);
+			auto EventsCallingDestination = FEventGraphUtils::FindExecutionNodes(LinkOwnerNode, OwnerNode);
+			if (0 != EventsCallingDestination.Num())
+			{
+				return true;
+			}
+		}
+		return false;
+	}
+};
+
+FBPTerminal* FKismetFunctionContext::CreateLocalTerminal(ETerminalSpecification Spec)
+{
+	FBPTerminal* Result = NULL;
+	switch (Spec)
+	{
+	case ETerminalSpecification::TS_ForcedShared:
+		// ensure(IsEventGraph()); it's used in function by UK2Node_AddComponent
+		Result = new (EventGraphLocals)FBPTerminal();
+		break;
+	case ETerminalSpecification::TS_Literal:
+		Result = new (Literals) FBPTerminal();
+		Result->bIsLiteral = true;
+		break;
+	default:
+		const bool bIsLocal = !IsEventGraph();
+		Result = new (bIsLocal ? Locals : EventGraphLocals) FBPTerminal();
+		Result->bIsLocal = bIsLocal;
+		break;
+	}
+	return Result;
+}
+
+FBPTerminal* FKismetFunctionContext::CreateLocalTerminalFromPinAutoChooseScope(UEdGraphPin* Net, const FString& NewName)
+{
+	check(Net);
+	bool bSharedTerm = IsEventGraph();
+
+	const bool bRegularBlueprint = Blueprint && (Blueprint->GetClass() == UBlueprint::StaticClass());
+	static FBoolConfigValueHelper UseLocalGraphVariables(TEXT("Kismet"), TEXT("bUseLocalGraphVariables"), GEngineIni);
+	if (bSharedTerm && bRegularBlueprint && UseLocalGraphVariables)
+	{
+		SCOPE_CYCLE_COUNTER(EKismetCompilerStats_ChooseTerminalScope);
+
+		// Pin's connections are checked, to tell if created terminal is shared, or if it could be a local variable.
+		ensure(EEdGraphPinDirection::EGPD_Output == Net->Direction);
+		bSharedTerm = FEventGraphUtils::PinRepresentsSharedTerminal(*Net);
+	}
+	FBPTerminal* Term = new (bSharedTerm ? EventGraphLocals : Locals) FBPTerminal();
+	Term->CopyFromPin(Net, NewName);
+	return Term;
+}
+
 #undef LOCTEXT_NAMESPACE
 
 //////////////////////////////////////////////////////////////////////////
diff --git a/Engine/Source/Editor/KismetCompiler/Public/EdGraphCompilerUtilities.h b/Engine/Source/Editor/KismetCompiler/Public/EdGraphCompilerUtilities.h
--- a/Engine/Source/Editor/KismetCompiler/Public/EdGraphCompilerUtilities.h
+++ b/Engine/Source/Editor/KismetCompiler/Public/EdGraphCompilerUtilities.h
@@ -25,7 +25,7 @@ DECLARE_CYCLE_STAT_EXTERN(TEXT("Postcompile Function"), EKismetCompilerStats_Pos
 DECLARE_CYCLE_STAT_EXTERN(TEXT("Finalization"), EKismetCompilerStats_FinalizationWork, STATGROUP_KismetCompiler, );
 DECLARE_CYCLE_STAT_EXTERN(TEXT("Code Gen"), EKismetCompilerStats_CodeGenerationTime, STATGROUP_KismetCompiler, );
 DECLARE_CYCLE_STAT_EXTERN(TEXT("Update Instances"), EKismetCompilerStats_UpdateBlueprintGeneratedClass, STATGROUP_KismetCompiler, );
-
+DECLARE_CYCLE_STAT_EXTERN(TEXT("Choose Terminal Scope"), EKismetCompilerStats_ChooseTerminalScope, STATGROUP_KismetCompiler, );
 
 //////////////////////////////////////////////////////////////////////////
 
diff --git a/Engine/Source/Editor/KismetCompiler/Public/KismetCompiledFunctionContext.h b/Engine/Source/Editor/KismetCompiler/Public/KismetCompiledFunctionContext.h
--- a/Engine/Source/Editor/KismetCompiler/Public/KismetCompiledFunctionContext.h
+++ b/Engine/Source/Editor/KismetCompiler/Public/KismetCompiledFunctionContext.h
@@ -19,6 +19,13 @@ namespace KismetCompilerDebugOptions
 	enum { EmitNodeComments = DebuggingCompiler };
 }
 
+enum ETerminalSpecification
+{
+	TS_Unspecified,
+	TS_Literal,
+	TS_ForcedShared,
+};
+
 //////////////////////////////////////////////////////////////////////////
 // FKismetFunctionContext
 
@@ -528,6 +535,9 @@ public:
 			}
 		}
 	}
+
+	KISMETCOMPILER_API FBPTerminal* CreateLocalTerminalFromPinAutoChooseScope(UEdGraphPin* Net, const FString& NewName);
+	KISMETCOMPILER_API FBPTerminal* CreateLocalTerminal(ETerminalSpecification Spec = ETerminalSpecification::TS_Unspecified);
 };
 
 //////////////////////////////////////////////////////////////////////////
