Can get the current action key filter from the action registrar.

diff --git a/Engine/Source/Editor/BlueprintGraph/Public/BlueprintActionDatabaseRegistrar.h b/Engine/Source/Editor/BlueprintGraph/Public/BlueprintActionDatabaseRegistrar.h
--- a/Engine/Source/Editor/BlueprintGraph/Public/BlueprintActionDatabaseRegistrar.h
+++ b/Engine/Source/Editor/BlueprintGraph/Public/BlueprintActionDatabaseRegistrar.h
@@ -56,6 +56,12 @@ public:
 	 */
 	bool IsOpenForRegistration(UObject const* OwnerKey);
 
+	/** Returns the current key the registrar is being filtered using */
+	UObject const* GetActionKeyFilter() const
+	{
+		return ActionKeyFilter;
+	}
+
 private:
 	typedef FBlueprintActionDatabase::FActionRegistry			FActionRegistry;
 	typedef FBlueprintActionDatabase::FUnloadedActionRegistry	FUnloadedActionRegistry;
