BP: When FMemberReference doesn't store SKEL_ classes, the CreateDelegate node (created in UK2Node_Event::ExpandNode) cannot find any function created in the same scope. A unnecessary validation was removed (there is another validation, after all UFunctions are created).

diff --git a/Engine/Source/Editor/BlueprintGraph/Private/K2Node_Event.cpp b/Engine/Source/Editor/BlueprintGraph/Private/K2Node_Event.cpp
--- a/Engine/Source/Editor/BlueprintGraph/Private/K2Node_Event.cpp
+++ b/Engine/Source/Editor/BlueprintGraph/Private/K2Node_Event.cpp
@@ -547,13 +547,8 @@ void UK2Node_Event::ExpandNode(class FKismetCompilerContext& CompilerContext, UE
 			CreateDelegateNode->AllocateDefaultPins();
 			CompilerContext.MovePinLinksToIntermediate(*OrgDelegatePin, *CreateDelegateNode->GetDelegateOutPin());
 			Schema->TryCreateConnection(SelfNode->FindPinChecked(Schema->PN_Self), CreateDelegateNode->GetObjectInPin());
+			// When called UFunction is defined in the same class, it wasn't created yet (previously the Skeletal class was checked). So no "CreateDelegateNode->HandleAnyChangeWithoutNotifying();" is called.
 			CreateDelegateNode->SetFunction(FunctionName);
-			CreateDelegateNode->HandleAnyChangeWithoutNotifying();
-			if (CreateDelegateNode->GetFunctionName() != FunctionName)
-			{
-				CreateDelegateNode->SetFunction(FunctionName);
-				CompilerContext.MessageLog.Warning(*FString::Printf(*LOCTEXT("EventDelegateError", "Invalid delegate connection @@. Try recompile.").ToString()), OrgDelegatePin);
-			}
 		}
 	}
 }
