Bubbling up Make/Break node options, when you're dragging off a struct pin.

diff --git a/Engine/Source/Editor/BlueprintGraph/Private/K2Node_BreakStruct.cpp b/Engine/Source/Editor/BlueprintGraph/Private/K2Node_BreakStruct.cpp
--- a/Engine/Source/Editor/BlueprintGraph/Private/K2Node_BreakStruct.cpp
+++ b/Engine/Source/Editor/BlueprintGraph/Private/K2Node_BreakStruct.cpp
@@ -5,6 +5,7 @@
 #include "BlueprintFieldNodeSpawner.h"
 #include "EditorCategoryUtils.h"
 #include "BlueprintActionDatabaseRegistrar.h"
+#include "BlueprintActionFilter.h"	// for FBlueprintActionContext
 
 #define LOCTEXT_NAMESPACE "K2Node_BreakStruct"
 
@@ -321,6 +322,19 @@ void UK2Node_BreakStruct::GetMenuActions(FBlueprintActionDatabaseRegistrar& Acti
 		BreakNode->StructType = NonConstStructPtr.Get();
 	};
 
+	auto CategoryOverrideLambda = [](FBlueprintActionContext const& Context, IBlueprintNodeBinder::FBindingSet const& /*Bindings*/, FBlueprintActionUiSpec* UiSpecOut, TWeakObjectPtr<UScriptStruct> StructPtr)
+	{
+		for (UEdGraphPin* Pin : Context.Pins)
+		{
+			UScriptStruct* PinStruct = Cast<UScriptStruct>(Pin->PinType.PinSubCategoryObject.Get());
+			if ((PinStruct != nullptr) && (StructPtr.Get() == PinStruct) && (Pin->Direction == EGPD_Output))
+			{
+				UiSpecOut->Category = LOCTEXT("EmptyCategory", "|"); 
+				break;
+			}
+		}
+	};
+
 	for (TObjectIterator<UScriptStruct> StructIt; StructIt; ++StructIt)
 	{
 		UScriptStruct const* Struct = (*StructIt);
@@ -342,6 +356,7 @@ void UK2Node_BreakStruct::GetMenuActions(FBlueprintActionDatabaseRegistrar& Acti
 		check(NodeSpawner != nullptr);
 		TWeakObjectPtr<UScriptStruct> NonConstStructPtr = Struct;
 		NodeSpawner->SetNodeFieldDelegate = UBlueprintFieldNodeSpawner::FSetNodeFieldDelegate::CreateStatic(SetNodeStructLambda, NonConstStructPtr);
+		NodeSpawner->DynamicUiSignatureGetter = UBlueprintFieldNodeSpawner::FUiSpecOverrideDelegate::CreateStatic(CategoryOverrideLambda, NonConstStructPtr);
 
 		// this struct could belong to a class, or is a user defined struct 
 		// (asset), that's why we want to make sure to register it along with 
diff --git a/Engine/Source/Editor/BlueprintGraph/Private/K2Node_MakeStruct.cpp b/Engine/Source/Editor/BlueprintGraph/Private/K2Node_MakeStruct.cpp
--- a/Engine/Source/Editor/BlueprintGraph/Private/K2Node_MakeStruct.cpp
+++ b/Engine/Source/Editor/BlueprintGraph/Private/K2Node_MakeStruct.cpp
@@ -6,6 +6,7 @@
 #include "BlueprintFieldNodeSpawner.h"
 #include "EditorCategoryUtils.h"
 #include "BlueprintActionDatabaseRegistrar.h"
+#include "BlueprintActionFilter.h"	// for FBlueprintActionContext
 
 #define LOCTEXT_NAMESPACE "K2Node_MakeStruct"
 
@@ -278,6 +279,19 @@ void UK2Node_MakeStruct::GetMenuActions(FBlueprintActionDatabaseRegistrar& Actio
 		MakeNode->StructType = NonConstStructPtr.Get();
 	};
 
+	auto CategoryOverrideLambda = [](FBlueprintActionContext const& Context, IBlueprintNodeBinder::FBindingSet const& /*Bindings*/, FBlueprintActionUiSpec* UiSpecOut, TWeakObjectPtr<UScriptStruct> StructPtr)
+	{
+		for (UEdGraphPin* Pin : Context.Pins)
+		{
+			UScriptStruct* PinStruct = Cast<UScriptStruct>(Pin->PinType.PinSubCategoryObject.Get());
+			if ((PinStruct != nullptr) && (StructPtr.Get() == PinStruct) && (Pin->Direction == EGPD_Input))
+			{
+				UiSpecOut->Category = LOCTEXT("EmptyCategory", "|");
+				break;
+			}
+		}
+	};
+
 	for (TObjectIterator<UScriptStruct> StructIt; StructIt; ++StructIt)
 	{
 		UScriptStruct const* Struct = (*StructIt);
@@ -299,6 +313,7 @@ void UK2Node_MakeStruct::GetMenuActions(FBlueprintActionDatabaseRegistrar& Actio
 		check(NodeSpawner != nullptr);
 		TWeakObjectPtr<UScriptStruct> NonConstStructPtr = Struct;
 		NodeSpawner->SetNodeFieldDelegate = UBlueprintFieldNodeSpawner::FSetNodeFieldDelegate::CreateStatic(SetNodeStructLambda, NonConstStructPtr);
+		NodeSpawner->DynamicUiSignatureGetter = UBlueprintFieldNodeSpawner::FUiSpecOverrideDelegate::CreateStatic(CategoryOverrideLambda, NonConstStructPtr);
 
 		// this struct could belong to a class, or is a user defined struct 
 		// (asset), that's why we want to make sure to register it along with 
