Missing include for Non-Unity builds.

diff --git a/Engine/Source/Editor/BlueprintGraph/Private/BlueprintComponentNodeSpawner.cpp b/Engine/Source/Editor/BlueprintGraph/Private/BlueprintComponentNodeSpawner.cpp
--- a/Engine/Source/Editor/BlueprintGraph/Private/BlueprintComponentNodeSpawner.cpp
+++ b/Engine/Source/Editor/BlueprintGraph/Private/BlueprintComponentNodeSpawner.cpp
@@ -6,6 +6,7 @@
 #include "ClassIconFinder.h" // for FindIconNameForClass()
 #include "BlueprintNodeTemplateCache.h" // for IsTemplateOuter()
 #include "ComponentAssetBroker.h" // for GetComponentsForAsset()/AssignAssetToComponent()
+#include "BlueprintActionFilter.h"	// for FBlueprintActionContext
 
 #define LOCTEXT_NAMESPACE "BlueprintComponenetNodeSpawner"
 
