Cast nodes to blueprint classes show up in the new context menu. Sometimes we need the GeneratedClass, sometimes we need the SkeletonGeneratedClass...

diff --git a/Engine/Source/Editor/BlueprintGraph/Private/BlueprintActionDatabase.cpp b/Engine/Source/Editor/BlueprintGraph/Private/BlueprintActionDatabase.cpp
--- a/Engine/Source/Editor/BlueprintGraph/Private/BlueprintActionDatabase.cpp
+++ b/Engine/Source/Editor/BlueprintGraph/Private/BlueprintActionDatabase.cpp
@@ -573,8 +573,11 @@ static void BlueprintActionDatabaseImpl::AddClassPropertyActions(UClass const* c
 }
 
 //------------------------------------------------------------------------------
-static void BlueprintActionDatabaseImpl::AddClassCastActions(UClass* const Class, FActionList& ActionListOut)
+static void BlueprintActionDatabaseImpl::AddClassCastActions(UClass* Class, FActionList& ActionListOut)
 {
+	Class = Class->GetAuthoritativeClass();
+	check(Class);
+
 	UEdGraphSchema_K2 const* K2Schema = GetDefault<UEdGraphSchema_K2>();
 	bool bIsCastPermitted  = UEdGraphSchema_K2::IsAllowableBlueprintVariableType(Class);
 
diff --git a/Engine/Source/Editor/BlueprintGraph/Private/BlueprintActionFilter.cpp b/Engine/Source/Editor/BlueprintGraph/Private/BlueprintActionFilter.cpp
--- a/Engine/Source/Editor/BlueprintGraph/Private/BlueprintActionFilter.cpp
+++ b/Engine/Source/Editor/BlueprintGraph/Private/BlueprintActionFilter.cpp
@@ -1280,6 +1280,15 @@ UClass const* FBlueprintActionInfo::GetOwnerClass()
 		{
 			CachedOwnerClass = GetAssociatedMemberField()->GetOwnerClass();
 		}
+
+		if( CachedOwnerClass == nullptr )
+		{
+			if (const UBlueprint* AsBlueprint = Cast<UBlueprint>(ActionOwner))
+			{
+				CachedOwnerClass = AsBlueprint->SkeletonGeneratedClass;
+			}
+		}
+
 		CacheFlags |= EBlueprintActionInfoFlags::CachedClass;
 	}
 	return CachedOwnerClass;
