Dragging off of output BP pin types and placing a variable will autoconnect to the "Target" pin if available in the new menu system.

"Target" pins for variables are no longer different when hidden vs visible

diff --git a/Engine/Source/Editor/BlueprintGraph/Classes/K2Node_Variable.h b/Engine/Source/Editor/BlueprintGraph/Classes/K2Node_Variable.h
--- a/Engine/Source/Editor/BlueprintGraph/Classes/K2Node_Variable.h
+++ b/Engine/Source/Editor/BlueprintGraph/Classes/K2Node_Variable.h
@@ -55,6 +55,7 @@ public:
 	virtual FString GetDocumentationExcerptName() const override;
 	virtual FName GetCornerIcon() const override;
 	virtual FName GetPaletteIcon(FLinearColor& OutColor) const override;
+	virtual void AutowireNewNode(UEdGraphPin* FromPin) override;
 	// End UEdGraphNode interface
 
 	// Begin K2Node interface
diff --git a/Engine/Source/Editor/BlueprintGraph/Private/BlueprintVariableNodeSpawner.cpp b/Engine/Source/Editor/BlueprintGraph/Private/BlueprintVariableNodeSpawner.cpp
--- a/Engine/Source/Editor/BlueprintGraph/Private/BlueprintVariableNodeSpawner.cpp
+++ b/Engine/Source/Editor/BlueprintGraph/Private/BlueprintVariableNodeSpawner.cpp
@@ -70,7 +70,13 @@ UBlueprintVariableNodeSpawner* UBlueprintVariableNodeSpawner::Create(TSubclassOf
 		{
 			UBlueprint* Blueprint = FBlueprintEditorUtils::FindBlueprintForNodeChecked(NewNode);
 			UClass* OwnerClass = Property->GetOwnerClass();
-			bool const bIsSelfContext = Blueprint->SkeletonGeneratedClass->IsChildOf(OwnerClass);
+
+			// We need to use a generated class instead of a skeleton class for IsChildOf, so if the OwnerClass has a Blueprint, grab the GeneratedClass
+			if(OwnerClass)
+			{
+				OwnerClass = OwnerClass->GetAuthoritativeClass();
+			}
+			bool const bIsSelfContext = Blueprint->SkeletonGeneratedClass->GetAuthoritativeClass() == OwnerClass || Blueprint->SkeletonGeneratedClass->IsChildOf(OwnerClass);
 
 			UK2Node_Variable* VarNode = CastChecked<UK2Node_Variable>(NewNode);
 			VarNode->SetFromProperty(Property, bIsSelfContext);
diff --git a/Engine/Source/Editor/BlueprintGraph/Private/K2Node_Variable.cpp b/Engine/Source/Editor/BlueprintGraph/Private/K2Node_Variable.cpp
--- a/Engine/Source/Editor/BlueprintGraph/Private/K2Node_Variable.cpp
+++ b/Engine/Source/Editor/BlueprintGraph/Private/K2Node_Variable.cpp
@@ -101,24 +101,20 @@ bool UK2Node_Variable::CreatePinForVariable(EEdGraphPinDirection Direction)
 void UK2Node_Variable::CreatePinForSelf()
 {
 	const UEdGraphSchema_K2* K2Schema = GetDefault<UEdGraphSchema_K2>();
-	UClass* VariableClass = GetVariableSourceClass();
 	// Create the self pin
 	if (!K2Schema->FindSelfPin(*this, EGPD_Input))
 	{
 		// Do not create a self pin for locally scoped variables
 		if( !VariableReference.IsLocalScope() )
 		{
+			UClass* MemberParentClass = VariableReference.GetMemberParentClass(this);
+			UClass* AuthoritativeClass = MemberParentClass ? MemberParentClass->GetAuthoritativeClass() : NULL;
+			UEdGraphPin* TargetPin = CreatePin(EGPD_Input, K2Schema->PC_Object, TEXT(""), AuthoritativeClass, false, false, K2Schema->PN_Self);
+			TargetPin->PinFriendlyName =  LOCTEXT("Target", "Target");
+
 			if (VariableReference.IsSelfContext() && (ESelfContextInfo::NotSelfContext != SelfContextInfo))
 			{
-				UEdGraphPin* SelfPin = CreatePin(EGPD_Input, K2Schema->PC_Object, K2Schema->PSC_Self, NULL, false, false, K2Schema->PN_Self);
-				SelfPin->bHidden = true; // don't show in 'self' context
-			}
-			else
-			{
-				UClass* MemberParentClass = VariableReference.GetMemberParentClass(this);
-				UClass* AuthoritativeClass = MemberParentClass ? MemberParentClass->GetAuthoritativeClass() : NULL;
-				UEdGraphPin* TargetPin = CreatePin(EGPD_Input, K2Schema->PC_Object, TEXT(""), AuthoritativeClass, false, false, K2Schema->PN_Self);
-				TargetPin->PinFriendlyName =  LOCTEXT("Target", "Target");
+				TargetPin->bHidden = true; // don't show in 'self' context
 			}
 		}
 	}
@@ -549,6 +545,52 @@ FString UK2Node_Variable::GetDocumentationExcerptName() const
 	return GetVarName().ToString();
 }
 
+void UK2Node_Variable::AutowireNewNode(UEdGraphPin* FromPin)
+{
+	const UEdGraphSchema_K2* K2Schema = CastChecked<UEdGraphSchema_K2>(GetSchema());
+
+	// Do some auto-connection
+	if (FromPin != NULL)
+	{
+		bool bConnected = false;
+		if(FromPin->Direction == EGPD_Output)
+		{
+			// If the source pin has a valid PinSubCategoryObject, we might be doing BP Comms, so check if it is a class
+			if(FromPin->PinType.PinSubCategoryObject.IsValid() && FromPin->PinType.PinSubCategoryObject->IsA(UClass::StaticClass()))
+			{
+				UProperty* VariableProperty = VariableReference.ResolveMember<UProperty>(this);
+				if(VariableProperty)
+				{
+					UClass* PropertyOwner = VariableProperty->GetOwnerClass();
+
+					// BP Comms is highly likely at this point, if the source pin's type is a child of the variable's owner class, let's conform the "Target" pin
+					if(FromPin->PinType.PinSubCategoryObject == PropertyOwner || dynamic_cast<UClass*>(FromPin->PinType.PinSubCategoryObject.Get())->IsChildOf(PropertyOwner))
+					{
+						UEdGraphPin* TargetPin = FindPin(K2Schema->PN_Self);
+						TargetPin->PinType.PinSubCategoryObject = PropertyOwner;
+
+						if(K2Schema->TryCreateConnection(FromPin, TargetPin))
+						{
+							bConnected = true;
+
+							// Setup the VariableReference correctly since it may no longer be a self member
+							VariableReference.SetFromField<UProperty>(GetPropertyForVariable(), false);
+							TargetPin->bHidden = false;
+							FromPin->GetOwningNode()->NodeConnectionListChanged();
+							this->NodeConnectionListChanged();
+						}
+					}
+				}
+			}
+		}
+
+		if(!bConnected)
+		{
+			Super::AutowireNewNode(FromPin);
+		}
+	}
+}
+
 FBPVariableDescription const* UK2Node_Variable::GetBlueprintVarDescription() const
 {
 	FName const& VarName = VariableReference.GetMemberName();
