Interface message nodes now exposed to UObject pins in the new BP menu system.

diff --git a/Engine/Source/Editor/BlueprintGraph/Private/BlueprintActionDatabase.cpp b/Engine/Source/Editor/BlueprintGraph/Private/BlueprintActionDatabase.cpp
--- a/Engine/Source/Editor/BlueprintGraph/Private/BlueprintActionDatabase.cpp
+++ b/Engine/Source/Editor/BlueprintGraph/Private/BlueprintActionDatabase.cpp
@@ -163,6 +163,7 @@ static UBlueprintNodeSpawner* FBlueprintNodeSpawnerFactory::MakeMessageNodeSpawn
 	};
 	NodeSpawner->SetNodeFieldDelegate = UBlueprintFunctionNodeSpawner::FSetNodeFieldDelegate::CreateStatic(SetNodeFunctionLambda);
 
+	NodeSpawner->DefaultMenuSignature.MenuName = FText::Format(LOCTEXT("MessageNodeMenuName", "{0} (Message)"), NodeSpawner->DefaultMenuSignature.MenuName);
 	return NodeSpawner;
 }
 
@@ -505,6 +506,9 @@ static void BlueprintActionDatabaseImpl::AddClassFunctionActions(UClass const* c
 
 			if (FKismetEditorUtilities::IsClassABlueprintInterface(Class))
 			{
+				FuncSpawner->DefaultMenuSignature.MenuName = FText::Format(LOCTEXT("InterfaceCallMenuName", "{0} (Interface Call)"), 
+					FuncSpawner->DefaultMenuSignature.MenuName);
+
 				ActionListOut.Add(MakeMessageNodeSpawner(Function));
 			}
 		}
@@ -520,7 +524,7 @@ static void BlueprintActionDatabaseImpl::AddClassPropertyActions(UClass const* c
 	bool const bIsActorClass = Class->IsChildOf<AActor>();
 	
 	// loop over all the properties in the specified class; exclude-super because 
-	// we can always get the super properties by looking up that class separately 
+	// we can always get the super properties by looking up that class separateHavely 
 	for (TFieldIterator<UProperty> PropertyIt(Class, EFieldIteratorFlags::ExcludeSuper); PropertyIt; ++PropertyIt)
 	{
 		UProperty* Property = *PropertyIt;
diff --git a/Engine/Source/Editor/BlueprintGraph/Private/BlueprintActionFilter.cpp b/Engine/Source/Editor/BlueprintGraph/Private/BlueprintActionFilter.cpp
--- a/Engine/Source/Editor/BlueprintGraph/Private/BlueprintActionFilter.cpp
+++ b/Engine/Source/Editor/BlueprintGraph/Private/BlueprintActionFilter.cpp
@@ -307,6 +307,15 @@ namespace BlueprintActionFilterImpl
 	 * @return true if the action is an animnotification that is incompatible with the current skeleton
 	 */
 	static bool IsIncompatibleAnimNotification( FBlueprintActionFilter const& Filter, FBlueprintActionInfo& BlueprintAction); 
+
+	/**
+	 * 
+	 * 
+	 * @param  Filter	
+	 * @param  BlueprintAction	
+	 * @return 
+	 */
+	static bool IsExtraneousInterfaceCall(FBlueprintActionFilter const& Filter, FBlueprintActionInfo& BlueprintAction);
 };
 
 //------------------------------------------------------------------------------
@@ -641,7 +650,7 @@ static bool BlueprintActionFilterImpl::IsNonTargetMember(FBlueprintActionFilter
 
 		// global (and static library) fields can stay (unless explicitly
 		// excluded... save that for a separate test)
-		bool const bSkip = bPermitNonTargetGlobals && IsGloballyAccessible(ClassField);
+		bool const bSkip = (bPermitNonTargetGlobals && IsGloballyAccessible(ClassField)) || BlueprintAction.GetNodeClass()->IsChildOf<UK2Node_Message>();
 		if (!bSkip)
 		{
 			for (UClass const* Class : Filter.TargetClasses)
@@ -923,6 +932,14 @@ static bool BlueprintActionFilterImpl::IsPinCompatibleWithTargetSelf(UEdGraphPin
 {
 	bool bIsCompatible = false;
 	UClass const* TargetClass = BlueprintAction.GetOwnerClass();
+	if (BlueprintAction.GetNodeClass()->IsChildOf<UK2Node_Message>())
+	{
+		// message nodes are a special case, they are intended to call a certain 
+		// function, but will take any arbitrary object (and invokes the function
+		// if that object implements the interface, otherwise the node is passed
+		// through)
+		TargetClass = UObject::StaticClass();
+	}
 
 	if ((Pin->Direction == EGPD_Output) && (TargetClass != nullptr))
 	{
@@ -1175,6 +1192,7 @@ static bool BlueprintActionFilterImpl::IsNodeTemplateSelfFiltered(FBlueprintActi
 	return bIsFilteredOut;
 }
 
+//------------------------------------------------------------------------------
 static bool BlueprintActionFilterImpl::IsIncompatibleAnimNotification(FBlueprintActionFilter const& Filter, FBlueprintActionInfo& BlueprintAction)
 {
 	bool bIsFilteredOut = false;
@@ -1210,6 +1228,58 @@ static bool BlueprintActionFilterImpl::IsIncompatibleAnimNotification(FBlueprint
 	return bIsFilteredOut;
 }
 
+//------------------------------------------------------------------------------
+static bool BlueprintActionFilterImpl::IsExtraneousInterfaceCall(FBlueprintActionFilter const& Filter, FBlueprintActionInfo& BlueprintAction)
+{
+	bool bIsFilteredOut = false;
+
+	if (BlueprintAction.GetNodeClass()->IsChildOf<UK2Node_Message>())
+	{
+		UFunction const* Function = BlueprintAction.GetAssociatedFunction();
+		checkSlow(Function != nullptr);
+
+		UClass* InterfaceClass = Function->GetOwnerClass();
+		checkSlow(InterfaceClass->IsChildOf<UInterface>());
+
+		bIsFilteredOut = (Filter.TargetClasses.Num() > 0);
+		for (const UClass* TargetClass : Filter.TargetClasses)
+		{
+			// if the target class implements (or is) the message spawner's 
+			// interface owner, then we don't need the message node (the regular
+			// interface call shouldt be available)
+			if (!IsClassOfType(TargetClass, InterfaceClass))
+			{
+				bIsFilteredOut = false;
+				break;
+			}
+		}
+	}
+	else if (UFunction const* Function = BlueprintAction.GetAssociatedFunction())
+	{
+		UClass* FuncClass = Function->GetOwnerClass();
+		bool const bIsInterfaceAction = FuncClass->IsChildOf<UInterface>();
+
+		if (bIsInterfaceAction)
+		{
+			bIsFilteredOut = (Filter.TargetClasses.Num() > 0);
+			for (const UClass* TargetClass : Filter.TargetClasses)
+			{
+				bool const bTargetIsBlueprint = (Cast<UBlueprint>(TargetClass->ClassGeneratedBy) != nullptr);
+				// Blueprints fold interface functions in with their other 
+				// functions, and we'd rather those get called (because it is more
+				// optimal, as it avoids a conversion between object/interface)
+				if (!bTargetIsBlueprint || !IsClassOfType(TargetClass, FuncClass))
+				{
+					bIsFilteredOut = false;
+					break;
+				}
+			}
+		}
+	}
+
+	return bIsFilteredOut;
+}
+
 /*******************************************************************************
  * FBlueprintActionInfo
  ******************************************************************************/
@@ -1275,18 +1345,14 @@ UClass const* FBlueprintActionInfo::GetOwnerClass()
 		{
 			CachedOwnerClass = nullptr;
 		}
-
-		if (CachedOwnerClass == nullptr)
+		else if (const UBlueprint* AsBlueprint = Cast<UBlueprint>(ActionOwner))
 		{
-			CachedOwnerClass = GetAssociatedMemberField()->GetOwnerClass();
+			CachedOwnerClass = AsBlueprint->SkeletonGeneratedClass;
 		}
 
-		if( CachedOwnerClass == nullptr )
+		if (CachedOwnerClass == nullptr)
 		{
-			if (const UBlueprint* AsBlueprint = Cast<UBlueprint>(ActionOwner))
-			{
-				CachedOwnerClass = AsBlueprint->SkeletonGeneratedClass;
-			}
+			CachedOwnerClass = GetAssociatedMemberField()->GetOwnerClass();
 		}
 
 		CacheFlags |= EBlueprintActionInfoFlags::CachedClass;
@@ -1393,6 +1459,7 @@ FBlueprintActionFilter::FBlueprintActionFilter(uint32 Flags/*= 0x00*/)
 	AddRejectionTest(FRejectionTestDelegate::CreateStatic(IsRestrictedClassMember));
 	AddRejectionTest(FRejectionTestDelegate::CreateStatic(IsIncompatibleWithGraphType));
 	AddRejectionTest(FRejectionTestDelegate::CreateStatic(IsSchemaIncompatible));
+	AddRejectionTest(FRejectionTestDelegate::CreateStatic(IsExtraneousInterfaceCall));
 
 	if (!(Flags & BPFILTER_PermitDeprecated))
 	{
@@ -1460,6 +1527,11 @@ bool FBlueprintActionFilter::IsFilteredByThis(FBlueprintActionInfo& BlueprintAct
 {
 	FBlueprintActionFilter const& FilterRef = *this;
 
+	if (BlueprintAction.NodeSpawner->DefaultMenuSignature.MenuName.ToString().StartsWith("My"))
+	{
+		printf("");
+	}
+
 	bool bIsFiltered = false;
 	// iterate backwards so that custom user test are ran first (and the slow
 	// internal tests are ran last).
