Bubbling up tagged make/break functions for struct pins (vector, rotator, etc.).

diff --git a/Engine/Source/Editor/BlueprintGraph/Classes/EdGraphSchema_K2.h b/Engine/Source/Editor/BlueprintGraph/Classes/EdGraphSchema_K2.h
--- a/Engine/Source/Editor/BlueprintGraph/Classes/EdGraphSchema_K2.h
+++ b/Engine/Source/Editor/BlueprintGraph/Classes/EdGraphSchema_K2.h
@@ -141,6 +141,11 @@ public:
 
 	/** Metadata to identify an DataTable Pin. Depending on which DataTable is selected, we display different RowName options */
 	static const FName MD_DataTablePin;
+
+	/** Metadata that flags make/break functions for specific struct types. */
+	static const FName MD_NativeMakeFunction;
+	static const FName MD_NativeBreakFunction;
+	
 private:
 	// This class should never be instantiated
 	FBlueprintMetadata() {}
diff --git a/Engine/Source/Editor/BlueprintGraph/Private/BlueprintFunctionNodeSpawner.cpp b/Engine/Source/Editor/BlueprintGraph/Private/BlueprintFunctionNodeSpawner.cpp
--- a/Engine/Source/Editor/BlueprintGraph/Private/BlueprintFunctionNodeSpawner.cpp
+++ b/Engine/Source/Editor/BlueprintGraph/Private/BlueprintFunctionNodeSpawner.cpp
@@ -47,6 +47,16 @@ namespace BlueprintFunctionNodeSpawnerImpl
 	 * @return 
 	 */
 	static FVector2D CalculateBindingPosition(UEdGraphNode* const InputNode);
+
+	/**
+	 * 
+	 * 
+	 * @param  Struct	
+	 * @param  Function	
+	 * @param  OperatorMetaTag	
+	 * @return 
+	 */
+	static bool IsStructOperatorFunc(const UScriptStruct* Struct, const UFunction* Function, FName const OperatorMetaTag);
 }
 
 //------------------------------------------------------------------------------
@@ -131,6 +141,27 @@ static FVector2D BlueprintFunctionNodeSpawnerImpl::CalculateBindingPosition(UEdG
 	return AttachingNodePos;
 }
 
+//------------------------------------------------------------------------------
+static bool BlueprintFunctionNodeSpawnerImpl::IsStructOperatorFunc(const UScriptStruct* Struct, const UFunction* Function, FName const OperatorMetaTag)
+{
+	bool bIsOperatorFunc = false;
+
+	FString NamedOperatorFunction = Struct->GetMetaData(OperatorMetaTag);
+	if (!NamedOperatorFunction.IsEmpty())
+	{
+		UObject* OperatorOuter = nullptr;
+		if (ResolveName(OperatorOuter, NamedOperatorFunction, /*Create =*/false, /*Throw =*/false))
+		{
+			if ((Function->GetOuter() == OperatorOuter) &&
+				(Function->GetName()  == NamedOperatorFunction))
+			{
+				bIsOperatorFunc = true;
+			}
+		}
+	}
+	return bIsOperatorFunc;
+}
+
 /*******************************************************************************
  * UBlueprintFunctionNodeSpawner
  ******************************************************************************/
@@ -254,12 +285,19 @@ void UBlueprintFunctionNodeSpawner::Prime()
 //------------------------------------------------------------------------------
 FBlueprintActionUiSpec UBlueprintFunctionNodeSpawner::GetUiSpec(FBlueprintActionContext const& Context, FBindingSet const& Bindings) const
 {
+	// for FallbackCategory, IsStructOperatorFunc(), etc.
+	using namespace BlueprintFunctionNodeSpawnerImpl;
+
 	UEdGraph* TargetGraph = (Context.Graphs.Num() > 0) ? Context.Graphs[0] : nullptr;
 	FBlueprintActionUiSpec MenuSignature = PrimeDefaultUiSpec(TargetGraph);
 
+	//
+	// stick uncategorized functions in either "Call Function" (for self 
+	// members), or "<ClassName>|..." for external members
+
 	FString const CategoryString = MenuSignature.Category.ToString();
 	// FText compares are slow, so let's use FString compares
-	bool const bIsUncategorized = (CategoryString == BlueprintFunctionNodeSpawnerImpl::FallbackCategory.ToString());
+	bool const bIsUncategorized = (CategoryString == FallbackCategory.ToString());
 	if (bIsUncategorized)
 	{
 		checkSlow(Context.Blueprints.Num() > 0);
@@ -282,11 +320,36 @@ FBlueprintActionUiSpec UBlueprintFunctionNodeSpawner::GetUiSpec(FBlueprintAction
 
 		if (!TargetClass->IsChildOf(FunctionClass))
 		{
-			MenuSignature.Category = FEditorCategoryUtils::BuildCategoryString(FCommonEditorCategory::Class, 
-				FText::FromString(FunctionClass->GetDisplayNameText().ToString() + TEXT("|") + CategoryString));
+			MenuSignature.Category = FEditorCategoryUtils::BuildCategoryString( FCommonEditorCategory::Class, 
+				FText::FromString(FunctionClass->GetDisplayNameText().ToString()) );
 		}
 	}
 
+	//
+	// bubble up important make/break functions (when dragging from a struct pin)
+
+	for (UEdGraphPin* Pin : Context.Pins)
+	{
+		const UScriptStruct* PinStruct = Cast<const UScriptStruct>(Pin->PinType.PinSubCategoryObject.Get());
+		if (PinStruct != nullptr)
+		{
+			UFunction const* WrappedFunction = GetFunction();
+			checkSlow(WrappedFunction != nullptr);
+
+			bool const bIsStructOperator = IsStructOperatorFunc(PinStruct, WrappedFunction, FBlueprintMetadata::MD_NativeBreakFunction) ||
+				IsStructOperatorFunc(PinStruct, WrappedFunction, FBlueprintMetadata::MD_NativeMakeFunction);
+
+			if (bIsStructOperator)
+			{
+				MenuSignature.Category = LOCTEXT("EmptyFunctionCategory", "|");
+				break;
+			}
+		}
+	}
+
+	//
+	// call out functions bound to sub-component (members); give them a unique name
+
 	if (Bindings.Num() == 1)
 	{
 		UObjectProperty const* ObjectProperty = Cast<UObjectProperty>(Bindings.CreateConstIterator()->Get());
diff --git a/Engine/Source/Editor/BlueprintGraph/Private/BlueprintVariableNodeSpawner.cpp b/Engine/Source/Editor/BlueprintGraph/Private/BlueprintVariableNodeSpawner.cpp
--- a/Engine/Source/Editor/BlueprintGraph/Private/BlueprintVariableNodeSpawner.cpp
+++ b/Engine/Source/Editor/BlueprintGraph/Private/BlueprintVariableNodeSpawner.cpp
@@ -187,7 +187,7 @@ FBlueprintActionUiSpec UBlueprintVariableNodeSpawner::GetUiSpec(FBlueprintAction
 		if (!TargetClass->IsChildOf(VariableClass))
 		{
 			MenuSignature.Category = FEditorCategoryUtils::BuildCategoryString(FCommonEditorCategory::Class,
-				FText::FromString(VariableClass->GetDisplayNameText().ToString() + TEXT("|") + MenuSignature.Category.ToString()));
+				FText::FromString(VariableClass->GetDisplayNameText().ToString()));
 		}
 	}
 	DynamicUiSignatureGetter.ExecuteIfBound(Context, Bindings, &MenuSignature);
diff --git a/Engine/Source/Editor/BlueprintGraph/Private/EdGraphSchema_K2.cpp b/Engine/Source/Editor/BlueprintGraph/Private/EdGraphSchema_K2.cpp
--- a/Engine/Source/Editor/BlueprintGraph/Private/EdGraphSchema_K2.cpp
+++ b/Engine/Source/Editor/BlueprintGraph/Private/EdGraphSchema_K2.cpp
@@ -84,6 +84,9 @@ const FName FBlueprintMetadata::MD_CallInEditor(TEXT("CallInEditor"));
 
 const FName FBlueprintMetadata::MD_DataTablePin(TEXT("DataTablePin"));
 
+const FName FBlueprintMetadata::MD_NativeMakeFunction(TEXT("HasNativeMake"));
+const FName FBlueprintMetadata::MD_NativeBreakFunction(TEXT("HasNativeBreak"));
+
 //////////////////////////////////////////////////////////////////////////
 
 #define LOCTEXT_NAMESPACE "KismetSchema"
