Fixing it so self var nodes can be autowired in as the target pin for other getters.

diff --git a/Engine/Source/Editor/BlueprintGraph/Private/K2Node_Variable.cpp b/Engine/Source/Editor/BlueprintGraph/Private/K2Node_Variable.cpp
--- a/Engine/Source/Editor/BlueprintGraph/Private/K2Node_Variable.cpp
+++ b/Engine/Source/Editor/BlueprintGraph/Private/K2Node_Variable.cpp
@@ -562,6 +562,10 @@ void UK2Node_Variable::AutowireNewNode(UEdGraphPin* FromPin)
 				if(VariableProperty)
 				{
 					UClass* PropertyOwner = VariableProperty->GetOwnerClass();
+					if (PropertyOwner != nullptr)
+					{
+						PropertyOwner = PropertyOwner->GetAuthoritativeClass();
+					}
 
 					// BP Comms is highly likely at this point, if the source pin's type is a child of the variable's owner class, let's conform the "Target" pin
 					if(FromPin->PinType.PinSubCategoryObject == PropertyOwner || dynamic_cast<UClass*>(FromPin->PinType.PinSubCategoryObject.Get())->IsChildOf(PropertyOwner))
