New blueprint context menu will no longer allow you to create multiple matinee controllers for one asset in the level blueprint. Behavior is slightly different from old context menu, when you attempt to create the second controller the blueprint editor will highlight it (just as if you tried to create a second event node)

diff --git a/Engine/Source/Editor/BlueprintGraph/Classes/BlueprintBoundNodeSpawner.h b/Engine/Source/Editor/BlueprintGraph/Classes/BlueprintBoundNodeSpawner.h
--- a/Engine/Source/Editor/BlueprintGraph/Classes/BlueprintBoundNodeSpawner.h
+++ b/Engine/Source/Editor/BlueprintGraph/Classes/BlueprintBoundNodeSpawner.h
@@ -15,6 +15,7 @@ class BLUEPRINTGRAPH_API UBlueprintBoundNodeSpawner : public UBlueprintNodeSpawn
 	GENERATED_UCLASS_BODY()
 	DECLARE_DELEGATE_RetVal_OneParam(bool, FCanBindObjectDelegate, UObject const*);
 	DECLARE_DELEGATE_RetVal_TwoParams(bool, FOnBindObjectDelegate, UEdGraphNode*, UObject*);
+	DECLARE_DELEGATE_RetVal_TwoParams(UEdGraphNode*, FFindPreExistingNodeDelegate, const UBlueprint*, IBlueprintNodeBinder::FBindingSet const& );
 
 public:
 	/**
@@ -24,6 +25,7 @@ public:
 
 	// UBlueprintNodeSpawner interface
 	virtual FBlueprintNodeSignature GetSpawnerSignature() const override;
+	virtual UEdGraphNode* Invoke(UEdGraph* ParentGraph, FBindingSet const& Bindings, FVector2D const Location) const override;
 	// End UBlueprintNodeSpawner interface
 
 	// FBlueprintNodeBinder interface
@@ -41,4 +43,10 @@ public:
 	 * A delegate to perform specialized node setup during binding
 	 */
 	FOnBindObjectDelegate OnBindObjectDelegate;
+
+	/**
+	 * A delegate to find a node that is already spawned, instead of spawning a node
+	 * we will focus on the pre-existing node found by the delegate
+	 */
+	FFindPreExistingNodeDelegate FindPreExistingNodeDelegate;
 };
diff --git a/Engine/Source/Editor/BlueprintGraph/Private/BlueprintBoundNodeSpawner.cpp b/Engine/Source/Editor/BlueprintGraph/Private/BlueprintBoundNodeSpawner.cpp
--- a/Engine/Source/Editor/BlueprintGraph/Private/BlueprintBoundNodeSpawner.cpp
+++ b/Engine/Source/Editor/BlueprintGraph/Private/BlueprintBoundNodeSpawner.cpp
@@ -38,6 +38,22 @@ FBlueprintNodeSignature UBlueprintBoundNodeSpawner::GetSpawnerSignature() const
 	return FBlueprintNodeSignature();
 }
 
+//------------------------------------------------------------------------------
+UEdGraphNode* UBlueprintBoundNodeSpawner::Invoke(UEdGraph* ParentGraph, FBindingSet const& Bindings, FVector2D const Location) const
+{
+	UEdGraphNode* Result = nullptr;
+	if( FindPreExistingNodeDelegate.IsBound() )
+	{
+		UBlueprint* Blueprint = FBlueprintEditorUtils::FindBlueprintForGraphChecked(ParentGraph);
+		Result = FindPreExistingNodeDelegate.Execute( Blueprint, Bindings );
+	}
+	if( Result == nullptr)
+	{
+		Result = UBlueprintNodeSpawner::Invoke( ParentGraph, Bindings, Location );
+	}
+	return Result;
+}
+
 //------------------------------------------------------------------------------
 bool UBlueprintBoundNodeSpawner::IsBindingCompatible(UObject const* BindingCandidate) const
 {
diff --git a/Engine/Source/Editor/BlueprintGraph/Private/K2Node_MatineeController.cpp b/Engine/Source/Editor/BlueprintGraph/Private/K2Node_MatineeController.cpp
--- a/Engine/Source/Editor/BlueprintGraph/Private/K2Node_MatineeController.cpp
+++ b/Engine/Source/Editor/BlueprintGraph/Private/K2Node_MatineeController.cpp
@@ -3,6 +3,7 @@
 #include "BlueprintGraphPrivatePCH.h"
 
 #include "BlueprintActionDatabaseRegistrar.h"
+#include "BlueprintActionFilter.h"
 #include "BlueprintBoundNodeSpawner.h"
 #include "KismetCompiler.h"
 #include "Runtime/Engine/Public/MatineeDelegates.h"
@@ -167,10 +168,25 @@ void UK2Node_MatineeController::GetMenuActions(FBlueprintActionDatabaseRegistrar
 		return true;
 	};
 
+	auto FindPreExistingNodeLambda = []( const UBlueprint* Blueprint, IBlueprintNodeBinder::FBindingSet const& BindingSet ) -> UEdGraphNode*
+	{
+		TArray<UK2Node_MatineeController*> ExistingMatineeControllers;
+		FBlueprintEditorUtils::GetAllNodesOfClass(Blueprint, ExistingMatineeControllers);
+		for (auto Controller : ExistingMatineeControllers)
+		{
+			if (BindingSet.Contains( Controller->MatineeActor ))
+			{
+				return Controller;
+			}
+		}
+		return nullptr;
+	};
+
 	UBlueprintBoundNodeSpawner* NodeSpawner = UBlueprintBoundNodeSpawner::Create(GetClass());
 	NodeSpawner->CanBindObjectDelegate = UBlueprintBoundNodeSpawner::FCanBindObjectDelegate::CreateStatic(CanBindObjectLambda);
 	NodeSpawner->OnBindObjectDelegate  = UBlueprintBoundNodeSpawner::FOnBindObjectDelegate::CreateStatic(PostBindSetupLambda);
 	NodeSpawner->DynamicUiSignatureGetter = UBlueprintBoundNodeSpawner::FUiSpecOverrideDelegate::CreateStatic(UiSpecOverride);
+	NodeSpawner->FindPreExistingNodeDelegate = UBlueprintBoundNodeSpawner::FFindPreExistingNodeDelegate::CreateStatic( FindPreExistingNodeLambda );
 	ActionRegistrar.AddBlueprintAction(NodeSpawner);
 }
 
