BP: SpawnActorFromClass - can use class type from link

diff --git a/Engine/Source/Editor/BlueprintGraph/Classes/K2Node_SpawnActorFromClass.h b/Engine/Source/Editor/BlueprintGraph/Classes/K2Node_SpawnActorFromClass.h
--- a/Engine/Source/Editor/BlueprintGraph/Classes/K2Node_SpawnActorFromClass.h
+++ b/Engine/Source/Editor/BlueprintGraph/Classes/K2Node_SpawnActorFromClass.h
@@ -16,6 +16,7 @@ class UK2Node_SpawnActorFromClass : public UK2Node
 	virtual FLinearColor GetNodeTitleColor() const override;
 	virtual FText GetNodeTitle(ENodeTitleType::Type TitleType) const override;
 	virtual void PinDefaultValueChanged(UEdGraphPin* Pin) override;
+	virtual void PinConnectionListChanged(UEdGraphPin* Pin) override;
 	virtual FText GetTooltipText() const override;
 	virtual void ExpandNode(class FKismetCompilerContext& CompilerContext, UEdGraph* SourceGraph) override;
 	virtual bool HasExternalBlueprintDependencies(TArray<class UStruct*>* OptionalOutput) const override;
@@ -55,6 +56,9 @@ class UK2Node_SpawnActorFromClass : public UK2Node
 	BLUEPRINTGRAPH_API UClass* GetClassToSpawn(const TArray<UEdGraphPin*>* InPinsToSearch=NULL) const;
 
 private:
+
+	void OnClassPinChanged();
+
 	/** Tooltip text for this node. */
 	FText NodeTooltip;
 
diff --git a/Engine/Source/Editor/BlueprintGraph/Private/K2Node_SpawnActorFromClass.cpp b/Engine/Source/Editor/BlueprintGraph/Private/K2Node_SpawnActorFromClass.cpp
--- a/Engine/Source/Editor/BlueprintGraph/Private/K2Node_SpawnActorFromClass.cpp
+++ b/Engine/Source/Editor/BlueprintGraph/Private/K2Node_SpawnActorFromClass.cpp
@@ -113,6 +113,11 @@ UClass* UK2Node_SpawnActorFromClass::GetClassToSpawn(const TArray<UEdGraphPin*>*
 	{
 		UseSpawnClass = CastChecked<UClass>(ClassPin->DefaultObject);
 	}
+	else if (ClassPin && (1 == ClassPin->LinkedTo.Num()))
+	{
+		auto SourcePin = ClassPin->LinkedTo[0];
+		UseSpawnClass = SourcePin ? Cast<UClass>(SourcePin->PinType.PinSubCategoryObject.Get()) : NULL;
+	}
 
 	return UseSpawnClass;
 }
@@ -150,43 +155,55 @@ bool UK2Node_SpawnActorFromClass::IsSpawnVarPin(UEdGraphPin* Pin)
 			Pin->PinName != FK2Node_SpawnActorFromClassHelper::SpawnTransformPinName );
 }
 
-
-void UK2Node_SpawnActorFromClass::PinDefaultValueChanged(UEdGraphPin* ChangedPin) 
+void UK2Node_SpawnActorFromClass::OnClassPinChanged()
 {
-	if (ChangedPin->PinName == FK2Node_SpawnActorFromClassHelper::ClassPinName)
-	{
-		const UEdGraphSchema_K2* K2Schema = GetDefault<UEdGraphSchema_K2>();
+	const UEdGraphSchema_K2* K2Schema = GetDefault<UEdGraphSchema_K2>();
 
-		// Because the archetype has changed, we break the output link as the output pin type will change
-		UEdGraphPin* ResultPin = GetResultPin();
-		ResultPin->BreakAllPinLinks();
+	// Because the archetype has changed, we break the output link as the output pin type will change
+	UEdGraphPin* ResultPin = GetResultPin();
+	ResultPin->BreakAllPinLinks();
 
-		// Remove all pins related to archetype variables
-		TArray<UEdGraphPin*> OldPins = Pins;
-		for (int32 i = 0; i < OldPins.Num(); i++)
+	// Remove all pins related to archetype variables
+	TArray<UEdGraphPin*> OldPins = Pins;
+	for (int32 i = 0; i < OldPins.Num(); i++)
+	{
+		UEdGraphPin* OldPin = OldPins[i];
+		if (IsSpawnVarPin(OldPin))
 		{
-			UEdGraphPin* OldPin = OldPins[i];
-			if (IsSpawnVarPin(OldPin))
-			{
-				OldPin->BreakAllPinLinks();
-				Pins.Remove(OldPin);
-			}
+			OldPin->BreakAllPinLinks();
+			Pins.Remove(OldPin);
 		}
+	}
 
-		CachedNodeTitle.MarkDirty();
+	CachedNodeTitle.MarkDirty();
 
-		UClass* UseSpawnClass = GetClassToSpawn();
-		if (UseSpawnClass != NULL)
-		{
-			CreatePinsForClass(UseSpawnClass);
-		}
+	UClass* UseSpawnClass = GetClassToSpawn();
+	if (UseSpawnClass != NULL)
+	{
+		CreatePinsForClass(UseSpawnClass);
+	}
 
-		// Refresh the UI for the graph so the pin changes show up
-		UEdGraph* Graph = GetGraph();
-		Graph->NotifyGraphChanged();
+	// Refresh the UI for the graph so the pin changes show up
+	UEdGraph* Graph = GetGraph();
+	Graph->NotifyGraphChanged();
 
-		// Mark dirty
-		FBlueprintEditorUtils::MarkBlueprintAsModified(GetBlueprint());
+	// Mark dirty
+	FBlueprintEditorUtils::MarkBlueprintAsModified(GetBlueprint());
+}
+
+void UK2Node_SpawnActorFromClass::PinConnectionListChanged(UEdGraphPin* ChangedPin)
+{
+	if (ChangedPin && (ChangedPin->PinName == FK2Node_SpawnActorFromClassHelper::ClassPinName))
+	{
+		OnClassPinChanged();
+	}
+}
+
+void UK2Node_SpawnActorFromClass::PinDefaultValueChanged(UEdGraphPin* ChangedPin) 
+{
+	if (ChangedPin && (ChangedPin->PinName == FK2Node_SpawnActorFromClassHelper::ClassPinName))
+	{
+		OnClassPinChanged();
 	}
 }
 
