Blueprint - We now support making FText an exposed on spawn parameter.

diff --git a/Engine/Source/Editor/BlueprintGraph/Private/EdGraphSchema_K2.cpp b/Engine/Source/Editor/BlueprintGraph/Private/EdGraphSchema_K2.cpp
--- a/Engine/Source/Editor/BlueprintGraph/Private/EdGraphSchema_K2.cpp
+++ b/Engine/Source/Editor/BlueprintGraph/Private/EdGraphSchema_K2.cpp
@@ -3473,6 +3473,11 @@ UFunction* UEdGraphSchema_K2::FindSetVariableByNameFunction(const FEdGraphPinTyp
 		static FName SetStringName(GET_FUNCTION_NAME_CHECKED(UKismetSystemLibrary, SetStringPropertyByName));
 		SetFunctionName = SetStringName;
 	}
+	else if ( PinType.PinCategory == K2Schema->PC_Text)
+	{
+		static FName SetTextName(GET_FUNCTION_NAME_CHECKED(UKismetSystemLibrary, SetTextPropertyByName));
+		SetFunctionName = SetTextName;
+	}
 	else if(PinType.PinCategory == K2Schema->PC_Name)
 	{
 		static FName SetNameName(GET_FUNCTION_NAME_CHECKED(UKismetSystemLibrary, SetNamePropertyByName));
diff --git a/Engine/Source/Editor/BlueprintGraph/Private/K2Node_ConstructObjectFromClass.cpp b/Engine/Source/Editor/BlueprintGraph/Private/K2Node_ConstructObjectFromClass.cpp
--- a/Engine/Source/Editor/BlueprintGraph/Private/K2Node_ConstructObjectFromClass.cpp
+++ b/Engine/Source/Editor/BlueprintGraph/Private/K2Node_ConstructObjectFromClass.cpp
@@ -166,6 +166,8 @@ void UK2Node_ConstructObjectFromClass::PinDefaultValueChanged(UEdGraphPin* Chang
 			}
 		}
 
+		CachedNodeTitle.MarkDirty();
+
 		UClass* UseSpawnClass = GetClassToSpawn();
 		if (UseSpawnClass != NULL)
 		{
diff --git a/Engine/Source/Programs/UnrealHeaderTool/Private/HeaderParser.cpp b/Engine/Source/Programs/UnrealHeaderTool/Private/HeaderParser.cpp
--- a/Engine/Source/Programs/UnrealHeaderTool/Private/HeaderParser.cpp
+++ b/Engine/Source/Programs/UnrealHeaderTool/Private/HeaderParser.cpp
@@ -5892,6 +5892,7 @@ struct FExposeOnSpawnValidator
 		case CPT_Bool:
 		case CPT_ObjectReference:
 		case CPT_String:
+		case CPT_Text:
 		case CPT_Name:
 		case CPT_Vector:
 		case CPT_Rotation:
