Making vector axis getter nodes distinguisable (and searchable with the phrase "get").
TTP #343490

diff --git a/Engine/Source/Editor/BlueprintGraph/Private/K2Node_GetInputVectorAxisValue.cpp b/Engine/Source/Editor/BlueprintGraph/Private/K2Node_GetInputVectorAxisValue.cpp
--- a/Engine/Source/Editor/BlueprintGraph/Private/K2Node_GetInputVectorAxisValue.cpp
+++ b/Engine/Source/Editor/BlueprintGraph/Private/K2Node_GetInputVectorAxisValue.cpp
@@ -97,6 +97,7 @@ void UK2Node_GetInputVectorAxisValue::GetMenuActions(FBlueprintActionDatabaseReg
 
 		UBlueprintNodeSpawner* NodeSpawner = UBlueprintNodeSpawner::Create(GetClass());
 		check(NodeSpawner != nullptr);
+		NodeSpawner->DefaultMenuSignature.MenuName = FText::Format(NSLOCTEXT("K2Node_GetInputVectorAxisValue", "MenuName", "Get {0}"), Key.GetDisplayName());
 
 		NodeSpawner->CustomizeNodeDelegate = UBlueprintNodeSpawner::FCustomizeNodeDelegate::CreateStatic(CustomizeInputNodeLambda, Key);
 		ActionRegistrar.AddBlueprintAction(ActionKey, NodeSpawner);
