Fixing it so delegate spawners are distinguishable in the BP context menu.

diff --git a/Engine/Source/Editor/BlueprintGraph/Private/BlueprintDelegateNodeSpawner.cpp b/Engine/Source/Editor/BlueprintGraph/Private/BlueprintDelegateNodeSpawner.cpp
--- a/Engine/Source/Editor/BlueprintGraph/Private/BlueprintDelegateNodeSpawner.cpp
+++ b/Engine/Source/Editor/BlueprintGraph/Private/BlueprintDelegateNodeSpawner.cpp
@@ -74,10 +74,10 @@ UBlueprintDelegateNodeSpawner* UBlueprintDelegateNodeSpawner::Create(TSubclassOf
 	//--------------------------------------
 
 	FBlueprintActionUiSpec& MenuSignature = NodeSpawner->DefaultMenuSignature;
-	MenuSignature.MenuName = BlueprintDelegateNodeSpawnerImpl::GetDefaultMenuName(Property);
+	//MenuSignature.MenuName, will be pulled from the node template
 	MenuSignature.Category = BlueprintDelegateNodeSpawnerImpl::GetDefaultMenuCategory(Property);
-	//MenuSignature.Tooltip  =
-	//MenuSignature.Keywords = 
+	//MenuSignature.Tooltip,  will be pulled from the node template
+	//MenuSignature.Keywords, will be pulled from the node template
 	MenuSignature.IconName = BlueprintDelegateNodeSpawnerImpl::GetDefaultMenuIcon(Property, MenuSignature.IconTint);
 
 	//--------------------------------------
