Fixes for variable target pin to work when "self" variables are included via an external macro blueprint.

#ttp 348416 - CRITICAL: Regression: SHADOW: BP: ERRORS: Several blueprints are failing to compile

diff --git a/Engine/Source/Editor/BlueprintGraph/Private/K2Node_Variable.cpp b/Engine/Source/Editor/BlueprintGraph/Private/K2Node_Variable.cpp
--- a/Engine/Source/Editor/BlueprintGraph/Private/K2Node_Variable.cpp
+++ b/Engine/Source/Editor/BlueprintGraph/Private/K2Node_Variable.cpp
@@ -107,12 +107,25 @@ void UK2Node_Variable::CreatePinForSelf()
 		// Do not create a self pin for locally scoped variables
 		if( !VariableReference.IsLocalScope() )
 		{
+			bool bSelfTarget = VariableReference.IsSelfContext() && (ESelfContextInfo::NotSelfContext != SelfContextInfo);
 			UClass* MemberParentClass = VariableReference.GetMemberParentClass(this);
-			UClass* AuthoritativeClass = MemberParentClass ? MemberParentClass->GetAuthoritativeClass() : NULL;
-			UEdGraphPin* TargetPin = CreatePin(EGPD_Input, K2Schema->PC_Object, TEXT(""), AuthoritativeClass, false, false, K2Schema->PN_Self);
+			UClass* TargetClass = nullptr;
+			
+			// Self Target pins should always make the class be the owning class of the property,
+			// so if the node is from a Macro Blueprint, it will hook up as self in any placed Blueprint
+			if(bSelfTarget)
+			{
+				TargetClass = VariableReference.ResolveMember<UProperty>(this)->GetOwnerClass()->GetAuthoritativeClass();
+			}
+			else
+			{
+				TargetClass = MemberParentClass ? MemberParentClass->GetAuthoritativeClass() : NULL;
+			}
+
+			UEdGraphPin* TargetPin = CreatePin(EGPD_Input, K2Schema->PC_Object, TEXT(""), TargetClass, false, false, K2Schema->PN_Self);
 			TargetPin->PinFriendlyName =  LOCTEXT("Target", "Target");
 
-			if (VariableReference.IsSelfContext() && (ESelfContextInfo::NotSelfContext != SelfContextInfo))
+			if (bSelfTarget)
 			{
 				TargetPin->bHidden = true; // don't show in 'self' context
 			}
