Fixing up the name/icon we show for actor reference nodes in the context menu.
TTP #348515/#348514

diff --git a/Engine/Source/Editor/BlueprintGraph/Private/K2Node_Literal.cpp b/Engine/Source/Editor/BlueprintGraph/Private/K2Node_Literal.cpp
--- a/Engine/Source/Editor/BlueprintGraph/Private/K2Node_Literal.cpp
+++ b/Engine/Source/Editor/BlueprintGraph/Private/K2Node_Literal.cpp
@@ -175,13 +175,39 @@ void UK2Node_Literal::GetMenuActions(FBlueprintActionDatabaseRegistrar& ActionRe
 	{
 		if (Bindings.Num() == 1)
 		{
-			UiSpecOut->MenuName = FText::Format(NSLOCTEXT("K2Node", "LiteralTitle", "Create a Reference to {0}"), 
-				FText::FromString( (*(Bindings.CreateConstIterator()))->GetName() ));
+			const AActor* ActorObj = CastChecked<AActor>(Bindings.CreateConstIterator()->Get());
+
+			UiSpecOut->MenuName = FText::Format( NSLOCTEXT("K2Node", "LiteralTitle", "Create a Reference to {0}"), 
+				FText::FromString(ActorObj->GetActorLabel()) );
+
+			const FName IconName = FClassIconFinder::FindIconNameForActor(ActorObj);
+			if (!IconName.IsNone())
+			{
+				UiSpecOut->IconName = IconName;
+			}
 		}
 		else if (Bindings.Num() > 1)
 		{
 			UiSpecOut->MenuName = FText::Format(NSLOCTEXT("K2Node", "FallbackLiteralTitle", "Create References to {0} selected Actors"), 
-				FText::AsNumber(Bindings.Num()));
+				FText::AsNumber(Bindings.Num()) );
+
+			auto BindingIt = Bindings.CreateConstIterator();
+
+			UClass* CommonClass = BindingIt->Get()->GetClass();
+			for (++BindingIt; BindingIt; ++BindingIt)
+			{
+				UClass* Class = BindingIt->Get()->GetClass();
+				while (!Class->IsChildOf(CommonClass))
+				{
+					CommonClass = CommonClass->GetSuperClass();
+				}
+			}
+
+			const FName IconName = FClassIconFinder::FindIconNameForClass(CommonClass);
+			if (!IconName.IsNone())
+			{
+				UiSpecOut->IconName = IconName;
+			}
 		}
 		else
 		{
@@ -190,9 +216,10 @@ void UK2Node_Literal::GetMenuActions(FBlueprintActionDatabaseRegistrar& ActionRe
 	};
 
 	UBlueprintBoundNodeSpawner* NodeSpawner = UBlueprintBoundNodeSpawner::Create(GetClass());
-	NodeSpawner->CanBindObjectDelegate = UBlueprintBoundNodeSpawner::FCanBindObjectDelegate::CreateStatic(CanBindObjectLambda);
-	NodeSpawner->OnBindObjectDelegate = UBlueprintBoundNodeSpawner::FOnBindObjectDelegate::CreateStatic(PostBindSetupLambda);
+	NodeSpawner->CanBindObjectDelegate    = UBlueprintBoundNodeSpawner::FCanBindObjectDelegate::CreateStatic(CanBindObjectLambda);
+	NodeSpawner->OnBindObjectDelegate     = UBlueprintBoundNodeSpawner::FOnBindObjectDelegate::CreateStatic(PostBindSetupLambda);
 	NodeSpawner->DynamicUiSignatureGetter = UBlueprintBoundNodeSpawner::FUiSpecOverrideDelegate::CreateStatic(UiSpecOverride);
+
 	ActionRegistrar.AddBlueprintAction(NodeSpawner);
 }
 
