Not offering hidden category options for bound function/events.
TTP #348564

diff --git a/Engine/Source/Editor/BlueprintGraph/Private/BlueprintActionFilter.cpp b/Engine/Source/Editor/BlueprintGraph/Private/BlueprintActionFilter.cpp
--- a/Engine/Source/Editor/BlueprintGraph/Private/BlueprintActionFilter.cpp
+++ b/Engine/Source/Editor/BlueprintGraph/Private/BlueprintActionFilter.cpp
@@ -703,6 +703,13 @@ static bool BlueprintActionFilterImpl::IsFieldCategoryHidden(FBlueprintActionFil
 				break;
 			}
 		}
+
+		// handled now in each IsBindingCompatible() check
+// 		for (auto BindingIt = BlueprintAction.GetBindings().CreateConstIterator(); BindingIt && !bIsFilteredOut; ++BindingIt)
+// 		{
+// 			UClass* BindingClass = FBlueprintNodeSpawnerUtils::GetBindingClass(BindingIt->Get());
+// 			bIsFilteredOut = IsFieldHiddenDelegate.Execute(BindingClass);
+// 		}
 	}
 
 	return bIsFilteredOut;
@@ -1391,9 +1398,12 @@ UProperty const* FBlueprintActionInfo::GetAssociatedProperty()
 		else
 		{
 			CachedActionProperty = FBlueprintNodeSpawnerUtils::GetAssociatedProperty(NodeSpawner);
-			CachedActionField = CachedActionProperty;
+			if (CachedActionProperty != nullptr)
+			{
+				CachedActionField = CachedActionProperty;
+				CacheFlags |= EBlueprintActionInfoFlags::CachedProperty;
+			}
 		}
-		CacheFlags |= (EBlueprintActionInfoFlags::CachedProperty | EBlueprintActionInfoFlags::CachedField);
 	}
 	return CachedActionProperty;
 }
@@ -1410,9 +1420,13 @@ UFunction const* FBlueprintActionInfo::GetAssociatedFunction()
 		else
 		{
 			CachedActionFunction = FBlueprintNodeSpawnerUtils::GetAssociatedFunction(NodeSpawner);
-			CachedActionField = CachedActionFunction;
+			if (CachedActionFunction != nullptr)
+			{
+				CachedActionField = CachedActionFunction;
+				CacheFlags |= EBlueprintActionInfoFlags::CachedProperty;
+			}
 		}
-		CacheFlags |= (EBlueprintActionInfoFlags::CachedFunction | EBlueprintActionInfoFlags::CachedField);
+		CacheFlags |= EBlueprintActionInfoFlags::CachedFunction;
 	}
 	return CachedActionFunction;
 }
diff --git a/Engine/Source/Editor/BlueprintGraph/Private/BlueprintBoundEventNodeSpawner.cpp b/Engine/Source/Editor/BlueprintGraph/Private/BlueprintBoundEventNodeSpawner.cpp
--- a/Engine/Source/Editor/BlueprintGraph/Private/BlueprintBoundEventNodeSpawner.cpp
+++ b/Engine/Source/Editor/BlueprintGraph/Private/BlueprintBoundEventNodeSpawner.cpp
@@ -2,9 +2,10 @@
 
 #include "BlueprintGraphPrivatePCH.h"
 #include "BlueprintBoundEventNodeSpawner.h"
-#include "ObjectEditorUtils.h"		// for GetCategory()
+#include "ObjectEditorUtils.h"		// for GetCategory()/IsVariableCategoryHiddenFromClass()
 #include "KismetEditorUtilities.h"	// for FindBoundEventForComponent()/FindBoundEventForActor()
 #include "EditorCategoryUtils.h"	// for GetCommonCategory()
+#include "BlueprintNodeSpawnerUtils.h" // for GetBindingClass()
 
 #define LOCTEXT_NAMESPACE "BlueprintBoundEventNodeSpawner"
 
@@ -119,27 +120,22 @@ UK2Node_Event const* UBlueprintBoundEventNodeSpawner::FindPreExistingEvent(UBlue
 //------------------------------------------------------------------------------
 bool UBlueprintBoundEventNodeSpawner::IsBindingCompatible(UObject const* BindingCandidate) const
 {
-	bool bCanBind = false;
-
-	UClass* DelegateOwner = nullptr;
-	if (UMulticastDelegateProperty const* Delegate = GetEventDelegate())
-	{
-		DelegateOwner = Delegate->GetOwnerClass();
-	}
-
+	bool bMatchesNodeType = false;
 	if (NodeClass->IsChildOf<UK2Node_ComponentBoundEvent>())
 	{
 		UObjectProperty const* BindingProperty = Cast<UObjectProperty>(BindingCandidate);
-
-		bCanBind = (BindingProperty != nullptr) && (DelegateOwner != nullptr) &&
-			BindingProperty->PropertyClass->IsChildOf(DelegateOwner);
+		bMatchesNodeType = (BindingProperty != nullptr);
 	}
 	else if (NodeClass->IsChildOf<UK2Node_ActorBoundEvent>())
 	{
-		bCanBind = BindingCandidate->IsA<AActor>() && (DelegateOwner != nullptr) &&
-			BindingCandidate->GetClass()->IsChildOf(DelegateOwner);
+		bMatchesNodeType = BindingCandidate->IsA<AActor>();
 	}
-	return bCanBind;
+
+	const UMulticastDelegateProperty* Delegate = GetEventDelegate();
+	UClass* DelegateOwner = Delegate->GetOwnerClass();
+	UClass* BindingClass  = FBlueprintNodeSpawnerUtils::GetBindingClass(BindingCandidate);
+
+	return bMatchesNodeType && BindingClass->IsChildOf(DelegateOwner) && !FObjectEditorUtils::IsVariableCategoryHiddenFromClass(Delegate, BindingClass);
 }
 
 //------------------------------------------------------------------------------
diff --git a/Engine/Source/Editor/BlueprintGraph/Private/BlueprintFunctionNodeSpawner.cpp b/Engine/Source/Editor/BlueprintGraph/Private/BlueprintFunctionNodeSpawner.cpp
--- a/Engine/Source/Editor/BlueprintGraph/Private/BlueprintFunctionNodeSpawner.cpp
+++ b/Engine/Source/Editor/BlueprintGraph/Private/BlueprintFunctionNodeSpawner.cpp
@@ -6,8 +6,10 @@
 #include "GameFramework/Actor.h"
 #include "BlueprintVariableNodeSpawner.h"
 #include "BlueprintNodeTemplateCache.h" // for IsTemplateOuter()
-#include "BlueprintActionFilter.h" // for FBlueprintActionContext
-#include "EditorCategoryUtils.h" // for BuildCategoryString()
+#include "BlueprintActionFilter.h"	// for FBlueprintActionContext
+#include "EditorCategoryUtils.h"	// for BuildCategoryString()
+#include "ObjectEditorUtils.h"		// for IsFunctionHiddenFromClass()
+#include "BlueprintNodeSpawnerUtils.h" // for GetBindingClass()
 
 #define LOCTEXT_NAMESPACE "BlueprintFunctionNodeSpawner"
 
@@ -397,24 +399,19 @@ bool UBlueprintFunctionNodeSpawner::CanBindMultipleObjects() const
 //------------------------------------------------------------------------------
 bool UBlueprintFunctionNodeSpawner::IsBindingCompatible(UObject const* BindingCandidate) const
 {
-	bool bCanBind = false;
-
 	UFunction const* Function = GetFunction();
-	if ((Function != nullptr) && (NodeClass == UK2Node_CallFunction::StaticClass()))
-	{
-		UClass* BindingClass = BindingCandidate->GetClass();
-		if (UObjectProperty const* ObjectProperty = Cast<UObjectProperty>(BindingCandidate))
-		{
-			BindingClass = ObjectProperty->PropertyClass;
-		}
+	checkSlow(Function != nullptr);
 
-		if (UClass const* FuncOwner = Function->GetOwnerClass())
-		{
-			bCanBind = BindingClass->IsChildOf(FuncOwner);
-		}
+	bool const bNodeTypeMatches = (NodeClass == UK2Node_CallFunction::StaticClass());
+	bool bClassOwnerMatches = false;
+
+	UClass* BindingClass = FBlueprintNodeSpawnerUtils::GetBindingClass(BindingCandidate);
+	if (UClass const* FuncOwner = Function->GetOwnerClass())
+	{
+		bClassOwnerMatches = BindingClass->IsChildOf(FuncOwner);
 	}
 
-	return bCanBind;
+	return bNodeTypeMatches && bClassOwnerMatches && !FObjectEditorUtils::IsFunctionHiddenFromClass(Function, BindingClass);
 }
 
 //------------------------------------------------------------------------------
diff --git a/Engine/Source/Editor/BlueprintGraph/Private/BlueprintNodeSpawnerUtils.cpp b/Engine/Source/Editor/BlueprintGraph/Private/BlueprintNodeSpawnerUtils.cpp
--- a/Engine/Source/Editor/BlueprintGraph/Private/BlueprintNodeSpawnerUtils.cpp
+++ b/Engine/Source/Editor/BlueprintGraph/Private/BlueprintNodeSpawnerUtils.cpp
@@ -23,6 +23,11 @@ UField const* FBlueprintNodeSpawnerUtils::GetAssociatedField(UBlueprintNodeSpawn
 	{
 		ClassField = Property;
 	}
+	// @TODO: have to fix up some of the filter cases to ignore structs/enums
+// 	else if (UBlueprintFieldNodeSpawner const* FieldNodeSpawner = Cast<UBlueprintFieldNodeSpawner>(BlueprintAction))
+// 	{
+// 		ClassField = FieldNodeSpawner->GetField();
+// 	}
 	return ClassField;
 }
 
@@ -62,3 +67,14 @@ UProperty const* FBlueprintNodeSpawnerUtils::GetAssociatedProperty(UBlueprintNod
 	}
 	return Property;
 }
+
+//------------------------------------------------------------------------------
+UClass* FBlueprintNodeSpawnerUtils::GetBindingClass(const UObject* Binding)
+{
+	UClass* BindingClass = Binding->GetClass();
+	if (const UObjectProperty* ObjProperty = Cast<UObjectProperty>(Binding))
+	{
+		BindingClass = ObjProperty->PropertyClass;
+	}
+	return BindingClass;
+}
diff --git a/Engine/Source/Editor/BlueprintGraph/Private/BlueprintNodeSpawnerUtils.h b/Engine/Source/Editor/BlueprintGraph/Private/BlueprintNodeSpawnerUtils.h
--- a/Engine/Source/Editor/BlueprintGraph/Private/BlueprintNodeSpawnerUtils.h
+++ b/Engine/Source/Editor/BlueprintGraph/Private/BlueprintNodeSpawnerUtils.h
@@ -35,5 +35,15 @@ struct FBlueprintNodeSpawnerUtils
 	 * @return The action's associated property (null if it doesn't have one).
 	 */
 	static UProperty const* GetAssociatedProperty(UBlueprintNodeSpawner const* BlueprintAction);
+
+	/**
+	 * Utility function to pull UClass info from a tentative binding object 
+	 * (defaults to the object's class, uses the PropertyClass if the object is
+	 * an object property).
+	 * 
+	 * @param  Binding	The binding you want a class for.
+	 * @return A UClass that corresponds to the supplied binding.
+	 */
+	static UClass* GetBindingClass(const UObject* Binding);
 };
 
diff --git a/Engine/Source/Editor/UnrealEd/Private/Commandlets/DumpBlueprintsInfoCommandlet.cpp b/Engine/Source/Editor/UnrealEd/Private/Commandlets/DumpBlueprintsInfoCommandlet.cpp
--- a/Engine/Source/Editor/UnrealEd/Private/Commandlets/DumpBlueprintsInfoCommandlet.cpp
+++ b/Engine/Source/Editor/UnrealEd/Private/Commandlets/DumpBlueprintsInfoCommandlet.cpp
@@ -100,8 +100,8 @@ DumpBlueprintsInfo commandlet params: \n\
                         '{1}' and '{2}' as placeholders for filenames, like so:\n\
                         -diffcmd=\"AraxisP4Diff.exe {2} {1}\".                 \n\
 \n\
-    -experimental       Uses an new way of constructing Blueprint action menus \n\
-                        (that will replace the current system).                \n\
+    -legacy             Uses the old way of constructing Blueprint action      \n\
+                        menus.                                                 \n\
 \n\
     -name=<Filename>    Overrides the default filename. Leave off the extention\n\
                         (this will add .json to the end). When -multifile is   \n\
@@ -419,13 +419,13 @@ DumpBlueprintsInfo commandlet params: \n\
 //------------------------------------------------------------------------------
 DumpBlueprintInfoUtils::CommandletOptions::CommandletOptions(TArray<FString> const& Switches)
 	: BlueprintClass(nullptr)
-	, DumpFlags(BPDUMP_UnfilteredPalette | BPDUMP_UseLegacyMenuBuilder)
+	, DumpFlags(BPDUMP_UnfilteredPalette)
 	, PaletteFilter(nullptr)
 	, GraphFilter(GT_MAX)
 	, SelectedObjectType(nullptr)
 	, InterfaceClass(nullptr)	
 {
-	uint32 NewDumpFlags = BPDUMP_UseLegacyMenuBuilder;
+	uint32 NewDumpFlags = 0x00;
 	for (FString const& Switch : Switches)
 	{
 		if (Switch.StartsWith("class="))
@@ -585,9 +585,9 @@ DumpBlueprintInfoUtils::CommandletOptions::CommandletOptions(TArray<FString> con
 		{
 			NewDumpFlags |= BPDUMP_RecordTiming;
 		}
-		else if (!Switch.Compare("experimental", ESearchCase::IgnoreCase))
+		else if (!Switch.Compare("legacy", ESearchCase::IgnoreCase))
 		{
-			NewDumpFlags &= ~BPDUMP_UseLegacyMenuBuilder;
+			NewDumpFlags |= BPDUMP_UseLegacyMenuBuilder;
 		}
 		else if (Switch.StartsWith("name="))
 		{
