[346160] Do not allow unsafe functions to be placed in a construction script.

- Modified CanPaletteItemBePlaced() in SBlueprintSubPalette.cpp to leverage the UEdGraphSchema_K2::CanFunctionBeUsedInClass() API.
- Modified UEdGraphSchema_K2::CanFunctionBeUsedInClass() to accept an optional parameter for a "reason" FText on failure and added localizable strings.
- Added a UK2Node_CallFunction::CanPasteHere() override impl to also invoke UEdGraphSchema_K2::CanFunctionBeUsedInClass() in order to gate function nodes on paste into a graph.

diff --git a/Engine/Source/Editor/BlueprintGraph/Classes/EdGraphSchema_K2.h b/Engine/Source/Editor/BlueprintGraph/Classes/EdGraphSchema_K2.h
--- a/Engine/Source/Editor/BlueprintGraph/Classes/EdGraphSchema_K2.h
+++ b/Engine/Source/Editor/BlueprintGraph/Classes/EdGraphSchema_K2.h
@@ -900,8 +900,9 @@ public:
 	 * @param	bInShowInherited	Allows for inherited functions
 	 * @param	bInCalledForEach	Call for each element in an array (a node accepts array)
 	 * @param	InTargetInfo		Allows spawning nodes which also create a target variable as well
+	 * @param	OutReason			Allows callers to receive a localized string containing more detail when the function is determined to be invalid (optional)
 	 */
-	bool CanFunctionBeUsedInClass(const UClass* InClass, UFunction* InFunction, const UEdGraph* InDestGraph, uint32 InFunctionTypes, bool bInShowInherited, bool bInCalledForEach, const FFunctionTargetInfo& InTargetInfo) const;
+	bool CanFunctionBeUsedInClass(const UClass* InClass, UFunction* InFunction, const UEdGraph* InDestGraph, uint32 InFunctionTypes, bool bInShowInherited, bool bInCalledForEach, const FFunctionTargetInfo& InTargetInfo, FText* OutReason = nullptr) const;
 
 	/**
 	 * Makes connections into/or out of the gateway node, connect directly to the associated networks on the opposite side of the tunnel
diff --git a/Engine/Source/Editor/BlueprintGraph/Classes/K2Node_CallFunction.h b/Engine/Source/Editor/BlueprintGraph/Classes/K2Node_CallFunction.h
--- a/Engine/Source/Editor/BlueprintGraph/Classes/K2Node_CallFunction.h
+++ b/Engine/Source/Editor/BlueprintGraph/Classes/K2Node_CallFunction.h
@@ -72,6 +72,7 @@ public:
 	virtual FString GetDocumentationLink() const override;
 	virtual FString GetDocumentationExcerptName() const override;
 	virtual FName GetPaletteIcon(FLinearColor& OutColor) const override;
+	virtual bool CanPasteHere(const UEdGraph* TargetGraph) const override;
 	// End of UEdGraphNode interface
 
 	// UK2Node interface
diff --git a/Engine/Source/Editor/BlueprintGraph/Private/EdGraphSchema_K2.cpp b/Engine/Source/Editor/BlueprintGraph/Private/EdGraphSchema_K2.cpp
--- a/Engine/Source/Editor/BlueprintGraph/Private/EdGraphSchema_K2.cpp
+++ b/Engine/Source/Editor/BlueprintGraph/Private/EdGraphSchema_K2.cpp
@@ -319,14 +319,14 @@ bool UEdGraphSchema_K2::DoesFunctionHaveOutParameters( const UFunction* Function
 	return false;
 }
 
-bool UEdGraphSchema_K2::CanFunctionBeUsedInClass(const UClass* InClass, UFunction* InFunction, const UEdGraph* InDestGraph, uint32 InAllowedFunctionTypes, bool bInShowInherited, bool bInCalledForEach, const FFunctionTargetInfo& InTargetInfo) const
+bool UEdGraphSchema_K2::CanFunctionBeUsedInClass(const UClass* InClass, UFunction* InFunction, const UEdGraph* InDestGraph, uint32 InAllowedFunctionTypes, bool bInShowInherited, bool bInCalledForEach, const FFunctionTargetInfo& InTargetInfo, FText* OutReason) const
 {
 	if (CanUserKismetCallFunction(InFunction))
 	{
 		bool bLatentFuncsAllowed = true;
 		bool bIsConstructionScript = false;
 
-		if(InDestGraph != NULL)
+		if(InDestGraph != nullptr)
 		{
 			bLatentFuncsAllowed = (GetGraphType(InDestGraph) == GT_Ubergraph || (GetGraphType(InDestGraph) == GT_Macro));
 			bIsConstructionScript = IsConstructionScript(InDestGraph);
@@ -338,6 +338,11 @@ bool UEdGraphSchema_K2::CanFunctionBeUsedInClass(const UClass* InClass, UFunctio
 			const bool bAllowPureFuncs = (InAllowedFunctionTypes & FT_Pure) != 0;
 			if (!bAllowPureFuncs)
 			{
+				if(OutReason != nullptr)
+				{
+					*OutReason = LOCTEXT("PureFunctionsNotAllowed", "Pure functions are not allowed.");
+				}
+
 				return false;
 			}
 		}
@@ -346,6 +351,11 @@ bool UEdGraphSchema_K2::CanFunctionBeUsedInClass(const UClass* InClass, UFunctio
 			const bool bAllowImperativeFuncs = (InAllowedFunctionTypes & FT_Imperative) != 0;
 			if (!bAllowImperativeFuncs)
 			{
+				if(OutReason != nullptr)
+				{
+					*OutReason = LOCTEXT("ImpureFunctionsNotAllowed", "Impure functions are not allowed.");
+				}
+
 				return false;
 			}
 		}
@@ -354,12 +364,22 @@ bool UEdGraphSchema_K2::CanFunctionBeUsedInClass(const UClass* InClass, UFunctio
 		const bool bAllowConstFuncs = (InAllowedFunctionTypes & FT_Const) != 0;
 		if (bIsConstFunc && !bAllowConstFuncs)
 		{
+			if(OutReason != nullptr)
+			{
+				*OutReason = LOCTEXT("ConstFunctionsNotAllowed", "Const functions are not allowed.");
+			}
+
 			return false;
 		}
 
 		const bool bIsLatent = InFunction->HasMetaData(FBlueprintMetadata::MD_Latent);
 		if (bIsLatent && !bLatentFuncsAllowed)
 		{
+			if(OutReason != nullptr)
+			{
+				*OutReason = LOCTEXT("LatentFunctionsNotAllowed", "Latent functions cannot be used here.");
+			}
+
 			return false;
 		}
 
@@ -370,11 +390,21 @@ bool UEdGraphSchema_K2::CanFunctionBeUsedInClass(const UClass* InClass, UFunctio
 			const bool bAllowProtectedFuncs = (InAllowedFunctionTypes & FT_Protected) != 0;
 			if (!bAllowProtectedFuncs)
 			{
+				if(OutReason != nullptr)
+				{
+					*OutReason = LOCTEXT("ProtectedFunctionsNotAllowed", "Protected functions are not allowed.");
+				}
+
 				return false;
 			}
 
 			if (!bFuncBelongsToSubClass)
 			{
+				if(OutReason != nullptr)
+				{
+					*OutReason = LOCTEXT("ProtectedFunctionInaccessible", "Function is protected and inaccessible.");
+				}
+
 				return false;
 			}
 		}
@@ -383,12 +413,22 @@ bool UEdGraphSchema_K2::CanFunctionBeUsedInClass(const UClass* InClass, UFunctio
 		const bool bFuncBelongsToClass = bFuncBelongsToSubClass && (InFunction->GetOuterUClass() == InClass);
 		if (bIsPrivate && !bFuncBelongsToClass)
 		{
+			if(OutReason != nullptr)
+			{
+				*OutReason = LOCTEXT("PrivateFunctionInaccessible", "Function is private and inaccessible.");
+			}
+
 			return false;
 		}
 
 		const bool bIsUnsafeForConstruction = InFunction->GetBoolMetaData(FBlueprintMetadata::MD_UnsafeForConstructionScripts);	
 		if (bIsUnsafeForConstruction && bIsConstructionScript)
 		{
+			if(OutReason != nullptr)
+			{
+				*OutReason = LOCTEXT("FunctionUnsafeForConstructionScript", "Function cannot be used in a Construction Script.");
+			}
+
 			return false;
 		}
 
@@ -400,6 +440,11 @@ bool UEdGraphSchema_K2::CanFunctionBeUsedInClass(const UClass* InClass, UFunctio
 				UClass* ParentClass = FBlueprintEditorUtils::FindBlueprintForGraphChecked(InDestGraph)->ParentClass;
 				if (!ParentClass->GetDefaultObject()->ImplementsGetWorld() && !ParentClass->HasMetaData(FBlueprintMetadata::MD_ShowWorldContextPin))
 				{
+					if(OutReason != nullptr)
+					{
+						*OutReason = LOCTEXT("FunctionRequiresWorldContext", "Function requires a world context.");
+					}
+
 					return false;
 				}
 			}
@@ -408,6 +453,11 @@ bool UEdGraphSchema_K2::CanFunctionBeUsedInClass(const UClass* InClass, UFunctio
 		const bool bFunctionHidden = FObjectEditorUtils::IsFunctionHiddenFromClass(InFunction, InClass);
 		if (bFunctionHidden)
 		{
+			if(OutReason != nullptr)
+			{
+				*OutReason = LOCTEXT("HiddenFunctionInaccessible", "Function is hidden and inaccessible.");
+			}
+
 			return false;
 		}
 
@@ -418,6 +468,38 @@ bool UEdGraphSchema_K2::CanFunctionBeUsedInClass(const UClass* InClass, UFunctio
 		const bool bAllowForEachCall = !bFunctionStatic && !bIsLatent && !bIsPureFunc && !bIsConstFunc && !bHasReturnParams && !bHasArrayPointerParms;
 		if (bInCalledForEach && !bAllowForEachCall)
 		{
+			if(OutReason != nullptr)
+			{
+				if(bFunctionStatic)
+				{
+					*OutReason = LOCTEXT("StaticFunctionsNotAllowedInForEachContext", "Static functions cannot be used within a ForEach context.");
+				}
+				else if(bIsLatent)
+				{
+					*OutReason = LOCTEXT("LatentFunctionsNotAllowedInForEachContext", "Latent functions cannot be used within a ForEach context.");
+				}
+				else if(bIsPureFunc)
+				{
+					*OutReason = LOCTEXT("PureFunctionsNotAllowedInForEachContext", "Pure functions cannot be used within a ForEach context.");
+				}
+				else if(bIsConstFunc)
+				{
+					*OutReason = LOCTEXT("ConstFunctionsNotAllowedInForEachContext", "Const functions cannot be used within a ForEach context.");
+				}
+				else if(bHasReturnParams)
+				{
+					*OutReason = LOCTEXT("FunctionsWithReturnValueNotAllowedInForEachContext", "Functions that return a value cannot be used within a ForEach context.");
+				}
+				else if(bHasArrayPointerParms)
+				{
+					*OutReason = LOCTEXT("FunctionsWithArrayParmsNotAllowedInForEachContext", "Functions with array parameters cannot be used within a ForEach context.");
+				}
+				else
+				{
+					*OutReason = LOCTEXT("FunctionNotAllowedInForEachContext", "Function cannot be used within a ForEach context.");
+				}
+			}
+
 			return false;
 		}
 
@@ -428,12 +510,22 @@ bool UEdGraphSchema_K2::CanFunctionBeUsedInClass(const UClass* InClass, UFunctio
 		const bool bAllowReturnValuesForNoneOrSingleActors = !bClassIsAnActor || InTargetInfo.Actors.Num() <= 1 || !bFunctionHasReturnOrOutParameters;
 		if (!bAllowReturnValuesForNoneOrSingleActors)
 		{
+			if(OutReason != nullptr)
+			{
+				*OutReason = LOCTEXT("FunctionNotAllowedWithMultipleTargets", "Functions that return a value cannot be used with multiple targets.");
+			}
+
 			return false;
 		}
 
 		return true;
 	}
 
+	if(OutReason != nullptr)
+	{
+		*OutReason = LOCTEXT("FunctionInvalid", "Invalid function.");
+	}
+
 	return false;
 }
 
diff --git a/Engine/Source/Editor/BlueprintGraph/Private/K2Node_CallFunction.cpp b/Engine/Source/Editor/BlueprintGraph/Private/K2Node_CallFunction.cpp
--- a/Engine/Source/Editor/BlueprintGraph/Private/K2Node_CallFunction.cpp
+++ b/Engine/Source/Editor/BlueprintGraph/Private/K2Node_CallFunction.cpp
@@ -681,6 +681,24 @@ bool UK2Node_CallFunction::CanFunctionSupportMultipleTargets(UFunction const* Fu
 	return !bHasReturnParam && bIsImpure && !bIsLatent;
 }
 
+bool UK2Node_CallFunction::CanPasteHere(const UEdGraph* TargetGraph) const
+{
+	bool bCanPaste = Super::CanPasteHere(TargetGraph);
+	if(bCanPaste)
+	{
+		const UEdGraphSchema_K2* K2Schema = GetDefault<UEdGraphSchema_K2>();
+		uint32 AllowedFunctionTypes = UEdGraphSchema_K2::EFunctionType::FT_Pure | UEdGraphSchema_K2::EFunctionType::FT_Const | UEdGraphSchema_K2::EFunctionType::FT_Protected;
+		if(K2Schema->DoesGraphSupportImpureFunctions(TargetGraph))
+		{
+			AllowedFunctionTypes |= UEdGraphSchema_K2::EFunctionType::FT_Imperative;
+		}
+
+		bCanPaste = K2Schema->CanFunctionBeUsedInClass(FBlueprintEditorUtils::FindBlueprintForGraphChecked(TargetGraph)->GeneratedClass, GetTargetFunction(), TargetGraph, AllowedFunctionTypes, true, false, FFunctionTargetInfo());
+	}
+	
+	return bCanPaste;
+}
+
 static FLinearColor GetPalletteIconColor(UFunction const* Function)
 {
 	bool const bIsPure = (Function != nullptr) && Function->HasAnyFunctionFlags(FUNC_BlueprintPure);
diff --git a/Engine/Source/Editor/Kismet/Private/SBlueprintSubPalette.cpp b/Engine/Source/Editor/Kismet/Private/SBlueprintSubPalette.cpp
--- a/Engine/Source/Editor/Kismet/Private/SBlueprintSubPalette.cpp
+++ b/Engine/Source/Editor/Kismet/Private/SBlueprintSubPalette.cpp
@@ -59,7 +59,7 @@ static bool CanPaletteItemBePlaced(TSharedPtr<FEdGraphSchemaAction> DropActionIn
 	else if (UK2Node const* NodeToBePlaced = FK2SchemaActionUtils::ExtractNodeTemplateFromAction(DropActionIn))
 	{
 		UEdGraphSchema const* const GraphSchema = HoveredGraphIn->GetSchema();
-		check(GraphSchema != NULL);
+		check(GraphSchema != nullptr);
 
 		bool bIsFunctionGraph = (GraphSchema->GetGraphType(HoveredGraphIn) == EGraphType::GT_Function);
 
@@ -68,24 +68,31 @@ static bool CanPaletteItemBePlaced(TSharedPtr<FEdGraphSchemaAction> DropActionIn
 			FName const FuncName = CallFuncNode->FunctionReference.GetMemberName();
 			check(FuncName != NAME_None);
 			UClass const* const FuncOwner = CallFuncNode->FunctionReference.GetMemberParentClass(CallFuncNode);
-			check(FuncOwner != NULL);
+			check(FuncOwner != nullptr);
 
-			UFunction const* const Function = FindField<UFunction>(FuncOwner, FuncName);
+			UFunction* const Function = FindField<UFunction>(FuncOwner, FuncName);
+			UEdGraphSchema_K2 const* const K2Schema = Cast<UEdGraphSchema_K2 const>(GraphSchema);
 
-			if (!HoveredGraphIn->GetSchema()->IsA(UEdGraphSchema_K2::StaticClass()))
+			if (Function == nullptr)
 			{
 				bCanBePlaced = false;
-				ImpededReasonOut = LOCTEXT("CannotCreateInThisSchema", "Cannot call functions in this type of graph");
+				ImpededReasonOut = LOCTEXT("InvalidFuncAction", "Invalid function for placement");
 			}
-			else if (Function == NULL)
+			else if (K2Schema == nullptr)
 			{
 				bCanBePlaced = false;
-				ImpededReasonOut = LOCTEXT("InvalidFuncAction", "Invalid function for placement");
+				ImpededReasonOut = LOCTEXT("CannotCreateInThisSchema", "Cannot call functions in this type of graph");
 			}
-			else if (bIsFunctionGraph && Function->HasMetaData(FBlueprintMetadata::MD_Latent))
+			else
 			{
-				bCanBePlaced = false;
-				ImpededReasonOut = LOCTEXT("CannotCreateLatentInGraph", "Cannot call latent functions in function graphs");
+				uint32 AllowedFunctionTypes = UEdGraphSchema_K2::EFunctionType::FT_Pure | UEdGraphSchema_K2::EFunctionType::FT_Const | UEdGraphSchema_K2::EFunctionType::FT_Protected;
+				if(K2Schema->DoesGraphSupportImpureFunctions(HoveredGraphIn))
+				{
+					AllowedFunctionTypes |= UEdGraphSchema_K2::EFunctionType::FT_Imperative;
+				}
+
+				const UClass* GeneratedClass = FBlueprintEditorUtils::FindBlueprintForGraphChecked(HoveredGraphIn)->GeneratedClass;
+				bCanBePlaced = K2Schema->CanFunctionBeUsedInClass(GeneratedClass, Function, HoveredGraphIn, AllowedFunctionTypes, true, false, FFunctionTargetInfo(), &ImpededReasonOut);
 			}
 		}
 		else if (UK2Node_Event const* EventNode = Cast<UK2Node_Event const>(NodeToBePlaced))
