MathExpression node will no longer crash if placed in a Macro

Moved FinalNodeBackToMacroSourceMap into the CompilerResultsLog so compiling nodes can find their source nodes when they are contained within a Macro.

#ttp 346745 - LIVE: EDITOR: BP: CRASH: Editor crashes when using a math expression inside of a macro

diff --git a/Engine/Source/Editor/BlueprintGraph/Private/K2Node_MathExpression.cpp b/Engine/Source/Editor/BlueprintGraph/Private/K2Node_MathExpression.cpp
--- a/Engine/Source/Editor/BlueprintGraph/Private/K2Node_MathExpression.cpp
+++ b/Engine/Source/Editor/BlueprintGraph/Private/K2Node_MathExpression.cpp
@@ -2436,16 +2436,34 @@ void UK2Node_MathExpression::ValidateNodeDuringCompilation(FCompilerResultsLog&
 		MessageLog.Append(*CachedMessageLog);
 	}
 	// else, this may be some intermediate node in the compile, let's look at the errors from the original...
-	else if (UK2Node_MathExpression const* SourceNode = MessageLog.FindSourceObjectTypeChecked<UK2Node_MathExpression>(this))
+	else 
 	{
-		// if the expressions match, then the errors should
-		check(SourceNode->Expression == Expression);
-		
-		// take the same errors from the original node (so we don't have to
-		// re-parse/re-gen to fish out the same errors)
-		if (SourceNode->CachedMessageLog.IsValid())
+		if(UObject const* SourceObject = MessageLog.FindSourceObject(this))
 		{
-			MessageLog.Append(*SourceNode->CachedMessageLog);
+			UK2Node_MathExpression const* MathExpression = nullptr;
+
+			// If the source object is a MacroInstance, we need to look elsewhere for the original MathExpression
+			if(Cast<UK2Node_MacroInstance>(SourceObject))
+			{
+				MathExpression = CastChecked<UK2Node_MathExpression>(MessageLog.FinalNodeBackToMacroSourceMap.FindSourceObject(this));
+			}
+			else
+			{
+				MathExpression = MessageLog.FindSourceObjectTypeChecked<UK2Node_MathExpression>(this);
+			}
+
+			// Should always be able to find the source math expression
+			check(MathExpression);
+
+			// if the expressions match, then the errors should
+			check(MathExpression->Expression == Expression);
+
+			// take the same errors from the original node (so we don't have to
+			// re-parse/re-gen to fish out the same errors)
+			if (MathExpression->CachedMessageLog.IsValid())
+			{
+				MessageLog.Append(*MathExpression->CachedMessageLog);
+			}
 		}
 	}
 }
diff --git a/Engine/Source/Editor/KismetCompiler/Private/KismetCompiler.cpp b/Engine/Source/Editor/KismetCompiler/Private/KismetCompiler.cpp
--- a/Engine/Source/Editor/KismetCompiler/Private/KismetCompiler.cpp
+++ b/Engine/Source/Editor/KismetCompiler/Private/KismetCompiler.cpp
@@ -2830,13 +2830,13 @@ void FKismetCompilerContext::ExpandTunnelsAndMacros(UEdGraph* SourceGraph)
 					UEdGraphNode* MacroSourceNode = Cast<UEdGraphNode>(MessageLog.FindSourceObject(DuplicatedNode));
 					if (MacroSourceNode)
 					{
-						FinalNodeBackToMacroSourceMap.NotifyIntermediateObjectCreation(DuplicatedNode, MacroSourceNode);
+						MessageLog.FinalNodeBackToMacroSourceMap.NotifyIntermediateObjectCreation(DuplicatedNode, MacroSourceNode);
 
 						// Also record mappings from final macro source node to intermediate macro instance nodes (there may be more than one)
-						UEdGraphNode* MacroInstanceSourceNode = Cast<UEdGraphNode>(FinalNodeBackToMacroSourceMap.FindSourceObject(MacroInstanceNode));
+						UEdGraphNode* MacroInstanceSourceNode = Cast<UEdGraphNode>(MessageLog.FinalNodeBackToMacroSourceMap.FindSourceObject(MacroInstanceNode));
 						if (MacroInstanceSourceNode && MacroInstanceSourceNode != MacroInstanceNode)
 						{
-							MacroSourceToMacroInstanceNodeMap.Add(MacroSourceNode, MacroInstanceSourceNode);
+							MessageLog.MacroSourceToMacroInstanceNodeMap.Add(MacroSourceNode, MacroInstanceSourceNode);
 						}
 					}
 
diff --git a/Engine/Source/Editor/KismetCompiler/Private/KismetCompilerVMBackend.cpp b/Engine/Source/Editor/KismetCompiler/Private/KismetCompilerVMBackend.cpp
--- a/Engine/Source/Editor/KismetCompiler/Private/KismetCompilerVMBackend.cpp
+++ b/Engine/Source/Editor/KismetCompiler/Private/KismetCompilerVMBackend.cpp
@@ -1118,7 +1118,7 @@ public:
 				if (TrueSourceNode)
 				{
 					// If this is a debug site for an expanded macro instruction, there should also be a macro source node associated with it
-					UEdGraphNode* MacroSourceNode = Cast<UEdGraphNode>(CompilerContext.FinalNodeBackToMacroSourceMap.FindSourceObject(SourceNode));
+					UEdGraphNode* MacroSourceNode = Cast<UEdGraphNode>(CompilerContext.MessageLog.FinalNodeBackToMacroSourceMap.FindSourceObject(SourceNode));
 					if (MacroSourceNode == SourceNode)
 					{
 						// The function above will return the given node if not found in the map. In that case there is no associated source macro node, so we clear it.
@@ -1143,7 +1143,7 @@ public:
 						}
 
 						// Gather up all the macro instance nodes that lead to this macro source node
-						CompilerContext.MacroSourceToMacroInstanceNodeMap.MultiFind(MacroSourceNode, MacroInstanceNodes);
+						CompilerContext.MessageLog.MacroSourceToMacroInstanceNodeMap.MultiFind(MacroSourceNode, MacroInstanceNodes);
 					}
 
 					int32 Offset = Writer.ScriptBuffer.Num();
diff --git a/Engine/Source/Editor/KismetCompiler/Public/KismetCompiler.h b/Engine/Source/Editor/KismetCompiler/Public/KismetCompiler.h
--- a/Engine/Source/Editor/KismetCompiler/Public/KismetCompiler.h
+++ b/Engine/Source/Editor/KismetCompiler/Public/KismetCompiler.h
@@ -82,10 +82,6 @@ public:
 
 	// Map from a name to the number of times it's been 'created' (identical nodes create the same variable names, so they need something appended)
 	FNetNameMapping ClassScopeNetNameMap;
-
-	// Special maps used for autocreated macros to preserve information about their source
-	FBacktrackMap FinalNodeBackToMacroSourceMap;
-	TMultiMap<TWeakObjectPtr<UEdGraphNode>, TWeakObjectPtr<UEdGraphNode>> MacroSourceToMacroInstanceNodeMap;
 public:
 	FKismetCompilerContext(UBlueprint* SourceSketch, FCompilerResultsLog& InMessageLog, const FKismetCompilerOptions& InCompilerOptions, TArray<UObject*>* InObjLoaded);
 	virtual ~FKismetCompilerContext();
diff --git a/Engine/Source/Editor/UnrealEd/Public/Kismet2/CompilerResultsLog.h b/Engine/Source/Editor/UnrealEd/Public/Kismet2/CompilerResultsLog.h
--- a/Engine/Source/Editor/UnrealEd/Public/Kismet2/CompilerResultsLog.h
+++ b/Engine/Source/Editor/UnrealEd/Public/Kismet2/CompilerResultsLog.h
@@ -85,6 +85,10 @@ public:
 	// Should detailed results be appended to the final summary log?
 	bool bLogDetailedResults;
 
+	// Special maps used for autocreated macros to preserve information about their source
+	FBacktrackMap FinalNodeBackToMacroSourceMap;
+	TMultiMap<TWeakObjectPtr<UEdGraphNode>, TWeakObjectPtr<UEdGraphNode>> MacroSourceToMacroInstanceNodeMap;
+
 	// Minimum event time (ms) for inclusion into the final summary log
 	int EventDisplayThresholdMs;
 
