Duplicating Blueprints will no longer leave all "self" pins on variables in an invalid state in referencing the original Blueprint.

diff --git a/Engine/Source/Editor/BlueprintGraph/Classes/K2Node_Variable.h b/Engine/Source/Editor/BlueprintGraph/Classes/K2Node_Variable.h
--- a/Engine/Source/Editor/BlueprintGraph/Classes/K2Node_Variable.h
+++ b/Engine/Source/Editor/BlueprintGraph/Classes/K2Node_Variable.h
@@ -45,6 +45,7 @@ protected:
 
 public:
 	// Begin UObject interface
+	virtual void PostDuplicate(bool bDuplicateForPIE) override;
 	virtual void Serialize(FArchive& Ar) override;
 	// End UObject interface
 
diff --git a/Engine/Source/Editor/BlueprintGraph/Private/K2Node_Variable.cpp b/Engine/Source/Editor/BlueprintGraph/Private/K2Node_Variable.cpp
--- a/Engine/Source/Editor/BlueprintGraph/Private/K2Node_Variable.cpp
+++ b/Engine/Source/Editor/BlueprintGraph/Private/K2Node_Variable.cpp
@@ -322,6 +322,34 @@ UEdGraphPin* UK2Node_Variable::GetValuePin() const
 	return Pin;
 }
 
+void UK2Node_Variable::PostDuplicate(bool bDuplicateForPIE)
+{
+	Super::PostDuplicate(bDuplicateForPIE);
+	if (!bDuplicateForPIE && (!this->HasAnyFlags(RF_Transient)))
+	{
+		// Self context variable nodes need to be updated with the new Blueprint class
+		if(VariableReference.IsSelfContext())
+		{
+			const UEdGraphSchema_K2* K2Schema = GetDefault<UEdGraphSchema_K2>();
+			if(UEdGraphPin* SelfPin = K2Schema->FindSelfPin(*this, EGPD_Input))
+			{
+				UClass* TargetClass = nullptr;
+
+				if(UProperty* Property = VariableReference.ResolveMember<UProperty>(this))
+				{
+					TargetClass = Property->GetOwnerClass()->GetAuthoritativeClass();
+				}
+				else
+				{
+					TargetClass = GetBlueprint()->SkeletonGeneratedClass->GetAuthoritativeClass();
+				}
+
+				SelfPin->PinType.PinSubCategoryObject = TargetClass;
+			}
+		}
+	}
+}
+
 void UK2Node_Variable::ValidateNodeDuringCompilation(class FCompilerResultsLog& MessageLog) const
 {
 	Super::ValidateNodeDuringCompilation(MessageLog);
