Pasting invalid self variables into Macro Blueprints is no longer allowed and promoting or replacing variables with missing properties is no longer an option.

Pasting any nodes into interface BPs is not possible.

When right clicking on a variable node with a missing property, you can replace the node with any valid local variables (in functions).

#ttp 340631 - bp: macros: one can create variables within macros by pasting them into a macro from a blueprint (not sure if this is intentional because the create variable ui has been removed)

diff --git a/Engine/Source/Editor/BlueprintGraph/Classes/EdGraphSchema_K2.h b/Engine/Source/Editor/BlueprintGraph/Classes/EdGraphSchema_K2.h
--- a/Engine/Source/Editor/BlueprintGraph/Classes/EdGraphSchema_K2.h
+++ b/Engine/Source/Editor/BlueprintGraph/Classes/EdGraphSchema_K2.h
@@ -939,8 +939,11 @@ public:
 	/** Create the variable that the broken node refers to */
 	static void OnCreateNonExistentVariable(class UK2Node_Variable* Variable, UBlueprint* OwnerBlueprint);
 
+	/** Create the local variable that the broken node refers to */
+	static void OnCreateNonExistentLocalVariable(class UK2Node_Variable* Variable, UBlueprint* OwnerBlueprint);
+
 	/** Replace the variable that a variable node refers to when the variable it refers to does not exist */
-	static void OnReplaceVariableForVariableNode(class UK2Node_Variable* Variable, UBlueprint* OwnerBlueprint, FString VariableName);
+	static void OnReplaceVariableForVariableNode(class UK2Node_Variable* Variable, UBlueprint* OwnerBlueprint, FString VariableName, bool bIsSelfMember);
 
 	/** Create sub menu that shows all possible variables that can be used to replace the existing variable reference */
 	static void GetReplaceNonExistentVariableMenu(class FMenuBuilder& MenuBuilder, class UK2Node_Variable* Variable, UBlueprint* OwnerBlueprint);
diff --git a/Engine/Source/Editor/BlueprintGraph/Classes/K2Node_Variable.h b/Engine/Source/Editor/BlueprintGraph/Classes/K2Node_Variable.h
--- a/Engine/Source/Editor/BlueprintGraph/Classes/K2Node_Variable.h
+++ b/Engine/Source/Editor/BlueprintGraph/Classes/K2Node_Variable.h
@@ -57,6 +57,7 @@ public:
 	virtual FName GetCornerIcon() const override;
 	virtual FName GetPaletteIcon(FLinearColor& OutColor) const override;
 	virtual void AutowireNewNode(UEdGraphPin* FromPin) override;
+	virtual bool CanPasteHere(const UEdGraph* TargetGraph) const override;
 	// End UEdGraphNode interface
 
 	// Begin K2Node interface
diff --git a/Engine/Source/Editor/BlueprintGraph/Private/EdGraphSchema_K2.cpp b/Engine/Source/Editor/BlueprintGraph/Private/EdGraphSchema_K2.cpp
--- a/Engine/Source/Editor/BlueprintGraph/Private/EdGraphSchema_K2.cpp
+++ b/Engine/Source/Editor/BlueprintGraph/Private/EdGraphSchema_K2.cpp
@@ -910,7 +910,16 @@ static bool IsUsingNonExistantVariable(const UEdGraphNode* InGraphNode, UBluepri
 			{
 				TArray<FName> CurrentVars;
 				FBlueprintEditorUtils::GetClassVariableList(OwnerBlueprint, CurrentVars);
-				if (false == CurrentVars.Contains(Variable->GetVarName()))
+				if ( false == CurrentVars.Contains(Variable->GetVarName()) )
+				{
+					bNonExistantVariable = true;
+				}
+			}
+			else if(Variable->VariableReference.IsLocalScope())
+			{
+				// If there is no member scope, or we can't find the local variable in the member scope, then it's non-existant
+				if(!Variable->VariableReference.GetMemberScope(Variable) 
+					||  (Variable->VariableReference.GetMemberScope(Variable) && !FBlueprintEditorUtils::FindLocalVariable(OwnerBlueprint, Variable->VariableReference.GetMemberScope(Variable), Variable->GetVarName())))
 				{
 					bNonExistantVariable = true;
 				}
@@ -1224,6 +1233,8 @@ void UEdGraphSchema_K2::OnCreateNonExistentVariable( UK2Node_Variable* Variable,
 {
 	if (UEdGraphPin* Pin = Variable->FindPin(Variable->GetVarNameString()))
 	{
+		const FScopedTransaction Transaction( LOCTEXT("CreateMissingVariable", "Create Missing Variable") );
+
 		if (FBlueprintEditorUtils::AddMemberVariable(OwnerBlueprint,Variable->GetVarName(), Pin->PinType))
 		{
 			Variable->VariableReference.SetSelfMember( Variable->GetVarName() );
@@ -1231,12 +1242,38 @@ void UEdGraphSchema_K2::OnCreateNonExistentVariable( UK2Node_Variable* Variable,
 	}	
 }
 
-void UEdGraphSchema_K2::OnReplaceVariableForVariableNode( UK2Node_Variable* Variable, UBlueprint* OwnerBlueprint, FString VariableName)
+void UEdGraphSchema_K2::OnCreateNonExistentLocalVariable( UK2Node_Variable* Variable,  UBlueprint* OwnerBlueprint)
 {
 	if (UEdGraphPin* Pin = Variable->FindPin(Variable->GetVarNameString()))
 	{
-		Variable->VariableReference.SetSelfMember( FName(*VariableName) );
+		const FScopedTransaction Transaction( LOCTEXT("CreateMissingLocalVariable", "Create Missing Local Variable") );
+
+		FName VarName = Variable->GetVarName();
+		if (FBlueprintEditorUtils::AddLocalVariable(OwnerBlueprint, Variable->GetGraph(), VarName, Pin->PinType))
+		{
+			FGuid LocalVarGuid = FBlueprintEditorUtils::FindLocalVariableGuidByName(OwnerBlueprint, Variable->GetGraph(), VarName);
+			if (LocalVarGuid.IsValid())
+			{
+				Variable->VariableReference.SetLocalMember(VarName, Variable->GetGraph()->GetName(), LocalVarGuid);
+			}
+		}
+	}	
+}
+
+void UEdGraphSchema_K2::OnReplaceVariableForVariableNode( UK2Node_Variable* Variable, UBlueprint* OwnerBlueprint, FString VariableName, bool bIsSelfMember)
+{
+	if(UEdGraphPin* Pin = Variable->FindPin(Variable->GetVarNameString()))
+	{
+		if (bIsSelfMember)
+		{
+			Variable->VariableReference.SetSelfMember( FName(*VariableName) );
+		}
+		else
+		{
+			Variable->VariableReference.SetLocalMember( FName(*VariableName), Variable->GetGraph()->GetName(), FBlueprintEditorUtils::FindLocalVariableGuidByName(OwnerBlueprint, Variable->GetGraph(), *VariableName));
+		}
 		Pin->PinName = VariableName;
+		Variable->ReconstructNode();
 	}
 }
 
@@ -1247,14 +1284,30 @@ void UEdGraphSchema_K2::GetReplaceNonExistentVariableMenu(FMenuBuilder& MenuBuil
 		TArray<FName> Variables;
 		FBlueprintEditorUtils::GetNewVariablesOfType(OwnerBlueprint, Pin->PinType, Variables);
 
+		MenuBuilder.BeginSection(NAME_None, LOCTEXT("Variables", "Variables"));
 		for (TArray<FName>::TIterator VarIt(Variables); VarIt; ++VarIt)
 		{
 			const FText AlternativeVar = FText::FromName( *VarIt );
 			const FText Desc = FText::Format( LOCTEXT("ReplaceNonExistantVarToolTip", "Variable '{0}' does not exist, replace with matching variable '{0}'?"), Variable->GetVarNameText(), AlternativeVar );
 
 			MenuBuilder.AddMenuEntry( AlternativeVar, Desc, FSlateIcon(), FUIAction(
-				FExecuteAction::CreateStatic(&UEdGraphSchema_K2::OnReplaceVariableForVariableNode, const_cast<UK2Node_Variable* >(Variable),OwnerBlueprint, (*VarIt).ToString() ) ) );
+				FExecuteAction::CreateStatic(&UEdGraphSchema_K2::OnReplaceVariableForVariableNode, const_cast<UK2Node_Variable* >(Variable),OwnerBlueprint, (*VarIt).ToString(), /*bIsSelfMember=*/true ) ) );
 		}
+		MenuBuilder.EndSection();
+
+		TArray<FName> LocalVariables;
+		FBlueprintEditorUtils::GetLocalVariablesOfType(Variable->GetGraph(), Pin->PinType, LocalVariables);
+
+		MenuBuilder.BeginSection(NAME_None, LOCTEXT("LocalVariables", "LocalVariables"));
+		for (TArray<FName>::TIterator VarIt(LocalVariables); VarIt; ++VarIt)
+		{
+			const FText AlternativeVar = FText::FromName( *VarIt );
+			const FText Desc = FText::Format( LOCTEXT("ReplaceNonExistantLocalVarToolTip", "Variable '{0}' does not exist, replace with matching localvariable '{0}'?"), Variable->GetVarNameText(), AlternativeVar );
+
+			MenuBuilder.AddMenuEntry( AlternativeVar, Desc, FSlateIcon(), FUIAction(
+				FExecuteAction::CreateStatic(&UEdGraphSchema_K2::OnReplaceVariableForVariableNode, const_cast<UK2Node_Variable* >(Variable),OwnerBlueprint, (*VarIt).ToString(), /*bIsSelfMember=*/false ) ) );
+		}
+		MenuBuilder.EndSection();
 	}
 }
 
@@ -1263,12 +1316,27 @@ void UEdGraphSchema_K2::GetNonExistentVariableMenu( const UEdGraphNode* InGraphN
 
 	if (const UK2Node_Variable* Variable = Cast<const UK2Node_Variable>(InGraphNode))
 	{
-		// create missing variable
-		{			
-			const FText Label = FText::Format( LOCTEXT("CreateNonExistentVar", "Create variable '{0}'"), Variable->GetVarNameText());
-			const FText Desc = FText::Format( LOCTEXT("CreateNonExistentVarToolTip", "Variable '{0}' does not exist, create it?"), Variable->GetVarNameText());
-			MenuBuilder->AddMenuEntry( Label, Desc, FSlateIcon(), FUIAction(
-				FExecuteAction::CreateStatic( &UEdGraphSchema_K2::OnCreateNonExistentVariable, const_cast<UK2Node_Variable* >(Variable),OwnerBlueprint) ) );
+		// Creating missing variables should never occur in a Macro Library or Interface, they do not support variables
+		if(OwnerBlueprint->BlueprintType != BPTYPE_MacroLibrary && OwnerBlueprint->BlueprintType != BPTYPE_Interface )
+		{		
+			// Creating missing member variables should never occur in a Function Library, they do not support variables
+			if(OwnerBlueprint->BlueprintType != BPTYPE_FunctionLibrary)
+			{
+				// create missing variable
+				const FText Label = FText::Format( LOCTEXT("CreateNonExistentVar", "Create variable '{0}'"), Variable->GetVarNameText());
+				const FText Desc = FText::Format( LOCTEXT("CreateNonExistentVarToolTip", "Variable '{0}' does not exist, create it?"), Variable->GetVarNameText());
+				MenuBuilder->AddMenuEntry( Label, Desc, FSlateIcon(), FUIAction(
+					FExecuteAction::CreateStatic( &UEdGraphSchema_K2::OnCreateNonExistentVariable, const_cast<UK2Node_Variable* >(Variable),OwnerBlueprint) ) );
+			}
+
+			// Only allow creating missing local variables if in a function graph
+			if(InGraphNode->GetGraph()->GetSchema()->GetGraphType(InGraphNode->GetGraph()) == GT_Function)
+			{
+				const FText Label = FText::Format( LOCTEXT("CreateNonExistentLocalVar", "Create local variable '{0}'"), Variable->GetVarNameText());
+				const FText Desc = FText::Format( LOCTEXT("CreateNonExistentLocalVarToolTip", "Local variable '{0}' does not exist, create it?"), Variable->GetVarNameText());
+				MenuBuilder->AddMenuEntry( Label, Desc, FSlateIcon(), FUIAction(
+					FExecuteAction::CreateStatic( &UEdGraphSchema_K2::OnCreateNonExistentLocalVariable, const_cast<UK2Node_Variable* >(Variable),OwnerBlueprint) ) );
+			}
 		}
 
 		// delete this node
@@ -1283,7 +1351,10 @@ void UEdGraphSchema_K2::GetNonExistentVariableMenu( const UEdGraphNode* InGraphN
 			TArray<FName> Variables;
 			FBlueprintEditorUtils::GetNewVariablesOfType(OwnerBlueprint, Pin->PinType, Variables);
 
-			if (Variables.Num() > 0)
+			TArray<FName> LocalVariables;
+			FBlueprintEditorUtils::GetLocalVariablesOfType(Variable->GetGraph(), Pin->PinType, LocalVariables);
+
+			if (Variables.Num() > 0 || LocalVariables.Num() > 0)
 			{
 				MenuBuilder->AddSubMenu(
 					FText::Format( LOCTEXT("ReplaceVariableWith", "Replace variable '{0}' with..."), Variable->GetVarNameText()),
diff --git a/Engine/Source/Editor/BlueprintGraph/Private/K2Node.cpp b/Engine/Source/Editor/BlueprintGraph/Private/K2Node.cpp
--- a/Engine/Source/Editor/BlueprintGraph/Private/K2Node.cpp
+++ b/Engine/Source/Editor/BlueprintGraph/Private/K2Node.cpp
@@ -1285,6 +1285,7 @@ void FMemberReference::SetLocalMember(FName InMemberName, FString InScopeName, c
 	MemberName = InMemberName;
 	MemberScope = InScopeName;
 	MemberGuid = InMemberGuid;
+	bSelfContext = false;
 }
 
 void FMemberReference::InvalidateSelfScope()
diff --git a/Engine/Source/Editor/BlueprintGraph/Private/K2Node_Variable.cpp b/Engine/Source/Editor/BlueprintGraph/Private/K2Node_Variable.cpp
--- a/Engine/Source/Editor/BlueprintGraph/Private/K2Node_Variable.cpp
+++ b/Engine/Source/Editor/BlueprintGraph/Private/K2Node_Variable.cpp
@@ -667,4 +667,22 @@ FBPVariableDescription const* UK2Node_Variable::GetBlueprintVarDescription() con
 	return nullptr;
 }
 
+bool UK2Node_Variable::CanPasteHere(const UEdGraph* TargetGraph) const
+{
+	// Do not allow pasting of variables in BPs that cannot handle them
+	if ( FBlueprintEditorUtils::FindBlueprintForGraph(TargetGraph)->BlueprintType == BPTYPE_MacroLibrary && VariableReference.IsSelfContext() )
+	{
+		// Self variables must be from a parent class to the macro BP
+		if(UProperty* Property = VariableReference.ResolveMember<UProperty>(this))
+		{
+			const UClass* CurrentClass = GetBlueprint()->SkeletonGeneratedClass->GetAuthoritativeClass();
+			const UClass* PropertyClass = Property->GetOwnerClass()->GetAuthoritativeClass();
+			const bool bIsChildOf = CurrentClass->IsChildOf(PropertyClass);
+			return bIsChildOf;
+		}
+		return false;
+	}
+	return true;
+}
+
 #undef LOCTEXT_NAMESPACE
diff --git a/Engine/Source/Editor/Kismet/Private/BlueprintEditor.cpp b/Engine/Source/Editor/Kismet/Private/BlueprintEditor.cpp
--- a/Engine/Source/Editor/Kismet/Private/BlueprintEditor.cpp
+++ b/Engine/Source/Editor/Kismet/Private/BlueprintEditor.cpp
@@ -4640,6 +4640,12 @@ void FBlueprintEditor::PasteNodesHere(class UEdGraph* DestinationGraph, const FV
 
 bool FBlueprintEditor::CanPasteNodes() const
 {
+	// Do not allow pasting into interface Blueprints
+	if (GetBlueprintObj()->BlueprintType == BPTYPE_Interface)
+	{
+		return false;
+	}
+
 	// Find the graph editor with focus
 	TSharedPtr<SGraphEditor> FocusedGraphEd = FocusedGraphEdPtr.Pin();
 	if (!FocusedGraphEd.IsValid())
@@ -5721,8 +5727,6 @@ void FBlueprintEditor::OnAddNewVariable()
 		++Index;
 	}
 
-	GetBlueprintObj()->Modify();
-
 	bool bSuccess = MyBlueprintWidget.IsValid() && FBlueprintEditorUtils::AddMemberVariable(GetBlueprintObj(), VarName, MyBlueprintWidget->GetLastPinTypeUsed());
 
 	if(!bSuccess)
@@ -5741,34 +5745,17 @@ void FBlueprintEditor::OnAddNewLocalVariable()
 	UEdGraph* TargetGraph = FBlueprintEditorUtils::GetTopLevelGraph(FocusedGraphEdPtr.Pin()->GetCurrentGraph());
 	check(TargetGraph->GetSchema()->GetGraphType(TargetGraph) == GT_Function);
 
-	if(TargetGraph != NULL)
-	{
-		const FScopedTransaction Transaction( LOCTEXT("AddLocalVariable", "Add Local Variable") );
-		GetBlueprintObj()->Modify();
-
-		// Figure out a decent place to stick the node
-		const FVector2D NewNodePos = TargetGraph->GetGoodPlaceForNewNode();
-
-		TArray<UK2Node_FunctionEntry*> FunctionEntryNodes;
-		TargetGraph->GetNodesOfClass(FunctionEntryNodes);
-		check(FunctionEntryNodes.Num());
-
-		const UEdGraphSchema_K2* K2Schema = GetDefault<UEdGraphSchema_K2>();
-
-		// Now create new variable
-		FBPVariableDescription NewVar;
-
-		NewVar.VarName = FBlueprintEditorUtils::FindUniqueKismetName(GetBlueprintObj(), TEXT("NewLocalVar"), FindField<UFunction>(GetBlueprintObj()->SkeletonGeneratedClass, TargetGraph->GetFName()));
-		NewVar.VarGuid = FGuid::NewGuid();
-		NewVar.VarType = GetMyBlueprintWidget()->GetLastPinTypeUsed();
-		NewVar.FriendlyName = FName::NameToDisplayString( NewVar.VarName.ToString(), (NewVar.VarType.PinCategory == K2Schema->PC_Boolean) ? true : false );
-		NewVar.Category = K2Schema->VR_DefaultCategory;
+	FName VarName = FBlueprintEditorUtils::FindUniqueKismetName(GetBlueprintObj(), TEXT("NewLocalVar"), FindField<UFunction>(GetBlueprintObj()->SkeletonGeneratedClass, TargetGraph->GetFName()));
 
-		FunctionEntryNodes[0]->Modify();
-		FunctionEntryNodes[0]->LocalVariables.Add(NewVar);
+	bool bSuccess = MyBlueprintWidget.IsValid() && FBlueprintEditorUtils::AddLocalVariable(GetBlueprintObj(), TargetGraph, VarName, MyBlueprintWidget->GetLastPinTypeUsed());
 
-		FBlueprintEditorUtils::MarkBlueprintAsStructurallyModified(GetBlueprintObj());
-		RenameNewlyAddedAction(NewVar.VarName);
+	if(!bSuccess)
+	{
+		LogSimpleMessage( LOCTEXT("AddLocalVariable_Error", "Adding new local variable failed.") );
+	}
+	else
+	{
+		RenameNewlyAddedAction(VarName);
 	}
 }
 
diff --git a/Engine/Source/Editor/UnrealEd/Private/Kismet2/BlueprintEditorUtils.cpp b/Engine/Source/Editor/UnrealEd/Private/Kismet2/BlueprintEditorUtils.cpp
--- a/Engine/Source/Editor/UnrealEd/Private/Kismet2/BlueprintEditorUtils.cpp
+++ b/Engine/Source/Editor/UnrealEd/Private/Kismet2/BlueprintEditorUtils.cpp
@@ -3201,6 +3201,28 @@ void FBlueprintEditorUtils::GetNewVariablesOfType( const UBlueprint* Blueprint,
 	}
 }
 
+void FBlueprintEditorUtils::GetLocalVariablesOfType( const UEdGraph* Graph, const FEdGraphPinType& Type, TArray<FName>& OutVars)
+{
+	if(Graph && Graph->GetSchema()->GetGraphType(Graph) == GT_Function)
+	{
+		TArray<UK2Node_FunctionEntry*> GraphNodes;
+		Graph->GetNodesOfClass<UK2Node_FunctionEntry>(GraphNodes);
+
+		bool bFoundLocalVariable = false;
+
+		// There is only ever 1 function entry
+		check(GraphNodes.Num() == 1);
+
+		for( auto& LocalVar : GraphNodes[0]->LocalVariables )
+		{
+			if(LocalVar.VarType == Type)
+			{
+				OutVars.Add(LocalVar.VarName);
+			}
+		}
+	}
+}
+
 // Adds a member variable to the blueprint.  It cannot mask a variable in any superclass.
 bool FBlueprintEditorUtils::AddMemberVariable(UBlueprint* Blueprint, const FName& NewVarName, const FEdGraphPinType& NewVarType, const FString& DefaultValue/* = FString()*/)
 {
@@ -3218,6 +3240,8 @@ bool FBlueprintEditorUtils::AddMemberVariable(UBlueprint* Blueprint, const FName
 		return false; // fail
 	}
 
+	Blueprint->Modify();
+
 	const UEdGraphSchema_K2* K2Schema = GetDefault<UEdGraphSchema_K2>();
 
 	// Now create new variable
@@ -3598,6 +3622,38 @@ FBPVariableDescription FBlueprintEditorUtils::DuplicateVariableDescription(UBlue
 	return NewVar;
 }
 
+bool FBlueprintEditorUtils::AddLocalVariable(UBlueprint* Blueprint, UEdGraph* InTargetGraph, const FName& InNewVarName, const FEdGraphPinType& InNewVarType)
+{
+	if(InTargetGraph != NULL && InTargetGraph->GetSchema()->GetGraphType(InTargetGraph) == GT_Function)
+	{
+		const FScopedTransaction Transaction( LOCTEXT("AddLocalVariable", "Add Local Variable") );
+		Blueprint->Modify();
+
+		TArray<UK2Node_FunctionEntry*> FunctionEntryNodes;
+		InTargetGraph->GetNodesOfClass(FunctionEntryNodes);
+		check(FunctionEntryNodes.Num());
+
+		const UEdGraphSchema_K2* K2Schema = GetDefault<UEdGraphSchema_K2>();
+
+		// Now create new variable
+		FBPVariableDescription NewVar;
+
+		NewVar.VarName = InNewVarName;
+		NewVar.VarGuid = FGuid::NewGuid();
+		NewVar.VarType = InNewVarType;
+		NewVar.FriendlyName = FName::NameToDisplayString( NewVar.VarName.ToString(), (NewVar.VarType.PinCategory == K2Schema->PC_Boolean) ? true : false );
+		NewVar.Category = K2Schema->VR_DefaultCategory;
+
+		FunctionEntryNodes[0]->Modify();
+		FunctionEntryNodes[0]->LocalVariables.Add(NewVar);
+
+		FBlueprintEditorUtils::MarkBlueprintAsStructurallyModified(Blueprint);
+
+		return true;
+	}
+	return false;
+}
+
 void FBlueprintEditorUtils::RemoveLocalVariable(UBlueprint* InBlueprint, const UStruct* InScope, const FName& InVarName)
 {
 	UEdGraph* ScopeGraph = FindScopeGraph(InBlueprint, InScope);
@@ -3725,15 +3781,13 @@ FBPVariableDescription* FBlueprintEditorUtils::FindLocalVariable(UBlueprint* InB
 	return FindLocalVariable(InBlueprint, InScope, InVariableName, &DummyFunctionEntry);
 }
 
-FBPVariableDescription* FBlueprintEditorUtils::FindLocalVariable(const UBlueprint* InBlueprint, const UStruct* InScope, const FName& InVariableName, class UK2Node_FunctionEntry** OutFunctionEntry)
+FBPVariableDescription* FBlueprintEditorUtils::FindLocalVariable(const UBlueprint* InBlueprint, const UEdGraph* InScopeGraph, const FName& InVariableName, class UK2Node_FunctionEntry** OutFunctionEntry)
 {
-	UEdGraph* ScopeGraph = FindScopeGraph(InBlueprint, InScope);
-
 	FBPVariableDescription* ReturnVariable = NULL;
-	if(ScopeGraph)
+	if(InScopeGraph)
 	{
 		TArray<UK2Node_FunctionEntry*> GraphNodes;
-		ScopeGraph->GetNodesOfClass<UK2Node_FunctionEntry>(GraphNodes);
+		InScopeGraph->GetNodesOfClass<UK2Node_FunctionEntry>(GraphNodes);
 
 		bool bFoundLocalVariable = false;
 
@@ -3761,6 +3815,13 @@ FBPVariableDescription* FBlueprintEditorUtils::FindLocalVariable(const UBlueprin
 	return ReturnVariable;
 }
 
+FBPVariableDescription* FBlueprintEditorUtils::FindLocalVariable(const UBlueprint* InBlueprint, const UStruct* InScope, const FName& InVariableName, class UK2Node_FunctionEntry** OutFunctionEntry)
+{
+	UEdGraph* ScopeGraph = FindScopeGraph(InBlueprint, InScope);
+
+	return FindLocalVariable(InBlueprint, ScopeGraph, InVariableName, OutFunctionEntry);
+}
+
 FName FBlueprintEditorUtils::FindLocalVariableNameByGuid(UBlueprint* InBlueprint, const FGuid& InVariableGuid)
 {
 	// Search through all function entry nodes for a local variable with the passed Guid
@@ -3799,6 +3860,17 @@ FGuid FBlueprintEditorUtils::FindLocalVariableGuidByName(UBlueprint* InBlueprint
 	return ReturnGuid;
 }
 
+FGuid FBlueprintEditorUtils::FindLocalVariableGuidByName(UBlueprint* InBlueprint, const UEdGraph* InScopeGraph, const FName InVariableName)
+{
+	FGuid ReturnGuid;
+	if(FBPVariableDescription* LocalVariable = FindLocalVariable(InBlueprint, InScopeGraph, InVariableName))
+	{
+		ReturnGuid = LocalVariable->VarGuid;
+	}
+
+	return ReturnGuid;
+}
+
 void FBlueprintEditorUtils::ChangeLocalVariableType(UBlueprint* InBlueprint, const UStruct* InScope, const FName& InVariableName, const FEdGraphPinType& NewPinType)
 {
 	if (InVariableName != NAME_None)
diff --git a/Engine/Source/Editor/UnrealEd/Public/Kismet2/BlueprintEditorUtils.h b/Engine/Source/Editor/UnrealEd/Public/Kismet2/BlueprintEditorUtils.h
--- a/Engine/Source/Editor/UnrealEd/Public/Kismet2/BlueprintEditorUtils.h
+++ b/Engine/Source/Editor/UnrealEd/Public/Kismet2/BlueprintEditorUtils.h
@@ -487,6 +487,14 @@ public:
 	 */
 	static void GetNewVariablesOfType( const UBlueprint* Blueprint, const FEdGraphPinType& Type, TArray<FName>& OutVars);
 
+	/**
+	 * Gets local variables of specified type
+	 *
+	 * @param 			FEdGraphPinType	 			Type of variables to look for
+	 * @param [in,out]	VisibleVariables			The visible variables will be appened to this array.
+	 */
+	static void GetLocalVariablesOfType( const UEdGraph* Graph, const FEdGraphPinType& Type, TArray<FName>& OutVars);
+
 	/**
 	 * Adds a member variable to the blueprint.  It cannot mask a variable in any superclass.
 	 *
@@ -561,6 +569,17 @@ public:
 	 */
 	static UEdGraph* FindScopeGraph(const UBlueprint* InBlueprint, const UStruct* InScope);
 
+	/**
+	 * Adds a local variable to the function graph.  It cannot mask a member variable or a variable in any superclass.
+	 *
+	 * @param	NewVarName	Name of the new variable.
+	 * @param	NewVarType	Type of the new variable.
+	 * @param	DefaultValue	Default value stored as string
+	 *
+	 * @return	true if it succeeds, false if it fails.
+	 */
+	static bool AddLocalVariable(UBlueprint* Blueprint, UEdGraph* InTargetGraph, const FName& InNewVarName, const FEdGraphPinType& InNewVarType);
+
 	/**
 	 * Removes a member variable if it was declared in this blueprint and not in a base class.
 	 *
@@ -579,6 +598,17 @@ public:
 	 */
 	static FBPVariableDescription* FindLocalVariable(UBlueprint* InBlueprint, const UStruct* InScope, const FName& InVariableName);
 
+	/**
+	 * Returns a local variable
+	 *
+	 * @param InBlueprint		Blueprint to search for the local variable
+	 * @param InScopeGraph		Local variable's graph
+	 * @param InVariableName	Name of the variable to search for
+	 * @param OutFunctionEntry	Optional output parameter. If not null, the found function entry is returned.
+	 * @return					The local variable description
+	 */
+	static FBPVariableDescription* FindLocalVariable(const UBlueprint* InBlueprint, const UEdGraph* InScopeGraph, const FName& InVariableName, class UK2Node_FunctionEntry** OutFunctionEntry = NULL);
+
 	/**
 	 * Returns a local variable
 	 *
@@ -610,6 +640,16 @@ public:
 	 */
 	static FGuid FindLocalVariableGuidByName(UBlueprint* InBlueprint, const UStruct* InScope, const FName InVariableName);
 
+	/**
+	 * Finds a local variable Guid using the variable's name
+	 *
+	 * @param InBlueprint		Blueprint to search for the local variable
+	 * @param InScopeGraph		Local variable's graph
+	 * @param InVariableGuid	Local variable's name to search for
+	 * @return					The Guid associated with the local variable
+	 */
+	static FGuid FindLocalVariableGuidByName(UBlueprint* InBlueprint, const UEdGraph* InScopeGraph, const FName InVariableName);
+
 	/**
 	 * Rename a local variable
 	 *
