Temp fix for DoOnceMultiInput nodes losing connection

diff --git a/Engine/Source/Editor/BlueprintGraph/Classes/K2Node_DoOnceMultiInput.h b/Engine/Source/Editor/BlueprintGraph/Classes/K2Node_DoOnceMultiInput.h
--- a/Engine/Source/Editor/BlueprintGraph/Classes/K2Node_DoOnceMultiInput.h
+++ b/Engine/Source/Editor/BlueprintGraph/Classes/K2Node_DoOnceMultiInput.h
@@ -20,6 +20,9 @@ class UK2Node_DoOnceMultiInput : public UK2Node
 	UPROPERTY(transient)
 	class UK2Node_TemporaryVariable* DataNode;
 
+	void ReallocatePinsDuringReconstruction(TArray<UEdGraphPin*>&) override;
+	ERedirectType DoPinsMatchForReconstruction(const UEdGraphPin* NewPin, int32 NewPinIndex, const UEdGraphPin* OldPin, int32 OldPinIndex) const override;
+
 private:
 
 	const static int32 NumBaseInputs = 1;
@@ -54,4 +57,4 @@ public:
 	virtual void GetMenuActions(FBlueprintActionDatabaseRegistrar& ActionRegistrar) const override;
 	virtual FText GetMenuCategory() const override;
 	// End of UK2Node interface
-};
\ No newline at end of file
+};
diff --git a/Engine/Source/Editor/BlueprintGraph/Private/K2Node_DoOnceMultiInput.cpp b/Engine/Source/Editor/BlueprintGraph/Private/K2Node_DoOnceMultiInput.cpp
--- a/Engine/Source/Editor/BlueprintGraph/Private/K2Node_DoOnceMultiInput.cpp
+++ b/Engine/Source/Editor/BlueprintGraph/Private/K2Node_DoOnceMultiInput.cpp
@@ -25,6 +25,38 @@ void UK2Node_DoOnceMultiInput::GetMenuEntries(FGraphContextMenuBuilder& ContextM
 	NodeAction->NodeTemplate = NodeTemplate;
 }
 
+void UK2Node_DoOnceMultiInput::ReallocatePinsDuringReconstruction(TArray<UEdGraphPin*>& OldPins)
+{
+	Super::ReallocatePinsDuringReconstruction(OldPins);
+}
+
+UK2Node::ERedirectType UK2Node_DoOnceMultiInput::DoPinsMatchForReconstruction(const UEdGraphPin* NewPin, int32 NewPinIndex, const UEdGraphPin* OldPin, int32 OldPinIndex) const
+{
+	// Temp work around: remove whitespaces from pin names before doing string comparison.
+
+	FString NewName = NewPin->PinName;
+	FString OldName = OldPin->PinName;
+
+	NewName.ReplaceInline(TEXT(" "), TEXT(""));
+	OldName.ReplaceInline(TEXT(" "), TEXT(""));
+	
+	if (NewName.Equals(OldName, ESearchCase::IgnoreCase))
+	{
+		// Make sure we're not dealing with a menu node
+		UEdGraph* OuterGraph = GetGraph();
+		if (OuterGraph && OuterGraph->Schema)
+		{
+			const UEdGraphSchema_K2* K2Schema = Cast<const UEdGraphSchema_K2>(GetSchema());
+			if (!K2Schema || K2Schema->IsSelfPin(*NewPin) || K2Schema->ArePinTypesCompatible(OldPin->PinType, NewPin->PinType))
+			{
+				return ERedirectType_Name;
+			}
+		}
+	}
+
+	return Super::DoPinsMatchForReconstruction(NewPin, NewPinIndex, OldPin, OldPinIndex);
+}
+
 FText UK2Node_DoOnceMultiInput::GetNodeTitle(ENodeTitleType::Type TitleType) const
 {
 	return LOCTEXT("DoOnceMultiInput", "DoOnce MultiInput");
