ôAutoCreateRefö û is handled in UK2Node_CallFunction::ExpandNode. The additional variable is created in local scope.

diff --git a/Engine/Source/Editor/BlueprintGraph/Classes/K2Node_AddComponent.h b/Engine/Source/Editor/BlueprintGraph/Classes/K2Node_AddComponent.h
--- a/Engine/Source/Editor/BlueprintGraph/Classes/K2Node_AddComponent.h
+++ b/Engine/Source/Editor/BlueprintGraph/Classes/K2Node_AddComponent.h
@@ -33,7 +33,6 @@ class UK2Node_AddComponent : public UK2Node_CallFunction
 	virtual UActorComponent* GetTemplateFromNode() const override;
 	virtual void ValidateNodeDuringCompilation(class FCompilerResultsLog& MessageLog) const override;
 	virtual void ReallocatePinsDuringReconstruction(TArray<UEdGraphPin*>& OldPins) override;
-	virtual class FNodeHandlingFunctor* CreateNodeHandler(class FKismetCompilerContext& CompilerContext) const override;
 	virtual void ExpandNode(class FKismetCompilerContext& CompilerContext, UEdGraph* SourceGraph) override;
 	// End of UK2Node interface
 
diff --git a/Engine/Source/Editor/BlueprintGraph/Classes/K2Node_CallFunction.h b/Engine/Source/Editor/BlueprintGraph/Classes/K2Node_CallFunction.h
--- a/Engine/Source/Editor/BlueprintGraph/Classes/K2Node_CallFunction.h
+++ b/Engine/Source/Editor/BlueprintGraph/Classes/K2Node_CallFunction.h
@@ -157,6 +157,8 @@ public:
 
 	static void CallForEachElementInArrayExpansion(UK2Node* Node, UEdGraphPin* MultiSelf, FKismetCompilerContext& CompilerContext, UEdGraph* SourceGraph);
 
+	static UEdGraphPin* InnerHandleAutoCreateRef(UK2Node* Node, UEdGraphPin* Pin, FKismetCompilerContext& CompilerContext, UEdGraph* SourceGraph, bool bForceAssignment);
+
 	/**
 	 * Returns the graph for this function, if available. In cases of calling an event, it will be the ubergraph for the event
 	 *
diff --git a/Engine/Source/Editor/BlueprintGraph/Classes/K2Node_PureAssignmentStatement.h b/Engine/Source/Editor/BlueprintGraph/Classes/K2Node_PureAssignmentStatement.h
new file mode 100644
--- /dev/null
+++ b/Engine/Source/Editor/BlueprintGraph/Classes/K2Node_PureAssignmentStatement.h
@@ -0,0 +1,37 @@
+// Copyright 1998-2014 Epic Games, Inc. All Rights Reserved.
+
+#pragma once
+#include "K2Node.h"
+#include "K2Node_PureAssignmentStatement.generated.h"
+
+UCLASS(MinimalAPI)
+class UK2Node_PureAssignmentStatement : public UK2Node
+{
+	GENERATED_UCLASS_BODY()
+
+	// Name of the Variable pin for this node
+	static FString VariablePinName;
+	// Name of the Value pin for this node
+	static FString ValuePinName;
+	// Name of the output pin for this node
+	static FString OutputPinName;
+
+	// Begin UEdGraphNode interface
+	virtual void AllocateDefaultPins() override;
+	// End UEdGraphNode interface
+
+	// Begin UK2Node interface
+	virtual bool IsNodePure() const { return true; }
+	virtual void NotifyPinConnectionListChanged(UEdGraphPin* Pin) override;
+	virtual class FNodeHandlingFunctor* CreateNodeHandler(class FKismetCompilerContext& CompilerContext) const override;
+	virtual int32 GetNodeRefreshPriority() const override { return EBaseNodeRefreshPriority::Low_UsesDependentWildcard; }
+	// End UK2Node interface
+
+	/** Get the Then output pin */
+	BLUEPRINTGRAPH_API UEdGraphPin* GetOutputPin() const;
+	/** Get the Variable input pin */
+	BLUEPRINTGRAPH_API UEdGraphPin* GetVariablePin() const;
+	/** Get the Value input pin */
+	BLUEPRINTGRAPH_API UEdGraphPin* GetValuePin() const;
+};
+
diff --git a/Engine/Source/Editor/BlueprintGraph/Private/CallFunctionHandler.cpp b/Engine/Source/Editor/BlueprintGraph/Private/CallFunctionHandler.cpp
--- a/Engine/Source/Editor/BlueprintGraph/Private/CallFunctionHandler.cpp
+++ b/Engine/Source/Editor/BlueprintGraph/Private/CallFunctionHandler.cpp
@@ -496,21 +496,12 @@ void FKCHandler_CallFunction::RegisterNets(FKismetFunctionContext& Context, UEdG
 	UFunction* Function = FindFunction(Context, Node);
 	if (Function != NULL)
 	{
-		// Auto-created term defaults
-		TArray<FString> AutoCreateRefTermPinNames;
-		const bool bHasAutoCreateRefTerms = Function->HasMetaData(FBlueprintMetadata::MD_AutoCreateRefTerm);
-		if (bHasAutoCreateRefTerms)
-		{
-			CompilerContext.GetSchema()->GetAutoEmitTermParameters(Function, AutoCreateRefTermPinNames);
-		}
-
 		TArray<FString> DefaultToSelfParamNames;
 		TArray<FString> RequiresSetValue;
 
 		if (Function->HasMetaData(FBlueprintMetadata::MD_DefaultToSelf))
 		{
 			FString const DefaltToSelfPinName = Function->GetMetaData(FBlueprintMetadata::MD_DefaultToSelf);
-			AutoCreateRefTermPinNames.Remove(DefaltToSelfPinName);
 
 			DefaultToSelfParamNames.Add(DefaltToSelfPinName);
 		}
@@ -519,7 +510,6 @@ void FKCHandler_CallFunction::RegisterNets(FKismetFunctionContext& Context, UEdG
 			const bool bHasIntrinsicWorldContext = Context.Blueprint->ParentClass->GetDefaultObject()->ImplementsGetWorld();
 
 			FString const WorldContextPinName = Function->GetMetaData(FBlueprintMetadata::MD_WorldContext);
-			AutoCreateRefTermPinNames.Remove(WorldContextPinName);
 
 			if (bHasIntrinsicWorldContext)
 			{
@@ -553,37 +543,6 @@ void FKCHandler_CallFunction::RegisterNets(FKismetFunctionContext& Context, UEdG
 					CompilerContext.MessageLog.Error(*NSLOCTEXT("KismetCompiler", "PinMustHaveConnection_Error", "Pin @@ must have a connection").ToString(), Pin);
 				}
 			}
-
-			// Check for default term for ref parameters
-			if (bHasAutoCreateRefTerms && (AutoCreateRefTermPinNames.Contains(Pin->PinName)))
-			{
-				if (Pin->PinType.bIsReference
-					&& !CompilerContext.GetSchema()->IsMetaPin(*Pin)
-					&& !CompilerContext.GetSchema()->IsSelfPin(*Pin)
-					&& (Pin->Direction == EGPD_Input)
-					&& !bIsConnected
-					&& (Pin->PinType.bIsArray || !Pin->DefaultValue.IsEmpty() || Pin->DefaultObject))
-				{
-					// These terms have to be class members, because they hold default values.  If we don't have a valid ubergraph context, we've already failed, so no use creating the term
-					if (CompilerContext.UbergraphContext != NULL)
-					{
-						check(CompilerContext.UbergraphContext->NetNameMap);
-						FBPTerminal* Term = Context.CreateLocalTerminal(ETerminalSpecification::TS_ForcedShared);
-						Term->CopyFromPin(Pin, *CompilerContext.UbergraphContext->NetNameMap->MakeValidName(Pin));
-						Term->Name += TEXT("_RefProperty");
-						if (!Pin->PinType.bIsArray)
-						{
-							Term->PropertyDefault = Pin->DefaultObject ? Pin->DefaultObject->GetPathName() : Pin->DefaultValue;
-						}
-
-						Context.NetMap.Add(Pin, Term);
-					}
-					else
-					{
-						CompilerContext.MessageLog.Error(*NSLOCTEXT("KismetCompiler", "PinMustHaveConnection_Error", "Pin @@ must have a connection").ToString(), Pin);
-					}
-				}
-			}
 		}
 	}
 
diff --git a/Engine/Source/Editor/BlueprintGraph/Private/K2Node_AddComponent.cpp b/Engine/Source/Editor/BlueprintGraph/Private/K2Node_AddComponent.cpp
--- a/Engine/Source/Editor/BlueprintGraph/Private/K2Node_AddComponent.cpp
+++ b/Engine/Source/Editor/BlueprintGraph/Private/K2Node_AddComponent.cpp
@@ -14,50 +14,6 @@
 //////////////////////////////////////////////////////////////////////////
 // FKCHandler_AddComponent
 
-class FKCHandler_AddComponent : public FKCHandler_CallFunction
-{
-public:
-	FKCHandler_AddComponent(FKismetCompilerContext& InCompilerContext)
-		: FKCHandler_CallFunction(InCompilerContext)
-	{
-	}
-
-	virtual void RegisterNets(FKismetFunctionContext& Context, UEdGraphNode* Node) override
-	{
-		FKCHandler_CallFunction::RegisterNets(Context, Node);
-
-		UK2Node_AddComponent* AddComponentNode = CastChecked<UK2Node_AddComponent>(Node);
-		UEdGraphPin* TransformPin = AddComponentNode->GetRelativeTransformPin();
-
-		// If the transform pin is not connected, create a term that grabs its value from the template, so that value is used and not replaced with the identity.
-		if (TransformPin->LinkedTo.Num() == 0)
-		{
-			// The context can be NULL if the graph failed validation previously; no point in even trying to register nets, as the skeleton will be used instead
-			if (CompilerContext.UbergraphContext != NULL)
-			{
-				check(CompilerContext.UbergraphContext->NetNameMap);
-				// If there's no term, add a dummy identity transform here.  Term is defined in the event graph's scope, because a default value needs to be set
-				FBPTerminal* TransformTerm = Context.CreateLocalTerminal(ETerminalSpecification::TS_ForcedShared);
-				TransformTerm->CopyFromPin(TransformPin, FString::Printf(TEXT("%s_AddComponentDefaultTransform"), *CompilerContext.UbergraphContext->NetNameMap->MakeValidName(TransformPin)));
-				TransformTerm->bIsConst = true;
-				TransformTerm->PropertyDefault = TEXT("");
-
-				// Try and find the template and get relative transform from it
-				UEdGraphPin* TemplateNamePin = AddComponentNode->GetTemplateNamePinChecked();
-				const FString TemplateName = TemplateNamePin->DefaultValue;
-				USceneComponent* SceneCompTemplate = Cast<USceneComponent>(Context.GetBlueprint()->FindTemplateByName(FName(*TemplateName)));
-				if (SceneCompTemplate != NULL)
-				{
-					FTransform TemplateTransform = FTransform(SceneCompTemplate->RelativeRotation, SceneCompTemplate->RelativeLocation, SceneCompTemplate->RelativeScale3D);
-					TransformTerm->PropertyDefault = TemplateTransform.ToString();
-				}
-
-				Context.NetMap.Add(TransformPin, TransformTerm);
-			}
-		}
-	}
-};
-
 UK2Node_AddComponent::UK2Node_AddComponent(const class FPostConstructInitializeProperties& PCIP)
 	: Super(PCIP)
 {
@@ -406,15 +362,37 @@ bool UK2Node_AddComponent::IsCompatibleWithGraph(UEdGraph const* Graph) const
 	return (Blueprint != nullptr) && FBlueprintEditorUtils::IsActorBased(Blueprint) && Super::IsCompatibleWithGraph(Graph);
 }
 
-FNodeHandlingFunctor* UK2Node_AddComponent::CreateNodeHandler(FKismetCompilerContext& CompilerContext) const
-{
-	return new FKCHandler_AddComponent(CompilerContext);
-}
-
 void UK2Node_AddComponent::ExpandNode(class FKismetCompilerContext& CompilerContext, UEdGraph* SourceGraph)
 {
 	Super::ExpandNode(CompilerContext, SourceGraph);
 
+	
+	if (CompilerContext.bIsFullCompile)
+	{
+		auto TransformPin = GetRelativeTransformPin();
+		if (TransformPin && !TransformPin->LinkedTo.Num())
+		{
+			FString DefaultValue;
+
+			// Try and find the template and get relative transform from it
+			UEdGraphPin* TemplateNamePin = GetTemplateNamePinChecked();
+			const FString TemplateName = TemplateNamePin->DefaultValue;
+			check(CompilerContext.Blueprint);
+			USceneComponent* SceneCompTemplate = Cast<USceneComponent>(CompilerContext.Blueprint->FindTemplateByName(FName(*TemplateName)));
+			if (SceneCompTemplate)
+			{
+				FTransform TemplateTransform = FTransform(SceneCompTemplate->RelativeRotation, SceneCompTemplate->RelativeLocation, SceneCompTemplate->RelativeScale3D);
+				DefaultValue = TemplateTransform.ToString();
+			}
+
+			auto ValuePin = InnerHandleAutoCreateRef(this, TransformPin, CompilerContext, SourceGraph, !DefaultValue.IsEmpty());
+			if (ValuePin)
+			{
+				ValuePin->DefaultValue = DefaultValue;
+			}
+		}
+	}
+
 	if (CompilerContext.bIsFullCompile && bHasExposedVariable)
 	{
 		static FString ObjectParamName = FString(TEXT("Object"));
diff --git a/Engine/Source/Editor/BlueprintGraph/Private/K2Node_CallArrayFunction.cpp b/Engine/Source/Editor/BlueprintGraph/Private/K2Node_CallArrayFunction.cpp
--- a/Engine/Source/Editor/BlueprintGraph/Private/K2Node_CallArrayFunction.cpp
+++ b/Engine/Source/Editor/BlueprintGraph/Private/K2Node_CallArrayFunction.cpp
@@ -223,6 +223,11 @@ void UK2Node_CallArrayFunction::PropagateArrayTypeInfo(const UEdGraphPin* Source
 				CurrentPin->PinType.PinSubCategory = SourcePin->PinType.PinSubCategory;
 				CurrentPin->PinType.PinSubCategoryObject = SourcePin->PinType.PinSubCategoryObject;
 
+				// Reset default values
+				CurrentPin->DefaultValue = FString();
+				CurrentPin->DefaultTextValue = FText();
+				CurrentPin->DefaultObject = NULL;
+
 				// Verify that all previous connections to this pin are still valid with the new type
 				for (TArray<UEdGraphPin*>::TIterator ConnectionIt(CurrentPin->LinkedTo); ConnectionIt; ++ConnectionIt)
 				{
diff --git a/Engine/Source/Editor/BlueprintGraph/Private/K2Node_CallFunction.cpp b/Engine/Source/Editor/BlueprintGraph/Private/K2Node_CallFunction.cpp
--- a/Engine/Source/Editor/BlueprintGraph/Private/K2Node_CallFunction.cpp
+++ b/Engine/Source/Editor/BlueprintGraph/Private/K2Node_CallFunction.cpp
@@ -8,6 +8,7 @@
 #include "K2Node_SwitchEnum.h"
 #include "Kismet/KismetMathLibrary.h"
 #include "Kismet/KismetArrayLibrary.h"
+#include "K2Node_PureAssignmentStatement.h"
 
 #define LOCTEXT_NAMESPACE "K2Node"
 
@@ -1382,50 +1383,50 @@ void UK2Node_CallFunction::ExpandNode(class FKismetCompilerContext& CompilerCont
 					// Expanded as input execs pins
 					if (EnumParamPin->Direction == EGPD_Input)
 					{
-					// Create normal exec input
-					UEdGraphPin* ExecutePin = CreatePin(EGPD_Input, Schema->PC_Exec, TEXT(""), NULL, false, false, Schema->PN_Execute);
-
-					// Create temp enum variable
-					UK2Node_TemporaryVariable* TempEnumVarNode = CompilerContext.SpawnIntermediateNode<UK2Node_TemporaryVariable>(this, SourceGraph);
-					TempEnumVarNode->VariableType.PinCategory = Schema->PC_Byte;
-					TempEnumVarNode->VariableType.PinSubCategoryObject = EnumProp->Enum;
-					TempEnumVarNode->AllocateDefaultPins();
-					// Get the output pin
-					UEdGraphPin* TempEnumVarOutput = TempEnumVarNode->GetVariablePin();
-
-					// Connect temp enum variable to (hidden) enum pin
-					Schema->TryCreateConnection(TempEnumVarOutput, EnumParamPin);
-
-					// Now we want to iterate over other exec inputs...
-					for(int32 PinIdx=Pins.Num()-1; PinIdx>=0; PinIdx--)
-					{
-						UEdGraphPin* Pin = Pins[PinIdx];
-						if( Pin != NULL && 
-							Pin != ExecutePin &&
-							Pin->Direction == EGPD_Input && 
-							Pin->PinType.PinCategory == Schema->PC_Exec )
+						// Create normal exec input
+						UEdGraphPin* ExecutePin = CreatePin(EGPD_Input, Schema->PC_Exec, TEXT(""), NULL, false, false, Schema->PN_Execute);
+
+						// Create temp enum variable
+						UK2Node_TemporaryVariable* TempEnumVarNode = CompilerContext.SpawnIntermediateNode<UK2Node_TemporaryVariable>(this, SourceGraph);
+						TempEnumVarNode->VariableType.PinCategory = Schema->PC_Byte;
+						TempEnumVarNode->VariableType.PinSubCategoryObject = EnumProp->Enum;
+						TempEnumVarNode->AllocateDefaultPins();
+						// Get the output pin
+						UEdGraphPin* TempEnumVarOutput = TempEnumVarNode->GetVariablePin();
+
+						// Connect temp enum variable to (hidden) enum pin
+						Schema->TryCreateConnection(TempEnumVarOutput, EnumParamPin);
+
+						// Now we want to iterate over other exec inputs...
+						for(int32 PinIdx=Pins.Num()-1; PinIdx>=0; PinIdx--)
 						{
-							// Create node to set the temp enum var
-							UK2Node_AssignmentStatement* AssignNode = CompilerContext.SpawnIntermediateNode<UK2Node_AssignmentStatement>(this, SourceGraph);
-							AssignNode->AllocateDefaultPins();
+							UEdGraphPin* Pin = Pins[PinIdx];
+							if( Pin != NULL && 
+								Pin != ExecutePin &&
+								Pin->Direction == EGPD_Input && 
+								Pin->PinType.PinCategory == Schema->PC_Exec )
+							{
+								// Create node to set the temp enum var
+								UK2Node_AssignmentStatement* AssignNode = CompilerContext.SpawnIntermediateNode<UK2Node_AssignmentStatement>(this, SourceGraph);
+								AssignNode->AllocateDefaultPins();
 
-							// Move connections from fake 'enum exec' pint to this assignment node
-								CompilerContext.MovePinLinksToIntermediate(*Pin, *AssignNode->GetExecPin());
+								// Move connections from fake 'enum exec' pint to this assignment node
+									CompilerContext.MovePinLinksToIntermediate(*Pin, *AssignNode->GetExecPin());
 
-							// Connect this to out temp enum var
-							Schema->TryCreateConnection(AssignNode->GetVariablePin(), TempEnumVarOutput);
+								// Connect this to out temp enum var
+								Schema->TryCreateConnection(AssignNode->GetVariablePin(), TempEnumVarOutput);
 
-							// Connect exec output to 'real' exec pin
-							Schema->TryCreateConnection(AssignNode->GetThenPin(), ExecutePin);
+								// Connect exec output to 'real' exec pin
+								Schema->TryCreateConnection(AssignNode->GetThenPin(), ExecutePin);
 
-							// set the literal enum value to set to
-							AssignNode->GetValuePin()->DefaultValue = Pin->PinName;
+								// set the literal enum value to set to
+								AssignNode->GetValuePin()->DefaultValue = Pin->PinName;
 
-							// Finally remove this 'cosmetic' exec pin
-							Pins.RemoveAt(PinIdx);
+								// Finally remove this 'cosmetic' exec pin
+								Pins.RemoveAt(PinIdx);
+							}
 						}
 					}
-				}
 					// Expanded as output execs pins
 					else if (EnumParamPin->Direction == EGPD_Output)
 					{
@@ -1464,6 +1465,44 @@ void UK2Node_CallFunction::ExpandNode(class FKismetCompilerContext& CompilerCont
 			}
 		}
 
+		// AUTO CREATED REFS
+		{
+			TArray<FString> AutoCreateRefTermPinNames;
+			const bool bHasAutoCreateRefTerms = Function->HasMetaData(FBlueprintMetadata::MD_AutoCreateRefTerm);
+			if (bHasAutoCreateRefTerms)
+			{
+				CompilerContext.GetSchema()->GetAutoEmitTermParameters(Function, AutoCreateRefTermPinNames);
+			}
+
+			for (auto Pin : Pins)
+			{
+				if (Pin && bHasAutoCreateRefTerms && AutoCreateRefTermPinNames.Contains(Pin->PinName))
+				{
+					const bool bHasDefaultValue = !Pin->DefaultValue.IsEmpty() || Pin->DefaultObject || !Pin->DefaultTextValue.IsEmpty();
+					const bool bValidAutoRefPin = Pin->PinType.bIsReference
+						&& !CompilerContext.GetSchema()->IsMetaPin(*Pin)
+						&& (Pin->Direction == EGPD_Input)
+						&& !Pin->LinkedTo.Num()
+						&& (Pin->PinType.bIsArray || bHasDefaultValue);
+					if (bValidAutoRefPin)
+					{
+						//default values can be reset when the pin is connected
+						const auto DefaultValue = Pin->DefaultValue;
+						const auto DefaultObject = Pin->DefaultObject;
+						const auto DefaultTextValue = Pin->DefaultTextValue;
+
+						auto ValuePin = InnerHandleAutoCreateRef(this, Pin, CompilerContext, SourceGraph, bHasDefaultValue);
+						if (ValuePin)
+						{
+							ValuePin->DefaultValue = DefaultValue;
+							ValuePin->DefaultObject = DefaultObject;
+							ValuePin->DefaultTextValue = DefaultTextValue;
+						}
+					}
+				}
+			}
+		}
+
 		// Then we go through and expand out array iteration if necessary
 		const bool bAllowMultipleSelfs = AllowMultipleSelfs(true);
 		UEdGraphPin* MultiSelf = Schema->FindSelfPin(*this, EEdGraphPinDirection::EGPD_Input);
@@ -1481,6 +1520,42 @@ void UK2Node_CallFunction::ExpandNode(class FKismetCompilerContext& CompilerCont
 	}
 }
 
+UEdGraphPin* UK2Node_CallFunction::InnerHandleAutoCreateRef(UK2Node* Node, UEdGraphPin* Pin, FKismetCompilerContext& CompilerContext, UEdGraph* SourceGraph, bool bForceAssignment)
+{
+	const bool bAddAssigment = !Pin->PinType.bIsArray && bForceAssignment;
+
+	// ADD LOCAL VARIABLE
+	UK2Node_TemporaryVariable* LocalVariable = CompilerContext.SpawnIntermediateNode<UK2Node_TemporaryVariable>(Node, SourceGraph);
+	LocalVariable->VariableType = Pin->PinType;
+	LocalVariable->VariableType.bIsReference = false;
+	LocalVariable->AllocateDefaultPins();
+	if (!bAddAssigment)
+	{
+		if (!CompilerContext.GetSchema()->TryCreateConnection(LocalVariable->GetVariablePin(), Pin))
+		{
+			CompilerContext.MessageLog.Error(*LOCTEXT("AutoCreateRefTermPin_NotConnected", "AutoCreateRefTerm Expansion: Pin @@ cannot be connected to @@").ToString(), LocalVariable->GetVariablePin(), Pin);
+			return NULL;
+		}
+	}
+	// ADD ASSIGMENT
+	else
+	{
+		// TODO connect to dest..
+		UK2Node_PureAssignmentStatement* AssignDefaultValue = CompilerContext.SpawnIntermediateNode<UK2Node_PureAssignmentStatement>(Node, SourceGraph);
+		AssignDefaultValue->AllocateDefaultPins();
+		const bool bVariableConnected = CompilerContext.GetSchema()->TryCreateConnection(AssignDefaultValue->GetVariablePin(), LocalVariable->GetVariablePin());
+		const bool bOutputConnected = CompilerContext.GetSchema()->TryCreateConnection(AssignDefaultValue->GetOutputPin(), Pin);
+		if (!bVariableConnected || !bOutputConnected)
+		{
+			CompilerContext.MessageLog.Error(*LOCTEXT("AutoCreateRefTermPin_AssignmentError", "AutoCreateRefTerm Expansion: Assignment Error @@").ToString(), AssignDefaultValue);
+			return NULL;
+		}
+		CompilerContext.GetSchema()->SetPinDefaultValueBasedOnType(AssignDefaultValue->GetValuePin());
+		return AssignDefaultValue->GetValuePin();
+	}
+	return NULL;
+}
+
 void UK2Node_CallFunction::CallForEachElementInArrayExpansion(UK2Node* Node, UEdGraphPin* MultiSelf, FKismetCompilerContext& CompilerContext, UEdGraph* SourceGraph)
 {
 	const UEdGraphSchema_K2* Schema = CompilerContext.GetSchema();
diff --git a/Engine/Source/Editor/BlueprintGraph/Private/K2Node_PureAssignmentStatement.cpp b/Engine/Source/Editor/BlueprintGraph/Private/K2Node_PureAssignmentStatement.cpp
new file mode 100644
--- /dev/null
+++ b/Engine/Source/Editor/BlueprintGraph/Private/K2Node_PureAssignmentStatement.cpp
@@ -0,0 +1,132 @@
+// Copyright 1998-2014 Epic Games, Inc. All Rights Reserved.
+
+
+#include "BlueprintGraphPrivatePCH.h"
+#include "KismetCompiler.h"
+#include "VariableSetHandler.h"
+#include "BlueprintNodeSpawner.h"
+#include "EditorCategoryUtils.h"
+#include "K2Node_PureAssignmentStatement.h"
+
+#define LOCTEXT_NAMESPACE "K2Node_PureAssignmentStatement"
+
+//////////////////////////////////////////////////////////////////////////
+// FKCHandler_PureAssignmentStatement
+
+class FKCHandler_PureAssignmentStatement : public FKCHandler_VariableSet
+{
+public:
+	FKCHandler_PureAssignmentStatement(FKismetCompilerContext& InCompilerContext)
+		: FKCHandler_VariableSet(InCompilerContext)
+	{
+	}
+
+	virtual void RegisterNets(FKismetFunctionContext& Context, UEdGraphNode* Node) override
+	{
+		auto PureAssignmentNode = CastChecked<UK2Node_PureAssignmentStatement>(Node);
+		{
+			auto VariablePin = PureAssignmentNode->GetVariablePin();
+			if (VariablePin->LinkedTo.Num() == 0)
+			{
+				CompilerContext.MessageLog.Error(*LOCTEXT("NoVarriableConnected_Error", "A variable needs to be connected to @@").ToString(), VariablePin);
+				return;
+			}
+
+			auto OutputPin = PureAssignmentNode->GetOutputPin();
+			if (OutputPin->LinkedTo.Num() == 0)
+			{
+				CompilerContext.MessageLog.Error(*LOCTEXT("NoVarriableConnected_Error", "A output pin needs to be connected to @@").ToString(), OutputPin);
+				return;
+			}
+
+			// At the moment, a term for variable should be already registered
+			auto VariableTerm = Context.NetMap.Find(FEdGraphUtilities::GetNetFromPin(VariablePin));
+			if (!VariableTerm || !*VariableTerm)
+			{
+				CompilerContext.MessageLog.Error(*LOCTEXT("NoVarriableTerm_Error", "ICE: no variable term found in @@").ToString(), Node);
+				return;
+			}
+			Context.NetMap.Add(OutputPin, *VariableTerm);
+		}
+
+		auto ValuePin = PureAssignmentNode->GetValuePin();
+		ValidateAndRegisterNetIfLiteral(Context, ValuePin);
+	}
+
+	virtual void Compile(FKismetFunctionContext& Context, UEdGraphNode* Node) override
+	{
+		auto PureAssignmentNode = CastChecked<UK2Node_PureAssignmentStatement>(Node);
+		auto VariablePin = PureAssignmentNode->GetVariablePin();
+		auto ValuePin = PureAssignmentNode->GetValuePin();
+
+		InnerAssignment(Context, Node, VariablePin, ValuePin);
+	}
+};
+
+// The node is (was originally) needed in expansion step for handling CreateAutoRefTerm parameters in pure functions 
+
+FString UK2Node_PureAssignmentStatement::VariablePinName = FString(TEXT("Variable"));
+FString UK2Node_PureAssignmentStatement::ValuePinName = FString(TEXT("Value"));
+FString UK2Node_PureAssignmentStatement::OutputPinName = FString(TEXT("ReturnValue"));
+
+
+UK2Node_PureAssignmentStatement::UK2Node_PureAssignmentStatement(const class FPostConstructInitializeProperties& PCIP)
+	: Super(PCIP)
+{
+}
+
+void UK2Node_PureAssignmentStatement::AllocateDefaultPins()
+{
+	const UEdGraphSchema_K2* Schema = GetDefault<UEdGraphSchema_K2>();
+
+	CreatePin(EGPD_Input, Schema->PC_Wildcard, TEXT(""), NULL, false, false, VariablePinName);
+	CreatePin(EGPD_Input, Schema->PC_Wildcard, TEXT(""), NULL, false, false, ValuePinName);
+	CreatePin(EGPD_Output, Schema->PC_Wildcard, TEXT(""), NULL, false, false, OutputPinName);
+
+	Super::AllocateDefaultPins();
+}
+
+void UK2Node_PureAssignmentStatement::NotifyPinConnectionListChanged(UEdGraphPin* Pin)
+{
+	Super::NotifyPinConnectionListChanged(Pin);
+
+	if (Pin->LinkedTo.Num() > 0)
+	{
+		const auto PinType = Pin->LinkedTo[0]->PinType;
+
+		auto VariablePin = GetVariablePin();
+		VariablePin->PinType = PinType;
+		UEdGraphSchema_K2::ValidateExistingConnections(VariablePin);
+
+		auto OutputPin = GetOutputPin();
+		OutputPin->PinType = PinType;
+		UEdGraphSchema_K2::ValidateExistingConnections(OutputPin);
+
+		auto ValuePin = GetValuePin();
+		ValuePin->PinType = PinType;
+		ValuePin->PinType.bIsReference = false;
+		UEdGraphSchema_K2::ValidateExistingConnections(ValuePin);
+	}
+}
+
+UEdGraphPin* UK2Node_PureAssignmentStatement::GetOutputPin() const
+{
+	return FindPinChecked(OutputPinName);
+}
+
+UEdGraphPin* UK2Node_PureAssignmentStatement::GetVariablePin() const
+{
+	return FindPinChecked(VariablePinName);
+}
+
+UEdGraphPin* UK2Node_PureAssignmentStatement::GetValuePin() const
+{
+	return FindPinChecked(ValuePinName);
+}
+
+FNodeHandlingFunctor* UK2Node_PureAssignmentStatement::CreateNodeHandler(FKismetCompilerContext& CompilerContext) const
+{
+	return new FKCHandler_PureAssignmentStatement(CompilerContext);
+}
+
+#undef LOCTEXT_NAMESPACE
