Blueprint - Fixing a crash in Call Function node not checking for null of the function pointer.

diff --git a/Engine/Source/Editor/BlueprintGraph/Private/K2Node_CallFunction.cpp b/Engine/Source/Editor/BlueprintGraph/Private/K2Node_CallFunction.cpp
--- a/Engine/Source/Editor/BlueprintGraph/Private/K2Node_CallFunction.cpp
+++ b/Engine/Source/Editor/BlueprintGraph/Private/K2Node_CallFunction.cpp
@@ -1467,36 +1467,39 @@ void UK2Node_CallFunction::ExpandNode(class FKismetCompilerContext& CompilerCont
 
 		// AUTO CREATED REFS
 		{
-			TArray<FString> AutoCreateRefTermPinNames;
-			const bool bHasAutoCreateRefTerms = Function->HasMetaData(FBlueprintMetadata::MD_AutoCreateRefTerm);
-			if (bHasAutoCreateRefTerms)
+			if ( Function )
 			{
-				CompilerContext.GetSchema()->GetAutoEmitTermParameters(Function, AutoCreateRefTermPinNames);
-			}
+				TArray<FString> AutoCreateRefTermPinNames;
+				const bool bHasAutoCreateRefTerms = Function->HasMetaData(FBlueprintMetadata::MD_AutoCreateRefTerm);
+				if ( bHasAutoCreateRefTerms )
+				{
+					CompilerContext.GetSchema()->GetAutoEmitTermParameters(Function, AutoCreateRefTermPinNames);
+				}
 
-			for (auto Pin : Pins)
-			{
-				if (Pin && bHasAutoCreateRefTerms && AutoCreateRefTermPinNames.Contains(Pin->PinName))
+				for ( auto Pin : Pins )
 				{
-					const bool bHasDefaultValue = !Pin->DefaultValue.IsEmpty() || Pin->DefaultObject || !Pin->DefaultTextValue.IsEmpty();
-					const bool bValidAutoRefPin = Pin->PinType.bIsReference
-						&& !CompilerContext.GetSchema()->IsMetaPin(*Pin)
-						&& (Pin->Direction == EGPD_Input)
-						&& !Pin->LinkedTo.Num()
-						&& (Pin->PinType.bIsArray || bHasDefaultValue);
-					if (bValidAutoRefPin)
+					if ( Pin && bHasAutoCreateRefTerms && AutoCreateRefTermPinNames.Contains(Pin->PinName) )
 					{
-						//default values can be reset when the pin is connected
-						const auto DefaultValue = Pin->DefaultValue;
-						const auto DefaultObject = Pin->DefaultObject;
-						const auto DefaultTextValue = Pin->DefaultTextValue;
-
-						auto ValuePin = InnerHandleAutoCreateRef(this, Pin, CompilerContext, SourceGraph, bHasDefaultValue);
-						if (ValuePin)
+						const bool bHasDefaultValue = !Pin->DefaultValue.IsEmpty() || Pin->DefaultObject || !Pin->DefaultTextValue.IsEmpty();
+						const bool bValidAutoRefPin = Pin->PinType.bIsReference
+							&& !CompilerContext.GetSchema()->IsMetaPin(*Pin)
+							&& ( Pin->Direction == EGPD_Input )
+							&& !Pin->LinkedTo.Num()
+							&& ( Pin->PinType.bIsArray || bHasDefaultValue );
+						if ( bValidAutoRefPin )
 						{
-							ValuePin->DefaultValue = DefaultValue;
-							ValuePin->DefaultObject = DefaultObject;
-							ValuePin->DefaultTextValue = DefaultTextValue;
+							//default values can be reset when the pin is connected
+							const auto DefaultValue = Pin->DefaultValue;
+							const auto DefaultObject = Pin->DefaultObject;
+							const auto DefaultTextValue = Pin->DefaultTextValue;
+
+							auto ValuePin = InnerHandleAutoCreateRef(this, Pin, CompilerContext, SourceGraph, bHasDefaultValue);
+							if ( ValuePin )
+							{
+								ValuePin->DefaultValue = DefaultValue;
+								ValuePin->DefaultObject = DefaultObject;
+								ValuePin->DefaultTextValue = DefaultTextValue;
+							}
 						}
 					}
 				}
