Setting up a call function node with a dynamic output value (determined by an class picker input). #NUX
TTP #345805

diff --git a/Engine/Source/Editor/BlueprintGraph/Classes/EdGraphSchema_K2.h b/Engine/Source/Editor/BlueprintGraph/Classes/EdGraphSchema_K2.h
--- a/Engine/Source/Editor/BlueprintGraph/Classes/EdGraphSchema_K2.h
+++ b/Engine/Source/Editor/BlueprintGraph/Classes/EdGraphSchema_K2.h
@@ -145,6 +145,11 @@ public:
 	/** Metadata that flags make/break functions for specific struct types. */
 	static const FName MD_NativeMakeFunction;
 	static const FName MD_NativeBreakFunction;
+
+	/** Metadata that flags function params that govern what type of object the function returns */
+	static const FName MD_DynamicOutputType;
+	/** Metadata that flags the function output param that will be controlled by the "MD_DynamicOutputType" pin */
+	static const FName MD_DynamicOutputParam;
 	
 private:
 	// This class should never be instantiated
diff --git a/Engine/Source/Editor/BlueprintGraph/Classes/K2Node_CallFunction.h b/Engine/Source/Editor/BlueprintGraph/Classes/K2Node_CallFunction.h
--- a/Engine/Source/Editor/BlueprintGraph/Classes/K2Node_CallFunction.h
+++ b/Engine/Source/Editor/BlueprintGraph/Classes/K2Node_CallFunction.h
@@ -73,6 +73,7 @@ public:
 	virtual FString GetDocumentationExcerptName() const override;
 	virtual FName GetPaletteIcon(FLinearColor& OutColor) const override;
 	virtual bool CanPasteHere(const UEdGraph* TargetGraph) const override;
+	virtual void PinDefaultValueChanged(UEdGraphPin* Pin) override;
 	// End of UEdGraphNode interface
 
 	// UK2Node interface
diff --git a/Engine/Source/Editor/BlueprintGraph/Private/EdGraphSchema_K2.cpp b/Engine/Source/Editor/BlueprintGraph/Private/EdGraphSchema_K2.cpp
--- a/Engine/Source/Editor/BlueprintGraph/Private/EdGraphSchema_K2.cpp
+++ b/Engine/Source/Editor/BlueprintGraph/Private/EdGraphSchema_K2.cpp
@@ -87,6 +87,9 @@ const FName FBlueprintMetadata::MD_DataTablePin(TEXT("DataTablePin"));
 const FName FBlueprintMetadata::MD_NativeMakeFunction(TEXT("HasNativeMake"));
 const FName FBlueprintMetadata::MD_NativeBreakFunction(TEXT("HasNativeBreak"));
 
+const FName FBlueprintMetadata::MD_DynamicOutputType(TEXT("DeterminesOutputType"));
+const FName FBlueprintMetadata::MD_DynamicOutputParam(TEXT("DynamicOutputParam"));
+
 //////////////////////////////////////////////////////////////////////////
 
 #define LOCTEXT_NAMESPACE "KismetSchema"
@@ -2839,6 +2842,13 @@ FText UEdGraphSchema_K2::TypeToText(const FEdGraphPinType& Type)
 		}
 		else
 		{
+			UObject* SubCategoryObj = Type.PinSubCategoryObject.Get();
+			FString SubCategoryObjName = SubCategoryObj->GetName();
+			if (UField* SubCategoryField = Cast<UField>(SubCategoryObj))
+			{
+				SubCategoryObjName = SubCategoryField->GetDisplayNameText().ToString();
+			}
+
 			if( !Type.bIsWeakPointer )
 			{
 				UClass* PSCOAsClass = Cast<UClass>(Type.PinSubCategoryObject.Get());
@@ -2854,14 +2864,15 @@ FText UEdGraphSchema_K2::TypeToText(const FEdGraphPinType& Type)
 				{
 					Args.Add(TEXT("Category"), (!bIsInterface ? GetCategoryText(Type.PinCategory) : GetCategoryText(PC_Interface)));
 				}
-				Args.Add(TEXT("ObjectName"), FText::FromString(FName::NameToDisplayString(Type.PinSubCategoryObject.Get()->GetName(), false)));
+
+				Args.Add(TEXT("ObjectName"), FText::FromString(FName::NameToDisplayString(SubCategoryObjName, /*bIsBool =*/false)));
 				PropertyText = FText::Format(LOCTEXT("ObjectAsText", "{ObjectName} {Category}"), Args);
 			}
 			else
 			{
 				FFormatNamedArguments Args;
 				Args.Add(TEXT("Category"), FText::FromString(Type.PinCategory));
-				Args.Add(TEXT("ObjectName"), FText::FromString(Type.PinSubCategoryObject.Get()->GetName()));
+				Args.Add(TEXT("ObjectName"), FText::FromString(SubCategoryObjName));
 				PropertyText = FText::Format(LOCTEXT("WeakPtrAsText", "{ObjectName} Weak {Category}"), Args);
 			}
 		}
@@ -2882,7 +2893,7 @@ FText UEdGraphSchema_K2::TypeToText(const FEdGraphPinType& Type)
 	{
 		FFormatNamedArguments Args;
 		Args.Add(TEXT("PropertyTitle"), PropertyText);
-		PropertyText = FText::Format(LOCTEXT("ArrayAsText", "Array of {PropertyTitle}"), Args);
+		PropertyText = FText::Format(LOCTEXT("ArrayAsText", "Array of {PropertyTitle}s"), Args);
 	}
 	else if (Type.bIsReference)
 	{
diff --git a/Engine/Source/Editor/BlueprintGraph/Private/K2Node_CallFunction.cpp b/Engine/Source/Editor/BlueprintGraph/Private/K2Node_CallFunction.cpp
--- a/Engine/Source/Editor/BlueprintGraph/Private/K2Node_CallFunction.cpp
+++ b/Engine/Source/Editor/BlueprintGraph/Private/K2Node_CallFunction.cpp
@@ -12,6 +12,10 @@
 
 #define LOCTEXT_NAMESPACE "K2Node"
 
+/*******************************************************************************
+ *  FCustomStructureParamHelper
+ ******************************************************************************/
+
 struct FCustomStructureParamHelper
 {
 	static FName GetCustomStructureParamName()
@@ -82,6 +86,316 @@ struct FCustomStructureParamHelper
 	}
 };
 
+/*******************************************************************************
+ *  FDynamicOutputUtils
+ ******************************************************************************/
+
+struct FDynamicOutputHelper
+{
+public:
+	FDynamicOutputHelper(UEdGraphPin* InAlteredPin)
+		: MutatingPin(InAlteredPin)
+	{}
+
+	/**
+	 * Attempts to change the output pin's type so that it reflects the class 
+	 * specified by the input class pin.
+	 */
+	void ConformOutputType() const;
+
+	/**
+	 * Retrieves the class pin that is used to determine the function's output type.
+	 * 
+	 * @return Null if the "DeterminesOutputType" metadata targets an invalid 
+	 *         param (or if the metadata isn't present), otherwise a class picker pin.
+	 */
+	static UEdGraphPin* GetTypePickerPin(const UK2Node_CallFunction* FuncNode);
+
+	/**
+	 * Attempts to pull out class info from the supplied pin. Starts with the 
+	 * pin's default, and then falls back onto the pin's native type. Will poll
+	 * any connections that the pin may have.
+	 * 
+	 * @param  Pin	The pin you want a class from.
+	 * @return A class that the pin represents (could be null if the pin isn't a class pin).
+	 */
+	static UClass* GetPinClass(UEdGraphPin* Pin);
+
+	/**
+	 * Intended to be used by ValidateNodeDuringCompilation(). Will check to 
+	 * make sure the dynamic output's connections are still valid (they could
+	 * become invalid as the the pin's type changes).
+	 * 
+	 * @param  FuncNode		The node you wish to validate.
+	 * @param  MessageLog	The log to post errors/warnings to.
+	 */
+	static void VerifyNode(const UK2Node_CallFunction* FuncNode, FCompilerResultsLog& MessageLog);
+
+private:
+	/**
+	 * 
+	 * 
+	 * @return 
+	 */
+	UK2Node_CallFunction* GetFunctionNode() const;
+
+	/**
+	 * 
+	 * 
+	 * @return 
+	 */
+	UFunction* GetTargetFunction() const;
+
+	/**
+	 * Checks if the supplied pin is the class picker that governs the 
+	 * function's output type.
+	 * 
+	 * @param  Pin	The pin to test.
+	 * @return True if the pin corresponds to the param that was flagged by the "DeterminesOutputType" metadata.
+	 */
+	bool IsTypePickerPin(UEdGraphPin* Pin) const;
+
+	/**
+	 * Retrieves the object output pin that is altered as the class input is 
+	 * changed (favors params flagged by "DynamicOutputParam" metadata).
+	 * 
+	 * @return Null if the output param cannot be altered from the class input, 
+	 *         otherwise a output pin that will mutate type as the class input is changed.
+	 */
+	static UEdGraphPin* GetDynamicOutPin(const UK2Node_CallFunction* FuncNode);
+
+	/**
+	 * Checks if the specified type is an object type that reflects the picker 
+	 * pin's class.
+	 * 
+	 * @param  TypeToTest	The type you want to check.
+	 * @return True if the type is likely the output governed by the class picker pin, otherwise false.
+	 */
+	static bool CanConformPinType(const UK2Node_CallFunction* FuncNode, const FEdGraphPinType& TypeToTest);
+
+private:
+	UEdGraphPin* MutatingPin;
+};
+
+void FDynamicOutputHelper::ConformOutputType() const
+{
+	if (IsTypePickerPin(MutatingPin))
+	{
+		UClass* PickedClass = GetPinClass(MutatingPin);
+		UK2Node_CallFunction* FuncNode = GetFunctionNode();
+
+		if (UEdGraphPin* DynamicOutPin = GetDynamicOutPin(FuncNode))
+		{
+			DynamicOutPin->PinType.PinSubCategoryObject = PickedClass;
+
+			if (UFunction* Function = FuncNode->GetTargetFunction())
+			{
+				DynamicOutPin->PinToolTip.Empty();
+				FuncNode->GeneratePinTooltipFromFunction(*DynamicOutPin, Function);
+			}
+
+			// leave the connection, and instead bring the user's attention to 
+			// it via a ValidateNodeDuringCompilation() error
+// 			const UEdGraphSchema* Schema = FuncNode->GetSchema();
+// 			for (int32 LinkIndex = 0; LinkIndex < DynamicOutPin->LinkedTo.Num();)
+// 			{
+// 				UEdGraphPin* Link = DynamicOutPin->LinkedTo[LinkIndex];
+// 				// if this can no longer be linked to the other pin, then we 
+// 				// should disconnect it (because the pin's type just changed)
+// 				if (Schema->CanCreateConnection(DynamicOutPin, Link).Response == CONNECT_RESPONSE_DISALLOW)
+// 				{
+// 					DynamicOutPin->BreakLinkTo(Link);
+// 					// @TODO: warn/notify somehow
+// 				}
+// 				else
+// 				{
+// 					++LinkIndex;
+// 				}
+// 			}
+		}
+	}
+}
+
+UEdGraphPin* FDynamicOutputHelper::GetTypePickerPin(const UK2Node_CallFunction* FuncNode)
+{
+	UEdGraphPin* TypePickerPin = nullptr;
+
+	if (UFunction* Function = FuncNode->GetTargetFunction())
+	{
+		FString TypeDeterminingPinName = Function->GetMetaData(FBlueprintMetadata::MD_DynamicOutputType);
+		if (!TypeDeterminingPinName.IsEmpty())
+		{
+			TypePickerPin = FuncNode->FindPin(TypeDeterminingPinName);
+		}
+	}
+
+	if (TypePickerPin && !ensure(TypePickerPin->Direction == EGPD_Input))
+	{
+		TypePickerPin = nullptr;
+	}
+
+	return TypePickerPin;
+}
+
+UClass* FDynamicOutputHelper::GetPinClass(UEdGraphPin* Pin)
+{
+	UClass* PinClass = UObject::StaticClass();
+
+	bool const bIsClassPin = (Pin->PinType.PinCategory == UEdGraphSchema_K2::PC_Class);
+	if (bIsClassPin)
+	{
+		if (UClass* DefaultClass = Cast<UClass>(Pin->DefaultObject))
+		{
+			PinClass = DefaultClass;
+		}
+		else if (UClass* BaseClass = Cast<UClass>(Pin->PinType.PinSubCategoryObject.Get()))
+		{
+			PinClass = BaseClass;
+		}
+
+		if (Pin->LinkedTo.Num() > 0)
+		{
+			UClass* CommonInputClass = Cast<UClass>(Pin->LinkedTo[0]->PinType.PinSubCategoryObject.Get());
+			for (int32 LinkIndex = 1; LinkIndex < Pin->LinkedTo.Num(); ++LinkIndex)
+			{
+				UClass* LinkClass = Cast<UClass>(Pin->LinkedTo[LinkIndex]->PinType.PinSubCategoryObject.Get());
+				if (LinkClass == nullptr)
+				{
+					continue;
+				}
+
+				if (CommonInputClass == nullptr)
+				{
+					CommonInputClass = LinkClass;
+					continue;
+				}
+
+				while (!LinkClass->IsChildOf(CommonInputClass))
+				{
+					CommonInputClass = CommonInputClass->GetSuperClass();
+				}
+			}
+
+			PinClass = CommonInputClass;
+		}
+	}
+	return PinClass;
+}
+
+void FDynamicOutputHelper::VerifyNode(const UK2Node_CallFunction* FuncNode, FCompilerResultsLog& MessageLog)
+{
+	if (UEdGraphPin* DynamicOutPin = GetDynamicOutPin(FuncNode))
+	{
+		const UEdGraphSchema* Schema = FuncNode->GetSchema();
+		for (UEdGraphPin* Link : DynamicOutPin->LinkedTo)
+		{
+			if (Schema->CanCreateConnection(DynamicOutPin, Link).Response == CONNECT_RESPONSE_DISALLOW)
+			{
+				FText const ErrorFormat = LOCTEXT("BadConnection", "Invalid pin connection from '@@' to '@@'. You may have changed the type after the connections were made.");
+				MessageLog.Error(*ErrorFormat.ToString(), DynamicOutPin, Link);
+			}
+		}
+	}
+}
+
+UK2Node_CallFunction* FDynamicOutputHelper::GetFunctionNode() const
+{
+	return CastChecked<UK2Node_CallFunction>(MutatingPin->GetOwningNode());
+}
+
+UFunction* FDynamicOutputHelper::GetTargetFunction() const
+{
+	return GetFunctionNode()->GetTargetFunction();
+}
+
+bool FDynamicOutputHelper::IsTypePickerPin(UEdGraphPin* Pin) const
+{
+	bool bIsTypeDeterminingPin = false;
+
+	if (UFunction* Function = GetTargetFunction())
+	{
+		FString TypeDeterminingPinName = Function->GetMetaData(FBlueprintMetadata::MD_DynamicOutputType);
+		if (!TypeDeterminingPinName.IsEmpty())
+		{
+			bIsTypeDeterminingPin = (Pin->PinName == TypeDeterminingPinName);
+		}
+	}
+
+	bool const bPinIsClassPicker = (Pin->PinType.PinCategory == UEdGraphSchema_K2::PC_Class);
+	return bIsTypeDeterminingPin && bPinIsClassPicker && (Pin->Direction == EGPD_Input);
+}
+
+UEdGraphPin* FDynamicOutputHelper::GetDynamicOutPin(const UK2Node_CallFunction* FuncNode)
+{
+	UProperty* TaggedOutputParam = nullptr;
+	if (UFunction* Function = FuncNode->GetTargetFunction())
+	{
+		const FString& OutputPinName = Function->GetMetaData(FBlueprintMetadata::MD_DynamicOutputParam);
+		// we sort through properties, instead of pins, because the pin's type 
+		// could already be modified to some other class (for when we check CanConformPinType)
+		for (TFieldIterator<UProperty> ParamIt(Function); ParamIt && (ParamIt->PropertyFlags & CPF_Parm); ++ParamIt)
+		{
+			if (OutputPinName.IsEmpty() && ParamIt->HasAnyPropertyFlags(CPF_ReturnParm))
+			{
+				TaggedOutputParam = *ParamIt;
+				break;
+			}
+			else if (OutputPinName == ParamIt->GetName())
+			{
+				TaggedOutputParam = *ParamIt;
+				break;
+			}
+		}
+
+		if (TaggedOutputParam != nullptr)
+		{
+			const UEdGraphSchema_K2* K2Schema = GetDefault<UEdGraphSchema_K2>();
+			FEdGraphPinType PropertyPinType;
+
+			if (!K2Schema->ConvertPropertyToPinType(TaggedOutputParam, /*out*/PropertyPinType) || !CanConformPinType(FuncNode, PropertyPinType))
+			{
+				TaggedOutputParam = nullptr;
+			}
+		}
+	}
+
+	UEdGraphPin* DynamicOutPin = nullptr;
+	if (TaggedOutputParam != nullptr)
+	{
+		DynamicOutPin = FuncNode->FindPin(TaggedOutputParam->GetName());
+		if (DynamicOutPin && (DynamicOutPin->Direction != EGPD_Output))
+		{
+			DynamicOutPin = nullptr;
+		}
+	}
+	return DynamicOutPin;
+}
+
+bool FDynamicOutputHelper::CanConformPinType(const UK2Node_CallFunction* FuncNode, const FEdGraphPinType& TypeToTest)
+{
+	bool bIsProperType = false;
+	if (UEdGraphPin* TypePickerPin = GetTypePickerPin(FuncNode))
+	{
+		UClass* BasePickerClass = CastChecked<UClass>(TypePickerPin->PinType.PinSubCategoryObject.Get());
+
+		const FString& PinCategory = TypeToTest.PinCategory;
+		if ((PinCategory == UEdGraphSchema_K2::PC_Object) ||
+			(PinCategory == UEdGraphSchema_K2::PC_Interface) ||
+			(PinCategory == UEdGraphSchema_K2::PC_Class))
+		{
+			if (UClass* TypeClass = Cast<UClass>(TypeToTest.PinSubCategoryObject.Get()))
+			{
+				bIsProperType = BasePickerClass->IsChildOf(TypeClass);
+			}
+		}
+	}
+	return bIsProperType;
+}
+
+/*******************************************************************************
+ *  UK2Node_CallFunction
+ ******************************************************************************/
+
 UK2Node_CallFunction::UK2Node_CallFunction(const class FPostConstructInitializeProperties& PCIP)
 	: Super(PCIP)
 {
@@ -558,7 +872,6 @@ bool UK2Node_CallFunction::CreatePinsForFunctionCall(const UFunction* Function)
 	return bAllPinsGood;
 }
 
-
 void UK2Node_CallFunction::PostReconstructNode()
 {
 	Super::PostReconstructNode();
@@ -589,6 +902,11 @@ void UK2Node_CallFunction::PostReconstructNode()
 	{
 		ReturnPin->PinType.PinSubCategoryObject = TemplateComp->GetClass();
 	}
+
+	if (UEdGraphPin* TypePickerPin = FDynamicOutputHelper::GetTypePickerPin(this))
+	{
+		FDynamicOutputHelper(TypePickerPin).ConformOutputType();
+	}
 }
 
 void UK2Node_CallFunction::DestroyNode()
@@ -624,6 +942,14 @@ void UK2Node_CallFunction::NotifyPinConnectionListChanged(UEdGraphPin* Pin)
 			DestroyNode();
 		}
 	}
+
+	FDynamicOutputHelper(Pin).ConformOutputType();
+}
+
+void UK2Node_CallFunction::PinDefaultValueChanged(UEdGraphPin* Pin)
+{
+	Super::PinDefaultValueChanged(Pin);
+	FDynamicOutputHelper(Pin).ConformOutputType();
 }
 
 UFunction* UK2Node_CallFunction::GetTargetFunction() const
@@ -1283,6 +1609,8 @@ void UK2Node_CallFunction::ValidateNodeDuringCompilation(class FCompilerResultsL
 			}
 		}
 	}
+
+	FDynamicOutputHelper::VerifyNode(this, MessageLog);
 }
 
 void UK2Node_CallFunction::Serialize(FArchive& Ar)
diff --git a/Engine/Source/Editor/BlueprintGraph/Private/K2Node_SpawnActorFromClass.cpp b/Engine/Source/Editor/BlueprintGraph/Private/K2Node_SpawnActorFromClass.cpp
--- a/Engine/Source/Editor/BlueprintGraph/Private/K2Node_SpawnActorFromClass.cpp
+++ b/Engine/Source/Editor/BlueprintGraph/Private/K2Node_SpawnActorFromClass.cpp
@@ -182,6 +182,7 @@ void UK2Node_SpawnActorFromClass::OnClassPinChanged()
 	{
 		CreatePinsForClass(UseSpawnClass);
 	}
+	K2Schema->ConstructBasicPinTooltip(*ResultPin, LOCTEXT("ResultPinDescription", "The spawned Actor"), ResultPin->PinToolTip);
 
 	// Refresh the UI for the graph so the pin changes show up
 	UEdGraph* Graph = GetGraph();
@@ -297,8 +298,15 @@ FText UK2Node_SpawnActorFromClass::GetNodeTitle(ENodeTitleType::Type TitleType)
 			{
 				if (CachedNodeTitle.IsOutOfDate())
 				{
+					FText ClassName;
+					if (UClass* PickedClass = Cast<UClass>(ClassPin->DefaultObject))
+					{
+						ClassName = PickedClass->GetDisplayNameText();
+					}
+
 					FFormatNamedArguments Args;
-					Args.Add(TEXT("ClassName"), FText::FromString(ClassPin->DefaultObject->GetName()));
+					Args.Add(TEXT("ClassName"), ClassName);
+
 					// FText::Format() is slow, so we cache this to save on performance
 					CachedNodeTitle = FText::Format(NSLOCTEXT("K2Node", "SpawnActor", "SpawnActor {ClassName}"), Args);
 				}
diff --git a/Engine/Source/Runtime/Engine/Classes/Kismet/GameplayStatics.h b/Engine/Source/Runtime/Engine/Classes/Kismet/GameplayStatics.h
--- a/Engine/Source/Runtime/Engine/Classes/Kismet/GameplayStatics.h
+++ b/Engine/Source/Runtime/Engine/Classes/Kismet/GameplayStatics.h
@@ -60,7 +60,7 @@ class ENGINE_API UGameplayStatics : public UBlueprintFunctionLibrary
 	 *	@param	ActorClass	Class of Actor to find. Must be specified or result array will be empty.
 	 *	@param	OutActors	Output array of Actors of the specified class.
 	 */
-	UFUNCTION(BlueprintCallable, Category="Utilities",  meta=(WorldContext="WorldContextObject"))
+	UFUNCTION(BlueprintCallable, Category="Utilities",  meta=(WorldContext="WorldContextObject", DeterminesOutputType="ActorClass", DynamicOutputParam="OutActors"))
 	static void GetAllActorsOfClass(UObject* WorldContextObject, TSubclassOf<AActor> ActorClass, TArray<AActor*>& OutActors);
 
 	/** 
@@ -69,7 +69,7 @@ class ENGINE_API UGameplayStatics : public UBlueprintFunctionLibrary
 	 *	@param	Interface	Interface to find. Must be specified or result array will be empty.
 	 *	@param	OutActors	Output array of Actors of the specified class.
 	 */
-	UFUNCTION(BlueprintCallable, Category="Utilities",  meta=(WorldContext="WorldContextObject"))
+	UFUNCTION(BlueprintCallable, Category="Utilities",  meta=(WorldContext="WorldContextObject", DeterminesOutputType="Interface", DynamicOutputParam="OutActors"))
 	static void GetAllActorsWithInterface(UObject* WorldContextObject, TSubclassOf<UInterface> Interface, TArray<AActor*>& OutActors);
 
 	// --- Player functions ------------------------------
