Fixing issue where event nodes were using the skeleton class.

diff --git a/Engine/Source/Editor/BlueprintGraph/Private/BlueprintActionDatabase.cpp b/Engine/Source/Editor/BlueprintGraph/Private/BlueprintActionDatabase.cpp
--- a/Engine/Source/Editor/BlueprintGraph/Private/BlueprintActionDatabase.cpp
+++ b/Engine/Source/Editor/BlueprintGraph/Private/BlueprintActionDatabase.cpp
@@ -275,6 +275,14 @@ namespace BlueprintActionDatabaseImpl
 	 */
 	static bool IsPropertyBlueprintVisible(UProperty const* const Property);
 
+	/**
+	 * 
+	 * 
+	 * @param  Function	
+	 * @return 
+	 */
+	static bool IsBlueprintInterfaceFunction(const UFunction* Function);
+
 	/**
 	 * 
 	 * 
@@ -455,6 +463,33 @@ static bool BlueprintActionDatabaseImpl::IsPropertyBlueprintVisible(UProperty co
 	return !Property->HasAnyPropertyFlags(CPF_Parm) && (bIsAccessible || (bIsDelegate && bIsAssignableOrCallable));
 }
 
+//------------------------------------------------------------------------------
+static bool BlueprintActionDatabaseImpl::IsBlueprintInterfaceFunction(const UFunction* Function)
+{
+	bool bIsBpInterfaceFunc = false;
+	if (UClass* FuncClass = Function->GetOwnerClass())
+	{
+		if (UBlueprint* BpOuter = Cast<UBlueprint>(FuncClass->ClassGeneratedBy))
+		{
+			FName FuncName = Function->GetFName();
+
+			for (int32 InterfaceIndex = 0; (InterfaceIndex < BpOuter->ImplementedInterfaces.Num()) && !bIsBpInterfaceFunc; ++InterfaceIndex)
+			{
+				FBPInterfaceDescription& InterfaceDesc = BpOuter->ImplementedInterfaces[InterfaceIndex];
+				for (TFieldIterator<UFunction> FunctionIt(InterfaceDesc.Interface, EFieldIteratorFlags::IncludeSuper); FunctionIt; ++FunctionIt)
+				{
+					if (FunctionIt->GetFName() == FuncName)
+					{
+						bIsBpInterfaceFunc = true;
+						break;
+					}
+				}
+			}
+		}
+	}
+	return bIsBpInterfaceFunc;
+}
+
 //------------------------------------------------------------------------------
 static void BlueprintActionDatabaseImpl::GetClassMemberActions(UClass* const Class, FActionList& ActionListOut)
 {
@@ -486,8 +521,9 @@ static void BlueprintActionDatabaseImpl::AddClassFunctionActions(UClass const* c
 	for (TFieldIterator<UFunction> FunctionIt(Class, EFieldIteratorFlags::ExcludeSuper); FunctionIt; ++FunctionIt)
 	{
 		UFunction* Function = *FunctionIt;
+		bool const bIsBpInterfaceFunc = BlueprintActionDatabaseImpl::IsBlueprintInterfaceFunction(Function);
 
-		if (UEdGraphSchema_K2::FunctionCanBePlacedAsEvent(Function))
+		if (UEdGraphSchema_K2::FunctionCanBePlacedAsEvent(Function) && !bIsBpInterfaceFunc)
 		{
 			if (UBlueprintEventNodeSpawner* NodeSpawner = UBlueprintEventNodeSpawner::Create(Function))
 			{
diff --git a/Engine/Source/Editor/BlueprintGraph/Private/BlueprintActionFilter.cpp b/Engine/Source/Editor/BlueprintGraph/Private/BlueprintActionFilter.cpp
--- a/Engine/Source/Editor/BlueprintGraph/Private/BlueprintActionFilter.cpp
+++ b/Engine/Source/Editor/BlueprintGraph/Private/BlueprintActionFilter.cpp
@@ -1240,7 +1240,8 @@ static bool BlueprintActionFilterImpl::IsExtraneousInterfaceCall(FBlueprintActio
 {
 	bool bIsFilteredOut = false;
 
-	if (BlueprintAction.GetNodeClass()->IsChildOf<UK2Node_Message>())
+	const UClass* NodeClass = BlueprintAction.GetNodeClass();
+	if (NodeClass->IsChildOf<UK2Node_Message>())
 	{
 		UFunction const* Function = BlueprintAction.GetAssociatedFunction();
 		checkSlow(Function != nullptr);
@@ -1253,7 +1254,7 @@ static bool BlueprintActionFilterImpl::IsExtraneousInterfaceCall(FBlueprintActio
 		{
 			// if the target class implements (or is) the message spawner's 
 			// interface owner, then we don't need the message node (the regular
-			// interface call shouldt be available)
+			// interface call shouldn't be available)
 			if (!IsClassOfType(TargetClass, InterfaceClass))
 			{
 				bIsFilteredOut = false;
@@ -1266,7 +1267,7 @@ static bool BlueprintActionFilterImpl::IsExtraneousInterfaceCall(FBlueprintActio
 		UClass* FuncClass = Function->GetOwnerClass();
 		bool const bIsInterfaceAction = FuncClass->IsChildOf<UInterface>();
 
-		if (bIsInterfaceAction)
+		if (bIsInterfaceAction && !NodeClass->IsChildOf<UK2Node_Event>())
 		{
 			bIsFilteredOut = (Filter.TargetClasses.Num() > 0);
 			for (const UClass* TargetClass : Filter.TargetClasses)
@@ -1541,11 +1542,6 @@ bool FBlueprintActionFilter::IsFilteredByThis(FBlueprintActionInfo& BlueprintAct
 {
 	FBlueprintActionFilter const& FilterRef = *this;
 
-	if (BlueprintAction.NodeSpawner->DefaultMenuSignature.MenuName.ToString().StartsWith("My"))
-	{
-		printf("");
-	}
-
 	bool bIsFiltered = false;
 	// iterate backwards so that custom user test are ran first (and the slow
 	// internal tests are ran last).
diff --git a/Engine/Source/Editor/BlueprintGraph/Private/BlueprintEventNodeSpawner.cpp b/Engine/Source/Editor/BlueprintGraph/Private/BlueprintEventNodeSpawner.cpp
--- a/Engine/Source/Editor/BlueprintGraph/Private/BlueprintEventNodeSpawner.cpp
+++ b/Engine/Source/Editor/BlueprintGraph/Private/BlueprintEventNodeSpawner.cpp
@@ -158,7 +158,7 @@ UEdGraphNode* UBlueprintEventNodeSpawner::Invoke(UEdGraph* ParentGraph, FBinding
 			if (EventFunc != nullptr)
 			{
 				EventNode->EventSignatureName  = EventName;
-				EventNode->EventSignatureClass = EventFunc->GetOuterUClass();
+				EventNode->EventSignatureClass = EventFunc->GetOuterUClass()->GetAuthoritativeClass();
 				EventNode->bOverrideFunction   = true;
 			}
 			else if (!bIsTemplateNode)
@@ -197,7 +197,7 @@ UK2Node_Event const* UBlueprintEventNodeSpawner::FindPreExistingEvent(UBlueprint
 	else
 	{
 		check(EventFunc != nullptr);
-		UClass* ClassOwner = EventFunc->GetOwnerClass();
+		UClass* ClassOwner = EventFunc->GetOwnerClass()->GetAuthoritativeClass();
 
 		PreExistingNode = FBlueprintEditorUtils::FindOverrideForFunction(Blueprint, ClassOwner, EventFunc->GetFName());
 	}
