Can right click on any variable get/set node to replace it with a member or local variable.

#ttp 346457 - NUXF-66 Can't rename a 'Set Variable' node variable
#NUX

diff --git a/Engine/Source/Editor/BlueprintGraph/Classes/EdGraphSchema_K2.h b/Engine/Source/Editor/BlueprintGraph/Classes/EdGraphSchema_K2.h
--- a/Engine/Source/Editor/BlueprintGraph/Classes/EdGraphSchema_K2.h
+++ b/Engine/Source/Editor/BlueprintGraph/Classes/EdGraphSchema_K2.h
@@ -877,6 +877,16 @@ public:
 	/** Create menu for variable get/set nodes which refer to a variable which does not exist. */
 	void GetNonExistentVariableMenu(const UEdGraphNode* InGraphNode, UBlueprint* OwnerBlueprint, FMenuBuilder* MenuBuilder) const;
 
+	/**
+	 * Create menu for variable get/set nodes which allows for the replacement of variables
+	 *
+	 * @param InGraphNode					Variable node to replace
+	 * @param InOwnerBlueprint				The owning Blueprint of the variable
+	 * @param InMenuBuilder					MenuBuilder to place the menu items into
+	 * @param bInReplaceExistingVariable	TRUE if replacing an existing variable, will keep the variable from appearing on the list
+	 */
+	void GetReplaceVariableMenu(const UEdGraphNode* InGraphNode, UBlueprint* InOwnerBlueprint, FMenuBuilder* InMenuBuilder, bool bInReplaceExistingVariable = false) const;
+
 	// Calculates an average position between the nodes owning the two specified pins
 	static FVector2D CalculateAveragePositionBetweenNodes(UEdGraphPin* InputPin, UEdGraphPin* OutputPin);
 
@@ -951,7 +961,7 @@ public:
 	static void OnReplaceVariableForVariableNode(class UK2Node_Variable* Variable, UBlueprint* OwnerBlueprint, FString VariableName, bool bIsSelfMember);
 
 	/** Create sub menu that shows all possible variables that can be used to replace the existing variable reference */
-	static void GetReplaceNonExistentVariableMenu(class FMenuBuilder& MenuBuilder, class UK2Node_Variable* Variable, UBlueprint* OwnerBlueprint);
+	static void GetReplaceVariableMenu(class FMenuBuilder& MenuBuilder, class UK2Node_Variable* Variable, UBlueprint* OwnerBlueprint, bool bReplaceExistingVariable = false);
 
 private:
 
diff --git a/Engine/Source/Editor/BlueprintGraph/Private/EdGraphSchema_K2.cpp b/Engine/Source/Editor/BlueprintGraph/Private/EdGraphSchema_K2.cpp
--- a/Engine/Source/Editor/BlueprintGraph/Private/EdGraphSchema_K2.cpp
+++ b/Engine/Source/Editor/BlueprintGraph/Private/EdGraphSchema_K2.cpp
@@ -1181,6 +1181,7 @@ void UEdGraphSchema_K2::GetContextMenuActions(const UEdGraph* CurrentGraph, cons
 					{
 						MenuBuilder->AddMenuEntry(FGraphEditorCommands::Get().FindVariableReferences);
 						MenuBuilder->AddMenuEntry(FGraphEditorCommands::Get().GotoNativeVariableDefinition);
+						GetReplaceVariableMenu(InGraphNode,OwnerBlueprint, MenuBuilder, true);
 					}
 
 					if (InGraphNode->IsA(UK2Node_SetFieldsInStruct::StaticClass()))
@@ -1273,6 +1274,10 @@ void UEdGraphSchema_K2::OnReplaceVariableForVariableNode( UK2Node_Variable* Vari
 {
 	if(UEdGraphPin* Pin = Variable->FindPin(Variable->GetVarNameString()))
 	{
+		const FScopedTransaction Transaction( NSLOCTEXT("UnrealEd", "GraphEd_ReplaceVariable", "Replace Variable") );
+		Variable->Modify();
+		Pin->Modify();
+
 		if (bIsSelfMember)
 		{
 			Variable->VariableReference.SetSelfMember( FName(*VariableName) );
@@ -1286,35 +1291,71 @@ void UEdGraphSchema_K2::OnReplaceVariableForVariableNode( UK2Node_Variable* Vari
 	}
 }
 
-void UEdGraphSchema_K2::GetReplaceNonExistentVariableMenu(FMenuBuilder& MenuBuilder, UK2Node_Variable* Variable,  UBlueprint* OwnerBlueprint)
+void UEdGraphSchema_K2::GetReplaceVariableMenu(FMenuBuilder& MenuBuilder, UK2Node_Variable* Variable,  UBlueprint* OwnerBlueprint, bool bReplaceExistingVariable/*=false*/)
 {
 	if (UEdGraphPin* Pin = Variable->FindPin(Variable->GetVarNameString()))
 	{
+		FName ExistingVariableName = bReplaceExistingVariable? Variable->GetVarName() : NAME_None;
+
+		FText ReplaceVariableWithTooltipFormat;
+		if(!bReplaceExistingVariable)
+		{
+			ReplaceVariableWithTooltipFormat = LOCTEXT("ReplaceNonExistantVarToolTip", "Variable '{OldVariable}' does not exist, replace with matching variable '{AlternateVariable}'?");
+		}
+		else
+		{
+			ReplaceVariableWithTooltipFormat = LOCTEXT("ReplaceExistantVarToolTip", "Replace Variable '{OldVariable}' with matching variable '{AlternateVariable}'?");
+		}
+
 		TArray<FName> Variables;
 		FBlueprintEditorUtils::GetNewVariablesOfType(OwnerBlueprint, Pin->PinType, Variables);
 
 		MenuBuilder.BeginSection(NAME_None, LOCTEXT("Variables", "Variables"));
 		for (TArray<FName>::TIterator VarIt(Variables); VarIt; ++VarIt)
 		{
-			const FText AlternativeVar = FText::FromName( *VarIt );
-			const FText Desc = FText::Format( LOCTEXT("ReplaceNonExistantVarToolTip", "Variable '{0}' does not exist, replace with matching variable '{0}'?"), Variable->GetVarNameText(), AlternativeVar );
+			if(*VarIt != ExistingVariableName)
+			{
+				const FText AlternativeVar = FText::FromName( *VarIt );
 
-			MenuBuilder.AddMenuEntry( AlternativeVar, Desc, FSlateIcon(), FUIAction(
-				FExecuteAction::CreateStatic(&UEdGraphSchema_K2::OnReplaceVariableForVariableNode, const_cast<UK2Node_Variable* >(Variable),OwnerBlueprint, (*VarIt).ToString(), /*bIsSelfMember=*/true ) ) );
+				FFormatNamedArguments TooltipArgs;
+				TooltipArgs.Add(TEXT("OldVariable"), Variable->GetVarNameText());
+				TooltipArgs.Add(TEXT("AlternateVariable"), AlternativeVar);
+				const FText Desc = FText::Format(ReplaceVariableWithTooltipFormat, TooltipArgs);
+
+				MenuBuilder.AddMenuEntry( AlternativeVar, Desc, FSlateIcon(), FUIAction(
+					FExecuteAction::CreateStatic(&UEdGraphSchema_K2::OnReplaceVariableForVariableNode, const_cast<UK2Node_Variable* >(Variable),OwnerBlueprint, (*VarIt).ToString(), /*bIsSelfMember=*/true ) ) );
+			}
 		}
 		MenuBuilder.EndSection();
 
+		FText ReplaceLocalVariableWithTooltipFormat;
+		if(!bReplaceExistingVariable)
+		{
+			ReplaceLocalVariableWithTooltipFormat = LOCTEXT("ReplaceNonExistantLocalVarToolTip", "Variable '{OldVariable}' does not exist, replace with matching local variable '{AlternateVariable}'?");
+		}
+		else
+		{
+			ReplaceLocalVariableWithTooltipFormat = LOCTEXT("ReplaceExistantLocalVarToolTip", "Replace Variable '{OldVariable}' with matching local variable '{AlternateVariable}'?");
+		}
+
 		TArray<FName> LocalVariables;
 		FBlueprintEditorUtils::GetLocalVariablesOfType(Variable->GetGraph(), Pin->PinType, LocalVariables);
 
 		MenuBuilder.BeginSection(NAME_None, LOCTEXT("LocalVariables", "LocalVariables"));
 		for (TArray<FName>::TIterator VarIt(LocalVariables); VarIt; ++VarIt)
 		{
-			const FText AlternativeVar = FText::FromName( *VarIt );
-			const FText Desc = FText::Format( LOCTEXT("ReplaceNonExistantLocalVarToolTip", "Variable '{0}' does not exist, replace with matching localvariable '{0}'?"), Variable->GetVarNameText(), AlternativeVar );
+			if(*VarIt != ExistingVariableName)
+			{
+				const FText AlternativeVar = FText::FromName( *VarIt );
+
+				FFormatNamedArguments TooltipArgs;
+				TooltipArgs.Add(TEXT("OldVariable"), Variable->GetVarNameText());
+				TooltipArgs.Add(TEXT("AlternateVariable"), AlternativeVar);
+				const FText Desc = FText::Format( ReplaceLocalVariableWithTooltipFormat, TooltipArgs );
 
-			MenuBuilder.AddMenuEntry( AlternativeVar, Desc, FSlateIcon(), FUIAction(
-				FExecuteAction::CreateStatic(&UEdGraphSchema_K2::OnReplaceVariableForVariableNode, const_cast<UK2Node_Variable* >(Variable),OwnerBlueprint, (*VarIt).ToString(), /*bIsSelfMember=*/false ) ) );
+				MenuBuilder.AddMenuEntry( AlternativeVar, Desc, FSlateIcon(), FUIAction(
+					FExecuteAction::CreateStatic(&UEdGraphSchema_K2::OnReplaceVariableForVariableNode, const_cast<UK2Node_Variable* >(Variable),OwnerBlueprint, (*VarIt).ToString(), /*bIsSelfMember=*/false ) ) );
+			}
 		}
 		MenuBuilder.EndSection();
 	}
@@ -1354,22 +1395,44 @@ void UEdGraphSchema_K2::GetNonExistentVariableMenu( const UEdGraphNode* InGraphN
 			MenuBuilder->AddMenuEntry( FGenericCommands::Get().Delete, NAME_None, FGenericCommands::Get().Delete->GetLabel(), Desc);
 		}
 
+		GetReplaceVariableMenu(InGraphNode, OwnerBlueprint, MenuBuilder);
+	}
+}
+
+void UEdGraphSchema_K2::GetReplaceVariableMenu(const UEdGraphNode* InGraphNode, UBlueprint* InOwnerBlueprint, FMenuBuilder* InMenuBuilder, bool bInReplaceExistingVariable/* = false*/) const
+{
+	if (const UK2Node_Variable* Variable = Cast<const UK2Node_Variable>(InGraphNode))
+	{
 		// replace with matching variables
 		if (UEdGraphPin* Pin = Variable->FindPin(Variable->GetVarNameString()))
 		{
+			FName ExistingVariableName = bInReplaceExistingVariable? Variable->GetVarName() : NAME_None;
+
 			TArray<FName> Variables;
-			FBlueprintEditorUtils::GetNewVariablesOfType(OwnerBlueprint, Pin->PinType, Variables);
+			FBlueprintEditorUtils::GetNewVariablesOfType(InOwnerBlueprint, Pin->PinType, Variables);
+			Variables.RemoveSwap(ExistingVariableName);
 
 			TArray<FName> LocalVariables;
 			FBlueprintEditorUtils::GetLocalVariablesOfType(Variable->GetGraph(), Pin->PinType, LocalVariables);
+			LocalVariables.RemoveSwap(ExistingVariableName);
 
 			if (Variables.Num() > 0 || LocalVariables.Num() > 0)
 			{
-				MenuBuilder->AddSubMenu(
+				FText ReplaceVariableWithTooltip;
+				if(bInReplaceExistingVariable)
+				{
+					ReplaceVariableWithTooltip = LOCTEXT("ReplaceVariableWithToolTip", "Replace Variable '{0}' with another variable?");
+				}
+				else
+				{
+					ReplaceVariableWithTooltip = LOCTEXT("ReplaceMissingVariableWithToolTip", "Variable '{0}' does not exist, replace with another variable?");
+				}
+
+				InMenuBuilder->AddSubMenu(
 					FText::Format( LOCTEXT("ReplaceVariableWith", "Replace variable '{0}' with..."), Variable->GetVarNameText()),
-					FText::Format( LOCTEXT("ReplaceVariableWithToolTip", "Variable '{0}' does not exist, replace with another variable?"), Variable->GetVarNameText()),
-					FNewMenuDelegate::CreateStatic( &UEdGraphSchema_K2::GetReplaceNonExistentVariableMenu,
-					const_cast<UK2Node_Variable*>(Variable),OwnerBlueprint));
+					FText::Format( ReplaceVariableWithTooltip, Variable->GetVarNameText()),
+					FNewMenuDelegate::CreateStatic( &UEdGraphSchema_K2::GetReplaceVariableMenu,
+					const_cast<UK2Node_Variable*>(Variable),InOwnerBlueprint, bInReplaceExistingVariable));
 			}
 		}
 	}
