Rename FPostConstructInitializeProperties to something simpler

diff --git a/Engine/Source/Editor/AnimGraph/Private/AnimGraphNode_SaveCachedPose.cpp b/Engine/Source/Editor/AnimGraph/Private/AnimGraphNode_SaveCachedPose.cpp
--- a/Engine/Source/Editor/AnimGraph/Private/AnimGraphNode_SaveCachedPose.cpp
+++ b/Engine/Source/Editor/AnimGraph/Private/AnimGraphNode_SaveCachedPose.cpp
@@ -34,8 +34,8 @@ public:
 
 #define LOCTEXT_NAMESPACE "A3Nodes"
 
-UAnimGraphNode_SaveCachedPose::UAnimGraphNode_SaveCachedPose(const FPostConstructInitializeProperties& PCIP)
-	: Super(PCIP)
+UAnimGraphNode_SaveCachedPose::UAnimGraphNode_SaveCachedPose(const FObjectInitializer& ObjectInitializer)
+	: Super(ObjectInitializer)
 {
 	bCanRenameNode = true;
 }
diff --git a/Engine/Source/Editor/AnimGraph/Private/AnimGraphNode_SequenceEvaluator.cpp b/Engine/Source/Editor/AnimGraph/Private/AnimGraphNode_SequenceEvaluator.cpp
--- a/Engine/Source/Editor/AnimGraph/Private/AnimGraphNode_SequenceEvaluator.cpp
+++ b/Engine/Source/Editor/AnimGraph/Private/AnimGraphNode_SequenceEvaluator.cpp
@@ -11,8 +11,8 @@
 
 #define LOCTEXT_NAMESPACE "A3Nodes"
 
-UAnimGraphNode_SequenceEvaluator::UAnimGraphNode_SequenceEvaluator(const FPostConstructInitializeProperties& PCIP)
-	: Super(PCIP)
+UAnimGraphNode_SequenceEvaluator::UAnimGraphNode_SequenceEvaluator(const FObjectInitializer& ObjectInitializer)
+	: Super(ObjectInitializer)
 {
 }
 
diff --git a/Engine/Source/Editor/AnimGraph/Private/AnimGraphNode_SequencePlayer.cpp b/Engine/Source/Editor/AnimGraph/Private/AnimGraphNode_SequencePlayer.cpp
--- a/Engine/Source/Editor/AnimGraph/Private/AnimGraphNode_SequencePlayer.cpp
+++ b/Engine/Source/Editor/AnimGraph/Private/AnimGraphNode_SequencePlayer.cpp
@@ -50,8 +50,8 @@ public:
 /////////////////////////////////////////////////////
 // UAnimGraphNode_SequencePlayer
 
-UAnimGraphNode_SequencePlayer::UAnimGraphNode_SequencePlayer(const FPostConstructInitializeProperties& PCIP)
-	: Super(PCIP)
+UAnimGraphNode_SequencePlayer::UAnimGraphNode_SequencePlayer(const FObjectInitializer& ObjectInitializer)
+	: Super(ObjectInitializer)
 {
 }
 
diff --git a/Engine/Source/Editor/AnimGraph/Private/AnimGraphNode_SkeletalControlBase.cpp b/Engine/Source/Editor/AnimGraph/Private/AnimGraphNode_SkeletalControlBase.cpp
--- a/Engine/Source/Editor/AnimGraph/Private/AnimGraphNode_SkeletalControlBase.cpp
+++ b/Engine/Source/Editor/AnimGraph/Private/AnimGraphNode_SkeletalControlBase.cpp
@@ -10,8 +10,8 @@
 
 #define LOCTEXT_NAMESPACE "A3Nodes"
 
-UAnimGraphNode_SkeletalControlBase::UAnimGraphNode_SkeletalControlBase(const FPostConstructInitializeProperties& PCIP)
-	: Super(PCIP)
+UAnimGraphNode_SkeletalControlBase::UAnimGraphNode_SkeletalControlBase(const FObjectInitializer& ObjectInitializer)
+	: Super(ObjectInitializer)
 {
 }
 
diff --git a/Engine/Source/Editor/AnimGraph/Private/AnimGraphNode_Slot.cpp b/Engine/Source/Editor/AnimGraph/Private/AnimGraphNode_Slot.cpp
--- a/Engine/Source/Editor/AnimGraph/Private/AnimGraphNode_Slot.cpp
+++ b/Engine/Source/Editor/AnimGraph/Private/AnimGraphNode_Slot.cpp
@@ -11,8 +11,8 @@
 
 #define LOCTEXT_NAMESPACE "A3Nodes"
 
-UAnimGraphNode_Slot::UAnimGraphNode_Slot(const FPostConstructInitializeProperties& PCIP)
-	: Super(PCIP)
+UAnimGraphNode_Slot::UAnimGraphNode_Slot(const FObjectInitializer& ObjectInitializer)
+	: Super(ObjectInitializer)
 {
 }
 
diff --git a/Engine/Source/Editor/AnimGraph/Private/AnimGraphNode_SpringBone.cpp b/Engine/Source/Editor/AnimGraph/Private/AnimGraphNode_SpringBone.cpp
--- a/Engine/Source/Editor/AnimGraph/Private/AnimGraphNode_SpringBone.cpp
+++ b/Engine/Source/Editor/AnimGraph/Private/AnimGraphNode_SpringBone.cpp
@@ -8,8 +8,8 @@
 
 #define LOCTEXT_NAMESPACE "A3Nodes"
 
-UAnimGraphNode_SpringBone::UAnimGraphNode_SpringBone(const FPostConstructInitializeProperties& PCIP)
-	: Super(PCIP)
+UAnimGraphNode_SpringBone::UAnimGraphNode_SpringBone(const FObjectInitializer& ObjectInitializer)
+	: Super(ObjectInitializer)
 {
 }
 
diff --git a/Engine/Source/Editor/AnimGraph/Private/AnimGraphNode_StateMachine.cpp b/Engine/Source/Editor/AnimGraph/Private/AnimGraphNode_StateMachine.cpp
--- a/Engine/Source/Editor/AnimGraph/Private/AnimGraphNode_StateMachine.cpp
+++ b/Engine/Source/Editor/AnimGraph/Private/AnimGraphNode_StateMachine.cpp
@@ -6,8 +6,8 @@
 /////////////////////////////////////////////////////
 // UAnimGraphNode_StateMachine
 
-UAnimGraphNode_StateMachine::UAnimGraphNode_StateMachine(const FPostConstructInitializeProperties& PCIP)
-	: Super(PCIP)
+UAnimGraphNode_StateMachine::UAnimGraphNode_StateMachine(const FObjectInitializer& ObjectInitializer)
+	: Super(ObjectInitializer)
 {
 	bCanRenameNode = true;
 }
diff --git a/Engine/Source/Editor/AnimGraph/Private/AnimGraphNode_StateMachineBase.cpp b/Engine/Source/Editor/AnimGraph/Private/AnimGraphNode_StateMachineBase.cpp
--- a/Engine/Source/Editor/AnimGraph/Private/AnimGraphNode_StateMachineBase.cpp
+++ b/Engine/Source/Editor/AnimGraph/Private/AnimGraphNode_StateMachineBase.cpp
@@ -42,8 +42,8 @@ public:
 
 #define LOCTEXT_NAMESPACE "A3Nodes"
 
-UAnimGraphNode_StateMachineBase::UAnimGraphNode_StateMachineBase(const FPostConstructInitializeProperties& PCIP)
-	: Super(PCIP)
+UAnimGraphNode_StateMachineBase::UAnimGraphNode_StateMachineBase(const FObjectInitializer& ObjectInitializer)
+	: Super(ObjectInitializer)
 {
 }
 
diff --git a/Engine/Source/Editor/AnimGraph/Private/AnimGraphNode_StateResult.cpp b/Engine/Source/Editor/AnimGraph/Private/AnimGraphNode_StateResult.cpp
--- a/Engine/Source/Editor/AnimGraph/Private/AnimGraphNode_StateResult.cpp
+++ b/Engine/Source/Editor/AnimGraph/Private/AnimGraphNode_StateResult.cpp
@@ -6,7 +6,7 @@
 /////////////////////////////////////////////////////
 // UAnimGraphNode_StateResult
 
-UAnimGraphNode_StateResult::UAnimGraphNode_StateResult(const FPostConstructInitializeProperties& PCIP)
-	: Super(PCIP)
+UAnimGraphNode_StateResult::UAnimGraphNode_StateResult(const FObjectInitializer& ObjectInitializer)
+	: Super(ObjectInitializer)
 {
 }
diff --git a/Engine/Source/Editor/AnimGraph/Private/AnimGraphNode_Trail.cpp b/Engine/Source/Editor/AnimGraph/Private/AnimGraphNode_Trail.cpp
--- a/Engine/Source/Editor/AnimGraph/Private/AnimGraphNode_Trail.cpp
+++ b/Engine/Source/Editor/AnimGraph/Private/AnimGraphNode_Trail.cpp
@@ -8,8 +8,8 @@
 
 #define LOCTEXT_NAMESPACE "A3Nodes"
 
-UAnimGraphNode_Trail::UAnimGraphNode_Trail(const FPostConstructInitializeProperties& PCIP)
-	: Super(PCIP)
+UAnimGraphNode_Trail::UAnimGraphNode_Trail(const FObjectInitializer& ObjectInitializer)
+	: Super(ObjectInitializer)
 {
 }
 
diff --git a/Engine/Source/Editor/AnimGraph/Private/AnimGraphNode_TransitionPoseEvaluator.cpp b/Engine/Source/Editor/AnimGraph/Private/AnimGraphNode_TransitionPoseEvaluator.cpp
--- a/Engine/Source/Editor/AnimGraph/Private/AnimGraphNode_TransitionPoseEvaluator.cpp
+++ b/Engine/Source/Editor/AnimGraph/Private/AnimGraphNode_TransitionPoseEvaluator.cpp
@@ -12,8 +12,8 @@
 
 #define LOCTEXT_NAMESPACE "UAnimGraphNode_TransitionPoseEvaluator"
 
-UAnimGraphNode_TransitionPoseEvaluator::UAnimGraphNode_TransitionPoseEvaluator(const FPostConstructInitializeProperties& PCIP)
-	: Super(PCIP)
+UAnimGraphNode_TransitionPoseEvaluator::UAnimGraphNode_TransitionPoseEvaluator(const FObjectInitializer& ObjectInitializer)
+	: Super(ObjectInitializer)
 {
 }
 
diff --git a/Engine/Source/Editor/AnimGraph/Private/AnimGraphNode_TransitionResult.cpp b/Engine/Source/Editor/AnimGraph/Private/AnimGraphNode_TransitionResult.cpp
--- a/Engine/Source/Editor/AnimGraph/Private/AnimGraphNode_TransitionResult.cpp
+++ b/Engine/Source/Editor/AnimGraph/Private/AnimGraphNode_TransitionResult.cpp
@@ -12,8 +12,8 @@
 
 #define LOCTEXT_NAMESPACE "UAnimGraphNode_TransitionResult"
 
-UAnimGraphNode_TransitionResult::UAnimGraphNode_TransitionResult(const FPostConstructInitializeProperties& PCIP)
-	: Super(PCIP)
+UAnimGraphNode_TransitionResult::UAnimGraphNode_TransitionResult(const FObjectInitializer& ObjectInitializer)
+	: Super(ObjectInitializer)
 {
 }
 
diff --git a/Engine/Source/Editor/AnimGraph/Private/AnimGraphNode_TwoBoneIK.cpp b/Engine/Source/Editor/AnimGraph/Private/AnimGraphNode_TwoBoneIK.cpp
--- a/Engine/Source/Editor/AnimGraph/Private/AnimGraphNode_TwoBoneIK.cpp
+++ b/Engine/Source/Editor/AnimGraph/Private/AnimGraphNode_TwoBoneIK.cpp
@@ -31,8 +31,8 @@ public:
 // UAnimGraphNode_TwoBoneIK
 
 
-UAnimGraphNode_TwoBoneIK::UAnimGraphNode_TwoBoneIK(const FPostConstructInitializeProperties& PCIP)
-	: Super(PCIP)
+UAnimGraphNode_TwoBoneIK::UAnimGraphNode_TwoBoneIK(const FObjectInitializer& ObjectInitializer)
+	: Super(ObjectInitializer)
 {
 }
 
diff --git a/Engine/Source/Editor/AnimGraph/Private/AnimGraphNode_TwoWayBlend.cpp b/Engine/Source/Editor/AnimGraph/Private/AnimGraphNode_TwoWayBlend.cpp
--- a/Engine/Source/Editor/AnimGraph/Private/AnimGraphNode_TwoWayBlend.cpp
+++ b/Engine/Source/Editor/AnimGraph/Private/AnimGraphNode_TwoWayBlend.cpp
@@ -8,8 +8,8 @@
 
 #define LOCTEXT_NAMESPACE "AnimGraphNode_TwoWayBlend"
 
-UAnimGraphNode_TwoWayBlend::UAnimGraphNode_TwoWayBlend(const FPostConstructInitializeProperties& PCIP)
-	: Super(PCIP)
+UAnimGraphNode_TwoWayBlend::UAnimGraphNode_TwoWayBlend(const FObjectInitializer& ObjectInitializer)
+	: Super(ObjectInitializer)
 {
 }
 
diff --git a/Engine/Source/Editor/AnimGraph/Private/AnimGraphNode_UseCachedPose.cpp b/Engine/Source/Editor/AnimGraph/Private/AnimGraphNode_UseCachedPose.cpp
--- a/Engine/Source/Editor/AnimGraph/Private/AnimGraphNode_UseCachedPose.cpp
+++ b/Engine/Source/Editor/AnimGraph/Private/AnimGraphNode_UseCachedPose.cpp
@@ -17,8 +17,8 @@
 
 #define LOCTEXT_NAMESPACE "A3Nodes"
 
-UAnimGraphNode_UseCachedPose::UAnimGraphNode_UseCachedPose(const FPostConstructInitializeProperties& PCIP)
-	: Super(PCIP)
+UAnimGraphNode_UseCachedPose::UAnimGraphNode_UseCachedPose(const FObjectInitializer& ObjectInitializer)
+	: Super(ObjectInitializer)
 {
 }
 
diff --git a/Engine/Source/Editor/AnimGraph/Private/AnimGraphNode_WheelHandler.cpp b/Engine/Source/Editor/AnimGraph/Private/AnimGraphNode_WheelHandler.cpp
--- a/Engine/Source/Editor/AnimGraph/Private/AnimGraphNode_WheelHandler.cpp
+++ b/Engine/Source/Editor/AnimGraph/Private/AnimGraphNode_WheelHandler.cpp
@@ -9,8 +9,8 @@
 
 #define LOCTEXT_NAMESPACE "A3Nodes"
 
-UAnimGraphNode_WheelHandler::UAnimGraphNode_WheelHandler(const FPostConstructInitializeProperties& PCIP)
-	: Super(PCIP)
+UAnimGraphNode_WheelHandler::UAnimGraphNode_WheelHandler(const FObjectInitializer& ObjectInitializer)
+	: Super(ObjectInitializer)
 {
 }
 
diff --git a/Engine/Source/Editor/AnimGraph/Private/AnimPreviewInstance.cpp b/Engine/Source/Editor/AnimGraph/Private/AnimPreviewInstance.cpp
--- a/Engine/Source/Editor/AnimGraph/Private/AnimPreviewInstance.cpp
+++ b/Engine/Source/Editor/AnimGraph/Private/AnimPreviewInstance.cpp
@@ -4,8 +4,8 @@
 #include "AnimGraphPrivatePCH.h"
 #include "AnimPreviewInstance.h"
 
-UAnimPreviewInstance::UAnimPreviewInstance(const class FPostConstructInitializeProperties& PCIP)
-	: Super(PCIP)
+UAnimPreviewInstance::UAnimPreviewInstance(const FObjectInitializer& ObjectInitializer)
+	: Super(ObjectInitializer)
 	, SkeletalControlAlpha(1.0f)
 #if WITH_EDITORONLY_DATA
 	, bForceRetargetBasePose(false)
diff --git a/Engine/Source/Editor/AnimGraph/Private/AnimStateConduitNode.cpp b/Engine/Source/Editor/AnimGraph/Private/AnimStateConduitNode.cpp
--- a/Engine/Source/Editor/AnimGraph/Private/AnimStateConduitNode.cpp
+++ b/Engine/Source/Editor/AnimGraph/Private/AnimStateConduitNode.cpp
@@ -14,8 +14,8 @@
 /////////////////////////////////////////////////////
 // UAnimStateConduitNode
 
-UAnimStateConduitNode::UAnimStateConduitNode(const class FPostConstructInitializeProperties& PCIP)
-	: Super(PCIP)
+UAnimStateConduitNode::UAnimStateConduitNode(const FObjectInitializer& ObjectInitializer)
+	: Super(ObjectInitializer)
 {
 	bCanRenameNode = true;
 }
diff --git a/Engine/Source/Editor/AnimGraph/Private/AnimStateEntryNode.cpp b/Engine/Source/Editor/AnimGraph/Private/AnimStateEntryNode.cpp
--- a/Engine/Source/Editor/AnimGraph/Private/AnimStateEntryNode.cpp
+++ b/Engine/Source/Editor/AnimGraph/Private/AnimStateEntryNode.cpp
@@ -14,8 +14,8 @@
 /////////////////////////////////////////////////////
 // UAnimStateEntryNode
 
-UAnimStateEntryNode::UAnimStateEntryNode(const class FPostConstructInitializeProperties& PCIP)
-	: Super(PCIP)
+UAnimStateEntryNode::UAnimStateEntryNode(const FObjectInitializer& ObjectInitializer)
+	: Super(ObjectInitializer)
 {
 }
 
diff --git a/Engine/Source/Editor/AnimGraph/Private/AnimStateNode.cpp b/Engine/Source/Editor/AnimGraph/Private/AnimStateNode.cpp
--- a/Engine/Source/Editor/AnimGraph/Private/AnimStateNode.cpp
+++ b/Engine/Source/Editor/AnimGraph/Private/AnimStateNode.cpp
@@ -18,8 +18,8 @@
 /////////////////////////////////////////////////////
 // UAnimStateNode
 
-UAnimStateNode::UAnimStateNode(const class FPostConstructInitializeProperties& PCIP)
-	: Super(PCIP)
+UAnimStateNode::UAnimStateNode(const FObjectInitializer& ObjectInitializer)
+	: Super(ObjectInitializer)
 {
 	bCanRenameNode = true;
 }
diff --git a/Engine/Source/Editor/AnimGraph/Private/AnimStateNodeBase.cpp b/Engine/Source/Editor/AnimGraph/Private/AnimStateNodeBase.cpp
--- a/Engine/Source/Editor/AnimGraph/Private/AnimStateNodeBase.cpp
+++ b/Engine/Source/Editor/AnimGraph/Private/AnimStateNodeBase.cpp
@@ -37,8 +37,8 @@ public:
 /////////////////////////////////////////////////////
 // UAnimStateNodeBase
 
-UAnimStateNodeBase::UAnimStateNodeBase(const class FPostConstructInitializeProperties& PCIP)
-	: Super(PCIP)
+UAnimStateNodeBase::UAnimStateNodeBase(const FObjectInitializer& ObjectInitializer)
+	: Super(ObjectInitializer)
 {
 }
 
diff --git a/Engine/Source/Editor/AnimGraph/Private/AnimStateTransitionNode.cpp b/Engine/Source/Editor/AnimGraph/Private/AnimStateTransitionNode.cpp
--- a/Engine/Source/Editor/AnimGraph/Private/AnimStateTransitionNode.cpp
+++ b/Engine/Source/Editor/AnimGraph/Private/AnimStateTransitionNode.cpp
@@ -62,8 +62,8 @@ protected:
 /////////////////////////////////////////////////////
 // UAnimStateTransitionNode
 
-UAnimStateTransitionNode::UAnimStateTransitionNode(const class FPostConstructInitializeProperties& PCIP)
-	: Super(PCIP)
+UAnimStateTransitionNode::UAnimStateTransitionNode(const FObjectInitializer& ObjectInitializer)
+	: Super(ObjectInitializer)
 {
 
 	CrossfadeDuration = 0.2f;
diff --git a/Engine/Source/Editor/AnimGraph/Private/AnimationConduitGraphSchema.cpp b/Engine/Source/Editor/AnimGraph/Private/AnimationConduitGraphSchema.cpp
--- a/Engine/Source/Editor/AnimGraph/Private/AnimationConduitGraphSchema.cpp
+++ b/Engine/Source/Editor/AnimGraph/Private/AnimationConduitGraphSchema.cpp
@@ -9,8 +9,8 @@
 /////////////////////////////////////////////////////
 // UAnimationConduitGraphSchema
 
-UAnimationConduitGraphSchema::UAnimationConduitGraphSchema(const class FPostConstructInitializeProperties& PCIP)
-	: Super(PCIP)
+UAnimationConduitGraphSchema::UAnimationConduitGraphSchema(const FObjectInitializer& ObjectInitializer)
+	: Super(ObjectInitializer)
 {
 }
 
diff --git a/Engine/Source/Editor/AnimGraph/Private/AnimationCustomTransitionGraph.cpp b/Engine/Source/Editor/AnimGraph/Private/AnimationCustomTransitionGraph.cpp
--- a/Engine/Source/Editor/AnimGraph/Private/AnimationCustomTransitionGraph.cpp
+++ b/Engine/Source/Editor/AnimGraph/Private/AnimationCustomTransitionGraph.cpp
@@ -8,8 +8,8 @@
 /////////////////////////////////////////////////////
 // UAnimationStateGraph
 
-UAnimationCustomTransitionGraph::UAnimationCustomTransitionGraph(const FPostConstructInitializeProperties& PCIP)
-	: Super(PCIP)
+UAnimationCustomTransitionGraph::UAnimationCustomTransitionGraph(const FObjectInitializer& ObjectInitializer)
+	: Super(ObjectInitializer)
 {
 }
 
diff --git a/Engine/Source/Editor/AnimGraph/Private/AnimationCustomTransitionSchema.cpp b/Engine/Source/Editor/AnimGraph/Private/AnimationCustomTransitionSchema.cpp
--- a/Engine/Source/Editor/AnimGraph/Private/AnimationCustomTransitionSchema.cpp
+++ b/Engine/Source/Editor/AnimGraph/Private/AnimationCustomTransitionSchema.cpp
@@ -10,8 +10,8 @@
 /////////////////////////////////////////////////////
 // UAnimationCustomTransitionSchema
 
-UAnimationCustomTransitionSchema::UAnimationCustomTransitionSchema(const class FPostConstructInitializeProperties& PCIP)
-	: Super(PCIP)
+UAnimationCustomTransitionSchema::UAnimationCustomTransitionSchema(const FObjectInitializer& ObjectInitializer)
+	: Super(ObjectInitializer)
 {
 }
 
diff --git a/Engine/Source/Editor/AnimGraph/Private/AnimationGraph.cpp b/Engine/Source/Editor/AnimGraph/Private/AnimationGraph.cpp
--- a/Engine/Source/Editor/AnimGraph/Private/AnimationGraph.cpp
+++ b/Engine/Source/Editor/AnimGraph/Private/AnimationGraph.cpp
@@ -8,8 +8,8 @@
 /////////////////////////////////////////////////////
 // UAnimationGraph
 
-UAnimationGraph::UAnimationGraph(const FPostConstructInitializeProperties& PCIP)
-	: Super(PCIP)
+UAnimationGraph::UAnimationGraph(const FObjectInitializer& ObjectInitializer)
+	: Super(ObjectInitializer)
 {
 }
 
diff --git a/Engine/Source/Editor/AnimGraph/Private/AnimationGraphSchema.cpp b/Engine/Source/Editor/AnimGraph/Private/AnimationGraphSchema.cpp
--- a/Engine/Source/Editor/AnimGraph/Private/AnimationGraphSchema.cpp
+++ b/Engine/Source/Editor/AnimGraph/Private/AnimationGraphSchema.cpp
@@ -25,8 +25,8 @@
 /////////////////////////////////////////////////////
 // UAnimationGraphSchema
 
-UAnimationGraphSchema::UAnimationGraphSchema(const class FPostConstructInitializeProperties& PCIP)
-	: Super(PCIP)
+UAnimationGraphSchema::UAnimationGraphSchema(const FObjectInitializer& ObjectInitializer)
+	: Super(ObjectInitializer)
 {
 	PN_SequenceName = TEXT("Sequence");
 
diff --git a/Engine/Source/Editor/AnimGraph/Private/AnimationStateGraph.cpp b/Engine/Source/Editor/AnimGraph/Private/AnimationStateGraph.cpp
--- a/Engine/Source/Editor/AnimGraph/Private/AnimationStateGraph.cpp
+++ b/Engine/Source/Editor/AnimGraph/Private/AnimationStateGraph.cpp
@@ -8,8 +8,8 @@
 /////////////////////////////////////////////////////
 // UAnimationStateGraph
 
-UAnimationStateGraph::UAnimationStateGraph(const FPostConstructInitializeProperties& PCIP)
-	: Super(PCIP)
+UAnimationStateGraph::UAnimationStateGraph(const FObjectInitializer& ObjectInitializer)
+	: Super(ObjectInitializer)
 {
 }
 
diff --git a/Engine/Source/Editor/AnimGraph/Private/AnimationStateGraphSchema.cpp b/Engine/Source/Editor/AnimGraph/Private/AnimationStateGraphSchema.cpp
--- a/Engine/Source/Editor/AnimGraph/Private/AnimationStateGraphSchema.cpp
+++ b/Engine/Source/Editor/AnimGraph/Private/AnimationStateGraphSchema.cpp
@@ -15,8 +15,8 @@
 /////////////////////////////////////////////////////
 // UAnimationStateGraphSchema
 
-UAnimationStateGraphSchema::UAnimationStateGraphSchema(const class FPostConstructInitializeProperties& PCIP)
-	: Super(PCIP)
+UAnimationStateGraphSchema::UAnimationStateGraphSchema(const FObjectInitializer& ObjectInitializer)
+	: Super(ObjectInitializer)
 {
 }
 
diff --git a/Engine/Source/Editor/AnimGraph/Private/AnimationStateMachineGraph.cpp b/Engine/Source/Editor/AnimGraph/Private/AnimationStateMachineGraph.cpp
--- a/Engine/Source/Editor/AnimGraph/Private/AnimationStateMachineGraph.cpp
+++ b/Engine/Source/Editor/AnimGraph/Private/AnimationStateMachineGraph.cpp
@@ -8,8 +8,8 @@
 /////////////////////////////////////////////////////
 // UAnimationStateMachineGraph
 
-UAnimationStateMachineGraph::UAnimationStateMachineGraph(const FPostConstructInitializeProperties& PCIP)
-	: Super(PCIP)
+UAnimationStateMachineGraph::UAnimationStateMachineGraph(const FObjectInitializer& ObjectInitializer)
+	: Super(ObjectInitializer)
 {
 	bAllowDeletion = false;
 	bAllowRenaming = true;
diff --git a/Engine/Source/Editor/AnimGraph/Private/AnimationStateMachineSchema.cpp b/Engine/Source/Editor/AnimGraph/Private/AnimationStateMachineSchema.cpp
--- a/Engine/Source/Editor/AnimGraph/Private/AnimationStateMachineSchema.cpp
+++ b/Engine/Source/Editor/AnimGraph/Private/AnimationStateMachineSchema.cpp
@@ -108,8 +108,8 @@ UEdGraphNode* FEdGraphSchemaAction_NewStateComment::PerformAction(class UEdGraph
 /////////////////////////////////////////////////////
 // UAnimationStateMachineSchema
 
-UAnimationStateMachineSchema::UAnimationStateMachineSchema(const class FPostConstructInitializeProperties& PCIP)
-	: Super(PCIP)
+UAnimationStateMachineSchema::UAnimationStateMachineSchema(const FObjectInitializer& ObjectInitializer)
+	: Super(ObjectInitializer)
 {
 	PC_Exec = TEXT("exec");
 }
diff --git a/Engine/Source/Editor/AnimGraph/Private/AnimationTransitionGraph.cpp b/Engine/Source/Editor/AnimGraph/Private/AnimationTransitionGraph.cpp
--- a/Engine/Source/Editor/AnimGraph/Private/AnimationTransitionGraph.cpp
+++ b/Engine/Source/Editor/AnimGraph/Private/AnimationTransitionGraph.cpp
@@ -8,8 +8,8 @@
 /////////////////////////////////////////////////////
 // UAnimationTransitionGraph
 
-UAnimationTransitionGraph::UAnimationTransitionGraph(const FPostConstructInitializeProperties& PCIP)
-	: Super(PCIP)
+UAnimationTransitionGraph::UAnimationTransitionGraph(const FObjectInitializer& ObjectInitializer)
+	: Super(ObjectInitializer)
 {
 }
 
diff --git a/Engine/Source/Editor/AnimGraph/Private/AnimationTransitionSchema.cpp b/Engine/Source/Editor/AnimGraph/Private/AnimationTransitionSchema.cpp
--- a/Engine/Source/Editor/AnimGraph/Private/AnimationTransitionSchema.cpp
+++ b/Engine/Source/Editor/AnimGraph/Private/AnimationTransitionSchema.cpp
@@ -25,8 +25,8 @@
 
 #define LOCTEXT_NAMESPACE "AnimationTransitionSchema"
 
-UAnimationTransitionSchema::UAnimationTransitionSchema(const class FPostConstructInitializeProperties& PCIP)
-	: Super(PCIP)
+UAnimationTransitionSchema::UAnimationTransitionSchema(const FObjectInitializer& ObjectInitializer)
+	: Super(ObjectInitializer)
 {
 	// Initialize defaults
 }
diff --git a/Engine/Source/Editor/AnimGraph/Private/K2Node_TransitionRuleGetter.cpp b/Engine/Source/Editor/AnimGraph/Private/K2Node_TransitionRuleGetter.cpp
--- a/Engine/Source/Editor/AnimGraph/Private/K2Node_TransitionRuleGetter.cpp
+++ b/Engine/Source/Editor/AnimGraph/Private/K2Node_TransitionRuleGetter.cpp
@@ -15,8 +15,8 @@
 
 #define LOCTEXT_NAMESPACE "TransitionRuleGetter"
 
-UK2Node_TransitionRuleGetter::UK2Node_TransitionRuleGetter(const class FPostConstructInitializeProperties& PCIP)
-	: Super(PCIP)
+UK2Node_TransitionRuleGetter::UK2Node_TransitionRuleGetter(const FObjectInitializer& ObjectInitializer)
+	: Super(ObjectInitializer)
 {
 }
 
diff --git a/Engine/Source/Editor/BehaviorTreeEditor/Private/BehaviorTreeDecoratorGraph.cpp b/Engine/Source/Editor/BehaviorTreeEditor/Private/BehaviorTreeDecoratorGraph.cpp
--- a/Engine/Source/Editor/BehaviorTreeEditor/Private/BehaviorTreeDecoratorGraph.cpp
+++ b/Engine/Source/Editor/BehaviorTreeEditor/Private/BehaviorTreeDecoratorGraph.cpp
@@ -5,7 +5,7 @@
 //////////////////////////////////////////////////////////////////////////
 // BehaviorTreeGraph
 
-UBehaviorTreeDecoratorGraph::UBehaviorTreeDecoratorGraph(const class FPostConstructInitializeProperties& PCIP) : Super(PCIP)
+UBehaviorTreeDecoratorGraph::UBehaviorTreeDecoratorGraph(const FObjectInitializer& ObjectInitializer) : Super(ObjectInitializer)
 {
 	Schema = UEdGraphSchema_BehaviorTreeDecorator::StaticClass();
 }
diff --git a/Engine/Source/Editor/BehaviorTreeEditor/Private/BehaviorTreeDecoratorGraphNode.cpp b/Engine/Source/Editor/BehaviorTreeEditor/Private/BehaviorTreeDecoratorGraphNode.cpp
--- a/Engine/Source/Editor/BehaviorTreeEditor/Private/BehaviorTreeDecoratorGraphNode.cpp
+++ b/Engine/Source/Editor/BehaviorTreeEditor/Private/BehaviorTreeDecoratorGraphNode.cpp
@@ -2,7 +2,7 @@
 
 #include "BehaviorTreeEditorPrivatePCH.h"
 
-UBehaviorTreeDecoratorGraphNode::UBehaviorTreeDecoratorGraphNode(const class FPostConstructInitializeProperties& PCIP) : Super(PCIP)
+UBehaviorTreeDecoratorGraphNode::UBehaviorTreeDecoratorGraphNode(const FObjectInitializer& ObjectInitializer) : Super(ObjectInitializer)
 {
 	bAllowModifingInputs = true;
 }
diff --git a/Engine/Source/Editor/BehaviorTreeEditor/Private/BehaviorTreeDecoratorGraphNode_Decorator.cpp b/Engine/Source/Editor/BehaviorTreeEditor/Private/BehaviorTreeDecoratorGraphNode_Decorator.cpp
--- a/Engine/Source/Editor/BehaviorTreeEditor/Private/BehaviorTreeDecoratorGraphNode_Decorator.cpp
+++ b/Engine/Source/Editor/BehaviorTreeEditor/Private/BehaviorTreeDecoratorGraphNode_Decorator.cpp
@@ -6,7 +6,7 @@
 #include "BehaviorTreeGraphNode_CompositeDecorator.h"
 #include "BehaviorTreeDecoratorGraphNode_Decorator.h"
 
-UBehaviorTreeDecoratorGraphNode_Decorator::UBehaviorTreeDecoratorGraphNode_Decorator(const class FPostConstructInitializeProperties& PCIP) : Super(PCIP)
+UBehaviorTreeDecoratorGraphNode_Decorator::UBehaviorTreeDecoratorGraphNode_Decorator(const FObjectInitializer& ObjectInitializer) : Super(ObjectInitializer)
 {
 	NodeInstance = NULL;
 }
diff --git a/Engine/Source/Editor/BehaviorTreeEditor/Private/BehaviorTreeDecoratorGraphNode_Logic.cpp b/Engine/Source/Editor/BehaviorTreeEditor/Private/BehaviorTreeDecoratorGraphNode_Logic.cpp
--- a/Engine/Source/Editor/BehaviorTreeEditor/Private/BehaviorTreeDecoratorGraphNode_Logic.cpp
+++ b/Engine/Source/Editor/BehaviorTreeEditor/Private/BehaviorTreeDecoratorGraphNode_Logic.cpp
@@ -2,7 +2,7 @@
 
 #include "BehaviorTreeEditorPrivatePCH.h"
 
-UBehaviorTreeDecoratorGraphNode_Logic::UBehaviorTreeDecoratorGraphNode_Logic(const class FPostConstructInitializeProperties& PCIP) : Super(PCIP)
+UBehaviorTreeDecoratorGraphNode_Logic::UBehaviorTreeDecoratorGraphNode_Logic(const FObjectInitializer& ObjectInitializer) : Super(ObjectInitializer)
 {
 }
 
diff --git a/Engine/Source/Editor/BehaviorTreeEditor/Private/BehaviorTreeEditorTypes.cpp b/Engine/Source/Editor/BehaviorTreeEditor/Private/BehaviorTreeEditorTypes.cpp
--- a/Engine/Source/Editor/BehaviorTreeEditor/Private/BehaviorTreeEditorTypes.cpp
+++ b/Engine/Source/Editor/BehaviorTreeEditor/Private/BehaviorTreeEditorTypes.cpp
@@ -547,7 +547,7 @@ void FClassBrowseHelper::UpdateAvailableNodeClasses()
 
 //////////////////////////////////////////////////////////////////////////
 
-UBehaviorTreeEditorTypes::UBehaviorTreeEditorTypes(const class FPostConstructInitializeProperties& PCIP) : Super(PCIP)
+UBehaviorTreeEditorTypes::UBehaviorTreeEditorTypes(const FObjectInitializer& ObjectInitializer) : Super(ObjectInitializer)
 {
 }
 
diff --git a/Engine/Source/Editor/BehaviorTreeEditor/Private/BehaviorTreeFactory.cpp b/Engine/Source/Editor/BehaviorTreeEditor/Private/BehaviorTreeFactory.cpp
--- a/Engine/Source/Editor/BehaviorTreeEditor/Private/BehaviorTreeFactory.cpp
+++ b/Engine/Source/Editor/BehaviorTreeEditor/Private/BehaviorTreeFactory.cpp
@@ -6,8 +6,8 @@
 
 #define LOCTEXT_NAMESPACE "BehaviorTreeFactory"
 
-UBehaviorTreeFactory::UBehaviorTreeFactory(const class FPostConstructInitializeProperties& PCIP)
-	: Super(PCIP)
+UBehaviorTreeFactory::UBehaviorTreeFactory(const FObjectInitializer& ObjectInitializer)
+	: Super(ObjectInitializer)
 {
 	SupportedClass = UBehaviorTree::StaticClass();
 	bCreateNew = true;
diff --git a/Engine/Source/Editor/BehaviorTreeEditor/Private/BehaviorTreeGraph.cpp b/Engine/Source/Editor/BehaviorTreeEditor/Private/BehaviorTreeGraph.cpp
--- a/Engine/Source/Editor/BehaviorTreeEditor/Private/BehaviorTreeGraph.cpp
+++ b/Engine/Source/Editor/BehaviorTreeEditor/Private/BehaviorTreeGraph.cpp
@@ -9,7 +9,7 @@
 //////////////////////////////////////////////////////////////////////////
 // BehaviorTreeGraph
 
-UBehaviorTreeGraph::UBehaviorTreeGraph(const class FPostConstructInitializeProperties& PCIP) : Super(PCIP)
+UBehaviorTreeGraph::UBehaviorTreeGraph(const FObjectInitializer& ObjectInitializer) : Super(ObjectInitializer)
 {
 	Schema = UEdGraphSchema_BehaviorTree::StaticClass();
 	bLockUpdates = false;
diff --git a/Engine/Source/Editor/BehaviorTreeEditor/Private/BehaviorTreeGraphNode.cpp b/Engine/Source/Editor/BehaviorTreeEditor/Private/BehaviorTreeGraphNode.cpp
--- a/Engine/Source/Editor/BehaviorTreeEditor/Private/BehaviorTreeGraphNode.cpp
+++ b/Engine/Source/Editor/BehaviorTreeEditor/Private/BehaviorTreeGraphNode.cpp
@@ -12,7 +12,7 @@
 
 #define LOCTEXT_NAMESPACE "BehaviorTreeGraphNode"
 
-UBehaviorTreeGraphNode::UBehaviorTreeGraphNode(const class FPostConstructInitializeProperties& PCIP) : Super(PCIP)
+UBehaviorTreeGraphNode::UBehaviorTreeGraphNode(const FObjectInitializer& ObjectInitializer) : Super(ObjectInitializer)
 {
 	NodeInstance = NULL;
 	bHighlightInAbortRange0 = false;
diff --git a/Engine/Source/Editor/BehaviorTreeEditor/Private/BehaviorTreeGraphNode_Composite.cpp b/Engine/Source/Editor/BehaviorTreeEditor/Private/BehaviorTreeGraphNode_Composite.cpp
--- a/Engine/Source/Editor/BehaviorTreeEditor/Private/BehaviorTreeGraphNode_Composite.cpp
+++ b/Engine/Source/Editor/BehaviorTreeEditor/Private/BehaviorTreeGraphNode_Composite.cpp
@@ -4,8 +4,8 @@
 
 #define LOCTEXT_NAMESPACE "BehaviorTreeGraphNode_Composite"
 
-UBehaviorTreeGraphNode_Composite::UBehaviorTreeGraphNode_Composite(const class FPostConstructInitializeProperties& PCIP)
-	: Super(PCIP)
+UBehaviorTreeGraphNode_Composite::UBehaviorTreeGraphNode_Composite(const FObjectInitializer& ObjectInitializer)
+	: Super(ObjectInitializer)
 {
 }
 
diff --git a/Engine/Source/Editor/BehaviorTreeEditor/Private/BehaviorTreeGraphNode_CompositeDecorator.cpp b/Engine/Source/Editor/BehaviorTreeEditor/Private/BehaviorTreeGraphNode_CompositeDecorator.cpp
--- a/Engine/Source/Editor/BehaviorTreeEditor/Private/BehaviorTreeGraphNode_CompositeDecorator.cpp
+++ b/Engine/Source/Editor/BehaviorTreeEditor/Private/BehaviorTreeGraphNode_CompositeDecorator.cpp
@@ -7,7 +7,7 @@
 
 #define LOCTEXT_NAMESPACE "BehaviorTreeGraphNode"
 
-UBehaviorTreeGraphNode_CompositeDecorator::UBehaviorTreeGraphNode_CompositeDecorator(const class FPostConstructInitializeProperties& PCIP) : Super(PCIP)
+UBehaviorTreeGraphNode_CompositeDecorator::UBehaviorTreeGraphNode_CompositeDecorator(const FObjectInitializer& ObjectInitializer) : Super(ObjectInitializer)
 {
 	bShowOperations = true;
 	bCanAbortFlow = false;
diff --git a/Engine/Source/Editor/BehaviorTreeEditor/Private/BehaviorTreeGraphNode_Decorator.cpp b/Engine/Source/Editor/BehaviorTreeEditor/Private/BehaviorTreeGraphNode_Decorator.cpp
--- a/Engine/Source/Editor/BehaviorTreeEditor/Private/BehaviorTreeGraphNode_Decorator.cpp
+++ b/Engine/Source/Editor/BehaviorTreeEditor/Private/BehaviorTreeGraphNode_Decorator.cpp
@@ -3,8 +3,8 @@
 #include "BehaviorTreeEditorPrivatePCH.h"
 #include "BehaviorTree/BTDecorator.h"
 
-UBehaviorTreeGraphNode_Decorator::UBehaviorTreeGraphNode_Decorator(const class FPostConstructInitializeProperties& PCIP)
-	: Super(PCIP)
+UBehaviorTreeGraphNode_Decorator::UBehaviorTreeGraphNode_Decorator(const FObjectInitializer& ObjectInitializer)
+	: Super(ObjectInitializer)
 {
 }
 
diff --git a/Engine/Source/Editor/BehaviorTreeEditor/Private/BehaviorTreeGraphNode_Root.cpp b/Engine/Source/Editor/BehaviorTreeEditor/Private/BehaviorTreeGraphNode_Root.cpp
--- a/Engine/Source/Editor/BehaviorTreeEditor/Private/BehaviorTreeGraphNode_Root.cpp
+++ b/Engine/Source/Editor/BehaviorTreeEditor/Private/BehaviorTreeGraphNode_Root.cpp
@@ -4,7 +4,7 @@
 #include "BehaviorTree/BlackboardData.h"
 #include "BehaviorTree/BehaviorTree.h"
 
-UBehaviorTreeGraphNode_Root::UBehaviorTreeGraphNode_Root(const class FPostConstructInitializeProperties& PCIP) : Super(PCIP)
+UBehaviorTreeGraphNode_Root::UBehaviorTreeGraphNode_Root(const FObjectInitializer& ObjectInitializer) : Super(ObjectInitializer)
 {
 }
 
diff --git a/Engine/Source/Editor/BehaviorTreeEditor/Private/BehaviorTreeGraphNode_Service.cpp b/Engine/Source/Editor/BehaviorTreeEditor/Private/BehaviorTreeGraphNode_Service.cpp
--- a/Engine/Source/Editor/BehaviorTreeEditor/Private/BehaviorTreeGraphNode_Service.cpp
+++ b/Engine/Source/Editor/BehaviorTreeEditor/Private/BehaviorTreeGraphNode_Service.cpp
@@ -3,8 +3,8 @@
 #include "BehaviorTreeEditorPrivatePCH.h"
 #include "BehaviorTree/BTService.h"
 
-UBehaviorTreeGraphNode_Service::UBehaviorTreeGraphNode_Service(const class FPostConstructInitializeProperties& PCIP)
-	: Super(PCIP)
+UBehaviorTreeGraphNode_Service::UBehaviorTreeGraphNode_Service(const FObjectInitializer& ObjectInitializer)
+	: Super(ObjectInitializer)
 {
 }
 
diff --git a/Engine/Source/Editor/BehaviorTreeEditor/Private/BehaviorTreeGraphNode_SimpleParallel.cpp b/Engine/Source/Editor/BehaviorTreeEditor/Private/BehaviorTreeGraphNode_SimpleParallel.cpp
--- a/Engine/Source/Editor/BehaviorTreeEditor/Private/BehaviorTreeGraphNode_SimpleParallel.cpp
+++ b/Engine/Source/Editor/BehaviorTreeEditor/Private/BehaviorTreeGraphNode_SimpleParallel.cpp
@@ -5,7 +5,7 @@
 
 #define LOCTEXT_NAMESPACE "BehaviorTreeSchema"
 
-UBehaviorTreeGraphNode_SimpleParallel::UBehaviorTreeGraphNode_SimpleParallel(const class FPostConstructInitializeProperties& PCIP) : Super(PCIP)
+UBehaviorTreeGraphNode_SimpleParallel::UBehaviorTreeGraphNode_SimpleParallel(const FObjectInitializer& ObjectInitializer) : Super(ObjectInitializer)
 {
 }
 
diff --git a/Engine/Source/Editor/BehaviorTreeEditor/Private/BehaviorTreeGraphNode_SubtreeTask.cpp b/Engine/Source/Editor/BehaviorTreeEditor/Private/BehaviorTreeGraphNode_SubtreeTask.cpp
--- a/Engine/Source/Editor/BehaviorTreeEditor/Private/BehaviorTreeGraphNode_SubtreeTask.cpp
+++ b/Engine/Source/Editor/BehaviorTreeEditor/Private/BehaviorTreeGraphNode_SubtreeTask.cpp
@@ -3,7 +3,7 @@
 #include "BehaviorTreeEditorPrivatePCH.h"
 #include "BehaviorTree/Tasks/BTTask_RunBehavior.h"
 
-UBehaviorTreeGraphNode_SubtreeTask::UBehaviorTreeGraphNode_SubtreeTask(const class FPostConstructInitializeProperties& PCIP) : Super(PCIP)
+UBehaviorTreeGraphNode_SubtreeTask::UBehaviorTreeGraphNode_SubtreeTask(const FObjectInitializer& ObjectInitializer) : Super(ObjectInitializer)
 {
 	SubtreeVersion = 0;
 }
diff --git a/Engine/Source/Editor/BehaviorTreeEditor/Private/BehaviorTreeGraphNode_Task.cpp b/Engine/Source/Editor/BehaviorTreeEditor/Private/BehaviorTreeGraphNode_Task.cpp
--- a/Engine/Source/Editor/BehaviorTreeEditor/Private/BehaviorTreeGraphNode_Task.cpp
+++ b/Engine/Source/Editor/BehaviorTreeEditor/Private/BehaviorTreeGraphNode_Task.cpp
@@ -2,7 +2,7 @@
 
 #include "BehaviorTreeEditorPrivatePCH.h"
 
-UBehaviorTreeGraphNode_Task::UBehaviorTreeGraphNode_Task(const class FPostConstructInitializeProperties& PCIP) : Super(PCIP)
+UBehaviorTreeGraphNode_Task::UBehaviorTreeGraphNode_Task(const FObjectInitializer& ObjectInitializer) : Super(ObjectInitializer)
 {
 }
 
diff --git a/Engine/Source/Editor/BehaviorTreeEditor/Private/BlackboardDataFactory.cpp b/Engine/Source/Editor/BehaviorTreeEditor/Private/BlackboardDataFactory.cpp
--- a/Engine/Source/Editor/BehaviorTreeEditor/Private/BlackboardDataFactory.cpp
+++ b/Engine/Source/Editor/BehaviorTreeEditor/Private/BlackboardDataFactory.cpp
@@ -6,8 +6,8 @@
 
 #define LOCTEXT_NAMESPACE "BehaviorTreeFactory"
 
-UBlackboardDataFactory::UBlackboardDataFactory(const class FPostConstructInitializeProperties& PCIP)
-	: Super(PCIP)
+UBlackboardDataFactory::UBlackboardDataFactory(const FObjectInitializer& ObjectInitializer)
+	: Super(ObjectInitializer)
 {
 	SupportedClass = UBlackboardData::StaticClass();
 	bCreateNew = true;
diff --git a/Engine/Source/Editor/BehaviorTreeEditor/Private/EdGraphSchema_BehaviorTree.cpp b/Engine/Source/Editor/BehaviorTreeEditor/Private/EdGraphSchema_BehaviorTree.cpp
--- a/Engine/Source/Editor/BehaviorTreeEditor/Private/EdGraphSchema_BehaviorTree.cpp
+++ b/Engine/Source/Editor/BehaviorTreeEditor/Private/EdGraphSchema_BehaviorTree.cpp
@@ -220,7 +220,7 @@ void FBehaviorTreeSchemaAction_NewSubNode::AddReferencedObjects( FReferenceColle
 }
 //////////////////////////////////////////////////////////////////////////
 
-UEdGraphSchema_BehaviorTree::UEdGraphSchema_BehaviorTree(const class FPostConstructInitializeProperties& PCIP) : Super(PCIP)
+UEdGraphSchema_BehaviorTree::UEdGraphSchema_BehaviorTree(const FObjectInitializer& ObjectInitializer) : Super(ObjectInitializer)
 {
 }
 
diff --git a/Engine/Source/Editor/BehaviorTreeEditor/Private/EdGraphSchema_BehaviorTreeDecorator.cpp b/Engine/Source/Editor/BehaviorTreeEditor/Private/EdGraphSchema_BehaviorTreeDecorator.cpp
--- a/Engine/Source/Editor/BehaviorTreeEditor/Private/EdGraphSchema_BehaviorTreeDecorator.cpp
+++ b/Engine/Source/Editor/BehaviorTreeEditor/Private/EdGraphSchema_BehaviorTreeDecorator.cpp
@@ -99,7 +99,7 @@ void FDecoratorSchemaAction_NewNode::AddReferencedObjects( FReferenceCollector&
 
 //////////////////////////////////////////////////////////////////////////
 
-UEdGraphSchema_BehaviorTreeDecorator::UEdGraphSchema_BehaviorTreeDecorator(const class FPostConstructInitializeProperties& PCIP) : Super(PCIP)
+UEdGraphSchema_BehaviorTreeDecorator::UEdGraphSchema_BehaviorTreeDecorator(const FObjectInitializer& ObjectInitializer) : Super(ObjectInitializer)
 {
 	PC_Boolean = TEXT("bool");
 }
diff --git a/Engine/Source/Editor/BlueprintGraph/Private/BlueprintBoundEventNodeSpawner.cpp b/Engine/Source/Editor/BlueprintGraph/Private/BlueprintBoundEventNodeSpawner.cpp
--- a/Engine/Source/Editor/BlueprintGraph/Private/BlueprintBoundEventNodeSpawner.cpp
+++ b/Engine/Source/Editor/BlueprintGraph/Private/BlueprintBoundEventNodeSpawner.cpp
@@ -66,8 +66,8 @@ UBlueprintBoundEventNodeSpawner* UBlueprintBoundEventNodeSpawner::Create(TSubcla
 }
 
 //------------------------------------------------------------------------------
-UBlueprintBoundEventNodeSpawner::UBlueprintBoundEventNodeSpawner(class FPostConstructInitializeProperties const& PCIP)
-	: Super(PCIP)
+UBlueprintBoundEventNodeSpawner::UBlueprintBoundEventNodeSpawner(FObjectInitializer const& ObjectInitializer)
+	: Super(ObjectInitializer)
 	, EventDelegate(nullptr)
 {
 }
diff --git a/Engine/Source/Editor/BlueprintGraph/Private/BlueprintBoundNodeSpawner.cpp b/Engine/Source/Editor/BlueprintGraph/Private/BlueprintBoundNodeSpawner.cpp
--- a/Engine/Source/Editor/BlueprintGraph/Private/BlueprintBoundNodeSpawner.cpp
+++ b/Engine/Source/Editor/BlueprintGraph/Private/BlueprintBoundNodeSpawner.cpp
@@ -24,8 +24,8 @@ UBlueprintBoundNodeSpawner* UBlueprintBoundNodeSpawner::Create(TSubclassOf<UEdGr
 }
 
 //------------------------------------------------------------------------------
-UBlueprintBoundNodeSpawner::UBlueprintBoundNodeSpawner(class FPostConstructInitializeProperties const& PCIP)
-	: Super(PCIP)
+UBlueprintBoundNodeSpawner::UBlueprintBoundNodeSpawner(FObjectInitializer const& ObjectInitializer)
+	: Super(ObjectInitializer)
 {
 }
 
diff --git a/Engine/Source/Editor/BlueprintGraph/Private/BlueprintComponentNodeSpawner.cpp b/Engine/Source/Editor/BlueprintGraph/Private/BlueprintComponentNodeSpawner.cpp
--- a/Engine/Source/Editor/BlueprintGraph/Private/BlueprintComponentNodeSpawner.cpp
+++ b/Engine/Source/Editor/BlueprintGraph/Private/BlueprintComponentNodeSpawner.cpp
@@ -72,8 +72,8 @@ UBlueprintComponentNodeSpawner* UBlueprintComponentNodeSpawner::Create(TSubclass
 }
 
 //------------------------------------------------------------------------------
-UBlueprintComponentNodeSpawner::UBlueprintComponentNodeSpawner(class FPostConstructInitializeProperties const& PCIP)
-	: Super(PCIP)
+UBlueprintComponentNodeSpawner::UBlueprintComponentNodeSpawner(FObjectInitializer const& ObjectInitializer)
+	: Super(ObjectInitializer)
 {
 }
 
diff --git a/Engine/Source/Editor/BlueprintGraph/Private/BlueprintDelegateNodeSpawner.cpp b/Engine/Source/Editor/BlueprintGraph/Private/BlueprintDelegateNodeSpawner.cpp
--- a/Engine/Source/Editor/BlueprintGraph/Private/BlueprintDelegateNodeSpawner.cpp
+++ b/Engine/Source/Editor/BlueprintGraph/Private/BlueprintDelegateNodeSpawner.cpp
@@ -104,8 +104,8 @@ UBlueprintDelegateNodeSpawner* UBlueprintDelegateNodeSpawner::Create(TSubclassOf
 }
 
 //------------------------------------------------------------------------------
-UBlueprintDelegateNodeSpawner::UBlueprintDelegateNodeSpawner(class FPostConstructInitializeProperties const& PCIP)
-	: Super(PCIP)
+UBlueprintDelegateNodeSpawner::UBlueprintDelegateNodeSpawner(FObjectInitializer const& ObjectInitializer)
+	: Super(ObjectInitializer)
 {
 }
 
diff --git a/Engine/Source/Editor/BlueprintGraph/Private/BlueprintEditorSettings.cpp b/Engine/Source/Editor/BlueprintGraph/Private/BlueprintEditorSettings.cpp
--- a/Engine/Source/Editor/BlueprintGraph/Private/BlueprintEditorSettings.cpp
+++ b/Engine/Source/Editor/BlueprintGraph/Private/BlueprintEditorSettings.cpp
@@ -5,8 +5,8 @@
 #include "Editor/UnrealEd/Classes/Settings/EditorExperimentalSettings.h"
 #include "Editor/UnrealEd/Classes/Editor/EditorUserSettings.h"
 
-UBlueprintEditorSettings::UBlueprintEditorSettings(const class FPostConstructInitializeProperties& PCIP)
-	: Super(PCIP)
+UBlueprintEditorSettings::UBlueprintEditorSettings(const FObjectInitializer& ObjectInitializer)
+	: Super(ObjectInitializer)
 {
 	bSaveOnCompile = false;
 	bShowGraphInstructionText = true;
diff --git a/Engine/Source/Editor/BlueprintGraph/Private/BlueprintEventNodeSpawner.cpp b/Engine/Source/Editor/BlueprintGraph/Private/BlueprintEventNodeSpawner.cpp
--- a/Engine/Source/Editor/BlueprintGraph/Private/BlueprintEventNodeSpawner.cpp
+++ b/Engine/Source/Editor/BlueprintGraph/Private/BlueprintEventNodeSpawner.cpp
@@ -107,8 +107,8 @@ UBlueprintEventNodeSpawner* UBlueprintEventNodeSpawner::Create(TSubclassOf<UK2No
 }
 
 //------------------------------------------------------------------------------
-UBlueprintEventNodeSpawner::UBlueprintEventNodeSpawner(class FPostConstructInitializeProperties const& PCIP)
-	: Super(PCIP)
+UBlueprintEventNodeSpawner::UBlueprintEventNodeSpawner(FObjectInitializer const& ObjectInitializer)
+	: Super(ObjectInitializer)
 	, EventFunc(nullptr)
 {
 }
diff --git a/Engine/Source/Editor/BlueprintGraph/Private/BlueprintFieldNodeSpawner.cpp b/Engine/Source/Editor/BlueprintGraph/Private/BlueprintFieldNodeSpawner.cpp
--- a/Engine/Source/Editor/BlueprintGraph/Private/BlueprintFieldNodeSpawner.cpp
+++ b/Engine/Source/Editor/BlueprintGraph/Private/BlueprintFieldNodeSpawner.cpp
@@ -20,8 +20,8 @@ UBlueprintFieldNodeSpawner* UBlueprintFieldNodeSpawner::Create(TSubclassOf<UK2No
 }
 
 //------------------------------------------------------------------------------
-UBlueprintFieldNodeSpawner::UBlueprintFieldNodeSpawner(class FPostConstructInitializeProperties const& PCIP)
-	: Super(PCIP)
+UBlueprintFieldNodeSpawner::UBlueprintFieldNodeSpawner(FObjectInitializer const& ObjectInitializer)
+	: Super(ObjectInitializer)
 	, Field(nullptr)
 {
 }
diff --git a/Engine/Source/Editor/BlueprintGraph/Private/BlueprintFunctionNodeSpawner.cpp b/Engine/Source/Editor/BlueprintGraph/Private/BlueprintFunctionNodeSpawner.cpp
--- a/Engine/Source/Editor/BlueprintGraph/Private/BlueprintFunctionNodeSpawner.cpp
+++ b/Engine/Source/Editor/BlueprintGraph/Private/BlueprintFunctionNodeSpawner.cpp
@@ -272,8 +272,8 @@ UBlueprintFunctionNodeSpawner* UBlueprintFunctionNodeSpawner::Create(TSubclassOf
 }
 
 //------------------------------------------------------------------------------
-UBlueprintFunctionNodeSpawner::UBlueprintFunctionNodeSpawner(class FPostConstructInitializeProperties const& PCIP)
-	: Super(PCIP)
+UBlueprintFunctionNodeSpawner::UBlueprintFunctionNodeSpawner(FObjectInitializer const& ObjectInitializer)
+	: Super(ObjectInitializer)
 {
 }
 
diff --git a/Engine/Source/Editor/BlueprintGraph/Private/BlueprintNodeSpawner.cpp b/Engine/Source/Editor/BlueprintGraph/Private/BlueprintNodeSpawner.cpp
--- a/Engine/Source/Editor/BlueprintGraph/Private/BlueprintNodeSpawner.cpp
+++ b/Engine/Source/Editor/BlueprintGraph/Private/BlueprintNodeSpawner.cpp
@@ -53,8 +53,8 @@ UBlueprintNodeSpawner* UBlueprintNodeSpawner::Create(TSubclassOf<UEdGraphNode> c
 }
 
 //------------------------------------------------------------------------------
-UBlueprintNodeSpawner::UBlueprintNodeSpawner(class FPostConstructInitializeProperties const& PCIP)
-	: Super(PCIP)
+UBlueprintNodeSpawner::UBlueprintNodeSpawner(FObjectInitializer const& ObjectInitializer)
+	: Super(ObjectInitializer)
 {
 }
 
diff --git a/Engine/Source/Editor/BlueprintGraph/Private/BlueprintVariableNodeSpawner.cpp b/Engine/Source/Editor/BlueprintGraph/Private/BlueprintVariableNodeSpawner.cpp
--- a/Engine/Source/Editor/BlueprintGraph/Private/BlueprintVariableNodeSpawner.cpp
+++ b/Engine/Source/Editor/BlueprintGraph/Private/BlueprintVariableNodeSpawner.cpp
@@ -135,8 +135,8 @@ UBlueprintVariableNodeSpawner* UBlueprintVariableNodeSpawner::Create(TSubclassOf
 }
 
 //------------------------------------------------------------------------------
-UBlueprintVariableNodeSpawner::UBlueprintVariableNodeSpawner(class FPostConstructInitializeProperties const& PCIP)
-	: Super(PCIP)
+UBlueprintVariableNodeSpawner::UBlueprintVariableNodeSpawner(FObjectInitializer const& ObjectInitializer)
+	: Super(ObjectInitializer)
 {
 }
 
diff --git a/Engine/Source/Editor/BlueprintGraph/Private/EdGraphSchema_K2.cpp b/Engine/Source/Editor/BlueprintGraph/Private/EdGraphSchema_K2.cpp
--- a/Engine/Source/Editor/BlueprintGraph/Private/EdGraphSchema_K2.cpp
+++ b/Engine/Source/Editor/BlueprintGraph/Private/EdGraphSchema_K2.cpp
@@ -292,8 +292,8 @@ const UScriptStruct* UEdGraphSchema_K2::ColorStruct = nullptr;
 
 bool UEdGraphSchema_K2::bGeneratingDocumentation = false;
 
-UEdGraphSchema_K2::UEdGraphSchema_K2(const class FPostConstructInitializeProperties& PCIP)
-	: Super(PCIP)
+UEdGraphSchema_K2::UEdGraphSchema_K2(const FObjectInitializer& ObjectInitializer)
+	: Super(ObjectInitializer)
 {
 	// Initialize cached static references to well-known struct types
 	if (VectorStruct == nullptr)
diff --git a/Engine/Source/Editor/BlueprintGraph/Private/K2Node.cpp b/Engine/Source/Editor/BlueprintGraph/Private/K2Node.cpp
--- a/Engine/Source/Editor/BlueprintGraph/Private/K2Node.cpp
+++ b/Engine/Source/Editor/BlueprintGraph/Private/K2Node.cpp
@@ -16,8 +16,8 @@ static const uint32 MaxArrayPinTooltipLineCount = 10;
 /////////////////////////////////////////////////////
 // UK2Node
 
-UK2Node::UK2Node(const class FPostConstructInitializeProperties& PCIP)
-	: Super(PCIP)
+UK2Node::UK2Node(const FObjectInitializer& ObjectInitializer)
+	: Super(ObjectInitializer)
 {
 }
 
diff --git a/Engine/Source/Editor/BlueprintGraph/Private/K2Node_AIMoveTo.cpp b/Engine/Source/Editor/BlueprintGraph/Private/K2Node_AIMoveTo.cpp
--- a/Engine/Source/Editor/BlueprintGraph/Private/K2Node_AIMoveTo.cpp
+++ b/Engine/Source/Editor/BlueprintGraph/Private/K2Node_AIMoveTo.cpp
@@ -10,8 +10,8 @@
 
 #define LOCTEXT_NAMESPACE "K2Node_AIMoveTo"
 
-UK2Node_AIMoveTo::UK2Node_AIMoveTo(const class FPostConstructInitializeProperties& PCIP)
-	: Super(PCIP)
+UK2Node_AIMoveTo::UK2Node_AIMoveTo(const FObjectInitializer& ObjectInitializer)
+	: Super(ObjectInitializer)
 {
 	ProxyFactoryFunctionName = GET_FUNCTION_NAME_CHECKED(UAIBlueprintHelperLibrary, CreateMoveToProxyObject);
 	ProxyFactoryClass = UAIBlueprintHelperLibrary::StaticClass();
diff --git a/Engine/Source/Editor/BlueprintGraph/Private/K2Node_ActorBoundEvent.cpp b/Engine/Source/Editor/BlueprintGraph/Private/K2Node_ActorBoundEvent.cpp
--- a/Engine/Source/Editor/BlueprintGraph/Private/K2Node_ActorBoundEvent.cpp
+++ b/Engine/Source/Editor/BlueprintGraph/Private/K2Node_ActorBoundEvent.cpp
@@ -33,8 +33,8 @@ public:
 	}
 };
 
-UK2Node_ActorBoundEvent::UK2Node_ActorBoundEvent(const class FPostConstructInitializeProperties& PCIP)
-	: Super(PCIP)
+UK2Node_ActorBoundEvent::UK2Node_ActorBoundEvent(const FObjectInitializer& ObjectInitializer)
+	: Super(ObjectInitializer)
 {
 }
 
diff --git a/Engine/Source/Editor/BlueprintGraph/Private/K2Node_AddComponent.cpp b/Engine/Source/Editor/BlueprintGraph/Private/K2Node_AddComponent.cpp
--- a/Engine/Source/Editor/BlueprintGraph/Private/K2Node_AddComponent.cpp
+++ b/Engine/Source/Editor/BlueprintGraph/Private/K2Node_AddComponent.cpp
@@ -14,8 +14,8 @@
 //////////////////////////////////////////////////////////////////////////
 // FKCHandler_AddComponent
 
-UK2Node_AddComponent::UK2Node_AddComponent(const class FPostConstructInitializeProperties& PCIP)
-	: Super(PCIP)
+UK2Node_AddComponent::UK2Node_AddComponent(const FObjectInitializer& ObjectInitializer)
+	: Super(ObjectInitializer)
 {
 	bIsPureFunc = false;
 }
diff --git a/Engine/Source/Editor/BlueprintGraph/Private/K2Node_AssignDelegate.cpp b/Engine/Source/Editor/BlueprintGraph/Private/K2Node_AssignDelegate.cpp
--- a/Engine/Source/Editor/BlueprintGraph/Private/K2Node_AssignDelegate.cpp
+++ b/Engine/Source/Editor/BlueprintGraph/Private/K2Node_AssignDelegate.cpp
@@ -12,8 +12,8 @@
 #define LOCTEXT_NAMESPACE "K2Node_AssignDelegate"
 
 //------------------------------------------------------------------------------
-UK2Node_AssignDelegate::UK2Node_AssignDelegate(const class FPostConstructInitializeProperties& PCIP)
-	: Super(PCIP)
+UK2Node_AssignDelegate::UK2Node_AssignDelegate(const FObjectInitializer& ObjectInitializer)
+	: Super(ObjectInitializer)
 {
 }
 
diff --git a/Engine/Source/Editor/BlueprintGraph/Private/K2Node_AssignmentStatement.cpp b/Engine/Source/Editor/BlueprintGraph/Private/K2Node_AssignmentStatement.cpp
--- a/Engine/Source/Editor/BlueprintGraph/Private/K2Node_AssignmentStatement.cpp
+++ b/Engine/Source/Editor/BlueprintGraph/Private/K2Node_AssignmentStatement.cpp
@@ -57,8 +57,8 @@ FString UK2Node_AssignmentStatement::VariablePinName = FString(TEXT("Variable"))
 FString UK2Node_AssignmentStatement::ValuePinName = FString(TEXT("Value"));
 
 
-UK2Node_AssignmentStatement::UK2Node_AssignmentStatement(const class FPostConstructInitializeProperties& PCIP)
-	: Super(PCIP)
+UK2Node_AssignmentStatement::UK2Node_AssignmentStatement(const FObjectInitializer& ObjectInitializer)
+	: Super(ObjectInitializer)
 {
 }
 
diff --git a/Engine/Source/Editor/BlueprintGraph/Private/K2Node_BaseAsyncTask.cpp b/Engine/Source/Editor/BlueprintGraph/Private/K2Node_BaseAsyncTask.cpp
--- a/Engine/Source/Editor/BlueprintGraph/Private/K2Node_BaseAsyncTask.cpp
+++ b/Engine/Source/Editor/BlueprintGraph/Private/K2Node_BaseAsyncTask.cpp
@@ -11,8 +11,8 @@
 
 #define LOCTEXT_NAMESPACE "UK2Node_BaseAsyncTask"
 
-UK2Node_BaseAsyncTask::UK2Node_BaseAsyncTask(const class FPostConstructInitializeProperties& PCIP)
-	: Super(PCIP)
+UK2Node_BaseAsyncTask::UK2Node_BaseAsyncTask(const FObjectInitializer& ObjectInitializer)
+	: Super(ObjectInitializer)
 	, ProxyFactoryFunctionName(NAME_None)
 	, ProxyFactoryClass(NULL)
 	, ProxyClass(NULL)
diff --git a/Engine/Source/Editor/BlueprintGraph/Private/K2Node_BreakStruct.cpp b/Engine/Source/Editor/BlueprintGraph/Private/K2Node_BreakStruct.cpp
--- a/Engine/Source/Editor/BlueprintGraph/Private/K2Node_BreakStruct.cpp
+++ b/Engine/Source/Editor/BlueprintGraph/Private/K2Node_BreakStruct.cpp
@@ -120,7 +120,7 @@ public:
 };
 
 
-UK2Node_BreakStruct::UK2Node_BreakStruct(const class FPostConstructInitializeProperties& PCIP) : Super(PCIP)
+UK2Node_BreakStruct::UK2Node_BreakStruct(const FObjectInitializer& ObjectInitializer) : Super(ObjectInitializer)
 {
 }
 
diff --git a/Engine/Source/Editor/BlueprintGraph/Private/K2Node_CallArrayFunction.cpp b/Engine/Source/Editor/BlueprintGraph/Private/K2Node_CallArrayFunction.cpp
--- a/Engine/Source/Editor/BlueprintGraph/Private/K2Node_CallArrayFunction.cpp
+++ b/Engine/Source/Editor/BlueprintGraph/Private/K2Node_CallArrayFunction.cpp
@@ -2,8 +2,8 @@
 
 #include "BlueprintGraphPrivatePCH.h"
 
-UK2Node_CallArrayFunction::UK2Node_CallArrayFunction(const class FPostConstructInitializeProperties& PCIP)
-	: Super(PCIP)
+UK2Node_CallArrayFunction::UK2Node_CallArrayFunction(const FObjectInitializer& ObjectInitializer)
+	: Super(ObjectInitializer)
 {
 }
 
diff --git a/Engine/Source/Editor/BlueprintGraph/Private/K2Node_CallDataTableFunction.cpp b/Engine/Source/Editor/BlueprintGraph/Private/K2Node_CallDataTableFunction.cpp
--- a/Engine/Source/Editor/BlueprintGraph/Private/K2Node_CallDataTableFunction.cpp
+++ b/Engine/Source/Editor/BlueprintGraph/Private/K2Node_CallDataTableFunction.cpp
@@ -4,8 +4,8 @@
 
 #include "CompilerResultsLog.h"
 
-UK2Node_CallDataTableFunction::UK2Node_CallDataTableFunction(const class FPostConstructInitializeProperties& PCIP)
-	: Super(PCIP)
+UK2Node_CallDataTableFunction::UK2Node_CallDataTableFunction(const FObjectInitializer& ObjectInitializer)
+	: Super(ObjectInitializer)
 {
 }
 
diff --git a/Engine/Source/Editor/BlueprintGraph/Private/K2Node_CallFunction.cpp b/Engine/Source/Editor/BlueprintGraph/Private/K2Node_CallFunction.cpp
--- a/Engine/Source/Editor/BlueprintGraph/Private/K2Node_CallFunction.cpp
+++ b/Engine/Source/Editor/BlueprintGraph/Private/K2Node_CallFunction.cpp
@@ -396,8 +396,8 @@ bool FDynamicOutputHelper::CanConformPinType(const UK2Node_CallFunction* FuncNod
  *  UK2Node_CallFunction
  ******************************************************************************/
 
-UK2Node_CallFunction::UK2Node_CallFunction(const class FPostConstructInitializeProperties& PCIP)
-	: Super(PCIP)
+UK2Node_CallFunction::UK2Node_CallFunction(const FObjectInitializer& ObjectInitializer)
+	: Super(ObjectInitializer)
 {
 }
 
diff --git a/Engine/Source/Editor/BlueprintGraph/Private/K2Node_CallFunctionOnMember.cpp b/Engine/Source/Editor/BlueprintGraph/Private/K2Node_CallFunctionOnMember.cpp
--- a/Engine/Source/Editor/BlueprintGraph/Private/K2Node_CallFunctionOnMember.cpp
+++ b/Engine/Source/Editor/BlueprintGraph/Private/K2Node_CallFunctionOnMember.cpp
@@ -6,8 +6,8 @@
 
 #define LOCTEXT_NAMESPACE "K2Node"
 
-UK2Node_CallFunctionOnMember::UK2Node_CallFunctionOnMember(const class FPostConstructInitializeProperties& PCIP)
-	: Super(PCIP)
+UK2Node_CallFunctionOnMember::UK2Node_CallFunctionOnMember(const FObjectInitializer& ObjectInitializer)
+	: Super(ObjectInitializer)
 {
 
 }
diff --git a/Engine/Source/Editor/BlueprintGraph/Private/K2Node_CallMaterialParameterCollectionFunction.cpp b/Engine/Source/Editor/BlueprintGraph/Private/K2Node_CallMaterialParameterCollectionFunction.cpp
--- a/Engine/Source/Editor/BlueprintGraph/Private/K2Node_CallMaterialParameterCollectionFunction.cpp
+++ b/Engine/Source/Editor/BlueprintGraph/Private/K2Node_CallMaterialParameterCollectionFunction.cpp
@@ -4,8 +4,8 @@
 
 #include "CompilerResultsLog.h"
 
-UK2Node_CallMaterialParameterCollectionFunction::UK2Node_CallMaterialParameterCollectionFunction(const class FPostConstructInitializeProperties& PCIP)
-	: Super(PCIP)
+UK2Node_CallMaterialParameterCollectionFunction::UK2Node_CallMaterialParameterCollectionFunction(const FObjectInitializer& ObjectInitializer)
+	: Super(ObjectInitializer)
 {
 }
 
diff --git a/Engine/Source/Editor/BlueprintGraph/Private/K2Node_CallParentFunction.cpp b/Engine/Source/Editor/BlueprintGraph/Private/K2Node_CallParentFunction.cpp
--- a/Engine/Source/Editor/BlueprintGraph/Private/K2Node_CallParentFunction.cpp
+++ b/Engine/Source/Editor/BlueprintGraph/Private/K2Node_CallParentFunction.cpp
@@ -4,8 +4,8 @@
 
 #define LOCTEXT_NAMESPACE "K2Node"
 
-UK2Node_CallParentFunction::UK2Node_CallParentFunction(const class FPostConstructInitializeProperties& PCIP)
-	: Super(PCIP)
+UK2Node_CallParentFunction::UK2Node_CallParentFunction(const FObjectInitializer& ObjectInitializer)
+	: Super(ObjectInitializer)
 {
 	bIsFinalFunction = true;
 }
diff --git a/Engine/Source/Editor/BlueprintGraph/Private/K2Node_CastByteToEnum.cpp b/Engine/Source/Editor/BlueprintGraph/Private/K2Node_CastByteToEnum.cpp
--- a/Engine/Source/Editor/BlueprintGraph/Private/K2Node_CastByteToEnum.cpp
+++ b/Engine/Source/Editor/BlueprintGraph/Private/K2Node_CastByteToEnum.cpp
@@ -11,8 +11,8 @@
 
 const FString UK2Node_CastByteToEnum::ByteInputPinName = TEXT("Byte");
 
-UK2Node_CastByteToEnum::UK2Node_CastByteToEnum(const class FPostConstructInitializeProperties& PCIP)
-	: Super(PCIP)
+UK2Node_CastByteToEnum::UK2Node_CastByteToEnum(const FObjectInitializer& ObjectInitializer)
+	: Super(ObjectInitializer)
 {
 }
 
diff --git a/Engine/Source/Editor/BlueprintGraph/Private/K2Node_ClassDynamicCast.cpp b/Engine/Source/Editor/BlueprintGraph/Private/K2Node_ClassDynamicCast.cpp
--- a/Engine/Source/Editor/BlueprintGraph/Private/K2Node_ClassDynamicCast.cpp
+++ b/Engine/Source/Editor/BlueprintGraph/Private/K2Node_ClassDynamicCast.cpp
@@ -16,8 +16,8 @@ struct FClassDynamicCastHelper
 	}
 };
 
-UK2Node_ClassDynamicCast::UK2Node_ClassDynamicCast(const class FPostConstructInitializeProperties& PCIP)
-	: Super(PCIP)
+UK2Node_ClassDynamicCast::UK2Node_ClassDynamicCast(const FObjectInitializer& ObjectInitializer)
+	: Super(ObjectInitializer)
 {
 }
 
diff --git a/Engine/Source/Editor/BlueprintGraph/Private/K2Node_CommutativeAssociativeBinaryOperator.cpp b/Engine/Source/Editor/BlueprintGraph/Private/K2Node_CommutativeAssociativeBinaryOperator.cpp
--- a/Engine/Source/Editor/BlueprintGraph/Private/K2Node_CommutativeAssociativeBinaryOperator.cpp
+++ b/Engine/Source/Editor/BlueprintGraph/Private/K2Node_CommutativeAssociativeBinaryOperator.cpp
@@ -21,8 +21,8 @@ FString UK2Node_CommutativeAssociativeBinaryOperator::GetNameForPin(int32 PinInd
 	return Name;
 }
 
-UK2Node_CommutativeAssociativeBinaryOperator::UK2Node_CommutativeAssociativeBinaryOperator(const class FPostConstructInitializeProperties& PCIP)
-	: Super(PCIP)
+UK2Node_CommutativeAssociativeBinaryOperator::UK2Node_CommutativeAssociativeBinaryOperator(const FObjectInitializer& ObjectInitializer)
+	: Super(ObjectInitializer)
 {
 	NumAdditionalInputs = 0;
 }
diff --git a/Engine/Source/Editor/BlueprintGraph/Private/K2Node_ComponentBoundEvent.cpp b/Engine/Source/Editor/BlueprintGraph/Private/K2Node_ComponentBoundEvent.cpp
--- a/Engine/Source/Editor/BlueprintGraph/Private/K2Node_ComponentBoundEvent.cpp
+++ b/Engine/Source/Editor/BlueprintGraph/Private/K2Node_ComponentBoundEvent.cpp
@@ -5,8 +5,8 @@
 
 #define LOCTEXT_NAMESPACE "K2Node"
 
-UK2Node_ComponentBoundEvent::UK2Node_ComponentBoundEvent(const class FPostConstructInitializeProperties& PCIP)
-	: Super(PCIP)
+UK2Node_ComponentBoundEvent::UK2Node_ComponentBoundEvent(const FObjectInitializer& ObjectInitializer)
+	: Super(ObjectInitializer)
 {
 }
 
diff --git a/Engine/Source/Editor/BlueprintGraph/Private/K2Node_Composite.cpp b/Engine/Source/Editor/BlueprintGraph/Private/K2Node_Composite.cpp
--- a/Engine/Source/Editor/BlueprintGraph/Private/K2Node_Composite.cpp
+++ b/Engine/Source/Editor/BlueprintGraph/Private/K2Node_Composite.cpp
@@ -7,8 +7,8 @@
 
 #define LOCTEXT_NAMESPACE "K2Node"
 
-UK2Node_Composite::UK2Node_Composite(const class FPostConstructInitializeProperties& PCIP)
-	: Super(PCIP)
+UK2Node_Composite::UK2Node_Composite(const FObjectInitializer& ObjectInitializer)
+	: Super(ObjectInitializer)
 {
 	bCanHaveInputs = true;
 	bCanHaveOutputs = true;
diff --git a/Engine/Source/Editor/BlueprintGraph/Private/K2Node_ConstructObjectFromClass.cpp b/Engine/Source/Editor/BlueprintGraph/Private/K2Node_ConstructObjectFromClass.cpp
--- a/Engine/Source/Editor/BlueprintGraph/Private/K2Node_ConstructObjectFromClass.cpp
+++ b/Engine/Source/Editor/BlueprintGraph/Private/K2Node_ConstructObjectFromClass.cpp
@@ -18,8 +18,8 @@ FString FK2Node_ConstructObjectFromClassHelper::ClassPinName(TEXT("Class"));
 
 #define LOCTEXT_NAMESPACE "K2Node_ConstructObjectFromClass"
 
-UK2Node_ConstructObjectFromClass::UK2Node_ConstructObjectFromClass(const class FPostConstructInitializeProperties& PCIP)
-	: Super(PCIP)
+UK2Node_ConstructObjectFromClass::UK2Node_ConstructObjectFromClass(const FObjectInitializer& ObjectInitializer)
+	: Super(ObjectInitializer)
 {
 	NodeTooltip = LOCTEXT("NodeTooltip", "Attempts to spawn a new object");
 }
diff --git a/Engine/Source/Editor/BlueprintGraph/Private/K2Node_CreateDelegate.cpp b/Engine/Source/Editor/BlueprintGraph/Private/K2Node_CreateDelegate.cpp
--- a/Engine/Source/Editor/BlueprintGraph/Private/K2Node_CreateDelegate.cpp
+++ b/Engine/Source/Editor/BlueprintGraph/Private/K2Node_CreateDelegate.cpp
@@ -11,8 +11,8 @@ struct FK2Node_CreateDelegate_Helper
 FString FK2Node_CreateDelegate_Helper::ObjectInputName(TEXT("InputObject"));
 FString FK2Node_CreateDelegate_Helper::DelegateOutputName(TEXT("OutputDelegate"));
 
-UK2Node_CreateDelegate::UK2Node_CreateDelegate(const class FPostConstructInitializeProperties& PCIP)
-	: Super(PCIP)
+UK2Node_CreateDelegate::UK2Node_CreateDelegate(const FObjectInitializer& ObjectInitializer)
+	: Super(ObjectInitializer)
 {
 }
 
diff --git a/Engine/Source/Editor/BlueprintGraph/Private/K2Node_CustomEvent.cpp b/Engine/Source/Editor/BlueprintGraph/Private/K2Node_CustomEvent.cpp
--- a/Engine/Source/Editor/BlueprintGraph/Private/K2Node_CustomEvent.cpp
+++ b/Engine/Source/Editor/BlueprintGraph/Private/K2Node_CustomEvent.cpp
@@ -95,8 +95,8 @@ private:
 	UK2Node_CustomEvent const* CustomEvent;
 };
 
-UK2Node_CustomEvent::UK2Node_CustomEvent(const class FPostConstructInitializeProperties& PCIP)
-	: Super(PCIP)
+UK2Node_CustomEvent::UK2Node_CustomEvent(const FObjectInitializer& ObjectInitializer)
+	: Super(ObjectInitializer)
 {
 	bOverrideFunction = false;
 	bIsEditable = true;
diff --git a/Engine/Source/Editor/BlueprintGraph/Private/K2Node_DeadClass.cpp b/Engine/Source/Editor/BlueprintGraph/Private/K2Node_DeadClass.cpp
--- a/Engine/Source/Editor/BlueprintGraph/Private/K2Node_DeadClass.cpp
+++ b/Engine/Source/Editor/BlueprintGraph/Private/K2Node_DeadClass.cpp
@@ -3,8 +3,8 @@
 #include "BlueprintGraphPrivatePCH.h"
 #include "K2Node_DeadClass.h"
 
-UK2Node_DeadClass::UK2Node_DeadClass(const class FPostConstructInitializeProperties& PCIP)
-	: Super(PCIP)
+UK2Node_DeadClass::UK2Node_DeadClass(const FObjectInitializer& ObjectInitializer)
+	: Super(ObjectInitializer)
 {
 	return;
 }
diff --git a/Engine/Source/Editor/BlueprintGraph/Private/K2Node_DelegateSet.cpp b/Engine/Source/Editor/BlueprintGraph/Private/K2Node_DelegateSet.cpp
--- a/Engine/Source/Editor/BlueprintGraph/Private/K2Node_DelegateSet.cpp
+++ b/Engine/Source/Editor/BlueprintGraph/Private/K2Node_DelegateSet.cpp
@@ -123,8 +123,8 @@ public:
 	}
 };
 
-UK2Node_DelegateSet::UK2Node_DelegateSet(const class FPostConstructInitializeProperties& PCIP)
-	: Super(PCIP)
+UK2Node_DelegateSet::UK2Node_DelegateSet(const FObjectInitializer& ObjectInitializer)
+	: Super(ObjectInitializer)
 {
 }
 
diff --git a/Engine/Source/Editor/BlueprintGraph/Private/K2Node_DoOnceMultiInput.cpp b/Engine/Source/Editor/BlueprintGraph/Private/K2Node_DoOnceMultiInput.cpp
--- a/Engine/Source/Editor/BlueprintGraph/Private/K2Node_DoOnceMultiInput.cpp
+++ b/Engine/Source/Editor/BlueprintGraph/Private/K2Node_DoOnceMultiInput.cpp
@@ -76,8 +76,8 @@ FString UK2Node_DoOnceMultiInput::GetNameForPin(int32 PinIndex, bool In)
 	return Name;
 }
 
-UK2Node_DoOnceMultiInput::UK2Node_DoOnceMultiInput(const class FPostConstructInitializeProperties& PCIP)
-	: Super(PCIP)
+UK2Node_DoOnceMultiInput::UK2Node_DoOnceMultiInput(const FObjectInitializer& ObjectInitializer)
+	: Super(ObjectInitializer)
 {
 	NumAdditionalInputs = 0;
 }
diff --git a/Engine/Source/Editor/BlueprintGraph/Private/K2Node_DynamicCast.cpp b/Engine/Source/Editor/BlueprintGraph/Private/K2Node_DynamicCast.cpp
--- a/Engine/Source/Editor/BlueprintGraph/Private/K2Node_DynamicCast.cpp
+++ b/Engine/Source/Editor/BlueprintGraph/Private/K2Node_DynamicCast.cpp
@@ -7,8 +7,8 @@
 
 #define LOCTEXT_NAMESPACE "K2Node_DynamicCast"
 
-UK2Node_DynamicCast::UK2Node_DynamicCast(const class FPostConstructInitializeProperties& PCIP)
-	: Super(PCIP)
+UK2Node_DynamicCast::UK2Node_DynamicCast(const FObjectInitializer& ObjectInitializer)
+	: Super(ObjectInitializer)
 {
 }
 
diff --git a/Engine/Source/Editor/BlueprintGraph/Private/K2Node_EaseFunction.cpp b/Engine/Source/Editor/BlueprintGraph/Private/K2Node_EaseFunction.cpp
--- a/Engine/Source/Editor/BlueprintGraph/Private/K2Node_EaseFunction.cpp
+++ b/Engine/Source/Editor/BlueprintGraph/Private/K2Node_EaseFunction.cpp
@@ -68,8 +68,8 @@ struct FEaseFunctionNodeHelper
 
 };
 
-UK2Node_EaseFunction::UK2Node_EaseFunction(const class FPostConstructInitializeProperties& PCIP)
-	: Super(PCIP)
+UK2Node_EaseFunction::UK2Node_EaseFunction(const FObjectInitializer& ObjectInitializer)
+	: Super(ObjectInitializer)
 	, CachedEaseFuncPin(NULL)
 {
 	NodeTooltip = LOCTEXT("NodeTooltip", "Interpolates from value A to value B using a user specified easing function");
diff --git a/Engine/Source/Editor/BlueprintGraph/Private/K2Node_EditablePinBase.cpp b/Engine/Source/Editor/BlueprintGraph/Private/K2Node_EditablePinBase.cpp
--- a/Engine/Source/Editor/BlueprintGraph/Private/K2Node_EditablePinBase.cpp
+++ b/Engine/Source/Editor/BlueprintGraph/Private/K2Node_EditablePinBase.cpp
@@ -2,8 +2,8 @@
 
 #include "BlueprintGraphPrivatePCH.h"
 
-UK2Node_EditablePinBase::UK2Node_EditablePinBase(const class FPostConstructInitializeProperties& PCIP)
-	: Super(PCIP)
+UK2Node_EditablePinBase::UK2Node_EditablePinBase(const FObjectInitializer& ObjectInitializer)
+	: Super(ObjectInitializer)
 {
 }
 
diff --git a/Engine/Source/Editor/BlueprintGraph/Private/K2Node_EnumEquality.cpp b/Engine/Source/Editor/BlueprintGraph/Private/K2Node_EnumEquality.cpp
--- a/Engine/Source/Editor/BlueprintGraph/Private/K2Node_EnumEquality.cpp
+++ b/Engine/Source/Editor/BlueprintGraph/Private/K2Node_EnumEquality.cpp
@@ -12,8 +12,8 @@
 
 #define LOCTEXT_NAMESPACE "K2Node_EnumEquality"
 
-UK2Node_EnumEquality::UK2Node_EnumEquality(const class FPostConstructInitializeProperties& PCIP)
-	: Super(PCIP)
+UK2Node_EnumEquality::UK2Node_EnumEquality(const FObjectInitializer& ObjectInitializer)
+	: Super(ObjectInitializer)
 {
 }
 
diff --git a/Engine/Source/Editor/BlueprintGraph/Private/K2Node_EnumLiteral.cpp b/Engine/Source/Editor/BlueprintGraph/Private/K2Node_EnumLiteral.cpp
--- a/Engine/Source/Editor/BlueprintGraph/Private/K2Node_EnumLiteral.cpp
+++ b/Engine/Source/Editor/BlueprintGraph/Private/K2Node_EnumLiteral.cpp
@@ -14,8 +14,8 @@ const FString& UK2Node_EnumLiteral::GetEnumInputPinName()
 	return Name;
 }
 
-UK2Node_EnumLiteral::UK2Node_EnumLiteral(const class FPostConstructInitializeProperties& PCIP)
-	: Super(PCIP)
+UK2Node_EnumLiteral::UK2Node_EnumLiteral(const FObjectInitializer& ObjectInitializer)
+	: Super(ObjectInitializer)
 {
 }
 
diff --git a/Engine/Source/Editor/BlueprintGraph/Private/K2Node_Event.cpp b/Engine/Source/Editor/BlueprintGraph/Private/K2Node_Event.cpp
--- a/Engine/Source/Editor/BlueprintGraph/Private/K2Node_Event.cpp
+++ b/Engine/Source/Editor/BlueprintGraph/Private/K2Node_Event.cpp
@@ -32,8 +32,8 @@ bool UK2Node_Event::IsCosmeticTickEvent() const
 
 #define LOCTEXT_NAMESPACE "K2Node_Event"
 
-UK2Node_Event::UK2Node_Event(const class FPostConstructInitializeProperties& PCIP)
-	: Super(PCIP)
+UK2Node_Event::UK2Node_Event(const FObjectInitializer& ObjectInitializer)
+	: Super(ObjectInitializer)
 {
 	FunctionFlags = 0;
 }
diff --git a/Engine/Source/Editor/BlueprintGraph/Private/K2Node_ExecutionSequence.cpp b/Engine/Source/Editor/BlueprintGraph/Private/K2Node_ExecutionSequence.cpp
--- a/Engine/Source/Editor/BlueprintGraph/Private/K2Node_ExecutionSequence.cpp
+++ b/Engine/Source/Editor/BlueprintGraph/Private/K2Node_ExecutionSequence.cpp
@@ -129,8 +129,8 @@ public:
 	}
 };
 
-UK2Node_ExecutionSequence::UK2Node_ExecutionSequence(const class FPostConstructInitializeProperties& PCIP)
-	: Super(PCIP)
+UK2Node_ExecutionSequence::UK2Node_ExecutionSequence(const FObjectInitializer& ObjectInitializer)
+	: Super(ObjectInitializer)
 {
 }
 
diff --git a/Engine/Source/Editor/BlueprintGraph/Private/K2Node_ForEachElementInEnum.cpp b/Engine/Source/Editor/BlueprintGraph/Private/K2Node_ForEachElementInEnum.cpp
--- a/Engine/Source/Editor/BlueprintGraph/Private/K2Node_ForEachElementInEnum.cpp
+++ b/Engine/Source/Editor/BlueprintGraph/Private/K2Node_ForEachElementInEnum.cpp
@@ -96,8 +96,8 @@ struct FForExpandNodeHelper
 	}
 };
 
-UK2Node_ForEachElementInEnum::UK2Node_ForEachElementInEnum(const class FPostConstructInitializeProperties& PCIP)
-	: Super(PCIP)
+UK2Node_ForEachElementInEnum::UK2Node_ForEachElementInEnum(const FObjectInitializer& ObjectInitializer)
+	: Super(ObjectInitializer)
 {
 }
 
diff --git a/Engine/Source/Editor/BlueprintGraph/Private/K2Node_FormatText.cpp b/Engine/Source/Editor/BlueprintGraph/Private/K2Node_FormatText.cpp
--- a/Engine/Source/Editor/BlueprintGraph/Private/K2Node_FormatText.cpp
+++ b/Engine/Source/Editor/BlueprintGraph/Private/K2Node_FormatText.cpp
@@ -25,8 +25,8 @@ struct FFormatTextNodeHelper
 	}
 };
 
-UK2Node_FormatText::UK2Node_FormatText(const class FPostConstructInitializeProperties& PCIP)
-	: Super(PCIP)
+UK2Node_FormatText::UK2Node_FormatText(const FObjectInitializer& ObjectInitializer)
+	: Super(ObjectInitializer)
 	, CachedFormatPin(NULL)
 {
 	NodeTooltip = LOCTEXT("NodeTooltip", "Builds a formatted string using available specifier values. Use {} to denote specifiers.");
diff --git a/Engine/Source/Editor/BlueprintGraph/Private/K2Node_FunctionEntry.cpp b/Engine/Source/Editor/BlueprintGraph/Private/K2Node_FunctionEntry.cpp
--- a/Engine/Source/Editor/BlueprintGraph/Private/K2Node_FunctionEntry.cpp
+++ b/Engine/Source/Editor/BlueprintGraph/Private/K2Node_FunctionEntry.cpp
@@ -108,8 +108,8 @@ struct FFunctionEntryHelper
 	}
 };
 
-UK2Node_FunctionEntry::UK2Node_FunctionEntry(const class FPostConstructInitializeProperties& PCIP)
-	: Super(PCIP)
+UK2Node_FunctionEntry::UK2Node_FunctionEntry(const FObjectInitializer& ObjectInitializer)
+	: Super(ObjectInitializer)
 {
 }
 
diff --git a/Engine/Source/Editor/BlueprintGraph/Private/K2Node_FunctionResult.cpp b/Engine/Source/Editor/BlueprintGraph/Private/K2Node_FunctionResult.cpp
--- a/Engine/Source/Editor/BlueprintGraph/Private/K2Node_FunctionResult.cpp
+++ b/Engine/Source/Editor/BlueprintGraph/Private/K2Node_FunctionResult.cpp
@@ -67,8 +67,8 @@ public:
 	}
 };
 
-UK2Node_FunctionResult::UK2Node_FunctionResult(const class FPostConstructInitializeProperties& PCIP)
-	: Super(PCIP)
+UK2Node_FunctionResult::UK2Node_FunctionResult(const FObjectInitializer& ObjectInitializer)
+	: Super(ObjectInitializer)
 {
 }
 
diff --git a/Engine/Source/Editor/BlueprintGraph/Private/K2Node_FunctionTerminator.cpp b/Engine/Source/Editor/BlueprintGraph/Private/K2Node_FunctionTerminator.cpp
--- a/Engine/Source/Editor/BlueprintGraph/Private/K2Node_FunctionTerminator.cpp
+++ b/Engine/Source/Editor/BlueprintGraph/Private/K2Node_FunctionTerminator.cpp
@@ -3,8 +3,8 @@
 
 #include "BlueprintGraphPrivatePCH.h"
 
-UK2Node_FunctionTerminator::UK2Node_FunctionTerminator(const class FPostConstructInitializeProperties& PCIP)
-	: Super(PCIP)
+UK2Node_FunctionTerminator::UK2Node_FunctionTerminator(const FObjectInitializer& ObjectInitializer)
+	: Super(ObjectInitializer)
 {
 }
 
diff --git a/Engine/Source/Editor/BlueprintGraph/Private/K2Node_GetDataTableRow.cpp b/Engine/Source/Editor/BlueprintGraph/Private/K2Node_GetDataTableRow.cpp
--- a/Engine/Source/Editor/BlueprintGraph/Private/K2Node_GetDataTableRow.cpp
+++ b/Engine/Source/Editor/BlueprintGraph/Private/K2Node_GetDataTableRow.cpp
@@ -21,8 +21,8 @@ FString UK2Node_GetDataTableRowHelper::DataTablePinName(LOCTEXT("DataTablePinNam
 FString UK2Node_GetDataTableRowHelper::RowNotFoundPinName(LOCTEXT("RowNotFoundPinName","RowNotFound").ToString());
 FString UK2Node_GetDataTableRowHelper::RowNamePinName(LOCTEXT("RowNamePinName","RowName").ToString());
 
-UK2Node_GetDataTableRow::UK2Node_GetDataTableRow(const class FPostConstructInitializeProperties& PCIP)
-	: Super(PCIP)
+UK2Node_GetDataTableRow::UK2Node_GetDataTableRow(const FObjectInitializer& ObjectInitializer)
+	: Super(ObjectInitializer)
 {
 	NodeTooltip = LOCTEXT("NodeTooltip", "Attempts to retrieve a TableRow from a DataTable via it's RowName");
 }
diff --git a/Engine/Source/Editor/BlueprintGraph/Private/K2Node_GetEnumeratorName.cpp b/Engine/Source/Editor/BlueprintGraph/Private/K2Node_GetEnumeratorName.cpp
--- a/Engine/Source/Editor/BlueprintGraph/Private/K2Node_GetEnumeratorName.cpp
+++ b/Engine/Source/Editor/BlueprintGraph/Private/K2Node_GetEnumeratorName.cpp
@@ -12,8 +12,8 @@
 
 FString UK2Node_GetEnumeratorName::EnumeratorPinName = TEXT("Enumerator");
 
-UK2Node_GetEnumeratorName::UK2Node_GetEnumeratorName(const class FPostConstructInitializeProperties& PCIP)
-	: Super(PCIP)
+UK2Node_GetEnumeratorName::UK2Node_GetEnumeratorName(const FObjectInitializer& ObjectInitializer)
+	: Super(ObjectInitializer)
 {
 }
 
diff --git a/Engine/Source/Editor/BlueprintGraph/Private/K2Node_GetEnumeratorNameAsString.cpp b/Engine/Source/Editor/BlueprintGraph/Private/K2Node_GetEnumeratorNameAsString.cpp
--- a/Engine/Source/Editor/BlueprintGraph/Private/K2Node_GetEnumeratorNameAsString.cpp
+++ b/Engine/Source/Editor/BlueprintGraph/Private/K2Node_GetEnumeratorNameAsString.cpp
@@ -9,8 +9,8 @@
 #include "EditorCategoryUtils.h"
 #include "BlueprintActionDatabaseRegistrar.h"
 
-UK2Node_GetEnumeratorNameAsString::UK2Node_GetEnumeratorNameAsString(const class FPostConstructInitializeProperties& PCIP)
-	: Super(PCIP)
+UK2Node_GetEnumeratorNameAsString::UK2Node_GetEnumeratorNameAsString(const FObjectInitializer& ObjectInitializer)
+	: Super(ObjectInitializer)
 {
 }
 
diff --git a/Engine/Source/Editor/BlueprintGraph/Private/K2Node_GetInputAxisKeyValue.cpp b/Engine/Source/Editor/BlueprintGraph/Private/K2Node_GetInputAxisKeyValue.cpp
--- a/Engine/Source/Editor/BlueprintGraph/Private/K2Node_GetInputAxisKeyValue.cpp
+++ b/Engine/Source/Editor/BlueprintGraph/Private/K2Node_GetInputAxisKeyValue.cpp
@@ -12,8 +12,8 @@
 
 #define LOCTEXT_NAMESPACE "K2Node_GetInputAxisKeyValue"
 
-UK2Node_GetInputAxisKeyValue::UK2Node_GetInputAxisKeyValue(const class FPostConstructInitializeProperties& PCIP)
-	: Super(PCIP)
+UK2Node_GetInputAxisKeyValue::UK2Node_GetInputAxisKeyValue(const FObjectInitializer& ObjectInitializer)
+	: Super(ObjectInitializer)
 {
 	bConsumeInput = true;
 }
diff --git a/Engine/Source/Editor/BlueprintGraph/Private/K2Node_GetInputAxisValue.cpp b/Engine/Source/Editor/BlueprintGraph/Private/K2Node_GetInputAxisValue.cpp
--- a/Engine/Source/Editor/BlueprintGraph/Private/K2Node_GetInputAxisValue.cpp
+++ b/Engine/Source/Editor/BlueprintGraph/Private/K2Node_GetInputAxisValue.cpp
@@ -12,8 +12,8 @@
 
 #define LOCTEXT_NAMESPACE "K2Node_GetInputAxisValue"
 
-UK2Node_GetInputAxisValue::UK2Node_GetInputAxisValue(const class FPostConstructInitializeProperties& PCIP)
-	: Super(PCIP)
+UK2Node_GetInputAxisValue::UK2Node_GetInputAxisValue(const FObjectInitializer& ObjectInitializer)
+	: Super(ObjectInitializer)
 {
 	bConsumeInput = true;
 }
diff --git a/Engine/Source/Editor/BlueprintGraph/Private/K2Node_GetInputVectorAxisValue.cpp b/Engine/Source/Editor/BlueprintGraph/Private/K2Node_GetInputVectorAxisValue.cpp
--- a/Engine/Source/Editor/BlueprintGraph/Private/K2Node_GetInputVectorAxisValue.cpp
+++ b/Engine/Source/Editor/BlueprintGraph/Private/K2Node_GetInputVectorAxisValue.cpp
@@ -7,8 +7,8 @@
 #include "Engine/InputVectorAxisDelegateBinding.h"
 #include "BlueprintActionDatabaseRegistrar.h"
 
-UK2Node_GetInputVectorAxisValue::UK2Node_GetInputVectorAxisValue(const class FPostConstructInitializeProperties& PCIP)
-	: Super(PCIP)
+UK2Node_GetInputVectorAxisValue::UK2Node_GetInputVectorAxisValue(const FObjectInitializer& ObjectInitializer)
+	: Super(ObjectInitializer)
 {
 	bConsumeInput = true;
 }
diff --git a/Engine/Source/Editor/BlueprintGraph/Private/K2Node_GetNumEnumEntries.cpp b/Engine/Source/Editor/BlueprintGraph/Private/K2Node_GetNumEnumEntries.cpp
--- a/Engine/Source/Editor/BlueprintGraph/Private/K2Node_GetNumEnumEntries.cpp
+++ b/Engine/Source/Editor/BlueprintGraph/Private/K2Node_GetNumEnumEntries.cpp
@@ -9,8 +9,8 @@
 #include "EditorCategoryUtils.h"
 #include "BlueprintActionDatabaseRegistrar.h"
 
-UK2Node_GetNumEnumEntries::UK2Node_GetNumEnumEntries(const class FPostConstructInitializeProperties& PCIP)
-	: Super(PCIP)
+UK2Node_GetNumEnumEntries::UK2Node_GetNumEnumEntries(const FObjectInitializer& ObjectInitializer)
+	: Super(ObjectInitializer)
 {
 }
 
diff --git a/Engine/Source/Editor/BlueprintGraph/Private/K2Node_IfThenElse.cpp b/Engine/Source/Editor/BlueprintGraph/Private/K2Node_IfThenElse.cpp
--- a/Engine/Source/Editor/BlueprintGraph/Private/K2Node_IfThenElse.cpp
+++ b/Engine/Source/Editor/BlueprintGraph/Private/K2Node_IfThenElse.cpp
@@ -79,8 +79,8 @@ public:
 	}
 };
 
-UK2Node_IfThenElse::UK2Node_IfThenElse(const class FPostConstructInitializeProperties& PCIP)
-	: Super(PCIP)
+UK2Node_IfThenElse::UK2Node_IfThenElse(const FObjectInitializer& ObjectInitializer)
+	: Super(ObjectInitializer)
 {
 }
 
diff --git a/Engine/Source/Editor/BlueprintGraph/Private/K2Node_InputAction.cpp b/Engine/Source/Editor/BlueprintGraph/Private/K2Node_InputAction.cpp
--- a/Engine/Source/Editor/BlueprintGraph/Private/K2Node_InputAction.cpp
+++ b/Engine/Source/Editor/BlueprintGraph/Private/K2Node_InputAction.cpp
@@ -12,8 +12,8 @@
 
 #define LOCTEXT_NAMESPACE "K2Node_InputAction"
 
-UK2Node_InputAction::UK2Node_InputAction(const class FPostConstructInitializeProperties& PCIP)
-	: Super(PCIP)
+UK2Node_InputAction::UK2Node_InputAction(const FObjectInitializer& ObjectInitializer)
+	: Super(ObjectInitializer)
 {
 	bConsumeInput = true;
 	bOverrideParentBinding = true;
diff --git a/Engine/Source/Editor/BlueprintGraph/Private/K2Node_InputActionEvent.cpp b/Engine/Source/Editor/BlueprintGraph/Private/K2Node_InputActionEvent.cpp
--- a/Engine/Source/Editor/BlueprintGraph/Private/K2Node_InputActionEvent.cpp
+++ b/Engine/Source/Editor/BlueprintGraph/Private/K2Node_InputActionEvent.cpp
@@ -4,8 +4,8 @@
 #include "K2Node_InputActionEvent.h"
 #include "Engine/InputActionDelegateBinding.h"
 
-UK2Node_InputActionEvent::UK2Node_InputActionEvent(const class FPostConstructInitializeProperties& PCIP)
-	: Super(PCIP)
+UK2Node_InputActionEvent::UK2Node_InputActionEvent(const FObjectInitializer& ObjectInitializer)
+	: Super(ObjectInitializer)
 {
 	bConsumeInput = true;
 	bOverrideParentBinding = true;
diff --git a/Engine/Source/Editor/BlueprintGraph/Private/K2Node_InputAxisEvent.cpp b/Engine/Source/Editor/BlueprintGraph/Private/K2Node_InputAxisEvent.cpp
--- a/Engine/Source/Editor/BlueprintGraph/Private/K2Node_InputAxisEvent.cpp
+++ b/Engine/Source/Editor/BlueprintGraph/Private/K2Node_InputAxisEvent.cpp
@@ -9,8 +9,8 @@
 
 #define LOCTEXT_NAMESPACE "K2Node_InputAxisEvent"
 
-UK2Node_InputAxisEvent::UK2Node_InputAxisEvent(const class FPostConstructInitializeProperties& PCIP)
-	: Super(PCIP)
+UK2Node_InputAxisEvent::UK2Node_InputAxisEvent(const FObjectInitializer& ObjectInitializer)
+	: Super(ObjectInitializer)
 {
 	bConsumeInput = true;
 	bOverrideParentBinding = true;
diff --git a/Engine/Source/Editor/BlueprintGraph/Private/K2Node_InputAxisKeyEvent.cpp b/Engine/Source/Editor/BlueprintGraph/Private/K2Node_InputAxisKeyEvent.cpp
--- a/Engine/Source/Editor/BlueprintGraph/Private/K2Node_InputAxisKeyEvent.cpp
+++ b/Engine/Source/Editor/BlueprintGraph/Private/K2Node_InputAxisKeyEvent.cpp
@@ -12,8 +12,8 @@
 
 #define LOCTEXT_NAMESPACE "UK2Node_InputAxisKeyEvent"
 
-UK2Node_InputAxisKeyEvent::UK2Node_InputAxisKeyEvent(const class FPostConstructInitializeProperties& PCIP)
-: Super(PCIP)
+UK2Node_InputAxisKeyEvent::UK2Node_InputAxisKeyEvent(const FObjectInitializer& ObjectInitializer)
+: Super(ObjectInitializer)
 {
 	bConsumeInput = true;
 	bOverrideParentBinding = true;
diff --git a/Engine/Source/Editor/BlueprintGraph/Private/K2Node_InputKey.cpp b/Engine/Source/Editor/BlueprintGraph/Private/K2Node_InputKey.cpp
--- a/Engine/Source/Editor/BlueprintGraph/Private/K2Node_InputKey.cpp
+++ b/Engine/Source/Editor/BlueprintGraph/Private/K2Node_InputKey.cpp
@@ -12,8 +12,8 @@
 
 #define LOCTEXT_NAMESPACE "UK2Node_InputKey"
 
-UK2Node_InputKey::UK2Node_InputKey(const class FPostConstructInitializeProperties& PCIP)
-	: Super(PCIP)
+UK2Node_InputKey::UK2Node_InputKey(const FObjectInitializer& ObjectInitializer)
+	: Super(ObjectInitializer)
 {
 	bConsumeInput = true;
 	bOverrideParentBinding = true;
diff --git a/Engine/Source/Editor/BlueprintGraph/Private/K2Node_InputKeyEvent.cpp b/Engine/Source/Editor/BlueprintGraph/Private/K2Node_InputKeyEvent.cpp
--- a/Engine/Source/Editor/BlueprintGraph/Private/K2Node_InputKeyEvent.cpp
+++ b/Engine/Source/Editor/BlueprintGraph/Private/K2Node_InputKeyEvent.cpp
@@ -4,8 +4,8 @@
 #include "K2Node_InputKeyEvent.h"
 #include "Engine/InputKeyDelegateBinding.h"
 
-UK2Node_InputKeyEvent::UK2Node_InputKeyEvent(const class FPostConstructInitializeProperties& PCIP)
-	: Super(PCIP)
+UK2Node_InputKeyEvent::UK2Node_InputKeyEvent(const FObjectInitializer& ObjectInitializer)
+	: Super(ObjectInitializer)
 {
 	bConsumeInput = true;
 	bOverrideParentBinding = true;
diff --git a/Engine/Source/Editor/BlueprintGraph/Private/K2Node_InputTouch.cpp b/Engine/Source/Editor/BlueprintGraph/Private/K2Node_InputTouch.cpp
--- a/Engine/Source/Editor/BlueprintGraph/Private/K2Node_InputTouch.cpp
+++ b/Engine/Source/Editor/BlueprintGraph/Private/K2Node_InputTouch.cpp
@@ -8,8 +8,8 @@
 #include "EditorCategoryUtils.h"
 #include "BlueprintActionDatabaseRegistrar.h"
 
-UK2Node_InputTouch::UK2Node_InputTouch(const class FPostConstructInitializeProperties& PCIP)
-	: Super(PCIP)
+UK2Node_InputTouch::UK2Node_InputTouch(const FObjectInitializer& ObjectInitializer)
+	: Super(ObjectInitializer)
 {
 	bConsumeInput = true;
 	bOverrideParentBinding = true;
diff --git a/Engine/Source/Editor/BlueprintGraph/Private/K2Node_InputTouchEvent.cpp b/Engine/Source/Editor/BlueprintGraph/Private/K2Node_InputTouchEvent.cpp
--- a/Engine/Source/Editor/BlueprintGraph/Private/K2Node_InputTouchEvent.cpp
+++ b/Engine/Source/Editor/BlueprintGraph/Private/K2Node_InputTouchEvent.cpp
@@ -5,8 +5,8 @@
 #include "CompilerResultsLog.h"
 #include "Engine/InputTouchDelegateBinding.h"
 
-UK2Node_InputTouchEvent::UK2Node_InputTouchEvent(const class FPostConstructInitializeProperties& PCIP)
-	: Super(PCIP)
+UK2Node_InputTouchEvent::UK2Node_InputTouchEvent(const FObjectInitializer& ObjectInitializer)
+	: Super(ObjectInitializer)
 {
 	bConsumeInput = true;
 	bOverrideParentBinding = true;
diff --git a/Engine/Source/Editor/BlueprintGraph/Private/K2Node_InputVectorAxisEvent.cpp b/Engine/Source/Editor/BlueprintGraph/Private/K2Node_InputVectorAxisEvent.cpp
--- a/Engine/Source/Editor/BlueprintGraph/Private/K2Node_InputVectorAxisEvent.cpp
+++ b/Engine/Source/Editor/BlueprintGraph/Private/K2Node_InputVectorAxisEvent.cpp
@@ -7,8 +7,8 @@
 #include "Engine/InputVectorAxisDelegateBinding.h"
 #include "BlueprintActionDatabaseRegistrar.h"
 
-UK2Node_InputVectorAxisEvent::UK2Node_InputVectorAxisEvent(const class FPostConstructInitializeProperties& PCIP)
-: Super(PCIP)
+UK2Node_InputVectorAxisEvent::UK2Node_InputVectorAxisEvent(const FObjectInitializer& ObjectInitializer)
+: Super(ObjectInitializer)
 {
 	EventSignatureName = TEXT("InputVectorAxisHandlerDynamicSignature__DelegateSignature");
 }
diff --git a/Engine/Source/Editor/BlueprintGraph/Private/K2Node_Knot.cpp b/Engine/Source/Editor/BlueprintGraph/Private/K2Node_Knot.cpp
--- a/Engine/Source/Editor/BlueprintGraph/Private/K2Node_Knot.cpp
+++ b/Engine/Source/Editor/BlueprintGraph/Private/K2Node_Knot.cpp
@@ -11,8 +11,8 @@
 /////////////////////////////////////////////////////
 // UK2Node_Knot
 
-UK2Node_Knot::UK2Node_Knot(const class FPostConstructInitializeProperties& PCIP)
-	: Super(PCIP)
+UK2Node_Knot::UK2Node_Knot(const FObjectInitializer& ObjectInitializer)
+	: Super(ObjectInitializer)
 {
 	bCanRenameNode = true;
 }
diff --git a/Engine/Source/Editor/BlueprintGraph/Private/K2Node_Literal.cpp b/Engine/Source/Editor/BlueprintGraph/Private/K2Node_Literal.cpp
--- a/Engine/Source/Editor/BlueprintGraph/Private/K2Node_Literal.cpp
+++ b/Engine/Source/Editor/BlueprintGraph/Private/K2Node_Literal.cpp
@@ -60,8 +60,8 @@ public:
 	}
 };
 
-UK2Node_Literal::UK2Node_Literal(const class FPostConstructInitializeProperties& PCIP)
-	: Super(PCIP)
+UK2Node_Literal::UK2Node_Literal(const FObjectInitializer& ObjectInitializer)
+	: Super(ObjectInitializer)
 {
 }
 
diff --git a/Engine/Source/Editor/BlueprintGraph/Private/K2Node_LocalVariable.cpp b/Engine/Source/Editor/BlueprintGraph/Private/K2Node_LocalVariable.cpp
--- a/Engine/Source/Editor/BlueprintGraph/Private/K2Node_LocalVariable.cpp
+++ b/Engine/Source/Editor/BlueprintGraph/Private/K2Node_LocalVariable.cpp
@@ -8,8 +8,8 @@
 #include "K2Node_FunctionEntry.h"
 #include "K2Node_LocalVariable.h"
 
-UDEPRECATED_K2Node_LocalVariable::UDEPRECATED_K2Node_LocalVariable(const class FPostConstructInitializeProperties& PCIP)
-	: Super(PCIP)
+UDEPRECATED_K2Node_LocalVariable::UDEPRECATED_K2Node_LocalVariable(const FObjectInitializer& ObjectInitializer)
+	: Super(ObjectInitializer)
 {
 	bCanRenameNode = true;
 	CustomVariableName = TEXT("NewLocalVar");
diff --git a/Engine/Source/Editor/BlueprintGraph/Private/K2Node_MCDelegate.cpp b/Engine/Source/Editor/BlueprintGraph/Private/K2Node_MCDelegate.cpp
--- a/Engine/Source/Editor/BlueprintGraph/Private/K2Node_MCDelegate.cpp
+++ b/Engine/Source/Editor/BlueprintGraph/Private/K2Node_MCDelegate.cpp
@@ -12,8 +12,8 @@ FString FK2Node_BaseMCDelegateHelper::DelegatePinName(TEXT("Delegate"));
 
 /////// UK2Node_BaseMCDelegate ///////////
 
-UK2Node_BaseMCDelegate::UK2Node_BaseMCDelegate(const class FPostConstructInitializeProperties& PCIP)
-	: Super(PCIP)
+UK2Node_BaseMCDelegate::UK2Node_BaseMCDelegate(const FObjectInitializer& ObjectInitializer)
+	: Super(ObjectInitializer)
 {
 }
 
@@ -212,8 +212,8 @@ void UK2Node_BaseMCDelegate::GetNodeAttributes( TArray<TKeyValuePair<FString, FS
 
 /////// UK2Node_AddDelegate ///////////
 
-UK2Node_AddDelegate::UK2Node_AddDelegate(const class FPostConstructInitializeProperties& PCIP)
-	: Super(PCIP)
+UK2Node_AddDelegate::UK2Node_AddDelegate(const FObjectInitializer& ObjectInitializer)
+	: Super(ObjectInitializer)
 {
 }
 
@@ -256,8 +256,8 @@ void UK2Node_AddDelegate::GetNodeAttributes( TArray<TKeyValuePair<FString, FStri
 
 /////// UK2Node_ClearDelegate ///////////
 
-UK2Node_ClearDelegate::UK2Node_ClearDelegate(const class FPostConstructInitializeProperties& PCIP)
-	: Super(PCIP)
+UK2Node_ClearDelegate::UK2Node_ClearDelegate(const FObjectInitializer& ObjectInitializer)
+	: Super(ObjectInitializer)
 {
 }
 
@@ -280,8 +280,8 @@ FNodeHandlingFunctor* UK2Node_ClearDelegate::CreateNodeHandler(FKismetCompilerCo
 
 /////// UK2Node_RemoveDelegate ///////////
 
-UK2Node_RemoveDelegate::UK2Node_RemoveDelegate(const class FPostConstructInitializeProperties& PCIP)
-	: Super(PCIP)
+UK2Node_RemoveDelegate::UK2Node_RemoveDelegate(const FObjectInitializer& ObjectInitializer)
+	: Super(ObjectInitializer)
 {
 }
 
@@ -317,8 +317,8 @@ FNodeHandlingFunctor* UK2Node_RemoveDelegate::CreateNodeHandler(FKismetCompilerC
 
 /////// UK2Node_CallDelegate ///////////
 
-UK2Node_CallDelegate::UK2Node_CallDelegate(const class FPostConstructInitializeProperties& PCIP)
-	: Super(PCIP)
+UK2Node_CallDelegate::UK2Node_CallDelegate(const FObjectInitializer& ObjectInitializer)
+	: Super(ObjectInitializer)
 {
 }
 
diff --git a/Engine/Source/Editor/BlueprintGraph/Private/K2Node_MacroInstance.cpp b/Engine/Source/Editor/BlueprintGraph/Private/K2Node_MacroInstance.cpp
--- a/Engine/Source/Editor/BlueprintGraph/Private/K2Node_MacroInstance.cpp
+++ b/Engine/Source/Editor/BlueprintGraph/Private/K2Node_MacroInstance.cpp
@@ -7,8 +7,8 @@
 
 #define LOCTEXT_NAMESPACE "K2Node_MacroInstance"
 
-UK2Node_MacroInstance::UK2Node_MacroInstance(const class FPostConstructInitializeProperties& PCIP)
-	: Super(PCIP)
+UK2Node_MacroInstance::UK2Node_MacroInstance(const FObjectInitializer& ObjectInitializer)
+	: Super(ObjectInitializer)
 {
 	bReconstructNode = false;
 }
diff --git a/Engine/Source/Editor/BlueprintGraph/Private/K2Node_MakeArray.cpp b/Engine/Source/Editor/BlueprintGraph/Private/K2Node_MakeArray.cpp
--- a/Engine/Source/Editor/BlueprintGraph/Private/K2Node_MakeArray.cpp
+++ b/Engine/Source/Editor/BlueprintGraph/Private/K2Node_MakeArray.cpp
@@ -70,8 +70,8 @@ public:
 /////////////////////////////////////////////////////
 // UK2Node_MakeArray
 
-UK2Node_MakeArray::UK2Node_MakeArray(const class FPostConstructInitializeProperties& PCIP)
-	: Super(PCIP)
+UK2Node_MakeArray::UK2Node_MakeArray(const FObjectInitializer& ObjectInitializer)
+	: Super(ObjectInitializer)
 {
 
 	NumInputs = 1;
diff --git a/Engine/Source/Editor/BlueprintGraph/Private/K2Node_MakeStruct.cpp b/Engine/Source/Editor/BlueprintGraph/Private/K2Node_MakeStruct.cpp
--- a/Engine/Source/Editor/BlueprintGraph/Private/K2Node_MakeStruct.cpp
+++ b/Engine/Source/Editor/BlueprintGraph/Private/K2Node_MakeStruct.cpp
@@ -74,8 +74,8 @@ bool UK2Node_MakeStruct::FMakeStructPinManager::CanTreatPropertyAsOptional(UProp
 		return UK2Node_MakeStruct::CanBeExposed(TestProperty);
 	}
 
-UK2Node_MakeStruct::UK2Node_MakeStruct(const class FPostConstructInitializeProperties& PCIP)
-	: Super(PCIP)
+UK2Node_MakeStruct::UK2Node_MakeStruct(const FObjectInitializer& ObjectInitializer)
+	: Super(ObjectInitializer)
 {
 }
 
diff --git a/Engine/Source/Editor/BlueprintGraph/Private/K2Node_MathExpression.cpp b/Engine/Source/Editor/BlueprintGraph/Private/K2Node_MathExpression.cpp
--- a/Engine/Source/Editor/BlueprintGraph/Private/K2Node_MathExpression.cpp
+++ b/Engine/Source/Editor/BlueprintGraph/Private/K2Node_MathExpression.cpp
@@ -2268,8 +2268,8 @@ private:
 *******************************************************************************/
 
 //------------------------------------------------------------------------------
-UK2Node_MathExpression::UK2Node_MathExpression(const FPostConstructInitializeProperties& PCIP)
-	: Super(PCIP)
+UK2Node_MathExpression::UK2Node_MathExpression(const FObjectInitializer& ObjectInitializer)
+	: Super(ObjectInitializer)
 {
 	// renaming the node rebuilds the expression (the node name is where they 
 	// specify the math equation)
diff --git a/Engine/Source/Editor/BlueprintGraph/Private/K2Node_MatineeController.cpp b/Engine/Source/Editor/BlueprintGraph/Private/K2Node_MatineeController.cpp
--- a/Engine/Source/Editor/BlueprintGraph/Private/K2Node_MatineeController.cpp
+++ b/Engine/Source/Editor/BlueprintGraph/Private/K2Node_MatineeController.cpp
@@ -10,8 +10,8 @@
 #include "Matinee/MatineeActor.h"
 #include "Matinee/InterpData.h"
 
-UK2Node_MatineeController::UK2Node_MatineeController(const class FPostConstructInitializeProperties& PCIP)
-	: Super(PCIP)
+UK2Node_MatineeController::UK2Node_MatineeController(const FObjectInitializer& ObjectInitializer)
+	: Super(ObjectInitializer)
 {
 	FMatineeDelegates::Get().OnEventKeyframeAdded.AddUObject(this, &UK2Node_MatineeController::OnEventKeyframeAdded);
 	FMatineeDelegates::Get().OnEventKeyframeRenamed.AddUObject(this, &UK2Node_MatineeController::OnEventKeyframeRenamed);
diff --git a/Engine/Source/Editor/BlueprintGraph/Private/K2Node_Message.cpp b/Engine/Source/Editor/BlueprintGraph/Private/K2Node_Message.cpp
--- a/Engine/Source/Editor/BlueprintGraph/Private/K2Node_Message.cpp
+++ b/Engine/Source/Editor/BlueprintGraph/Private/K2Node_Message.cpp
@@ -8,8 +8,8 @@
 
 #define LOCTEXT_NAMESPACE "K2Node_Message"
 
-UK2Node_Message::UK2Node_Message(const class FPostConstructInitializeProperties& PCIP)
-	: Super(PCIP)
+UK2Node_Message::UK2Node_Message(const FObjectInitializer& ObjectInitializer)
+	: Super(ObjectInitializer)
 {
 }
 
diff --git a/Engine/Source/Editor/BlueprintGraph/Private/K2Node_MultiGate.cpp b/Engine/Source/Editor/BlueprintGraph/Private/K2Node_MultiGate.cpp
--- a/Engine/Source/Editor/BlueprintGraph/Private/K2Node_MultiGate.cpp
+++ b/Engine/Source/Editor/BlueprintGraph/Private/K2Node_MultiGate.cpp
@@ -452,8 +452,8 @@ public:
 	}
 };
 
-UK2Node_MultiGate::UK2Node_MultiGate(const class FPostConstructInitializeProperties& PCIP)
-	: Super(PCIP)
+UK2Node_MultiGate::UK2Node_MultiGate(const FObjectInitializer& ObjectInitializer)
+	: Super(ObjectInitializer)
 {
 }
 
diff --git a/Engine/Source/Editor/BlueprintGraph/Private/K2Node_PureAssignmentStatement.cpp b/Engine/Source/Editor/BlueprintGraph/Private/K2Node_PureAssignmentStatement.cpp
--- a/Engine/Source/Editor/BlueprintGraph/Private/K2Node_PureAssignmentStatement.cpp
+++ b/Engine/Source/Editor/BlueprintGraph/Private/K2Node_PureAssignmentStatement.cpp
@@ -70,8 +70,8 @@ FString UK2Node_PureAssignmentStatement::ValuePinName = FString(TEXT("Value"));
 FString UK2Node_PureAssignmentStatement::OutputPinName = FString(TEXT("ReturnValue"));
 
 
-UK2Node_PureAssignmentStatement::UK2Node_PureAssignmentStatement(const class FPostConstructInitializeProperties& PCIP)
-	: Super(PCIP)
+UK2Node_PureAssignmentStatement::UK2Node_PureAssignmentStatement(const FObjectInitializer& ObjectInitializer)
+	: Super(ObjectInitializer)
 {
 }
 
diff --git a/Engine/Source/Editor/BlueprintGraph/Private/K2Node_Select.cpp b/Engine/Source/Editor/BlueprintGraph/Private/K2Node_Select.cpp
--- a/Engine/Source/Editor/BlueprintGraph/Private/K2Node_Select.cpp
+++ b/Engine/Source/Editor/BlueprintGraph/Private/K2Node_Select.cpp
@@ -187,8 +187,8 @@ public:
 	}
 };
 
-UK2Node_Select::UK2Node_Select(const class FPostConstructInitializeProperties& PCIP)
-	: Super(PCIP)
+UK2Node_Select::UK2Node_Select(const FObjectInitializer& ObjectInitializer)
+	: Super(ObjectInitializer)
 {
 
 	NumOptionPins = 2;
diff --git a/Engine/Source/Editor/BlueprintGraph/Private/K2Node_SelectEnum.cpp b/Engine/Source/Editor/BlueprintGraph/Private/K2Node_SelectEnum.cpp
--- a/Engine/Source/Editor/BlueprintGraph/Private/K2Node_SelectEnum.cpp
+++ b/Engine/Source/Editor/BlueprintGraph/Private/K2Node_SelectEnum.cpp
@@ -6,8 +6,8 @@
 #include "BlueprintFieldNodeSpawner.h"
 #include "BlueprintActionDatabaseRegistrar.h"
 
-UDEPRECATED_K2Node_SelectEnum::UDEPRECATED_K2Node_SelectEnum(const class FPostConstructInitializeProperties& PCIP)
-	: Super(PCIP)
+UDEPRECATED_K2Node_SelectEnum::UDEPRECATED_K2Node_SelectEnum(const FObjectInitializer& ObjectInitializer)
+	: Super(ObjectInitializer)
 {
 
 	const UEdGraphSchema_K2* Schema = GetDefault<UEdGraphSchema_K2>();
diff --git a/Engine/Source/Editor/BlueprintGraph/Private/K2Node_Self.cpp b/Engine/Source/Editor/BlueprintGraph/Private/K2Node_Self.cpp
--- a/Engine/Source/Editor/BlueprintGraph/Private/K2Node_Self.cpp
+++ b/Engine/Source/Editor/BlueprintGraph/Private/K2Node_Self.cpp
@@ -32,8 +32,8 @@ public:
 	}
 };
 
-UK2Node_Self::UK2Node_Self(const class FPostConstructInitializeProperties& PCIP)
-	: Super(PCIP)
+UK2Node_Self::UK2Node_Self(const FObjectInitializer& ObjectInitializer)
+	: Super(ObjectInitializer)
 {
 }
 
diff --git a/Engine/Source/Editor/BlueprintGraph/Private/K2Node_SetFieldsInStruct.cpp b/Engine/Source/Editor/BlueprintGraph/Private/K2Node_SetFieldsInStruct.cpp
--- a/Engine/Source/Editor/BlueprintGraph/Private/K2Node_SetFieldsInStruct.cpp
+++ b/Engine/Source/Editor/BlueprintGraph/Private/K2Node_SetFieldsInStruct.cpp
@@ -29,8 +29,8 @@ public:
 	}
 };
 
-UK2Node_SetFieldsInStruct::UK2Node_SetFieldsInStruct(const class FPostConstructInitializeProperties& PCIP)
-	: Super(PCIP)
+UK2Node_SetFieldsInStruct::UK2Node_SetFieldsInStruct(const FObjectInitializer& ObjectInitializer)
+	: Super(ObjectInitializer)
 {
 }
 
diff --git a/Engine/Source/Editor/BlueprintGraph/Private/K2Node_SpawnActor.cpp b/Engine/Source/Editor/BlueprintGraph/Private/K2Node_SpawnActor.cpp
--- a/Engine/Source/Editor/BlueprintGraph/Private/K2Node_SpawnActor.cpp
+++ b/Engine/Source/Editor/BlueprintGraph/Private/K2Node_SpawnActor.cpp
@@ -11,8 +11,8 @@ static FString NoCollisionFailPinName(TEXT("SpawnEvenIfColliding"));
 
 #define LOCTEXT_NAMESPACE "K2Node_SpawnActor"
 
-UK2Node_SpawnActor::UK2Node_SpawnActor(const class FPostConstructInitializeProperties& PCIP)
-	: Super(PCIP)
+UK2Node_SpawnActor::UK2Node_SpawnActor(const FObjectInitializer& ObjectInitializer)
+	: Super(ObjectInitializer)
 {
 	NodeTooltip = LOCTEXT("NodeTooltip", "Attempts to spawn a new Actor with the specified transform");
 }
diff --git a/Engine/Source/Editor/BlueprintGraph/Private/K2Node_SpawnActorFromClass.cpp b/Engine/Source/Editor/BlueprintGraph/Private/K2Node_SpawnActorFromClass.cpp
--- a/Engine/Source/Editor/BlueprintGraph/Private/K2Node_SpawnActorFromClass.cpp
+++ b/Engine/Source/Editor/BlueprintGraph/Private/K2Node_SpawnActorFromClass.cpp
@@ -22,8 +22,8 @@ FString FK2Node_SpawnActorFromClassHelper::NoCollisionFailPinName(TEXT("SpawnEve
 
 #define LOCTEXT_NAMESPACE "K2Node_SpawnActorFromClass"
 
-UK2Node_SpawnActorFromClass::UK2Node_SpawnActorFromClass(const class FPostConstructInitializeProperties& PCIP)
-	: Super(PCIP)
+UK2Node_SpawnActorFromClass::UK2Node_SpawnActorFromClass(const FObjectInitializer& ObjectInitializer)
+	: Super(ObjectInitializer)
 {
 	NodeTooltip = LOCTEXT("NodeTooltip", "Attempts to spawn a new Actor with the specified transform");
 }
diff --git a/Engine/Source/Editor/BlueprintGraph/Private/K2Node_StructMemberGet.cpp b/Engine/Source/Editor/BlueprintGraph/Private/K2Node_StructMemberGet.cpp
--- a/Engine/Source/Editor/BlueprintGraph/Private/K2Node_StructMemberGet.cpp
+++ b/Engine/Source/Editor/BlueprintGraph/Private/K2Node_StructMemberGet.cpp
@@ -8,8 +8,8 @@
 
 #define LOCTEXT_NAMESPACE "K2Node"
 
-UK2Node_StructMemberGet::UK2Node_StructMemberGet(const class FPostConstructInitializeProperties& PCIP)
-	: Super(PCIP)
+UK2Node_StructMemberGet::UK2Node_StructMemberGet(const FObjectInitializer& ObjectInitializer)
+	: Super(ObjectInitializer)
 {
 }
 
diff --git a/Engine/Source/Editor/BlueprintGraph/Private/K2Node_StructMemberSet.cpp b/Engine/Source/Editor/BlueprintGraph/Private/K2Node_StructMemberSet.cpp
--- a/Engine/Source/Editor/BlueprintGraph/Private/K2Node_StructMemberSet.cpp
+++ b/Engine/Source/Editor/BlueprintGraph/Private/K2Node_StructMemberSet.cpp
@@ -9,8 +9,8 @@
 
 #define LOCTEXT_NAMESPACE "K2Node"
 
-UK2Node_StructMemberSet::UK2Node_StructMemberSet(const class FPostConstructInitializeProperties& PCIP)
-	: Super(PCIP)
+UK2Node_StructMemberSet::UK2Node_StructMemberSet(const FObjectInitializer& ObjectInitializer)
+	: Super(ObjectInitializer)
 {
 }
 
diff --git a/Engine/Source/Editor/BlueprintGraph/Private/K2Node_StructOperation.cpp b/Engine/Source/Editor/BlueprintGraph/Private/K2Node_StructOperation.cpp
--- a/Engine/Source/Editor/BlueprintGraph/Private/K2Node_StructOperation.cpp
+++ b/Engine/Source/Editor/BlueprintGraph/Private/K2Node_StructOperation.cpp
@@ -7,8 +7,8 @@
 //////////////////////////////////////////////////////////////////////////
 // UK2Node_StructOperation
 
-UK2Node_StructOperation::UK2Node_StructOperation(const class FPostConstructInitializeProperties& PCIP)
-	: Super(PCIP)
+UK2Node_StructOperation::UK2Node_StructOperation(const FObjectInitializer& ObjectInitializer)
+	: Super(ObjectInitializer)
 {
 }
 
diff --git a/Engine/Source/Editor/BlueprintGraph/Private/K2Node_Switch.cpp b/Engine/Source/Editor/BlueprintGraph/Private/K2Node_Switch.cpp
--- a/Engine/Source/Editor/BlueprintGraph/Private/K2Node_Switch.cpp
+++ b/Engine/Source/Editor/BlueprintGraph/Private/K2Node_Switch.cpp
@@ -134,8 +134,8 @@ private:
 	FEdGraphPinType ExpectedSelectionPinType;
 };
 
-UK2Node_Switch::UK2Node_Switch(const class FPostConstructInitializeProperties& PCIP)
-	: Super(PCIP)
+UK2Node_Switch::UK2Node_Switch(const FObjectInitializer& ObjectInitializer)
+	: Super(ObjectInitializer)
 {
 
 	bHasDefaultPin = true;
diff --git a/Engine/Source/Editor/BlueprintGraph/Private/K2Node_SwitchEnum.cpp b/Engine/Source/Editor/BlueprintGraph/Private/K2Node_SwitchEnum.cpp
--- a/Engine/Source/Editor/BlueprintGraph/Private/K2Node_SwitchEnum.cpp
+++ b/Engine/Source/Editor/BlueprintGraph/Private/K2Node_SwitchEnum.cpp
@@ -9,8 +9,8 @@
 
 #define LOCTEXT_NAMESPACE "K2Node"
 
-UK2Node_SwitchEnum::UK2Node_SwitchEnum(const class FPostConstructInitializeProperties& PCIP)
-	: Super(PCIP)
+UK2Node_SwitchEnum::UK2Node_SwitchEnum(const FObjectInitializer& ObjectInitializer)
+	: Super(ObjectInitializer)
 {
 
 	bHasDefaultPin = false;
diff --git a/Engine/Source/Editor/BlueprintGraph/Private/K2Node_SwitchInteger.cpp b/Engine/Source/Editor/BlueprintGraph/Private/K2Node_SwitchInteger.cpp
--- a/Engine/Source/Editor/BlueprintGraph/Private/K2Node_SwitchInteger.cpp
+++ b/Engine/Source/Editor/BlueprintGraph/Private/K2Node_SwitchInteger.cpp
@@ -5,8 +5,8 @@
 #include "BlueprintNodeSpawner.h"
 #include "BlueprintActionDatabaseRegistrar.h"
 
-UK2Node_SwitchInteger::UK2Node_SwitchInteger(const class FPostConstructInitializeProperties& PCIP)
-	: Super(PCIP)
+UK2Node_SwitchInteger::UK2Node_SwitchInteger(const FObjectInitializer& ObjectInitializer)
+	: Super(ObjectInitializer)
 {
 	StartIndex = 0;
 
