BP: GetDataTableRow - some fixes

diff --git a/Engine/Source/Editor/BlueprintGraph/Classes/K2Node_GetDataTableRow.h b/Engine/Source/Editor/BlueprintGraph/Classes/K2Node_GetDataTableRow.h
--- a/Engine/Source/Editor/BlueprintGraph/Classes/K2Node_GetDataTableRow.h
+++ b/Engine/Source/Editor/BlueprintGraph/Classes/K2Node_GetDataTableRow.h
@@ -18,7 +18,6 @@ class BLUEPRINTGRAPH_API UK2Node_GetDataTableRow : public UK2Node
 	virtual void AllocateDefaultPins() override;
 	virtual FText GetNodeTitle(ENodeTitleType::Type TitleType) const override;
 	virtual void PinDefaultValueChanged(UEdGraphPin* Pin) override;
-	virtual void NotifyPinConnectionListChanged(UEdGraphPin* Pin) override;
 	virtual FText GetTooltipText() const override;
 	virtual void ExpandNode(class FKismetCompilerContext& CompilerContext, UEdGraph* SourceGraph) override;
 	virtual FName GetPaletteIcon(FLinearColor& OutColor) const override
@@ -33,6 +32,7 @@ class BLUEPRINTGRAPH_API UK2Node_GetDataTableRow : public UK2Node
 	virtual void ReallocatePinsDuringReconstruction(TArray<UEdGraphPin*>& OldPins) override;
 	virtual void GetMenuActions(FBlueprintActionDatabaseRegistrar& ActionRegistrar) const override;
 	virtual FText GetMenuCategory() const override;
+	virtual bool IsConnectionDisallowed(const UEdGraphPin* MyPin, const UEdGraphPin* OtherPin, FString& OutReason) const override;
 	// End UK2Node interface
 
 
diff --git a/Engine/Source/Editor/BlueprintGraph/Private/K2Node_GetDataTableRow.cpp b/Engine/Source/Editor/BlueprintGraph/Private/K2Node_GetDataTableRow.cpp
--- a/Engine/Source/Editor/BlueprintGraph/Private/K2Node_GetDataTableRow.cpp
+++ b/Engine/Source/Editor/BlueprintGraph/Private/K2Node_GetDataTableRow.cpp
@@ -193,22 +193,14 @@ FText UK2Node_GetDataTableRow::GetMenuCategory() const
 	return FEditorCategoryUtils::GetCommonCategory(FCommonEditorCategory::Utilities);
 }
 
-void UK2Node_GetDataTableRow::NotifyPinConnectionListChanged(UEdGraphPin* Pin)
+bool UK2Node_GetDataTableRow::IsConnectionDisallowed(const UEdGraphPin* MyPin, const UEdGraphPin* OtherPin, FString& OutReason) const
 {
-	if (Pin && Pin->PinName == UK2Node_GetDataTableRowHelper::DataTablePinName)
+	if (MyPin && MyPin->PinName == UK2Node_GetDataTableRowHelper::DataTablePinName)
 	{
-		// Because the archetype has changed, we break the output link as the output pin type will change
-		UEdGraphPin* ResultPin = GetResultPin();
-		ResultPin->BreakAllPinLinks();
-
-		// If the pin has been connected, clear the default values so that we don't hold on to references
-		if (Pin->LinkedTo.Num() == 1)
-		{
-			UScriptStruct* RowStruct = GetDataTableRowStructType(&(Pin->LinkedTo));
-
-			SetReturnTypeForStruct(RowStruct);
-		}
+		OutReason = LOCTEXT("OnlyLiteralDTSupported", "Only literal data table is supported").ToString();
+		return true;
 	}
+	return false;
 }
 
 void UK2Node_GetDataTableRow::PinDefaultValueChanged(UEdGraphPin* ChangedPin) 
@@ -218,6 +210,14 @@ void UK2Node_GetDataTableRow::PinDefaultValueChanged(UEdGraphPin* ChangedPin)
 		UScriptStruct* RowStruct = GetDataTableRowStructType();
 
 		SetReturnTypeForStruct(RowStruct);
+
+		auto RowNamePin = GetRowNamePin();
+		auto DataTable = Cast<UDataTable>(ChangedPin->DefaultObject);
+		if (RowNamePin && RowNamePin->DefaultValue.IsEmpty() && DataTable)
+		{
+			auto Iterator = DataTable->RowMap.CreateConstIterator();
+			RowNamePin->DefaultValue = Iterator ? Iterator.Key().ToString() : FString();
+		}
 	}
 }
 
