BP: BLueprint Function Library always has a WorldContext pin. Proper Default self object for static function.

diff --git a/Engine/Source/Editor/BlueprintGraph/Private/K2Node_CallFunction.cpp b/Engine/Source/Editor/BlueprintGraph/Private/K2Node_CallFunction.cpp
--- a/Engine/Source/Editor/BlueprintGraph/Private/K2Node_CallFunction.cpp
+++ b/Engine/Source/Editor/BlueprintGraph/Private/K2Node_CallFunction.cpp
@@ -790,7 +790,8 @@ bool UK2Node_CallFunction::CreatePinsForFunctionCall(const UFunction* Function)
 			// For static methods, wire up the self to the CDO of the class if it's not us
 			if (!bIsFunctionCompatibleWithSelf)
 			{
-				SelfPin->DefaultObject = FunctionOwnerClass->GetDefaultObject();
+				auto AuthoritativeClass = FunctionOwnerClass->GetAuthoritativeClass();
+				SelfPin->DefaultObject = AuthoritativeClass->GetDefaultObject();
 			}
 
 			// Purity doesn't matter with a static function, we can always hide the self pin since we know how to call the method
diff --git a/Engine/Source/Editor/BlueprintGraph/Private/K2Node_FunctionEntry.cpp b/Engine/Source/Editor/BlueprintGraph/Private/K2Node_FunctionEntry.cpp
--- a/Engine/Source/Editor/BlueprintGraph/Private/K2Node_FunctionEntry.cpp
+++ b/Engine/Source/Editor/BlueprintGraph/Private/K2Node_FunctionEntry.cpp
@@ -157,14 +157,8 @@ UEdGraphPin* UK2Node_FunctionEntry::GetAutoWorldContextPin() const
 
 void UK2Node_FunctionEntry::RemoveUnnecessaryAutoWorldContext()
 {
-	auto WorldContextPin = GetAutoWorldContextPin();
-	if (WorldContextPin)
-	{
-		if (!WorldContextPin->LinkedTo.Num())
-		{
-			Pins.Remove(WorldContextPin);
-		}
-	}
+	// Don't remove. After cl#2310036 Skeleton function must have the same signature as regular one.
+	// For skeleton function, no expansion is done, so the pin is always unconnected. The unconnected pin is removed, so the signature is different.
 }
 
 void UK2Node_FunctionEntry::RemoveOutputPin(UEdGraphPin* PinToRemove)
