BP: WorldContext parameter in static function - improvements. When a static function is called in unsafe context a warning is displayed.

diff --git a/Engine/Source/Editor/BlueprintGraph/Private/EdGraphSchema_K2.cpp b/Engine/Source/Editor/BlueprintGraph/Private/EdGraphSchema_K2.cpp
--- a/Engine/Source/Editor/BlueprintGraph/Private/EdGraphSchema_K2.cpp
+++ b/Engine/Source/Editor/BlueprintGraph/Private/EdGraphSchema_K2.cpp
@@ -441,8 +441,11 @@ bool UEdGraphSchema_K2::CanFunctionBeUsedInClass(const UClass* InClass, UFunctio
 		{
 			if (InDestGraph && !InFunction->HasMetaData(FBlueprintMetadata::MD_CallableWithoutWorldContext))
 			{
-				UClass* ParentClass = FBlueprintEditorUtils::FindBlueprintForGraphChecked(InDestGraph)->ParentClass;
-				if (!ParentClass->GetDefaultObject()->ImplementsGetWorld() && !ParentClass->HasMetaData(FBlueprintMetadata::MD_ShowWorldContextPin))
+				auto BP = FBlueprintEditorUtils::FindBlueprintForGraph(InDestGraph);
+				const bool bIsFunctLib = BP && (EBlueprintType::BPTYPE_FunctionLibrary == BP->BlueprintType);
+				UClass* ParentClass = BP ? BP->ParentClass : NULL;
+				const bool bIncompatibleParrent = ParentClass && (!ParentClass->GetDefaultObject()->ImplementsGetWorld() && !ParentClass->HasMetaData(FBlueprintMetadata::MD_ShowWorldContextPin));
+				if (!bIsFunctLib && bIncompatibleParrent)
 				{
 					if(OutReason != nullptr)
 					{
diff --git a/Engine/Source/Editor/BlueprintGraph/Private/K2Node_CallFunction.cpp b/Engine/Source/Editor/BlueprintGraph/Private/K2Node_CallFunction.cpp
--- a/Engine/Source/Editor/BlueprintGraph/Private/K2Node_CallFunction.cpp
+++ b/Engine/Source/Editor/BlueprintGraph/Private/K2Node_CallFunction.cpp
@@ -1540,12 +1540,12 @@ void UK2Node_CallFunction::ValidateNodeDuringCompilation(class FCompilerResultsL
 {
 	Super::ValidateNodeDuringCompilation(MessageLog);
 
+	const UBlueprint* Blueprint = GetBlueprint();
 	UFunction *Function = GetTargetFunction();
 	if (Function == NULL)
 	{
 		FString OwnerName;
 
-		UBlueprint* Blueprint = GetBlueprint();
 		if (Blueprint != nullptr)
 		{
 			OwnerName = Blueprint->GetName();
@@ -1600,11 +1600,15 @@ void UK2Node_CallFunction::ValidateNodeDuringCompilation(class FCompilerResultsL
 		}
 
 		// enforce WorldContext restrictions
-		else if (   Function->HasMetaData(FBlueprintMetadata::MD_WorldContext) 
-			     && !Function->HasMetaData(FBlueprintMetadata::MD_CallableWithoutWorldContext))
-		{
-			UClass* ParentClass = GetBlueprint()->ParentClass;
-			if (!ParentClass->GetDefaultObject()->ImplementsGetWorld() && !ParentClass->HasMetaData(FBlueprintMetadata::MD_ShowWorldContextPin))
+		const bool bInsideBpFuncLibrary = Blueprint && (BPTYPE_FunctionLibrary == Blueprint->BlueprintType);
+		if (!bInsideBpFuncLibrary && 
+			Function->HasMetaData(FBlueprintMetadata::MD_WorldContext) && 
+			!Function->HasMetaData(FBlueprintMetadata::MD_CallableWithoutWorldContext))
+		{
+			check(Blueprint);
+			UClass* ParentClass = Blueprint->ParentClass;
+			check(ParentClass);
+			if (ParentClass && !ParentClass->GetDefaultObject()->ImplementsGetWorld() && !ParentClass->HasMetaData(FBlueprintMetadata::MD_ShowWorldContextPin))
 			{
 				MessageLog.Warning(*LOCTEXT("FunctionUnsafeInContext", "Function '@@' is unsafe to call from blueprints of class '@@'.").ToString(), this, ParentClass);
 			}
diff --git a/Engine/Source/Editor/KismetCompiler/Private/KismetCompiler.cpp b/Engine/Source/Editor/KismetCompiler/Private/KismetCompiler.cpp
--- a/Engine/Source/Editor/KismetCompiler/Private/KismetCompiler.cpp
+++ b/Engine/Source/Editor/KismetCompiler/Private/KismetCompiler.cpp
@@ -1582,8 +1582,7 @@ void FKismetCompilerContext::FinishCompilingFunction(FKismetFunctionContext& Con
 	}
 	if (auto WorldContextPin = EntryNode->GetAutoWorldContextPin())
 	{
-		Function->SetMetaData(FBlueprintMetadata::MD_DefaultToSelf, *WorldContextPin->PinName); 
-		Function->SetMetaData(TEXT("HidePin"), *WorldContextPin->PinName);
+		Function->SetMetaData(FBlueprintMetadata::MD_WorldContext, *WorldContextPin->PinName);
 	}
 
 	for (int32 EntryPinIndex = 0; EntryPinIndex < EntryNode->Pins.Num(); ++EntryPinIndex)
