Any latent action nodes derived from BaseAsyncTask are now supported in macros. Reminder, these macros will prevent themselve from being placeable in functions.

Added an error for BaseAsyncTask nodes to inform when they find themselves in a function graph, detailing the function and the macro that expanded.

diff --git a/Engine/Source/Editor/BlueprintGraph/Classes/K2Node_BaseAsyncTask.h b/Engine/Source/Editor/BlueprintGraph/Classes/K2Node_BaseAsyncTask.h
--- a/Engine/Source/Editor/BlueprintGraph/Classes/K2Node_BaseAsyncTask.h
+++ b/Engine/Source/Editor/BlueprintGraph/Classes/K2Node_BaseAsyncTask.h
@@ -21,6 +21,7 @@ class BLUEPRINTGRAPH_API UK2Node_BaseAsyncTask : public UK2Node
 	virtual FText GetTooltipText() const override;
 	virtual FText GetNodeTitle(ENodeTitleType::Type TitleType) const override;
 	virtual bool IsCompatibleWithGraph(const UEdGraph* TargetGraph) const override;
+	virtual void ValidateNodeDuringCompilation(class FCompilerResultsLog& MessageLog) const override;
 	// End of UEdGraphNode interface
 
 	// UK2Node interface
diff --git a/Engine/Source/Editor/BlueprintGraph/Private/K2Node_BaseAsyncTask.cpp b/Engine/Source/Editor/BlueprintGraph/Private/K2Node_BaseAsyncTask.cpp
--- a/Engine/Source/Editor/BlueprintGraph/Private/K2Node_BaseAsyncTask.cpp
+++ b/Engine/Source/Editor/BlueprintGraph/Private/K2Node_BaseAsyncTask.cpp
@@ -34,11 +34,12 @@ FText UK2Node_BaseAsyncTask::GetNodeTitle(ENodeTitleType::Type TitleType) const
 
 bool UK2Node_BaseAsyncTask::IsCompatibleWithGraph(const UEdGraph* TargetGraph) const
 {
-	bool bIsCompatible = true;
-	// Can only place events in ubergraphs, and basicasync task creates an event node:
-	if (TargetGraph->GetSchema()->GetGraphType(TargetGraph) != EGraphType::GT_Ubergraph)
+	bool bIsCompatible = false;
+	// Can only place events in ubergraphs and macros (other code will help prevent macros with latents from ending up in functions), and basicasync task creates an event node:
+	EGraphType GraphType = TargetGraph->GetSchema()->GetGraphType(TargetGraph);
+	if (GraphType == EGraphType::GT_Ubergraph || GraphType == EGraphType::GT_Macro)
 	{
-		bIsCompatible = false;
+		bIsCompatible = true;
 	}
 	return bIsCompatible && Super::IsCompatibleWithGraph(TargetGraph);
 }
@@ -409,4 +410,22 @@ UFunction* UK2Node_BaseAsyncTask::GetFactoryFunction() const
 	return FactoryFunction;
 }
 
+void UK2Node_BaseAsyncTask::ValidateNodeDuringCompilation(class FCompilerResultsLog& MessageLog) const
+{
+	UEdGraphSchema_K2 const* K2Schema = GetDefault<UEdGraphSchema_K2>();
+
+	if(UObject const* SourceObject = MessageLog.FindSourceObject(this))
+	{
+		// Lets check if it's a result of macro expansion, to give a helpful error
+		if(UK2Node_MacroInstance const* MacroInstance = Cast<UK2Node_MacroInstance>(SourceObject))
+		{
+			// Since it's not possible to check the graph's type, just check if this is a ubergraph using the schema's name for it
+			if(!(GetGraph()->HasAnyFlags(RF_Transient) && GetGraph()->GetName().StartsWith(K2Schema->FN_ExecuteUbergraphBase.ToString())))
+			{
+				MessageLog.Error(*LOCTEXT("AsyncTaskInFunctionFromMacro", "@@ is being used in Function '@@' resulting from expansion of Macro '@@'").ToString(), this, GetGraph(), MacroInstance);
+			}
+		}
+	}
+}
+
 #undef LOCTEXT_NAMESPACE
diff --git a/Engine/Source/Editor/UnrealEd/Private/Kismet2/BlueprintEditorUtils.cpp b/Engine/Source/Editor/UnrealEd/Private/Kismet2/BlueprintEditorUtils.cpp
--- a/Engine/Source/Editor/UnrealEd/Private/Kismet2/BlueprintEditorUtils.cpp
+++ b/Engine/Source/Editor/UnrealEd/Private/Kismet2/BlueprintEditorUtils.cpp
@@ -6595,6 +6595,12 @@ bool FBlueprintEditorUtils::CheckIfGraphHasLatentFunctions(UEdGraph* InGraph)
 							return true;
 						}
 					}
+					else if(const UK2Node_BaseAsyncTask* BaseAsyncNode = Cast<UK2Node_BaseAsyncTask>(Node))
+					{
+						// Async tasks are latent nodes
+						TunnelNode->MetaData.HasLatentFunctions = 1;
+						return true;
+					}
 					else if(const UK2Node_MacroInstance* MacroInstanceNode = Cast<UK2Node_MacroInstance>(Node))
 					{
 						// Any macro graphs that haven't already been checked need to be checked for latent function calls
