BP Static Function have WorldContext parameter only when it's necessary

diff --git a/Engine/Source/Editor/BlueprintGraph/Private/K2Node_CallFunction.cpp b/Engine/Source/Editor/BlueprintGraph/Private/K2Node_CallFunction.cpp
--- a/Engine/Source/Editor/BlueprintGraph/Private/K2Node_CallFunction.cpp
+++ b/Engine/Source/Editor/BlueprintGraph/Private/K2Node_CallFunction.cpp
@@ -1696,7 +1696,10 @@ void UK2Node_CallFunction::ExpandNode(class FKismetCompilerContext& CompilerCont
 				}
 			};
 			FStructConnectHelper::Connect(DefaultToSelfMetaValue, this, BetterSelfPin, Schema, CompilerContext.MessageLog);
-			FStructConnectHelper::Connect(WorldContextMetaValue, this, BetterSelfPin, Schema, CompilerContext.MessageLog);
+			if (!Function->HasMetaData(FBlueprintMetadata::MD_CallableWithoutWorldContext))
+			{
+				FStructConnectHelper::Connect(WorldContextMetaValue, this, BetterSelfPin, Schema, CompilerContext.MessageLog);
+			}
 		}
 	}
 
diff --git a/Engine/Source/Editor/BlueprintGraph/Private/K2Node_FunctionEntry.cpp b/Engine/Source/Editor/BlueprintGraph/Private/K2Node_FunctionEntry.cpp
--- a/Engine/Source/Editor/BlueprintGraph/Private/K2Node_FunctionEntry.cpp
+++ b/Engine/Source/Editor/BlueprintGraph/Private/K2Node_FunctionEntry.cpp
@@ -157,8 +157,14 @@ UEdGraphPin* UK2Node_FunctionEntry::GetAutoWorldContextPin() const
 
 void UK2Node_FunctionEntry::RemoveUnnecessaryAutoWorldContext()
 {
-	// Don't remove. After cl#2310036 Skeleton function must have the same signature as regular one.
-	// For skeleton function, no expansion is done, so the pin is always unconnected. The unconnected pin is removed, so the signature is different.
+	auto WorldContextPin = GetAutoWorldContextPin();
+	if (WorldContextPin)
+	{
+		if (!WorldContextPin->LinkedTo.Num())
+		{
+			Pins.Remove(WorldContextPin);
+		}
+	}
 }
 
 void UK2Node_FunctionEntry::RemoveOutputPin(UEdGraphPin* PinToRemove)
diff --git a/Engine/Source/Editor/KismetCompiler/Private/KismetCompiler.cpp b/Engine/Source/Editor/KismetCompiler/Private/KismetCompiler.cpp
--- a/Engine/Source/Editor/KismetCompiler/Private/KismetCompiler.cpp
+++ b/Engine/Source/Editor/KismetCompiler/Private/KismetCompiler.cpp
@@ -2466,7 +2466,8 @@ void FKismetCompilerContext::MergeUbergraphPagesIn(UEdGraph* Ubergraph)
 // Expands out nodes that need it
 void FKismetCompilerContext::ExpansionStep(UEdGraph* Graph, bool bAllowUbergraphExpansions)
 {
-	if (bIsFullCompile)
+	// Node expansion may affect a signature of function in FunctionLibrary
+	if (bIsFullCompile || (EBlueprintType::BPTYPE_FunctionLibrary == Blueprint->BlueprintType))
 	{
 		BP_SCOPED_COMPILER_EVENT_STAT(EKismetCompilerStats_Expansion);
 
