Display True/False in select node for booleans
Don't allow add/remove pins for boolean select node

diff --git a/Engine/Source/Editor/BlueprintGraph/Private/K2Node_Select.cpp b/Engine/Source/Editor/BlueprintGraph/Private/K2Node_Select.cpp
--- a/Engine/Source/Editor/BlueprintGraph/Private/K2Node_Select.cpp
+++ b/Engine/Source/Editor/BlueprintGraph/Private/K2Node_Select.cpp
@@ -231,15 +231,22 @@ void UK2Node_Select::AllocateDefaultPins()
 			NewPin = CreatePin(EGPD_Input, Schema->PC_Wildcard, TEXT(""), NULL, false, false, PinName);
 		}
 
-		if (NewPin && Idx < EnumEntryFriendlyNames.Num())
+		if (NewPin)
 		{
-			if (EnumEntryFriendlyNames[Idx] != NAME_None)
+			if (IndexPinType.PinCategory == UEdGraphSchema_K2::PC_Boolean)
 			{
-				NewPin->PinFriendlyName = FText::FromName(EnumEntryFriendlyNames[Idx]);
+				NewPin->PinFriendlyName = (Idx == 0 ? GFalse : GTrue);
 			}
-			else
+			else if (Idx < EnumEntryFriendlyNames.Num())
 			{
-				NewPin->PinFriendlyName = FText::GetEmpty();
+				if (EnumEntryFriendlyNames[Idx] != NAME_None)
+				{
+					NewPin->PinFriendlyName = FText::FromName(EnumEntryFriendlyNames[Idx]);
+				}
+				else
+				{
+					NewPin->PinFriendlyName = FText::GetEmpty();
+				}
 			}
 		}
 	}
@@ -647,6 +654,10 @@ bool UK2Node_Select::CanAddOptionPinToNode() const
 	{
 		return false;
 	}
+	else if (IndexPinType.PinCategory == Schema->PC_Boolean)
+	{
+		return false;
+	}
 
 	return true;
 }
@@ -660,6 +671,10 @@ bool UK2Node_Select::CanRemoveOptionPinToNode() const
 	{
 		return false;
 	}
+	else if (IndexPinType.PinCategory == Schema->PC_Boolean)
+	{
+		return false;
+	}
 
 	return true;
 }
