[UE-3323] Making sure to preload the datatable used in constructing a GetDataTableRow's output pin (so we have the proper data).

diff --git a/Engine/Source/Editor/BlueprintGraph/Private/K2Node_GetDataTableRow.cpp b/Engine/Source/Editor/BlueprintGraph/Private/K2Node_GetDataTableRow.cpp
--- a/Engine/Source/Editor/BlueprintGraph/Private/K2Node_GetDataTableRow.cpp
+++ b/Engine/Source/Editor/BlueprintGraph/Private/K2Node_GetDataTableRow.cpp
@@ -160,7 +160,17 @@ void UK2Node_GetDataTableRow::GetDataTableRowNameList(TArray<TSharedPtr<FName>>&
 void UK2Node_GetDataTableRow::ReallocatePinsDuringReconstruction(TArray<UEdGraphPin*>& OldPins) 
 {
 	Super::ReallocatePinsDuringReconstruction(OldPins);
-		
+
+	if (UEdGraphPin* DataTablePin = GetDataTablePin(&OldPins))
+	{
+		if (UDataTable* DataTable = Cast<UDataTable>(DataTablePin->DefaultObject))
+		{
+			// make sure to properly load the data-table object so that we can 
+			// farm the "RowStruct" property from it (below, in GetDataTableRowStructType)
+			PreloadObject(DataTable);
+		}
+	}
+
 	UScriptStruct* RowStruct = GetDataTableRowStructType(&OldPins);
 	if( RowStruct != NULL )
 	{
