Preventing duplicate function entries from showing up in the BP context menu (removing overridden functions from the action database).

diff --git a/Engine/Source/Editor/BlueprintGraph/Private/BlueprintActionDatabase.cpp b/Engine/Source/Editor/BlueprintGraph/Private/BlueprintActionDatabase.cpp
--- a/Engine/Source/Editor/BlueprintGraph/Private/BlueprintActionDatabase.cpp
+++ b/Engine/Source/Editor/BlueprintGraph/Private/BlueprintActionDatabase.cpp
@@ -276,18 +276,30 @@ namespace BlueprintActionDatabaseImpl
 	static bool IsPropertyBlueprintVisible(UProperty const* const Property);
 
 	/**
-	 * 
-	 * 
-	 * @param  Function	
-	 * @return 
+	 * Checks to see if the specified function is a blueprint owned function
+	 * that was inherited from an implemented interface.
+	 *
+	 * @param  Function	 The function to check.
+	 * @return True if the function is owned by a blueprint, and some (implemented) interface has a matching function name.
 	 */
 	static bool IsBlueprintInterfaceFunction(const UFunction* Function);
 
 	/**
+	 * Checks to see if the specified function is a blueprint owned function
+	 * that was inherited from the blueprint's parent.
 	 * 
+	 * @param  Function	 The function to check.
+	 * @return True if the function is owned by a blueprint, and some parent has a matching function name.
+	 */
+	static bool IsInheritedBlueprintFunction(const UFunction* Function);
+
+	/**
+	 * Retrieves all the actions pertaining to a class and its fields (functions,
+	 * variables, delegates, etc.). Actions that are conceptually owned by the 
+	 * class.
 	 * 
-	 * @param  Class	
-	 * @param  ActionListOut	
+	 * @param  Class			The class you want actions for.
+	 * @param  ActionListOut	The array you want filled with the requested actions.
 	 */
 	static void GetClassMemberActions(UClass* const Class, FActionList& ActionListOut);
 
@@ -476,20 +488,31 @@ static bool BlueprintActionDatabaseImpl::IsBlueprintInterfaceFunction(const UFun
 			for (int32 InterfaceIndex = 0; (InterfaceIndex < BpOuter->ImplementedInterfaces.Num()) && !bIsBpInterfaceFunc; ++InterfaceIndex)
 			{
 				FBPInterfaceDescription& InterfaceDesc = BpOuter->ImplementedInterfaces[InterfaceIndex];
-				for (TFieldIterator<UFunction> FunctionIt(InterfaceDesc.Interface, EFieldIteratorFlags::IncludeSuper); FunctionIt; ++FunctionIt)
-				{
-					if (FunctionIt->GetFName() == FuncName)
-					{
-						bIsBpInterfaceFunc = true;
-						break;
-					}
-				}
+				bIsBpInterfaceFunc = (InterfaceDesc.Interface->FindFunctionByName(FuncName) != nullptr);
 			}
 		}
 	}
 	return bIsBpInterfaceFunc;
 }
 
+//------------------------------------------------------------------------------
+static bool BlueprintActionDatabaseImpl::IsInheritedBlueprintFunction(const UFunction* Function)
+{
+	bool bIsBpInheritedFunc = false;
+	if (UClass* FuncClass = Function->GetOwnerClass())
+	{
+		if (UBlueprint* BpOwner = Cast<UBlueprint>(FuncClass->ClassGeneratedBy))
+		{
+			FName FuncName = Function->GetFName();
+			if (UClass* ParentClass = BpOwner->ParentClass)
+			{
+				bIsBpInheritedFunc = (ParentClass->FindFunctionByName(FuncName, EIncludeSuperFlag::IncludeSuper) != nullptr);
+			}
+		}
+	}
+	return bIsBpInheritedFunc;
+}
+
 //------------------------------------------------------------------------------
 static void BlueprintActionDatabaseImpl::GetClassMemberActions(UClass* const Class, FActionList& ActionListOut)
 {
@@ -521,8 +544,16 @@ static void BlueprintActionDatabaseImpl::AddClassFunctionActions(UClass const* c
 	for (TFieldIterator<UFunction> FunctionIt(Class, EFieldIteratorFlags::ExcludeSuper); FunctionIt; ++FunctionIt)
 	{
 		UFunction* Function = *FunctionIt;
-		bool const bIsBpInterfaceFunc = BlueprintActionDatabaseImpl::IsBlueprintInterfaceFunction(Function);
 
+		bool const bIsInheritedFunction = BlueprintActionDatabaseImpl::IsInheritedBlueprintFunction(Function);
+		if (bIsInheritedFunction)
+		{
+			// inherited functions will be captured when the parent class is ran
+			// through this function (no need to duplicate)
+			continue;
+		}
+
+		bool const bIsBpInterfaceFunc = BlueprintActionDatabaseImpl::IsBlueprintInterfaceFunction(Function);
 		if (UEdGraphSchema_K2::FunctionCanBePlacedAsEvent(Function) && !bIsBpInterfaceFunc)
 		{
 			if (UBlueprintEventNodeSpawner* NodeSpawner = UBlueprintEventNodeSpawner::Create(Function))
