[UE-4562] Hiding message nodes from blueprint vars that couldn't have implemented the interface.

diff --git a/Engine/Source/Editor/BlueprintGraph/Private/BlueprintActionFilter.cpp b/Engine/Source/Editor/BlueprintGraph/Private/BlueprintActionFilter.cpp
--- a/Engine/Source/Editor/BlueprintGraph/Private/BlueprintActionFilter.cpp
+++ b/Engine/Source/Editor/BlueprintGraph/Private/BlueprintActionFilter.cpp
@@ -1249,13 +1249,22 @@ static bool BlueprintActionFilterImpl::IsExtraneousInterfaceCall(FBlueprintActio
 		UClass* InterfaceClass = Function->GetOwnerClass();
 		checkSlow(InterfaceClass->IsChildOf<UInterface>());
 
+		bool const bCanBeAddedToBlueprints = !InterfaceClass->HasMetaData(FBlueprintMetadata::MD_CannotImplementInterfaceInBlueprint);
+
 		bIsFilteredOut = (Filter.TargetClasses.Num() > 0);
 		for (const UClass* TargetClass : Filter.TargetClasses)
 		{
-			// if the target class implements (or is) the message spawner's 
-			// interface owner, then we don't need the message node (the regular
-			// interface call shouldn't be available)
-			if (!IsClassOfType(TargetClass, InterfaceClass))
+			bool const bImplementsInterface = IsClassOfType(TargetClass, InterfaceClass);
+			// if this is a blueprint class, and "CannotImplementInterfaceInBlueprint" 
+			// is set on the interface, then we know sub-classes cannot have
+			// the interface either (so there's no point to offering a message node)
+			bool const bIsBlueprintClass = (Cast<UBlueprintGeneratedClass>(TargetClass) != nullptr);
+
+			// if the class doesn't directly implement the interface (and it is 
+			// a possibility that some sub-class does), then we want to offer 
+			// the message node (in case the Target object is actually an 
+			// instance of a sub-class)
+			if (!bImplementsInterface && (!bIsBlueprintClass || bCanBeAddedToBlueprints))
 			{
 				bIsFilteredOut = false;
 				break;
