Adding support for a pure version of the cast node.

diff --git a/Engine/Source/Editor/BlueprintGraph/Classes/K2Node_DynamicCast.h b/Engine/Source/Editor/BlueprintGraph/Classes/K2Node_DynamicCast.h
--- a/Engine/Source/Editor/BlueprintGraph/Classes/K2Node_DynamicCast.h
+++ b/Engine/Source/Editor/BlueprintGraph/Classes/K2Node_DynamicCast.h
@@ -20,6 +20,7 @@ class UK2Node_DynamicCast : public UK2Node
 	virtual FLinearColor GetNodeTitleColor() const override;
 	virtual FText GetNodeTitle(ENodeTitleType::Type TitleType) const override;
 	virtual FName GetPaletteIcon(FLinearColor& OutColor) const override{ return TEXT("GraphEditor.Cast_16x"); }
+	virtual void GetContextMenuActions(const FGraphNodeContextMenuBuilder& Context) const override;
 	// End UEdGraphNode interface
 
 	// UK2Node interface
@@ -28,6 +29,7 @@ class UK2Node_DynamicCast : public UK2Node
 	virtual bool HasExternalBlueprintDependencies(TArray<class UStruct*>* OptionalOutput) const override;
 	virtual FText GetMenuCategory() const override;
 	virtual FBlueprintNodeSignature GetSignature() const override;
+	virtual bool IsNodePure() const override { return bIsPureCast; }
 	// End of UK2Node interface
 
 	/** Get the 'valid cast' exec pin */
@@ -42,8 +44,22 @@ class UK2Node_DynamicCast : public UK2Node
 	/** Get the input object to be casted pin */
 	BLUEPRINTGRAPH_API virtual UEdGraphPin* GetCastSourcePin() const;
 
+	/**
+	 * Will change the node's purity, and reallocate pins accordingly (adding/
+	 * removing exec pins).
+	 * 
+	 * @param  bNewPurity  The new value for bIsPureCast.
+	 */
+	void SetPurity(bool bNewPurity);
+
 protected:
+	/** Flips the node's purity (adding/removing exec pins as needed). */
+	void TogglePurity();
+
 	/** Constructing FText strings can be costly, so we cache the node's title */
 	FNodeTextCache CachedNodeTitle;
+
+	UPROPERTY()
+	bool bIsPureCast;
 };
 
diff --git a/Engine/Source/Editor/BlueprintGraph/Private/DynamicCastHandler.cpp b/Engine/Source/Editor/BlueprintGraph/Private/DynamicCastHandler.cpp
--- a/Engine/Source/Editor/BlueprintGraph/Private/DynamicCastHandler.cpp
+++ b/Engine/Source/Editor/BlueprintGraph/Private/DynamicCastHandler.cpp
@@ -35,22 +35,6 @@ void FKCHandler_DynamicCast::Compile(FKismetFunctionContext& Context, UEdGraphNo
 		CompilerContext.MessageLog.Error(*LOCTEXT("BadCastNoTargetType_Error", "Node @@ has an invalid target type, please delete and recreate it").ToString(), Node);
 	}
 
-	// Successful interface cast execution
-	UEdGraphPin* SuccessPin = DynamicCastNode->GetValidCastPin();
-	UEdGraphNode* SuccessNode = NULL;
-	if (SuccessPin->LinkedTo.Num() > 0)
-	{
-		SuccessNode = SuccessPin->LinkedTo[0]->GetOwningNode();
-	}
-
-	// Failed interface cast execution
-	UEdGraphPin* FailurePin = DynamicCastNode->GetInvalidCastPin();
-	UEdGraphNode* FailureNode = NULL;
-	if (FailurePin->LinkedTo.Num() > 0)
-	{
-		FailureNode = FailurePin->LinkedTo[0]->GetOwningNode();
-	}
-
 	// Self Pin
 	UEdGraphPin* SourceObjectPin = DynamicCastNode->GetCastSourcePin();
 	UEdGraphPin* PinToTry = FEdGraphUtilities::GetNetFromPin(SourceObjectPin);
@@ -132,26 +116,35 @@ void FKCHandler_DynamicCast::Compile(FKismetFunctionContext& Context, UEdGraphNo
 	CastStatement.RHS.Add(ClassTerm);
 	CastStatement.RHS.Add(*ObjectToCast);
 
-	// Check result of cast statement
-	FBlueprintCompiledStatement& CheckResultStatement = Context.AppendStatementForNode(Node);
-	const UClass* SubObjectClass = Cast<UClass>((*CastResultTerm)->Type.PinSubCategoryObject.Get());
-	const bool bIsInterfaceCast = (SubObjectClass != NULL && SubObjectClass->HasAnyClassFlags(CLASS_Interface));
-
-	CheckResultStatement.Type = KCST_ObjectToBool;
-	CheckResultStatement.LHS = BoolTerm;
-	CheckResultStatement.RHS.Add(*CastResultTerm);
-
-	// Failure condition...skip to the failed output
-	FBlueprintCompiledStatement& FailCastGoto = Context.AppendStatementForNode(Node);
-	FailCastGoto.Type = KCST_GotoIfNot;
-	FailCastGoto.LHS = BoolTerm;
-	Context.GotoFixupRequestMap.Add(&FailCastGoto, FailurePin);
-
-	// Successful cast...hit the success output node
-	FBlueprintCompiledStatement& SuccessCastGoto = Context.AppendStatementForNode(Node);
-	SuccessCastGoto.Type = KCST_UnconditionalGoto;
-	SuccessCastGoto.LHS = BoolTerm;
-	Context.GotoFixupRequestMap.Add(&SuccessCastGoto, SuccessPin);
+	UEdGraphPin* SuccessPin = DynamicCastNode->GetValidCastPin();
+	bool const bIsPureCast = (SuccessPin == nullptr);
+
+	if (!bIsPureCast)
+	{
+		// Check result of cast statement
+		FBlueprintCompiledStatement& CheckResultStatement = Context.AppendStatementForNode(Node);
+		const UClass* SubObjectClass = Cast<UClass>((*CastResultTerm)->Type.PinSubCategoryObject.Get());
+		const bool bIsInterfaceCast = (SubObjectClass != NULL && SubObjectClass->HasAnyClassFlags(CLASS_Interface));
+
+		CheckResultStatement.Type = KCST_ObjectToBool;
+		CheckResultStatement.LHS = BoolTerm;
+		CheckResultStatement.RHS.Add(*CastResultTerm);
+	
+		UEdGraphPin* FailurePin = DynamicCastNode->GetInvalidCastPin();
+		check(FailurePin != nullptr);
+		// Failure condition...skip to the failed output
+		FBlueprintCompiledStatement& FailCastGoto = Context.AppendStatementForNode(Node);
+		FailCastGoto.Type = KCST_GotoIfNot;
+		FailCastGoto.LHS = BoolTerm;
+		Context.GotoFixupRequestMap.Add(&FailCastGoto, FailurePin);
+
+		// Successful cast...hit the success output node
+		FBlueprintCompiledStatement& SuccessCastGoto = Context.AppendStatementForNode(Node);
+		SuccessCastGoto.Type = KCST_UnconditionalGoto;
+		SuccessCastGoto.LHS = BoolTerm;
+		Context.GotoFixupRequestMap.Add(&SuccessCastGoto, SuccessPin);
+	}
+	
 }
 
 #undef LOCTEXT_NAMESPACE
diff --git a/Engine/Source/Editor/BlueprintGraph/Private/K2Node_ClassDynamicCast.cpp b/Engine/Source/Editor/BlueprintGraph/Private/K2Node_ClassDynamicCast.cpp
--- a/Engine/Source/Editor/BlueprintGraph/Private/K2Node_ClassDynamicCast.cpp
+++ b/Engine/Source/Editor/BlueprintGraph/Private/K2Node_ClassDynamicCast.cpp
@@ -27,18 +27,26 @@ void UK2Node_ClassDynamicCast::AllocateDefaultPins()
 	//@TODO: Move this somewhere more sensible
 	ensure((TargetType == NULL) || (!TargetType->HasAnyClassFlags(CLASS_NewerVersionExists)));
 
-	const UEdGraphSchema_K2* K2Schema = GetDefault<UEdGraphSchema_K2>();
+	const UEdGraphSchema_K2* K2Schema = Cast<UEdGraphSchema_K2>(GetSchema());
+	check(K2Schema != nullptr);
+	if (!K2Schema->DoesGraphSupportImpureFunctions(GetGraph()))
+	{
+		bIsPureCast = true;
+	}
+
+	if (!bIsPureCast)
+	{
+		// Input - Execution Pin
+		CreatePin(EGPD_Input, K2Schema->PC_Exec, TEXT(""), NULL, false, false, K2Schema->PN_Execute);
 
-	// Input - Execution Pin
-	CreatePin(EGPD_Input, K2Schema->PC_Exec, TEXT(""), NULL, false, false, K2Schema->PN_Execute);
+		// Output - Execution Pins
+		CreatePin(EGPD_Output, K2Schema->PC_Exec, TEXT(""), NULL, false, false, K2Schema->PN_CastSucceeded);
+		CreatePin(EGPD_Output, K2Schema->PC_Exec, TEXT(""), NULL, false, false, K2Schema->PN_CastFailed);
+	}
 
 	// Input - Source type Pin
 	CreatePin(EGPD_Input, K2Schema->PC_Class, TEXT(""), UObject::StaticClass(), false, false, FClassDynamicCastHelper::GetClassToCastName());
 
-	// Output - Execution Pins
-	CreatePin(EGPD_Output, K2Schema->PC_Exec, TEXT(""), NULL, false, false, K2Schema->PN_CastSucceeded);
-	CreatePin(EGPD_Output, K2Schema->PC_Exec, TEXT(""), NULL, false, false, K2Schema->PN_CastFailed);
-
 	// Output - Data Pin
 	if (TargetType != NULL)
 	{
diff --git a/Engine/Source/Editor/BlueprintGraph/Private/K2Node_DynamicCast.cpp b/Engine/Source/Editor/BlueprintGraph/Private/K2Node_DynamicCast.cpp
--- a/Engine/Source/Editor/BlueprintGraph/Private/K2Node_DynamicCast.cpp
+++ b/Engine/Source/Editor/BlueprintGraph/Private/K2Node_DynamicCast.cpp
@@ -9,6 +9,7 @@
 
 UK2Node_DynamicCast::UK2Node_DynamicCast(const FObjectInitializer& ObjectInitializer)
 	: Super(ObjectInitializer)
+	, bIsPureCast(false)
 {
 }
 
@@ -18,22 +19,30 @@ void UK2Node_DynamicCast::AllocateDefaultPins()
 	//@TODO: Move this somewhere more sensible
 	check((TargetType == NULL) || (!TargetType->HasAnyClassFlags(CLASS_NewerVersionExists)));
 
-	const UEdGraphSchema_K2* K2Schema = GetDefault<UEdGraphSchema_K2>();
+	const UEdGraphSchema_K2* K2Schema = Cast<UEdGraphSchema_K2>(GetSchema());
+	check(K2Schema != nullptr);
+	if (!K2Schema->DoesGraphSupportImpureFunctions(GetGraph()))
+	{
+		bIsPureCast = true;
+	}
 
-	// Input - Execution Pin
-	CreatePin(EGPD_Input, K2Schema->PC_Exec, TEXT(""), NULL, false, false, K2Schema->PN_Execute);
+	if (!bIsPureCast)
+	{
+		// Input - Execution Pin
+		CreatePin(EGPD_Input, K2Schema->PC_Exec, TEXT(""), NULL, false, false, K2Schema->PN_Execute);
+
+		// Output - Execution Pins
+		CreatePin(EGPD_Output, K2Schema->PC_Exec, TEXT(""), NULL, false, false, K2Schema->PN_CastSucceeded);
+		CreatePin(EGPD_Output, K2Schema->PC_Exec, TEXT(""), NULL, false, false, K2Schema->PN_CastFailed);
+	}
 
 	// Input - Source type Pin
 	CreatePin(EGPD_Input, K2Schema->PC_Object, TEXT(""), UObject::StaticClass(), false, false, K2Schema->PN_ObjectToCast);
 
-	// Output - Execution Pins
-	CreatePin(EGPD_Output, K2Schema->PC_Exec, TEXT(""), NULL, false, false, K2Schema->PN_CastSucceeded);
-	CreatePin(EGPD_Output, K2Schema->PC_Exec, TEXT(""), NULL, false, false, K2Schema->PN_CastFailed);
-
 	// Output - Data Pin
 	if (TargetType != NULL)
 	{
-		FString CastResultPinName = K2Schema->PN_CastedValuePrefix + TargetType->GetName();
+		FString CastResultPinName = K2Schema->PN_CastedValuePrefix + TargetType->GetDisplayNameText().ToString();
 		if (TargetType->IsChildOf(UInterface::StaticClass()))
 		{
 			CreatePin(EGPD_Output, K2Schema->PC_Interface, TEXT(""), *TargetType, false, false, CastResultPinName);
@@ -81,13 +90,57 @@ FText UK2Node_DynamicCast::GetNodeTitle(ENodeTitleType::Type TitleType) const
 	return CachedNodeTitle;
 }
 
+void UK2Node_DynamicCast::GetContextMenuActions(const FGraphNodeContextMenuBuilder& Context) const
+{
+	Super::GetContextMenuActions(Context);
+
+	Context.MenuBuilder->BeginSection("K2NodeDynamicCast", LOCTEXT("DynamicCastHeader", "Cast"));
+	{
+		FText MenuEntryTitle   = LOCTEXT("MakePureTitle",   "Convert to pure cast");
+		FText MenuEntryTooltip = LOCTEXT("MakePureTooltip", "Removes the execution pins to make the node more versitile (NOTE: the cast could still, resulting in an invalid output).");
+
+		bool bCanTogglePurity = true;
+		auto CanExecutePurityToggle = [](bool const bCanTogglePurity)->bool
+		{
+			return bCanTogglePurity;
+		};
+
+		if (bIsPureCast)
+		{
+			MenuEntryTitle   = LOCTEXT("MakeImpureTitle",   "Convert to impure cast");
+			MenuEntryTooltip = LOCTEXT("MakeImpureTooltip", "Adds in branching execution pins so that you can separatly handle when the cast fails/succeeds.");
+		
+			const UEdGraphSchema_K2* K2Schema = Cast<UEdGraphSchema_K2>(GetSchema());
+			check(K2Schema != nullptr);
+
+			bCanTogglePurity = K2Schema->DoesGraphSupportImpureFunctions(GetGraph());
+			if (!bCanTogglePurity)
+			{
+				MenuEntryTooltip = LOCTEXT("CannotMakeImpureTooltip", "This graph does not support impure calls (and you should therefore test the cast's result for validity).");
+			}
+		}
+		
+		Context.MenuBuilder->AddMenuEntry(
+			MenuEntryTitle,
+			MenuEntryTooltip,
+			FSlateIcon(),
+			FUIAction(
+				FExecuteAction::CreateUObject(this, &UK2Node_DynamicCast::TogglePurity),
+				FCanExecuteAction::CreateStatic(CanExecutePurityToggle, bCanTogglePurity),
+				FIsActionChecked()
+			)
+		);
+	}
+	Context.MenuBuilder->EndSection();
+}
+
 UEdGraphPin* UK2Node_DynamicCast::GetValidCastPin() const
 {
 	const UEdGraphSchema_K2* K2Schema = GetDefault<UEdGraphSchema_K2>();
 
 	UEdGraphPin* Pin = FindPin(K2Schema->PN_CastSucceeded);
-	check(Pin != NULL);
-	check(Pin->Direction == EGPD_Output);
+	check((Pin != nullptr) || bIsPureCast);
+	check((Pin == nullptr) || (Pin->Direction == EGPD_Output));
 	return Pin;
 }
 
@@ -96,8 +149,8 @@ UEdGraphPin* UK2Node_DynamicCast::GetInvalidCastPin() const
 	const UEdGraphSchema_K2* K2Schema = GetDefault<UEdGraphSchema_K2>();
 
 	UEdGraphPin* Pin = FindPin(K2Schema->PN_CastFailed);
-	check(Pin != NULL);
-	check(Pin->Direction == EGPD_Output);
+	check((Pin != nullptr) || bIsPureCast);
+	check((Pin == nullptr) || (Pin->Direction == EGPD_Output));
 	return Pin;
 }
 
@@ -108,7 +161,7 @@ UEdGraphPin* UK2Node_DynamicCast::GetCastResultPin() const
 	if(TargetType != NULL)
 	{
 		const UEdGraphSchema_K2* K2Schema = GetDefault<UEdGraphSchema_K2>();
-		FString PinName = K2Schema->PN_CastedValuePrefix + TargetType->GetName();
+		FString PinName = K2Schema->PN_CastedValuePrefix + TargetType->GetDisplayNameText().ToString();
 		Pin = FindPin(PinName);
 	}
 		
@@ -125,6 +178,25 @@ UEdGraphPin* UK2Node_DynamicCast::GetCastSourcePin() const
 	return Pin;
 }
 
+void UK2Node_DynamicCast::SetPurity(bool bNewPurity)
+{
+	if (bNewPurity != bIsPureCast)
+	{
+		bIsPureCast = bNewPurity;
+
+		bool const bHasBeenConstructed = (Pins.Num() > 0);
+		if (bHasBeenConstructed)
+		{
+			ReconstructNode();
+		}
+	}
+}
+
+void UK2Node_DynamicCast::TogglePurity()
+{
+	SetPurity(!bIsPureCast);
+}
+
 UK2Node::ERedirectType UK2Node_DynamicCast::DoPinsMatchForReconstruction(const UEdGraphPin* NewPin, int32 NewPinIndex, const UEdGraphPin* OldPin, int32 OldPinIndex) const
 {
 	ERedirectType RedirectType = Super::DoPinsMatchForReconstruction(NewPin, NewPinIndex, OldPin, OldPinIndex);
