[UE-4562] Message nodes now work for pure/const functions.

diff --git a/Engine/Source/Editor/BlueprintGraph/Classes/K2Node_Message.h b/Engine/Source/Editor/BlueprintGraph/Classes/K2Node_Message.h
--- a/Engine/Source/Editor/BlueprintGraph/Classes/K2Node_Message.h
+++ b/Engine/Source/Editor/BlueprintGraph/Classes/K2Node_Message.h
@@ -14,6 +14,7 @@ class UK2Node_Message : public UK2Node_CallFunction
 
 	// Begin UEdGraphNode interface.
 	virtual FText GetNodeTitle(ENodeTitleType::Type TitleType) const override;
+	virtual void AllocateDefaultPins() override;
 	// End UEdGraphNode interface.
 
 	// Begin K2Node interface.
@@ -21,6 +22,7 @@ class UK2Node_Message : public UK2Node_CallFunction
 	virtual class FNodeHandlingFunctor* CreateNodeHandler(class FKismetCompilerContext& CompilerContext) const override;
 	virtual void ExpandNode(class FKismetCompilerContext& CompilerContext, UEdGraph* SourceGraph) override;
 	virtual FName GetCornerIcon() const override;
+	virtual bool IsNodePure() const override { return false; }
 	// End K2Node interface.
 
 	// Begin K2Node_CallFunction Interface.
diff --git a/Engine/Source/Editor/BlueprintGraph/Private/K2Node_Message.cpp b/Engine/Source/Editor/BlueprintGraph/Private/K2Node_Message.cpp
--- a/Engine/Source/Editor/BlueprintGraph/Private/K2Node_Message.cpp
+++ b/Engine/Source/Editor/BlueprintGraph/Private/K2Node_Message.cpp
@@ -5,6 +5,7 @@
 #include "K2Node_CallArrayFunction.h"
 #include "K2Node_Message.h"
 #include "Kismet/KismetArrayLibrary.h"
+#include "K2Node_TemporaryVariable.h"
 
 #define LOCTEXT_NAMESPACE "K2Node_Message"
 
@@ -45,6 +46,25 @@ FText UK2Node_Message::GetNodeTitle(ENodeTitleType::Type TitleType) const
 	return CachedNodeTitles.GetCachedTitle(TitleType);
 }
 
+void UK2Node_Message::AllocateDefaultPins()
+{
+	UFunction* MessageNodeFunction = GetTargetFunction();
+	// since we have branching logic in ExpandNode(), this has to be an impure
+	// node with exec pins
+	//
+	// @TODO: make it so we can have impure message nodes using a custom 
+	//        FNodeHandlingFunctor, instead of ExpandNode()
+	if (MessageNodeFunction && MessageNodeFunction->HasAnyFunctionFlags(FUNC_BlueprintPure))
+	{
+		// Input - Execution Pin
+		CreatePin(EGPD_Input,  UEdGraphSchema_K2::PC_Exec, TEXT(""), NULL, false, false, UEdGraphSchema_K2::PN_Execute);
+		// Output - Execution Pin
+		CreatePin(EGPD_Output, UEdGraphSchema_K2::PC_Exec, TEXT(""), NULL, false, false, UEdGraphSchema_K2::PN_Then);
+	}
+
+	Super::AllocateDefaultPins();
+}
+
 UEdGraphPin* UK2Node_Message::CreateSelfPin(const UFunction* Function)
 {
 	const UEdGraphSchema_K2* K2Schema = GetDefault<UEdGraphSchema_K2>();
@@ -80,7 +100,7 @@ void UK2Node_Message::ExpandNode(class FKismetCompilerContext& CompilerContext,
 	const bool bThenPinConnected = ThenPin && (ThenPin->LinkedTo.Num() > 0);
 
 	// Skip ourselves if our exec isn't wired up
-	if( bExecPinConnected )
+	if (bExecPinConnected)
 	{
 		// Make sure our interface is valid
 		if (FunctionReference.GetMemberParentClass(this) == NULL)
@@ -115,7 +135,6 @@ void UK2Node_Message::ExpandNode(class FKismetCompilerContext& CompilerContext,
 		CastToInterfaceNode->TargetType = MessageNodeFunction->GetOuterUClass();
 		CastToInterfaceNode->AllocateDefaultPins();
 
-		UEdGraphPin* CastToInterfaceValidPin = CastToInterfaceNode->GetValidCastPin();
 		UEdGraphPin* CastToInterfaceResultPin = CastToInterfaceNode->GetCastResultPin();
 		if( !CastToInterfaceResultPin )
 		{
@@ -125,8 +144,14 @@ void UK2Node_Message::ExpandNode(class FKismetCompilerContext& CompilerContext,
 
 		CastToInterfaceResultPin->PinType.PinSubCategoryObject = *CastToInterfaceNode->TargetType;
 
-		// Wire up the connections
-		CompilerContext.MovePinLinksToIntermediate(*ExecPin, *CastToInterfaceNode->GetExecPin());
+		if (ExecPin != nullptr)
+		{
+			UEdGraphPin* CastExecInput = CastToInterfaceNode->GetExecPin();
+			check(CastExecInput != nullptr);
+
+			// Wire up the connections
+			CompilerContext.MovePinLinksToIntermediate(*ExecPin, *CastExecInput);
+		}
 
 		UEdGraphPin* CastToInterfaceSourceObjectPin = CastToInterfaceNode->GetCastSourcePin();
 		CompilerContext.MovePinLinksToIntermediate(*MessageSelfPin, *CastToInterfaceSourceObjectPin);
@@ -137,19 +162,27 @@ void UK2Node_Message::ExpandNode(class FKismetCompilerContext& CompilerContext,
 		FunctionCallNode->FunctionReference = FunctionReference;
 		FunctionCallNode->AllocateDefaultPins();
 
-		// Wire up the connections
-		UEdGraphPin* CallFunctionExecPin = Schema->FindExecutionPin(*FunctionCallNode, EGPD_Input);
-		CastToInterfaceValidPin->MakeLinkTo(CallFunctionExecPin);
+		UEdGraphPin* CastToInterfaceValidPin = CastToInterfaceNode->GetValidCastPin();
+		check(CastToInterfaceValidPin != nullptr);
 
+		UEdGraphPin* LastOutCastFaildPin   = CastToInterfaceNode->GetInvalidCastPin();
+		check(LastOutCastFaildPin != nullptr);
+		UEdGraphPin* LastOutCastSuccessPin = CastToInterfaceValidPin;
+		// Wire up the connections
+		if (UEdGraphPin* CallFunctionExecPin = Schema->FindExecutionPin(*FunctionCallNode, EGPD_Input))
+		{
+			CastToInterfaceValidPin->MakeLinkTo(CallFunctionExecPin);
+			LastOutCastSuccessPin = Schema->FindExecutionPin(*FunctionCallNode, EGPD_Output);
+		}
+		
 		// Self pin
 		UEdGraphPin* FunctionCallSelfPin = Schema->FindSelfPin(*FunctionCallNode, EGPD_Input);
 		CastToInterfaceResultPin->MakeLinkTo(FunctionCallSelfPin);
 
-		UEdGraphPin* LastOutCastFaildPin = CastToInterfaceNode->GetInvalidCastPin();
-
 		UFunction* ArrayClearFunction = UKismetArrayLibrary::StaticClass()->FindFunctionByName(FName(TEXT("Array_Clear")));
 		check(ArrayClearFunction);
 
+		bool const bIsPureFunc = Super::IsNodePure();
 		// Variable pins - Try to associate variable inputs to the message node with the variable inputs and outputs to the call function node
 		for( int32 i = 0; i < Pins.Num(); i++ )
 		{
@@ -166,6 +199,46 @@ void UK2Node_Message::ExpandNode(class FKismetCompilerContext& CompilerContext,
 					// when cast fails all return values must be cleared.
 					if (EEdGraphPinDirection::EGPD_Output == CurrentPin->Direction)
 					{
+						UEdGraphPin* VarOutPin = FunctionCallPin;
+						if (bIsPureFunc)
+						{
+							// since we cannot directly use the output from the
+							// function call node (since it is pure, and invoking
+							// it with a null target would cause an error), we 
+							// have to use a temporary variable in it's place...
+							UK2Node_TemporaryVariable* TempVar = CompilerContext.SpawnIntermediateNode<UK2Node_TemporaryVariable>(this, SourceGraph);
+							TempVar->VariableType = CurrentPin->PinType;
+							TempVar->AllocateDefaultPins();
+
+							VarOutPin = TempVar->GetVariablePin();
+							// nodes using the function's outputs directly, now
+							// use this TempVar node instead
+							CompilerContext.MovePinLinksToIntermediate(*FunctionCallPin, *VarOutPin);
+
+							// on a successful cast, the temp var is filled with
+							// the function's value, on a failed cast, the var 
+							// is filled with a default value (DefaultValueNode, 
+							// below)... this is the node for the success case:
+							UK2Node_AssignmentStatement* AssignTempVar = CompilerContext.SpawnIntermediateNode<UK2Node_AssignmentStatement>(this, SourceGraph);
+							AssignTempVar->AllocateDefaultPins();
+							
+							// assign the output from the pure function node to
+							// the TempVar (either way, this message node is 
+							// returning the TempVar's value, so on a successful 
+							// cast we want it to have the function's result)
+							UEdGraphPin* ValueInPin = AssignTempVar->GetValuePin();
+							Schema->TryCreateConnection(FunctionCallPin, ValueInPin);
+							AssignTempVar->PinConnectionListChanged(ValueInPin);
+
+							UEdGraphPin* VarInPin = AssignTempVar->GetVariablePin();
+							Schema->TryCreateConnection(VarOutPin, VarInPin);
+							AssignTempVar->PinConnectionListChanged(VarInPin);
+							// fold this AssignTempVar node into the cast's 
+							// success execution chain
+							Schema->TryCreateConnection(AssignTempVar->GetExecPin(), LastOutCastSuccessPin);
+							LastOutCastSuccessPin = AssignTempVar->GetThenPin();
+						}
+
 						UK2Node* DefaultValueNode = NULL;
 						UEdGraphPin* DefaultValueThenPin = NULL;
 						if (CurrentPin->PinType.bIsArray)
@@ -177,7 +250,7 @@ void UK2Node_Message::ExpandNode(class FKismetCompilerContext& CompilerContext,
 
 							UEdGraphPin* ArrayPin = ClearArray->GetTargetArrayPin();
 							check(ArrayPin);
-							Schema->TryCreateConnection(ArrayPin, FunctionCallPin);
+							Schema->TryCreateConnection(ArrayPin, VarOutPin);
 							ClearArray->PinConnectionListChanged(ArrayPin);
 
 							DefaultValueThenPin = ClearArray->GetThenPin();
@@ -188,12 +261,13 @@ void UK2Node_Message::ExpandNode(class FKismetCompilerContext& CompilerContext,
 							DefaultValueNode = AssignDefaultValue;
 							AssignDefaultValue->AllocateDefaultPins();
 
-							Schema->TryCreateConnection(AssignDefaultValue->GetVariablePin(), FunctionCallPin);
+							Schema->TryCreateConnection(AssignDefaultValue->GetVariablePin(), VarOutPin);
 							AssignDefaultValue->PinConnectionListChanged(AssignDefaultValue->GetVariablePin());
 							Schema->SetPinDefaultValueBasedOnType(AssignDefaultValue->GetValuePin());
 
 							DefaultValueThenPin = AssignDefaultValue->GetThenPin();
 						}
+
 						UEdGraphPin* DefaultValueExecPin = DefaultValueNode->GetExecPin();
 						check(DefaultValueExecPin);
 						Schema->TryCreateConnection(DefaultValueExecPin, LastOutCastFaildPin);
@@ -210,18 +284,13 @@ void UK2Node_Message::ExpandNode(class FKismetCompilerContext& CompilerContext,
 
 		if( bThenPinConnected )
 		{
+			check(LastOutCastFaildPin != nullptr);
 			// Failure case for the cast runs straight through to the exit
 			CompilerContext.CopyPinLinksToIntermediate(*ThenPin, *LastOutCastFaildPin);
 
+			check(LastOutCastSuccessPin != nullptr);
 			// Copy all links from the invalid cast case above to the call function node
-			if (UEdGraphPin* CallFunctionThenPin = Schema->FindExecutionPin(*FunctionCallNode, EGPD_Output))
-			{
-				CompilerContext.MovePinLinksToIntermediate(*ThenPin, *CallFunctionThenPin);
-			}
-			else
-			{
-				CompilerContext.MessageLog.Warning(*LOCTEXT("ExpectedThenPinOnInterfaceCall", "Expected a then pin on the call in @@.  Does the interface method still exist?").ToString(), FunctionCallNode);
-			}
+			CompilerContext.MovePinLinksToIntermediate(*ThenPin, *LastOutCastSuccessPin);
 		}
 	}
 
