When adding/removing/renaming action or axis mappings in the input properties, the Blueprint database will refresh required actions.

Added an Editor Delegate to handle callback into the BlueprintActionDatabase when Action and Axis mappings are changed.

diff --git a/Engine/Source/Editor/BlueprintGraph/Private/K2Node_GetInputAxisValue.cpp b/Engine/Source/Editor/BlueprintGraph/Private/K2Node_GetInputAxisValue.cpp
--- a/Engine/Source/Editor/BlueprintGraph/Private/K2Node_GetInputAxisValue.cpp
+++ b/Engine/Source/Editor/BlueprintGraph/Private/K2Node_GetInputAxisValue.cpp
@@ -117,22 +117,32 @@ void UK2Node_GetInputAxisValue::GetMenuActions(FBlueprintActionDatabaseRegistrar
 	// type disappears, then the action should go with it)
 	UClass* ActionKey = GetClass();
 
-	for (FName const InputAxisName : AxisNames)
+	// to keep from needlessly instantiating a UBlueprintNodeSpawner, first   
+	// check to make sure that the registrar is looking for actions of this type
+	// (could be regenerating actions for a specific asset, and therefore the 
+	// registrar would only accept actions corresponding to that asset)
+	if (ActionRegistrar.IsOpenForRegistration(ActionKey))
 	{
-		// to keep from needlessly instantiating a UBlueprintNodeSpawner, first   
-		// check to make sure that the registrar is looking for actions of this type
-		// (could be regenerating actions for a specific asset, and therefore the 
-		// registrar would only accept actions corresponding to that asset)
-		if (!ActionRegistrar.IsOpenForRegistration(ActionKey))
+		auto RefreshClassActions = []()
 		{
-			continue;
+			FBlueprintActionDatabase::Get().RefreshClassActions(StaticClass());
+		};
+
+		static bool bRegisterOnce = true;
+		if(bRegisterOnce)
+		{
+			bRegisterOnce = false;
+			FEditorDelegates::OnActionAxisMappingsChanged.AddStatic(RefreshClassActions);
 		}
 
-		UBlueprintNodeSpawner* NodeSpawner = UBlueprintNodeSpawner::Create(GetClass());
-		check(NodeSpawner != nullptr);
+		for (FName const InputAxisName : AxisNames)
+		{
+			UBlueprintNodeSpawner* NodeSpawner = UBlueprintNodeSpawner::Create(GetClass());
+			check(NodeSpawner != nullptr);
 
-		NodeSpawner->CustomizeNodeDelegate = UBlueprintNodeSpawner::FCustomizeNodeDelegate::CreateStatic(CustomizeInputNodeLambda, InputAxisName);
-		ActionRegistrar.AddBlueprintAction(ActionKey, NodeSpawner);
+			NodeSpawner->CustomizeNodeDelegate = UBlueprintNodeSpawner::FCustomizeNodeDelegate::CreateStatic(CustomizeInputNodeLambda, InputAxisName);
+			ActionRegistrar.AddBlueprintAction(ActionKey, NodeSpawner);
+		}
 	}
 }
 
diff --git a/Engine/Source/Editor/BlueprintGraph/Private/K2Node_InputAction.cpp b/Engine/Source/Editor/BlueprintGraph/Private/K2Node_InputAction.cpp
--- a/Engine/Source/Editor/BlueprintGraph/Private/K2Node_InputAction.cpp
+++ b/Engine/Source/Editor/BlueprintGraph/Private/K2Node_InputAction.cpp
@@ -158,22 +158,32 @@ void UK2Node_InputAction::GetMenuActions(FBlueprintActionDatabaseRegistrar& Acti
 	// type disappears, then the action should go with it)
 	UClass* ActionKey = GetClass();
 
-	for (FName const InputActionName : ActionNames)
+	// to keep from needlessly instantiating a UBlueprintNodeSpawner, first   
+	// check to make sure that the registrar is looking for actions of this type
+	// (could be regenerating actions for a specific asset, and therefore the 
+	// registrar would only accept actions corresponding to that asset)
+	if (ActionRegistrar.IsOpenForRegistration(ActionKey))
 	{
-		// to keep from needlessly instantiating a UBlueprintNodeSpawner, first   
-		// check to make sure that the registrar is looking for actions of this type
-		// (could be regenerating actions for a specific asset, and therefore the 
-		// registrar would only accept actions corresponding to that asset)
-		if (!ActionRegistrar.IsOpenForRegistration(ActionKey))
+		auto RefreshClassActions = []()
 		{
-			continue;
+			FBlueprintActionDatabase::Get().RefreshClassActions(StaticClass());
+		};
+
+		static bool bRegisterOnce = true;
+		if(bRegisterOnce)
+		{
+			bRegisterOnce = false;
+			FEditorDelegates::OnActionAxisMappingsChanged.AddStatic(RefreshClassActions);
 		}
 
-		UBlueprintNodeSpawner* NodeSpawner = UBlueprintNodeSpawner::Create(GetClass());
-		check(NodeSpawner != nullptr);
+		for (FName const InputActionName : ActionNames)
+		{
+			UBlueprintNodeSpawner* NodeSpawner = UBlueprintNodeSpawner::Create(GetClass());
+			check(NodeSpawner != nullptr);
 
-		NodeSpawner->CustomizeNodeDelegate = UBlueprintNodeSpawner::FCustomizeNodeDelegate::CreateStatic(CustomizeInputNodeLambda, InputActionName);
-		ActionRegistrar.AddBlueprintAction(ActionKey, NodeSpawner);
+			NodeSpawner->CustomizeNodeDelegate = UBlueprintNodeSpawner::FCustomizeNodeDelegate::CreateStatic(CustomizeInputNodeLambda, InputActionName);
+			ActionRegistrar.AddBlueprintAction(ActionKey, NodeSpawner);
+		}
 	}
 }
 
diff --git a/Engine/Source/Editor/BlueprintGraph/Private/K2Node_InputAxisEvent.cpp b/Engine/Source/Editor/BlueprintGraph/Private/K2Node_InputAxisEvent.cpp
--- a/Engine/Source/Editor/BlueprintGraph/Private/K2Node_InputAxisEvent.cpp
+++ b/Engine/Source/Editor/BlueprintGraph/Private/K2Node_InputAxisEvent.cpp
@@ -129,22 +129,32 @@ void UK2Node_InputAxisEvent::GetMenuActions(FBlueprintActionDatabaseRegistrar& A
 	// type disappears, then the action should go with it)
 	UClass* ActionKey = GetClass();
 
-	for (FName const InputAxisName : AxisNames)
+	// to keep from needlessly instantiating a UBlueprintNodeSpawner, first   
+	// check to make sure that the registrar is looking for actions of this type
+	// (could be regenerating actions for a specific asset, and therefore the 
+	// registrar would only accept actions corresponding to that asset)
+	if (ActionRegistrar.IsOpenForRegistration(ActionKey))
 	{
-		// to keep from needlessly instantiating a UBlueprintNodeSpawner, first   
-		// check to make sure that the registrar is looking for actions of this type
-		// (could be regenerating actions for a specific asset, and therefore the 
-		// registrar would only accept actions corresponding to that asset)
-		if (!ActionRegistrar.IsOpenForRegistration(ActionKey))
+		auto RefreshClassActions = []()
 		{
-			continue;
+			FBlueprintActionDatabase::Get().RefreshClassActions(StaticClass());
+		};
+
+		static bool bRegisterOnce = true;
+		if(bRegisterOnce)
+		{
+			bRegisterOnce = false;
+			FEditorDelegates::OnActionAxisMappingsChanged.AddStatic(RefreshClassActions);
 		}
 
-		UBlueprintNodeSpawner* NodeSpawner = UBlueprintNodeSpawner::Create(GetClass());
-		check(NodeSpawner != nullptr);
+		for (FName const InputAxisName : AxisNames)
+		{
+			UBlueprintNodeSpawner* NodeSpawner = UBlueprintNodeSpawner::Create(GetClass());
+			check(NodeSpawner != nullptr);
 
-		NodeSpawner->CustomizeNodeDelegate = UBlueprintNodeSpawner::FCustomizeNodeDelegate::CreateStatic(CustomizeInputNodeLambda, InputAxisName);
-		ActionRegistrar.AddBlueprintAction(ActionKey, NodeSpawner);
+			NodeSpawner->CustomizeNodeDelegate = UBlueprintNodeSpawner::FCustomizeNodeDelegate::CreateStatic(CustomizeInputNodeLambda, InputAxisName);
+			ActionRegistrar.AddBlueprintAction(ActionKey, NodeSpawner);
+		}
 	}
 }
 
diff --git a/Engine/Source/Editor/UnrealEd/Private/Editor.cpp b/Engine/Source/Editor/UnrealEd/Private/Editor.cpp
--- a/Engine/Source/Editor/UnrealEd/Private/Editor.cpp
+++ b/Engine/Source/Editor/UnrealEd/Private/Editor.cpp
@@ -142,6 +142,7 @@ FEditorDelegates::FOnBlueprintContextMenuCreated		FEditorDelegates::OnBlueprintC
 FSimpleMulticastDelegate								FEditorDelegates::OnShutdownPostPackagesSaved;
 FEditorDelegates::FOnAssetsPreDelete					FEditorDelegates::OnAssetsPreDelete;
 FEditorDelegates::FOnAssetsDeleted						FEditorDelegates::OnAssetsDeleted;
+FSimpleMulticastDelegate								FEditorDelegates::OnActionAxisMappingsChanged;
 
 /*-----------------------------------------------------------------------------
 	Globals.
diff --git a/Engine/Source/Editor/UnrealEd/Public/Editor.h b/Engine/Source/Editor/UnrealEd/Public/Editor.h
--- a/Engine/Source/Editor/UnrealEd/Public/Editor.h
+++ b/Engine/Source/Editor/UnrealEd/Public/Editor.h
@@ -185,6 +185,8 @@ struct UNREALED_API FEditorDelegates
 	static FOnAssetsPreDelete OnAssetsPreDelete;
 	/** Called when one or more assets have been deleted */
 	static FOnAssetsDeleted OnAssetsDeleted;
+	/** Called when Action or Axis mappings have been changed */
+	static FSimpleMulticastDelegate OnActionAxisMappingsChanged;
 };
 
 /**
diff --git a/Engine/Source/Runtime/Engine/Private/UserInterface/InputSettings.cpp b/Engine/Source/Runtime/Engine/Private/UserInterface/InputSettings.cpp
--- a/Engine/Source/Runtime/Engine/Private/UserInterface/InputSettings.cpp
+++ b/Engine/Source/Runtime/Engine/Private/UserInterface/InputSettings.cpp
@@ -6,6 +6,10 @@
 
 #include "EnginePrivate.h"
 
+#if WITH_EDITOR
+#include "UnrealEd.h"
+#endif
+
 UInputSettings::UInputSettings(const FObjectInitializer& ObjectInitializer)
 : Super(ObjectInitializer)
 {
@@ -47,6 +51,7 @@ void UInputSettings::PostEditChangeChainProperty(FPropertyChangedChainEvent& Pro
 	if (MemberPropertyName == "ActionMappings" || MemberPropertyName == "AxisMappings" || MemberPropertyName == "AxisConfig")
 	{
 		ForceRebuildKeymaps();
+		FEditorDelegates::OnActionAxisMappingsChanged.Broadcast();
 	}
 }
 
