[UE-2031] Providing a more verbose error message, explaining why class and object pins are not compatible (and using a GetObjectClass node to convert between the two when we can).

diff --git a/Engine/Source/Editor/BlueprintGraph/Private/EdGraphSchema_K2.cpp b/Engine/Source/Editor/BlueprintGraph/Private/EdGraphSchema_K2.cpp
--- a/Engine/Source/Editor/BlueprintGraph/Private/EdGraphSchema_K2.cpp
+++ b/Engine/Source/Editor/BlueprintGraph/Private/EdGraphSchema_K2.cpp
@@ -1603,6 +1603,63 @@ const FPinConnectionResponse UEdGraphSchema_K2::DetermineConnectionResponseOfCom
 	}
 }
 
+static FText GetPinIncompatibilityMessage(const UEdGraphPin* PinA, const UEdGraphPin* PinB)
+{
+	const FEdGraphPinType& PinAType = PinA->PinType;
+	const FEdGraphPinType& PinBType = PinB->PinType;
+
+	FFormatNamedArguments MessageArgs;
+	MessageArgs.Add(TEXT("PinAName"), PinA->GetDisplayName());
+	MessageArgs.Add(TEXT("PinBName"), PinB->GetDisplayName());
+	MessageArgs.Add(TEXT("PinAType"), UEdGraphSchema_K2::TypeToText(PinAType));
+	MessageArgs.Add(TEXT("PinBType"), UEdGraphSchema_K2::TypeToText(PinBType));
+
+	const UEdGraphPin* InputPin = (PinA->Direction == EGPD_Input) ? PinA : PinB;
+	const FEdGraphPinType& InputType  = InputPin->PinType;
+	const UEdGraphPin* OutputPin = (InputPin == PinA) ? PinB : PinA;
+	const FEdGraphPinType& OutputType = OutputPin->PinType;
+
+	FText MessageFormat = LOCTEXT("DefaultPinIncompatibilityMessage", "{PinAType} is not compatible with {PinBType}.");
+
+	if (OutputType.PinCategory == UEdGraphSchema_K2::PC_Struct)
+	{
+		if (InputType.PinCategory == UEdGraphSchema_K2::PC_Struct)
+		{
+			MessageFormat = LOCTEXT("StructsIncompatible", "Only exactly matching structures are considered compatible.");
+
+			const UStruct* OutStruct = CastChecked<const UStruct>(OutputType.PinSubCategoryObject.Get());
+			const UStruct* InStruct  = CastChecked<const UStruct>(InputType.PinSubCategoryObject.Get());
+			if ((OutStruct != nullptr) && (InStruct != nullptr) && OutStruct->IsChildOf(InStruct))
+			{
+				MessageFormat = LOCTEXT("ChildStructIncompatible", "Only exactly matching structures are considered compatible. Derived structures are disallowed.");
+			}
+		}
+	}
+	else if (OutputType.PinCategory == UEdGraphSchema_K2::PC_Class)
+	{
+		if ((InputType.PinCategory == UEdGraphSchema_K2::PC_Object) || 
+			(InputType.PinCategory == UEdGraphSchema_K2::PC_Interface))
+		{
+			MessageArgs.Add(TEXT("OutputName"), OutputPin->GetDisplayName());
+			MessageArgs.Add(TEXT("InputName"),  InputPin->GetDisplayName());
+			MessageFormat = LOCTEXT("ClassObjectIncompatible", "'{PinAName}' and '{PinBName}' are incompatible ('{OutputName}' is an object type, and '{InputName}' is a reference to an object instance).");
+		}
+	}
+	else if ((OutputType.PinCategory == UEdGraphSchema_K2::PC_Object) )//|| (OutputType.PinCategory == UEdGraphSchema_K2::PC_Interface))
+	{
+		if (InputType.PinCategory == UEdGraphSchema_K2::PC_Class)
+		{
+			MessageArgs.Add(TEXT("OutputName"), OutputPin->GetDisplayName());
+			MessageArgs.Add(TEXT("InputName"),  InputPin->GetDisplayName());
+			MessageArgs.Add(TEXT("InputType"),  UEdGraphSchema_K2::TypeToText(InputType));
+
+			MessageFormat = LOCTEXT("CannotGetClass", "'{PinAName}' and '{PinBName}' are not inherently compatible ('{InputName}' is an object type, and '{OutputName}' is a reference to an object instance).\nWe cannot use {OutputName}'s class because it is not a child of {InputType}.");
+		}
+	}
+
+	return FText::Format(MessageFormat, MessageArgs);
+}
+
 const FPinConnectionResponse UEdGraphSchema_K2::CanCreateConnection(const UEdGraphPin* PinA, const UEdGraphPin* PinB) const
 {
 	const UK2Node* OwningNodeA = Cast<UK2Node>(PinA->GetOwningNodeUnchecked());
@@ -1680,19 +1737,8 @@ const FPinConnectionResponse UEdGraphSchema_K2::CanCreateConnection(const UEdGra
 		}
 		else
 		{
-			if (OutputPin && InputPin && (PC_Struct == OutputPin->PinType.PinCategory) && (PC_Struct == InputPin->PinType.PinCategory)
-				&& (OutputPin->PinType.PinSubCategoryObject != InputPin->PinType.PinSubCategoryObject))
-			{
-				FString Msg(TEXT("Only exactly matching structures are considered compatible."));
-				auto InStruct = Cast<const UStruct>(InputPin->PinType.PinSubCategoryObject.Get());
-				auto OutStruct = Cast<const UStruct>(OutputPin->PinType.PinSubCategoryObject.Get());
-				if (InStruct && OutStruct && OutStruct->IsChildOf(InStruct))
-				{
-					Msg += TEXT(" Derived structures are disallowed.");
-				}
-				return FPinConnectionResponse(CONNECT_RESPONSE_DISALLOW, Msg);
-			}
-			return FPinConnectionResponse(CONNECT_RESPONSE_DISALLOW, FString::Printf(TEXT("%s is not compatible with %s"), *TypeToText(PinA->PinType).ToString(), *TypeToText(PinB->PinType).ToString()));
+			FText IncompatibilityReasonText = GetPinIncompatibilityMessage(PinA, PinB);
+			return FPinConnectionResponse(CONNECT_RESPONSE_DISALLOW, IncompatibilityReasonText.ToString());
 		}
 	}
 }
@@ -1927,6 +1973,20 @@ bool UEdGraphSchema_K2::SearchForAutocastFunction(const UEdGraphPin* OutputPin,
 			TargetFunction = TEXT("Conv_InterfaceToObject");
 		}
 	}
+	else if (OutputPin->PinType.PinCategory == PC_Object)
+	{
+		UClass const* OutputClass = Cast<UClass const>(OutputPin->PinType.PinSubCategoryObject.Get());
+		if (InputPin->PinType.PinCategory == PC_Class)
+		{
+			UClass const* InputClass = Cast<UClass const>(InputPin->PinType.PinSubCategoryObject.Get());
+			if ((OutputClass != nullptr) &&
+				(InputClass != nullptr) &&
+				OutputClass->IsChildOf(InputClass))
+			{
+				TargetFunction = TEXT("GetObjectClass");
+			}
+		}
+	}
 
 	return TargetFunction != NAME_None;
 }
diff --git a/Engine/Source/Editor/BlueprintGraph/Private/K2Node_ClassDynamicCast.cpp b/Engine/Source/Editor/BlueprintGraph/Private/K2Node_ClassDynamicCast.cpp
--- a/Engine/Source/Editor/BlueprintGraph/Private/K2Node_ClassDynamicCast.cpp
+++ b/Engine/Source/Editor/BlueprintGraph/Private/K2Node_ClassDynamicCast.cpp
@@ -50,7 +50,7 @@ void UK2Node_ClassDynamicCast::AllocateDefaultPins()
 	// Output - Data Pin
 	if (TargetType != NULL)
 	{
-		const FString CastResultPinName = K2Schema->PN_CastedValuePrefix + TargetType->GetName();
+		const FString CastResultPinName = K2Schema->PN_CastedValuePrefix + TargetType->GetDisplayNameText().ToString();
 		CreatePin(EGPD_Output, K2Schema->PC_Class, TEXT(""), *TargetType, false, false, CastResultPinName);
 	}
 
