#ue4 - SpawnActorFromClass uses the original tooltips for properties marked "ExposeOnSpawn".

diff --git a/Engine/Source/Editor/BlueprintGraph/Private/K2Node_SpawnActorFromClass.cpp b/Engine/Source/Editor/BlueprintGraph/Private/K2Node_SpawnActorFromClass.cpp
--- a/Engine/Source/Editor/BlueprintGraph/Private/K2Node_SpawnActorFromClass.cpp
+++ b/Engine/Source/Editor/BlueprintGraph/Private/K2Node_SpawnActorFromClass.cpp
@@ -94,6 +94,9 @@ void UK2Node_SpawnActorFromClass::CreatePinsForClass(UClass* InClass)
 				const bool bDefaultValueSet = FBlueprintEditorUtils::PropertyValueToString(Property, reinterpret_cast<const uint8*>(ClassDefaultObject), DefaultValueAsString);
 				check( bDefaultValueSet );
 				K2Schema->TrySetDefaultValue(*Pin, DefaultValueAsString);
+
+				// Copy tooltip from the property.
+				K2Schema->ConstructBasicPinTooltip(*Pin, Property->GetToolTipText(), Pin->PinToolTip);
 			}
 		}
 	}
