Blueprint - Changing a couple of CastChecked to Casts to prevent crashes on 4.5 to 4.6 UMG blueprint upgrades.

diff --git a/Engine/Source/Editor/BlueprintGraph/Private/EdGraphSchema_K2.cpp b/Engine/Source/Editor/BlueprintGraph/Private/EdGraphSchema_K2.cpp
--- a/Engine/Source/Editor/BlueprintGraph/Private/EdGraphSchema_K2.cpp
+++ b/Engine/Source/Editor/BlueprintGraph/Private/EdGraphSchema_K2.cpp
@@ -1627,8 +1627,8 @@ static FText GetPinIncompatibilityMessage(const UEdGraphPin* PinA, const UEdGrap
 		{
 			MessageFormat = LOCTEXT("StructsIncompatible", "Only exactly matching structures are considered compatible.");
 
-			const UStruct* OutStruct = CastChecked<const UStruct>(OutputType.PinSubCategoryObject.Get());
-			const UStruct* InStruct  = CastChecked<const UStruct>(InputType.PinSubCategoryObject.Get());
+			const UStruct* OutStruct = Cast<const UStruct>(OutputType.PinSubCategoryObject.Get());
+			const UStruct* InStruct  = Cast<const UStruct>(InputType.PinSubCategoryObject.Get());
 			if ((OutStruct != nullptr) && (InStruct != nullptr) && OutStruct->IsChildOf(InStruct))
 			{
 				MessageFormat = LOCTEXT("ChildStructIncompatible", "Only exactly matching structures are considered compatible. Derived structures are disallowed.");
