#ue4 - Fix tooltips for expose on spawn properties only working for some types of classes.

diff --git a/Engine/Source/Editor/BlueprintGraph/Private/K2Node_AddComponent.cpp b/Engine/Source/Editor/BlueprintGraph/Private/K2Node_AddComponent.cpp
--- a/Engine/Source/Editor/BlueprintGraph/Private/K2Node_AddComponent.cpp
+++ b/Engine/Source/Editor/BlueprintGraph/Private/K2Node_AddComponent.cpp
@@ -71,6 +71,9 @@ void UK2Node_AddComponent::AllocatePinsForExposedVariables()
 						check(bDefaultValueSet);
 						K2Schema->TrySetDefaultValue(*Pin, DefaultValueAsString);
 					}
+
+					// Copy tooltip from the property.
+					K2Schema->ConstructBasicPinTooltip(*Pin, Property->GetToolTipText(), Pin->PinToolTip);
 				}
 			}
 		}
diff --git a/Engine/Source/Editor/BlueprintGraph/Private/K2Node_ConstructObjectFromClass.cpp b/Engine/Source/Editor/BlueprintGraph/Private/K2Node_ConstructObjectFromClass.cpp
--- a/Engine/Source/Editor/BlueprintGraph/Private/K2Node_ConstructObjectFromClass.cpp
+++ b/Engine/Source/Editor/BlueprintGraph/Private/K2Node_ConstructObjectFromClass.cpp
@@ -101,6 +101,12 @@ void UK2Node_ConstructObjectFromClass::CreatePinsForClass(UClass* InClass)
 				check( bDefaultValueSet );
 				K2Schema->TrySetDefaultValue(*Pin, DefaultValueAsString);
 			}
+
+			// Copy tooltip from the property.
+			if (Pin != nullptr)
+			{
+				K2Schema->ConstructBasicPinTooltip(*Pin, Property->GetToolTipText(), Pin->PinToolTip);
+			}
 		}
 	}
 
diff --git a/Engine/Source/Editor/BlueprintGraph/Private/K2Node_SpawnActor.cpp b/Engine/Source/Editor/BlueprintGraph/Private/K2Node_SpawnActor.cpp
--- a/Engine/Source/Editor/BlueprintGraph/Private/K2Node_SpawnActor.cpp
+++ b/Engine/Source/Editor/BlueprintGraph/Private/K2Node_SpawnActor.cpp
@@ -73,6 +73,12 @@ void UK2Node_SpawnActor::CreatePinsForClass(UClass* InClass)
 		{
 			UEdGraphPin* Pin = CreatePin(EGPD_Input, TEXT(""), TEXT(""), NULL, false, false, Property->GetName());
 			const bool bPinGood = (Pin != NULL) && K2Schema->ConvertPropertyToPinType(Property, /*out*/ Pin->PinType);	
+
+			// Copy tooltip from the property.
+			if (Pin != nullptr)
+			{
+				K2Schema->ConstructBasicPinTooltip(*Pin, Property->GetToolTipText(), Pin->PinToolTip);
+			}
 		}
 	}
 
diff --git a/Engine/Source/Editor/BlueprintGraph/Private/K2Node_SpawnActorFromClass.cpp b/Engine/Source/Editor/BlueprintGraph/Private/K2Node_SpawnActorFromClass.cpp
--- a/Engine/Source/Editor/BlueprintGraph/Private/K2Node_SpawnActorFromClass.cpp
+++ b/Engine/Source/Editor/BlueprintGraph/Private/K2Node_SpawnActorFromClass.cpp
@@ -94,8 +94,11 @@ void UK2Node_SpawnActorFromClass::CreatePinsForClass(UClass* InClass)
 				const bool bDefaultValueSet = FBlueprintEditorUtils::PropertyValueToString(Property, reinterpret_cast<const uint8*>(ClassDefaultObject), DefaultValueAsString);
 				check( bDefaultValueSet );
 				K2Schema->TrySetDefaultValue(*Pin, DefaultValueAsString);
+			}
 
-				// Copy tooltip from the property.
+			// Copy tooltip from the property.
+			if (Pin != nullptr)
+			{
 				K2Schema->ConstructBasicPinTooltip(*Pin, Property->GetToolTipText(), Pin->PinToolTip);
 			}
 		}
diff --git a/Engine/Source/Editor/GameplayAbilitiesEditor/Private/K2Node_LatentAbilityCall.cpp b/Engine/Source/Editor/GameplayAbilitiesEditor/Private/K2Node_LatentAbilityCall.cpp
--- a/Engine/Source/Editor/GameplayAbilitiesEditor/Private/K2Node_LatentAbilityCall.cpp
+++ b/Engine/Source/Editor/GameplayAbilitiesEditor/Private/K2Node_LatentAbilityCall.cpp
@@ -269,7 +269,10 @@ void UK2Node_LatentAbilityCall::CreatePinsForClass(UClass* InClass)
 			}
 
 			// Copy tooltip from the property.
-			K2Schema->ConstructBasicPinTooltip(*Pin, Property->GetToolTipText(), Pin->PinToolTip);
+			if (Pin != nullptr)
+			{
+				K2Schema->ConstructBasicPinTooltip(*Pin, Property->GetToolTipText(), Pin->PinToolTip);
+			}
 		}
 	}
 }
