Assign nodes should only be placeable in macro graphs.

diff --git a/Engine/Source/Editor/BlueprintGraph/Private/K2Node_AssignmentStatement.cpp b/Engine/Source/Editor/BlueprintGraph/Private/K2Node_AssignmentStatement.cpp
--- a/Engine/Source/Editor/BlueprintGraph/Private/K2Node_AssignmentStatement.cpp
+++ b/Engine/Source/Editor/BlueprintGraph/Private/K2Node_AssignmentStatement.cpp
@@ -91,7 +91,7 @@ bool UK2Node_AssignmentStatement::IsCompatibleWithGraph(UEdGraph const* TargetGr
 	if (bIsCompatible)
 	{
 		EGraphType const GraphType = TargetGraph->GetSchema()->GetGraphType(TargetGraph);
-		bIsCompatible = (GraphType != GT_Ubergraph) && (GraphType != GT_Animation);
+		bIsCompatible = GraphType == GT_Macro;
 	}
 
 	return bIsCompatible;
