UE-4475 Uasset causes editor to crash when opened/right clicked

diff --git a/Engine/Source/Editor/BlueprintGraph/Classes/K2Node_CastByteToEnum.h b/Engine/Source/Editor/BlueprintGraph/Classes/K2Node_CastByteToEnum.h
--- a/Engine/Source/Editor/BlueprintGraph/Classes/K2Node_CastByteToEnum.h
+++ b/Engine/Source/Editor/BlueprintGraph/Classes/K2Node_CastByteToEnum.h
@@ -24,6 +24,7 @@ class UK2Node_CastByteToEnum : public UK2Node
 	virtual FText GetTooltipText() const override;
 	virtual FText GetNodeTitle(ENodeTitleType::Type TitleType) const override;
 	virtual FName GetPaletteIcon(FLinearColor& OutColor) const override{ return TEXT("GraphEditor.Enum_16x"); }
+	virtual void ValidateNodeDuringCompilation(class FCompilerResultsLog& MessageLog) const override;
 	// End UEdGraphNode interface
 
 	// Begin UK2Node interface
diff --git a/Engine/Source/Editor/BlueprintGraph/Classes/K2Node_EnumLiteral.h b/Engine/Source/Editor/BlueprintGraph/Classes/K2Node_EnumLiteral.h
--- a/Engine/Source/Editor/BlueprintGraph/Classes/K2Node_EnumLiteral.h
+++ b/Engine/Source/Editor/BlueprintGraph/Classes/K2Node_EnumLiteral.h
@@ -21,6 +21,7 @@ class UK2Node_EnumLiteral : public UK2Node, public INodeDependingOnEnumInterface
 	virtual FText GetTooltipText() const override;
 	virtual FText GetNodeTitle(ENodeTitleType::Type TitleType) const override;
 	virtual FName GetPaletteIcon(FLinearColor& OutColor) const override;
+	virtual void ValidateNodeDuringCompilation(class FCompilerResultsLog& MessageLog) const override;
 	// End UEdGraphNode interface
 
 	// Begin UK2Node interface
diff --git a/Engine/Source/Editor/BlueprintGraph/Private/K2Node_CastByteToEnum.cpp b/Engine/Source/Editor/BlueprintGraph/Private/K2Node_CastByteToEnum.cpp
--- a/Engine/Source/Editor/BlueprintGraph/Private/K2Node_CastByteToEnum.cpp
+++ b/Engine/Source/Editor/BlueprintGraph/Private/K2Node_CastByteToEnum.cpp
@@ -16,10 +16,17 @@ UK2Node_CastByteToEnum::UK2Node_CastByteToEnum(const FObjectInitializer& ObjectI
 {
 }
 
-void UK2Node_CastByteToEnum::AllocateDefaultPins()
+void UK2Node_CastByteToEnum::ValidateNodeDuringCompilation(class FCompilerResultsLog& MessageLog) const
 {
-	check(Enum);
+	Super::ValidateNodeDuringCompilation(MessageLog);
+	if (!Enum)
+	{
+		MessageLog.Error(*FString::Printf(*NSLOCTEXT("K2Node", "CastByteToNullEnumError", "Undefined Enum in @@").ToString()), this);
+	}
+}
 
+void UK2Node_CastByteToEnum::AllocateDefaultPins()
+{
 	const UEdGraphSchema_K2* Schema = GetDefault<UEdGraphSchema_K2>();
 	
 	CreatePin(EGPD_Input, Schema->PC_Byte, TEXT(""), NULL, false, false, ByteInputPinName);
@@ -28,7 +35,11 @@ void UK2Node_CastByteToEnum::AllocateDefaultPins()
 
 FText UK2Node_CastByteToEnum::GetTooltipText() const
 {
-	if (CachedTooltip.IsOutOfDate())
+	if (Enum == nullptr)
+	{
+		return NSLOCTEXT("K2Node", "CastByteToEnum_NullTooltip", "Byte to Enum (bad enum)");
+	}
+	else if(CachedTooltip.IsOutOfDate())
 	{
 		// FText::Format() is slow, so we cache this to save on performance
 		CachedTooltip = FText::Format(
@@ -59,7 +70,7 @@ void UK2Node_CastByteToEnum::ExpandNode(class FKismetCompilerContext& CompilerCo
 {
 	Super::ExpandNode(CompilerContext, SourceGraph);
 
-	if (bSafe)
+	if (bSafe && Enum)
 	{
 		const UEdGraphSchema_K2* Schema = CompilerContext.GetSchema();
 
diff --git a/Engine/Source/Editor/BlueprintGraph/Private/K2Node_EnumLiteral.cpp b/Engine/Source/Editor/BlueprintGraph/Private/K2Node_EnumLiteral.cpp
--- a/Engine/Source/Editor/BlueprintGraph/Private/K2Node_EnumLiteral.cpp
+++ b/Engine/Source/Editor/BlueprintGraph/Private/K2Node_EnumLiteral.cpp
@@ -19,10 +19,17 @@ UK2Node_EnumLiteral::UK2Node_EnumLiteral(const FObjectInitializer& ObjectInitial
 {
 }
 
-void UK2Node_EnumLiteral::AllocateDefaultPins()
+void UK2Node_EnumLiteral::ValidateNodeDuringCompilation(class FCompilerResultsLog& MessageLog) const
 {
-	check(Enum);
+	Super::ValidateNodeDuringCompilation(MessageLog);
+	if (!Enum)
+	{
+		MessageLog.Error(*FString::Printf(*NSLOCTEXT("K2Node", "EnumLiteral_NullEnumError", "Undefined Enum in @@").ToString()), this);
+	}
+}
 
+void UK2Node_EnumLiteral::AllocateDefaultPins()
+{
 	const UEdGraphSchema_K2* Schema = GetDefault<UEdGraphSchema_K2>();
 
 	auto InputPin = CreatePin(EGPD_Input, Schema->PC_Byte, TEXT(""), Enum, false, false, GetEnumInputPinName());
