Blueprint/UMG - When modified the ComponentBoundEvent K2Node now invalidates its cached text.  This now causes graphs to show the proper text for components and widgets that have been renamed.

diff --git a/Engine/Source/Editor/BlueprintGraph/Classes/K2Node_ComponentBoundEvent.h b/Engine/Source/Editor/BlueprintGraph/Classes/K2Node_ComponentBoundEvent.h
--- a/Engine/Source/Editor/BlueprintGraph/Classes/K2Node_ComponentBoundEvent.h
+++ b/Engine/Source/Editor/BlueprintGraph/Classes/K2Node_ComponentBoundEvent.h
@@ -19,6 +19,10 @@ class UK2Node_ComponentBoundEvent : public UK2Node_Event
 	UPROPERTY()
 	FName ComponentPropertyName;
 
+	// Begin UObject interface
+	virtual bool Modify(bool bAlwaysMarkDirty = true) override;
+	// End UObject interface
+
 	// Begin UEdGraphNode interface
 	virtual FText GetNodeTitle(ENodeTitleType::Type TitleType) const override;
 	virtual FText GetTooltipText() const override;
diff --git a/Engine/Source/Editor/BlueprintGraph/Private/K2Node_ComponentBoundEvent.cpp b/Engine/Source/Editor/BlueprintGraph/Private/K2Node_ComponentBoundEvent.cpp
--- a/Engine/Source/Editor/BlueprintGraph/Private/K2Node_ComponentBoundEvent.cpp
+++ b/Engine/Source/Editor/BlueprintGraph/Private/K2Node_ComponentBoundEvent.cpp
@@ -10,6 +10,13 @@ UK2Node_ComponentBoundEvent::UK2Node_ComponentBoundEvent(const FObjectInitialize
 {
 }
 
+bool UK2Node_ComponentBoundEvent::Modify(bool bAlwaysMarkDirty)
+{
+	CachedNodeTitle.MarkDirty();
+
+	return Super::Modify(bAlwaysMarkDirty);
+}
+
 FText UK2Node_ComponentBoundEvent::GetNodeTitle(ENodeTitleType::Type TitleType) const
 {
 	if (CachedNodeTitle.IsOutOfDate())
