Prepare cpp files for header cleanup.

diff --git a/Engine/Source/Editor/BlueprintGraph/Classes/K2Node_ClearDelegate.h b/Engine/Source/Editor/BlueprintGraph/Classes/K2Node_ClearDelegate.h
--- a/Engine/Source/Editor/BlueprintGraph/Classes/K2Node_ClearDelegate.h
+++ b/Engine/Source/Editor/BlueprintGraph/Classes/K2Node_ClearDelegate.h
@@ -1,6 +1,7 @@
 // Copyright 1998-2014 Epic Games, Inc. All Rights Reserved.
 #pragma once
 
+#include "K2Node_BaseMCDelegate.h"
 #include "K2Node_ClearDelegate.generated.h"
 
 UCLASS(MinimalAPI)
diff --git a/Engine/Source/Editor/BlueprintGraph/Private/BlueprintActionDatabase.cpp b/Engine/Source/Editor/BlueprintGraph/Private/BlueprintActionDatabase.cpp
--- a/Engine/Source/Editor/BlueprintGraph/Private/BlueprintActionDatabase.cpp
+++ b/Engine/Source/Editor/BlueprintGraph/Private/BlueprintActionDatabase.cpp
@@ -41,6 +41,7 @@
 // used below in BlueprintActionDatabaseImpl::GetNodeSpectificActions()
 #include "EdGraphNode_Comment.h"
 #include "EdGraph/EdGraphNode_Documentation.h"
+#include "K2Node_FunctionEntry.h"
 
 #define LOCTEXT_NAMESPACE "BlueprintActionDatabase"
 
diff --git a/Engine/Source/Editor/BlueprintGraph/Private/BlueprintActionFilter.cpp b/Engine/Source/Editor/BlueprintGraph/Private/BlueprintActionFilter.cpp
--- a/Engine/Source/Editor/BlueprintGraph/Private/BlueprintActionFilter.cpp
+++ b/Engine/Source/Editor/BlueprintGraph/Private/BlueprintActionFilter.cpp
@@ -20,6 +20,7 @@
 #include "K2Node_Message.h"
 #include "K2Node_ExecutionSequence.h"
 #include "K2Node_DynamicCast.h"
+#include "K2Node_BaseAsyncTask.h"
 
 /*******************************************************************************
  * Static BlueprintActionFilter Helpers
diff --git a/Engine/Source/Editor/BlueprintGraph/Private/BlueprintBoundEventNodeSpawner.cpp b/Engine/Source/Editor/BlueprintGraph/Private/BlueprintBoundEventNodeSpawner.cpp
--- a/Engine/Source/Editor/BlueprintGraph/Private/BlueprintBoundEventNodeSpawner.cpp
+++ b/Engine/Source/Editor/BlueprintGraph/Private/BlueprintBoundEventNodeSpawner.cpp
@@ -6,6 +6,9 @@
 #include "KismetEditorUtilities.h"	// for FindBoundEventForComponent()/FindBoundEventForActor()
 #include "EditorCategoryUtils.h"	// for GetCommonCategory()
 #include "BlueprintNodeSpawnerUtils.h" // for GetBindingClass()
+#include "K2Node_ComponentBoundEvent.h"
+#include "K2Node_ActorBoundEvent.h"
+#include "EditorStyleSettings.h"
 
 #define LOCTEXT_NAMESPACE "BlueprintBoundEventNodeSpawner"
 
diff --git a/Engine/Source/Editor/BlueprintGraph/Private/BlueprintFunctionNodeSpawner.cpp b/Engine/Source/Editor/BlueprintGraph/Private/BlueprintFunctionNodeSpawner.cpp
--- a/Engine/Source/Editor/BlueprintGraph/Private/BlueprintFunctionNodeSpawner.cpp
+++ b/Engine/Source/Editor/BlueprintGraph/Private/BlueprintFunctionNodeSpawner.cpp
@@ -10,6 +10,13 @@
 #include "EditorCategoryUtils.h"	// for BuildCategoryString()
 #include "ObjectEditorUtils.h"		// for IsFunctionHiddenFromClass()
 #include "BlueprintNodeSpawnerUtils.h" // for GetBindingClass()
+#include "K2Node_Literal.h"
+#include "K2Node_CallFunctionOnMember.h"
+#include "K2Node_CommutativeAssociativeBinaryOperator.h"
+#include "K2Node_CallMaterialParameterCollectionFunction.h"
+#include "K2Node_CallDataTableFunction.h"
+#include "K2Node_CallArrayFunction.h"
+#include "K2Node_VariableGet.h"
 
 #define LOCTEXT_NAMESPACE "BlueprintFunctionNodeSpawner"
 
diff --git a/Engine/Source/Editor/BlueprintGraph/Private/BlueprintVariableNodeSpawner.cpp b/Engine/Source/Editor/BlueprintGraph/Private/BlueprintVariableNodeSpawner.cpp
--- a/Engine/Source/Editor/BlueprintGraph/Private/BlueprintVariableNodeSpawner.cpp
+++ b/Engine/Source/Editor/BlueprintGraph/Private/BlueprintVariableNodeSpawner.cpp
@@ -9,6 +9,8 @@
 #include "EdGraphSchema_K2.h"		// for ConvertPropertyToPinType()
 #include "EditorCategoryUtils.h"	// for BuildCategoryString()
 #include "BlueprintActionFilter.h"	// for FBlueprintActionContext
+#include "K2Node_VariableGet.h"
+#include "K2Node_VariableSet.h"
 
 #define LOCTEXT_NAMESPACE "BlueprintVariableNodeSpawner"
 
diff --git a/Engine/Source/Editor/BlueprintGraph/Private/DelegateNodeHandlers.cpp b/Engine/Source/Editor/BlueprintGraph/Private/DelegateNodeHandlers.cpp
--- a/Engine/Source/Editor/BlueprintGraph/Private/DelegateNodeHandlers.cpp
+++ b/Engine/Source/Editor/BlueprintGraph/Private/DelegateNodeHandlers.cpp
@@ -2,6 +2,7 @@
 
 #include "BlueprintGraphPrivatePCH.h"
 #include "DelegateNodeHandlers.h"
+#include "K2Node_ClearDelegate.h"
 
 #define LOCTEXT_NAMESPACE "DelegateNodeHandlers"
 
diff --git a/Engine/Source/Editor/BlueprintGraph/Private/EdGraphSchema_K2_Actions.cpp b/Engine/Source/Editor/BlueprintGraph/Private/EdGraphSchema_K2_Actions.cpp
--- a/Engine/Source/Editor/BlueprintGraph/Private/EdGraphSchema_K2_Actions.cpp
+++ b/Engine/Source/Editor/BlueprintGraph/Private/EdGraphSchema_K2_Actions.cpp
@@ -8,6 +8,12 @@
 #include "Kismet2/KismetEditorUtilities.h"
 #include "Editor/UnrealEd/Public/EdGraphUtilities.h"
 #include "EdGraph/EdGraphNode_Documentation.h"
+#include "EdGraphNode_Comment.h"
+#include "K2Node_Timeline.h"
+#include "K2Node_Literal.h"
+#include "K2Node_VariableGet.h"
+#include "K2Node_AddDelegate.h"
+#include "K2Node_CustomEvent.h"
 
 
 #define SNAP_GRID (16) // @todo ensure this is the same as SNodePanel::GetSnapGridSize()
diff --git a/Engine/Source/Editor/BlueprintGraph/Private/K2ActionMenuBuilder.cpp b/Engine/Source/Editor/BlueprintGraph/Private/K2ActionMenuBuilder.cpp
--- a/Engine/Source/Editor/BlueprintGraph/Private/K2ActionMenuBuilder.cpp
+++ b/Engine/Source/Editor/BlueprintGraph/Private/K2ActionMenuBuilder.cpp
@@ -27,6 +27,47 @@
 #include "K2Node_SetFieldsInStruct.h"
 
 #include "EdGraphSchema_K2_Actions.h"
+#include "K2Node_ClearDelegate.h"
+#include "K2Node_SpawnActorFromClass.h"
+#include "K2Node_InputAxisEvent.h"
+#include "K2Node_GetDataTableRow.h"
+#include "K2Node_InputAction.h"
+#include "K2Node_InputKey.h"
+#include "K2Node_InputTouch.h"
+#include "K2Node_MacroInstance.h"
+#include "K2Node_CreateDelegate.h"
+#include "K2Node_VariableSet.h"
+#include "K2Node_CommutativeAssociativeBinaryOperator.h"
+#include "K2Node_CallMaterialParameterCollectionFunction.h"
+#include "K2Node_CallDataTableFunction.h"
+#include "K2Node_CallFunctionOnMember.h"
+#include "K2Node_CallParentFunction.h"
+#include "K2Node_CallFunction.h"
+#include "K2Node_CallArrayFunction.h"
+#include "K2Node_IfThenElse.h"
+#include "K2Node_MakeArray.h"
+#include "K2Node_FormatText.h"
+#include "K2Node_EaseFunction.h"
+#include "K2Node_AssignmentStatement.h"
+#include "K2Node_Composite.h"
+#include "K2Node_Select.h"
+#include "K2Node_BreakStruct.h"
+#include "K2Node_SwitchInteger.h"
+#include "K2Node_SwitchName.h"
+#include "K2Node_VariableGet.h"
+#include "K2Node_FunctionEntry.h"
+#include "K2Node_Self.h"
+#include "K2Node_AddDelegate.h"
+#include "K2Node_RemoveDelegate.h"
+#include "K2Node_CallDelegate.h"
+#include "K2Node_Timeline.h"
+#include "K2Node_CustomEvent.h"
+#include "K2Node_Literal.h"
+#include "K2Node_ComponentBoundEvent.h"
+#include "K2Node_ActorBoundEvent.h"
+#include "K2Node_MatineeController.h"
+#include "K2Node_TemporaryVariable.h"
+#include "EditorStyleSettings.h"
 
 #define LOCTEXT_NAMESPACE "KismetSchema"
 
diff --git a/Engine/Source/Editor/BlueprintGraph/Private/K2Node.cpp b/Engine/Source/Editor/BlueprintGraph/Private/K2Node.cpp
--- a/Engine/Source/Editor/BlueprintGraph/Private/K2Node.cpp
+++ b/Engine/Source/Editor/BlueprintGraph/Private/K2Node.cpp
@@ -7,6 +7,7 @@
 #include "K2Node.h"
 #include "KismetDebugUtilities.h" // for HasDebuggingData(), GetWatchText()
 #include "KismetCompiler.h"
+#include "GraphEditorSettings.h"
 
 #define LOCTEXT_NAMESPACE "K2Node"
 
diff --git a/Engine/Source/Editor/BlueprintGraph/Private/K2Node_CallArrayFunction.cpp b/Engine/Source/Editor/BlueprintGraph/Private/K2Node_CallArrayFunction.cpp
--- a/Engine/Source/Editor/BlueprintGraph/Private/K2Node_CallArrayFunction.cpp
+++ b/Engine/Source/Editor/BlueprintGraph/Private/K2Node_CallArrayFunction.cpp
@@ -1,6 +1,7 @@
 // Copyright 1998-2014 Epic Games, Inc. All Rights Reserved.
 
 #include "BlueprintGraphPrivatePCH.h"
+#include "K2Node_CallArrayFunction.h"
 
 UK2Node_CallArrayFunction::UK2Node_CallArrayFunction(const FObjectInitializer& ObjectInitializer)
 	: Super(ObjectInitializer)
diff --git a/Engine/Source/Editor/BlueprintGraph/Private/K2Node_CallDataTableFunction.cpp b/Engine/Source/Editor/BlueprintGraph/Private/K2Node_CallDataTableFunction.cpp
--- a/Engine/Source/Editor/BlueprintGraph/Private/K2Node_CallDataTableFunction.cpp
+++ b/Engine/Source/Editor/BlueprintGraph/Private/K2Node_CallDataTableFunction.cpp
@@ -3,6 +3,7 @@
 #include "BlueprintGraphPrivatePCH.h"
 
 #include "CompilerResultsLog.h"
+#include "K2Node_CallDataTableFunction.h"
 
 UK2Node_CallDataTableFunction::UK2Node_CallDataTableFunction(const FObjectInitializer& ObjectInitializer)
 	: Super(ObjectInitializer)
diff --git a/Engine/Source/Editor/BlueprintGraph/Private/K2Node_CallFunction.cpp b/Engine/Source/Editor/BlueprintGraph/Private/K2Node_CallFunction.cpp
--- a/Engine/Source/Editor/BlueprintGraph/Private/K2Node_CallFunction.cpp
+++ b/Engine/Source/Editor/BlueprintGraph/Private/K2Node_CallFunction.cpp
@@ -9,6 +9,7 @@
 #include "Kismet/KismetMathLibrary.h"
 #include "Kismet/KismetArrayLibrary.h"
 #include "K2Node_PureAssignmentStatement.h"
+#include "GraphEditorSettings.h"
 
 #define LOCTEXT_NAMESPACE "K2Node"
 
diff --git a/Engine/Source/Editor/BlueprintGraph/Private/K2Node_CallMaterialParameterCollectionFunction.cpp b/Engine/Source/Editor/BlueprintGraph/Private/K2Node_CallMaterialParameterCollectionFunction.cpp
--- a/Engine/Source/Editor/BlueprintGraph/Private/K2Node_CallMaterialParameterCollectionFunction.cpp
+++ b/Engine/Source/Editor/BlueprintGraph/Private/K2Node_CallMaterialParameterCollectionFunction.cpp
@@ -3,6 +3,7 @@
 #include "BlueprintGraphPrivatePCH.h"
 
 #include "CompilerResultsLog.h"
+#include "K2Node_CallMaterialParameterCollectionFunction.h"
 
 UK2Node_CallMaterialParameterCollectionFunction::UK2Node_CallMaterialParameterCollectionFunction(const FObjectInitializer& ObjectInitializer)
 	: Super(ObjectInitializer)
diff --git a/Engine/Source/Editor/BlueprintGraph/Private/K2Node_CallParentFunction.cpp b/Engine/Source/Editor/BlueprintGraph/Private/K2Node_CallParentFunction.cpp
--- a/Engine/Source/Editor/BlueprintGraph/Private/K2Node_CallParentFunction.cpp
+++ b/Engine/Source/Editor/BlueprintGraph/Private/K2Node_CallParentFunction.cpp
@@ -1,6 +1,9 @@
 // Copyright 1998-2014 Epic Games, Inc. All Rights Reserved.
 
 #include "BlueprintGraphPrivatePCH.h"
+#include "K2Node_CallParentFunction.h"
+#include "EditorStyleSettings.h"
+#include "GraphEditorSettings.h"
 
 #define LOCTEXT_NAMESPACE "K2Node"
 
diff --git a/Engine/Source/Editor/BlueprintGraph/Private/K2Node_ComponentBoundEvent.cpp b/Engine/Source/Editor/BlueprintGraph/Private/K2Node_ComponentBoundEvent.cpp
--- a/Engine/Source/Editor/BlueprintGraph/Private/K2Node_ComponentBoundEvent.cpp
+++ b/Engine/Source/Editor/BlueprintGraph/Private/K2Node_ComponentBoundEvent.cpp
@@ -2,6 +2,7 @@
 
 #include "BlueprintGraphPrivatePCH.h"
 #include "Engine/ComponentDelegateBinding.h"
+#include "K2Node_ComponentBoundEvent.h"
 
 #define LOCTEXT_NAMESPACE "K2Node"
 
diff --git a/Engine/Source/Editor/BlueprintGraph/Private/K2Node_Composite.cpp b/Engine/Source/Editor/BlueprintGraph/Private/K2Node_Composite.cpp
--- a/Engine/Source/Editor/BlueprintGraph/Private/K2Node_Composite.cpp
+++ b/Engine/Source/Editor/BlueprintGraph/Private/K2Node_Composite.cpp
@@ -4,6 +4,7 @@
 
 #include "EdGraphUtilities.h"
 #include "Kismet2NameValidators.h"
+#include "K2Node_Composite.h"
 
 #define LOCTEXT_NAMESPACE "K2Node"
 
diff --git a/Engine/Source/Editor/BlueprintGraph/Private/K2Node_CustomEvent.cpp b/Engine/Source/Editor/BlueprintGraph/Private/K2Node_CustomEvent.cpp
--- a/Engine/Source/Editor/BlueprintGraph/Private/K2Node_CustomEvent.cpp
+++ b/Engine/Source/Editor/BlueprintGraph/Private/K2Node_CustomEvent.cpp
@@ -7,6 +7,8 @@
 #include "CompilerResultsLog.h"
 #include "BlueprintEventNodeSpawner.h"
 #include "BlueprintActionDatabaseRegistrar.h"
+#include "K2Node_CustomEvent.h"
+#include "K2Node_BaseMCDelegate.h"
 
 #define LOCTEXT_NAMESPACE "K2Node_CustomEvent"
 
diff --git a/Engine/Source/Editor/BlueprintGraph/Private/K2Node_Event.cpp b/Engine/Source/Editor/BlueprintGraph/Private/K2Node_Event.cpp
--- a/Engine/Source/Editor/BlueprintGraph/Private/K2Node_Event.cpp
+++ b/Engine/Source/Editor/BlueprintGraph/Private/K2Node_Event.cpp
@@ -5,6 +5,7 @@
 #include "CompilerResultsLog.h"
 #include "KismetCompiler.h"
 #include "EventEntryHandler.h"
+#include "GraphEditorSettings.h"
 
 const FString UK2Node_Event::DelegateOutputName(TEXT("OutputDelegate"));
 
diff --git a/Engine/Source/Editor/BlueprintGraph/Private/K2Node_FunctionTerminator.cpp b/Engine/Source/Editor/BlueprintGraph/Private/K2Node_FunctionTerminator.cpp
--- a/Engine/Source/Editor/BlueprintGraph/Private/K2Node_FunctionTerminator.cpp
+++ b/Engine/Source/Editor/BlueprintGraph/Private/K2Node_FunctionTerminator.cpp
@@ -2,6 +2,8 @@
 
 
 #include "BlueprintGraphPrivatePCH.h"
+#include "K2Node_FunctionTerminator.h"
+#include "GraphEditorSettings.h"
 
 UK2Node_FunctionTerminator::UK2Node_FunctionTerminator(const FObjectInitializer& ObjectInitializer)
 	: Super(ObjectInitializer)
diff --git a/Engine/Source/Editor/BlueprintGraph/Private/K2Node_InputAxisEvent.cpp b/Engine/Source/Editor/BlueprintGraph/Private/K2Node_InputAxisEvent.cpp
--- a/Engine/Source/Editor/BlueprintGraph/Private/K2Node_InputAxisEvent.cpp
+++ b/Engine/Source/Editor/BlueprintGraph/Private/K2Node_InputAxisEvent.cpp
@@ -6,6 +6,7 @@
 #include "EditorCategoryUtils.h"
 #include "Engine/InputAxisDelegateBinding.h"
 #include "BlueprintActionDatabaseRegistrar.h"
+#include "K2Node_InputAxisEvent.h"
 
 #define LOCTEXT_NAMESPACE "K2Node_InputAxisEvent"
 
diff --git a/Engine/Source/Editor/BlueprintGraph/Private/K2Node_Knot.cpp b/Engine/Source/Editor/BlueprintGraph/Private/K2Node_Knot.cpp
--- a/Engine/Source/Editor/BlueprintGraph/Private/K2Node_Knot.cpp
+++ b/Engine/Source/Editor/BlueprintGraph/Private/K2Node_Knot.cpp
@@ -5,6 +5,7 @@
 #include "Kismet2NameValidators.h"
 #include "BlueprintNodeSpawner.h"
 #include "BlueprintActionDatabaseRegistrar.h"
+#include "K2Node_Knot.h"
 
 #define LOCTEXT_NAMESPACE "K2Node_Knot"
 
diff --git a/Engine/Source/Editor/BlueprintGraph/Private/K2Node_LocalVariable.cpp b/Engine/Source/Editor/BlueprintGraph/Private/K2Node_LocalVariable.cpp
--- a/Engine/Source/Editor/BlueprintGraph/Private/K2Node_LocalVariable.cpp
+++ b/Engine/Source/Editor/BlueprintGraph/Private/K2Node_LocalVariable.cpp
@@ -7,6 +7,9 @@
 
 #include "K2Node_FunctionEntry.h"
 #include "K2Node_LocalVariable.h"
+#include "K2Node_AssignmentStatement.h"
+#include "K2Node_VariableSet.h"
+#include "K2Node_VariableGet.h"
 
 UDEPRECATED_K2Node_LocalVariable::UDEPRECATED_K2Node_LocalVariable(const FObjectInitializer& ObjectInitializer)
 	: Super(ObjectInitializer)
diff --git a/Engine/Source/Editor/BlueprintGraph/Private/K2Node_MacroInstance.cpp b/Engine/Source/Editor/BlueprintGraph/Private/K2Node_MacroInstance.cpp
--- a/Engine/Source/Editor/BlueprintGraph/Private/K2Node_MacroInstance.cpp
+++ b/Engine/Source/Editor/BlueprintGraph/Private/K2Node_MacroInstance.cpp
@@ -4,6 +4,8 @@
 #include "SlateBasics.h"
 #include "EditorCategoryUtils.h"
 #include "BlueprintActionFilter.h"
+#include "K2Node_MacroInstance.h"
+#include "EditorStyleSet.h"
 
 #define LOCTEXT_NAMESPACE "K2Node_MacroInstance"
 
diff --git a/Engine/Source/Editor/BlueprintGraph/Private/K2Node_SwitchInteger.cpp b/Engine/Source/Editor/BlueprintGraph/Private/K2Node_SwitchInteger.cpp
--- a/Engine/Source/Editor/BlueprintGraph/Private/K2Node_SwitchInteger.cpp
+++ b/Engine/Source/Editor/BlueprintGraph/Private/K2Node_SwitchInteger.cpp
@@ -4,6 +4,7 @@
 #include "../../../Runtime/Engine/Classes/Kismet/KismetMathLibrary.h"
 #include "BlueprintNodeSpawner.h"
 #include "BlueprintActionDatabaseRegistrar.h"
+#include "K2Node_SwitchInteger.h"
 
 UK2Node_SwitchInteger::UK2Node_SwitchInteger(const FObjectInitializer& ObjectInitializer)
 	: Super(ObjectInitializer)
diff --git a/Engine/Source/Editor/BlueprintGraph/Private/K2Node_SwitchName.cpp b/Engine/Source/Editor/BlueprintGraph/Private/K2Node_SwitchName.cpp
--- a/Engine/Source/Editor/BlueprintGraph/Private/K2Node_SwitchName.cpp
+++ b/Engine/Source/Editor/BlueprintGraph/Private/K2Node_SwitchName.cpp
@@ -4,6 +4,7 @@
 #include "../../../Runtime/Engine/Classes/Kismet/KismetMathLibrary.h"
 #include "BlueprintNodeSpawner.h"
 #include "BlueprintActionDatabaseRegistrar.h"
+#include "K2Node_SwitchName.h"
 
 UK2Node_SwitchName::UK2Node_SwitchName(const FObjectInitializer& ObjectInitializer)
 	: Super(ObjectInitializer)
diff --git a/Engine/Source/Editor/BlueprintGraph/Private/K2Node_Tunnel.cpp b/Engine/Source/Editor/BlueprintGraph/Private/K2Node_Tunnel.cpp
--- a/Engine/Source/Editor/BlueprintGraph/Private/K2Node_Tunnel.cpp
+++ b/Engine/Source/Editor/BlueprintGraph/Private/K2Node_Tunnel.cpp
@@ -5,6 +5,7 @@
 
 #include "K2Node_MacroInstance.h"
 #include "K2Node_Tunnel.h"
+#include "K2Node_Composite.h"
 
 UK2Node_Tunnel::UK2Node_Tunnel(const FObjectInitializer& ObjectInitializer)
 	: Super(ObjectInitializer)
diff --git a/Engine/Source/Editor/GraphEditor/Public/GraphEditorActions.h b/Engine/Source/Editor/GraphEditor/Public/GraphEditorActions.h
--- a/Engine/Source/Editor/GraphEditor/Public/GraphEditorActions.h
+++ b/Engine/Source/Editor/GraphEditor/Public/GraphEditorActions.h
@@ -3,6 +3,7 @@
 #pragma once
 
 #include "SlateBasics.h"
+#include "EditorStyleSet.h"
 
 class FGraphEditorCommandsImpl : public TCommands<FGraphEditorCommandsImpl>
 {
diff --git a/Engine/Source/Editor/GraphEditor/Public/SGraphPin.h b/Engine/Source/Editor/GraphEditor/Public/SGraphPin.h
--- a/Engine/Source/Editor/GraphEditor/Public/SGraphPin.h
+++ b/Engine/Source/Editor/GraphEditor/Public/SGraphPin.h
@@ -1,6 +1,7 @@
 // Copyright 1998-2014 Epic Games, Inc. All Rights Reserved.
 
 #pragma once
+#include "SImage.h"
 
 class SGraphNode;
 
diff --git a/Engine/Source/Editor/GraphEditor/Public/SNameComboBox.h b/Engine/Source/Editor/GraphEditor/Public/SNameComboBox.h
--- a/Engine/Source/Editor/GraphEditor/Public/SNameComboBox.h
+++ b/Engine/Source/Editor/GraphEditor/Public/SNameComboBox.h
@@ -6,6 +6,7 @@
 #include "SharedPointer.h"
 #include "SlateEnums.h"
 #include "SCompoundWidget.h"
+#include "SComboBox.h"
 //#include "UnrealString.h"
 
 /**
diff --git a/Engine/Source/Editor/WorldBrowser/Private/LevelCollectionModel.cpp b/Engine/Source/Editor/WorldBrowser/Private/LevelCollectionModel.cpp
--- a/Engine/Source/Editor/WorldBrowser/Private/LevelCollectionModel.cpp
+++ b/Engine/Source/Editor/WorldBrowser/Private/LevelCollectionModel.cpp
@@ -8,6 +8,8 @@
 
 #include "LevelCollectionModel.h"
 #include "ShaderCompiler.h"
+#include "Engine/LevelStreaming.h"
+#include "Engine/Selection.h"
 
 #define LOCTEXT_NAMESPACE "WorldBrowser"
 
diff --git a/Engine/Source/Editor/WorldBrowser/Private/LevelModel.cpp b/Engine/Source/Editor/WorldBrowser/Private/LevelModel.cpp
--- a/Engine/Source/Editor/WorldBrowser/Private/LevelModel.cpp
+++ b/Engine/Source/Editor/WorldBrowser/Private/LevelModel.cpp
@@ -6,6 +6,8 @@
 #include "LevelModel.h"
 #include "LevelCollectionModel.h"
 #include "LevelBrowserSettings.h"
+#include "Engine/Selection.h"
+#include "Engine/LevelStreaming.h"
 
 #define LOCTEXT_NAMESPACE "WorldBrowser"
 
diff --git a/Engine/Source/Editor/WorldBrowser/Private/SWorldHierarchyItem.cpp b/Engine/Source/Editor/WorldBrowser/Private/SWorldHierarchyItem.cpp
--- a/Engine/Source/Editor/WorldBrowser/Private/SWorldHierarchyItem.cpp
+++ b/Engine/Source/Editor/WorldBrowser/Private/SWorldHierarchyItem.cpp
@@ -3,6 +3,8 @@
 
 #include "SWorldHierarchyItem.h"
 #include "SColorPicker.h"
+#include "Engine/LevelStreamingAlwaysLoaded.h"
+#include "Engine/LevelStreamingKismet.h"
 
 
 #define LOCTEXT_NAMESPACE "WorldBrowser"
diff --git a/Engine/Source/Editor/WorldBrowser/Private/StreamingLevels/StreamingLevelCollectionModel.cpp b/Engine/Source/Editor/WorldBrowser/Private/StreamingLevels/StreamingLevelCollectionModel.cpp
--- a/Engine/Source/Editor/WorldBrowser/Private/StreamingLevels/StreamingLevelCollectionModel.cpp
+++ b/Engine/Source/Editor/WorldBrowser/Private/StreamingLevels/StreamingLevelCollectionModel.cpp
@@ -10,6 +10,11 @@
 
 #include "StreamingLevelCustomization.h"
 #include "StreamingLevelCollectionModel.h"
+#include "Engine/LevelStreaming.h"
+#include "Engine/Selection.h"
+#include "Engine/LevelStreamingVolume.h"
+#include "Engine/LevelStreamingAlwaysLoaded.h"
+#include "Engine/LevelStreamingKismet.h"
 
 #define LOCTEXT_NAMESPACE "WorldBrowser"
 
diff --git a/Engine/Source/Editor/WorldBrowser/Private/StreamingLevels/StreamingLevelCollectionModel.h b/Engine/Source/Editor/WorldBrowser/Private/StreamingLevels/StreamingLevelCollectionModel.h
--- a/Engine/Source/Editor/WorldBrowser/Private/StreamingLevels/StreamingLevelCollectionModel.h
+++ b/Engine/Source/Editor/WorldBrowser/Private/StreamingLevels/StreamingLevelCollectionModel.h
@@ -2,6 +2,7 @@
 #pragma once
 
 #include "StreamingLevelModel.h"
+#include "EditorUndoClient.h"
 
 /** The non-UI solution specific presentation logic for a LevelsView */
 class FStreamingLevelCollectionModel 
diff --git a/Engine/Source/Editor/WorldBrowser/Private/StreamingLevels/StreamingLevelCustomization.cpp b/Engine/Source/Editor/WorldBrowser/Private/StreamingLevels/StreamingLevelCustomization.cpp
--- a/Engine/Source/Editor/WorldBrowser/Private/StreamingLevels/StreamingLevelCustomization.cpp
+++ b/Engine/Source/Editor/WorldBrowser/Private/StreamingLevels/StreamingLevelCustomization.cpp
@@ -13,6 +13,9 @@
 #include "StreamingLevelCustomization.h"
 #include "StreamingLevelEdMode.h"
 #include "SNumericEntryBox.h"
+#include "Engine/LevelStreaming.h"
+#include "Engine/LevelStreamingVolume.h"
+#include "AssetData.h"
 
 #define LOCTEXT_NAMESPACE "WorldBrowser"
 
diff --git a/Engine/Source/Editor/WorldBrowser/Private/StreamingLevels/StreamingLevelEdMode.cpp b/Engine/Source/Editor/WorldBrowser/Private/StreamingLevels/StreamingLevelEdMode.cpp
--- a/Engine/Source/Editor/WorldBrowser/Private/StreamingLevels/StreamingLevelEdMode.cpp
+++ b/Engine/Source/Editor/WorldBrowser/Private/StreamingLevels/StreamingLevelEdMode.cpp
@@ -2,6 +2,9 @@
 
 #include "WorldBrowserPrivatePCH.h"
 #include "StreamingLevelEdMode.h"
+#include "Engine/LevelStreaming.h"
+#include "GameFramework/WorldSettings.h"
+#include "Materials/MaterialInstanceDynamic.h"
 
 #define LOCTEXT_NAMESPACE "WorldBrowser"
 
diff --git a/Engine/Source/Editor/WorldBrowser/Private/StreamingLevels/StreamingLevelModel.cpp b/Engine/Source/Editor/WorldBrowser/Private/StreamingLevels/StreamingLevelModel.cpp
--- a/Engine/Source/Editor/WorldBrowser/Private/StreamingLevels/StreamingLevelModel.cpp
+++ b/Engine/Source/Editor/WorldBrowser/Private/StreamingLevels/StreamingLevelModel.cpp
@@ -4,6 +4,7 @@
 
 #include "StreamingLevelCollectionModel.h"
 #include "StreamingLevelModel.h"
+#include "Engine/LevelStreaming.h"
 
 #define LOCTEXT_NAMESPACE "WorldBrowser"
 
diff --git a/Engine/Source/Editor/WorldBrowser/Private/StreamingLevels/StreamingLevelModel.h b/Engine/Source/Editor/WorldBrowser/Private/StreamingLevels/StreamingLevelModel.h
--- a/Engine/Source/Editor/WorldBrowser/Private/StreamingLevels/StreamingLevelModel.h
+++ b/Engine/Source/Editor/WorldBrowser/Private/StreamingLevels/StreamingLevelModel.h
@@ -1,5 +1,7 @@
 // Copyright 1998-2014 Epic Games, Inc. All Rights Reserved.
 #pragma once
+#include "EditorUndoClient.h"
+#include "Layers/Layer.h"
 
 class FStreamingLevelCollectionModel;
 
diff --git a/Engine/Source/Editor/WorldBrowser/Private/Tiles/WorldTileCollectionModel.cpp b/Engine/Source/Editor/WorldBrowser/Private/Tiles/WorldTileCollectionModel.cpp
--- a/Engine/Source/Editor/WorldBrowser/Private/Tiles/WorldTileCollectionModel.cpp
+++ b/Engine/Source/Editor/WorldBrowser/Private/Tiles/WorldTileCollectionModel.cpp
@@ -16,6 +16,11 @@
 #include "WorldTileDetailsCustomization.h"
 #include "WorldTileCollectionModel.h"
 #include "STiledLandscapeImportDlg.h"
+#include "Engine/StaticMeshActor.h"
+#include "Engine/StaticMesh.h"
+#include "Engine/LevelStreaming.h"
+#include "GameFramework/WorldSettings.h"
+#include "Engine/Light.h"
 
 #define LOCTEXT_NAMESPACE "WorldBrowser"
 
diff --git a/Engine/Source/Editor/WorldBrowser/Private/Tiles/WorldTileCollectionModel.h b/Engine/Source/Editor/WorldBrowser/Private/Tiles/WorldTileCollectionModel.h
--- a/Engine/Source/Editor/WorldBrowser/Private/Tiles/WorldTileCollectionModel.h
+++ b/Engine/Source/Editor/WorldBrowser/Private/Tiles/WorldTileCollectionModel.h
@@ -2,6 +2,7 @@
 #pragma once
 
 #include "WorldTileModel.h"
+#include "EditorUndoClient.h"
 
 /** The non-UI solution specific presentation logic for a world composition */
 class FWorldTileCollectionModel
diff --git a/Engine/Source/Editor/WorldBrowser/Private/Tiles/WorldTileModel.cpp b/Engine/Source/Editor/WorldBrowser/Private/Tiles/WorldTileModel.cpp
--- a/Engine/Source/Editor/WorldBrowser/Private/Tiles/WorldTileModel.cpp
+++ b/Engine/Source/Editor/WorldBrowser/Private/Tiles/WorldTileModel.cpp
@@ -7,9 +7,13 @@
 #include "WorldTileDetails.h"
 #include "WorldTileCollectionModel.h"
 #include "WorldTileModel.h"
-
+#include "Engine/LevelBounds.h"
 #include "Engine/WorldComposition.h"
 #include "Landscape.h"
+#include "GameFramework/WorldSettings.h"
+#include "Engine/LevelStreaming.h"
+#include "Engine/LevelStreamingKismet.h"
+
 
 #define LOCTEXT_NAMESPACE "WorldBrowser"
 DEFINE_LOG_CATEGORY_STATIC(WorldBrowser, Log, All);
diff --git a/Engine/Source/Runtime/AIModule/Classes/AITypes.h b/Engine/Source/Runtime/AIModule/Classes/AITypes.h
--- a/Engine/Source/Runtime/AIModule/Classes/AITypes.h
+++ b/Engine/Source/Runtime/AIModule/Classes/AITypes.h
@@ -127,7 +127,7 @@ namespace EAIRequestPriority
 
 namespace EAIRequestPriority
 {
-	static const int32 Lowest = Logic;
+	static const int32 Lowest = EAIRequestPriority::Logic;
 };
 
 UENUM()
diff --git a/Engine/Source/Runtime/AIModule/Classes/Perception/PawnSensingComponent.h b/Engine/Source/Runtime/AIModule/Classes/Perception/PawnSensingComponent.h
--- a/Engine/Source/Runtime/AIModule/Classes/Perception/PawnSensingComponent.h
+++ b/Engine/Source/Runtime/AIModule/Classes/Perception/PawnSensingComponent.h
@@ -3,8 +3,11 @@
 
 
 #pragma once
+#include "Components/ActorComponent.h"
 #include "PawnSensingComponent.generated.h"
 
+class UPawnNoiseEmitterComponent;
+
 /**
  * SensingComponent encapsulates sensory (ie sight and hearing) settings and functionality for an Actor,
  * allowing the actor to see/hear Pawns in the world. It does nothing on network clients.
diff --git a/Engine/Source/Runtime/Engine/Classes/AI/Navigation/NavCollision.h b/Engine/Source/Runtime/Engine/Classes/AI/Navigation/NavCollision.h
--- a/Engine/Source/Runtime/Engine/Classes/AI/Navigation/NavCollision.h
+++ b/Engine/Source/Runtime/Engine/Classes/AI/Navigation/NavCollision.h
@@ -2,6 +2,7 @@
 
 #pragma once
 
+#include "NavigationTypes.h"
 #include "NavCollision.generated.h"
 
 USTRUCT()
diff --git a/Engine/Source/Runtime/Engine/Classes/Animation/AnimBlueprint.h b/Engine/Source/Runtime/Engine/Classes/Animation/AnimBlueprint.h
--- a/Engine/Source/Runtime/Engine/Classes/Animation/AnimBlueprint.h
+++ b/Engine/Source/Runtime/Engine/Classes/Animation/AnimBlueprint.h
@@ -1,7 +1,7 @@
 // Copyright 1998-2014 Epic Games, Inc. All Rights Reserved.
 
 #pragma once
-
+#include "Engine/Blueprint.h"
 #include "AnimBlueprint.generated.h"
 
 USTRUCT()
diff --git a/Engine/Source/Runtime/Engine/Classes/Atmosphere/AtmosphericFog.h b/Engine/Source/Runtime/Engine/Classes/Atmosphere/AtmosphericFog.h
--- a/Engine/Source/Runtime/Engine/Classes/Atmosphere/AtmosphericFog.h
+++ b/Engine/Source/Runtime/Engine/Classes/Atmosphere/AtmosphericFog.h
@@ -2,6 +2,7 @@
 
 
 #pragma once
+#include "GameFramework/Info.h"
 #include "AtmosphericFog.generated.h"
 
 /** 
diff --git a/Engine/Source/Runtime/Engine/Classes/Engine/BlockingVolume.h b/Engine/Source/Runtime/Engine/Classes/Engine/BlockingVolume.h
--- a/Engine/Source/Runtime/Engine/Classes/Engine/BlockingVolume.h
+++ b/Engine/Source/Runtime/Engine/Classes/Engine/BlockingVolume.h
@@ -2,6 +2,7 @@
 
 
 #pragma once
+#include "GameFramework/Volume.h"
 #include "BlockingVolume.generated.h"
 
 /** An invisible volume used to block other actors. */
diff --git a/Engine/Source/Runtime/Engine/Classes/Engine/DirectionalLight.h b/Engine/Source/Runtime/Engine/Classes/Engine/DirectionalLight.h
--- a/Engine/Source/Runtime/Engine/Classes/Engine/DirectionalLight.h
+++ b/Engine/Source/Runtime/Engine/Classes/Engine/DirectionalLight.h
@@ -4,6 +4,8 @@
 #include "Engine/Light.h"
 #include "DirectionalLight.generated.h"
 
+class UArrowComponent;
+
 UCLASS(ClassGroup=(Lights, DirectionalLights), MinimalAPI, meta=(ChildCanTick))
 class ADirectionalLight : public ALight
 {
diff --git a/Engine/Source/Runtime/Engine/Classes/Engine/LevelStreamingKismet.h b/Engine/Source/Runtime/Engine/Classes/Engine/LevelStreamingKismet.h
--- a/Engine/Source/Runtime/Engine/Classes/Engine/LevelStreamingKismet.h
+++ b/Engine/Source/Runtime/Engine/Classes/Engine/LevelStreamingKismet.h
@@ -8,6 +8,7 @@
  */
 
 #pragma once
+#include "LevelStreaming.h"
 #include "LevelStreamingKismet.generated.h"
 
 UCLASS(MinimalAPI, BlueprintType)
diff --git a/Engine/Source/Runtime/Engine/Classes/Engine/LocalPlayer.h b/Engine/Source/Runtime/Engine/Classes/Engine/LocalPlayer.h
--- a/Engine/Source/Runtime/Engine/Classes/Engine/LocalPlayer.h
+++ b/Engine/Source/Runtime/Engine/Classes/Engine/LocalPlayer.h
@@ -5,6 +5,7 @@
 //=============================================================================
 
 #pragma once
+#include "Player.h"
 #include "LocalPlayer.generated.h"
 
 /** A context object that binds to a LocalPlayer. Useful for UI or other things that need to pass around player references */
diff --git a/Engine/Source/Runtime/Engine/Classes/Engine/NavigationObjectBase.h b/Engine/Source/Runtime/Engine/Classes/Engine/NavigationObjectBase.h
--- a/Engine/Source/Runtime/Engine/Classes/Engine/NavigationObjectBase.h
+++ b/Engine/Source/Runtime/Engine/Classes/Engine/NavigationObjectBase.h
@@ -2,6 +2,7 @@
 
 
 #pragma once
+#include "AI/Navigation/NavAgentInterface.h"
 #include "NavigationObjectBase.generated.h"
 
 UCLASS(hidecategories=(Lighting, LightColor, Force), ClassGroup=Navigation, NotBlueprintable, abstract)
diff --git a/Engine/Source/Runtime/Engine/Classes/Engine/PointLight.h b/Engine/Source/Runtime/Engine/Classes/Engine/PointLight.h
--- a/Engine/Source/Runtime/Engine/Classes/Engine/PointLight.h
+++ b/Engine/Source/Runtime/Engine/Classes/Engine/PointLight.h
@@ -1,7 +1,7 @@
 // Copyright 1998-2014 Epic Games, Inc. All Rights Reserved.
 
 #pragma once
-
+#include "Light.h"
 #include "PointLight.generated.h"
 
 UCLASS(ClassGroup=(Lights, PointLights),MinimalAPI, meta=(ChildCanTick))
diff --git a/Engine/Source/Runtime/Engine/Classes/Engine/RendererSettings.h b/Engine/Source/Runtime/Engine/Classes/Engine/RendererSettings.h
--- a/Engine/Source/Runtime/Engine/Classes/Engine/RendererSettings.h
+++ b/Engine/Source/Runtime/Engine/Classes/Engine/RendererSettings.h
@@ -2,6 +2,7 @@
 
 #pragma once
 
+#include "Curves/CurveFloat.h"
 #include "UserInterfaceSettings.h"
 
 #include "RendererSettings.generated.h"
diff --git a/Engine/Source/Runtime/Engine/Classes/Engine/SCS_Node.h b/Engine/Source/Runtime/Engine/Classes/Engine/SCS_Node.h
--- a/Engine/Source/Runtime/Engine/Classes/Engine/SCS_Node.h
+++ b/Engine/Source/Runtime/Engine/Classes/Engine/SCS_Node.h
@@ -2,6 +2,7 @@
 
 
 #pragma once
+#include "SimpleConstructionScript.h"
 #include "SCS_Node.generated.h"
 
 DECLARE_DELEGATE_OneParam( FSCSNodeNameChanged, const FName& );
diff --git a/Engine/Source/Runtime/Engine/Classes/Engine/SceneCapture2D.h b/Engine/Source/Runtime/Engine/Classes/Engine/SceneCapture2D.h
--- a/Engine/Source/Runtime/Engine/Classes/Engine/SceneCapture2D.h
+++ b/Engine/Source/Runtime/Engine/Classes/Engine/SceneCapture2D.h
@@ -6,6 +6,7 @@
  */
 
 #pragma once
+#include "SceneCapture.h"
 #include "SceneCapture2D.generated.h"
 
 UCLASS(hidecategories=(Collision, Material, Attachment, Actor), MinimalAPI)
diff --git a/Engine/Source/Runtime/Engine/Classes/Engine/SceneCaptureCube.h b/Engine/Source/Runtime/Engine/Classes/Engine/SceneCaptureCube.h
--- a/Engine/Source/Runtime/Engine/Classes/Engine/SceneCaptureCube.h
+++ b/Engine/Source/Runtime/Engine/Classes/Engine/SceneCaptureCube.h
@@ -6,6 +6,7 @@
  */
 
 #pragma once
+#include "SceneCapture.h"
 #include "SceneCaptureCube.generated.h"
 
 UCLASS(hidecategories = (Collision, Material, Attachment, Actor), MinimalAPI)
diff --git a/Engine/Source/Runtime/Engine/Classes/Engine/SkyLight.h b/Engine/Source/Runtime/Engine/Classes/Engine/SkyLight.h
--- a/Engine/Source/Runtime/Engine/Classes/Engine/SkyLight.h
+++ b/Engine/Source/Runtime/Engine/Classes/Engine/SkyLight.h
@@ -1,7 +1,7 @@
 // Copyright 1998-2014 Epic Games, Inc. All Rights Reserved.
 
 #pragma once
-
+#include "GameFramework/Info.h"
 #include "SkyLight.generated.h"
 
 UCLASS(ClassGroup=Lights, hidecategories=(Input,Collision,Replication,Info), showcategories=("Input|MouseInput", "Input|TouchInput"), ConversionRoot, Blueprintable)
diff --git a/Engine/Source/Runtime/Engine/Classes/Engine/SphereReflectionCapture.h b/Engine/Source/Runtime/Engine/Classes/Engine/SphereReflectionCapture.h
--- a/Engine/Source/Runtime/Engine/Classes/Engine/SphereReflectionCapture.h
+++ b/Engine/Source/Runtime/Engine/Classes/Engine/SphereReflectionCapture.h
@@ -1,6 +1,7 @@
 // Copyright 1998-2014 Epic Games, Inc. All Rights Reserved.
 
 #pragma once
+#include "ReflectionCapture.h"
 #include "SphereReflectionCapture.generated.h"
 
 class UDrawSphereComponent;
diff --git a/Engine/Source/Runtime/Engine/Classes/Engine/UserInterfaceSettings.h b/Engine/Source/Runtime/Engine/Classes/Engine/UserInterfaceSettings.h
--- a/Engine/Source/Runtime/Engine/Classes/Engine/UserInterfaceSettings.h
+++ b/Engine/Source/Runtime/Engine/Classes/Engine/UserInterfaceSettings.h
@@ -1,7 +1,7 @@
 // Copyright 1998-2014 Epic Games, Inc. All Rights Reserved.
 
 #pragma once
-
+#include "Curves/CurveFloat.h"
 #include "UserInterfaceSettings.generated.h"
 
 
diff --git a/Engine/Source/Runtime/Engine/Classes/Engine/WindDirectionalSource.h b/Engine/Source/Runtime/Engine/Classes/Engine/WindDirectionalSource.h
--- a/Engine/Source/Runtime/Engine/Classes/Engine/WindDirectionalSource.h
+++ b/Engine/Source/Runtime/Engine/Classes/Engine/WindDirectionalSource.h
@@ -1,6 +1,6 @@
 // Copyright 1998-2014 Epic Games, Inc. All Rights Reserved.
 #pragma once
-
+#include "GameFramework/Info.h"
 #include "WindDirectionalSource.generated.h"
 
 UCLASS(ClassGroup=Wind, showcategories=(Rendering, "Utilities|Transformation"))
diff --git a/Engine/Source/Runtime/Engine/Classes/Kismet/KismetStringLibrary.h b/Engine/Source/Runtime/Engine/Classes/Kismet/KismetStringLibrary.h
--- a/Engine/Source/Runtime/Engine/Classes/Kismet/KismetStringLibrary.h
+++ b/Engine/Source/Runtime/Engine/Classes/Kismet/KismetStringLibrary.h
@@ -1,6 +1,7 @@
 // Copyright 1998-2014 Epic Games, Inc. All Rights Reserved.
 
 #pragma once
+#include "BlueprintFunctionLibrary.h"
 #include "KismetStringLibrary.generated.h"
 
 UCLASS(MinimalAPI)
diff --git a/Engine/Source/Runtime/Engine/Classes/Kismet/KismetSystemLibrary.h b/Engine/Source/Runtime/Engine/Classes/Kismet/KismetSystemLibrary.h
--- a/Engine/Source/Runtime/Engine/Classes/Kismet/KismetSystemLibrary.h
+++ b/Engine/Source/Runtime/Engine/Classes/Kismet/KismetSystemLibrary.h
@@ -3,6 +3,7 @@
 #pragma once
 
 #include "Engine/LatentActionManager.h"
+#include "BlueprintFunctionLibrary.h"
 #include "KismetSystemLibrary.generated.h"
 
 /** It's just a proxy for FTimerDynamicDelegate */
diff --git a/Engine/Source/Runtime/Engine/Classes/Lightmass/LightmassCharacterIndirectDetailVolume.h b/Engine/Source/Runtime/Engine/Classes/Lightmass/LightmassCharacterIndirectDetailVolume.h
--- a/Engine/Source/Runtime/Engine/Classes/Lightmass/LightmassCharacterIndirectDetailVolume.h
+++ b/Engine/Source/Runtime/Engine/Classes/Lightmass/LightmassCharacterIndirectDetailVolume.h
@@ -5,6 +5,7 @@
 //=============================================================================
 #pragma once
 
+#include "GameFramework/Volume.h"
 #include "LightmassCharacterIndirectDetailVolume.generated.h"
 
 UCLASS(hidecategories=(Collision, Brush, Attachment, Physics, Volume), MinimalAPI)
diff --git a/Engine/Source/Runtime/Engine/Classes/Lightmass/LightmassImportanceVolume.h b/Engine/Source/Runtime/Engine/Classes/Lightmass/LightmassImportanceVolume.h
--- a/Engine/Source/Runtime/Engine/Classes/Lightmass/LightmassImportanceVolume.h
+++ b/Engine/Source/Runtime/Engine/Classes/Lightmass/LightmassImportanceVolume.h
@@ -6,6 +6,7 @@
 //=============================================================================
 #pragma once
 
+#include "GameFramework/Volume.h"
 #include "LightmassImportanceVolume.generated.h"
 
 UCLASS(hidecategories=(Collision, Brush, Attachment, Physics, Volume), MinimalAPI)
diff --git a/Engine/Source/Runtime/Engine/Classes/Materials/MaterialInstanceDynamic.h b/Engine/Source/Runtime/Engine/Classes/Materials/MaterialInstanceDynamic.h
--- a/Engine/Source/Runtime/Engine/Classes/Materials/MaterialInstanceDynamic.h
+++ b/Engine/Source/Runtime/Engine/Classes/Materials/MaterialInstanceDynamic.h
@@ -2,6 +2,7 @@
 
 
 #pragma once
+#include "MaterialInstance.h"
 #include "MaterialInstanceDynamic.generated.h"
 
 UCLASS(hidecategories=Object, collapsecategories, BlueprintType, MinimalAPI)
diff --git a/Engine/Source/Runtime/Engine/Classes/Particles/Parameter/ParticleModuleParameterDynamic.h b/Engine/Source/Runtime/Engine/Classes/Particles/Parameter/ParticleModuleParameterDynamic.h
--- a/Engine/Source/Runtime/Engine/Classes/Particles/Parameter/ParticleModuleParameterDynamic.h
+++ b/Engine/Source/Runtime/Engine/Classes/Particles/Parameter/ParticleModuleParameterDynamic.h
@@ -8,8 +8,11 @@
 #pragma once
 #include "Particles/Parameter/ParticleModuleParameterBase.h"
 #include "Particles/ParticleSystemComponent.h"
+#include "Distributions/DistributionFloatConstant.h"
 #include "ParticleModuleParameterDynamic.generated.h"
 
+class UDistributionFloatConstant;
+
 /**
  *	EmitterDynamicParameterValue
  *	Enumeration indicating the way a dynamic parameter should be set.
diff --git a/Engine/Source/Runtime/Engine/Classes/Sound/SoundNode.h b/Engine/Source/Runtime/Engine/Classes/Sound/SoundNode.h
--- a/Engine/Source/Runtime/Engine/Classes/Sound/SoundNode.h
+++ b/Engine/Source/Runtime/Engine/Classes/Sound/SoundNode.h
@@ -8,6 +8,8 @@
 	USoundNode helper macros. 
 -----------------------------------------------------------------------------*/
 
+struct FActiveSound;
+
 #define DECLARE_SOUNDNODE_ELEMENT(Type,Name)													\
 	Type& Name = *((Type*)(Payload));															\
 	Payload += sizeof(Type);														
diff --git a/Engine/Source/Runtime/Engine/Classes/Sound/SoundNodeAttenuation.h b/Engine/Source/Runtime/Engine/Classes/Sound/SoundNodeAttenuation.h
--- a/Engine/Source/Runtime/Engine/Classes/Sound/SoundNodeAttenuation.h
+++ b/Engine/Source/Runtime/Engine/Classes/Sound/SoundNodeAttenuation.h
@@ -3,8 +3,12 @@
  
 #pragma once
 #include "Sound/SoundNode.h"
+#include "SoundAttenuation.h"
 #include "SoundNodeAttenuation.generated.h"
 
+struct FAttenuationSettings;
+class USoundAttenuation;
+
 /** 
  * Defines how a sound's volume changes based on distance to the listener
  */
diff --git a/Engine/Source/Runtime/Engine/Private/Blueprint.cpp b/Engine/Source/Runtime/Engine/Private/Blueprint.cpp
--- a/Engine/Source/Runtime/Engine/Private/Blueprint.cpp
+++ b/Engine/Source/Runtime/Engine/Private/Blueprint.cpp
@@ -17,6 +17,9 @@
 #include "MessageLog.h"
 #include "Editor/UnrealEd/Classes/Settings/EditorLoadingSavingSettings.h"
 #include "BlueprintEditorUtils.h"
+#include "Engine/TimelineTemplate.h"
+#include "Engine/SimpleConstructionScript.h"
+#include "Engine/SCS_Node.h"
 #endif
 #include "Components/TimelineComponent.h"
 
diff --git a/Engine/Source/Runtime/Engine/Public/AnimationCompression.h b/Engine/Source/Runtime/Engine/Public/AnimationCompression.h
--- a/Engine/Source/Runtime/Engine/Public/AnimationCompression.h
+++ b/Engine/Source/Runtime/Engine/Public/AnimationCompression.h
@@ -9,6 +9,7 @@
 
 #include "FloatPacker.h"
 #include "Animation/AnimSequence.h"
+#include "EngineLogs.h"
 
 DECLARE_LOG_CATEGORY_EXTERN(LogAnimationCompression, Log, All);
 
diff --git a/Engine/Source/Runtime/Engine/Public/GenericOctree.h b/Engine/Source/Runtime/Engine/Public/GenericOctree.h
--- a/Engine/Source/Runtime/Engine/Public/GenericOctree.h
+++ b/Engine/Source/Runtime/Engine/Public/GenericOctree.h
@@ -7,6 +7,8 @@
 #ifndef __GENERIC_OCTREE_H__
 #define __GENERIC_OCTREE_H__
 
+class FOctreeElementId;
+
 /** A concise iteration over the children of an octree node. */
 #define FOREACH_OCTREE_CHILD_NODE(ChildRef) \
 	for(FOctreeChildNodeRef ChildRef(0);!ChildRef.IsNULL();ChildRef.Advance())
diff --git a/Engine/Source/Runtime/Engine/Public/PrecomputedLightVolume.h b/Engine/Source/Runtime/Engine/Public/PrecomputedLightVolume.h
--- a/Engine/Source/Runtime/Engine/Public/PrecomputedLightVolume.h
+++ b/Engine/Source/Runtime/Engine/Public/PrecomputedLightVolume.h
@@ -8,6 +8,7 @@
 #define __PRECOMPUTEDLIGHTVOLUME_H__
 
 #include "GenericOctree.h"
+#include "GenericOctreePublic.h"
 
 /** Incident radiance stored for a point. */
 class FVolumeLightingSample
diff --git a/Engine/Source/Runtime/Slate/Public/Framework/Application/SlateApplication.h b/Engine/Source/Runtime/Slate/Public/Framework/Application/SlateApplication.h
--- a/Engine/Source/Runtime/Slate/Public/Framework/Application/SlateApplication.h
+++ b/Engine/Source/Runtime/Slate/Public/Framework/Application/SlateApplication.h
@@ -4,6 +4,8 @@
 
 #include "GenericApplicationMessageHandler.h"
 #include "MenuStack.h"
+#include "SlateDelegates.h"
+#include "SlateApplicationBase.h"
 
 class SToolTip;
 class SViewport;
diff --git a/Engine/Source/Runtime/Slate/Public/Framework/MultiBox/MultiBoxBuilder.h b/Engine/Source/Runtime/Slate/Public/Framework/MultiBox/MultiBoxBuilder.h
--- a/Engine/Source/Runtime/Slate/Public/Framework/MultiBox/MultiBoxBuilder.h
+++ b/Engine/Source/Runtime/Slate/Public/Framework/MultiBox/MultiBoxBuilder.h
@@ -3,11 +3,13 @@
 #pragma once
 
 #include "MultiBoxExtender.h"
+#include "SlateDelegates.h"
 
 
 class FUICommandInfo;
 class FUICommandList;
-
+struct FSlateIcon;
+struct FUIAction;
 
 /** Delegate used by multi-box to call a user function to populate a new menu.  Used for spawning sub-menus and pull-down menus. */
 DECLARE_DELEGATE_OneParam( FNewMenuDelegate, class FMenuBuilder& );
diff --git a/Engine/Source/Runtime/Slate/Public/Framework/MultiBox/MultiBoxExtender.h b/Engine/Source/Runtime/Slate/Public/Framework/MultiBox/MultiBoxExtender.h
--- a/Engine/Source/Runtime/Slate/Public/Framework/MultiBox/MultiBoxExtender.h
+++ b/Engine/Source/Runtime/Slate/Public/Framework/MultiBox/MultiBoxExtender.h
@@ -2,6 +2,7 @@
 
 #pragma once
 
+class FUICommandList;
 
 /** Called on your extension to add new pull-down menus to a menu bar */
 DECLARE_DELEGATE_OneParam( FMenuBarExtensionDelegate, class FMenuBarBuilder& );
diff --git a/Engine/Source/Runtime/Slate/Public/Framework/Text/IRichTextMarkupWriter.h b/Engine/Source/Runtime/Slate/Public/Framework/Text/IRichTextMarkupWriter.h
--- a/Engine/Source/Runtime/Slate/Public/Framework/Text/IRichTextMarkupWriter.h
+++ b/Engine/Source/Runtime/Slate/Public/Framework/Text/IRichTextMarkupWriter.h
@@ -1,5 +1,6 @@
 // Copyright 1998-2014 Epic Games, Inc. All Rights Reserved.
 #pragma once
+#include "IRun.h"
 
 #if WITH_FANCY_TEXT
 
diff --git a/Engine/Source/Runtime/Slate/Public/Framework/Text/IRun.h b/Engine/Source/Runtime/Slate/Public/Framework/Text/IRun.h
--- a/Engine/Source/Runtime/Slate/Public/Framework/Text/IRun.h
+++ b/Engine/Source/Runtime/Slate/Public/Framework/Text/IRun.h
@@ -1,6 +1,8 @@
 // Copyright 1998-2014 Epic Games, Inc. All Rights Reserved.
 #pragma once
 
+enum class ETextHitPoint : uint8;
+
 struct SLATE_API FRunInfo
 {
 	FRunInfo()
diff --git a/Engine/Source/Runtime/Slate/Public/Framework/Text/ITextDecorator.h b/Engine/Source/Runtime/Slate/Public/Framework/Text/ITextDecorator.h
--- a/Engine/Source/Runtime/Slate/Public/Framework/Text/ITextDecorator.h
+++ b/Engine/Source/Runtime/Slate/Public/Framework/Text/ITextDecorator.h
@@ -1,5 +1,9 @@
 // Copyright 1998-2014 Epic Games, Inc. All Rights Reserved.
 #pragma once
+#include "TextRange.h"
+#include "IRun.h"
+
+class ISlateRun;
 
 #if WITH_FANCY_TEXT
 
diff --git a/Engine/Source/Runtime/Slate/Public/Framework/Text/SlateWidgetRun.h b/Engine/Source/Runtime/Slate/Public/Framework/Text/SlateWidgetRun.h
--- a/Engine/Source/Runtime/Slate/Public/Framework/Text/SlateWidgetRun.h
+++ b/Engine/Source/Runtime/Slate/Public/Framework/Text/SlateWidgetRun.h
@@ -1,5 +1,6 @@
 // Copyright 1998-2014 Epic Games, Inc. All Rights Reserved.
 #pragma once
+#include "ISlateRun.h"
 
 #if WITH_FANCY_TEXT
 
diff --git a/Engine/Source/Runtime/Slate/Public/Framework/Text/TextDecorators.h b/Engine/Source/Runtime/Slate/Public/Framework/Text/TextDecorators.h
--- a/Engine/Source/Runtime/Slate/Public/Framework/Text/TextDecorators.h
+++ b/Engine/Source/Runtime/Slate/Public/Framework/Text/TextDecorators.h
@@ -4,6 +4,8 @@
 #if WITH_FANCY_TEXT
 
 #include "ITextDecorator.h"
+#include "SlateWidgetRun.h"
+#include "SlateHyperlinkRun.h"
 
 class SLATE_API FWidgetDecorator : public ITextDecorator
 {
diff --git a/Engine/Source/Runtime/Slate/Public/Framework/Text/TextLayout.h b/Engine/Source/Runtime/Slate/Public/Framework/Text/TextLayout.h
--- a/Engine/Source/Runtime/Slate/Public/Framework/Text/TextLayout.h
+++ b/Engine/Source/Runtime/Slate/Public/Framework/Text/TextLayout.h
@@ -2,6 +2,9 @@
 #pragma once
 
 #include "IBreakIterator.h"
+#include "TextRunRenderer.h"
+#include "TextLineHighlight.h"
+#include "ILineHighlighter.h"
 #include "TextLayout.generated.h"
 
 #define TEXT_LAYOUT_DEBUG 0
diff --git a/Engine/Source/Runtime/Slate/Public/Widgets/Input/SCheckBox.h b/Engine/Source/Runtime/Slate/Public/Widgets/Input/SCheckBox.h
--- a/Engine/Source/Runtime/Slate/Public/Widgets/Input/SCheckBox.h
+++ b/Engine/Source/Runtime/Slate/Public/Widgets/Input/SCheckBox.h
@@ -2,6 +2,8 @@
 
 #pragma once
 
+#include "SBorder.h"
+#include "SlateDelegates.h"
 #include "SCheckBox.generated.h"
 
 /** Current state of the check box */
diff --git a/Engine/Source/Runtime/Slate/Public/Widgets/Input/SComboBox.h b/Engine/Source/Runtime/Slate/Public/Widgets/Input/SComboBox.h
--- a/Engine/Source/Runtime/Slate/Public/Widgets/Input/SComboBox.h
+++ b/Engine/Source/Runtime/Slate/Public/Widgets/Input/SComboBox.h
@@ -1,6 +1,11 @@
 // Copyright 1998-2014 Epic Games, Inc. All Rights Reserved.
 
 #pragma once
+#include "STableRow.h"
+#include "SMenuAnchor.h"
+#include "SComboButton.h"
+#include "SListView.h"
+#include "SlateApplication.h"
 
 
 DECLARE_DELEGATE( FOnComboBoxOpening )
diff --git a/Engine/Source/Runtime/Slate/Public/Widgets/Input/SMenuAnchor.h b/Engine/Source/Runtime/Slate/Public/Widgets/Input/SMenuAnchor.h
--- a/Engine/Source/Runtime/Slate/Public/Widgets/Input/SMenuAnchor.h
+++ b/Engine/Source/Runtime/Slate/Public/Widgets/Input/SMenuAnchor.h
@@ -1,6 +1,7 @@
 // Copyright 1998-2014 Epic Games, Inc. All Rights Reserved.
 
 #pragma once
+#include "STextBlock.h"
 
 /** Notification when popup is opened/closed. */
 DECLARE_DELEGATE_OneParam(FOnIsOpenChanged, bool)
diff --git a/Engine/Source/Runtime/Slate/Public/Widgets/Layout/SBorder.h b/Engine/Source/Runtime/Slate/Public/Widgets/Layout/SBorder.h
--- a/Engine/Source/Runtime/Slate/Public/Widgets/Layout/SBorder.h
+++ b/Engine/Source/Runtime/Slate/Public/Widgets/Layout/SBorder.h
@@ -1,6 +1,7 @@
 // Copyright 1998-2014 Epic Games, Inc. All Rights Reserved.
 
 #pragma once
+#include "SlateDelegates.h"
 
 /**
  * A border is an widget that can be used to contain other widgets.
diff --git a/Engine/Source/Runtime/Slate/Public/Widgets/Layout/SScrollBar.h b/Engine/Source/Runtime/Slate/Public/Widgets/Layout/SScrollBar.h
--- a/Engine/Source/Runtime/Slate/Public/Widgets/Layout/SScrollBar.h
+++ b/Engine/Source/Runtime/Slate/Public/Widgets/Layout/SScrollBar.h
@@ -1,6 +1,7 @@
 // Copyright 1998-2014 Epic Games, Inc. All Rights Reserved.
  
 #pragma once
+#include "SSpacer.h"
 
 
 DECLARE_DELEGATE_OneParam(
diff --git a/Engine/Source/Runtime/Slate/Public/Widgets/Views/SExpanderArrow.h b/Engine/Source/Runtime/Slate/Public/Widgets/Views/SExpanderArrow.h
--- a/Engine/Source/Runtime/Slate/Public/Widgets/Views/SExpanderArrow.h
+++ b/Engine/Source/Runtime/Slate/Public/Widgets/Views/SExpanderArrow.h
@@ -1,6 +1,7 @@
 // Copyright 1998-2014 Epic Games, Inc. All Rights Reserved.
 
 #pragma once
+#include "SButton.h"
 
 
 /**
diff --git a/Engine/Source/Runtime/Slate/Public/Widgets/Views/SHeaderRow.h b/Engine/Source/Runtime/Slate/Public/Widgets/Views/SHeaderRow.h
--- a/Engine/Source/Runtime/Slate/Public/Widgets/Views/SHeaderRow.h
+++ b/Engine/Source/Runtime/Slate/Public/Widgets/Views/SHeaderRow.h
@@ -1,6 +1,8 @@
 // Copyright 1998-2014 Epic Games, Inc. All Rights Reserved.
 
 #pragma once
+#include "SBorder.h"
+#include "SScrollBar.h"
 
 
 class SSplitter;
diff --git a/Engine/Source/Runtime/Slate/Public/Widgets/Views/SListView.h b/Engine/Source/Runtime/Slate/Public/Widgets/Views/SListView.h
--- a/Engine/Source/Runtime/Slate/Public/Widgets/Views/SListView.h
+++ b/Engine/Source/Runtime/Slate/Public/Widgets/Views/SListView.h
@@ -1,6 +1,7 @@
 // Copyright 1998-2014 Epic Games, Inc. All Rights Reserved.
 
 #pragma once
+#include "STableViewBase.h"
 
 
 /**
diff --git a/Engine/Source/Runtime/Slate/Public/Widgets/Views/STableRow.h b/Engine/Source/Runtime/Slate/Public/Widgets/Views/STableRow.h
--- a/Engine/Source/Runtime/Slate/Public/Widgets/Views/STableRow.h
+++ b/Engine/Source/Runtime/Slate/Public/Widgets/Views/STableRow.h
@@ -1,6 +1,9 @@
 // Copyright 1998-2014 Epic Games, Inc. All Rights Reserved.
 
 #pragma once
+#include "TableViewTypeTraits.h"
+#include "ITypedTableView.h"
+#include "SExpanderArrow.h"
 
 
 /**
diff --git a/Engine/Source/Runtime/UMG/Public/Blueprint/UserWidget.h b/Engine/Source/Runtime/UMG/Public/Blueprint/UserWidget.h
--- a/Engine/Source/Runtime/UMG/Public/Blueprint/UserWidget.h
+++ b/Engine/Source/Runtime/UMG/Public/Blueprint/UserWidget.h
@@ -8,6 +8,7 @@
 
 #include "NamedSlotInterface.h"
 #include "Slate/Anchors.h"
+#include "Engine/LocalPlayer.h"
 #include "UserWidget.generated.h"
 
 static FGeometry NullGeometry;
@@ -18,6 +19,7 @@ static FWidgetStyle NullStyle;
 class UWidget;
 class UWidgetAnimation;
 struct FAnchors;
+struct FLocalPlayerContext;
 
 /**
  * The state passed into OnPaint that we can expose as a single painting structure to blueprints to
