Only display an asset picker for placeable actors if the node is in a level script blueprint

diff --git a/Engine/Source/Editor/BlueprintGraph/Private/EdGraphSchema_K2.cpp b/Engine/Source/Editor/BlueprintGraph/Private/EdGraphSchema_K2.cpp
--- a/Engine/Source/Editor/BlueprintGraph/Private/EdGraphSchema_K2.cpp
+++ b/Engine/Source/Editor/BlueprintGraph/Private/EdGraphSchema_K2.cpp
@@ -3603,6 +3603,13 @@ bool UEdGraphSchema_K2::ShouldShowAssetPickerForPin(UEdGraphPin* Pin) const
 			// Don't show literal buttons for component type objects
 			bShow = !ObjectClass->IsChildOf(UActorComponent::StaticClass());
 
+			if (bShow && ObjectClass->IsChildOf(AActor::StaticClass()))
+			{
+				// Only show the picker for Actor classes if the class is placeable and we are in the level script
+				bShow = !ObjectClass->HasAllClassFlags(CLASS_NotPlaceable)
+							&& FBlueprintEditorUtils::IsLevelScriptBlueprint(FBlueprintEditorUtils::FindBlueprintForNode(Pin->GetOwningNode()));
+			}
+
 			if (bShow)
 			{
 				if (UK2Node_CallFunction* CallFunctionNode = Cast<UK2Node_CallFunction>(Pin->GetOwningNode()))
