BP: CreateDelegate node is listed in menu.

diff --git a/Engine/Source/Editor/BlueprintGraph/Classes/K2Node_CreateDelegate.h b/Engine/Source/Editor/BlueprintGraph/Classes/K2Node_CreateDelegate.h
--- a/Engine/Source/Editor/BlueprintGraph/Classes/K2Node_CreateDelegate.h
+++ b/Engine/Source/Editor/BlueprintGraph/Classes/K2Node_CreateDelegate.h
@@ -28,6 +28,8 @@ public:
 	virtual bool IsNodePure() const override { return true; }
 	virtual void PostReconstructNode() override;
 	virtual class FNodeHandlingFunctor* CreateNodeHandler(class FKismetCompilerContext& CompilerContext) const override;
+	virtual void GetMenuActions(FBlueprintActionDatabaseRegistrar& ActionRegistrar) const override;
+	virtual FText GetMenuCategory() const override;
 	// End of UK2Node interface
 
 	bool IsValid(FString* OutMsg = NULL, bool bDontUseSkeletalClassForSelf = false) const;
diff --git a/Engine/Source/Editor/BlueprintGraph/Private/K2Node_CreateDelegate.cpp b/Engine/Source/Editor/BlueprintGraph/Private/K2Node_CreateDelegate.cpp
--- a/Engine/Source/Editor/BlueprintGraph/Private/K2Node_CreateDelegate.cpp
+++ b/Engine/Source/Editor/BlueprintGraph/Private/K2Node_CreateDelegate.cpp
@@ -2,6 +2,9 @@
 #include "BlueprintGraphPrivatePCH.h"
 
 #include "DelegateNodeHandlers.h"
+#include "BlueprintNodeSpawner.h"
+#include "EditorCategoryUtils.h"
+#include "BlueprintActionDatabaseRegistrar.h"
 
 struct FK2Node_CreateDelegate_Helper
 {
@@ -387,4 +390,20 @@ void UK2Node_CreateDelegate::SetFunction(FName Name)
 {
 	SelectedFunctionName = Name;
 	SelectedFunctionGuid.Invalidate();
+}
+
+FText UK2Node_CreateDelegate::GetMenuCategory() const
+{
+	return FEditorCategoryUtils::GetCommonCategory(FCommonEditorCategory::Delegates);
+}
+
+void UK2Node_CreateDelegate::GetMenuActions(FBlueprintActionDatabaseRegistrar& ActionRegistrar) const
+{
+	UClass* NodeClass = GetClass();
+	if (ActionRegistrar.IsOpenForRegistration(NodeClass))
+	{
+		UBlueprintNodeSpawner* NodeSpawner = UBlueprintNodeSpawner::Create(NodeClass);
+		check(NodeSpawner);
+		ActionRegistrar.AddBlueprintAction(NodeClass, NodeSpawner);
+	}
 }
\ No newline at end of file
