BP: K2Node_GetDataTableRow unused member variable removed

diff --git a/Engine/Source/Editor/BlueprintGraph/Classes/K2Node_GetDataTableRow.h b/Engine/Source/Editor/BlueprintGraph/Classes/K2Node_GetDataTableRow.h
--- a/Engine/Source/Editor/BlueprintGraph/Classes/K2Node_GetDataTableRow.h
+++ b/Engine/Source/Editor/BlueprintGraph/Classes/K2Node_GetDataTableRow.h
@@ -12,9 +12,6 @@ class BLUEPRINTGRAPH_API UK2Node_GetDataTableRow : public UK2Node
 {
 	GENERATED_UCLASS_BODY()
 
-	UPROPERTY()
-	FDataTableRowHandle DataTableRowHandle;
-
 	// Begin UEdGraphNode interface.
 	virtual void AllocateDefaultPins() override;
 	virtual FText GetNodeTitle(ENodeTitleType::Type TitleType) const override;
