Adding Interfaces to Blueprints now utilizes the ClassPicker and will display unloaded Interface Blueprints.

Variable Type picking will now display unloaded Blueprint objects, interfaces, and classes, as well as unloaded structs and enums.

diff --git a/Engine/Source/Editor/BlueprintGraph/Classes/EdGraphSchema_K2.h b/Engine/Source/Editor/BlueprintGraph/Classes/EdGraphSchema_K2.h
--- a/Engine/Source/Editor/BlueprintGraph/Classes/EdGraphSchema_K2.h
+++ b/Engine/Source/Editor/BlueprintGraph/Classes/EdGraphSchema_K2.h
@@ -244,6 +244,7 @@ class BLUEPRINTGRAPH_API UEdGraphSchema_K2 : public UEdGraphSchema
 	static const FString PC_Text;
 	static const FString PC_Struct;    // SubCategoryObject is the ScriptStruct of the struct passed thru this pin, 'self' is not a valid SubCategory. DefaultObject should always be empty, the DefaultValue string may be used for supported structs.
 	static const FString PC_Wildcard;    // Special matching rules are imposed by the node itself
+	static const FString PC_Enum;    // SubCategoryObject is the UEnum object passed thru this pin
 
 	// Common PinType.PinSubCategory values
 	static const FString PSC_Self;    // Category=PC_Object or PC_Class, indicates the class being compiled
@@ -801,14 +802,6 @@ public:
 	 */
 	bool DoesTypeHaveSubtypes( const FString& FriendlyTypeName ) const;
 
-	/**
-	 * Gets a list of variable subtypes that are valid for the specified type
-	 *
-	 * @param	Type			The type to grab subtypes for
-	 * @param	SubtypesList	(out) Upon return, this will be a list of valid subtype objects for the specified type
-	 */
-	void GetVariableSubtypes(const FString& Type, TArray<UObject*>& SubtypesList) const;
-
 	/**
 	 * Returns true if the types and directions of two pins are schema compatible. Handles
 	 * outputting a more derived type to an input pin expecting a less derived type.
diff --git a/Engine/Source/Editor/BlueprintGraph/Private/EdGraphSchema_K2.cpp b/Engine/Source/Editor/BlueprintGraph/Private/EdGraphSchema_K2.cpp
--- a/Engine/Source/Editor/BlueprintGraph/Private/EdGraphSchema_K2.cpp
+++ b/Engine/Source/Editor/BlueprintGraph/Private/EdGraphSchema_K2.cpp
@@ -100,8 +100,10 @@ UEdGraphSchema_K2::FPinTypeTreeInfo::FPinTypeTreeInfo(const FText& InFriendlyCat
 	Init(InFriendlyCategoryName, CategoryName, Schema, InTooltip, bInReadOnly);
 }
 
-struct FGatherStructTypesFromAssetsHelper
+/** Helper class to gather variable types */
+class FGatherTypesHelper
 {
+private:
 	typedef TSharedPtr<UEdGraphSchema_K2::FPinTypeTreeInfo> FPinTypeTreeInfoPtr;
 	struct FCompareChildren
 	{
@@ -111,28 +113,186 @@ struct FGatherStructTypesFromAssetsHelper
 		}
 	};
 
-	static void Gather(const FString& CategoryName, const UEdGraphSchema_K2* Schema, TArray<FPinTypeTreeInfoPtr>& OutChildren)
+	/** Helper function to add an unloaded asset to the children, does no validation on passed data */
+	static void AddUnloadedAsset(const FAssetData& InAsset, const FString& InCategoryName, TArray<FPinTypeTreeInfoPtr>& OutChildren)
 	{
+		const FString* TooltipPtr = InAsset.TagsAndValues.Find(TEXT("Tooltip"));
+		const FString Tooltip = (TooltipPtr && !TooltipPtr->IsEmpty()) ? *TooltipPtr : InAsset.ObjectPath.ToString();
+
+		FPinTypeTreeInfoPtr TypeTreeInfo = MakeShareable(new UEdGraphSchema_K2::FPinTypeTreeInfo(InCategoryName, InAsset.ToStringReference(), FText::FromString(Tooltip)));
+		TypeTreeInfo->FriendlyName = FText::FromName(InAsset.AssetName);
+		OutChildren.Add(TypeTreeInfo);
+	}
+
+	/**
+	 * Gets a list of variable subtypes that are valid for the specified type
+	 *
+	 * @param	Type			The type to grab subtypes for
+	 * @param	SubtypesList	(out) Upon return, this will be a list of valid subtype objects for the specified type
+	 */
+	static void GetVariableSubtypes(const FString& Type, TArray<UObject*>& OutSubtypesList)
+	{
+		OutSubtypesList.Empty();
+
+		if (Type == UEdGraphSchema_K2::PC_Struct)
+		{
+			// Find script structs marked with "BlueprintType=true" in their metadata, and add to the list
+			for (TObjectIterator<UScriptStruct> StructIt; StructIt; ++StructIt)
+			{
+				UScriptStruct* ScriptStruct = *StructIt;
+				if (UEdGraphSchema_K2::IsAllowableBlueprintVariableType(ScriptStruct))
+				{
+					OutSubtypesList.Add(ScriptStruct);
+				}
+			}
+		}
+		else if (Type == UEdGraphSchema_K2::PC_Class)
+		{
+			// Generate a list of all potential objects which have "BlueprintType=true" in their metadata
+			for (TObjectIterator<UClass> ClassIt; ClassIt; ++ClassIt)
+			{
+				UClass* CurrentClass = *ClassIt;
+				if (UEdGraphSchema_K2::IsAllowableBlueprintVariableType(CurrentClass) && !CurrentClass->HasAnyClassFlags(CLASS_Deprecated))
+				{
+					OutSubtypesList.Add(CurrentClass);
+				}
+			}
+		}
+		else if (Type == UEdGraphSchema_K2::PC_Object)
+		{
+			// Generate a list of all potential objects which have "BlueprintType=true" in their metadata (that aren't interfaces)
+			for (TObjectIterator<UClass> ClassIt; ClassIt; ++ClassIt)
+			{
+				UClass* CurrentClass = *ClassIt;
+				if (!CurrentClass->IsChildOf(UInterface::StaticClass()) && UEdGraphSchema_K2::IsAllowableBlueprintVariableType(CurrentClass) && !CurrentClass->HasAnyClassFlags(CLASS_Deprecated))
+				{
+					OutSubtypesList.Add(CurrentClass);
+				}
+			}
+		}
+		else if (Type == UEdGraphSchema_K2::PC_Interface)
+		{
+			// Generate a list of all potential objects which have "BlueprintType=true" in their metadata (only ones that are interfaces)
+			for (TObjectIterator<UClass> ClassIt; ClassIt; ++ClassIt)
+			{
+				UClass* CurrentClass = *ClassIt;
+				if (CurrentClass->IsChildOf(UInterface::StaticClass()) && UEdGraphSchema_K2::IsAllowableBlueprintVariableType(CurrentClass))
+				{
+					OutSubtypesList.Add(CurrentClass);
+				}
+			}
+		}
+		else if (Type == UEdGraphSchema_K2::PC_Enum)
+		{
+			// Generate a list of all potential enums which have "BlueprintType=true" in their metadata
+			for (TObjectIterator<UEnum> EnumIt; EnumIt; ++EnumIt)
+			{
+				UEnum* CurrentEnum = *EnumIt;
+				if (UEdGraphSchema_K2::IsAllowableBlueprintVariableType(CurrentEnum))
+				{
+					OutSubtypesList.Add(CurrentEnum);
+				}
+			}
+		}
+	}
+
+public:
+	/**
+	 * Gathers all valid sub-types (loaded and unloaded) of a passed category and sorts them alphabetically
+	 * @param FriendlyName		Friendly name to be used for the tooltip if there is no available data
+	 * @param CategoryName		Category (type) to find sub-types of
+	 * @param Schema			Schema to use
+	 * @param OutChildren		All the gathered children
+	 */
+	static void Gather(const FText& FriendlyName, const FString& CategoryName, const UEdGraphSchema_K2* Schema, TArray<FPinTypeTreeInfoPtr>& OutChildren)
+	{
+		// Gather any loaded subtype children first
+		TArray<UObject*> LoadedSubTypes;
+		GetVariableSubtypes(CategoryName, LoadedSubTypes);
+
+		FEdGraphPinType LoadedPinSubtype;
+		LoadedPinSubtype.PinCategory = (CategoryName == UEdGraphSchema_K2::PC_Enum ? UEdGraphSchema_K2::PC_Byte : CategoryName);
+		LoadedPinSubtype.PinSubCategory = TEXT("");
+		LoadedPinSubtype.PinSubCategoryObject = NULL;
+
+		// Add any loaded subtypes to the children list
+		for (auto it = LoadedSubTypes.CreateIterator(); it; ++it)
+		{
+			FText SubtypeTooltip;
+			UStruct* Struct = Cast<UStruct>(*it);
+			if(Struct != NULL)
+			{
+				SubtypeTooltip = (Struct ? Struct->GetToolTipText() : FriendlyName);
+			}
+			
+			OutChildren.Add( MakeShareable(new UEdGraphSchema_K2::FPinTypeTreeInfo(LoadedPinSubtype.PinCategory, *it, SubtypeTooltip)) );
+		}
+
+		UClass* AssetRegistrySearchClass = nullptr;
+		if(Schema->PC_Struct == CategoryName)
+		{
+			AssetRegistrySearchClass = UUserDefinedStruct::StaticClass();
+		}
+		else if(Schema->PC_Enum == CategoryName)
+		{
+			AssetRegistrySearchClass = UUserDefinedEnum::StaticClass();
+		}
+		else if(Schema->PC_Object == CategoryName || Schema->PC_Class == CategoryName || Schema->PC_Interface == CategoryName)
+		{
+			AssetRegistrySearchClass = UBlueprint::StaticClass();
+		}
+
+		// Gather all assets of the specified class
 		const FAssetRegistryModule& AssetRegistryModule = FModuleManager::LoadModuleChecked<FAssetRegistryModule>(TEXT("AssetRegistry"));
 		TArray<FAssetData> AssetData;
-		AssetRegistryModule.Get().GetAssetsByClass(UUserDefinedStruct::StaticClass()->GetFName(), AssetData);
-		for (int32 AssetIndex = 0; AssetIndex < AssetData.Num(); ++AssetIndex)
+		AssetRegistryModule.Get().GetAssetsByClass(AssetRegistrySearchClass->GetFName(), AssetData);
+
+		// Search through all the unloaded assets, filtered by certain rules based on struct/enum or object/class/interface
+		if (Schema->PC_Struct == CategoryName || Schema->PC_Enum == CategoryName)
 		{
-			const FAssetData& Asset = AssetData[AssetIndex];
-			if (Asset.IsValid() && !Asset.IsAssetLoaded())
+			for (int32 AssetIndex = 0; AssetIndex < AssetData.Num(); ++AssetIndex)
 			{
-				const FString* pDescription = Asset.TagsAndValues.Find(TEXT("Tooltip"));
-				const FString Tooltip = (pDescription && !pDescription->IsEmpty()) ? *pDescription : Asset.ObjectPath.ToString();
-
-				FPinTypeTreeInfoPtr TypeTreeInfo = MakeShareable(new UEdGraphSchema_K2::FPinTypeTreeInfo(CategoryName, Asset.ToStringReference(), FText::FromString(Tooltip)));
-				TypeTreeInfo->FriendlyName = FText::FromName(Asset.AssetName);
-				OutChildren.Add(TypeTreeInfo);
+				const FAssetData& Asset = AssetData[AssetIndex];
+				if (Asset.IsValid() && !Asset.IsAssetLoaded())
+				{
+					AddUnloadedAsset(Asset, CategoryName, OutChildren);
+				}
 			}
 		}
+		else if(Schema->PC_Object == CategoryName || Schema->PC_Class == CategoryName || Schema->PC_Interface == CategoryName)
+		{
+			const FString BPTypeAllowed = (Schema->PC_Interface == CategoryName)? TEXT("BPTYPE_Interface") : TEXT("BPTYPE_Normal");
+
+			for (int32 AssetIndex = 0; AssetIndex < AssetData.Num(); ++AssetIndex)
+			{
+				const FAssetData& Asset = AssetData[AssetIndex];
+
+				if (Asset.IsValid() && !Asset.IsAssetLoaded())
+				{
+					// Based on the category, only certain Blueprint types are available
+					const FString* BlueprintTypeStr = Asset.TagsAndValues.Find("BlueprintType");
+					if(BlueprintTypeStr && *BlueprintTypeStr == BPTypeAllowed)
+					{
+						uint32 ClassFlags = 0;
+						const FString* ClassFlagsStr = Asset.TagsAndValues.Find("ClassFlags");
+						if(ClassFlagsStr)
+						{
+							ClassFlags = FCString::Atoi(**ClassFlagsStr);
+						}
 
+						// Do not allow deprecated Blueprints to be added
+						if(!(ClassFlags & CLASS_Deprecated))
+						{
+							AddUnloadedAsset(Asset, CategoryName, OutChildren);
+						}
+					}
+				}
+			}
+		}
 		OutChildren.Sort(FCompareChildren());
 	}
 
+	/** Loads an asset based on the AssetReference through the asset registry */
 	static UObject* LoadAsset(const FStringAssetReference& AssetReference)
 	{
 		if (AssetReference.IsValid())
@@ -141,15 +301,55 @@ struct FGatherStructTypesFromAssetsHelper
 			const FAssetData AssetData = AssetRegistryModule.Get().GetAssetByObjectPath(*AssetReference.ToString());
 			return AssetData.GetAsset();
 		}
-		return NULL;
+		return nullptr;
 	}
 };
 
 const FEdGraphPinType& UEdGraphSchema_K2::FPinTypeTreeInfo::GetPinType(bool bForceLoadedSubCategoryObject)
 {
-	if (bForceLoadedSubCategoryObject && !PinType.PinSubCategoryObject.IsValid() && SubCategoryObjectAssetReference.IsValid())
+	if (bForceLoadedSubCategoryObject)
 	{
-		PinType.PinSubCategoryObject = FGatherStructTypesFromAssetsHelper::LoadAsset(SubCategoryObjectAssetReference);
+		// Only attempt to load the sub category object if we need to
+		if ( SubCategoryObjectAssetReference.IsValid() && (!PinType.PinSubCategoryObject.IsValid() || FStringAssetReference(PinType.PinSubCategoryObject.Get()) != SubCategoryObjectAssetReference) )
+		{
+			UObject* LoadedObject = FGatherTypesHelper::LoadAsset(SubCategoryObjectAssetReference);
+
+			if(UBlueprint* BlueprintObject = Cast<UBlueprint>(LoadedObject))
+			{
+				PinType.PinSubCategoryObject = *BlueprintObject->GeneratedClass;
+			}
+			else
+			{
+				PinType.PinSubCategoryObject = LoadedObject;
+			}
+		}
+	}
+	else
+	{
+		if (SubCategoryObjectAssetReference.IsValid())
+		{
+			const FAssetRegistryModule& AssetRegistryModule = FModuleManager::LoadModuleChecked<FAssetRegistryModule>(TEXT("AssetRegistry"));
+			const FAssetData AssetData = AssetRegistryModule.Get().GetAssetByObjectPath(*SubCategoryObjectAssetReference.ToString());
+
+			if(!AssetData.IsAssetLoaded())
+			{
+				UObject* LoadedObject = FindObject<UClass>(ANY_PACKAGE, *AssetData.AssetClass.ToString());
+
+				// If the unloaded asset is a Blueprint, we need to pull the generated class and assign that
+				if(UBlueprint* BlueprintObject = Cast<UBlueprint>(LoadedObject))
+				{
+					PinType.PinSubCategoryObject = *BlueprintObject->GeneratedClass;
+				}
+				else
+				{
+					PinType.PinSubCategoryObject = LoadedObject;
+				}
+			}
+			else
+			{
+				PinType.PinSubCategoryObject = AssetData.GetAsset();
+			}
+		}
 	}
 	return PinType;
 }
@@ -161,7 +361,7 @@ void UEdGraphSchema_K2::FPinTypeTreeInfo::Init(const FText& InFriendlyName, cons
 
 	FriendlyName = InFriendlyName;
 	Tooltip = InTooltip;
-	PinType.PinCategory = (CategoryName == TEXT("Enum") ? PC_Byte : CategoryName);
+	PinType.PinCategory = (CategoryName == PC_Enum? PC_Byte : CategoryName);
 	PinType.PinSubCategory = TEXT("");
 	PinType.PinSubCategoryObject = NULL;
 
@@ -169,24 +369,7 @@ void UEdGraphSchema_K2::FPinTypeTreeInfo::Init(const FText& InFriendlyName, cons
 
 	if (Schema->DoesTypeHaveSubtypes(CategoryName))
 	{
-		TArray<UObject*> Subtypes;
-		Schema->GetVariableSubtypes(CategoryName, Subtypes);
-		for (auto it = Subtypes.CreateIterator(); it; ++it)
-		{
-			FText SubtypeTooltip;
-			UStruct* Struct = Cast<UStruct>(*it);
-			if(Struct != NULL)
-			{
-				SubtypeTooltip = (Struct ? Struct->GetToolTipText() : InFriendlyName);
-			}
-
-			Children.Add( MakeShareable(new FPinTypeTreeInfo(PinType.PinCategory, *it, SubtypeTooltip)) );
-		}
-
-		if (Schema->PC_Struct == CategoryName)
-		{
-			FGatherStructTypesFromAssetsHelper::Gather(CategoryName, Schema, Children);
-		}
+		FGatherTypesHelper::Gather(InFriendlyName, CategoryName, Schema, Children);
 	}
 }
 
@@ -254,6 +437,7 @@ const FString UEdGraphSchema_K2::PC_String(TEXT("string"));
 const FString UEdGraphSchema_K2::PC_Text(TEXT("text"));
 const FString UEdGraphSchema_K2::PC_Struct(TEXT("struct"));
 const FString UEdGraphSchema_K2::PC_Wildcard(TEXT("wildcard"));
+const FString UEdGraphSchema_K2::PC_Enum(TEXT("enum"));
 const FString UEdGraphSchema_K2::PSC_Self(TEXT("self"));
 const FString UEdGraphSchema_K2::PSC_Index(TEXT("index"));
 const FString UEdGraphSchema_K2::PN_Execute(TEXT("execute"));
@@ -2944,6 +3128,7 @@ FText UEdGraphSchema_K2::GetCategoryText(const FString& Category, const bool bFo
 		CategoryDescriptions.Add(PC_Text, LOCTEXT("TextCategory","Text"));
 		CategoryDescriptions.Add(PC_Struct, LOCTEXT("StructCategory","Structure"));
 		CategoryDescriptions.Add(PC_Wildcard, LOCTEXT("WildcardCategory","Wildcard"));
+		CategoryDescriptions.Add(PC_Enum, LOCTEXT("EnumCategory","Enum"));
 	}
 
 	if (bForMenu)
@@ -3077,7 +3262,7 @@ void UEdGraphSchema_K2::GetVariableTypeTree( TArray< TSharedPtr<FPinTypeTreeInfo
 	TypeTree.Add( MakeShareable( new FPinTypeTreeInfo(GetCategoryText(PC_Object, true), PC_Object, this, LOCTEXT("ObjectType", "Object pointer."), true) ) );
 	TypeTree.Add( MakeShareable( new FPinTypeTreeInfo(GetCategoryText(PC_Interface, true), PC_Interface, this, LOCTEXT("InterfaceType", "Interface pointer."), true) ) );
 	TypeTree.Add( MakeShareable( new FPinTypeTreeInfo(GetCategoryText(PC_Class, true), PC_Class, this, LOCTEXT("ClassType", "Class pointers."), true) ) );
-	TypeTree.Add( MakeShareable( new FPinTypeTreeInfo(LOCTEXT("EnumCategory", "Enum"), TEXT("Enum"), this, LOCTEXT("EnumType", "Enumeration types."), true) ) );
+	TypeTree.Add( MakeShareable( new FPinTypeTreeInfo(GetCategoryText(PC_Enum, true), PC_Enum, this, LOCTEXT("EnumType", "Enumeration types."), true) ) );
 }
 
 void UEdGraphSchema_K2::GetVariableIndexTypeTree( TArray< TSharedPtr<FPinTypeTreeInfo> >& TypeTree, bool bAllowExec, bool bAllowWildcard ) const
@@ -3101,80 +3286,12 @@ void UEdGraphSchema_K2::GetVariableIndexTypeTree( TArray< TSharedPtr<FPinTypeTre
 	}
 
 	// Add the types that have subtrees
-	TypeTree.Add( MakeShareable( new FPinTypeTreeInfo(LOCTEXT("EnumCategory", "Enum"), TEXT("Enum"), this, LOCTEXT("EnumIndexType", "Enumeration types."), true) ) );
+	TypeTree.Add( MakeShareable( new FPinTypeTreeInfo(GetCategoryText(PC_Enum, true), PC_Enum, this, LOCTEXT("EnumIndexType", "Enumeration types."), true) ) );
 }
 
 bool UEdGraphSchema_K2::DoesTypeHaveSubtypes(const FString& Category) const
 {
-	return (Category == PC_Struct) || (Category == PC_Object) || (Category == PC_Interface) || (Category == PC_Class) || (Category == TEXT("Enum"));
-}
-
-void UEdGraphSchema_K2::GetVariableSubtypes(const FString& Type, TArray<UObject*>& SubtypesList) const
-{
-	SubtypesList.Empty();
-
-	if (Type == PC_Struct)
-	{
-		// Find script structs marked with "BlueprintType=true" in their metadata, and add to the list
-		for (TObjectIterator<UScriptStruct> StructIt; StructIt; ++StructIt)
-		{
-			UScriptStruct* ScriptStruct = *StructIt;
-			if (UEdGraphSchema_K2::IsAllowableBlueprintVariableType(ScriptStruct))
-			{
-				SubtypesList.Add(ScriptStruct);
-			}
-		}
-	}
-	else if (Type == PC_Class)
-	{
-		// Generate a list of all potential objects which have "BlueprintType=true" in their metadata
-		for (TObjectIterator<UClass> ClassIt; ClassIt; ++ClassIt)
-		{
-			UClass* CurrentClass = *ClassIt;
-			if (UEdGraphSchema_K2::IsAllowableBlueprintVariableType(CurrentClass) && !CurrentClass->HasAnyClassFlags(CLASS_Deprecated))
-			{
-				SubtypesList.Add(CurrentClass);
-			}
-		}
-	}
-	else if (Type == PC_Object)
-	{
-		// Generate a list of all potential objects which have "BlueprintType=true" in their metadata (that aren't interfaces)
-		for (TObjectIterator<UClass> ClassIt; ClassIt; ++ClassIt)
-		{
-			UClass* CurrentClass = *ClassIt;
-			if (!CurrentClass->IsChildOf(UInterface::StaticClass()) && UEdGraphSchema_K2::IsAllowableBlueprintVariableType(CurrentClass) && !CurrentClass->HasAnyClassFlags(CLASS_Deprecated))
-			{
-				SubtypesList.Add(CurrentClass);
-			}
-		}
-	}
-	else if (Type == PC_Interface)
-	{
-		// Generate a list of all potential objects which have "BlueprintType=true" in their metadata (only ones that are interfaces)
-		for (TObjectIterator<UClass> ClassIt; ClassIt; ++ClassIt)
-		{
-			UClass* CurrentClass = *ClassIt;
-			if (CurrentClass->IsChildOf(UInterface::StaticClass()) && UEdGraphSchema_K2::IsAllowableBlueprintVariableType(CurrentClass))
-			{
-				SubtypesList.Add(CurrentClass);
-			}
-		}
-	}
-	else if (Type == TEXT("Enum"))
-	{
-		// Generate a list of all potential enums which have "BlueprintType=true" in their metadata
-		for (TObjectIterator<UEnum> EnumIt; EnumIt; ++EnumIt)
-		{
-			UEnum* CurrentEnum = *EnumIt;
-			if (UEdGraphSchema_K2::IsAllowableBlueprintVariableType(CurrentEnum))
-			{
-				SubtypesList.Add(CurrentEnum);
-			}
-		}
-	}
-
-	SubtypesList.Sort();
+	return (Category == PC_Struct) || (Category == PC_Object) || (Category == PC_Interface) || (Category == PC_Class) || (Category == PC_Enum);
 }
 
 struct FWildcardArrayPinHelper
diff --git a/Engine/Source/Editor/Kismet/Private/BlueprintDetailsCustomization.cpp b/Engine/Source/Editor/Kismet/Private/BlueprintDetailsCustomization.cpp
--- a/Engine/Source/Editor/Kismet/Private/BlueprintDetailsCustomization.cpp
+++ b/Engine/Source/Editor/Kismet/Private/BlueprintDetailsCustomization.cpp
@@ -885,7 +885,6 @@ void FBlueprintVarActionDetails::OnVarTypeChanged(const FEdGraphPinType& NewPinT
 
 		if (VarName != NAME_None)
 		{
-			FBlueprintEditorUtils::ChangeMemberVariableType(GetBlueprintObj(), VarName, NewPinType);
 			// Set the MyBP tab's last pin type used as this, for adding lots of variables of the same type
 			MyBlueprint.Pin()->GetLastPinTypeUsed() = NewPinType;
 
@@ -3941,103 +3940,47 @@ void FBlueprintInterfaceLayout::OnRemoveInterface(FInterfaceName InterfaceName)
 	RegenerateChildrenDelegate.ExecuteIfBound();
 }
 
-/** Helper function for the interface menu */
-bool IsInterfaceImplemented(const UBlueprint* Blueprint, const UClass* TestInterface)
+void FBlueprintInterfaceLayout::OnClassPicked(UClass* PickedClass)
 {
-	const auto InterfaceName = TestInterface->GetFName();
-
-	// First look in the blueprint's ImplementedInterfaces list
-	for(TArray<FBPInterfaceDescription>::TConstIterator it(Blueprint->ImplementedInterfaces); it; ++it)
+	if (AddInterfaceComboButton.IsValid())
 	{
-		const FBPInterfaceDescription& CurrentInterface = *it;
-		if( CurrentInterface.Interface && CurrentInterface.Interface->GetFName() == InterfaceName )
-		{
-			return true;
-		}
+		AddInterfaceComboButton->SetIsOpen(false);
 	}
 
-	// Now go up the inheritance tree and see if the interface in implemented via its' ancestry
-	UClass* BlueprintParent =  Blueprint->ParentClass;
-	while (BlueprintParent)
-	{
-		for (TArray<FImplementedInterface>::TIterator It(BlueprintParent->Interfaces); It; ++It)
-		{
-			const FImplementedInterface& CurrentInterface = *It;
-			if (CurrentInterface.Class && (CurrentInterface.Class->GetFName() == InterfaceName))
-			{
-				return true;
-			}
-		}
-		BlueprintParent = BlueprintParent->GetSuperClass();
-	}
+	UBlueprint* Blueprint = GlobalOptionsDetailsPtr.Pin()->GetBlueprintObj();
+	check(Blueprint);
 
-	return false;
+	FBlueprintEditorUtils::ImplementNewInterface( Blueprint, PickedClass->GetFName() );
+
+	RegenerateChildrenDelegate.ExecuteIfBound();
 }
 
 TSharedRef<SWidget> FBlueprintInterfaceLayout::OnGetAddInterfaceMenuContent()
 {
 	UBlueprint* Blueprint = GlobalOptionsDetailsPtr.Pin()->GetBlueprintObj();
 
-	// Generate a list of interfaces that can still be implemented
-	UnimplementedInterfaces.Empty();
-	for (TObjectIterator<UClass> It; It; ++It)
-	{
-		UClass* CurrentInterface = *It;
-		if (FKismetEditorUtilities::CanBlueprintImplementInterface(Blueprint, CurrentInterface) &&
-			!IsInterfaceImplemented(Blueprint, CurrentInterface) &&
-			!FKismetEditorUtilities::IsClassABlueprintSkeleton(CurrentInterface))
-		{
-			UnimplementedInterfaces.Add(MakeShareable(new FInterfaceName(CurrentInterface->GetFName(), CurrentInterface->GetDisplayNameText())));
-		}
-	}
-
-	if (UnimplementedInterfaces.Num() > 0)
-	{
-		return SNew(SListView<TSharedPtr<FInterfaceName>>)
-			.ListItemsSource(&UnimplementedInterfaces)
-			.OnGenerateRow(this, &FBlueprintInterfaceLayout::GenerateInterfaceListRow)
-			.OnSelectionChanged(this, &FBlueprintInterfaceLayout::OnInterfaceListSelectionChanged);
-	}
-	else
-	{
-		return SNew(STextBlock)
-			.Text(LOCTEXT("BlueprintNoInterfacesAvailable", "No interfaces available to implement."));
-	}
-}
-
-TSharedRef<ITableRow> FBlueprintInterfaceLayout::GenerateInterfaceListRow(TSharedPtr<FInterfaceName> InterfaceName, const TSharedRef<STableViewBase>& OwningList)
-{
-	return SNew(STableRow< TSharedPtr<FInterfaceName> >, OwningList)
+	TArray<UBlueprint*> Blueprints;
+	Blueprints.Add(Blueprint);
+	TSharedRef<SWidget> ClassPicker = FBlueprintEditorUtils::ConstructBlueprintInterfaceClassPicker(Blueprints, FOnClassPicked::CreateSP(this, &FBlueprintInterfaceLayout::OnClassPicked));
+	return
+		SNew(SBorder)
+		.BorderImage(FEditorStyle::GetBrush("Menu.Background"))
 		[
-			SNew(STextBlock)
-			.Text(InterfaceName.IsValid() ? InterfaceName->DisplayText : FText())
+			// Achieving fixed width by nesting items within a fixed width box.
+			SNew(SBox)
+			.WidthOverride(350.0f)
+			[
+				SNew(SVerticalBox)
+				+SVerticalBox::Slot()
+				.MaxHeight(400.0f)
+				.AutoHeight()
+				[
+					ClassPicker
+				]
+			]
 		];
 }
 
-void FBlueprintInterfaceLayout::OnInterfaceListSelectionChanged(TSharedPtr<FInterfaceName> Selection, ESelectInfo::Type SelectInfo)
-{
-	if (Selection.IsValid())
-	{
-		UBlueprint* Blueprint = GlobalOptionsDetailsPtr.Pin()->GetBlueprintObj();
-		check(Blueprint);
-
-		FName SelectedInterface = Selection->Name;
-		if (!FBlueprintEditorUtils::ImplementNewInterface( Blueprint, SelectedInterface ) )
-		{
-			GlobalOptionsDetailsPtr.Pin()->GetBlueprintEditorPtr().Pin()->LogSimpleMessage( LOCTEXT("ImplementInterface_Error", "Unable to implement interface. Check log for details") );
-		}
-		
-		if (AddInterfaceComboButton.IsValid())
-		{
-			AddInterfaceComboButton->SetIsOpen(false);
-		}
-
-		RegenerateChildrenDelegate.ExecuteIfBound();
-	}
-}
-
-
-
 UBlueprint* FBlueprintGlobalOptionsDetails::GetBlueprintObj() const
 {
 	if(BlueprintEditorPtr.IsValid())
diff --git a/Engine/Source/Editor/Kismet/Private/BlueprintDetailsCustomization.h b/Engine/Source/Editor/Kismet/Private/BlueprintDetailsCustomization.h
--- a/Engine/Source/Editor/Kismet/Private/BlueprintDetailsCustomization.h
+++ b/Engine/Source/Editor/Kismet/Private/BlueprintDetailsCustomization.h
@@ -547,8 +547,9 @@ private:
 	void OnRemoveInterface(FInterfaceName InterfaceName);
 
 	TSharedRef<SWidget> OnGetAddInterfaceMenuContent();
-	TSharedRef<ITableRow> GenerateInterfaceListRow(TSharedPtr<FInterfaceName> InterfaceName, const TSharedRef<STableViewBase>& OwningList);
-	void OnInterfaceListSelectionChanged(TSharedPtr<FInterfaceName> Selection, ESelectInfo::Type SelectInfo);
+
+	/** Callback function when an interface class is picked */
+	void OnClassPicked(UClass* PickedClass);
 
 private:
 	/** The parent graph action details customization */
diff --git a/Engine/Source/Editor/KismetWidgets/Private/SPinTypeSelector.cpp b/Engine/Source/Editor/KismetWidgets/Private/SPinTypeSelector.cpp
--- a/Engine/Source/Editor/KismetWidgets/Private/SPinTypeSelector.cpp
+++ b/Engine/Source/Editor/KismetWidgets/Private/SPinTypeSelector.cpp
@@ -154,15 +154,15 @@ TSharedRef<ITableRow> SPinTypeSelector::GenerateTypeTreeRow(FPinTypeTreeItem InI
 {
 	const bool bHasChildren = (InItem->Children.Num() > 0);
 	const FText Description = InItem->GetDescription();
+	const FEdGraphPinType& PinType = InItem->GetPinType(false);
 
 	// Determine the best icon the to represents this item
-	const FSlateBrush* IconBrush = GetIconFromPin(InItem->GetPinType(false));
+	const FSlateBrush* IconBrush = GetIconFromPin(PinType);
 
 	// Use tooltip if supplied, otherwise just repeat description
 	const FText OrgTooltip = InItem->GetToolTip();
 	const FText Tooltip = !OrgTooltip.IsEmpty() ? OrgTooltip : Description;
 
-	const FEdGraphPinType& PinType = InItem->GetPinType(true);
 	const FString PinTooltipExcerpt = ((PinType.PinCategory != UEdGraphSchema_K2::PC_Byte || PinType.PinSubCategoryObject == nullptr) ? PinType.PinCategory : TEXT("Enum")); 
 
 	return SNew( SComboRow<FPinTypeTreeItem>, OwnerTree )
@@ -175,7 +175,7 @@ TSharedRef<ITableRow> SPinTypeSelector::GenerateTypeTreeRow(FPinTypeTreeItem InI
 			[
 				SNew(SImage)
 				.Image(IconBrush)
-				.ColorAndOpacity(Schema->GetPinTypeColor(InItem->GetPinType(false)))
+				.ColorAndOpacity(Schema->GetPinTypeColor(PinType))
 				.Visibility( InItem->bReadOnly ? EVisibility::Collapsed : EVisibility::Visible )
 			]
 			+SHorizontalBox::Slot()
diff --git a/Engine/Source/Editor/UnrealEd/Private/Kismet2/BlueprintEditorUtils.cpp b/Engine/Source/Editor/UnrealEd/Private/Kismet2/BlueprintEditorUtils.cpp
--- a/Engine/Source/Editor/UnrealEd/Private/Kismet2/BlueprintEditorUtils.cpp
+++ b/Engine/Source/Editor/UnrealEd/Private/Kismet2/BlueprintEditorUtils.cpp
@@ -42,6 +42,9 @@
 #include "Engine/SCS_Node.h"
 #include "Engine/TimelineTemplate.h"
 
+#include "SNotificationList.h"
+#include "NotificationManager.h"
+
 #define LOCTEXT_NAMESPACE "Blueprint"
 
 DEFINE_LOG_CATEGORY(LogBlueprintDebug);
@@ -3550,7 +3553,8 @@ void FBlueprintEditorUtils::ChangeMemberVariableType(UBlueprint* Blueprint, cons
 				const FScopedTransaction Transaction( LOCTEXT("ChangeVariableType", "Change Variable Type") );
 				Blueprint->Modify();
 
-				Variable.VarType = NewPinType;
+				/** Only change the variable type if type selection is valid, some unloaded Blueprints will turn out to be bad */
+				bool bChangeVariableType = true;
 
 				// Destroy all our nodes, because the pin types could be incorrect now (as will the links)
 				RemoveVariableNodes(Blueprint, VariableName);
@@ -3558,20 +3562,36 @@ void FBlueprintEditorUtils::ChangeMemberVariableType(UBlueprint* Blueprint, cons
 				if ((NewPinType.PinCategory == K2Schema->PC_Object) || (NewPinType.PinCategory == K2Schema->PC_Interface))
 				{
 					// if it's a PC_Object, then it should have an associated UClass object
-					check(NewPinType.PinSubCategoryObject.IsValid());
-					const UClass* ClassObject = Cast<UClass>(NewPinType.PinSubCategoryObject.Get());
-					check(ClassObject != NULL);
-
-					if (ClassObject->IsChildOf(AActor::StaticClass()))
+					if(NewPinType.PinSubCategoryObject.IsValid())
 					{
-						// prevent Actor variables from having default values (because Blueprint templates are library elements that can 
-						// bridge multiple levels and different levels might not have the actor that the default is referencing).
-						Variable.PropertyFlags |= CPF_DisableEditOnTemplate;
+						const UClass* ClassObject = Cast<UClass>(NewPinType.PinSubCategoryObject.Get());
+						check(ClassObject != NULL);
+
+						if (ClassObject->IsChildOf(AActor::StaticClass()))
+						{
+							// prevent Actor variables from having default values (because Blueprint templates are library elements that can 
+							// bridge multiple levels and different levels might not have the actor that the default is referencing).
+							Variable.PropertyFlags |= CPF_DisableEditOnTemplate;
+						}
+						else 
+						{
+							// clear the disable-default-value flag that might have been present (if this was an AActor variable before)
+							Variable.PropertyFlags &= ~(CPF_DisableEditOnTemplate);
+						}
 					}
-					else 
+					else
 					{
-						// clear the disable-default-value flag that might have been present (if this was an AActor variable before)
-						Variable.PropertyFlags &= ~(CPF_DisableEditOnTemplate);
+						bChangeVariableType = false;
+
+						// Display a notification to inform the user that the variable type was invalid (likely due to corruption), it should no longer appear in the list.
+						FNotificationInfo Info( LOCTEXT("InvalidUnloadedBP", "The selected type was invalid once loaded, it has been removed from the list!") );
+						Info.ExpireDuration = 3.0f;
+						Info.bUseLargeFont = false;
+						TSharedPtr<SNotificationItem> Notification = FSlateNotificationManager::Get().AddNotification(Info);
+						if ( Notification.IsValid() )
+						{
+							Notification->SetCompletionState( SNotificationItem::CS_Fail );
+						}
 					}
 				}
 				else 
@@ -3580,6 +3600,11 @@ void FBlueprintEditorUtils::ChangeMemberVariableType(UBlueprint* Blueprint, cons
 					Variable.PropertyFlags &= ~(CPF_DisableEditOnTemplate);
 				}
 
+				if(bChangeVariableType)
+				{
+					Variable.VarType = NewPinType;
+				}
+
 				// And recompile
 				FBlueprintEditorUtils::MarkBlueprintAsStructurallyModified(Blueprint);
 			}
@@ -5821,6 +5846,112 @@ void FBlueprintEditorUtils::OpenReparentBlueprintMenu( const TArray< UBlueprint*
 		);
 }
 
+/** Filter class for ClassPicker handling allowed interfaces for a Blueprint */
+class FBlueprintInterfaceFilter : public IClassViewerFilter
+{
+public:
+	/** All children of these classes will be included unless filtered out by another setting. */
+	TSet< const UClass* > AllowedChildrenOfClasses;
+
+	/** Classes to not allow any children of into the Class Viewer/Picker. */
+	TSet< const UClass* > DisallowedChildrenOfClasses;
+
+	/** Classes to never show in this class viewer. */
+	TSet< const UClass* > DisallowedClasses;
+
+	virtual bool IsClassAllowed(const FClassViewerInitializationOptions& InInitOptions, const UClass* InClass, TSharedRef< FClassViewerFilterFuncs > InFilterFuncs ) override
+	{
+		// If it appears on the allowed child-of classes list (or there is nothing on that list)
+		//		AND it is NOT on the disallowed child-of classes list
+		//		AND it is NOT on the disallowed classes list
+		return InFilterFuncs->IfInChildOfClassesSet( AllowedChildrenOfClasses, InClass) != EFilterReturn::Failed && 
+			InFilterFuncs->IfInChildOfClassesSet(DisallowedChildrenOfClasses, InClass) != EFilterReturn::Passed && 
+			InFilterFuncs->IfInClassesSet(DisallowedClasses, InClass) != EFilterReturn::Passed &&
+			!InClass->HasAnyClassFlags(CLASS_Deprecated | CLASS_NewerVersionExists) &&
+			InClass->HasAnyClassFlags(CLASS_Interface) &&
+			// Here is some loaded classes only logic, Blueprints will never have this info
+			!InClass->HasMetaData(FBlueprintMetadata::MD_CannotImplementInterfaceInBlueprint);
+	}
+
+	virtual bool IsUnloadedClassAllowed(const FClassViewerInitializationOptions& InInitOptions, const TSharedRef< const IUnloadedBlueprintData > InUnloadedClassData, TSharedRef< FClassViewerFilterFuncs > InFilterFuncs) override
+	{
+		// Unloaded interfaces mean they must be Blueprint Interfaces
+
+
+		// If it appears on the allowed child-of classes list (or there is nothing on that list)
+		//		AND it is NOT on the disallowed child-of classes list
+		//		AND it is NOT on the disallowed classes list
+		return InFilterFuncs->IfInChildOfClassesSet( AllowedChildrenOfClasses, InUnloadedClassData) != EFilterReturn::Failed && 
+			InFilterFuncs->IfInChildOfClassesSet(DisallowedChildrenOfClasses, InUnloadedClassData) != EFilterReturn::Passed && 
+			InFilterFuncs->IfInClassesSet(DisallowedClasses, InUnloadedClassData) != EFilterReturn::Passed &&
+			!InUnloadedClassData->HasAnyClassFlags(CLASS_Deprecated | CLASS_NewerVersionExists) &&
+			InUnloadedClassData->HasAnyClassFlags(CLASS_Interface);
+	}
+};
+
+TSharedRef<SWidget> FBlueprintEditorUtils::ConstructBlueprintInterfaceClassPicker( const TArray< UBlueprint* >& Blueprints, const FOnClassPicked& OnPicked)
+{
+	TArray<UClass*> BlueprintClasses;
+	for( auto BlueprintIter = Blueprints.CreateConstIterator(); BlueprintIter; ++BlueprintIter )
+	{
+		const auto Blueprint = *BlueprintIter;
+		BlueprintClasses.Add(Blueprint->GeneratedClass);
+	}
+
+	// Fill in options
+	FClassViewerInitializationOptions Options;
+	Options.Mode = EClassViewerMode::ClassPicker;
+
+	TSharedPtr<FBlueprintInterfaceFilter> Filter = MakeShareable(new FBlueprintInterfaceFilter);
+	Options.ClassFilter = Filter;
+	Options.ViewerTitleString = LOCTEXT("ImplementInterfaceBlueprint", "Implement Interface").ToString();
+
+	for( auto BlueprintIter = Blueprints.CreateConstIterator(); BlueprintIter; ++BlueprintIter )
+	{
+		const auto Blueprint = *BlueprintIter;
+
+		// don't allow making me my own parent!
+		Filter->DisallowedClasses.Add(Blueprint->GeneratedClass);
+
+		UClass const* const ParentClass = Blueprint->ParentClass;
+		// see if the parent class has any prohibited interfaces
+		if ((ParentClass != NULL) && ParentClass->HasMetaData(FBlueprintMetadata::MD_ProhibitedInterfaces))
+		{
+			FString const& ProhibitedList = Blueprint->ParentClass->GetMetaData(FBlueprintMetadata::MD_ProhibitedInterfaces);
+
+			TArray<FString> ProhibitedInterfaceNames;
+			ProhibitedList.ParseIntoArray(&ProhibitedInterfaceNames, TEXT(","), true);
+
+			// loop over all the prohibited interfaces
+			for (int32 ExclusionIndex = 0; ExclusionIndex < ProhibitedInterfaceNames.Num(); ++ExclusionIndex)
+			{
+				FString const& ProhibitedInterfaceName = ProhibitedInterfaceNames[ExclusionIndex].Trim().RightChop(1);
+				UClass* ProhibitedInterface = (UClass*)StaticFindObject(UClass::StaticClass(), ANY_PACKAGE, *ProhibitedInterfaceName);
+				if(ProhibitedInterface)
+				{
+					Filter->DisallowedClasses.Add(ProhibitedInterface);
+					Filter->DisallowedChildrenOfClasses.Add(ProhibitedInterface);
+				}
+			}
+		}
+
+		// Do not allow adding interfaces that are already added to the Blueprint
+		for(TArray<FBPInterfaceDescription>::TConstIterator it(Blueprint->ImplementedInterfaces); it; ++it)
+		{
+			const FBPInterfaceDescription& CurrentInterface = *it;
+			Filter->DisallowedClasses.Add(CurrentInterface.Interface);
+		}
+	}
+
+	// never allow parenting to children of itself
+	for( auto ClassIt = BlueprintClasses.CreateIterator(); ClassIt; ++ClassIt )
+	{
+		Filter->DisallowedChildrenOfClasses.Add(*ClassIt);
+	}
+
+	return FModuleManager::LoadModuleChecked<FClassViewerModule>("ClassViewer").CreateClassViewer(Options, OnPicked);
+}
+
 void FBlueprintEditorUtils::UpdateOldPureFunctions( UBlueprint* Blueprint )
 {
 	if(UBlueprintGeneratedClass* GeneratedClass = Cast<UBlueprintGeneratedClass>(Blueprint->SkeletonGeneratedClass))
diff --git a/Engine/Source/Editor/UnrealEd/Public/Kismet2/BlueprintEditorUtils.h b/Engine/Source/Editor/UnrealEd/Public/Kismet2/BlueprintEditorUtils.h
--- a/Engine/Source/Editor/UnrealEd/Public/Kismet2/BlueprintEditorUtils.h
+++ b/Engine/Source/Editor/UnrealEd/Public/Kismet2/BlueprintEditorUtils.h
@@ -403,6 +403,9 @@ public:
 	static void OpenReparentBlueprintMenu( UBlueprint* Blueprint, const TSharedRef<SWidget>& ParentContent, const FOnClassPicked& OnPicked);
 	static void OpenReparentBlueprintMenu( const TArray< UBlueprint* >& Blueprints, const TSharedRef<SWidget>& ParentContent, const FOnClassPicked& OnPicked);
 
+	/** Constructs a class picker widget for adding interfaces for the specified blueprint(s) */
+	static TSharedRef<SWidget> ConstructBlueprintInterfaceClassPicker( const TArray< UBlueprint* >& Blueprints, const FOnClassPicked& OnPicked);
+
 	/** return find first native class in the hierarchy */
 	static UClass* FindFirstNativeClass(UClass* Class);
 
diff --git a/Engine/Source/Runtime/Engine/Private/Blueprint.cpp b/Engine/Source/Runtime/Engine/Private/Blueprint.cpp
--- a/Engine/Source/Runtime/Engine/Private/Blueprint.cpp
+++ b/Engine/Source/Runtime/Engine/Private/Blueprint.cpp
@@ -675,7 +675,7 @@ void UBlueprint::GetAssetRegistryTags(TArray<FAssetRegistryTag>& OutTags) const
 
 	//NumReplicatedProperties
 	int32 NumReplicatedProperties = 0;
-	UBlueprintGeneratedClass* BlueprintClass = Cast<UBlueprintGeneratedClass>(GenClass);
+	UBlueprintGeneratedClass* BlueprintClass = Cast<UBlueprintGeneratedClass>(SkeletonGeneratedClass);
 	if (BlueprintClass)
 	{
 		NumReplicatedProperties = BlueprintClass->NumReplicatedProperties;
