MakeArray will now always have proper default values
MakeArray will not reset to wildcard when it gets disconnected, however, a ResetToWildcard context menu option is available

diff --git a/Engine/Source/Editor/BlueprintGraph/Classes/K2Node_MakeArray.h b/Engine/Source/Editor/BlueprintGraph/Classes/K2Node_MakeArray.h
--- a/Engine/Source/Editor/BlueprintGraph/Classes/K2Node_MakeArray.h
+++ b/Engine/Source/Editor/BlueprintGraph/Classes/K2Node_MakeArray.h
@@ -49,6 +49,8 @@ protected:
 	/** If needed, will clear all pins to be wildcards */
 	void ClearPinTypeToWildcard();
 
+	bool CanResetToWildcard() const;
+
 	/** Propagates the pin type from the output (array) pin to the inputs, to make sure types are consistent */
 	void PropagatePinType();
 
diff --git a/Engine/Source/Editor/BlueprintGraph/Private/K2Node_MakeArray.cpp b/Engine/Source/Editor/BlueprintGraph/Private/K2Node_MakeArray.cpp
--- a/Engine/Source/Editor/BlueprintGraph/Private/K2Node_MakeArray.cpp
+++ b/Engine/Source/Editor/BlueprintGraph/Private/K2Node_MakeArray.cpp
@@ -94,47 +94,40 @@ UEdGraphPin* UK2Node_MakeArray::GetOutputPin() const
 
 void UK2Node_MakeArray::AllocateDefaultPins()
 {
-	const UEdGraphSchema_K2* K2Schema = GetDefault<UEdGraphSchema_K2>();
-
 	// Create the output pin
-	CreatePin(EGPD_Output, K2Schema->PC_Wildcard, TEXT(""), NULL, true, false, *OutputPinName);
+	CreatePin(EGPD_Output, UEdGraphSchema_K2::PC_Wildcard, TEXT(""), NULL, true, false, *OutputPinName);
 
 	// Create the input pins to create the arrays from
 	for (int32 i = 0; i < NumInputs; ++i)
 	{
-		CreatePin(EGPD_Input, K2Schema->PC_Wildcard, TEXT(""), NULL, false, false, *FString::Printf(TEXT("[%d]"), i));
+		CreatePin(EGPD_Input, UEdGraphSchema_K2::PC_Wildcard, TEXT(""), NULL, false, false, *FString::Printf(TEXT("[%d]"), i));
 	}
 }
 
-void UK2Node_MakeArray::ClearPinTypeToWildcard()
+bool UK2Node_MakeArray::CanResetToWildcard() const
 {
 	bool bClearPinsToWildcard = true;
 
 	// Check to see if we want to clear the wildcards.
-	for (int32 PinIndex = 0; PinIndex < Pins.Num(); ++PinIndex)
+	for (const UEdGraphPin* Pin : Pins)
 	{
-		if(Pins[PinIndex]->LinkedTo.Num() > 0)
+		if( Pin->LinkedTo.Num() > 0 )
 		{
 			// One of the pins is still linked, we will not be clearing the types.
 			bClearPinsToWildcard = false;
 			break;
 		}
-		else if( Pins[PinIndex]->Direction == EGPD_Input )
-		{
-			if( Pins[PinIndex]->GetDefaultAsString().IsEmpty() == false )
-			{
-				// One of the pins has data the user may not want to have cleared, we will not be clearing the types.
-				bClearPinsToWildcard = false;
-				break;
-			}
-		}
 	}
 
-	if( bClearPinsToWildcard == true )
+	return bClearPinsToWildcard;
+}
+
+void UK2Node_MakeArray::ClearPinTypeToWildcard()
+{
+	if( CanResetToWildcard() )
 	{
-		const UEdGraphSchema_K2* Schema = GetDefault<UEdGraphSchema_K2>();
 		UEdGraphPin* OutputPin = GetOutputPin();
-		OutputPin->PinType.PinCategory = Schema->PC_Wildcard;
+		OutputPin->PinType.PinCategory = UEdGraphSchema_K2::PC_Wildcard;
 		OutputPin->PinType.PinSubCategory = TEXT("");
 		OutputPin->PinType.PinSubCategoryObject = NULL;
 
@@ -146,8 +139,6 @@ void UK2Node_MakeArray::NotifyPinConnectionListChanged(UEdGraphPin* Pin)
 {
 	Super::NotifyPinConnectionListChanged(Pin);
 
-	const UEdGraphSchema_K2* Schema = GetDefault<UEdGraphSchema_K2>();
-
 	// Was this the first or last connection?
 	int32 NumPinsWithLinks = 0;
 	// Array to cache the input pins we might want to find these if we are removing the last link
@@ -191,7 +182,7 @@ void UK2Node_MakeArray::NotifyPinConnectionListChanged(UEdGraphPin* Pin)
 
 			if( bResetOutputPin == true )
 			{
-				OutputPin->PinType.PinCategory = Schema->PC_Wildcard;
+				OutputPin->PinType.PinCategory = UEdGraphSchema_K2::PC_Wildcard;
 				OutputPin->PinType.PinSubCategory = TEXT("");
 				OutputPin->PinType.PinSubCategoryObject = NULL;
 
@@ -231,6 +222,15 @@ void UK2Node_MakeArray::PropagatePinType()
 				CurrentPin->PinType.PinSubCategory = OutputPin->PinType.PinSubCategory;
 				CurrentPin->PinType.PinSubCategoryObject = OutputPin->PinType.PinSubCategoryObject;
 
+				if (CurrentPin->PinType.PinCategory == UEdGraphSchema_K2::PC_Wildcard)
+				{
+					CurrentPin->ResetDefaultValue();
+				}
+				else
+				{
+					Schema->SetPinDefaultValueBasedOnType(CurrentPin);
+				}
+
 				// Verify that all previous connections to this pin are still valid with the new type
 				for (TArray<UEdGraphPin*>::TIterator ConnectionIt(CurrentPin->LinkedTo); ConnectionIt; ++ConnectionIt)
 				{
@@ -307,8 +307,7 @@ void UK2Node_MakeArray::AddInputPin()
 	Modify();
 
 	++NumInputs;
-	const UEdGraphSchema_K2* K2Schema = GetDefault<UEdGraphSchema_K2>();
-	CreatePin(EGPD_Input, K2Schema->PC_Wildcard, TEXT(""), NULL, false, false, *FString::Printf(TEXT("[%d]"), (NumInputs-1)));
+	CreatePin(EGPD_Input, UEdGraphSchema_K2::PC_Wildcard, TEXT(""), NULL, false, false, *FString::Printf(TEXT("[%d]"), (NumInputs-1)));
 	PostReconstructNode();
 	
 	const bool bIsCompiling = GetBlueprint()->bBeingCompiled;
@@ -340,8 +339,6 @@ void UK2Node_MakeArray::RemoveInputPin(UEdGraphPin* Pin)
 
 		--NumInputs;
 
-		ClearPinTypeToWildcard();
-
 		FBlueprintEditorUtils::MarkBlueprintAsStructurallyModified(GetBlueprint());
 	}
 }
@@ -352,41 +349,47 @@ void UK2Node_MakeArray::GetContextMenuActions(const FGraphNodeContextMenuBuilder
 
 	if (!Context.bIsDebugging)
 	{
+		Context.MenuBuilder->BeginSection("K2NodeMakeArray", NSLOCTEXT("K2Nodes", "MakeArrayHeader", "MakeArray"));
+
 		if (Context.Pin != NULL)
 		{
 			// we only do this for normal BlendList/BlendList by enum, BlendList by Bool doesn't support add/remove pins
 			if (Context.Pin->Direction == EGPD_Input)
 			{
 				//@TODO: Only offer this option on arrayed pins
-				Context.MenuBuilder->BeginSection("K2NodeMakeArray", NSLOCTEXT("K2Nodes", "BlendListHeader", "MakeArray"));
-				{
-					Context.MenuBuilder->AddMenuEntry(
-						LOCTEXT("RemovePin", "Remove array element pin"),
-						LOCTEXT("RemovePinTooltip", "Remove this array element pin"),
-						FSlateIcon(),
-						FUIAction(
-							FExecuteAction::CreateUObject(this, &UK2Node_MakeArray::RemoveInputPin, const_cast<UEdGraphPin*>(Context.Pin))
-						)
-					);
-				}
-				Context.MenuBuilder->EndSection();
-			}
-		}
-		else
-		{
-			Context.MenuBuilder->BeginSection("K2NodeMakeArray", NSLOCTEXT("K2Nodes", "MakeArrayHeader", "MakeArray"));
-			{
 				Context.MenuBuilder->AddMenuEntry(
-					LOCTEXT("AddPin", "Add array element pin"),
-					LOCTEXT("AddPinTooltip", "Add another array element pin"),
+					LOCTEXT("RemovePin", "Remove array element pin"),
+					LOCTEXT("RemovePinTooltip", "Remove this array element pin"),
 					FSlateIcon(),
 					FUIAction(
-						FExecuteAction::CreateUObject(this, &UK2Node_MakeArray::AddInputPin)
+						FExecuteAction::CreateUObject(this, &UK2Node_MakeArray::RemoveInputPin, const_cast<UEdGraphPin*>(Context.Pin))
 					)
 				);
 			}
-			Context.MenuBuilder->EndSection();
 		}
+		else
+		{
+			Context.MenuBuilder->AddMenuEntry(
+				LOCTEXT("AddPin", "Add array element pin"),
+				LOCTEXT("AddPinTooltip", "Add another array element pin"),
+				FSlateIcon(),
+				FUIAction(
+					FExecuteAction::CreateUObject(this, &UK2Node_MakeArray::AddInputPin)
+				)
+			);
+		}
+
+		Context.MenuBuilder->AddMenuEntry(
+			LOCTEXT("ResetToWildcard", "Reset to wildcard"),
+			LOCTEXT("ResetToWildcardTooltip", "Reset the node to have wildcard input/outputs. Requires no pins are connected."),
+			FSlateIcon(),
+			FUIAction(
+				FExecuteAction::CreateUObject(this, &UK2Node_MakeArray::ClearPinTypeToWildcard),
+				FCanExecuteAction::CreateUObject(this, &UK2Node_MakeArray::CanResetToWildcard)
+			)
+		);
+
+		Context.MenuBuilder->EndSection();
 	}
 }
 
