Fixing possible crash.

Expression would be evaluated as (OutputClass && OutputClass->ImplementsInterface(InterfaceClass) || OutputClass->IsChildOf(InterfaceClass), so if OutputClass were null it would crash.

diff --git a/Engine/Source/Editor/BlueprintGraph/Private/EdGraphSchema_K2.cpp b/Engine/Source/Editor/BlueprintGraph/Private/EdGraphSchema_K2.cpp
--- a/Engine/Source/Editor/BlueprintGraph/Private/EdGraphSchema_K2.cpp
+++ b/Engine/Source/Editor/BlueprintGraph/Private/EdGraphSchema_K2.cpp
@@ -3444,7 +3444,7 @@ bool UEdGraphSchema_K2::ArePinTypesCompatible(const FEdGraphPinType& Output, con
 			OutputClass = CallingContext;
 		}
 
-		return OutputClass && OutputClass->ImplementsInterface(InterfaceClass) || OutputClass->IsChildOf(InterfaceClass);
+		return OutputClass && (OutputClass->ImplementsInterface(InterfaceClass) || OutputClass->IsChildOf(InterfaceClass));
 	}
 
 	// Pins representing BLueprint objects and subclass of UObject can match when EditoronlyBP.bAllowClassAndBlueprintPinMatching=true (BaseEngine.ini)
