#ue4 UE-5440:  Fix delegates not being able to be assigned in LSBPs by coercing the type from self as needed in autowire new node

diff --git a/Engine/Source/Editor/BlueprintGraph/Classes/K2Node_BaseMCDelegate.h b/Engine/Source/Editor/BlueprintGraph/Classes/K2Node_BaseMCDelegate.h
--- a/Engine/Source/Editor/BlueprintGraph/Classes/K2Node_BaseMCDelegate.h
+++ b/Engine/Source/Editor/BlueprintGraph/Classes/K2Node_BaseMCDelegate.h
@@ -25,6 +25,7 @@ public:
 	virtual void ExpandNode(class FKismetCompilerContext& CompilerContext, UEdGraph* SourceGraph) override;
 	virtual bool HasExternalBlueprintDependencies(TArray<class UStruct*>* OptionalOutput) const override;
 	virtual void GetNodeAttributes( TArray<TKeyValuePair<FString, FString>>& OutNodeAttributes ) const override;
+	virtual void AutowireNewNode(UEdGraphPin* FromPin) override;
 	// End of UK2Node interface
 
 	// UEdGraphNode interface
diff --git a/Engine/Source/Editor/BlueprintGraph/Private/K2Node_MCDelegate.cpp b/Engine/Source/Editor/BlueprintGraph/Private/K2Node_MCDelegate.cpp
--- a/Engine/Source/Editor/BlueprintGraph/Private/K2Node_MCDelegate.cpp
+++ b/Engine/Source/Editor/BlueprintGraph/Private/K2Node_MCDelegate.cpp
@@ -210,6 +210,55 @@ void UK2Node_BaseMCDelegate::GetNodeAttributes( TArray<TKeyValuePair<FString, FS
 	OutNodeAttributes.Add( TKeyValuePair<FString, FString>( TEXT( "Name" ), GetPropertyName().ToString() ));
 }
 
+void UK2Node_BaseMCDelegate::AutowireNewNode(UEdGraphPin* FromPin)
+{
+	const UEdGraphSchema_K2* K2Schema = CastChecked<UEdGraphSchema_K2>(GetSchema());
+
+	// Since nodes no longer have a sense of scope when they're placed, look at the connection we're coming from, and use that to coerce the Target pin
+	if (FromPin != nullptr)
+	{
+		bool bConnected = false;
+
+		// Only do the fixup if we're going coming from an output pin, which implies a contextual drag
+		if (FromPin->Direction == EGPD_Output)
+		{
+			if (FromPin->PinType.PinSubCategoryObject.IsValid() && FromPin->PinType.PinSubCategoryObject->IsA(UClass::StaticClass()))
+			{
+				UProperty* DelegateProperty = DelegateReference.ResolveMember<UProperty>(this);
+				if (DelegateProperty)
+				{
+					UClass* DelegateOwner = DelegateProperty->GetOwnerClass();
+					if (FromPin->PinType.PinSubCategoryObject == DelegateOwner || dynamic_cast<UClass*>(FromPin->PinType.PinSubCategoryObject.Get())->IsChildOf(DelegateOwner))
+					{
+						// If we get here, then the property delegate is also available on the FromPin's class.
+						// Fix up the type by propagating it from the FromPin to our target pin
+						UEdGraphPin* TargetPin = FindPin(K2Schema->PN_Self);
+						TargetPin->PinType.PinSubCategory.Empty();
+						TargetPin->PinType.PinSubCategoryObject = DelegateOwner;
+
+						// And finally, try to establish the connection
+						if (K2Schema->TryCreateConnection(FromPin, TargetPin))
+						{
+							bConnected = true;
+
+							DelegateReference.SetFromField<UProperty>(DelegateProperty, false);
+							TargetPin->bHidden = false;
+							FromPin->GetOwningNode()->NodeConnectionListChanged();
+							this->NodeConnectionListChanged();
+						}
+
+					}
+				}
+			}
+		}
+
+		if (!bConnected)
+		{
+			Super::AutowireNewNode(FromPin);
+		}
+	}
+}
+
 /////// UK2Node_AddDelegate ///////////
 
 UK2Node_AddDelegate::UK2Node_AddDelegate(const FObjectInitializer& ObjectInitializer)
