Improved error message - github PR#573 from user stephenwittle

diff --git a/Engine/Source/Editor/BlueprintGraph/Private/K2Node_Message.cpp b/Engine/Source/Editor/BlueprintGraph/Private/K2Node_Message.cpp
--- a/Engine/Source/Editor/BlueprintGraph/Private/K2Node_Message.cpp
+++ b/Engine/Source/Editor/BlueprintGraph/Private/K2Node_Message.cpp
@@ -126,7 +126,7 @@ void UK2Node_Message::ExpandNode(class FKismetCompilerContext& CompilerContext,
 		UEdGraphPin* MessageSelfPin = Schema->FindSelfPin(*this, EGPD_Input);
 		if( !MessageSelfPin || MessageSelfPin->LinkedTo.Num() == 0 )
 		{
-			CompilerContext.MessageLog.Error(*FString::Printf(*LOCTEXT("MessageNodeSelfPin_Error", "Message node @@ must have a self pin connection.").ToString()), this);
+			CompilerContext.MessageLog.Error(*FString::Printf(*LOCTEXT("MessageNodeSelfPin_Error", "Message node @@ must have a valid target or reference to self.").ToString()), this);
 			return;
 		}
 
