Delegate Pin in Event Node  uses MemberReference. Uncompatible delegate pins cannot be connected.

diff --git a/Engine/Source/Editor/BlueprintGraph/Classes/K2Node_Event.h b/Engine/Source/Editor/BlueprintGraph/Classes/K2Node_Event.h
--- a/Engine/Source/Editor/BlueprintGraph/Classes/K2Node_Event.h
+++ b/Engine/Source/Editor/BlueprintGraph/Classes/K2Node_Event.h
@@ -71,7 +71,7 @@ class UK2Node_Event : public UK2Node_EditablePinBase
 	BLUEPRINTGRAPH_API virtual bool IsFunctionEntryCompatible(const class UK2Node_FunctionEntry* EntryNode) const;
 	
 	UFunction* FindEventSignatureFunction();
-	BLUEPRINTGRAPH_API void UpdateDelegatePin();
+	BLUEPRINTGRAPH_API void UpdateDelegatePin(bool bSilent = false);
 	BLUEPRINTGRAPH_API FName GetFunctionName() const;
 	BLUEPRINTGRAPH_API virtual bool IsUsedByAuthorityOnlyDelegate() const { return false; }
 	BLUEPRINTGRAPH_API virtual bool IsCosmeticTickEvent() const;
diff --git a/Engine/Source/Editor/BlueprintGraph/Private/EdGraphSchema_K2.cpp b/Engine/Source/Editor/BlueprintGraph/Private/EdGraphSchema_K2.cpp
--- a/Engine/Source/Editor/BlueprintGraph/Private/EdGraphSchema_K2.cpp
+++ b/Engine/Source/Editor/BlueprintGraph/Private/EdGraphSchema_K2.cpp
@@ -3342,7 +3342,9 @@ bool UEdGraphSchema_K2::ArePinTypesCompatible(const FEdGraphPinType& Output, con
 	}
 	else if (Output.PinCategory == Input.PinCategory)
 	{
-		if ((Output.PinSubCategory == Input.PinSubCategory) && (Output.PinSubCategoryObject == Input.PinSubCategoryObject))
+		if ((Output.PinSubCategory == Input.PinSubCategory) 
+			&& (Output.PinSubCategoryObject == Input.PinSubCategoryObject)
+			&& (Output.PinSubCategoryMemberReference == Input.PinSubCategoryMemberReference))
 		{
 			return true;
 		}
@@ -3404,7 +3406,25 @@ bool UEdGraphSchema_K2::ArePinTypesCompatible(const FEdGraphPinType& Output, con
 		else if (PC_Delegate == Output.PinCategory || PC_MCDelegate == Output.PinCategory)
 		{
 			const UFunction* OutFunction = FMemberReference::ResolveSimpleMemberReference<UFunction>(Output.PinSubCategoryMemberReference);
+			if (!OutFunction && Output.PinSubCategoryMemberReference.MemberParentClass)
+			{
+				const auto ParentClass = Output.PinSubCategoryMemberReference.MemberParentClass;
+				const auto BPOwner = Cast<UBlueprint>(ParentClass->ClassGeneratedBy);
+				if (BPOwner && BPOwner->SkeletonGeneratedClass && (BPOwner->SkeletonGeneratedClass != ParentClass))
+				{
+					OutFunction = BPOwner->SkeletonGeneratedClass->FindFunctionByName(Output.PinSubCategoryMemberReference.MemberName);
+				}
+			}
 			const UFunction* InFunction = FMemberReference::ResolveSimpleMemberReference<UFunction>(Input.PinSubCategoryMemberReference);
+			if (!InFunction && Input.PinSubCategoryMemberReference.MemberParentClass)
+			{
+				const auto ParentClass = Input.PinSubCategoryMemberReference.MemberParentClass;
+				const auto BPOwner = Cast<UBlueprint>(ParentClass->ClassGeneratedBy);
+				if (BPOwner && BPOwner->SkeletonGeneratedClass && (BPOwner->SkeletonGeneratedClass != ParentClass))
+				{
+					InFunction = BPOwner->SkeletonGeneratedClass->FindFunctionByName(Input.PinSubCategoryMemberReference.MemberName);
+				}
+			}
 			return !OutFunction || !InFunction || OutFunction->IsSignatureCompatibleWith(InFunction);
 		}
 	}
diff --git a/Engine/Source/Editor/BlueprintGraph/Private/K2Node_Event.cpp b/Engine/Source/Editor/BlueprintGraph/Private/K2Node_Event.cpp
--- a/Engine/Source/Editor/BlueprintGraph/Private/K2Node_Event.cpp
+++ b/Engine/Source/Editor/BlueprintGraph/Private/K2Node_Event.cpp
@@ -160,28 +160,33 @@ void UK2Node_Event::PostReconstructNode()
 	UpdateDelegatePin();
 }
 
-void UK2Node_Event::UpdateDelegatePin()
+void UK2Node_Event::UpdateDelegatePin(bool bSilent)
 {
 	UEdGraphPin* Pin = FindPinChecked(DelegateOutputName);
 	checkSlow(EGPD_Output == Pin->Direction);
-	const UObject* OldSignature = Pin->PinType.PinSubCategoryObject.Get();
 
+	const UObject* OldSignature = FMemberReference::ResolveSimpleMemberReference<UFunction>(Pin->PinType.PinSubCategoryMemberReference);
+	if (!OldSignature)
+	{
+		OldSignature = Pin->PinType.PinSubCategoryObject.Get();
+	}
+
+	UFunction* NewSignature = NULL;
 	if(bOverrideFunction)
 	{
-		Pin->PinType.PinSubCategoryObject = EventSignatureClass->FindFunctionByName(EventSignatureName);
+		NewSignature = EventSignatureClass->FindFunctionByName(EventSignatureName);
 	}
 	else if(UBlueprint* Blueprint = GetBlueprint())
 	{
-		Pin->PinType.PinSubCategoryObject = Blueprint->SkeletonGeneratedClass 
+		NewSignature = Blueprint->SkeletonGeneratedClass
 			? Blueprint->SkeletonGeneratedClass->FindFunctionByName(CustomFunctionName)
 			: NULL;
 	}
-	else
-	{
-		Pin->PinType.PinSubCategoryObject = NULL;
-	}
 
-	if(OldSignature != Pin->PinType.PinSubCategoryObject.Get())
+	Pin->PinType.PinSubCategoryObject = NULL;
+	FMemberReference::FillSimpleMemberReference<UFunction>(NewSignature, Pin->PinType.PinSubCategoryMemberReference);
+
+	if ((OldSignature != NewSignature) && !bSilent)
 	{
 		PinTypeChanged(Pin);
 	}
@@ -234,6 +239,8 @@ void UK2Node_Event::AllocateDefaultPins()
 		CreatePinsForFunctionEntryExit(Function, /*bIsFunctionEntry=*/ true);
 	}
 
+	UpdateDelegatePin(true);
+
 	Super::AllocateDefaultPins();
 }
 
