Do not swap what Mac control/command map to at the Platform level and instead manage it at the modifier level
Simplify/fix FText name generation for modifier keys
Have FInputGesture store a bit mask and expose helper functions instead of public boolean parameters (required to switch meaning of command/control on Mac)

diff --git a/Engine/Source/Editor/BlueprintGraph/Classes/K2Node_InputKey.h b/Engine/Source/Editor/BlueprintGraph/Classes/K2Node_InputKey.h
--- a/Engine/Source/Editor/BlueprintGraph/Classes/K2Node_InputKey.h
+++ b/Engine/Source/Editor/BlueprintGraph/Classes/K2Node_InputKey.h
@@ -27,7 +27,7 @@ class UK2Node_InputKey : public UK2Node
 	UPROPERTY(EditAnywhere, Category="Input")
 	uint32 bOverrideParentBinding:1;
 
-	// Does this binding require the control key to be held
+	// Does this binding require the control key on PC or the command key on Mac to be held
 	UPROPERTY(EditAnywhere, Category="Modifier")
 	uint32 bControl:1;
 
@@ -39,12 +39,13 @@ class UK2Node_InputKey : public UK2Node
 	UPROPERTY(EditAnywhere, Category="Modifier")
 	uint32 bShift:1;
 
-	// Does this binding require the command key to be held
+	// Does this binding require the windows key on PC or the control key on Mac to be held
 	UPROPERTY(EditAnywhere, Category="Modifier")
 	uint32 bCommand:1;
 
 	// Begin UObject interface
 	virtual void PostLoad() override;
+	virtual void PostEditChangeProperty(struct FPropertyChangedEvent& PropertyChangedEvent) override;
 	// End UObject interface
 
 	// Begin UK2Node interface.
diff --git a/Engine/Source/Editor/BlueprintGraph/Private/K2Node_InputKey.cpp b/Engine/Source/Editor/BlueprintGraph/Private/K2Node_InputKey.cpp
--- a/Engine/Source/Editor/BlueprintGraph/Private/K2Node_InputKey.cpp
+++ b/Engine/Source/Editor/BlueprintGraph/Private/K2Node_InputKey.cpp
@@ -18,6 +18,17 @@ UK2Node_InputKey::UK2Node_InputKey(const FObjectInitializer& ObjectInitializer)
 {
 	bConsumeInput = true;
 	bOverrideParentBinding = true;
+
+#if PLATFORM_MAC && WITH_EDITOR
+	if (IsTemplate())
+	{
+		UProperty* ControlProp = GetClass()->FindPropertyByName(GET_MEMBER_NAME_CHECKED(UK2Node_InputKey, bControl));
+		UProperty* CommandProp = GetClass()->FindPropertyByName(GET_MEMBER_NAME_CHECKED(UK2Node_InputKey, bCommand));
+
+		ControlProp->SetMetaData(TEXT("DisplayName"), TEXT("Command"));
+		CommandProp->SetMetaData(TEXT("DisplayName"), TEXT("Control"));
+	}
+#endif
 }
 
 void UK2Node_InputKey::PostLoad()
@@ -31,6 +42,14 @@ void UK2Node_InputKey::PostLoad()
 	}
 }
 
+void UK2Node_InputKey::PostEditChangeProperty(FPropertyChangedEvent& PropertyChangedEvent)
+{
+	Super::PostEditChangeProperty(PropertyChangedEvent);
+
+	CachedNodeTitle.Clear();
+	CachedTooltip.Clear();
+}
+
 void UK2Node_InputKey::AllocateDefaultPins()
 {
 	const UEdGraphSchema_K2* K2Schema = GetDefault<UEdGraphSchema_K2>();
@@ -48,179 +67,82 @@ FLinearColor UK2Node_InputKey::GetNodeTitleColor() const
 
 FName UK2Node_InputKey::GetModifierName() const
 {
-	if ( bControl && !bAlt && !bShift && !bCommand )
-	{
-		return FName("Ctrl");
-	}
-	else if ( bControl && bAlt && !bShift && !bCommand )
-	{
-		return FName("Ctrl+Alt");
-	}
-	else if ( bControl && !bAlt && bShift && !bCommand )
-	{
-		return FName("Ctrl+Shift");
-	}
-	else if ( bControl && !bAlt && !bShift && bCommand )
-	{
-		return FName("Ctrl+Cmd");
-	}
-	else if ( bControl && bAlt && bShift && !bCommand )
-	{
-		return FName("Ctrl+Alt+Shift");
-	}
-	else if ( bControl && bAlt && !bShift && bCommand )
-	{
-		return FName("Ctrl+Cmd+Alt");
-	}
-	else if ( bControl && !bAlt && bShift && bCommand )
-	{
-		return FName("Ctrl+Cmd+Shift");
-	}
-	else if ( !bControl && bAlt && bShift && bCommand )
-	{
-		return FName("Cmd+Alt+Shift");
-	}
-	else if ( bControl && bAlt && bShift && bCommand )
-	{
-		return FName("Ctrl+Cmd+Alt+Shift");
-	}
-	else if ( !bControl && bAlt && !bShift && !bCommand )
-	{
-		return FName("Alt");
-	}
-	else if ( !bControl && bAlt && bShift && !bCommand )
-	{
-		return FName("Alt+Shift");
-	}
-	else if ( !bControl && bAlt && !bShift && bCommand )
-	{
-		return FName("Cmd+Alt");
-	}
-	else if ( !bControl && !bAlt && bShift && bCommand )
-	{
-		return FName("Cmd+Shift");
-	}
-	else if ( !bControl && !bAlt && bShift && !bCommand )
-	{
-		return FName("Shift");
-	}
-	else if ( !bControl && !bAlt && !bShift && bCommand )
-	{
-		return FName("Cmd");
-	}
-
-	return NAME_None;
+    FString ModName;
+    bool bPlus = false;
+    if (bControl)
+    {
+        ModName += TEXT("Ctrl");
+        bPlus = true;
+    }
+    if (bCommand)
+    {
+        if (bPlus) ModName += TEXT("+");
+        ModName += TEXT("Cmd");
+        bPlus = true;
+    }
+    if (bAlt)
+    {
+        if (bPlus) ModName += TEXT("+");
+        ModName += TEXT("Alt");
+        bPlus = true;
+    }
+    if (bShift)
+    {
+        if (bPlus) ModName += TEXT("+");
+        ModName += TEXT("Shift");
+        bPlus = true;
+    }
+
+	return FName(*ModName);
 }
 
 FText UK2Node_InputKey::GetModifierText() const
 {
-	//@todo This should be unified with other places in the editor [10/11/2013 justin.sargent]
-	if ( bControl && !bAlt && !bShift && !bCommand )
-	{
-#if PLATFORM_MAC
-		return NSLOCTEXT("UK2Node_InputKey", "KeyName_Command", "Cmd");
-#else
-		return NSLOCTEXT("UK2Node_InputKey", "KeyName_Control", "Ctrl");
-#endif
-	}
-	else if ( bControl && bAlt && !bShift && !bCommand )
-	{
 #if PLATFORM_MAC
-		return NSLOCTEXT("UK2Node_InputKey", "KeyName_Command + KeyName_Alt", "Cmd+Alt");
+    const FText CommandText = NSLOCTEXT("UK2Node_InputKey", "KeyName_Control", "Ctrl");
+    const FText ControlText = NSLOCTEXT("UK2Node_InputKey", "KeyName_Command", "Cmd");
 #else
-		return NSLOCTEXT("UK2Node_InputKey", "KeyName_Control + KeyName_Alt", "Ctrl+Alt");
+    const FText ControlText = NSLOCTEXT("UK2Node_InputKey", "KeyName_Control", "Ctrl");
+    const FText CommandText = NSLOCTEXT("UK2Node_InputKey", "KeyName_Command", "Cmd");
 #endif
+    const FText AltText = NSLOCTEXT("UK2Node_InputKey", "KeyName_Alt", "Alt");
+    const FText ShiftText = NSLOCTEXT("UK2Node_InputKey", "KeyName_Shift", "Shift");
+    
+	const FText AppenderText = NSLOCTEXT("UK2Node_InputKey", "ModAppender", "+");
+
+	FFormatNamedArguments Args;
+	int32 ModCount = 0;
+
+    if (bControl)
+    {
+		Args.Add(FString::Printf(TEXT("Mod%d"),++ModCount), ControlText);
+    }
+    if (bCommand)
+    {
+		Args.Add(FString::Printf(TEXT("Mod%d"),++ModCount), CommandText);
+    }
+    if (bAlt)
+    {
+		Args.Add(FString::Printf(TEXT("Mod%d"),++ModCount), AltText);
+    }
+    if (bShift)
+    {
+		Args.Add(FString::Printf(TEXT("Mod%d"),++ModCount), ShiftText);
+    }
+
+	for (int32 i = 1; i <= 4; ++i)
+	{
+		if (i > ModCount)
+		{
+			Args.Add(FString::Printf(TEXT("Mod%d"), i), FText::GetEmpty());
+		}
+
+		Args.Add(FString::Printf(TEXT("Appender%d"), i), (i < ModCount ? AppenderText : FText::GetEmpty()));
 	}
-	else if ( bControl && !bAlt && bShift && !bCommand )
-	{
-#if PLATFORM_MAC
-		return NSLOCTEXT("UK2Node_InputKey", "KeyName_Command + KeyName_Shift", "Cmd+Shift");
-#else
-		return NSLOCTEXT("UK2Node_InputKey", "KeyName_Control + KeyName_Shift", "Ctrl+Shift");
-#endif
-	}
-	else if ( bControl && !bAlt && !bShift && bCommand )
-	{
-#if PLATFORM_MAC
-		return NSLOCTEXT("UK2Node_InputKey", "KeyName_Command + KeyName_Control", "Cmd+Ctrl");
-#else
-		return NSLOCTEXT("UK2Node_InputKey", "KeyName_Control + KeyName_Command", "Ctrl+Cmd");
-#endif
-	}
-	else if ( !bControl && bAlt && !bShift && bCommand )
-	{
-#if PLATFORM_MAC
-		return NSLOCTEXT("UK2Node_InputKey", "KeyName_Control + KeyName_Alt", "Ctrl+Alt");
-#else
-		return NSLOCTEXT("UK2Node_InputKey", "KeyName_Command + KeyName_Alt", "Cmd+Alt");
-#endif
-	}
-	else if ( !bControl && !bAlt && bShift && bCommand )
-	{
-#if PLATFORM_MAC
-		return NSLOCTEXT("UK2Node_InputKey", "KeyName_Control + KeyName_Shift", "Ctrl+Shift");
-#else
-		return NSLOCTEXT("UK2Node_InputKey", "KeyName_Command + KeyName_Shift", "Cmd+Shift");
-#endif
-	}
-	else if ( bControl && bAlt && bShift && !bCommand )
-	{
-#if PLATFORM_MAC
-		return NSLOCTEXT("UK2Node_InputKey", "KeyName_Command + KeyName_Alt + KeyName_Shift", "Cmd+Alt+Shift");
-#else
-		return NSLOCTEXT("UK2Node_InputKey", "KeyName_Control + KeyName_Alt + KeyName_Shift", "Ctrl+Alt+Shift");
-#endif
-	}
-	else if ( bControl && bAlt && !bShift && bCommand )
-	{
-#if PLATFORM_MAC
-		return NSLOCTEXT("UK2Node_InputKey", "KeyName_Command + KeyName_Control + KeyName_Alt", "Cmd+Ctrl+Alt");
-#else
-		return NSLOCTEXT("UK2Node_InputKey", "KeyName_Control + KeyName_Command + KeyName_Alt", "Ctrl+Cmd+Alt");
-#endif
-	}
-	else if ( bControl && !bAlt && bShift && bCommand )
-	{
-#if PLATFORM_MAC
-		return NSLOCTEXT("UK2Node_InputKey", "KeyName_Command + KeyName_Control + KeyName_Shift", "Cmd+Ctrl+Shift");
-#else
-		return NSLOCTEXT("UK2Node_InputKey", "KeyName_Control + KeyName_Command + KeyName_Shift", "Ctrl+Cmd+Shift");
-#endif
-	}
-	else if ( bControl && bAlt && bShift && bCommand )
-	{
-#if PLATFORM_MAC
-		return NSLOCTEXT("UK2Node_InputKey", "KeyName_Command + KeyName_Control + KeyName_Alt + KeyName_Shift", "Cmd+Ctrl+Alt+Shift");
-#else
-		return NSLOCTEXT("UK2Node_InputKey", "KeyName_Control + KeyName_Command + KeyName_Alt + KeyName_Shift", "Ctrl+Cmd+Alt+Shift");
-#endif
-	}
-	else if ( !bControl && bAlt && !bShift && !bCommand )
-	{
-		return NSLOCTEXT("UK2Node_InputKey", "KeyName_Alt", "Alt");
-	}
-	else if ( !bControl && bAlt && bShift && !bCommand )
-	{
-		return NSLOCTEXT("UK2Node_InputKey", "KeyName_Alt + KeyName_Shift", "Alt+Shift");
-	}
-	else if ( !bControl && bAlt && !bShift && bCommand )
-	{
-#if PLATFORM_MAC
-		return NSLOCTEXT("UK2Node_InputKey", "KeyName_Control + KeyName_Alt", "Ctrl+Alt");
-#else
-		return NSLOCTEXT("UK2Node_InputKey", "KeyName_Command + KeyName_Alt", "Cmd+Alt");
-#endif
-	}
-	else if ( !bControl && !bAlt && bShift && !bCommand )
-	{
-		return NSLOCTEXT("UK2Node_InputKey", "KeyName_Shift", "Shift");
-	}
-	else if ( !bControl && !bAlt && !bShift && bCommand )
-	{
-		return NSLOCTEXT("UK2Node_InputKey", "KeyName_Command", "Cmd");
-	}
 
-	return FText::GetEmpty();
+	Args.Add(TEXT("Key"), GetKeyText());
+
+	return FText::Format(NSLOCTEXT("UK2Node_InputKey", "NodeTitle", "{Mod1}{Appender1}{Mod2}{Appender2}{Mod3}{Appender3}{Mod4}"), Args);
 }
 
 FText UK2Node_InputKey::GetKeyText() const
@@ -230,7 +152,7 @@ FText UK2Node_InputKey::GetKeyText() const
 
 FText UK2Node_InputKey::GetNodeTitle(ENodeTitleType::Type TitleType) const
 {
-	if (bControl || bAlt || bShift)
+	if (bControl || bAlt || bShift || bCommand)
 	{
 		if (CachedNodeTitle.IsOutOfDate())
 		{
diff --git a/Engine/Source/Editor/GraphEditor/Private/SNodePanel.cpp b/Engine/Source/Editor/GraphEditor/Private/SNodePanel.cpp
--- a/Engine/Source/Editor/GraphEditor/Private/SNodePanel.cpp
+++ b/Engine/Source/Editor/GraphEditor/Private/SNodePanel.cpp
@@ -854,11 +854,7 @@ FReply SNodePanel::OnMouseButtonUp( const FGeometry& MyGeometry, const FPointerE
 					OnSpawnNodeByShortcut.Execute(LastKeyGestureDetected, PanelCoordToGraphCoord(  MyGeometry.AbsoluteToLocal( MouseEvent.GetScreenSpacePosition() ) ));
 				}
 
-				LastKeyGestureDetected.Key = EKeys::Invalid;
-				LastKeyGestureDetected.bAlt = false;
-				LastKeyGestureDetected.bCtrl = false;
-				LastKeyGestureDetected.bShift = false;
-				LastKeyGestureDetected.bCmd = false;
+				LastKeyGestureDetected = FInputGesture();
 			}
 		}
 		else if ( Marquee.IsValid() )
@@ -923,11 +919,7 @@ FReply SNodePanel::OnKeyDown( const FGeometry& MyGeometry, const FKeyEvent& InKe
 {
 	if( IsEditable.Get() )
 	{
-		LastKeyGestureDetected.Key = InKeyEvent.GetKey();
-		LastKeyGestureDetected.bAlt = InKeyEvent.IsAltDown();
-		LastKeyGestureDetected.bCtrl = InKeyEvent.IsControlDown();
-		LastKeyGestureDetected.bShift = InKeyEvent.IsShiftDown();
-		LastKeyGestureDetected.bCmd = InKeyEvent.IsCommandDown();
+		LastKeyGestureDetected = FInputGesture(InKeyEvent.GetKey(), EModifierKey::FromBools(InKeyEvent.IsControlDown(), InKeyEvent.IsAltDown(), InKeyEvent.IsShiftDown(), InKeyEvent.IsCommandDown()));
 	}
 
 	return FReply::Unhandled();
@@ -937,11 +929,7 @@ FReply SNodePanel::OnKeyUp( const FGeometry& MyGeometry, const FKeyEvent& InKeyE
 {
 	if(LastKeyGestureDetected.Key == InKeyEvent.GetKey())
 	{
-		LastKeyGestureDetected.Key = EKeys::Invalid;
-		LastKeyGestureDetected.bAlt = false;
-		LastKeyGestureDetected.bCtrl = false;
-		LastKeyGestureDetected.bShift = false;
-		LastKeyGestureDetected.bCmd = false;
+		LastKeyGestureDetected = FInputGesture();
 	}
 
 	return FReply::Unhandled();
@@ -949,11 +937,7 @@ FReply SNodePanel::OnKeyUp( const FGeometry& MyGeometry, const FKeyEvent& InKeyE
 
 void SNodePanel::OnFocusLost( const FFocusEvent& InFocusEvent )
 {
-	LastKeyGestureDetected.Key = EKeys::Invalid;
-	LastKeyGestureDetected.bAlt = false;
-	LastKeyGestureDetected.bCtrl = false;
-	LastKeyGestureDetected.bShift = false;
-	LastKeyGestureDetected.bCmd = false;
+	LastKeyGestureDetected = FInputGesture();
 }
 
 FReply SNodePanel::OnTouchGesture( const FGeometry& MyGeometry, const FPointerEvent& GestureEvent )
diff --git a/Engine/Source/Editor/InputBindingEditor/Private/Widgets/SGestureEditor.cpp b/Engine/Source/Editor/InputBindingEditor/Private/Widgets/SGestureEditor.cpp
--- a/Engine/Source/Editor/InputBindingEditor/Private/Widgets/SGestureEditor.cpp
+++ b/Engine/Source/Editor/InputBindingEditor/Private/Widgets/SGestureEditor.cpp
@@ -54,10 +54,7 @@ FReply SGestureEditor::OnKeyDown( const FGeometry& MyGeometry, const FKeyEvent&
 	
 		StartChangingText();
 
-		EditingInputGesture.bCtrl = InKeyEvent.IsControlDown();
-		EditingInputGesture.bAlt = InKeyEvent.IsAltDown();
-		EditingInputGesture.bShift = InKeyEvent.IsShiftDown();
-		EditingInputGesture.bCmd = InKeyEvent.IsCommandDown();
+		EditingInputGesture.ModifierKeys = EModifierKey::FromBools(InKeyEvent.IsControlDown(), InKeyEvent.IsAltDown(), InKeyEvent.IsShiftDown(), InKeyEvent.IsCommandDown());
 
 		LoadText();
 		//@todo checking the length of localized string is not valid, at the very least in this manner [10/11/2013 justin.sargent]
diff --git a/Engine/Source/Editor/Kismet/Private/BlueprintEditorCommands.cpp b/Engine/Source/Editor/Kismet/Private/BlueprintEditorCommands.cpp
--- a/Engine/Source/Editor/Kismet/Private/BlueprintEditorCommands.cpp
+++ b/Engine/Source/Editor/Kismet/Private/BlueprintEditorCommands.cpp
@@ -374,11 +374,7 @@ void FBlueprintSpawnNodeCommands::RegisterCommands()
 				FParse::Bool(*NodeSpawns[x], TEXT("Cmd="), bCmd);
 			}
 
-			FInputGesture Gesture(Key);
-			Gesture.bAlt = bAlt;
-			Gesture.bCtrl = bCtrl;
-			Gesture.bShift = bShift;
-			Gesture.bCmd = bCmd;
+			FInputGesture Gesture(Key, EModifierKey::FromBools(bCtrl, bAlt, bShift, bCmd));
 
 			FText CommandLabelText = FText::FromString( CommandLabel );
 			FText Description = FText::Format( NSLOCTEXT("BlueprintEditor", "NodeSpawnDescription", "Hold down the bound keys and left click in the graph panel to spawn a {0} node."), CommandLabelText );
diff --git a/Engine/Source/Editor/LevelEditor/Private/SLevelViewport.cpp b/Engine/Source/Editor/LevelEditor/Private/SLevelViewport.cpp
--- a/Engine/Source/Editor/LevelEditor/Private/SLevelViewport.cpp
+++ b/Engine/Source/Editor/LevelEditor/Private/SLevelViewport.cpp
@@ -3329,10 +3329,7 @@ FText SLevelViewport::GetMouseCaptureLabelText() const
 			if (TargetPlayer && TargetPlayer->PlayerController && TargetPlayer->PlayerController->PlayerInput)
 			{
 				FKeyBind Binding = TargetPlayer->PlayerController->PlayerInput->GetExecBind(TEXT("ShowMouseCursor"));
-				Gesture.Key = Binding.Key;
-				Gesture.bAlt = Binding.Alt;
-				Gesture.bCtrl = Binding.Control;
-				Gesture.bShift = Binding.Shift;
+				Gesture = FInputGesture(Binding.Key, EModifierKey::FromBools(Binding.Control, Binding.Alt, Binding.Shift, Binding.Cmd));
 				bInitedGesture = true;
 			}
 		}
diff --git a/Engine/Source/Editor/MaterialEditor/Private/MaterialEditorActions.cpp b/Engine/Source/Editor/MaterialEditor/Private/MaterialEditorActions.cpp
--- a/Engine/Source/Editor/MaterialEditor/Private/MaterialEditorActions.cpp
+++ b/Engine/Source/Editor/MaterialEditor/Private/MaterialEditorActions.cpp
@@ -127,10 +127,7 @@ void FMaterialEditorSpawnNodeCommands::RegisterCommands()
 				FParse::Bool(*NodeSpawns[x], TEXT("Ctrl="), bCtrl);
 			}
 
-			FInputGesture Gesture(Key);
-			Gesture.bAlt = bAlt;
-			Gesture.bCtrl = bCtrl;
-			Gesture.bShift = bShift;
+			FInputGesture Gesture(Key, EModifierKey::FromBools(bCtrl, bAlt, bShift, false));
 
 			const FText CommandLabelText = FText::FromString( CommandLabel );
 			const FText Description = FText::Format( NSLOCTEXT("MaterialEditor", "NodeSpawnDescription", "Hold down the bound keys and left click in the graph panel to spawn a {0} node."), CommandLabelText);
diff --git a/Engine/Source/Editor/UnrealEd/Private/LevelEditorViewport.cpp b/Engine/Source/Editor/UnrealEd/Private/LevelEditorViewport.cpp
--- a/Engine/Source/Editor/UnrealEd/Private/LevelEditorViewport.cpp
+++ b/Engine/Source/Editor/UnrealEd/Private/LevelEditorViewport.cpp
@@ -2227,10 +2227,10 @@ static bool CommandAcceptsInput( FLevelEditorViewportClient& ViewportClient, FKe
 {
 	const FInputGesture& Gesture = *Command->GetActiveGesture();
 
-	return (!Gesture.bCtrl	|| ViewportClient.IsCtrlPressed() ) 
-		&& (!Gesture.bAlt	|| ViewportClient.IsAltPressed() ) 
-		&& (!Gesture.bShift || ViewportClient.IsShiftPressed() ) 
-		&& (!Gesture.bCmd	|| ViewportClient.IsCmdPressed() )
+	return (!Gesture.NeedsControl()	|| ViewportClient.IsCtrlPressed() ) 
+		&& (!Gesture.NeedsAlt()		|| ViewportClient.IsAltPressed() ) 
+		&& (!Gesture.NeedsShift()	|| ViewportClient.IsShiftPressed() ) 
+		&& (!Gesture.NeedsCommand()		|| ViewportClient.IsCmdPressed() )
 		&& Gesture.Key == Key;
 }
 
diff --git a/Engine/Source/Editor/UnrealEd/Private/SnappingUtils.cpp b/Engine/Source/Editor/UnrealEd/Private/SnappingUtils.cpp
--- a/Engine/Source/Editor/UnrealEd/Private/SnappingUtils.cpp
+++ b/Engine/Source/Editor/UnrealEd/Private/SnappingUtils.cpp
@@ -136,9 +136,9 @@ bool FEditorViewportSnapping::IsSnapToVertexEnabled()
 
 		const FInputGesture& Gesture = *Commands.HoldToEnableVertexSnapping->GetActiveGesture();
 
-		return (Gesture.bCtrl == GCurrentLevelEditingViewportClient->IsCtrlPressed() ) 
-			&& (Gesture.bAlt ==  GCurrentLevelEditingViewportClient->IsAltPressed() ) 
-			&& (Gesture.bShift == GCurrentLevelEditingViewportClient->IsShiftPressed() ) 
+		return (Gesture.NeedsControl() == GCurrentLevelEditingViewportClient->IsCtrlPressed() ) 
+			&& (Gesture.NeedsAlt() ==  GCurrentLevelEditingViewportClient->IsAltPressed() ) 
+			&& (Gesture.NeedsShift() == GCurrentLevelEditingViewportClient->IsShiftPressed() ) 
 			&& GCurrentLevelEditingViewportClient->Viewport->KeyState(Gesture.Key) == true;
 	}
 	else
diff --git a/Engine/Source/Editor/UnrealEd/Private/Tests/EditorBuildPromotionTests.cpp b/Engine/Source/Editor/UnrealEd/Private/Tests/EditorBuildPromotionTests.cpp
--- a/Engine/Source/Editor/UnrealEd/Private/Tests/EditorBuildPromotionTests.cpp
+++ b/Engine/Source/Editor/UnrealEd/Private/Tests/EditorBuildPromotionTests.cpp
@@ -358,7 +358,7 @@ namespace EditorBuildPromotionTestUtils
 			FSlateApplication::Get().SetKeyboardFocus(FocusWidget.ToSharedRef(), EFocusCause::SetDirectly);
 
 			//Send the command
-			FModifierKeysState ModifierKeys(InGesture.bShift, false, InGesture.bCtrl, false, InGesture.bAlt, false, InGesture.bCmd, false, false);
+			FModifierKeysState ModifierKeys(InGesture.NeedsShift(), false, InGesture.NeedsControl(), false, InGesture.NeedsAlt(), false, InGesture.NeedsCommand(), false, false);
 			FKeyEvent KeyEvent(InGesture.Key, ModifierKeys, 0/*UserIndex*/, false, 0, 0);
 			FSlateApplication::Get().ProcessKeyDownEvent(KeyEvent);
 			FSlateApplication::Get().ProcessKeyUpEvent(KeyEvent);
diff --git a/Engine/Source/Runtime/Core/Private/Mac/MacPlatformMisc.cpp b/Engine/Source/Runtime/Core/Private/Mac/MacPlatformMisc.cpp
--- a/Engine/Source/Runtime/Core/Private/Mac/MacPlatformMisc.cpp
+++ b/Engine/Source/Runtime/Core/Private/Mac/MacPlatformMisc.cpp
@@ -545,16 +545,15 @@ uint32 FMacPlatformMisc::GetKeyMap( uint16* KeyCodes, FString* KeyNames, uint32
 		ADDKEYMAP( kVK_F11, TEXT("F11") );
 		ADDKEYMAP( kVK_F12, TEXT("F12") );
 
-		// Mac pretends the Command key is Ctrl and Ctrl is Command key
-		ADDKEYMAP( MMK_RightCommand, TEXT("RightControl") );
-		ADDKEYMAP( MMK_LeftCommand, TEXT("LeftControl") );
+		ADDKEYMAP( MMK_RightControl, TEXT("RightControl") );
+		ADDKEYMAP( MMK_LeftControl, TEXT("LeftControl") );
 		ADDKEYMAP( MMK_LeftShift, TEXT("LeftShift") );
 		ADDKEYMAP( MMK_CapsLock, TEXT("CapsLock") );
 		ADDKEYMAP( MMK_LeftAlt, TEXT("LeftAlt") );
-		ADDKEYMAP( MMK_LeftControl, TEXT("LeftCommand") );
+		ADDKEYMAP( MMK_LeftCommand, TEXT("LeftCommand") );
 		ADDKEYMAP( MMK_RightShift, TEXT("RightShift") );
 		ADDKEYMAP( MMK_RightAlt, TEXT("RightAlt") );
-		ADDKEYMAP( MMK_RightControl, TEXT("RightCommand") );
+		ADDKEYMAP( MMK_RightCommand, TEXT("RightCommand") );
 	}
 
 	check(NumMappings < MaxMappings);
diff --git a/Engine/Source/Runtime/Core/Public/GenericPlatform/GenericApplication.h b/Engine/Source/Runtime/Core/Public/GenericPlatform/GenericApplication.h
--- a/Engine/Source/Runtime/Core/Public/GenericPlatform/GenericApplication.h
+++ b/Engine/Source/Runtime/Core/Public/GenericPlatform/GenericApplication.h
@@ -30,6 +30,17 @@ namespace EModifierKey
 
 	/** Cmd key (Control key on Mac, Win key on Windows) */
 	const Type Command = 1 << 3;
+
+	FORCEINLINE EModifierKey::Type FromBools(const bool bControl, const bool bAlt, const bool bShift, const bool bCommand)
+	{
+		EModifierKey::Type ModifierMask = EModifierKey::None;
+		if (bControl)	ModifierMask |= EModifierKey::Control;
+		if (bAlt)		ModifierMask |= EModifierKey::Alt;
+		if (bShift)		ModifierMask |= EModifierKey::Shift;
+		if (bCommand)	ModifierMask |= EModifierKey::Command;
+
+		return ModifierMask;
+	}
 };
 
 namespace EPopUpOrientation
diff --git a/Engine/Source/Runtime/Engine/Private/UserInterface/PlayerInput.cpp b/Engine/Source/Runtime/Engine/Private/UserInterface/PlayerInput.cpp
--- a/Engine/Source/Runtime/Engine/Private/UserInterface/PlayerInput.cpp
+++ b/Engine/Source/Runtime/Engine/Private/UserInterface/PlayerInput.cpp
@@ -1461,7 +1461,11 @@ bool UPlayerInput::IsAltPressed() const
 
 bool UPlayerInput::IsCtrlPressed() const
 {
+#if PLATFORM_MAC
+	return IsPressed(EKeys::LeftCommand) || IsPressed(EKeys::RightCommand);
+#else
 	return IsPressed(EKeys::LeftControl) || IsPressed(EKeys::RightControl);
+#endif
 }
 
 bool UPlayerInput::IsShiftPressed() const
@@ -1471,7 +1475,11 @@ bool UPlayerInput::IsShiftPressed() const
 
 bool UPlayerInput::IsCmdPressed() const
 {
+#if PLATFORM_MAC
+	return IsPressed(EKeys::LeftControl) || IsPressed(EKeys::RightControl);
+#else
 	return IsPressed(EKeys::LeftCommand) || IsPressed(EKeys::RightCommand);
+#endif
 }
 
 void UPlayerInput::ConditionalInitAxisProperties()
diff --git a/Engine/Source/Runtime/Slate/Private/Framework/Commands/InputBindingManager.cpp b/Engine/Source/Runtime/Slate/Private/Framework/Commands/InputBindingManager.cpp
--- a/Engine/Source/Runtime/Slate/Private/Framework/Commands/InputBindingManager.cpp
+++ b/Engine/Source/Runtime/Slate/Private/Framework/Commands/InputBindingManager.cpp
@@ -87,11 +87,12 @@ void FUserDefinedGestures::LoadGestures()
 					const FUserDefinedGestureKey GestureKey(BindingContext, CommandName);
 					FInputGesture& UserDefinedGesture = Gestures->FindOrAdd(GestureKey);
 
-					UserDefinedGesture.bCtrl = CtrlObj->AsBool();
-					UserDefinedGesture.bAlt = AltObj->AsBool();
-					UserDefinedGesture.bShift = ShiftObj->AsBool();
-					UserDefinedGesture.bCmd = CmdObj.IsValid() ? CmdObj->AsBool() : false; // Old config files may not have this
-					UserDefinedGesture.Key = *KeyObj->AsString();
+					UserDefinedGesture = FInputGesture(*KeyObj->AsString(), EModifierKey::FromBools(
+																				CtrlObj->AsBool(), 
+																				AltObj->AsBool(), 
+																				ShiftObj->AsBool(), 
+																				CmdObj.IsValid() ? CmdObj->AsBool() : false // Old config files may not have this
+																			));
 				}
 			}
 		}
@@ -140,11 +141,12 @@ void FUserDefinedGestures::LoadGestures()
 						const FUserDefinedGestureKey GestureKey(BindingContext, CommandName);
 						FInputGesture& UserDefinedGesture = Gestures->FindOrAdd(GestureKey);
 
-						UserDefinedGesture.bCtrl = CtrlObj->AsBool();
-						UserDefinedGesture.bAlt = AltObj->AsBool();
-						UserDefinedGesture.bShift = ShiftObj->AsBool();
-						UserDefinedGesture.bCmd = CmdObj.IsValid() ? CmdObj->AsBool() : false; // Old config files may not have this
-						UserDefinedGesture.Key = *KeyObj->AsString();
+						UserDefinedGesture = FInputGesture(*KeyObj->AsString(), EModifierKey::FromBools(
+																					CtrlObj->AsBool(), 
+																					AltObj->AsBool(), 
+																					ShiftObj->AsBool(), 
+																					CmdObj.IsValid() ? CmdObj->AsBool() : false // Old config files may not have this
+																				));
 					}
 				}
 			}
@@ -166,10 +168,10 @@ void FUserDefinedGestures::SaveGestures() const
 			// Set the gesture values for the command
 			GestureInfoObj->Values.Add( TEXT("BindingContext"), MakeShareable( new FJsonValueString( GestureInfo.Key.BindingContext.ToString() ) ) );
 			GestureInfoObj->Values.Add( TEXT("CommandName"), MakeShareable( new FJsonValueString( GestureInfo.Key.CommandName.ToString() ) ) );
-			GestureInfoObj->Values.Add( TEXT("Control"), MakeShareable( new FJsonValueBoolean( GestureInfo.Value.bCtrl ) ) );
-			GestureInfoObj->Values.Add( TEXT("Alt"), MakeShareable( new FJsonValueBoolean( GestureInfo.Value.bAlt ) ) );
-			GestureInfoObj->Values.Add( TEXT("Shift"), MakeShareable( new FJsonValueBoolean( GestureInfo.Value.bShift ) ) );
-			GestureInfoObj->Values.Add( TEXT("Command"), MakeShareable( new FJsonValueBoolean( GestureInfo.Value.bCmd ) ) );
+			GestureInfoObj->Values.Add( TEXT("Control"), MakeShareable( new FJsonValueBoolean( GestureInfo.Value.NeedsControl() ) ) );
+			GestureInfoObj->Values.Add( TEXT("Alt"), MakeShareable( new FJsonValueBoolean( GestureInfo.Value.NeedsAlt() ) ) );
+			GestureInfoObj->Values.Add( TEXT("Shift"), MakeShareable( new FJsonValueBoolean( GestureInfo.Value.NeedsShift() ) ) );
+			GestureInfoObj->Values.Add( TEXT("Command"), MakeShareable( new FJsonValueBoolean( GestureInfo.Value.NeedsCommand() ) ) );
 			GestureInfoObj->Values.Add( TEXT("Key"), MakeShareable( new FJsonValueString( GestureInfo.Value.Key.ToString() ) ) );
 
 			auto JsonWriter = TJsonWriterFactory< TCHAR, TCondensedJsonPrintPolicy<TCHAR> >::Create( &GestureRawJsonContent );
diff --git a/Engine/Source/Runtime/Slate/Private/Framework/Commands/InputGesture.cpp b/Engine/Source/Runtime/Slate/Private/Framework/Commands/InputGesture.cpp
--- a/Engine/Source/Runtime/Slate/Private/Framework/Commands/InputGesture.cpp
+++ b/Engine/Source/Runtime/Slate/Private/Framework/Commands/InputGesture.cpp
@@ -15,121 +15,55 @@
  */
 FText FInputGesture::GetInputText( ) const
 {
-	FText FormatString = LOCTEXT("ShortcutKeyWithNoAdditionalModifierKeys", "{Key}");
-
-	if ( bCtrl && !bAlt && !bShift && !bCmd )
-	{
-#if PLATFORM_MAC
-		FormatString = LOCTEXT("KeyName_Command", "Cmd+{Key}");
-#else
-		FormatString = LOCTEXT("KeyName_Control", "Ctrl+{Key}");
-#endif
-	}
-	else if ( bCtrl && bAlt && !bShift && !bCmd )
-	{
-#if PLATFORM_MAC
-		FormatString = LOCTEXT("KeyName_Command + KeyName_Alt", "Cmd+Alt+{Key}");
-#else
-		FormatString = LOCTEXT("KeyName_Control + KeyName_Alt", "Ctrl+Alt+{Key}");
-#endif
-	}
-	else if ( bCtrl && !bAlt && bShift && !bCmd )
-	{
 #if PLATFORM_MAC
-		FormatString = LOCTEXT("KeyName_Command + KeyName_Shift", "Cmd+Shift+{Key}");
+    const FText CommandText = LOCTEXT("KeyName_Control", "Ctrl");
+    const FText ControlText = LOCTEXT("KeyName_Command", "Cmd");
 #else
-		FormatString = LOCTEXT("KeyName_Control + KeyName_Shift", "Ctrl+Shift+{Key}");
+    const FText ControlText = LOCTEXT("KeyName_Control", "Ctrl");
+    const FText CommandText = LOCTEXT("KeyName_Command", "Cmd"); 
 #endif
-	}
-	else if ( bCtrl && !bAlt && !bShift && bCmd )
-	{
-#if PLATFORM_MAC
-		FormatString = LOCTEXT("KeyName_Command + KeyName_Control", "Cmd+Ctrl+{Key}");
-#else
-		FormatString = LOCTEXT("KeyName_Control + KeyName_Command", "Ctrl+Cmd+{Key}");
-#endif
-	}
-	else if ( !bCtrl && bAlt && !bShift && bCmd )
-	{
-#if PLATFORM_MAC
-		FormatString = LOCTEXT("KeyName_Control + KeyName_Alt", "Ctrl+Alt+{Key}");
-#else
-		FormatString = LOCTEXT("KeyName_Command + KeyName_Alt", "Cmd+Alt+{Key}");
-#endif
-	}
-	else if ( !bCtrl && !bAlt && bShift && bCmd )
-	{
-#if PLATFORM_MAC
-		FormatString = LOCTEXT("KeyName_Control + KeyName_Shift", "Ctrl+Shift+{Key}");
-#else
-		FormatString = LOCTEXT("KeyName_Command + KeyName_Shift", "Cmd+Shift+{Key}");
-#endif
-	}
-	else if ( bCtrl && bAlt && bShift && !bCmd )
-	{
-#if PLATFORM_MAC
-		FormatString = LOCTEXT("KeyName_Command + KeyName_Alt + KeyName_Shift", "Cmd+Alt+Shift+{Key}");
-#else
-		FormatString = LOCTEXT("KeyName_Control + KeyName_Alt + KeyName_Shift", "Ctrl+Alt+Shift+{Key}");
-#endif
-	}
-	else if ( !bCtrl && bAlt && bShift && bCmd )
-	{
-#if PLATFORM_MAC
-		FormatString = LOCTEXT("KeyName_Control + KeyName_Alt + KeyName_Shift", "Ctrl+Alt+Shift+{Key}");
-#else
-		FormatString = LOCTEXT("KeyName_Command + KeyName_Alt + KeyName_Shift", "Cmd+Alt+Shift+{Key}");
-#endif
-	}
-	else if ( bCtrl && !bAlt && bShift && bCmd )
-	{
-#if PLATFORM_MAC
-		FormatString = LOCTEXT("KeyName_Command + KeyName_Control + KeyName_Shift", "Cmd+Ctrl+Shift+{Key}");
-#else
-		FormatString = LOCTEXT("KeyName_Control + KeyName_Command + KeyName_Shift", "Ctrl+Cmd+Shift+{Key}");
-#endif
-	}
-	else if ( bCtrl && bAlt && !bShift && bCmd )
-	{
-#if PLATFORM_MAC
-		FormatString = LOCTEXT("KeyName_Command + KeyName_Control + KeyName_Alt", "Cmd+Ctrl+Alt+{Key}");
-#else
-		FormatString = LOCTEXT("KeyName_Control + KeyName_Command + KeyName_Alt", "Ctrl+Cmd+Alt+{Key}");
-#endif
-	}
-	else if ( bCtrl && bAlt && bShift && bCmd )
-	{
-#if PLATFORM_MAC
-		FormatString = LOCTEXT("KeyName_Command + KeyName_Control + KeyName_Alt + KeyName_Shift", "Cmd+Ctrl+Alt+Shift{Key}");
-#else
-		FormatString = LOCTEXT("KeyName_Control + KeyName_Command + KeyName_Alt + KeyName_Shift", "Ctrl+Cmd+Alt+Shift{Key}");
-#endif
-	}
-	else if ( !bCtrl && bAlt && !bShift && !bCmd )
-	{
-		FormatString = LOCTEXT("KeyName_Alt", "Alt+{Key}");
-	}
-	else if ( !bCtrl && bAlt && bShift && !bCmd )
-	{
-		FormatString = LOCTEXT("KeyName_Alt + KeyName_Shift", "Alt+Shift+{Key}");
-	}
-	else if ( !bCtrl && !bAlt && bShift && !bCmd )
-	{
-		FormatString = LOCTEXT("KeyName_Shift", "Shift+{Key}");
-	}
-	else if ( !bCtrl && !bAlt && !bShift && bCmd )
-	{
-#if PLATFORM_MAC
-		FormatString = LOCTEXT("KeyName_Control", "Ctrl+{Key}");
-#else
-		FormatString = LOCTEXT("KeyName_Command", "Cmd+{Key}");
-#endif
-	}
+    const FText AltText = LOCTEXT("KeyName_Alt", "Alt");
+    const FText ShiftText = LOCTEXT("KeyName_Shift", "Shift");
+    
+	const FText AppenderText = LOCTEXT("ModAppender", "+");
 
 	FFormatNamedArguments Args;
-	Args.Add( TEXT("Key"), GetKeyText() );
+	int32 ModCount = 0;
+
+    if (ModifierKeys & EModifierKey::Control)
+    {
+		Args.Add(FString::Printf(TEXT("Mod%d"),++ModCount), ControlText);
+    }
+    if (ModifierKeys & EModifierKey::Command)
+    {
+		Args.Add(FString::Printf(TEXT("Mod%d"),++ModCount), CommandText);
+    }
+    if (ModifierKeys & EModifierKey::Alt)
+    {
+		Args.Add(FString::Printf(TEXT("Mod%d"),++ModCount), AltText);
+    }
+    if (ModifierKeys & EModifierKey::Shift)
+    {
+		Args.Add(FString::Printf(TEXT("Mod%d"),++ModCount), ShiftText);
+    }
+
+	for (int32 i = 1; i <= 4; ++i)
+	{
+		if (i > ModCount)
+		{
+			Args.Add(FString::Printf(TEXT("Mod%d"), i), FText::GetEmpty());
+			Args.Add(FString::Printf(TEXT("Appender%d"), i), FText::GetEmpty());
+		}
+		else
+		{
+			Args.Add(FString::Printf(TEXT("Appender%d"), i), AppenderText);
+		}
+
+	}
+
+	Args.Add(TEXT("Key"), GetKeyText());
 
-	return FText::Format( FormatString, Args );
+	return FText::Format(LOCTEXT("FourModifiers", "{Mod1}{Appender1}{Mod2}{Appender2}{Mod3}{Appender3}{Mod4}{Appender4}{Key}"), Args);
 }
 
 
diff --git a/Engine/Source/Runtime/Slate/Private/Framework/Commands/UICommandList.cpp b/Engine/Source/Runtime/Slate/Private/Framework/Commands/UICommandList.cpp
--- a/Engine/Source/Runtime/Slate/Private/Framework/Commands/UICommandList.cpp
+++ b/Engine/Source/Runtime/Slate/Private/Framework/Commands/UICommandList.cpp
@@ -161,11 +161,7 @@ bool FUICommandList::ConditionalProcessCommandBindings( const FKey Key, bool bCt
 {
 	if ( !bRepeat && !FSlateApplication::Get().IsDragDropping() )
 	{
-		FInputGesture CheckGesture( Key );
-		CheckGesture.bCtrl = bCtrl;
-		CheckGesture.bAlt = bAlt;
-		CheckGesture.bShift = bShift;
-		CheckGesture.bCmd = bCmd;
+		FInputGesture CheckGesture( Key, EModifierKey::FromBools(bCtrl, bAlt, bShift, bCmd) );
 
 		if( CheckGesture.IsValidGesture() )
 		{
diff --git a/Engine/Source/Runtime/Slate/Private/Framework/MultiBox/Mac/MacMenu.cpp b/Engine/Source/Runtime/Slate/Private/Framework/MultiBox/Mac/MacMenu.cpp
--- a/Engine/Source/Runtime/Slate/Private/Framework/MultiBox/Mac/MacMenu.cpp
+++ b/Engine/Source/Runtime/Slate/Private/Framework/MultiBox/Mac/MacMenu.cpp
@@ -428,19 +428,19 @@ NSString* FSlateMacMenu::GetMenuItemKeyEquivalent(const TSharedRef<const class F
 		const TSharedRef<const FInputGesture>& Gesture = Block->GetAction()->GetActiveGesture();
 
 		*OutModifiers = 0;
-		if (Gesture->bCtrl)
+		if (Gesture->NeedsControl())
 		{
 			*OutModifiers |= NSCommandKeyMask;
 		}
-		if (Gesture->bShift)
+		if (Gesture->NeedsShift())
 		{
 			*OutModifiers |= NSShiftKeyMask;
 		}
-		if (Gesture->bAlt)
+		if (Gesture->NeedsAlt())
 		{
 			*OutModifiers |= NSAlternateKeyMask;
 		}
-		if (Gesture->bCmd)
+		if (Gesture->NeedsCommand())
 		{
 			*OutModifiers |= NSControlKeyMask;
 		}
diff --git a/Engine/Source/Runtime/Slate/Public/Framework/Commands/InputGesture.h b/Engine/Source/Runtime/Slate/Public/Framework/Commands/InputGesture.h
--- a/Engine/Source/Runtime/Slate/Public/Framework/Commands/InputGesture.h
+++ b/Engine/Source/Runtime/Slate/Public/Framework/Commands/InputGesture.h
@@ -12,24 +12,24 @@ struct SLATE_API FInputGesture
 	/** Key that must be pressed */
 	FKey Key;
 
-	/** True if control must be pressed */
-	uint32 bCtrl:1;
-	
-	/** True if alt must be pressed */
-	uint32 bAlt:1;
-	
-	/** True if shift must be pressed */
-	uint32 bShift:1;
-
-	/** True if command must be pressed */
-	uint32 bCmd:1;
+	EModifierKey::Type ModifierKeys;
+
+#if PLATFORM_MAC
+	bool NeedsControl() const { return (ModifierKeys & EModifierKey::Control) != 0; }
+	bool NeedsCommand() const { return (ModifierKeys & EModifierKey::Command) != 0; }
+#else
+	bool NeedsControl() const { return (ModifierKeys & EModifierKey::Command) != 0; }
+	bool NeedsCommand() const { return (ModifierKeys & EModifierKey::Control) != 0; }
+#endif
+	bool NeedsAlt() const { return (ModifierKeys & EModifierKey::Alt) != 0; }
+	bool NeedsShift() const { return (ModifierKeys & EModifierKey::Shift) != 0; }
 
 public:
 
 	/** Default constructor. */
-	FInputGesture( )
+	FInputGesture()
+		: ModifierKeys(EModifierKey::None)
 	{
-		bCtrl = bAlt = bShift = bCmd = 0;
 	}
 
 	/**
@@ -38,13 +38,10 @@ public:
 	 * @param InModifierKeys
 	 * @param InKey
 	 */
-	FInputGesture( EModifierKey::Type InModifierKeys, const FKey InKey )
+	FInputGesture( const EModifierKey::Type InModifierKeys, const FKey InKey )
 		: Key(InKey)
+		, ModifierKeys(InModifierKeys)
 	{
-		bCtrl = (InModifierKeys & EModifierKey::Control) != 0;
-		bAlt = (InModifierKeys & EModifierKey::Alt) != 0;
-		bShift = (InModifierKeys & EModifierKey::Shift) != 0;
-		bCmd = (InModifierKeys & EModifierKey::Command) != 0;
 	}
 
 	/**
@@ -54,8 +51,8 @@ public:
 	 */
 	FInputGesture( const FKey InKey )
 		: Key(InKey)
+		, ModifierKeys(EModifierKey::None)
 	{
-		bCtrl = bAlt = bShift = bCmd = 0;
 	}
 
 	/**
@@ -64,13 +61,10 @@ public:
 	 * @param InKey
 	 * @param InModifierKeys
 	 */
-	FInputGesture( const FKey InKey, EModifierKey::Type InModifierKeys )
+	FInputGesture( const FKey InKey, const EModifierKey::Type InModifierKeys )
 		: Key(InKey)
+		, ModifierKeys(InModifierKeys)
 	{
-		bCtrl = (InModifierKeys & EModifierKey::Control) != 0;
-		bAlt = (InModifierKeys & EModifierKey::Alt) != 0;
-		bShift = (InModifierKeys & EModifierKey::Shift) != 0;
-		bCmd = (InModifierKeys & EModifierKey::Command) != 0;
 	}
 
 	/**
@@ -80,10 +74,7 @@ public:
 	 */
 	FInputGesture( const FInputGesture& Other )
 		: Key(Other.Key)
-		, bCtrl(Other.bCtrl)
-		, bAlt(Other.bAlt)
-		, bShift(Other.bShift)
-		, bCmd(Other.bCmd)
+		, ModifierKeys(Other.ModifierKeys)
 	{ }
 
 public:
@@ -96,7 +87,7 @@ public:
 	 */
 	bool operator!=( const FInputGesture& Other ) const
 	{
-		return (Key != Other.Key) || (bCtrl != Other.bCtrl) || (bAlt != Other.bAlt) || (bShift != Other.bShift) || (bCmd != Other.bCmd);
+		return !(*this == Other);
 	}
 
 	/**
@@ -107,7 +98,7 @@ public:
 	 */
 	bool operator==( const FInputGesture& Other ) const
 	{
-		return (Key == Other.Key) && (bCtrl == Other.bCtrl) && (bAlt == Other.bAlt) && (bShift == Other.bShift) && (bCmd == Other.bCmd);
+		return (Key == Other.Key) && (ModifierKeys == Other.ModifierKeys);
 	}
 
 public:
@@ -133,7 +124,7 @@ public:
 	 */
 	bool HasAnyModifierKeys( ) const
 	{
-		return bCtrl || bAlt || bShift || bCmd;
+		return (ModifierKeys != EModifierKey::None);
 	}
 
 	/**
@@ -151,13 +142,9 @@ public:
 	 * Sets this gesture to a new key and modifier state based on the provided template
 	 * Should not be called directly.  Only used by the key binding editor to set user defined keys
 	 */
-	void Set( FInputGesture InTemplate )
+	void Set( const FInputGesture& InTemplate )
 	{
-		Key = InTemplate.Key;
-		bCtrl = InTemplate.bCtrl;
-		bAlt = InTemplate.bAlt;
-		bShift = InTemplate.bShift;
-		bCmd = InTemplate.bCmd;
+		*this = InTemplate;
 	}
 
 public:
@@ -170,6 +157,6 @@ public:
 	 */
 	friend uint32 GetTypeHash( const FInputGesture& Gesture )
 	{
-		return GetTypeHash(Gesture.Key) ^ ((Gesture.bCtrl << 3) | (Gesture.bShift << 2) | (Gesture.bAlt << 1) | Gesture.bCmd);
+		return GetTypeHash(Gesture.Key) ^ Gesture.ModifierKeys;
 	}
 };
