[UE-4609] Marking user-defined pins that are removed as pending kill (so that they avoid being flagged by a ReconstructNode error from CL 2238282)

diff --git a/Engine/Source/Editor/BlueprintGraph/Private/K2Node_EditablePinBase.cpp b/Engine/Source/Editor/BlueprintGraph/Private/K2Node_EditablePinBase.cpp
--- a/Engine/Source/Editor/BlueprintGraph/Private/K2Node_EditablePinBase.cpp
+++ b/Engine/Source/Editor/BlueprintGraph/Private/K2Node_EditablePinBase.cpp
@@ -48,6 +48,7 @@ void UK2Node_EditablePinBase::RemoveUserDefinedPin(TSharedPtr<FUserPinInfo> PinT
 		{
 			Pin->BreakAllPinLinks();
 			Pins.Remove(Pin);
+			Pin->MarkPendingKill();
 		}
 	}
 
