GraphSchema doesn't support EventDispatcher by default

diff --git a/Engine/Source/Editor/BlueprintGraph/Classes/EdGraphSchema_K2.h b/Engine/Source/Editor/BlueprintGraph/Classes/EdGraphSchema_K2.h
--- a/Engine/Source/Editor/BlueprintGraph/Classes/EdGraphSchema_K2.h
+++ b/Engine/Source/Editor/BlueprintGraph/Classes/EdGraphSchema_K2.h
@@ -414,6 +414,7 @@ public:
 	virtual void BackwardCompatibilityNodeConversion(UEdGraph* Graph, bool bOnlySafeChanges) const override;
 	virtual void SplitPin(UEdGraphPin* Pin) const override;
 	virtual void RecombinePin(UEdGraphPin* Pin) const override;
+	virtual bool DoesSupportEventDispatcher() const	override { return true; }
 	// End EdGraphSchema Interface
 
 	/**
diff --git a/Engine/Source/Editor/BlueprintGraph/Private/K2Node_MCDelegate.cpp b/Engine/Source/Editor/BlueprintGraph/Private/K2Node_MCDelegate.cpp
--- a/Engine/Source/Editor/BlueprintGraph/Private/K2Node_MCDelegate.cpp
+++ b/Engine/Source/Editor/BlueprintGraph/Private/K2Node_MCDelegate.cpp
@@ -212,10 +212,10 @@ void UK2Node_BaseMCDelegate::GetNodeAttributes( TArray<TKeyValuePair<FString, FS
 
 void UK2Node_BaseMCDelegate::AutowireNewNode(UEdGraphPin* FromPin)
 {
-	const UEdGraphSchema_K2* K2Schema = CastChecked<UEdGraphSchema_K2>(GetSchema());
+	const UEdGraphSchema_K2* K2Schema = Cast<UEdGraphSchema_K2>(GetSchema());
 
 	// Since nodes no longer have a sense of scope when they're placed, look at the connection we're coming from, and use that to coerce the Target pin
-	if (FromPin != nullptr)
+	if (FromPin && K2Schema)
 	{
 		bool bConnected = false;
 
diff --git a/Engine/Source/Editor/GraphEditor/Private/SGraphPanel.cpp b/Engine/Source/Editor/GraphEditor/Private/SGraphPanel.cpp
--- a/Engine/Source/Editor/GraphEditor/Private/SGraphPanel.cpp
+++ b/Engine/Source/Editor/GraphEditor/Private/SGraphPanel.cpp
@@ -808,7 +808,8 @@ FReply SGraphPanel::OnDrop( const FGeometry& MyGeometry, const FDragDropEvent& D
 	if (Operation->IsOfType<FGraphEditorDragDropAction>())
 	{
 		check(GraphObj);
-		if (GraphObj->GetSchema()->DoesSupportEventDispatcher())
+		auto Schema = GraphObj->GetSchema();
+		if (Schema && Schema->DoesSupportEventDispatcher())
 		{
 			TSharedPtr<FGraphEditorDragDropAction> DragConn = StaticCastSharedPtr<FGraphEditorDragDropAction>(Operation);
 			return DragConn->DroppedOnPanel(SharedThis(this), DragDropEvent.GetScreenSpacePosition(), NodeAddPosition, *GraphObj);
diff --git a/Engine/Source/Runtime/Engine/Classes/EdGraph/EdGraphSchema.h b/Engine/Source/Runtime/Engine/Classes/EdGraph/EdGraphSchema.h
--- a/Engine/Source/Runtime/Engine/Classes/EdGraph/EdGraphSchema.h
+++ b/Engine/Source/Runtime/Engine/Classes/EdGraph/EdGraphSchema.h
@@ -522,7 +522,7 @@ class ENGINE_API UEdGraphSchema : public UObject
 	 * 
 	 *	@return True if this schema supports Event Dispatcher 
 	 */
-	virtual bool DoesSupportEventDispatcher() const	{	return true; }
+	virtual bool DoesSupportEventDispatcher() const	{ return false; }
 	/**
 	 * Checks to see if the specified pin is being watched by the graph's debug system.
 	 * 
