Node_AssignDelegate title fix

diff --git a/Engine/Source/Editor/BlueprintGraph/Private/K2Node_AssignDelegate.cpp b/Engine/Source/Editor/BlueprintGraph/Private/K2Node_AssignDelegate.cpp
--- a/Engine/Source/Editor/BlueprintGraph/Private/K2Node_AssignDelegate.cpp
+++ b/Engine/Source/Editor/BlueprintGraph/Private/K2Node_AssignDelegate.cpp
@@ -31,14 +31,14 @@ FText UK2Node_AssignDelegate::GetNodeTitle(ENodeTitleType::Type TitleType) const
 				PropertyName = FText::FromString(bShowFriendlyNames ? UEditorEngine::GetFriendlyName(Property) : Property->GetName());
 
 				// FText::Format() is slow, so we cache this to save on performance
-				CachedNodeTitle = FText::Format(LOCTEXT("AssignDelegateTitle", "Assign {0}"), PropertyName);
+				CachedListTitle = FText::Format(LOCTEXT("AssignDelegateTitle", "Assign {0}"), PropertyName);
 			}
 			else
 			{
 				return LOCTEXT("InvalidPropertyTitle", "Assign <invalid delegate>");
 			}
 		}
-		return CachedNodeTitle;
+		return CachedListTitle;
 	}
 	
 	return Super::GetNodeTitle(TitleType);
