"Set" variable nodes now have an output pin that acts as a "Get" variable node.

Added an output pin to UK2Node_VariableSet

Changed UK2Node_Variable::CreatePinForVariable to accept an optional pin name.
Changed UK2Node_Variable::RecreatePinForVariable to accept an optional pin name.
 UK2NOde_Variable::DoPinsMatchForReconstruction will now confirm the pins are going in the same direction as well of compatible types before it determines that the pins match.

diff --git a/Engine/Source/Editor/BlueprintGraph/Classes/K2Node_Variable.h b/Engine/Source/Editor/BlueprintGraph/Classes/K2Node_Variable.h
--- a/Engine/Source/Editor/BlueprintGraph/Classes/K2Node_Variable.h
+++ b/Engine/Source/Editor/BlueprintGraph/Classes/K2Node_Variable.h
@@ -93,10 +93,11 @@ public:
 	 * Creates a reader or writer pin for a variable.
 	 *
 	 * @param	Direction	  	The direction of the variable access.
+	 * @param	InPinName		Optional pin name, will default to the variable name if not included.
 	 *
 	 * @return	true if it succeeds, false if it fails.
 	 */
-	bool CreatePinForVariable(EEdGraphPinDirection Direction);
+	bool CreatePinForVariable(EEdGraphPinDirection Direction, FString InPinName = FString());
 
 	/** Creates 'self' pin */
 	void CreatePinForSelf();
@@ -106,10 +107,11 @@ public:
 	 *
 	 * @param	Direction	  	The direction of the variable access.
 	 * @param	OldPins			Old pins.
+	 * @param	InPinName		Optional pin name, will default to the variable name if not included.
 	 *
 	 * @return	true if it succeeds, false if it fails.
 	 */
-	bool RecreatePinForVariable(EEdGraphPinDirection Direction, TArray<UEdGraphPin*>& OldPins);
+	bool RecreatePinForVariable(EEdGraphPinDirection Direction, TArray<UEdGraphPin*>& OldPins, FString InPinName = FString());
 
 	/** Get the class to look for this variable in */
 	UClass* GetVariableSourceClass() const;
diff --git a/Engine/Source/Editor/BlueprintGraph/Classes/K2Node_VariableSet.h b/Engine/Source/Editor/BlueprintGraph/Classes/K2Node_VariableSet.h
--- a/Engine/Source/Editor/BlueprintGraph/Classes/K2Node_VariableSet.h
+++ b/Engine/Source/Editor/BlueprintGraph/Classes/K2Node_VariableSet.h
@@ -16,13 +16,18 @@ class UK2Node_VariableSet : public UK2Node_Variable
 	virtual void ReallocatePinsDuringReconstruction(TArray<UEdGraphPin*>& OldPins) override;
 	virtual FText GetTooltipText() const override;
 	virtual FText GetNodeTitle(ENodeTitleType::Type TitleType) const override;
+	virtual bool ShouldOverridePinNames() const override { return true; }
+	virtual FString GetPinNameOverride(const UEdGraphPin& Pin) const override;
 	// End UEdGraphNode interface
 
 	// Begin K2Node interface
 	virtual bool ShouldShowNodeProperties() const override { return true; }
 	virtual class FNodeHandlingFunctor* CreateNodeHandler(class FKismetCompilerContext& CompilerContext) const override;
+	virtual void ExpandNode(class FKismetCompilerContext& CompilerContext, UEdGraph* SourceGraph) override;
 	// End K2Node interface
 
+	/** Retrieves the output pin name for the node */
+	FString GetVariableOutputPinName() const;
 
 	BLUEPRINTGRAPH_API bool HasLocalRepNotify() const;
 	BLUEPRINTGRAPH_API FName GetRepNotifyName() const;
@@ -31,6 +36,10 @@ class UK2Node_VariableSet : public UK2Node_Variable
 	static FText GetPropertyTooltip(UProperty const* VariableProperty);
 	static FText GetBlueprintVarTooltip(FBPVariableDescription const& VarDesc);
 
+private:
+	/** Sets up the tooltip for the output pin */
+	void CreateOutputPinTooltip();
+
 private:
 	/** Constructing FText strings can be costly, so we cache the node's title/tooltip */
 	FNodeTextCache CachedTooltip;
diff --git a/Engine/Source/Editor/BlueprintGraph/Private/K2Node_Variable.cpp b/Engine/Source/Editor/BlueprintGraph/Private/K2Node_Variable.cpp
--- a/Engine/Source/Editor/BlueprintGraph/Private/K2Node_Variable.cpp
+++ b/Engine/Source/Editor/BlueprintGraph/Private/K2Node_Variable.cpp
@@ -55,7 +55,7 @@ void UK2Node_Variable::SetFromProperty(const UProperty* Property, bool bSelfCont
 	VariableReference.SetFromField<UProperty>(Property, bSelfContext);
 }
 
-bool UK2Node_Variable::CreatePinForVariable(EEdGraphPinDirection Direction)
+bool UK2Node_Variable::CreatePinForVariable(EEdGraphPinDirection Direction, FString InPinName/* = FString()*/)
 {
 	const UEdGraphSchema_K2* K2Schema = GetDefault<UEdGraphSchema_K2>();
 
@@ -81,8 +81,9 @@ bool UK2Node_Variable::CreatePinForVariable(EEdGraphPinDirection Direction)
 
 	if (VariableProperty != NULL)
 	{
+		const FString PinName = InPinName.IsEmpty()? GetVarNameString() : InPinName;
 		// Create the pin
-		UEdGraphPin* VariablePin = CreatePin(Direction, TEXT(""), TEXT(""), NULL, false, false, GetVarNameString());
+		UEdGraphPin* VariablePin = CreatePin(Direction, TEXT(""), TEXT(""), NULL, false, false, PinName);
 		K2Schema->ConvertPropertyToPinType(VariableProperty, /*out*/ VariablePin->PinType);
 		K2Schema->SetPinDefaultValueBasedOnType(VariablePin);
 	}
@@ -144,14 +145,14 @@ void UK2Node_Variable::CreatePinForSelf()
 	}
 }
 
-bool UK2Node_Variable::RecreatePinForVariable(EEdGraphPinDirection Direction, TArray<UEdGraphPin*>& OldPins)
+bool UK2Node_Variable::RecreatePinForVariable(EEdGraphPinDirection Direction, TArray<UEdGraphPin*>& OldPins, FString InPinName/* = FString()*/)
 {
 	// probably the node was pasted to a blueprint without the variable
 	// we don't want to beak any connection, so the pin will be recreated from old one, but compiler will throw error
 
 	// find old variable pin
 	const UEdGraphPin* OldVariablePin = NULL;
-	const FString PinName = GetVarNameString();
+	const FString PinName = InPinName.IsEmpty()? GetVarNameString() : InPinName;
 	for(auto Iter = OldPins.CreateConstIterator(); Iter; ++Iter)
 	{
 		if(const UEdGraphPin* Pin = *Iter)
diff --git a/Engine/Source/Editor/BlueprintGraph/Private/K2Node_VariableSet.cpp b/Engine/Source/Editor/BlueprintGraph/Private/K2Node_VariableSet.cpp
--- a/Engine/Source/Editor/BlueprintGraph/Private/K2Node_VariableSet.cpp
+++ b/Engine/Source/Editor/BlueprintGraph/Private/K2Node_VariableSet.cpp
@@ -86,6 +86,11 @@ void UK2Node_VariableSet::AllocateDefaultPins()
 		{
 			CreatePinForSelf();
 		}
+
+		if(CreatePinForVariable(EGPD_Output, GetVariableOutputPinName()))
+		{
+			CreateOutputPinTooltip();
+		}
 	}
 
 	Super::AllocateDefaultPins();
@@ -107,6 +112,15 @@ void UK2Node_VariableSet::ReallocatePinsDuringReconstruction(TArray<UEdGraphPin*
 				return;
 			}
 		}
+
+		if(!CreatePinForVariable(EGPD_Output, GetVariableOutputPinName()))
+		{
+			if(!RecreatePinForVariable(EGPD_Output, OldPins, GetVariableOutputPinName()))
+			{
+				return;
+			}
+		}
+		CreateOutputPinTooltip();
 		CreatePinForSelf();
 	}
 }
@@ -342,4 +356,68 @@ FNodeHandlingFunctor* UK2Node_VariableSet::CreateNodeHandler(FKismetCompilerCont
 	return new FKCHandler_VariableSet(CompilerContext);
 }
 
+FString UK2Node_VariableSet::GetVariableOutputPinName() const
+{
+	return GetVarNameString() + TEXT("_Get");
+}
+
+void UK2Node_VariableSet::CreateOutputPinTooltip()
+{
+	UEdGraphPin* Pin = FindPin(GetVariableOutputPinName());
+	check(Pin);
+	Pin->PinToolTip = NSLOCTEXT("K2Node", "SetPinOutputTooltip", "Retrieves the value of the variable, can use instead of a separate Get node").ToString();
+}
+
+FString UK2Node_VariableSet::GetPinNameOverride(const UEdGraphPin& Pin) const
+{
+	// Stop the output pin for the variable, effectively the "get" pin, from displaying a name.
+	const UEdGraphSchema_K2* K2Schema = GetDefault<UEdGraphSchema_K2>();
+	if(Pin.Direction == EGPD_Output || Pin.PinType.PinCategory == K2Schema->PC_Exec)
+	{
+		return FString();
+	}
+
+	return !Pin.PinFriendlyName.IsEmpty() ? Pin.PinFriendlyName.ToString() : Pin.PinName;
+}
+
+void UK2Node_VariableSet::ExpandNode(class FKismetCompilerContext& CompilerContext, UEdGraph* SourceGraph)
+{
+	Super::ExpandNode(CompilerContext, SourceGraph);
+
+	if (CompilerContext.bIsFullCompile)
+	{
+		const UEdGraphSchema_K2* K2Schema = CompilerContext.GetSchema();
+
+		UEdGraphPin* Pin = FindPin(GetVariableOutputPinName());
+		if(Pin)
+		{
+			// If the output pin is linked, we need to spawn a separate "Get" node and hook it up.
+			if(Pin->LinkedTo.Num())
+			{
+				UProperty* VariableProperty = GetPropertyForVariable();
+
+				if(VariableProperty)
+				{
+					UK2Node_VariableGet* VariableGetNode = CompilerContext.SpawnIntermediateNode<UK2Node_VariableGet>(this, SourceGraph);
+					VariableGetNode->VariableReference = VariableReference;
+					VariableGetNode->AllocateDefaultPins();
+					CompilerContext.MessageLog.NotifyIntermediateObjectCreation(VariableGetNode, this);
+					CompilerContext.MovePinLinksToIntermediate(*Pin, *VariableGetNode->FindPin(GetVarNameString()));
+
+					// Duplicate the connection to the self pin.
+					UEdGraphPin* SetSelfPin = K2Schema->FindSelfPin(*this, EGPD_Input);
+					UEdGraphPin* GetSelfPin = K2Schema->FindSelfPin(*VariableGetNode, EGPD_Input);
+					if(SetSelfPin && GetSelfPin)
+					{
+						CompilerContext.CopyPinLinksToIntermediate(*SetSelfPin, *GetSelfPin);
+					}
+				}
+			}
+			Pins.Remove(Pin);
+		}
+		
+	}
+
+}
+
 #undef LOCTEXT_NAMESPACE
