[UE-5652] Changes to matinee events now don't leave controller node connections dangling.

diff --git a/Engine/Source/Editor/BlueprintGraph/Private/K2Node_MatineeController.cpp b/Engine/Source/Editor/BlueprintGraph/Private/K2Node_MatineeController.cpp
--- a/Engine/Source/Editor/BlueprintGraph/Private/K2Node_MatineeController.cpp
+++ b/Engine/Source/Editor/BlueprintGraph/Private/K2Node_MatineeController.cpp
@@ -214,19 +214,10 @@ void UK2Node_MatineeController::OnEventKeyframeRenamed(const AMatineeActor* InMa
 {
 	if ( MatineeActor == InMatineeActor )
 	{
-		// Rather than rename the pin directly, modify 'this' array as it's cheaper and quicker for the undo/redo system to not track UEdGraphPin's directly
 		if( UEdGraphPin* OldPin = FindPin(InOldPinName.ToString()) )
 		{
-			// Cache the pin index, then remove it from the array
-			const int32 PinIndex = GetPinIndex(OldPin);
-			check( PinIndex != INDEX_NONE );
-			DiscardPin(OldPin);
-
-			// Create a new pin at the old index and with the new name and copy it's data
-			const UEdGraphSchema_K2* K2Schema = GetDefault<UEdGraphSchema_K2>();
-			UEdGraphPin* NewPin = CreatePin(EGPD_Output, K2Schema->PC_Exec, TEXT(""), NULL, false, false, InNewPinName.ToString(), false, PinIndex);	
-			NewPin->CopyPersistentDataFromOldPin(*OldPin);
-			OldPin = NULL;
+			OldPin->Modify();
+			OldPin->PinName = InNewPinName.ToString();
 
 			GetGraph()->NotifyGraphChanged();
 		}
diff --git a/Engine/Source/Editor/Matinee/Private/MatineeTransaction.cpp b/Engine/Source/Editor/Matinee/Private/MatineeTransaction.cpp
--- a/Engine/Source/Editor/Matinee/Private/MatineeTransaction.cpp
+++ b/Engine/Source/Editor/Matinee/Private/MatineeTransaction.cpp
@@ -26,7 +26,9 @@ void FMatineeTransaction::SaveObject( UObject* Object )
 		Object->IsA( UInterpGroupInst::StaticClass() ) ||
 		Object->IsA( UInterpTrackInst::StaticClass() ) ||
 		Object->IsA( UMatineeOptions::StaticClass() ) ||
-		Object->IsA( UK2Node_MatineeController::StaticClass() ) )
+		Object->IsA( UK2Node_MatineeController::StaticClass() ) || 
+		// or bIsMatineeControllerPin...
+		(Object->IsA<UEdGraphPin>() && Object->GetOuter()->IsA<UK2Node_MatineeController>()) )
 	{
 		// Save the object.
 		new( Records )FObjectRecord( this, Object, NULL, 0, 0, 0, 0, NULL, NULL, NULL );
diff --git a/Engine/Source/Runtime/Engine/Private/EdGraph/EdGraphNode.cpp b/Engine/Source/Runtime/Engine/Private/EdGraph/EdGraphNode.cpp
--- a/Engine/Source/Runtime/Engine/Private/EdGraph/EdGraphNode.cpp
+++ b/Engine/Source/Runtime/Engine/Private/EdGraph/EdGraphNode.cpp
@@ -104,9 +104,10 @@ UEdGraphPin* UEdGraphNode::FindPinChecked(const FString& PinName) const
 void UEdGraphNode::DiscardPin(UEdGraphPin* Pin)
 {
 	check( Pin );
-
+	
 	Modify();
 	Pins.Remove( Pin );
+	Pin->BreakAllPinLinks();
 }
 
 void UEdGraphNode::BreakAllNodeLinks()
