BP: safe way to deal with obsolete literal values in AutoCreateRefs pins

diff --git a/Engine/Source/Editor/BlueprintGraph/Private/K2Node_CallFunction.cpp b/Engine/Source/Editor/BlueprintGraph/Private/K2Node_CallFunction.cpp
--- a/Engine/Source/Editor/BlueprintGraph/Private/K2Node_CallFunction.cpp
+++ b/Engine/Source/Editor/BlueprintGraph/Private/K2Node_CallFunction.cpp
@@ -1839,13 +1839,22 @@ void UK2Node_CallFunction::ExpandNode(class FKismetCompilerContext& CompilerCont
 						const auto DefaultValue = Pin->DefaultValue;
 						const auto DefaultObject = Pin->DefaultObject;
 						const auto DefaultTextValue = Pin->DefaultTextValue;
+						const auto AutogeneratedDefaultValue = Pin->AutogeneratedDefaultValue;
 
 						auto ValuePin = InnerHandleAutoCreateRef(this, Pin, CompilerContext, SourceGraph, bHasDefaultValue);
 						if ( ValuePin )
 						{
-							ValuePin->DefaultValue = DefaultValue;
-							ValuePin->DefaultObject = DefaultObject;
-							ValuePin->DefaultTextValue = DefaultTextValue;
+							if (!DefaultObject && DefaultTextValue.IsEmpty() && (DefaultValue == AutogeneratedDefaultValue))
+							{
+								// Use the latest code to set default value
+								Schema->SetPinDefaultValueBasedOnType(ValuePin);
+							}
+							else
+							{
+								ValuePin->DefaultValue = DefaultValue;
+								ValuePin->DefaultObject = DefaultObject;
+								ValuePin->DefaultTextValue = DefaultTextValue;
+							}
 						}
 					}
 				}
