[UE-5309] Editor setting to have cast nodes default to their pure form.

diff --git a/Engine/Source/Editor/BlueprintGraph/Private/BlueprintEditorSettings.cpp b/Engine/Source/Editor/BlueprintGraph/Private/BlueprintEditorSettings.cpp
--- a/Engine/Source/Editor/BlueprintGraph/Private/BlueprintEditorSettings.cpp
+++ b/Engine/Source/Editor/BlueprintGraph/Private/BlueprintEditorSettings.cpp
@@ -9,13 +9,14 @@ UBlueprintEditorSettings::UBlueprintEditorSettings(const FObjectInitializer& Obj
 	: Super(ObjectInitializer)
 	// Style Settings
 	, bDrawMidpointArrowsInBlueprints(false)
-	// UX Settings
 	, bShowGraphInstructionText(true)
+	// Workflow Settings
 	, bUseTargetContextForNodeMenu(true)
 	, bExposeAllMemberComponentFunctions(true)
 	, bShowContextualFavorites(false)
 	, bFlattenFavoritesMenus(true)
 	, bUseLegacyMenuingSystem(false)
+	, bFavorPureCastNodes(false)
 	// Compiler Settings
 	, SaveOnCompile(SoC_Never)
 	, bJumpToNodeErrors(false)
diff --git a/Engine/Source/Editor/BlueprintGraph/Private/K2Node_DynamicCast.cpp b/Engine/Source/Editor/BlueprintGraph/Private/K2Node_DynamicCast.cpp
--- a/Engine/Source/Editor/BlueprintGraph/Private/K2Node_DynamicCast.cpp
+++ b/Engine/Source/Editor/BlueprintGraph/Private/K2Node_DynamicCast.cpp
@@ -4,6 +4,7 @@
 #include "BlueprintGraphPrivatePCH.h"
 #include "DynamicCastHandler.h"
 #include "EditorCategoryUtils.h"
+#include "BlueprintEditorSettings.h"
 
 #define LOCTEXT_NAMESPACE "K2Node_DynamicCast"
 
@@ -11,6 +12,8 @@ UK2Node_DynamicCast::UK2Node_DynamicCast(const FObjectInitializer& ObjectInitial
 	: Super(ObjectInitializer)
 	, bIsPureCast(false)
 {
+	const UBlueprintEditorSettings* BlueprintSettings = GetDefault<UBlueprintEditorSettings>();
+	bIsPureCast = BlueprintSettings->bFavorPureCastNodes;
 }
 
 void UK2Node_DynamicCast::AllocateDefaultPins()
diff --git a/Engine/Source/Editor/BlueprintGraph/Public/BlueprintEditorSettings.h b/Engine/Source/Editor/BlueprintGraph/Public/BlueprintEditorSettings.h
--- a/Engine/Source/Editor/BlueprintGraph/Public/BlueprintEditorSettings.h
+++ b/Engine/Source/Editor/BlueprintGraph/Public/BlueprintEditorSettings.h
@@ -25,30 +25,34 @@ public:
 	UPROPERTY(EditAnywhere, config, Category=VisualStyle, meta=(DisplayName="Draw midpoint arrows in Blueprints"))
 	bool bDrawMidpointArrowsInBlueprints;
 
-// UX Settings
-public:
 	/** Determines if lightweight tutorial text shows up at the top of empty blueprint graphs */
-	UPROPERTY(EditAnywhere, config, Category=UserExperience)
+	UPROPERTY(EditAnywhere, config, Category = VisualStyle)
 	bool bShowGraphInstructionText;
 
+// Workflow Settings
+public:
 	/** If enabled, will use the blueprint's (or output pin's) class to narrow down context menu results. */
-	UPROPERTY(EditAnywhere, config, Category=UserExperience, meta=(DisplayName="Context Menu: Use Target Context"))
+	UPROPERTY(EditAnywhere, config, Category=Workflow, meta=(DisplayName="Context Menu: Use Target Context"))
 	bool bUseTargetContextForNodeMenu;
 
 	/** If enabled, then ALL component functions are exposed to the context menu (when the contextual target is a component owner). Ignores "ExposeFunctionCategories" metadata for components. */
-	UPROPERTY(EditAnywhere, config, Category=UserExperience, meta=(DisplayName="Context Menu: Expose All Sub-Component Functions"))
+	UPROPERTY(EditAnywhere, config, Category=Workflow, meta=(DisplayName="Context Menu: Expose All Sub-Component Functions"))
 	bool bExposeAllMemberComponentFunctions;
 
 	/** If enabled, then a separate section with your Blueprint favorites will be pined to the top of the context menu. */
-	UPROPERTY(EditAnywhere, config, Category=UserExperience, meta=(DisplayName="Context Menu: Show Favorites Section"))
+	UPROPERTY(EditAnywhere, config, Category=Workflow, meta=(DisplayName="Context Menu: Show Favorites Section"))
 	bool bShowContextualFavorites;
 
 	/** If enabled, then your Blueprint favorites will be uncategorized, leaving you with less nested categories to sort through. */
-	UPROPERTY(EditAnywhere, config, Category=UserExperience)
+	UPROPERTY(EditAnywhere, config, Category=Workflow)
 	bool bFlattenFavoritesMenus;
 
+	/** If enabled, then placed cast nodes will default to their "pure" form (meaning: without execution pins). */
+	UPROPERTY(EditAnywhere, config, Category=Workflow, meta=(DisplayName="Default to Using Pure Cast Nodes"))
+	bool bFavorPureCastNodes;
+
 	/** If set, then the new refactored menu system will be replaced with the old (legacy) system (as a fallback, in case the new system has unforeseen problems)*/
-	UPROPERTY(EditAnywhere, AdvancedDisplay, config, Category=UserExperience)
+	UPROPERTY(EditAnywhere, AdvancedDisplay, config, Category=Workflow)
 	bool bUseLegacyMenuingSystem;
 
 // Compiler Settings
