[UE-5692] Native public properties marked as EditDefaultsOnly and BlueprintReadOnly were previously classifed as 'protected' and filtered from the context menu in blueprints.  If a field is public and has EDO and BRO it is now public.  If it only has EDO it is still classified as protected.

diff --git a/Engine/Source/Editor/BlueprintGraph/Private/BlueprintActionFilter.cpp b/Engine/Source/Editor/BlueprintGraph/Private/BlueprintActionFilter.cpp
--- a/Engine/Source/Editor/BlueprintGraph/Private/BlueprintActionFilter.cpp
+++ b/Engine/Source/Editor/BlueprintGraph/Private/BlueprintActionFilter.cpp
@@ -487,12 +487,18 @@ static bool BlueprintActionFilterImpl::IsFieldInaccessible(FBlueprintActionFilte
 		{
 			// CPF_DisableEditOnInstance corresponds to the eyeball/editable 
 			// checkbox available in the blueprint editor. It is poorly named. 
+			bIsPublic = !Property->HasAnyPropertyFlags(CPF_DisableEditOnInstance);
+
+			// If we've disabled editing on the instance, it's entirely possible that the variable has
+			// the flags, EditDefaultsOnly, BlueprintReadOnly, which should still be public, but only
+			// for getting the value.  This is for Native properties.
+			bIsPublic = bIsPublic || Property->HasAnyPropertyFlags(CPF_BlueprintReadOnly);
+
 			// When CPF_DisableEditOnInstance is *not* set, the user has tagged 
 			// the variable as 'editable' and it is treated as public (delegate 
 			// properties don't have the visibility check box, so they default
 			// to public):
-			bIsPublic = !Property->HasAnyPropertyFlags(CPF_DisableEditOnInstance) || 
-				(Cast<UMulticastDelegateProperty>(Property) != nullptr);
+			bIsPublic = bIsPublic || (Cast<UMulticastDelegateProperty>(Property) != nullptr);
 
 			// If the variable is not public, and is not tagged as private, it should be protected
 			// by default. This branch handles variables declared in blueprints:
