[UE-5308] New boolean output on pure cast nodes to check the success/failure of the cast.

diff --git a/Engine/Source/Editor/BlueprintGraph/Classes/K2Node_ClassDynamicCast.h b/Engine/Source/Editor/BlueprintGraph/Classes/K2Node_ClassDynamicCast.h
--- a/Engine/Source/Editor/BlueprintGraph/Classes/K2Node_ClassDynamicCast.h
+++ b/Engine/Source/Editor/BlueprintGraph/Classes/K2Node_ClassDynamicCast.h
@@ -20,7 +20,9 @@ class UK2Node_ClassDynamicCast : public UK2Node_DynamicCast
 	virtual class FNodeHandlingFunctor* CreateNodeHandler(class FKismetCompilerContext& CompilerContext) const override;
 	// End of UK2Node interface
 
-	/** Get the input object to be casted pin */
+	// UK2Node_DynamicCast interface
 	virtual UEdGraphPin* GetCastSourcePin() const override;
+	virtual UEdGraphPin* GetBoolSuccessPin() const override;
+	// End of UK2Node_DynamicCast interface
 };
 
diff --git a/Engine/Source/Editor/BlueprintGraph/Classes/K2Node_DynamicCast.h b/Engine/Source/Editor/BlueprintGraph/Classes/K2Node_DynamicCast.h
--- a/Engine/Source/Editor/BlueprintGraph/Classes/K2Node_DynamicCast.h
+++ b/Engine/Source/Editor/BlueprintGraph/Classes/K2Node_DynamicCast.h
@@ -44,6 +44,9 @@ class UK2Node_DynamicCast : public UK2Node
 	/** Get the input object to be casted pin */
 	BLUEPRINTGRAPH_API virtual UEdGraphPin* GetCastSourcePin() const;
 
+	/** Get the boolean output pin that signifies a successful/failed cast. */
+	BLUEPRINTGRAPH_API virtual UEdGraphPin* GetBoolSuccessPin() const;
+
 	/**
 	 * Will change the node's purity, and reallocate pins accordingly (adding/
 	 * removing exec pins).
diff --git a/Engine/Source/Editor/BlueprintGraph/Private/DynamicCastHandler.cpp b/Engine/Source/Editor/BlueprintGraph/Private/DynamicCastHandler.cpp
--- a/Engine/Source/Editor/BlueprintGraph/Private/DynamicCastHandler.cpp
+++ b/Engine/Source/Editor/BlueprintGraph/Private/DynamicCastHandler.cpp
@@ -11,13 +11,6 @@
 void FKCHandler_DynamicCast::RegisterNets(FKismetFunctionContext& Context, UEdGraphNode* Node)
 {
 	FNodeHandlingFunctor::RegisterNets(Context, Node);
-
-	// Create a term to determine if the cast was successful or not
-	FBPTerminal* BoolTerm = Context.CreateLocalTerminal();
-	BoolTerm->Type.PinCategory = CompilerContext.GetSchema()->PC_Boolean;
-	BoolTerm->Source = Node;
-	BoolTerm->Name = Context.NetNameMap->MakeValidName(Node) + TEXT("_CastSuccess");
-	BoolTermMap.Add(Node, BoolTerm);
 }
 
 void FKCHandler_DynamicCast::RegisterNet(FKismetFunctionContext& Context, UEdGraphPin* Net)
@@ -73,9 +66,6 @@ void FKCHandler_DynamicCast::Compile(FKismetFunctionContext& Context, UEdGraphNo
 	ClassTerm->Source = Node;
 	ClassTerm->ObjectLiteral = DynamicCastNode->TargetType;
 
-	// Find the boolean intermediate result term, so we can track whether the cast was successful
-	FBPTerminal* BoolTerm = *BoolTermMap.Find(DynamicCastNode);
-
 	UClass const* const InputObjClass  = Cast<UClass>((*ObjectToCast)->Type.PinSubCategoryObject.Get());
 	UClass const* const OutputObjClass = Cast<UClass>((*CastResultTerm)->Type.PinSubCategoryObject.Get());
 
@@ -116,35 +106,33 @@ void FKCHandler_DynamicCast::Compile(FKismetFunctionContext& Context, UEdGraphNo
 	CastStatement.RHS.Add(ClassTerm);
 	CastStatement.RHS.Add(*ObjectToCast);
 
-	UEdGraphPin* SuccessPin = DynamicCastNode->GetValidCastPin();
-	bool const bIsPureCast = (SuccessPin == nullptr);
-
+	UEdGraphPin*  BoolSuccessPin  = DynamicCastNode->GetBoolSuccessPin();
+	FBPTerminal** BoolSuccessTerm = Context.NetMap.Find(BoolSuccessPin);
+	check(BoolSuccessTerm != nullptr);
+	// Check result of cast statement
+	FBlueprintCompiledStatement& CheckResultStatement = Context.AppendStatementForNode(Node);
+	CheckResultStatement.Type = KCST_ObjectToBool;
+	CheckResultStatement.LHS  = *BoolSuccessTerm;
+	CheckResultStatement.RHS.Add(*CastResultTerm);
+
+	UEdGraphPin* SuccessExecPin = DynamicCastNode->GetValidCastPin();
+	bool const bIsPureCast = (SuccessExecPin == nullptr);
 	if (!bIsPureCast)
 	{
-		// Check result of cast statement
-		FBlueprintCompiledStatement& CheckResultStatement = Context.AppendStatementForNode(Node);
-		const UClass* SubObjectClass = Cast<UClass>((*CastResultTerm)->Type.PinSubCategoryObject.Get());
-		const bool bIsInterfaceCast = (SubObjectClass != NULL && SubObjectClass->HasAnyClassFlags(CLASS_Interface));
-
-		CheckResultStatement.Type = KCST_ObjectToBool;
-		CheckResultStatement.LHS = BoolTerm;
-		CheckResultStatement.RHS.Add(*CastResultTerm);
-	
 		UEdGraphPin* FailurePin = DynamicCastNode->GetInvalidCastPin();
 		check(FailurePin != nullptr);
 		// Failure condition...skip to the failed output
 		FBlueprintCompiledStatement& FailCastGoto = Context.AppendStatementForNode(Node);
 		FailCastGoto.Type = KCST_GotoIfNot;
-		FailCastGoto.LHS = BoolTerm;
+		FailCastGoto.LHS  = *BoolSuccessTerm;
 		Context.GotoFixupRequestMap.Add(&FailCastGoto, FailurePin);
 
 		// Successful cast...hit the success output node
 		FBlueprintCompiledStatement& SuccessCastGoto = Context.AppendStatementForNode(Node);
 		SuccessCastGoto.Type = KCST_UnconditionalGoto;
-		SuccessCastGoto.LHS = BoolTerm;
-		Context.GotoFixupRequestMap.Add(&SuccessCastGoto, SuccessPin);
+		SuccessCastGoto.LHS  = *BoolSuccessTerm;
+		Context.GotoFixupRequestMap.Add(&SuccessCastGoto, SuccessExecPin);
 	}
-	
 }
 
 #undef LOCTEXT_NAMESPACE
diff --git a/Engine/Source/Editor/BlueprintGraph/Private/DynamicCastHandler.h b/Engine/Source/Editor/BlueprintGraph/Private/DynamicCastHandler.h
--- a/Engine/Source/Editor/BlueprintGraph/Private/DynamicCastHandler.h
+++ b/Engine/Source/Editor/BlueprintGraph/Private/DynamicCastHandler.h
@@ -10,7 +10,6 @@
 class FKCHandler_DynamicCast : public FNodeHandlingFunctor
 {
 protected:
-	TMap<UEdGraphNode*, FBPTerminal*> BoolTermMap;
 	EKismetCompiledStatementType CastType;
 
 public:
diff --git a/Engine/Source/Editor/BlueprintGraph/Private/K2Node_ClassDynamicCast.cpp b/Engine/Source/Editor/BlueprintGraph/Private/K2Node_ClassDynamicCast.cpp
--- a/Engine/Source/Editor/BlueprintGraph/Private/K2Node_ClassDynamicCast.cpp
+++ b/Engine/Source/Editor/BlueprintGraph/Private/K2Node_ClassDynamicCast.cpp
@@ -10,6 +10,8 @@
 
 struct FClassDynamicCastHelper
 {
+	static const FString CastSuccessPinName;
+
 	static const FString& GetClassToCastName()
 	{
 		static FString ClassToCastName(TEXT("Class"));
@@ -17,6 +19,8 @@ struct FClassDynamicCastHelper
 	}
 };
 
+const FString FClassDynamicCastHelper::CastSuccessPinName = TEXT("bSuccess");
+
 UK2Node_ClassDynamicCast::UK2Node_ClassDynamicCast(const FObjectInitializer& ObjectInitializer)
 	: Super(ObjectInitializer)
 {
@@ -55,6 +59,9 @@ void UK2Node_ClassDynamicCast::AllocateDefaultPins()
 		CreatePin(EGPD_Output, K2Schema->PC_Class, TEXT(""), *TargetType, false, false, CastResultPinName);
 	}
 
+	UEdGraphPin* BoolSuccessPin = CreatePin(EGPD_Output, K2Schema->PC_Boolean, TEXT(""), nullptr, /*bIsArray =*/false, /*bIsReference =*/false, FClassDynamicCastHelper::CastSuccessPinName);
+	BoolSuccessPin->bHidden = !bIsPureCast;
+
 	UK2Node::AllocateDefaultPins();
 }
 
@@ -80,6 +87,14 @@ UEdGraphPin* UK2Node_ClassDynamicCast::GetCastSourcePin() const
 	return Pin;
 }
 
+UEdGraphPin* UK2Node_ClassDynamicCast::GetBoolSuccessPin() const
+{
+	UEdGraphPin* Pin = FindPin(FClassDynamicCastHelper::CastSuccessPinName);
+	check(Pin != nullptr);
+	check(Pin->Direction == EGPD_Output);
+	return Pin;
+}
+
 FNodeHandlingFunctor* UK2Node_ClassDynamicCast::CreateNodeHandler(FKismetCompilerContext& CompilerContext) const
 {
 	return new FKCHandler_DynamicCast(CompilerContext, KCST_MetaCast);
diff --git a/Engine/Source/Editor/BlueprintGraph/Private/K2Node_DynamicCast.cpp b/Engine/Source/Editor/BlueprintGraph/Private/K2Node_DynamicCast.cpp
--- a/Engine/Source/Editor/BlueprintGraph/Private/K2Node_DynamicCast.cpp
+++ b/Engine/Source/Editor/BlueprintGraph/Private/K2Node_DynamicCast.cpp
@@ -8,6 +8,11 @@
 
 #define LOCTEXT_NAMESPACE "K2Node_DynamicCast"
 
+namespace UK2Node_DynamicCastImpl
+{
+	static const FString CastSuccessPinName("bSuccess");
+}
+
 UK2Node_DynamicCast::UK2Node_DynamicCast(const FObjectInitializer& ObjectInitializer)
 	: Super(ObjectInitializer)
 	, bIsPureCast(false)
@@ -56,6 +61,9 @@ void UK2Node_DynamicCast::AllocateDefaultPins()
 		}
 	}
 
+	UEdGraphPin* BoolSuccessPin = CreatePin(EGPD_Output, K2Schema->PC_Boolean, TEXT(""), nullptr, /*bIsArray =*/false, /*bIsReference =*/false, UK2Node_DynamicCastImpl::CastSuccessPinName);
+	BoolSuccessPin->bHidden = !bIsPureCast;
+
 	Super::AllocateDefaultPins();
 }
 
@@ -181,6 +189,14 @@ UEdGraphPin* UK2Node_DynamicCast::GetCastSourcePin() const
 	return Pin;
 }
 
+UEdGraphPin* UK2Node_DynamicCast::GetBoolSuccessPin() const
+{
+	UEdGraphPin* Pin = FindPin(UK2Node_DynamicCastImpl::CastSuccessPinName);
+	check(Pin != nullptr);
+	check(Pin->Direction == EGPD_Output);
+	return Pin;
+}
+
 void UK2Node_DynamicCast::SetPurity(bool bNewPurity)
 {
 	if (bNewPurity != bIsPureCast)
