[UE-2345] BP - enforce const-correctness in native const class method overrides

this change introduces enforcement of 'const-correctness' into implemented function graphs.

summary:
if you have a function declared in C++ like this:
UFUNCTION(BlueprintImplementableEvent)
int32 MyFunctionThatReturnsSomeValue() const;
if you implement that (BPIE) function in a Blueprint that's parented to that native class, it will now be flagged as 'const'. this makes any properties of 'self' read-only within the context of that graph, which means the compiler will emit an error if you try to set a property or otherwise call a non-const, non-static function with 'self' as the target.
if there happens to already be an implemented const function in a Blueprint that was in place prior to this change, the compiler will emit a warning instead of an error, in order to allow existing Blueprints that may currently be "violating" const within the context of a const BPIE function to still compile, while still alerting to issues that should probably be addressed.
notes:
1) this also applies to BlueprintNativeEvent (BPNE) implementations, and also when implementing BPIE/BPNE interface methods that are also declared as const
2) a const BPIE/BPNE function with no return value and no output parameters will be implemented as a "normal" impure function, and not as an event as in the non-const case
3) a const BPIE/BPNE function with a return value and/or output parameters will currently be implemented as a pure function, regardless of whether or not BlueprintCallable is specified
4) this CL also retains some consolidation of static function validation code that i had previously done, mostly to allow static functions to more easily be whitelisted for const function graphs

diff --git a/Engine/Source/Editor/BlueprintGraph/Classes/EdGraphSchema_K2.h b/Engine/Source/Editor/BlueprintGraph/Classes/EdGraphSchema_K2.h
--- a/Engine/Source/Editor/BlueprintGraph/Classes/EdGraphSchema_K2.h
+++ b/Engine/Source/Editor/BlueprintGraph/Classes/EdGraphSchema_K2.h
@@ -484,6 +484,23 @@ public:
 	 */
 	bool IsCompositeGraph(const UEdGraph* TestEdGraph) const;
 
+	/**
+	 * Checks to see if the specified graph is a const function graph
+	 *
+	 * @param	TestEdGraph		Graph to test
+	 * @param	bOutIsEnforcingConstCorrectness (Optional) Whether or not this graph is enforcing const correctness during compilation
+	 * @return	true if this is a const function graph
+	 */
+	bool IsConstFunctionGraph(const UEdGraph* TestEdGraph, bool* bOutIsEnforcingConstCorrectness = nullptr) const;
+
+	/**
+	 * Checks to see if the specified graph is a static function graph
+	 *
+	 * @param	TestEdGraph		Graph to test
+	 * @return	true if this is a const function graph
+	 */
+	bool IsStaticFunctionGraph(const UEdGraph* TestEdGraph) const;
+
 	/**
 	 * Checks to see if a pin is an execution pin.
 	 *
@@ -682,13 +699,13 @@ public:
 	 * Determine if a function has a parameter of a specific type.
 	 *
 	 * @param	InFunction	  	The function to search.
-	 * @param	CallingContext  The blueprint that you're looking to call the function from (some functions hide different pins depending on the blueprint they're in)
+	 * @param	InGraph			The graph that you're looking to call the function from (some functions hide different pins depending on the graph they're in)
 	 * @param	DesiredPinType	The type that at least one function parameter needs to be.
 	 * @param	bWantOutput   	The direction that the parameter needs to be.
 	 *
 	 * @return	true if at least one parameter is of the correct type and direction.
 	 */
-	bool FunctionHasParamOfType(const UFunction* InFunction, UBlueprint const* CallingContext, const FEdGraphPinType& DesiredPinType, bool bWantOutput) const;
+	bool FunctionHasParamOfType(const UFunction* InFunction, UEdGraph const* InGraph, const FEdGraphPinType& DesiredPinType, bool bWantOutput) const;
 
 	/**
 	 * Add the specified flags to the function entry node of the graph, to make sure they get compiled in to the generated function
diff --git a/Engine/Source/Editor/BlueprintGraph/Classes/K2Node_FunctionEntry.h b/Engine/Source/Editor/BlueprintGraph/Classes/K2Node_FunctionEntry.h
--- a/Engine/Source/Editor/BlueprintGraph/Classes/K2Node_FunctionEntry.h
+++ b/Engine/Source/Editor/BlueprintGraph/Classes/K2Node_FunctionEntry.h
@@ -25,6 +25,14 @@ class UK2Node_FunctionEntry : public UK2Node_FunctionTerminator
 	UPROPERTY()
 	TArray<struct FBPVariableDescription> LocalVariables;
 
+	/** Whether or not to enforce const-correctness for const function overrides */
+	UPROPERTY()
+	bool bEnforceConstCorrectness;
+
+	// Begin UObject interface
+	virtual void Serialize(FArchive& Ar) override;
+	// End of UObject interface
+
 	// Begin UEdGraphNode interface
 	virtual void AllocateDefaultPins() override;
 	virtual FText GetNodeTitle(ENodeTitleType::Type TitleType) const override;
diff --git a/Engine/Source/Editor/BlueprintGraph/Private/BlueprintActionFilter.cpp b/Engine/Source/Editor/BlueprintGraph/Private/BlueprintActionFilter.cpp
--- a/Engine/Source/Editor/BlueprintGraph/Private/BlueprintActionFilter.cpp
+++ b/Engine/Source/Editor/BlueprintGraph/Private/BlueprintActionFilter.cpp
@@ -1049,7 +1049,7 @@ static bool BlueprintActionFilterImpl::IsFunctionMissingPinParam(FBlueprintActio
 				// want to flip the connotation here
 				bool const bWantsOutputConnection = (PinDir == EGPD_Input) ^ bIsEventSpawner;
 
-				if (K2Schema->FunctionHasParamOfType(AssociatedFunc, K2Node->GetBlueprint(), PinType, bWantsOutputConnection))
+				if (K2Schema->FunctionHasParamOfType(AssociatedFunc, K2Node->GetGraph(), PinType, bWantsOutputConnection))
 				{
 					bIsFilteredOut = false;
 				}
diff --git a/Engine/Source/Editor/BlueprintGraph/Private/CallFunctionHandler.cpp b/Engine/Source/Editor/BlueprintGraph/Private/CallFunctionHandler.cpp
--- a/Engine/Source/Editor/BlueprintGraph/Private/CallFunctionHandler.cpp
+++ b/Engine/Source/Editor/BlueprintGraph/Private/CallFunctionHandler.cpp
@@ -312,18 +312,66 @@ void FKCHandler_CallFunction::CreateFunctionCallStatement(FKismetFunctionContext
 			TArray<FBPTerminal*> ContextTerms;
 			if (SelfPin != NULL)
 			{
+				const bool bIsConstSelfContext = Context.IsConstFunction();
+				const bool bIsNonConstFunction = !Function->HasAnyFunctionFlags(FUNC_Const|FUNC_Static);
+				const bool bEnforceConstCorrectness = Context.EnforceConstCorrectness();
+				auto CheckAndAddSelfTermLambda = [this, &Node, &ContextTerms, bIsConstSelfContext, bIsNonConstFunction, bEnforceConstCorrectness](FBPTerminal* Target)
+				{
+					bool bIsSelfTerm = true;
+					if(Target != nullptr)
+					{
+						const UEdGraphPin* SourcePin = Cast<UEdGraphPin>(Target->Source);
+						bIsSelfTerm = (SourcePin == nullptr || CompilerContext.GetSchema()->IsSelfPin(*SourcePin));
+					}
+
+					// Ensure const correctness within the context of the function call:
+					//	a) Attempting to call a non-const, non-static function within a const function graph (i.e. 'const self' as context)
+					//	b) Attempting to call a non-const, non-static function with a 'const' term linked to the target pin as the function context
+					if(bIsSelfTerm && bIsConstSelfContext && bIsNonConstFunction)
+					{
+						// If we're not enforcing const correctness in this context, emit a warning here rather than an error, and allow compilation of this statement to proceed
+						if(Target != nullptr)
+						{
+							if(bEnforceConstCorrectness)
+							{
+								CompilerContext.MessageLog.Error(*LOCTEXT("NonConstFunctionCallOnReadOnlyTarget_Error", "Function @@ can modify state and cannot be called on @@ because it is a read-only Target in this context").ToString(), Node, Target->Source);
+							}
+							else
+							{
+								CompilerContext.MessageLog.Warning(*LOCTEXT("NonConstFunctionCallOnReadOnlyTarget_Warning", "Function @@ can modify state and should not be called on @@ because it is considered to be a read-only Target in this context").ToString(), Node, Target->Source);
+							}
+						}
+						else
+						{
+							if(bEnforceConstCorrectness)
+							{
+								CompilerContext.MessageLog.Error(*LOCTEXT("NonConstFunctionCallOnReadOnlySelfScope_Error", "Function @@ can modify state and cannot be called on 'self' because it is a read-only Target in this context").ToString(), Node);
+							}
+							else
+							{
+								CompilerContext.MessageLog.Warning(*LOCTEXT("NonConstFunctionCallOnReadOnlySelfScope_Warning", "Function @@ can modify state and should not be called on 'self' because it is considered to be a read-only Target in this context").ToString(), Node);
+							}
+						}
+					}
+
+					ContextTerms.Add(Target);
+				};
+
 				if( SelfPin->LinkedTo.Num() > 0 )
 				{
 					for(int32 i = 0; i < SelfPin->LinkedTo.Num(); i++)
 					{
 						FBPTerminal** pContextTerm = Context.NetMap.Find(SelfPin->LinkedTo[i]);
-						ContextTerms.Add((pContextTerm != NULL) ? *pContextTerm : NULL);
+						if(ensure(pContextTerm != nullptr))
+						{
+							CheckAndAddSelfTermLambda(*pContextTerm);
+						}
 					}
 				}
 				else
 				{
 					FBPTerminal** pContextTerm = Context.NetMap.Find(SelfPin);
-					ContextTerms.Add((pContextTerm != NULL) ? *pContextTerm : NULL);
+					CheckAndAddSelfTermLambda((pContextTerm != nullptr) ? *pContextTerm : nullptr);
 				}
 			}
 
@@ -507,7 +555,7 @@ void FKCHandler_CallFunction::RegisterNets(FKismetFunctionContext& Context, UEdG
 		}
 		if (Function->HasMetaData(FBlueprintMetadata::MD_WorldContext))
 		{
-			const bool bHasIntrinsicWorldContext = Context.Blueprint->ParentClass->GetDefaultObject()->ImplementsGetWorld();
+			const bool bHasIntrinsicWorldContext = !K2Schema->IsStaticFunctionGraph(Context.SourceGraph) && Context.Blueprint->ParentClass->GetDefaultObject()->ImplementsGetWorld();
 
 			FString const WorldContextPinName = Function->GetMetaData(FBlueprintMetadata::MD_WorldContext);
 
diff --git a/Engine/Source/Editor/BlueprintGraph/Private/EdGraphSchema_K2.cpp b/Engine/Source/Editor/BlueprintGraph/Private/EdGraphSchema_K2.cpp
--- a/Engine/Source/Editor/BlueprintGraph/Private/EdGraphSchema_K2.cpp
+++ b/Engine/Source/Editor/BlueprintGraph/Private/EdGraphSchema_K2.cpp
@@ -768,8 +768,8 @@ struct FNoOutputParametersHelper
 
 bool UEdGraphSchema_K2::FunctionCanBePlacedAsEvent(const UFunction* InFunction)
 {
-	// First check we are override-able
-	if (!InFunction || !CanKismetOverrideFunction(InFunction))
+	// First check we are override-able, non-static and non-const
+	if (!InFunction || !CanKismetOverrideFunction(InFunction) || InFunction->HasAnyFunctionFlags(FUNC_Static|FUNC_Const))
 	{
 		return false;
 	}
@@ -807,10 +807,10 @@ void UEdGraphSchema_K2::GetAutoEmitTermParameters(const UFunction* Function, TAr
 	}
 }
 
-bool UEdGraphSchema_K2::FunctionHasParamOfType(const UFunction* InFunction, UBlueprint const* CallingContext, const FEdGraphPinType& DesiredPinType, bool bWantOutput) const
+bool UEdGraphSchema_K2::FunctionHasParamOfType(const UFunction* InFunction, UEdGraph const* InGraph, const FEdGraphPinType& DesiredPinType, bool bWantOutput) const
 {
 	TSet<FString> HiddenPins;
-	FBlueprintEditorUtils::GetHiddenPinsForFunction(CallingContext, InFunction, HiddenPins);
+	FBlueprintEditorUtils::GetHiddenPinsForFunction(InGraph, InFunction, HiddenPins);
 
 	// Iterate over all params of function
 	for (TFieldIterator<UProperty> PropIt(InFunction); PropIt && (PropIt->PropertyFlags & CPF_Parm); ++PropIt)
@@ -2855,6 +2855,9 @@ void UEdGraphSchema_K2::CreateFunctionGraphTerminators(UEdGraph& Graph, UClass*
 	UFunction* InterfaceToImplement = FindField<UFunction>(Class, GraphName);
 	if (InterfaceToImplement)
 	{
+		// Add modifier flags from the declaration
+		EntryNode->ExtraFlags |= InterfaceToImplement->FunctionFlags & (FUNC_Const | FUNC_Static | FUNC_BlueprintPure);
+
 		// See if any function params are marked as out
 		bool bHasOutParam =  false;
 		for( TFieldIterator<UProperty> It(InterfaceToImplement); It && (It->PropertyFlags & CPF_Parm); ++It )
@@ -3995,7 +3998,7 @@ bool UEdGraphSchema_K2::CanPromotePinToVariable( const UEdGraphPin& Pin ) const
 	const UK2Node* Node = Cast<UK2Node>(Pin.GetOwningNode());
 	const UBlueprint* OwningBlueprint = Node->GetBlueprint();
 	
-	if (!OwningBlueprint || (OwningBlueprint->BlueprintType == BPTYPE_MacroLibrary) || (OwningBlueprint->BlueprintType == BPTYPE_FunctionLibrary))
+	if (!OwningBlueprint || (OwningBlueprint->BlueprintType == BPTYPE_MacroLibrary) || (OwningBlueprint->BlueprintType == BPTYPE_FunctionLibrary) || IsStaticFunctionGraph(Node->GetGraph()))
 	{
 		// Never allow promotion in macros, because there's not a scope to define them in
 		bCanPromote = false;
@@ -4134,10 +4137,23 @@ void UEdGraphSchema_K2::GetGraphDisplayInformation(const UEdGraph& Graph, /*out*
 		DisplayInfo.DocExcerptName = TEXT("StateMachine");
 	}
 
-	// Add pure to notes if set
-	if (Function && Function->HasAnyFunctionFlags(FUNC_BlueprintPure))
+	// Add pure/static/const to notes if set
+	if (Function)
 	{
-		DisplayInfo.Notes.Add(TEXT("pure"));
+		if(Function->HasAnyFunctionFlags(FUNC_BlueprintPure))
+		{
+			DisplayInfo.Notes.Add(TEXT("pure"));
+		}
+
+		// since 'static' is implied in a function library, not going to display it (to be consistent with previous behavior)
+		if(Function->HasAnyFunctionFlags(FUNC_Static) && Blueprint->BlueprintType != BPTYPE_FunctionLibrary)
+		{
+			DisplayInfo.Notes.Add(TEXT("static"));
+		}
+		else if(Function->HasAnyFunctionFlags(FUNC_Const))
+		{
+			DisplayInfo.Notes.Add(TEXT("const"));
+		}
 	}
 
 	// Mark transient graphs as obviously so
@@ -4217,6 +4233,58 @@ bool UEdGraphSchema_K2::IsCompositeGraph( const UEdGraph* TestEdGraph ) const
 	return false;
 }
 
+bool UEdGraphSchema_K2::IsConstFunctionGraph( const UEdGraph* TestEdGraph, bool* bOutIsEnforcingConstCorrectness ) const
+{
+	check(TestEdGraph);
+
+	const EGraphType GraphType = GetGraphType(TestEdGraph);
+	if(GraphType == GT_Function) 
+	{
+		// Find the entry node for the function graph and see if the 'const' flag is set
+		for(auto I = TestEdGraph->Nodes.CreateConstIterator(); I; ++I)
+		{
+			UEdGraphNode* Node = *I;
+			if(auto EntryNode = Cast<UK2Node_FunctionEntry>(Node))
+			{
+				if(bOutIsEnforcingConstCorrectness != nullptr)
+				{
+					*bOutIsEnforcingConstCorrectness = EntryNode->bEnforceConstCorrectness;
+				}
+
+				return (EntryNode->ExtraFlags & FUNC_Const) != 0;
+			}
+		}
+	}
+
+	if(bOutIsEnforcingConstCorrectness != nullptr)
+	{
+		*bOutIsEnforcingConstCorrectness = false;
+	}
+
+	return false;
+}
+
+bool UEdGraphSchema_K2::IsStaticFunctionGraph( const UEdGraph* TestEdGraph ) const
+{
+	check(TestEdGraph);
+
+	const EGraphType GraphType = GetGraphType(TestEdGraph);
+	if(GraphType == GT_Function) 
+	{
+		// Find the entry node for the function graph and see if the 'static' flag is set
+		for(auto I = TestEdGraph->Nodes.CreateConstIterator(); I; ++I)
+		{
+			UEdGraphNode* Node = *I;
+			if(auto EntryNode = Cast<UK2Node_FunctionEntry>(Node))
+			{
+				return (EntryNode->ExtraFlags & FUNC_Static) != 0;
+			}
+		}
+	}
+
+	return false;
+}
+
 void UEdGraphSchema_K2::DroppedAssetsOnGraph(const TArray<FAssetData>& Assets, const FVector2D& GraphPosition, UEdGraph* Graph) const 
 {
 	UBlueprint* Blueprint = FBlueprintEditorUtils::FindBlueprintForGraph(Graph);
diff --git a/Engine/Source/Editor/BlueprintGraph/Private/K2ActionMenuBuilder.cpp b/Engine/Source/Editor/BlueprintGraph/Private/K2ActionMenuBuilder.cpp
--- a/Engine/Source/Editor/BlueprintGraph/Private/K2ActionMenuBuilder.cpp
+++ b/Engine/Source/Editor/BlueprintGraph/Private/K2ActionMenuBuilder.cpp
@@ -2569,7 +2569,7 @@ void FK2ActionMenuBuilder::GetFuncNodesWithPinType(FBlueprintGraphActionListBuil
 			if (K2Schema->CanUserKismetCallFunction(Function))
 			{
 				if ((false == bPureOnly || Function->HasAnyFunctionFlags(FUNC_BlueprintPure)) && 
-					K2Schema->FunctionHasParamOfType(Function, ClassBlueprint, DesiredPinType, bWantOutput)) 
+					K2Schema->FunctionHasParamOfType(Function, ContextMenuBuilder.CurrentGraph, DesiredPinType, bWantOutput)) 
 				{
 					const bool bShowMakeOnTop = bUseNativeMake && Function->HasMetaData(NativeMakeFunc);
 					const bool bShowBrakeOnTop = bUseNativeBrake && Function->HasMetaData(NativeBrakeFunc);
diff --git a/Engine/Source/Editor/BlueprintGraph/Private/K2Node_BaseAsyncTask.cpp b/Engine/Source/Editor/BlueprintGraph/Private/K2Node_BaseAsyncTask.cpp
--- a/Engine/Source/Editor/BlueprintGraph/Private/K2Node_BaseAsyncTask.cpp
+++ b/Engine/Source/Editor/BlueprintGraph/Private/K2Node_BaseAsyncTask.cpp
@@ -105,7 +105,7 @@ void UK2Node_BaseAsyncTask::AllocateDefaultPins()
 	if (Function)
 	{
 		TSet<FString> PinsToHide;
-		FBlueprintEditorUtils::GetHiddenPinsForFunction(GetBlueprint(), Function, PinsToHide);
+		FBlueprintEditorUtils::GetHiddenPinsForFunction(GetGraph(), Function, PinsToHide);
 		for (TFieldIterator<UProperty> PropIt(Function); PropIt && (PropIt->PropertyFlags & CPF_Parm); ++PropIt)
 		{
 			UProperty* Param = *PropIt;
diff --git a/Engine/Source/Editor/BlueprintGraph/Private/K2Node_CallFunction.cpp b/Engine/Source/Editor/BlueprintGraph/Private/K2Node_CallFunction.cpp
--- a/Engine/Source/Editor/BlueprintGraph/Private/K2Node_CallFunction.cpp
+++ b/Engine/Source/Editor/BlueprintGraph/Private/K2Node_CallFunction.cpp
@@ -780,11 +780,12 @@ bool UK2Node_CallFunction::CreatePinsForFunctionCall(const UFunction* Function)
 	const bool bIsProtectedFunc = Function->GetBoolMetaData(FBlueprintMetadata::MD_Protected);
 	const bool bIsStaticFunc = Function->HasAllFunctionFlags(FUNC_Static);
 
-	UBlueprint* BP = GetBlueprint();
+	UEdGraph const* const Graph = GetGraph();
+	UBlueprint* BP = FBlueprintEditorUtils::FindBlueprintForGraph(Graph);
 	ensure(BP);
 	if (BP != nullptr)
 	{
-		const bool bIsFunctionCompatibleWithSelf = BP->SkeletonGeneratedClass->IsChildOf(FunctionOwnerClass);
+		const bool bIsFunctionCompatibleWithSelf = !K2Schema->IsStaticFunctionGraph(GetGraph()) && BP->SkeletonGeneratedClass->IsChildOf(FunctionOwnerClass);
 
 		if (bIsStaticFunc)
 		{
@@ -807,7 +808,7 @@ bool UK2Node_CallFunction::CreatePinsForFunctionCall(const UFunction* Function)
 
 	// Build a list of the pins that should be hidden for this function (ones that are automagically filled in by the K2 compiler)
 	TSet<FString> PinsToHide;
-	FBlueprintEditorUtils::GetHiddenPinsForFunction(BP, Function, PinsToHide);
+	FBlueprintEditorUtils::GetHiddenPinsForFunction(Graph, Function, PinsToHide);
 
 	const bool bShowWorldContextPin = ((PinsToHide.Num() > 0) && BP && BP->ParentClass && BP->ParentClass->HasMetaData(FBlueprintMetadata::MD_ShowWorldContextPin));
 
@@ -1514,7 +1515,7 @@ void UK2Node_CallFunction::PostPasteNode()
 	{
 		// After pasting we need to go through and ensure the hidden the self pins is correct in case the source blueprint had different metadata
 		TSet<FString> PinsToHide;
-		FBlueprintEditorUtils::GetHiddenPinsForFunction(GetBlueprint(), Function, PinsToHide);
+		FBlueprintEditorUtils::GetHiddenPinsForFunction(GetGraph(), Function, PinsToHide);
 
 		const bool bShowWorldContextPin = ((PinsToHide.Num() > 0) && GetBlueprint()->ParentClass->HasMetaData(FBlueprintMetadata::MD_ShowWorldContextPin));
 
@@ -1679,8 +1680,7 @@ void UK2Node_CallFunction::ExpandNode(class FKismetCompilerContext& CompilerCont
 	UFunction* Function = GetTargetFunction();
 
 	// connect DefaultToSelf and WorldContext inside static functions to proper 'self'  
-	const bool bInsideBpFuncLibrary = CompilerContext.Blueprint && (BPTYPE_FunctionLibrary == CompilerContext.Blueprint->BlueprintType);
-	if (bInsideBpFuncLibrary && SourceGraph && Function)
+	if (SourceGraph && Schema->IsStaticFunctionGraph(SourceGraph) && Function)
 	{
 		TArray<UK2Node_FunctionEntry*> EntryPoints;
 		SourceGraph->GetNodesOfClass(EntryPoints);
diff --git a/Engine/Source/Editor/BlueprintGraph/Private/K2Node_FunctionEntry.cpp b/Engine/Source/Editor/BlueprintGraph/Private/K2Node_FunctionEntry.cpp
--- a/Engine/Source/Editor/BlueprintGraph/Private/K2Node_FunctionEntry.cpp
+++ b/Engine/Source/Editor/BlueprintGraph/Private/K2Node_FunctionEntry.cpp
@@ -102,15 +102,30 @@ struct FFunctionEntryHelper
 
 	static bool RequireWorldContextParameter(const UK2Node_FunctionEntry* Node)
 	{
-		auto Blueprint = Node->GetBlueprint();
-		ensure(Blueprint);
-		return Blueprint && (BPTYPE_FunctionLibrary == Blueprint->BlueprintType);
+		const UEdGraphSchema_K2* K2Schema = GetDefault<UEdGraphSchema_K2>();
+		return K2Schema->IsStaticFunctionGraph(Node->GetGraph());
 	}
 };
 
 UK2Node_FunctionEntry::UK2Node_FunctionEntry(const FObjectInitializer& ObjectInitializer)
 	: Super(ObjectInitializer)
 {
+	// Enforce const-correctness by default
+	bEnforceConstCorrectness = true;
+}
+
+void UK2Node_FunctionEntry::Serialize(FArchive& Ar)
+{
+	Super::Serialize(Ar);
+
+	if (Ar.IsLoading())
+	{
+		if (Ar.UE4Ver() < VER_UE4_BLUEPRINT_ENFORCE_CONST_IN_FUNCTION_OVERRIDES)
+		{
+			// Allow legacy implementations to violate const-correctness
+			bEnforceConstCorrectness = false;
+		}
+	}
 }
 
 FText UK2Node_FunctionEntry::GetNodeTitle(ENodeTitleType::Type TitleType) const
diff --git a/Engine/Source/Editor/BlueprintGraph/Private/K2Node_Self.cpp b/Engine/Source/Editor/BlueprintGraph/Private/K2Node_Self.cpp
--- a/Engine/Source/Editor/BlueprintGraph/Private/K2Node_Self.cpp
+++ b/Engine/Source/Editor/BlueprintGraph/Private/K2Node_Self.cpp
@@ -73,12 +73,11 @@ FNodeHandlingFunctor* UK2Node_Self::CreateNodeHandler(FKismetCompilerContext& Co
 void UK2Node_Self::ValidateNodeDuringCompilation(class FCompilerResultsLog& MessageLog) const
 {
 	Super::ValidateNodeDuringCompilation(MessageLog);
-	const UBlueprint* Blueprint = GetBlueprint();
-	const UClass* MyClass = Blueprint ? Blueprint->GeneratedClass : NULL;
-	const bool bValidClass = !MyClass || !MyClass->IsChildOf(UBlueprintFunctionLibrary::StaticClass());
-	if (!bValidClass)
+
+	const UEdGraphSchema_K2* Schema = GetDefault<UEdGraphSchema_K2>();
+	if(Schema->IsStaticFunctionGraph(GetGraph()))
 	{
-		MessageLog.Warning(*NSLOCTEXT("K2Node", "InvalidSelfNode", "Self node @@ cannot be used in static library.").ToString(), this);
+		MessageLog.Warning(*NSLOCTEXT("K2Node", "InvalidSelfNode", "Self node @@ cannot be used in a static function.").ToString(), this);
 	}
 }
 
diff --git a/Engine/Source/Editor/BlueprintGraph/Private/StructMemberNodeHandlers.cpp b/Engine/Source/Editor/BlueprintGraph/Private/StructMemberNodeHandlers.cpp
--- a/Engine/Source/Editor/BlueprintGraph/Private/StructMemberNodeHandlers.cpp
+++ b/Engine/Source/Editor/BlueprintGraph/Private/StructMemberNodeHandlers.cpp
@@ -28,7 +28,7 @@ static FBPTerminal* RegisterStructVar(FCompilerResultsLog& MessageLog, FKismetFu
 		//@TODO: Needed? Context.NetMap.Add(Net, Term);
 
 		// Read-only variables and variables in const classes are both const
-		if (BoundProperty->HasAnyPropertyFlags(CPF_BlueprintReadOnly) || Context.IsConstFunction())
+		if (BoundProperty->HasAnyPropertyFlags(CPF_BlueprintReadOnly) || (Context.IsConstFunction() && Context.NewClass->IsChildOf(SearchScope)))
 		{
 			Term->bIsConst = true;
 		}
@@ -64,7 +64,7 @@ static void ResolveAndRegisterScopedStructTerm(FCompilerResultsLog& MessageLog,
 		Term->Context = ContextTerm;
 
 		// Read-only variables and variables in const classes are both const
-		if (BoundProperty->HasAnyPropertyFlags(CPF_BlueprintReadOnly) || Context.IsConstFunction())
+		if (BoundProperty->HasAnyPropertyFlags(CPF_BlueprintReadOnly) || (Context.IsConstFunction() && Context.NewClass->IsChildOf(StructType)))
 		{
 			Term->bIsConst = true;
 		}
diff --git a/Engine/Source/Editor/BlueprintGraph/Private/VariableSetHandler.cpp b/Engine/Source/Editor/BlueprintGraph/Private/VariableSetHandler.cpp
--- a/Engine/Source/Editor/BlueprintGraph/Private/VariableSetHandler.cpp
+++ b/Engine/Source/Editor/BlueprintGraph/Private/VariableSetHandler.cpp
@@ -82,7 +82,23 @@ void FKCHandler_VariableSet::InnerAssignment(FKismetFunctionContext& Context, UE
 
 		if (!(*VariableTerm)->IsTermWritable())
 		{
-			CompilerContext.MessageLog.Error(*LOCTEXT("WriteConst_Error", "Cannot write to const @@").ToString(), VariablePin);
+			// If the term is not explicitly marked as read-only, then we're attempting to set a variable on a const target
+			if(!(*VariableTerm)->AssociatedVarProperty->HasAnyPropertyFlags(CPF_BlueprintReadOnly))
+			{
+				if(Context.EnforceConstCorrectness())
+				{
+					CompilerContext.MessageLog.Error(*LOCTEXT("WriteToReadOnlyContext_Error", "Variable @@ is read-only within this context and cannot be set to a new value").ToString(), VariablePin);
+				}
+				else
+				{
+					// Warn, but still allow compilation to succeed
+					CompilerContext.MessageLog.Warning(*LOCTEXT("WriteToReadOnlyContext_Warning", "Variable @@ is considered to be read-only within this context and should not be set to a new value").ToString(), VariablePin);
+				}
+			}
+			else
+			{
+				CompilerContext.MessageLog.Error(*LOCTEXT("WriteConst_Error", "Cannot write to const @@").ToString(), VariablePin);
+			}
 		}
 	}
 	else
diff --git a/Engine/Source/Editor/Kismet/Private/BlueprintActionMenuUtils.cpp b/Engine/Source/Editor/Kismet/Private/BlueprintActionMenuUtils.cpp
--- a/Engine/Source/Editor/Kismet/Private/BlueprintActionMenuUtils.cpp
+++ b/Engine/Source/Editor/Kismet/Private/BlueprintActionMenuUtils.cpp
@@ -483,6 +483,16 @@ void FBlueprintActionMenuUtils::MakeContextMenu(FBlueprintActionContext const& C
 		bCanHaveActorComponents &= FBlueprintEditorUtils::DoesSupportComponents(Blueprint);
 	}
 
+	if(bCanHaveActorComponents)
+	{
+		// Don't allow actor components in static function graphs
+		UEdGraphSchema_K2 const* K2Schema = GetDefault<UEdGraphSchema_K2>();
+		for (UEdGraph* Graph : Context.Graphs)
+		{
+			bCanHaveActorComponents &= !K2Schema->IsStaticFunctionGraph(Graph);
+		}
+	}
+
 	if (bAddTargetContext)
 	{
 		bool bContextPinIsObj = false;
diff --git a/Engine/Source/Editor/KismetCompiler/Private/KismetCompiler.cpp b/Engine/Source/Editor/KismetCompiler/Private/KismetCompiler.cpp
--- a/Engine/Source/Editor/KismetCompiler/Private/KismetCompiler.cpp
+++ b/Engine/Source/Editor/KismetCompiler/Private/KismetCompiler.cpp
@@ -1029,44 +1029,48 @@ void FKismetCompilerContext::ValidateSelfPinsInGraph(const UEdGraph* SourceGraph
 			{
 				if (const UEdGraphPin* Pin = Node->Pins[PinIndex])
 				{
-					if (Schema->IsSelfPin(*Pin) && (Pin->LinkedTo.Num() == 0))
+					if (Schema->IsSelfPin(*Pin) && (Pin->LinkedTo.Num() == 0) && Pin->DefaultObject == nullptr)
 					{
 						FEdGraphPinType SelfType;
 						SelfType.PinCategory = Schema->PC_Object;
 						SelfType.PinSubCategory = Schema->PSC_Self;
 
-						if (!Schema->ArePinTypesCompatible(SelfType, Pin->PinType, NewClass))
+						FString ErrorMsg;
+						if(Blueprint->BlueprintType != BPTYPE_FunctionLibrary && Schema->IsStaticFunctionGraph(SourceGraph))
 						{
-							if (Pin->DefaultObject == NULL)
+							ErrorMsg = FString::Printf(*LOCTEXT("PinMustHaveConnection_Static_Error", "'@@' must have a connection, because %s is a static function and will not be bound to instances of this blueprint.").ToString(), *SourceGraph->GetName());
+						}
+						else if (!Schema->ArePinTypesCompatible(SelfType, Pin->PinType, NewClass))
+						{
+							FString PinType = Pin->PinType.PinCategory;
+							if ((Pin->PinType.PinCategory == Schema->PC_Object)    || 
+								(Pin->PinType.PinCategory == Schema->PC_Interface) ||
+								(Pin->PinType.PinCategory == Schema->PC_Class))
 							{
-								FString PinType = Pin->PinType.PinCategory;
-								if ((Pin->PinType.PinCategory == Schema->PC_Object)    || 
-									(Pin->PinType.PinCategory == Schema->PC_Interface) ||
-									(Pin->PinType.PinCategory == Schema->PC_Class))
-								{
-									if (Pin->PinType.PinSubCategoryObject.IsValid())
-									{
-										PinType = Pin->PinType.PinSubCategoryObject->GetName();
-									}
-									else
-									{
-										PinType = TEXT("");
-									}
-								}
-
-								FString ErrorMsg;
-								if(PinType.IsEmpty())
+								if (Pin->PinType.PinSubCategoryObject.IsValid())
 								{
-									ErrorMsg = FString::Printf(*LOCTEXT("PinMustHaveConnection_NoType_Error", "'@@' must have a connection").ToString());
+									PinType = Pin->PinType.PinSubCategoryObject->GetName();
 								}
 								else
 								{
-									ErrorMsg = FString::Printf(*LOCTEXT("PinMustHaveConnection_Error", "This blueprint (self) is not a %s, therefore '@@' must have a connection").ToString(), *PinType);	
+									PinType = TEXT("");
 								}
+							}
 
-								MessageLog.Error(*ErrorMsg, Pin);
+							if(PinType.IsEmpty())
+							{
+								ErrorMsg = FString::Printf(*LOCTEXT("PinMustHaveConnection_NoType_Error", "'@@' must have a connection.").ToString());
+							}
+							else
+							{
+								ErrorMsg = FString::Printf(*LOCTEXT("PinMustHaveConnection_WrongClass_Error", "This blueprint (self) is not a %s, therefore '@@' must have a connection.").ToString(), *PinType);	
 							}
 						}
+
+						if(!ErrorMsg.IsEmpty())
+						{
+							MessageLog.Error(*ErrorMsg, Pin);
+						}
 					}
 				}
 			}
@@ -2467,8 +2471,8 @@ void FKismetCompilerContext::MergeUbergraphPagesIn(UEdGraph* Ubergraph)
 // Expands out nodes that need it
 void FKismetCompilerContext::ExpansionStep(UEdGraph* Graph, bool bAllowUbergraphExpansions)
 {
-	// Node expansion may affect a signature of function in FunctionLibrary
-	if (bIsFullCompile || (EBlueprintType::BPTYPE_FunctionLibrary == Blueprint->BlueprintType))
+	// Node expansion may affect the signature of a static function
+	if (bIsFullCompile || Schema->IsStaticFunctionGraph(Graph))
 	{
 		BP_SCOPED_COMPILER_EVENT_STAT(EKismetCompilerStats_Expansion);
 
@@ -2992,9 +2996,10 @@ void FKismetCompilerContext::ProcessOneFunctionGraph(UEdGraph* SourceGraph, bool
 			Context.MarkAsInterfaceStub();
 		}
 
-		if (FBlueprintEditorUtils::IsBlueprintConst(Blueprint))
+		bool bEnforceConstCorrectness = true;
+		if (FBlueprintEditorUtils::IsBlueprintConst(Blueprint) || Schema->IsConstFunctionGraph(Context.SourceGraph, &bEnforceConstCorrectness))
 		{
-			Context.MarkAsConstFunction();
+			Context.MarkAsConstFunction(bEnforceConstCorrectness);
 		}
 
 		if ( bInternalFunction )
diff --git a/Engine/Source/Editor/KismetCompiler/Private/KismetCompilerMisc.cpp b/Engine/Source/Editor/KismetCompiler/Private/KismetCompilerMisc.cpp
--- a/Engine/Source/Editor/KismetCompiler/Private/KismetCompilerMisc.cpp
+++ b/Engine/Source/Editor/KismetCompiler/Private/KismetCompilerMisc.cpp
@@ -833,7 +833,7 @@ void FNodeHandlingFunctor::ResolveAndRegisterScopedTerm(FKismetFunctionContext&
 		{
 			Term->bIsLocal = true;
 		}
-		else if (BoundProperty->HasAnyPropertyFlags(CPF_BlueprintReadOnly) || Context.IsConstFunction())
+		else if (BoundProperty->HasAnyPropertyFlags(CPF_BlueprintReadOnly) || (Context.IsConstFunction() && Context.NewClass->IsChildOf(SearchScope)))
 		{
 			// Read-only variables and variables in const classes are both const
 			Term->bIsConst = true;
@@ -999,6 +999,7 @@ FKismetFunctionContext::FKismetFunctionContext(FCompilerResultsLog& InMessageLog
 	, bCannotBeCalledFromOtherKismet(false)
 	, bIsInterfaceStub(false)
 	, bIsConstFunction(false)
+	, bEnforceConstCorrectness(false)
 	// only need debug-data when running in the editor app:
 	, bCreateDebugData(GIsEditor && !IsRunningCommandlet())
 	, bIsSimpleStubGraphWithNoParams(false)
diff --git a/Engine/Source/Editor/KismetCompiler/Public/KismetCompiledFunctionContext.h b/Engine/Source/Editor/KismetCompiler/Public/KismetCompiledFunctionContext.h
--- a/Engine/Source/Editor/KismetCompiler/Public/KismetCompiledFunctionContext.h
+++ b/Engine/Source/Editor/KismetCompiler/Public/KismetCompiledFunctionContext.h
@@ -77,6 +77,7 @@ public:
 	bool bCannotBeCalledFromOtherKismet;
 	bool bIsInterfaceStub;
 	bool bIsConstFunction;
+	bool bEnforceConstCorrectness;
 	bool bCreateDebugData;
 	bool bIsSimpleStubGraphWithNoParams;
 	uint32 NetFlags;
@@ -130,9 +131,10 @@ public:
 		bIsInterfaceStub = true;
 	}
 
-	void MarkAsConstFunction()
+	void MarkAsConstFunction(bool bInEnforceConstCorrectness)
 	{
 		bIsConstFunction = true;
+		bEnforceConstCorrectness = bInEnforceConstCorrectness;
 	}
 
 	bool IsInterfaceStub() const
@@ -156,6 +158,11 @@ public:
 		return bIsConstFunction;
 	}
 
+	bool EnforceConstCorrectness() const
+	{
+		return bEnforceConstCorrectness;
+	}
+
 	void MarkAsNetFunction(uint32 InFunctionFlags)
 	{
 		NetFlags = InFunctionFlags & (FUNC_NetFuncFlags);
diff --git a/Engine/Source/Editor/UnrealEd/Private/Kismet2/BlueprintEditorUtils.cpp b/Engine/Source/Editor/UnrealEd/Private/Kismet2/BlueprintEditorUtils.cpp
--- a/Engine/Source/Editor/UnrealEd/Private/Kismet2/BlueprintEditorUtils.cpp
+++ b/Engine/Source/Editor/UnrealEd/Private/Kismet2/BlueprintEditorUtils.cpp
@@ -3161,7 +3161,7 @@ UEdGraph* FBlueprintEditorUtils::GetDelegateSignatureGraphByName(UBlueprint* Blu
 }
 
 // Gets a list of pins that should hidden for a given function
-void FBlueprintEditorUtils::GetHiddenPinsForFunction(UBlueprint const* CallingContext, UFunction const* Function, TSet<FString>& HiddenPins)
+void FBlueprintEditorUtils::GetHiddenPinsForFunction(UEdGraph const* Graph, UFunction const* Function, TSet<FString>& HiddenPins)
 {
 	check(Function != NULL);
 	TMap<FName, FString>* MetaData = UMetaData::GetMapForObject(Function);	
@@ -3188,18 +3188,23 @@ void FBlueprintEditorUtils::GetHiddenPinsForFunction(UBlueprint const* CallingCo
 			}
 			else if (Key == FBlueprintMetadata::MD_WorldContext)
 			{
-				bool bHasIntrinsicWorldContext = false;
-
-				if (GEngine && CallingContext && CallingContext->ParentClass)
+				const UEdGraphSchema_K2* K2Schema = GetDefault<UEdGraphSchema_K2>();
+				if(!K2Schema->IsStaticFunctionGraph(Graph))
 				{
-					bHasIntrinsicWorldContext = CallingContext->ParentClass->GetDefaultObject()->ImplementsGetWorld();
-				}
+					bool bHasIntrinsicWorldContext = false;
 
-				// if the blueprint has world context that we can lookup with "self", 
-				// then we can hide this pin (and default it to self)
-				if (bHasIntrinsicWorldContext)
-				{
-					HiddenPins.Add(It.Value());
+					UBlueprint const* CallingContext = FindBlueprintForGraph(Graph);
+					if (GEngine && CallingContext && CallingContext->ParentClass)
+					{
+						bHasIntrinsicWorldContext = CallingContext->ParentClass->GetDefaultObject()->ImplementsGetWorld();
+					}
+
+					// if the blueprint has world context that we can lookup with "self", 
+					// then we can hide this pin (and default it to self)
+					if (bHasIntrinsicWorldContext)
+					{
+						HiddenPins.Add(It.Value());
+					}
 				}
 			}
 		}
diff --git a/Engine/Source/Editor/UnrealEd/Public/Kismet2/BlueprintEditorUtils.h b/Engine/Source/Editor/UnrealEd/Public/Kismet2/BlueprintEditorUtils.h
--- a/Engine/Source/Editor/UnrealEd/Public/Kismet2/BlueprintEditorUtils.h
+++ b/Engine/Source/Editor/UnrealEd/Public/Kismet2/BlueprintEditorUtils.h
@@ -452,13 +452,13 @@ public:
 	static bool IsMathExpressionGraph(const UEdGraph* InGraph);
 
 	/**
-	 * Gets a list of pins that should hidden for a given function
+	 * Gets a list of pins that should hidden for a given function in a given graph
 	 *
-	 * @param			CallingContext  The blueprint that you're looking to call the function from (some functions hide different pins depending on the blueprint they're in)
+	 * @param			Graph			The graph that you're looking to call the function from (some functions hide different pins depending on the graph they're in)
 	 * @param			Function		The function to consider
 	 * @param [out]		HiddenPins		Set of pins that should be hidden
 	 */
-	static void GetHiddenPinsForFunction(UBlueprint const* CallingContext, UFunction const* Function, TSet<FString>& HiddenPins);
+	static void GetHiddenPinsForFunction(UEdGraph const* Graph, UFunction const* Function, TSet<FString>& HiddenPins);
 
 	/** Makes sure that calls to parent functions are valid, and removes them if not */
 	static void ConformCallsToParentFunctions(UBlueprint* Blueprint);
diff --git a/Engine/Source/Programs/UnrealHeaderTool/Private/CodeGenerator.cpp b/Engine/Source/Programs/UnrealHeaderTool/Private/CodeGenerator.cpp
--- a/Engine/Source/Programs/UnrealHeaderTool/Private/CodeGenerator.cpp
+++ b/Engine/Source/Programs/UnrealHeaderTool/Private/CodeGenerator.cpp
@@ -1519,7 +1519,8 @@ void FNativeClassHeaderGenerator::ExportInterfaceCallFunctions( const TArray<UFu
 			GeneratedPackageCPP.Logf(TEXT("%sParms.%s=%s;%s"), FCString::Spc(12), *Param->GetName(), *Param->GetName(), LINE_TERMINATOR);
 		}
 
-		GeneratedPackageCPP.Logf(TEXT("%sO->ProcessEvent(Func, %s);%s"), FCString::Spc(12), bHasParms ? TEXT("&Parms") : TEXT("NULL"), LINE_TERMINATOR);
+		const FString ObjectRef = FunctionData.FunctionReference->HasAllFunctionFlags(FUNC_Const) ? FString::Printf(TEXT("const_cast<UObject*>(O)"), *ClassName) : TEXT("O");
+		GeneratedPackageCPP.Logf(TEXT("%s%s->ProcessEvent(Func, %s);%s"), FCString::Spc(12), *ObjectRef, bHasParms ? TEXT("&Parms") : TEXT("NULL"), LINE_TERMINATOR);
 
 		for (auto It = Parameters.Parms.CreateConstIterator(); It; ++It)
 		{
diff --git a/Engine/Source/Runtime/Core/Public/UObject/ObjectVersion.h b/Engine/Source/Runtime/Core/Public/UObject/ObjectVersion.h
--- a/Engine/Source/Runtime/Core/Public/UObject/ObjectVersion.h
+++ b/Engine/Source/Runtime/Core/Public/UObject/ObjectVersion.h
@@ -651,6 +651,8 @@ enum EUnrealEngineObjectUE4Version
 	VER_UE4_ANIMATION_ADD_TRACKCURVES,
 	// Removed BranchingPoints from AnimMontages and converted them to regular AnimNotifies.
 	VER_UE4_MONTAGE_BRANCHING_POINT_REMOVAL,
+	// Enforce const-correctness in Blueprint implementations of native C++ const class methods
+	VER_UE4_BLUEPRINT_ENFORCE_CONST_IN_FUNCTION_OVERRIDES,
 	// -----<new versions can be added before this line>-------------------------------------------------
 	// - this needs to be the last line (see note below)
 	VER_UE4_AUTOMATIC_VERSION_PLUS_ONE,
