BP, DataTable: Better validation for GetDataTableRow node. BlueprintCallable function UDataTableFunctionLibrary::GetDataTableRowNames added.

diff --git a/Engine/Source/Editor/BlueprintGraph/Classes/K2Node_GetDataTableRow.h b/Engine/Source/Editor/BlueprintGraph/Classes/K2Node_GetDataTableRow.h
--- a/Engine/Source/Editor/BlueprintGraph/Classes/K2Node_GetDataTableRow.h
+++ b/Engine/Source/Editor/BlueprintGraph/Classes/K2Node_GetDataTableRow.h
@@ -23,7 +23,6 @@ class BLUEPRINTGRAPH_API UK2Node_GetDataTableRow : public UK2Node
         OutColor = GetNodeTitleColor();
         return TEXT("Kismet.AllClasses.FunctionIcon");
     }
-	virtual void ValidateNodeDuringCompilation(class FCompilerResultsLog& MessageLog) const override;
 	// End UEdGraphNode interface.
 
 	// Begin UK2Node interface
@@ -32,6 +31,7 @@ class BLUEPRINTGRAPH_API UK2Node_GetDataTableRow : public UK2Node
 	virtual void GetMenuActions(FBlueprintActionDatabaseRegistrar& ActionRegistrar) const override;
 	virtual FText GetMenuCategory() const override;
 	virtual bool IsConnectionDisallowed(const UEdGraphPin* MyPin, const UEdGraphPin* OtherPin, FString& OutReason) const override;
+	virtual void EarlyValidation(class FCompilerResultsLog& MessageLog) const override;
 	// End UK2Node interface
 
 
diff --git a/Engine/Source/Editor/BlueprintGraph/Private/K2Node_GetDataTableRow.cpp b/Engine/Source/Editor/BlueprintGraph/Private/K2Node_GetDataTableRow.cpp
--- a/Engine/Source/Editor/BlueprintGraph/Private/K2Node_GetDataTableRow.cpp
+++ b/Engine/Source/Editor/BlueprintGraph/Private/K2Node_GetDataTableRow.cpp
@@ -376,7 +376,7 @@ void UK2Node_GetDataTableRow::ExpandNode(class FKismetCompilerContext& CompilerC
 	BreakAllNodeLinks();
 }
 
-void UK2Node_GetDataTableRow::ValidateNodeDuringCompilation(class FCompilerResultsLog& MessageLog) const
+void UK2Node_GetDataTableRow::EarlyValidation(class FCompilerResultsLog& MessageLog) const
 {
 	const auto DataTablePin = GetDataTablePin();
 	const auto RowNamePin = GetRowNamePin();
@@ -393,15 +393,18 @@ void UK2Node_GetDataTableRow::ValidateNodeDuringCompilation(class FCompilerResul
 		return;
 	}
 
-	const FName CurrentName = FName(*RowNamePin->GetDefaultAsString());
-	if (!DataTable->GetRowNames().Contains(CurrentName))
+	if (!RowNamePin->LinkedTo.Num())
 	{
-		const FString Msg = FString::Printf(
-			*LOCTEXT("WronRowName", "'%s' row name is not stored in '%s'. @@").ToString()
-			, *CurrentName.ToString()
-			, *GetFullNameSafe(DataTable));
-		MessageLog.Error(*Msg, this);
-		return;
+		const FName CurrentName = FName(*RowNamePin->GetDefaultAsString());
+		if (!DataTable->GetRowNames().Contains(CurrentName))
+		{
+			const FString Msg = FString::Printf(
+				*LOCTEXT("WronRowName", "'%s' row name is not stored in '%s'. @@").ToString()
+				, *CurrentName.ToString()
+				, *GetFullNameSafe(DataTable));
+			MessageLog.Error(*Msg, this);
+			return;
+		}
 	}
 }
 
diff --git a/Engine/Source/Runtime/Engine/Classes/Kismet/DataTableFunctionLibrary.h b/Engine/Source/Runtime/Engine/Classes/Kismet/DataTableFunctionLibrary.h
--- a/Engine/Source/Runtime/Engine/Classes/Kismet/DataTableFunctionLibrary.h
+++ b/Engine/Source/Runtime/Engine/Classes/Kismet/DataTableFunctionLibrary.h
@@ -28,6 +28,9 @@ class UDataTableFunctionLibrary : public UBlueprintFunctionLibrary
 	UFUNCTION(BlueprintCallable, Category = "DataTable", meta = (ExpandEnumAsExecs="OutResult", DataTablePin="CurveTable"))
 	static void EvaluateCurveTableRow(UCurveTable* CurveTable, FName RowName, float InXY, TEnumAsByte<EEvaluateCurveTableResult::Type>& OutResult, float& OutXY);
     
+	UFUNCTION(BlueprintCallable, Category = "DataTable")
+	static void GetDataTableRowNames(UDataTable* Table, TArray<FName>& OutRowNames);
+
     /** Get a Row from a DataTable given a RowName */
     UFUNCTION(BlueprintCallable, CustomThunk, Category = "DataTable", meta=(CustomStructureParam = "OutRow", BlueprintInternalUseOnly="true"))
     static bool GetDataTableRowFromName(UDataTable* Table, FName RowName, FTableRowBase& OutRow);
diff --git a/Engine/Source/Runtime/Engine/Private/DataTableFunctionLibrary.cpp b/Engine/Source/Runtime/Engine/Private/DataTableFunctionLibrary.cpp
--- a/Engine/Source/Runtime/Engine/Private/DataTableFunctionLibrary.cpp
+++ b/Engine/Source/Runtime/Engine/Private/DataTableFunctionLibrary.cpp
@@ -55,4 +55,16 @@ bool UDataTableFunctionLibrary::GetDataTableRowFromName(UDataTable* Table, FName
 	// We should never hit this!  stubs to avoid NoExport on the class.
 	check(0);
     return false;
+}
+
+void UDataTableFunctionLibrary::GetDataTableRowNames(UDataTable* Table, TArray<FName>& OutRowNames)
+{
+	if (Table)
+	{
+		OutRowNames = Table->GetRowNames();
+	}
+	else
+	{
+		OutRowNames.Empty();
+	}
 }
\ No newline at end of file
