Hidden functions can now be dragged from the palette and pasted into graphs, they will not appear in the context sensitive menu however.

Renamed UEdGraphSchema_K2::CanFunctionBeUsedInClass to CanFunctionBeUsedInGraph to better describe it's functionality and use, also removed bInShowInherited parameter as it was unused.

Fixed up the legacy Blueprint menu system to continue to filter out hidden functions after this change.

diff --git a/Engine/Source/Editor/BlueprintGraph/Classes/EdGraphSchema_K2.h b/Engine/Source/Editor/BlueprintGraph/Classes/EdGraphSchema_K2.h
--- a/Engine/Source/Editor/BlueprintGraph/Classes/EdGraphSchema_K2.h
+++ b/Engine/Source/Editor/BlueprintGraph/Classes/EdGraphSchema_K2.h
@@ -917,18 +917,17 @@ public:
 	bool DoesGraphSupportImpureFunctions(const UEdGraph* InGraph) const;
 
 	/**
-	 * Checks to see if the passed in function is valid in the class
+	 * Checks to see if the passed in function is valid in the graph for the current class
 	 *
 	 * @param	InClass  			Class being checked to see if the function is valid for
 	 * @param	InFunction			Function being checked
 	 * @param	InDestGraph			Graph we will be using action for (may be NULL)
 	 * @param	InFunctionTypes		Combination of EFunctionType to indicate types of functions accepted
-	 * @param	bInShowInherited	Allows for inherited functions
 	 * @param	bInCalledForEach	Call for each element in an array (a node accepts array)
 	 * @param	InTargetInfo		Allows spawning nodes which also create a target variable as well
 	 * @param	OutReason			Allows callers to receive a localized string containing more detail when the function is determined to be invalid (optional)
 	 */
-	bool CanFunctionBeUsedInClass(const UClass* InClass, UFunction* InFunction, const UEdGraph* InDestGraph, uint32 InFunctionTypes, bool bInShowInherited, bool bInCalledForEach, const FFunctionTargetInfo& InTargetInfo, FText* OutReason = nullptr) const;
+	bool CanFunctionBeUsedInGraph(const UClass* InClass, UFunction* InFunction, const UEdGraph* InDestGraph, uint32 InFunctionTypes, bool bInCalledForEach, const FFunctionTargetInfo& InTargetInfo, FText* OutReason = nullptr) const;
 
 	/**
 	 * Makes connections into/or out of the gateway node, connect directly to the associated networks on the opposite side of the tunnel
diff --git a/Engine/Source/Editor/BlueprintGraph/Private/EdGraphSchema_K2.cpp b/Engine/Source/Editor/BlueprintGraph/Private/EdGraphSchema_K2.cpp
--- a/Engine/Source/Editor/BlueprintGraph/Private/EdGraphSchema_K2.cpp
+++ b/Engine/Source/Editor/BlueprintGraph/Private/EdGraphSchema_K2.cpp
@@ -507,7 +507,7 @@ bool UEdGraphSchema_K2::DoesFunctionHaveOutParameters( const UFunction* Function
 	return false;
 }
 
-bool UEdGraphSchema_K2::CanFunctionBeUsedInClass(const UClass* InClass, UFunction* InFunction, const UEdGraph* InDestGraph, uint32 InAllowedFunctionTypes, bool bInShowInherited, bool bInCalledForEach, const FFunctionTargetInfo& InTargetInfo, FText* OutReason) const
+bool UEdGraphSchema_K2::CanFunctionBeUsedInGraph(const UClass* InClass, UFunction* InFunction, const UEdGraph* InDestGraph, uint32 InAllowedFunctionTypes, bool bInCalledForEach, const FFunctionTargetInfo& InTargetInfo, FText* OutReason) const
 {
 	if (CanUserKismetCallFunction(InFunction))
 	{
@@ -641,17 +641,6 @@ bool UEdGraphSchema_K2::CanFunctionBeUsedInClass(const UClass* InClass, UFunctio
 			}
 		}
 
-		const bool bFunctionHidden = FObjectEditorUtils::IsFunctionHiddenFromClass(InFunction, InClass);
-		if (bFunctionHidden)
-		{
-			if(OutReason != nullptr)
-			{
-				*OutReason = LOCTEXT("HiddenFunctionInaccessible", "Function is hidden and inaccessible.");
-			}
-
-			return false;
-		}
-
 		const bool bFunctionStatic = InFunction->HasAllFunctionFlags(FUNC_Static);
 		const bool bHasReturnParams = (InFunction->GetReturnProperty() != NULL);
 		const bool bHasArrayPointerParms = InFunction->HasMetaData(TEXT("ArrayParm"));
diff --git a/Engine/Source/Editor/BlueprintGraph/Private/K2ActionMenuBuilder.cpp b/Engine/Source/Editor/BlueprintGraph/Private/K2ActionMenuBuilder.cpp
--- a/Engine/Source/Editor/BlueprintGraph/Private/K2ActionMenuBuilder.cpp
+++ b/Engine/Source/Editor/BlueprintGraph/Private/K2ActionMenuBuilder.cpp
@@ -1489,9 +1489,14 @@ void FK2ActionMenuBuilder::GetFuncNodesForClass(FGraphActionListBuilderBase& Lis
 		for (TFieldIterator<UFunction> FunctionIt(Class, SuperClassFlag); FunctionIt; ++FunctionIt)
 		{
 			UFunction* Function = *FunctionIt;
-			if( K2Schema->CanFunctionBeUsedInClass(Class, Function, DestGraph, FunctionTypes, bShowInherited, bCalledForEach, TargetInfo) )
+			if( K2Schema->CanFunctionBeUsedInGraph(Class, Function, DestGraph, FunctionTypes, bCalledForEach, TargetInfo) )
 			{
-				FK2ActionMenuBuilder::AddSpawnInfoForFunction(Function, false, TargetInfo, FMemberReference(), BaseCategory, K2Schema->AG_LevelReference, ListBuilder, bCalledForEach);
+				// Check if the function is hidden
+				const bool bFunctionHidden = FObjectEditorUtils::IsFunctionHiddenFromClass(Function, Class);
+				if (!bFunctionHidden)
+				{
+					FK2ActionMenuBuilder::AddSpawnInfoForFunction(Function, false, TargetInfo, FMemberReference(), BaseCategory, K2Schema->AG_LevelReference, ListBuilder, bCalledForEach);
+				}
 			}
 		}
 
diff --git a/Engine/Source/Editor/BlueprintGraph/Private/K2Node_CallFunction.cpp b/Engine/Source/Editor/BlueprintGraph/Private/K2Node_CallFunction.cpp
--- a/Engine/Source/Editor/BlueprintGraph/Private/K2Node_CallFunction.cpp
+++ b/Engine/Source/Editor/BlueprintGraph/Private/K2Node_CallFunction.cpp
@@ -1023,7 +1023,7 @@ bool UK2Node_CallFunction::CanPasteHere(const UEdGraph* TargetGraph) const
 			AllowedFunctionTypes |= UEdGraphSchema_K2::EFunctionType::FT_Imperative;
 		}
 
-		bCanPaste = K2Schema->CanFunctionBeUsedInClass(FBlueprintEditorUtils::FindBlueprintForGraphChecked(TargetGraph)->GeneratedClass, GetTargetFunction(), TargetGraph, AllowedFunctionTypes, true, false, FFunctionTargetInfo());
+		bCanPaste = K2Schema->CanFunctionBeUsedInGraph(FBlueprintEditorUtils::FindBlueprintForGraphChecked(TargetGraph)->GeneratedClass, GetTargetFunction(), TargetGraph, AllowedFunctionTypes, false, FFunctionTargetInfo());
 	}
 	
 	return bCanPaste;
diff --git a/Engine/Source/Editor/Kismet/Private/BlueprintEditorCommands.cpp b/Engine/Source/Editor/Kismet/Private/BlueprintEditorCommands.cpp
--- a/Engine/Source/Editor/Kismet/Private/BlueprintEditorCommands.cpp
+++ b/Engine/Source/Editor/Kismet/Private/BlueprintEditorCommands.cpp
@@ -223,7 +223,7 @@ public:
 				FunctionTypes |= UEdGraphSchema_K2::EFunctionType::FT_Imperative;
 			}
 
-			if(K2Schema->CanFunctionBeUsedInClass(Blueprint->GeneratedClass, FunctionPtr, InDestGraph, FunctionTypes, true, false, FFunctionTargetInfo()))
+			if(K2Schema->CanFunctionBeUsedInGraph(Blueprint->GeneratedClass, FunctionPtr, InDestGraph, FunctionTypes, false, FFunctionTargetInfo()))
 			{
 				FK2ActionMenuBuilder::AddSpawnInfoForFunction(FunctionPtr, false, FFunctionTargetInfo(), FMemberReference(), TEXT(""), K2Schema->AG_LevelReference, InPaletteBuilder);
 			}
diff --git a/Engine/Source/Editor/Kismet/Private/SBlueprintSubPalette.cpp b/Engine/Source/Editor/Kismet/Private/SBlueprintSubPalette.cpp
--- a/Engine/Source/Editor/Kismet/Private/SBlueprintSubPalette.cpp
+++ b/Engine/Source/Editor/Kismet/Private/SBlueprintSubPalette.cpp
@@ -95,7 +95,7 @@ static bool CanPaletteItemBePlaced(TSharedPtr<FEdGraphSchemaAction> DropActionIn
 				}
 
 				const UClass* GeneratedClass = FBlueprintEditorUtils::FindBlueprintForGraphChecked(HoveredGraphIn)->GeneratedClass;
-				bCanBePlaced = K2Schema->CanFunctionBeUsedInClass(GeneratedClass, Function, HoveredGraphIn, AllowedFunctionTypes, true, false, FFunctionTargetInfo(), &ImpededReasonOut);
+				bCanBePlaced = K2Schema->CanFunctionBeUsedInGraph(GeneratedClass, Function, HoveredGraphIn, AllowedFunctionTypes, false, FFunctionTargetInfo(), &ImpededReasonOut);
 			}
 		}
 		else if (UK2Node_Event const* EventNode = Cast<UK2Node_Event const>(NodeToBePlaced))
