Fixes Blueprint compiler error when renaming variables that have placed set nodes in the graph.

The "Get" pin was named in such a way that compiling would result in being unable to find the pin to expand and remove it from the node.

diff --git a/Engine/Source/Editor/BlueprintGraph/Private/K2Node_VariableSet.cpp b/Engine/Source/Editor/BlueprintGraph/Private/K2Node_VariableSet.cpp
--- a/Engine/Source/Editor/BlueprintGraph/Private/K2Node_VariableSet.cpp
+++ b/Engine/Source/Editor/BlueprintGraph/Private/K2Node_VariableSet.cpp
@@ -358,7 +358,7 @@ FNodeHandlingFunctor* UK2Node_VariableSet::CreateNodeHandler(FKismetCompilerCont
 
 FString UK2Node_VariableSet::GetVariableOutputPinName() const
 {
-	return GetVarNameString() + TEXT("_Get");
+	return TEXT("Output_Get");
 }
 
 void UK2Node_VariableSet::CreateOutputPinTooltip()
