BP - event nodes no longer indicate structural changes when UGPF is enabled

diff --git a/Engine/Source/Editor/BlueprintGraph/Classes/K2Node_Event.h b/Engine/Source/Editor/BlueprintGraph/Classes/K2Node_Event.h
--- a/Engine/Source/Editor/BlueprintGraph/Classes/K2Node_Event.h
+++ b/Engine/Source/Editor/BlueprintGraph/Classes/K2Node_Event.h
@@ -53,7 +53,7 @@ class UK2Node_Event : public UK2Node_EditablePinBase
 
 	// Begin UK2Node interface
 	virtual bool DrawNodeAsEntry() const override { return true; }
-	virtual bool NodeCausesStructuralBlueprintChange() const override { return true; }
+	BLUEPRINTGRAPH_API virtual bool NodeCausesStructuralBlueprintChange() const override;
 	BLUEPRINTGRAPH_API virtual void GetRedirectPinNames(const UEdGraphPin& Pin, TArray<FString>& RedirectPinNames) const override;
 	BLUEPRINTGRAPH_API virtual void ValidateNodeDuringCompilation(class FCompilerResultsLog& MessageLog) const override;
 	BLUEPRINTGRAPH_API virtual void PinConnectionListChanged(UEdGraphPin* Pin) override;
diff --git a/Engine/Source/Editor/BlueprintGraph/Private/K2Node_Event.cpp b/Engine/Source/Editor/BlueprintGraph/Private/K2Node_Event.cpp
--- a/Engine/Source/Editor/BlueprintGraph/Private/K2Node_Event.cpp
+++ b/Engine/Source/Editor/BlueprintGraph/Private/K2Node_Event.cpp
@@ -255,6 +255,12 @@ void UK2Node_Event::ValidateNodeDuringCompilation(class FCompilerResultsLog& Mes
 	}
 }
 
+bool UK2Node_Event::NodeCausesStructuralBlueprintChange() const
+{
+	// will only change class structure when UGPF is disabled
+	return !UBlueprintGeneratedClass::UsePersistentUberGraphFrame();
+}
+
 void UK2Node_Event::GetRedirectPinNames(const UEdGraphPin& Pin, TArray<FString>& RedirectPinNames) const
 {
 	Super::GetRedirectPinNames(Pin, RedirectPinNames);
