[UE-6016] Making sure we use the authoritative class when checking for compatible blueprint node bindings.

diff --git a/Engine/Source/Editor/BlueprintGraph/Private/BlueprintBoundEventNodeSpawner.cpp b/Engine/Source/Editor/BlueprintGraph/Private/BlueprintBoundEventNodeSpawner.cpp
--- a/Engine/Source/Editor/BlueprintGraph/Private/BlueprintBoundEventNodeSpawner.cpp
+++ b/Engine/Source/Editor/BlueprintGraph/Private/BlueprintBoundEventNodeSpawner.cpp
@@ -135,8 +135,8 @@ bool UBlueprintBoundEventNodeSpawner::IsBindingCompatible(UObject const* Binding
 	}
 
 	const UMulticastDelegateProperty* Delegate = GetEventDelegate();
-	UClass* DelegateOwner = Delegate->GetOwnerClass();
-	UClass* BindingClass  = FBlueprintNodeSpawnerUtils::GetBindingClass(BindingCandidate);
+	UClass* DelegateOwner = Delegate->GetOwnerClass()->GetAuthoritativeClass();
+	UClass* BindingClass  = FBlueprintNodeSpawnerUtils::GetBindingClass(BindingCandidate)->GetAuthoritativeClass();
 
 	return bMatchesNodeType && BindingClass->IsChildOf(DelegateOwner) && !FObjectEditorUtils::IsVariableCategoryHiddenFromClass(Delegate, BindingClass);
 }
diff --git a/Engine/Source/Editor/BlueprintGraph/Private/BlueprintFunctionNodeSpawner.cpp b/Engine/Source/Editor/BlueprintGraph/Private/BlueprintFunctionNodeSpawner.cpp
--- a/Engine/Source/Editor/BlueprintGraph/Private/BlueprintFunctionNodeSpawner.cpp
+++ b/Engine/Source/Editor/BlueprintGraph/Private/BlueprintFunctionNodeSpawner.cpp
@@ -412,8 +412,8 @@ bool UBlueprintFunctionNodeSpawner::IsBindingCompatible(UObject const* BindingCa
 	bool const bNodeTypeMatches = (NodeClass == UK2Node_CallFunction::StaticClass());
 	bool bClassOwnerMatches = false;
 
-	UClass* BindingClass = FBlueprintNodeSpawnerUtils::GetBindingClass(BindingCandidate);
-	if (UClass const* FuncOwner = Function->GetOwnerClass())
+	UClass* BindingClass = FBlueprintNodeSpawnerUtils::GetBindingClass(BindingCandidate)->GetAuthoritativeClass();
+	if (UClass const* FuncOwner = Function->GetOwnerClass()->GetAuthoritativeClass())
 	{
 		bClassOwnerMatches = BindingClass->IsChildOf(FuncOwner);
 	}
