[UE-5763] Showing interface function nodes for blueprints when the interface comes from a native superclass.

diff --git a/Engine/Source/Editor/BlueprintGraph/Private/BlueprintActionFilter.cpp b/Engine/Source/Editor/BlueprintGraph/Private/BlueprintActionFilter.cpp
--- a/Engine/Source/Editor/BlueprintGraph/Private/BlueprintActionFilter.cpp
+++ b/Engine/Source/Editor/BlueprintGraph/Private/BlueprintActionFilter.cpp
@@ -49,6 +49,17 @@ namespace BlueprintActionFilterImpl
 	 * @return True if the ClassToTest is the same type as (or inherited from) TypeToCheckFor.
 	 */
 	static bool IsClassOfType(UClass const* ClassToTest, UClass const* TypeToCheckFor, bool const bNeedsExactMatch = false);
+
+	/**
+	 * This not only tells you if an interface is implemented by a class, but 
+	 * walks the class inheritance chain to find out which super class actually 
+	 * adds the interface.
+	 * 
+	 * @param  SubClass		The class you're querying for.
+	 * @param  Interface	The interface you want to look for.
+	 * @return Null if the class doesn't implement the function, otherwise the super class that adds it (could be the class passed in).
+	 */
+	static UClass const* FindInheritedInterfaceClass(UClass const* SubClass, TSubclassOf<UInterface> Interface);
 	
 	/**
 	 * Utility method to check and see if the specified node-spawner would 
@@ -359,6 +370,27 @@ static bool BlueprintActionFilterImpl::IsClassOfType(UClass const* ClassToTest,
 	}
 }
 
+//------------------------------------------------------------------------------
+static UClass const* BlueprintActionFilterImpl::FindInheritedInterfaceClass(UClass const* SubClass, TSubclassOf<UInterface> Interface)
+{
+	UClass const* ImplementsInterface = nullptr;
+
+	UClass const* ClassToCheck = SubClass;
+	while ((ClassToCheck != nullptr) && (ImplementsInterface == nullptr))
+	{
+		for (FImplementedInterface const& ClassInterface : ClassToCheck->Interfaces)
+		{
+			UClass const* InterfaceClass = ClassInterface.Class;
+			if (InterfaceClass->IsChildOf(Interface))
+			{
+				ImplementsInterface = ClassToCheck;
+			}
+		}
+		ClassToCheck = ClassToCheck->GetSuperClass();
+	}
+	return ImplementsInterface;
+}
+
 //------------------------------------------------------------------------------
 static bool BlueprintActionFilterImpl::IsImpure(FBlueprintActionInfo& BlueprintAction)
 {
@@ -1288,11 +1320,19 @@ static bool BlueprintActionFilterImpl::IsExtraneousInterfaceCall(FBlueprintActio
 			bIsFilteredOut = (Filter.TargetClasses.Num() > 0);
 			for (const UClass* TargetClass : Filter.TargetClasses)
 			{
-				bool const bTargetIsBlueprint = (Cast<UBlueprint>(TargetClass->ClassGeneratedBy) != nullptr);
-				// Blueprints fold interface functions in with their other 
-				// functions, and we'd rather those get called (because it is more
-				// optimal, as it avoids a conversion between object/interface)
-				if (!bTargetIsBlueprint || !IsClassOfType(TargetClass, FuncClass))
+				UClass const* InterfaceImplementingClass = FindInheritedInterfaceClass(TargetClass, FuncClass);
+				// interfaces that are added directly to a Blueprint (even in 
+				// the case of an interface on a parent blueprint) have their 
+				// functions stubbed-out/added to the blueprint class directly;
+				// in that case, we want to favor a call to the blueprint 
+				// version (not this interface call) because we can circumvent 
+				// the extra work converting from a interface to an object.
+				//
+				// however, if the interface belongs to a native class, then the
+				// blueprint doesn't get those extra functions, so this is is 
+				// our only way of calling the interface methods.
+				bool const bImplementedByBlueprint = (InterfaceImplementingClass != nullptr) && (Cast<UBlueprint>(InterfaceImplementingClass->ClassGeneratedBy) != nullptr);
+				if (!bImplementedByBlueprint)
 				{
 					bIsFilteredOut = false;
 					break;
