BP - fix regression caused by a missing delegate pin in older event nodes

diff --git a/Engine/Source/Editor/BlueprintGraph/Classes/K2Node_Event.h b/Engine/Source/Editor/BlueprintGraph/Classes/K2Node_Event.h
--- a/Engine/Source/Editor/BlueprintGraph/Classes/K2Node_Event.h
+++ b/Engine/Source/Editor/BlueprintGraph/Classes/K2Node_Event.h
@@ -36,6 +36,10 @@ class UK2Node_Event : public UK2Node_EditablePinBase
 	UPROPERTY()
 	uint32 FunctionFlags;
 
+	// UObject interface
+	BLUEPRINTGRAPH_API virtual void Serialize(FArchive& Ar) override;
+	// End of UObject interface
+
 	// Begin UEdGraphNode interface
 	BLUEPRINTGRAPH_API virtual void AllocateDefaultPins() override;
 	BLUEPRINTGRAPH_API virtual FText GetTooltipText() const override;
diff --git a/Engine/Source/Editor/BlueprintGraph/Private/K2Node_Event.cpp b/Engine/Source/Editor/BlueprintGraph/Private/K2Node_Event.cpp
--- a/Engine/Source/Editor/BlueprintGraph/Private/K2Node_Event.cpp
+++ b/Engine/Source/Editor/BlueprintGraph/Private/K2Node_Event.cpp
@@ -39,6 +39,18 @@ UK2Node_Event::UK2Node_Event(const FObjectInitializer& ObjectInitializer)
 	FunctionFlags = 0;
 }
 
+void UK2Node_Event::Serialize(FArchive& Ar)
+{
+	Super::Serialize(Ar);
+
+	// Fix up legacy nodes that may not yet have a delegate pin
+	if(Ar.IsLoading() && !FindPin(DelegateOutputName))
+	{
+		const UEdGraphSchema_K2* K2Schema = GetDefault<UEdGraphSchema_K2>();
+		CreatePin(EGPD_Output, K2Schema->PC_Delegate, TEXT(""), NULL, false, false, DelegateOutputName);
+	}
+}
+
 FNodeHandlingFunctor* UK2Node_Event::CreateNodeHandler(FKismetCompilerContext& CompilerContext) const
 {
 	return new FKCHandler_EventEntry(CompilerContext);
