BP: Some k2 nodes don't use NodeCausesStructuralBlueprintChange anymore.

diff --git a/Engine/Source/Editor/BlueprintGraph/Classes/K2Node.h b/Engine/Source/Editor/BlueprintGraph/Classes/K2Node.h
--- a/Engine/Source/Editor/BlueprintGraph/Classes/K2Node.h
+++ b/Engine/Source/Editor/BlueprintGraph/Classes/K2Node.h
@@ -337,6 +337,7 @@ class UK2Node : public UEdGraphNode
 	enum BLUEPRINTGRAPH_API EBaseNodeRefreshPriority
 	{
 		Low_UsesDependentWildcard = 100,
+		Low_ReceivesDelegateSignature = 150,
 		Normal = 200,
 	};
 
diff --git a/Engine/Source/Editor/BlueprintGraph/Classes/K2Node_BaseMCDelegate.h b/Engine/Source/Editor/BlueprintGraph/Classes/K2Node_BaseMCDelegate.h
--- a/Engine/Source/Editor/BlueprintGraph/Classes/K2Node_BaseMCDelegate.h
+++ b/Engine/Source/Editor/BlueprintGraph/Classes/K2Node_BaseMCDelegate.h
@@ -18,7 +18,6 @@ public:
 	// UK2Node interface
 	virtual bool IsNodePure() const override { return false; }
 	virtual ERedirectType DoPinsMatchForReconstruction(const UEdGraphPin* NewPin, int32 NewPinIndex, const UEdGraphPin* OldPin, int32 OldPinIndex) const override;
-	virtual bool NodeCausesStructuralBlueprintChange() const override { return true; }
 	virtual FString GetDocumentationLink() const override;
 	virtual FString GetDocumentationExcerptName() const override;
 	virtual bool AllowMultipleSelfs(bool bInputAsArray) const override { return true; }
diff --git a/Engine/Source/Editor/BlueprintGraph/Classes/K2Node_BreakStruct.h b/Engine/Source/Editor/BlueprintGraph/Classes/K2Node_BreakStruct.h
--- a/Engine/Source/Editor/BlueprintGraph/Classes/K2Node_BreakStruct.h
+++ b/Engine/Source/Editor/BlueprintGraph/Classes/K2Node_BreakStruct.h
@@ -23,6 +23,7 @@ class UK2Node_BreakStruct : public UK2Node_StructMemberGet
 	// End  UEdGraphNode interface
 
 	// Begin K2Node interface
+	virtual bool NodeCausesStructuralBlueprintChange() const override { return false; }
 	virtual bool IsNodePure() const override { return true; }
 	virtual bool DrawNodeAsVariable() const override { return false; }
 	virtual ERedirectType DoPinsMatchForReconstruction(const UEdGraphPin* NewPin, int32 NewPinIndex, const UEdGraphPin* OldPin, int32 OldPinIndex) const override;
diff --git a/Engine/Source/Editor/BlueprintGraph/Classes/K2Node_CreateDelegate.h b/Engine/Source/Editor/BlueprintGraph/Classes/K2Node_CreateDelegate.h
--- a/Engine/Source/Editor/BlueprintGraph/Classes/K2Node_CreateDelegate.h
+++ b/Engine/Source/Editor/BlueprintGraph/Classes/K2Node_CreateDelegate.h
@@ -30,6 +30,7 @@ public:
 	virtual class FNodeHandlingFunctor* CreateNodeHandler(class FKismetCompilerContext& CompilerContext) const override;
 	virtual void GetMenuActions(FBlueprintActionDatabaseRegistrar& ActionRegistrar) const override;
 	virtual FText GetMenuCategory() const override;
+	virtual int32 GetNodeRefreshPriority() const override { return EBaseNodeRefreshPriority::Low_ReceivesDelegateSignature; }
 	// End of UK2Node interface
 
 	bool IsValid(FString* OutMsg = NULL, bool bDontUseSkeletalClassForSelf = false) const;
diff --git a/Engine/Source/Editor/BlueprintGraph/Classes/K2Node_MakeStruct.h b/Engine/Source/Editor/BlueprintGraph/Classes/K2Node_MakeStruct.h
--- a/Engine/Source/Editor/BlueprintGraph/Classes/K2Node_MakeStruct.h
+++ b/Engine/Source/Editor/BlueprintGraph/Classes/K2Node_MakeStruct.h
@@ -24,6 +24,7 @@ class UK2Node_MakeStruct : public UK2Node_StructMemberSet
 	// End  UEdGraphNode interface
 
 	// Begin K2Node interface
+	virtual bool NodeCausesStructuralBlueprintChange() const override { return false; }
 	virtual bool IsNodePure() const override { return true; }
 	virtual bool DrawNodeAsVariable() const override { return false; }
 	virtual class FNodeHandlingFunctor* CreateNodeHandler(class FKismetCompilerContext& CompilerContext) const override;
