[UE-5249] BP graph context menu now filters out functions with metadata that would disallow it from being placed into the current graph context

change summary:
- implemented a UK2Node_CallFunction::IsActionFilteredOut() override
- made 'Function' param 'const' in UEdGraphSchema_K2::CanFunctionBeUsedInGraph()

diff --git a/Engine/Source/Editor/BlueprintGraph/Classes/EdGraphSchema_K2.h b/Engine/Source/Editor/BlueprintGraph/Classes/EdGraphSchema_K2.h
--- a/Engine/Source/Editor/BlueprintGraph/Classes/EdGraphSchema_K2.h
+++ b/Engine/Source/Editor/BlueprintGraph/Classes/EdGraphSchema_K2.h
@@ -934,7 +934,7 @@ public:
 	 * @param	InTargetInfo		Allows spawning nodes which also create a target variable as well
 	 * @param	OutReason			Allows callers to receive a localized string containing more detail when the function is determined to be invalid (optional)
 	 */
-	bool CanFunctionBeUsedInGraph(const UClass* InClass, UFunction* InFunction, const UEdGraph* InDestGraph, uint32 InFunctionTypes, bool bInCalledForEach, const FFunctionTargetInfo& InTargetInfo, FText* OutReason = nullptr) const;
+	bool CanFunctionBeUsedInGraph(const UClass* InClass, const UFunction* InFunction, const UEdGraph* InDestGraph, uint32 InFunctionTypes, bool bInCalledForEach, const FFunctionTargetInfo& InTargetInfo, FText* OutReason = nullptr) const;
 
 	/**
 	 * Makes connections into/or out of the gateway node, connect directly to the associated networks on the opposite side of the tunnel
diff --git a/Engine/Source/Editor/BlueprintGraph/Classes/K2Node_CallFunction.h b/Engine/Source/Editor/BlueprintGraph/Classes/K2Node_CallFunction.h
--- a/Engine/Source/Editor/BlueprintGraph/Classes/K2Node_CallFunction.h
+++ b/Engine/Source/Editor/BlueprintGraph/Classes/K2Node_CallFunction.h
@@ -95,6 +95,7 @@ public:
 	virtual FText GetToolTipHeading() const override;
 	virtual void GetNodeAttributes( TArray<TKeyValuePair<FString, FString>>& OutNodeAttributes ) const override;
 	virtual FText GetMenuCategory() const override;
+	virtual bool IsActionFilteredOut(class FBlueprintActionFilter const& Filter) override;
 	// End of UK2Node interface
 
 	/** Returns the UFunction that this class is pointing to */
diff --git a/Engine/Source/Editor/BlueprintGraph/Private/EdGraphSchema_K2.cpp b/Engine/Source/Editor/BlueprintGraph/Private/EdGraphSchema_K2.cpp
--- a/Engine/Source/Editor/BlueprintGraph/Private/EdGraphSchema_K2.cpp
+++ b/Engine/Source/Editor/BlueprintGraph/Private/EdGraphSchema_K2.cpp
@@ -507,7 +507,7 @@ bool UEdGraphSchema_K2::DoesFunctionHaveOutParameters( const UFunction* Function
 	return false;
 }
 
-bool UEdGraphSchema_K2::CanFunctionBeUsedInGraph(const UClass* InClass, UFunction* InFunction, const UEdGraph* InDestGraph, uint32 InAllowedFunctionTypes, bool bInCalledForEach, const FFunctionTargetInfo& InTargetInfo, FText* OutReason) const
+bool UEdGraphSchema_K2::CanFunctionBeUsedInGraph(const UClass* InClass, const UFunction* InFunction, const UEdGraph* InDestGraph, uint32 InAllowedFunctionTypes, bool bInCalledForEach, const FFunctionTargetInfo& InTargetInfo, FText* OutReason) const
 {
 	if (CanUserKismetCallFunction(InFunction))
 	{
diff --git a/Engine/Source/Editor/BlueprintGraph/Private/K2Node_CallFunction.cpp b/Engine/Source/Editor/BlueprintGraph/Private/K2Node_CallFunction.cpp
--- a/Engine/Source/Editor/BlueprintGraph/Private/K2Node_CallFunction.cpp
+++ b/Engine/Source/Editor/BlueprintGraph/Private/K2Node_CallFunction.cpp
@@ -1029,6 +1029,17 @@ bool UK2Node_CallFunction::CanPasteHere(const UEdGraph* TargetGraph) const
 	return bCanPaste;
 }
 
+bool UK2Node_CallFunction::IsActionFilteredOut(FBlueprintActionFilter const& Filter)
+{
+	bool bIsFilteredOut = false;
+	for(UEdGraph* TargetGraph : Filter.Context.Graphs)
+	{
+		bIsFilteredOut |= !CanPasteHere(TargetGraph);
+	}
+
+	return bIsFilteredOut;
+}
+
 static FLinearColor GetPalletteIconColor(UFunction const* Function)
 {
 	bool const bIsPure = (Function != nullptr) && Function->HasAnyFunctionFlags(FUNC_BlueprintPure);
