[UE-6283] Fixing it so internal macro calls cannot be placed/pasted into the wrong blueprint.

diff --git a/Engine/Source/Editor/BlueprintGraph/Classes/K2Node_MacroInstance.h b/Engine/Source/Editor/BlueprintGraph/Classes/K2Node_MacroInstance.h
--- a/Engine/Source/Editor/BlueprintGraph/Classes/K2Node_MacroInstance.h
+++ b/Engine/Source/Editor/BlueprintGraph/Classes/K2Node_MacroInstance.h
@@ -46,6 +46,7 @@ public:
 	virtual FString GetDocumentationLink() const override;
 	virtual FString GetDocumentationExcerptName() const override;
 	virtual FName GetPaletteIcon(FLinearColor& OutColor) const override;
+	virtual bool CanPasteHere(const UEdGraph* TargetGraph) const override;
 	// End UEdGraphNode interface
 
 	// Begin UK2Node interface
diff --git a/Engine/Source/Editor/BlueprintGraph/Private/K2Node_MacroInstance.cpp b/Engine/Source/Editor/BlueprintGraph/Private/K2Node_MacroInstance.cpp
--- a/Engine/Source/Editor/BlueprintGraph/Private/K2Node_MacroInstance.cpp
+++ b/Engine/Source/Editor/BlueprintGraph/Private/K2Node_MacroInstance.cpp
@@ -32,8 +32,7 @@ bool UK2Node_MacroInstance::IsActionFilteredOut(FBlueprintActionFilter const& Fi
 
 	for (UEdGraph* Graph : FilterContext.Graphs)
 	{
-		// Macro Instances are not allowed in it's own graph, nor in Function graphs if the macro has latent functions in it
-		if (Graph == GetMacroGraph() || (Graph->GetSchema()->GetGraphType(Graph) == GT_Function && FBlueprintEditorUtils::CheckIfGraphHasLatentFunctions(GetMacroGraph())) )
+		if (!CanPasteHere(Graph))
 		{
 			bIsFilteredOut = true;
 			break;
@@ -357,6 +356,35 @@ FName UK2Node_MacroInstance::GetPaletteIcon(FLinearColor& OutColor) const
 	return TEXT("GraphEditor.Macro_16x");
 }
 
+bool UK2Node_MacroInstance::CanPasteHere(const UEdGraph* TargetGraph) const
+{
+	bool bCanPaste = false;
+
+	UBlueprint* MacroBlueprint  = GetSourceBlueprint();
+	UBlueprint* TargetBlueprint = FBlueprintEditorUtils::FindBlueprintForGraph(TargetGraph);
+
+	if ((MacroBlueprint != nullptr) && (TargetBlueprint != nullptr))
+	{
+		UClass* TargetClass = (TargetBlueprint->GeneratedClass != nullptr) ? TargetBlueprint->GeneratedClass : TargetBlueprint->ParentClass;
+		UClass* MacroBaseClass = ((MacroBlueprint->BlueprintType == BPTYPE_MacroLibrary) || (MacroBlueprint->GeneratedClass == nullptr)) ? 
+			MacroBlueprint->ParentClass : MacroBlueprint->GeneratedClass;
+		
+		if ((MacroBaseClass != nullptr) && (TargetClass != nullptr))
+		{
+			bCanPaste = TargetClass->IsChildOf(MacroBaseClass);
+		}
+	}
+
+	// Macro Instances are not allowed in it's own graph
+	UEdGraph* MacroGraph = GetMacroGraph();
+	bCanPaste &= (MacroGraph != TargetGraph);
+	// nor in Function graphs if the macro has latent functions in it
+	bool const bIsTargetFuncGraph = (TargetGraph->GetSchema()->GetGraphType(TargetGraph) == GT_Function);
+	bCanPaste &= (!bIsTargetFuncGraph || !FBlueprintEditorUtils::CheckIfGraphHasLatentFunctions(MacroGraph));
+
+	return bCanPaste && Super::CanPasteHere(TargetGraph);
+}
+
 void UK2Node_MacroInstance::ReallocatePinsDuringReconstruction(TArray<UEdGraphPin*>& OldPins)
 {
 	Super::ReallocatePinsDuringReconstruction(OldPins);
