Fix VS2012 editor compilation

diff --git a/Engine/Source/Editor/BlueprintGraph/Private/BlueprintFieldNodeSpawner.cpp b/Engine/Source/Editor/BlueprintGraph/Private/BlueprintFieldNodeSpawner.cpp
--- a/Engine/Source/Editor/BlueprintGraph/Private/BlueprintFieldNodeSpawner.cpp
+++ b/Engine/Source/Editor/BlueprintGraph/Private/BlueprintFieldNodeSpawner.cpp
@@ -45,7 +45,7 @@ UEdGraphNode* UBlueprintFieldNodeSpawner::Invoke(UEdGraph* ParentGraph, FBinding
 	};
 
 	FCustomizeNodeDelegate PostSpawnSetupDelegate = FCustomizeNodeDelegate::CreateStatic(PostSpawnSetupLambda, GetField(), SetNodeFieldDelegate, CustomizeNodeDelegate);
-	return Super::SpawnNode(NodeClass, ParentGraph, Bindings, Location, PostSpawnSetupDelegate);
+	return Super::SpawnNode<UEdGraphNode>(NodeClass, ParentGraph, Bindings, Location, PostSpawnSetupDelegate);
 }
 
 //------------------------------------------------------------------------------
diff --git a/Engine/Source/Editor/BlueprintGraph/Private/BlueprintFunctionNodeSpawner.cpp b/Engine/Source/Editor/BlueprintGraph/Private/BlueprintFunctionNodeSpawner.cpp
--- a/Engine/Source/Editor/BlueprintGraph/Private/BlueprintFunctionNodeSpawner.cpp
+++ b/Engine/Source/Editor/BlueprintGraph/Private/BlueprintFunctionNodeSpawner.cpp
@@ -392,7 +392,7 @@ UEdGraphNode* UBlueprintFunctionNodeSpawner::Invoke(UEdGraph* ParentGraph, FBind
 	// bound nodes get positioned properly
 	BlueprintFunctionNodeSpawnerImpl::BindingOffset = FVector2D::ZeroVector;
 
-	return Super::SpawnNode(SpawnClass, ParentGraph, Bindings, Location, PostSpawnSetupDelegate);
+	return Super::SpawnNode<UEdGraphNode>(SpawnClass, ParentGraph, Bindings, Location, PostSpawnSetupDelegate);
 }
 
 //------------------------------------------------------------------------------
diff --git a/Engine/Source/Editor/BlueprintGraph/Private/BlueprintNodeSpawner.cpp b/Engine/Source/Editor/BlueprintGraph/Private/BlueprintNodeSpawner.cpp
--- a/Engine/Source/Editor/BlueprintGraph/Private/BlueprintNodeSpawner.cpp
+++ b/Engine/Source/Editor/BlueprintGraph/Private/BlueprintNodeSpawner.cpp
@@ -225,7 +225,7 @@ FBlueprintActionUiSpec UBlueprintNodeSpawner::GetUiSpec(FBlueprintActionContext
 //------------------------------------------------------------------------------
 UEdGraphNode* UBlueprintNodeSpawner::Invoke(UEdGraph* ParentGraph, FBindingSet const& Bindings, FVector2D const Location) const
 {
-	return SpawnNode(NodeClass, ParentGraph, Bindings, Location, CustomizeNodeDelegate);
+	return SpawnNode<UEdGraphNode>(NodeClass, ParentGraph, Bindings, Location, CustomizeNodeDelegate);
 }
 
 //------------------------------------------------------------------------------
diff --git a/Engine/Source/Editor/BlueprintGraph/Private/BlueprintVariableNodeSpawner.cpp b/Engine/Source/Editor/BlueprintGraph/Private/BlueprintVariableNodeSpawner.cpp
--- a/Engine/Source/Editor/BlueprintGraph/Private/BlueprintVariableNodeSpawner.cpp
+++ b/Engine/Source/Editor/BlueprintGraph/Private/BlueprintVariableNodeSpawner.cpp
@@ -213,7 +213,7 @@ UEdGraphNode* UBlueprintVariableNodeSpawner::Invoke(UEdGraph* ParentGraph, FBind
 			PostSpawnDelegate = FCustomizeNodeDelegate::CreateStatic(LocalVarSetupLambda, LocalVarDesc.VarName, LocalVarOuter, LocalVarDesc.VarGuid, CustomizeNodeDelegate);
 		}
 
-		NewNode = UBlueprintNodeSpawner::SpawnNode(NodeClass, ParentGraph, Bindings, Location, PostSpawnDelegate);
+		NewNode = UBlueprintNodeSpawner::SpawnNode<UEdGraphNode>(NodeClass, ParentGraph, Bindings, Location, PostSpawnDelegate);
 	}
 	else
 	{
diff --git a/Engine/Source/Editor/BlueprintGraph/Public/BlueprintNodeSpawner.h b/Engine/Source/Editor/BlueprintGraph/Public/BlueprintNodeSpawner.h
--- a/Engine/Source/Editor/BlueprintGraph/Public/BlueprintNodeSpawner.h
+++ b/Engine/Source/Editor/BlueprintGraph/Public/BlueprintNodeSpawner.h
@@ -186,7 +186,7 @@ protected:
 	 * @param  PostSpawnDelegate	A delegate to run after spawning the node, but prior to allocating the node's pins.
 	 * @return Null if it failed to spawn a node (if NodeClass is null), otherwise a newly spawned node.
 	 */
-	template <class NodeType = UEdGraphNode>
+	template <class NodeType>
 	NodeType* SpawnNode(TSubclassOf<UEdGraphNode> NodeClass, UEdGraph* ParentGraph, FBindingSet const& Bindings, FVector2D const Location, FCustomizeNodeDelegate PostSpawnDelegate) const;
 
 	/**
