Add 'Get' as keyword for get input axis value blueprint nodes

diff --git a/Engine/Source/Editor/BlueprintGraph/Classes/K2Node_GetInputAxisKeyValue.h b/Engine/Source/Editor/BlueprintGraph/Classes/K2Node_GetInputAxisKeyValue.h
--- a/Engine/Source/Editor/BlueprintGraph/Classes/K2Node_GetInputAxisKeyValue.h
+++ b/Engine/Source/Editor/BlueprintGraph/Classes/K2Node_GetInputAxisKeyValue.h
@@ -36,6 +36,7 @@ class UK2Node_GetInputAxisKeyValue : public UK2Node_CallFunction
 	virtual void RegisterDynamicBinding(UDynamicBlueprintBinding* BindingObject) const override;
 	virtual void GetMenuActions(FBlueprintActionDatabaseRegistrar& ActionRegistrar) const override;
 	virtual FText GetMenuCategory() const override;
+	virtual FString GetKeywords() const override;
 	virtual FBlueprintNodeSignature GetSignature() const override;
 	// End UK2Node interface
 	
diff --git a/Engine/Source/Editor/BlueprintGraph/Classes/K2Node_GetInputAxisValue.h b/Engine/Source/Editor/BlueprintGraph/Classes/K2Node_GetInputAxisValue.h
--- a/Engine/Source/Editor/BlueprintGraph/Classes/K2Node_GetInputAxisValue.h
+++ b/Engine/Source/Editor/BlueprintGraph/Classes/K2Node_GetInputAxisValue.h
@@ -25,6 +25,7 @@ class UK2Node_GetInputAxisValue : public UK2Node_CallFunction
 	virtual void AllocateDefaultPins() override;
 	virtual FText GetNodeTitle(ENodeTitleType::Type TitleType) const override;
 	virtual FText GetTooltipText() const override;
+	virtual FString GetKeywords() const override;
 	virtual bool IsCompatibleWithGraph(UEdGraph const* Graph) const override;
 	// End EdGraphNode interface
 
diff --git a/Engine/Source/Editor/BlueprintGraph/Classes/K2Node_GetInputVectorAxisValue.h b/Engine/Source/Editor/BlueprintGraph/Classes/K2Node_GetInputVectorAxisValue.h
--- a/Engine/Source/Editor/BlueprintGraph/Classes/K2Node_GetInputVectorAxisValue.h
+++ b/Engine/Source/Editor/BlueprintGraph/Classes/K2Node_GetInputVectorAxisValue.h
@@ -11,6 +11,7 @@ class UK2Node_GetInputVectorAxisValue : public UK2Node_GetInputAxisKeyValue
 
 	// Begin EdGraphNode interface
 	virtual FText GetTooltipText() const override;
+	virtual FString GetKeywords() const override;
 	// End EdGraphNode interface
 
 	// Begin UK2Node interface
diff --git a/Engine/Source/Editor/BlueprintGraph/Private/K2Node_GetInputAxisKeyValue.cpp b/Engine/Source/Editor/BlueprintGraph/Private/K2Node_GetInputAxisKeyValue.cpp
--- a/Engine/Source/Editor/BlueprintGraph/Private/K2Node_GetInputAxisKeyValue.cpp
+++ b/Engine/Source/Editor/BlueprintGraph/Private/K2Node_GetInputAxisKeyValue.cpp
@@ -52,6 +52,11 @@ FText UK2Node_GetInputAxisKeyValue::GetNodeTitle(ENodeTitleType::Type TitleType)
 	return CachedNodeTitle;
 }
 
+FString UK2Node_GetInputAxisKeyValue::GetKeywords() const
+{
+	return TEXT("Get");
+}
+
 FText UK2Node_GetInputAxisKeyValue::GetTooltipText() const
 {
 	if (CachedTooltip.IsOutOfDate())
diff --git a/Engine/Source/Editor/BlueprintGraph/Private/K2Node_GetInputAxisValue.cpp b/Engine/Source/Editor/BlueprintGraph/Private/K2Node_GetInputAxisValue.cpp
--- a/Engine/Source/Editor/BlueprintGraph/Private/K2Node_GetInputAxisValue.cpp
+++ b/Engine/Source/Editor/BlueprintGraph/Private/K2Node_GetInputAxisValue.cpp
@@ -51,6 +51,11 @@ FText UK2Node_GetInputAxisValue::GetNodeTitle(ENodeTitleType::Type TitleType) co
 	return CachedNodeTitle;
 }
 
+FString UK2Node_GetInputAxisValue::GetKeywords() const
+{
+	return TEXT("Get");
+}
+
 FText UK2Node_GetInputAxisValue::GetTooltipText() const
 {
 	if (CachedTooltip.IsOutOfDate())
diff --git a/Engine/Source/Editor/BlueprintGraph/Private/K2Node_GetInputVectorAxisValue.cpp b/Engine/Source/Editor/BlueprintGraph/Private/K2Node_GetInputVectorAxisValue.cpp
--- a/Engine/Source/Editor/BlueprintGraph/Private/K2Node_GetInputVectorAxisValue.cpp
+++ b/Engine/Source/Editor/BlueprintGraph/Private/K2Node_GetInputVectorAxisValue.cpp
@@ -19,6 +19,11 @@ void UK2Node_GetInputVectorAxisValue::Initialize(const FKey AxisKey)
 	SetFromFunction(AActor::StaticClass()->FindFunctionByName(TEXT("GetInputVectorAxisValue")));
 }
 
+FString UK2Node_GetInputVectorAxisValue::GetKeywords() const
+{
+	return TEXT("Get");
+}
+
 FText UK2Node_GetInputVectorAxisValue::GetTooltipText() const
 {
 	FFormatNamedArguments Args;
