BP: IsStaticFunctionGraph fixed

diff --git a/Engine/Source/Editor/BlueprintGraph/Private/EdGraphSchema_K2.cpp b/Engine/Source/Editor/BlueprintGraph/Private/EdGraphSchema_K2.cpp
--- a/Engine/Source/Editor/BlueprintGraph/Private/EdGraphSchema_K2.cpp
+++ b/Engine/Source/Editor/BlueprintGraph/Private/EdGraphSchema_K2.cpp
@@ -4258,6 +4258,12 @@ bool UEdGraphSchema_K2::IsStaticFunctionGraph( const UEdGraph* TestEdGraph ) con
 {
 	check(TestEdGraph);
 
+	const auto Blueprint = FBlueprintEditorUtils::FindBlueprintForGraph(TestEdGraph);
+	if (Blueprint && (EBlueprintType::BPTYPE_FunctionLibrary == Blueprint->BlueprintType))
+	{
+		return true;
+	}
+
 	const EGraphType GraphType = GetGraphType(TestEdGraph);
 	if(GraphType == GT_Function) 
 	{
