Added move pin to input touch event node

diff --git a/Engine/Source/Editor/BlueprintGraph/Classes/K2Node_InputTouch.h b/Engine/Source/Editor/BlueprintGraph/Classes/K2Node_InputTouch.h
--- a/Engine/Source/Editor/BlueprintGraph/Classes/K2Node_InputTouch.h
+++ b/Engine/Source/Editor/BlueprintGraph/Classes/K2Node_InputTouch.h
@@ -49,6 +49,9 @@ class UK2Node_InputTouch : public UK2Node
 	/** Get the 'released' input pin */
 	BLUEPRINTGRAPH_API UEdGraphPin* GetReleasedPin() const;
 
+	/** Get the 'moved' input pin */
+	BLUEPRINTGRAPH_API UEdGraphPin* GetMovedPin() const;
+
 	BLUEPRINTGRAPH_API UEdGraphPin* GetLocationPin() const;
 	BLUEPRINTGRAPH_API UEdGraphPin* GetFingerIndexPin() const;
 	
diff --git a/Engine/Source/Editor/BlueprintGraph/Private/K2Node_InputTouch.cpp b/Engine/Source/Editor/BlueprintGraph/Private/K2Node_InputTouch.cpp
--- a/Engine/Source/Editor/BlueprintGraph/Private/K2Node_InputTouch.cpp
+++ b/Engine/Source/Editor/BlueprintGraph/Private/K2Node_InputTouch.cpp
@@ -45,7 +45,8 @@ void UK2Node_InputTouch::AllocateDefaultPins()
 
 	CreatePin(EGPD_Output, K2Schema->PC_Exec, TEXT(""), NULL, false, false, TEXT("Pressed"));
 	CreatePin(EGPD_Output, K2Schema->PC_Exec, TEXT(""), NULL, false, false, TEXT("Released"));
-
+	CreatePin(EGPD_Output, K2Schema->PC_Exec, TEXT(""), NULL, false, false, TEXT("Moved"));
+	
 	UScriptStruct* VectorStruct = FindObjectChecked<UScriptStruct>(UObject::StaticClass(), TEXT("Vector"));
 	CreatePin(EGPD_Output, K2Schema->PC_Struct, TEXT(""), VectorStruct, false, false, TEXT("Location"));
 
@@ -114,6 +115,11 @@ UEdGraphPin* UK2Node_InputTouch::GetReleasedPin() const
 	return FindPin(TEXT("Released"));
 }
 
+UEdGraphPin* UK2Node_InputTouch::GetMovedPin() const
+{
+	return FindPin(TEXT("Moved"));
+}
+
 UEdGraphPin* UK2Node_InputTouch::GetLocationPin() const
 {
 	return FindPin(TEXT("Location"));
@@ -130,102 +136,91 @@ void UK2Node_InputTouch::ExpandNode(FKismetCompilerContext& CompilerContext, UEd
 
 	UEdGraphPin* InputTouchPressedPin = GetPressedPin();
 	UEdGraphPin* InputTouchReleasedPin = GetReleasedPin();
+	UEdGraphPin* InputTouchMovedPin = GetMovedPin();
+		
+	struct EventPinData
+	{
+		EventPinData(UEdGraphPin* InPin,TEnumAsByte<EInputEvent> InEvent ){	Pin=InPin;EventType=InEvent;};
+		UEdGraphPin* Pin;
+		TEnumAsByte<EInputEvent> EventType;
+	};
 
+	TArray<EventPinData> ActivePins;
+	if(( InputTouchPressedPin != nullptr ) && (InputTouchPressedPin->LinkedTo.Num() > 0 ))
+	{
+		ActivePins.Add(EventPinData(InputTouchPressedPin,IE_Pressed));
+	}
+	if((InputTouchReleasedPin != nullptr) && (InputTouchReleasedPin->LinkedTo.Num() > 0 ))
+	{
+		ActivePins.Add(EventPinData(InputTouchReleasedPin,IE_Released));
+	}
+	if((InputTouchMovedPin != nullptr) && ( InputTouchMovedPin->LinkedTo.Num() > 0 ))
+	{
+		ActivePins.Add(EventPinData(InputTouchMovedPin,IE_Repeat));
+	}
+	
 	const UEdGraphSchema_K2* Schema = CompilerContext.GetSchema();
 
-	// If both are linked we have to do more complicated behaviors
-	if (InputTouchPressedPin->LinkedTo.Num() > 0 && InputTouchReleasedPin->LinkedTo.Num() > 0)
+	// If more than one is linked we have to do more complicated behaviors
+	if( ActivePins.Num() > 1 )
 	{
-		// Create the input touch events
-		UK2Node_InputTouchEvent* InputTouchPressedEvent = CompilerContext.SpawnIntermediateNode<UK2Node_InputTouchEvent>(this, SourceGraph);
-		InputTouchPressedEvent->CustomFunctionName = FName( *FString::Printf(TEXT("InpTchEvt_%s"), *InputTouchPressedEvent->GetName()));
-		InputTouchPressedEvent->bConsumeInput = bConsumeInput;
-		InputTouchPressedEvent->bExecuteWhenPaused = bExecuteWhenPaused;
-		InputTouchPressedEvent->bOverrideParentBinding = bOverrideParentBinding;
-		InputTouchPressedEvent->InputKeyEvent = IE_Pressed;
-		InputTouchPressedEvent->EventSignatureName = TEXT("InputTouchHandlerDynamicSignature__DelegateSignature");
-		InputTouchPressedEvent->EventSignatureClass = UInputComponent::StaticClass();
-		InputTouchPressedEvent->bInternalEvent = true;
-		InputTouchPressedEvent->AllocateDefaultPins();
-
-		UK2Node_InputTouchEvent* InputTouchReleasedEvent = CompilerContext.SpawnIntermediateNode<UK2Node_InputTouchEvent>(this, SourceGraph);
-		InputTouchReleasedEvent->CustomFunctionName = FName( *FString::Printf(TEXT("InpTchEvt_%s"),  *InputTouchReleasedEvent->GetName()));
-		InputTouchReleasedEvent->bConsumeInput = bConsumeInput;
-		InputTouchReleasedEvent->bExecuteWhenPaused = bExecuteWhenPaused;
-		InputTouchReleasedEvent->bOverrideParentBinding = bOverrideParentBinding;
-		InputTouchReleasedEvent->InputKeyEvent = IE_Released;
-		InputTouchReleasedEvent->EventSignatureName = TEXT("InputTouchHandlerDynamicSignature__DelegateSignature");
-		InputTouchReleasedEvent->EventSignatureClass = UInputComponent::StaticClass();
-		InputTouchReleasedEvent->bInternalEvent = true;
-		InputTouchReleasedEvent->AllocateDefaultPins();
-
 		// Create a temporary variable to copy location in to
 		static UScriptStruct* VectorStruct = FindObjectChecked<UScriptStruct>(UObject::StaticClass(), TEXT("Vector"));
-
 		UK2Node_TemporaryVariable* TouchLocationVar = CompilerContext.SpawnIntermediateNode<UK2Node_TemporaryVariable>(this, SourceGraph);
 		TouchLocationVar->VariableType.PinCategory = Schema->PC_Struct;
 		TouchLocationVar->VariableType.PinSubCategoryObject = VectorStruct;
 		TouchLocationVar->AllocateDefaultPins();
 
-		// Create assignment nodes to assign the location
-		UK2Node_AssignmentStatement* TouchPressedLocationInitialize = CompilerContext.SpawnIntermediateNode<UK2Node_AssignmentStatement>(this, SourceGraph);
-		TouchPressedLocationInitialize->AllocateDefaultPins();
-		Schema->TryCreateConnection(TouchLocationVar->GetVariablePin(), TouchPressedLocationInitialize->GetVariablePin());
-		Schema->TryCreateConnection(TouchPressedLocationInitialize->GetValuePin(), InputTouchPressedEvent->FindPinChecked(TEXT("Location")));
-
-		UK2Node_AssignmentStatement* TouchReleasedLocationInitialize = CompilerContext.SpawnIntermediateNode<UK2Node_AssignmentStatement>(this, SourceGraph);
-		TouchReleasedLocationInitialize->AllocateDefaultPins();
-		Schema->TryCreateConnection(TouchLocationVar->GetVariablePin(), TouchReleasedLocationInitialize->GetVariablePin());
-		Schema->TryCreateConnection(TouchReleasedLocationInitialize->GetValuePin(), InputTouchReleasedEvent->FindPinChecked(TEXT("Location")));
-
-		// Connect the events to the assign location nodes
-		Schema->TryCreateConnection(Schema->FindExecutionPin(*InputTouchPressedEvent, EGPD_Output), TouchPressedLocationInitialize->GetExecPin());
-		Schema->TryCreateConnection(Schema->FindExecutionPin(*InputTouchReleasedEvent, EGPD_Output), TouchReleasedLocationInitialize->GetExecPin());
-
 		// Create a temporary variable to copy finger index in to
 		UK2Node_TemporaryVariable* TouchFingerVar = CompilerContext.SpawnIntermediateNode<UK2Node_TemporaryVariable>(this, SourceGraph);
 		TouchFingerVar->VariableType.PinCategory = Schema->PC_Byte;
 		TouchFingerVar->VariableType.PinSubCategoryObject = UK2Node_InputTouch::GetTouchIndexEnum();
 		TouchFingerVar->AllocateDefaultPins();
 
-		// Create assignment nodes to assign the finger index
-		UK2Node_AssignmentStatement* TouchPressedFingerInitialize = CompilerContext.SpawnIntermediateNode<UK2Node_AssignmentStatement>(this, SourceGraph);
-		TouchPressedFingerInitialize->AllocateDefaultPins();
-		Schema->TryCreateConnection(TouchFingerVar->GetVariablePin(), TouchPressedFingerInitialize->GetVariablePin());
-		Schema->TryCreateConnection(TouchPressedFingerInitialize->GetValuePin(), InputTouchPressedEvent->FindPinChecked(TEXT("FingerIndex")));
-
-		UK2Node_AssignmentStatement* TouchReleasedFingerInitialize = CompilerContext.SpawnIntermediateNode<UK2Node_AssignmentStatement>(this, SourceGraph);
-		TouchReleasedFingerInitialize->AllocateDefaultPins();
-		Schema->TryCreateConnection(TouchFingerVar->GetVariablePin(), TouchReleasedFingerInitialize->GetVariablePin());
-		Schema->TryCreateConnection(TouchReleasedFingerInitialize->GetValuePin(), InputTouchReleasedEvent->FindPinChecked(TEXT("FingerIndex")));
-
-		// Connect the assign location to the assign finger index nodes
-		Schema->TryCreateConnection(TouchPressedLocationInitialize->GetThenPin(), TouchPressedFingerInitialize->GetExecPin());
-		Schema->TryCreateConnection(TouchReleasedLocationInitialize->GetThenPin(), TouchReleasedFingerInitialize->GetExecPin());
-
-		// Move the original event connections to the then pin of the finger index assign
-		CompilerContext.MovePinLinksToIntermediate(*GetPressedPin(), *TouchPressedFingerInitialize->GetThenPin());
-		CompilerContext.MovePinLinksToIntermediate(*GetReleasedPin(), *TouchReleasedFingerInitialize->GetThenPin());
+		for (auto PinIt = ActivePins.CreateIterator(); PinIt; ++PinIt)
+		{			
+			UEdGraphPin *EachPin= (*PinIt).Pin;
+			// Create the input touch event
+			UK2Node_InputTouchEvent* InputTouchEvent = CompilerContext.SpawnIntermediateNode<UK2Node_InputTouchEvent>(this, SourceGraph);
+			InputTouchEvent->CustomFunctionName = FName(*FString::Printf(TEXT("InpTchEvt_%s"), *EachPin->GetName()));
+			InputTouchEvent->bConsumeInput = bConsumeInput;
+			InputTouchEvent->bExecuteWhenPaused = bExecuteWhenPaused;
+			InputTouchEvent->bOverrideParentBinding = bOverrideParentBinding;
+			InputTouchEvent->InputKeyEvent = (*PinIt).EventType;
+			InputTouchEvent->EventSignatureName = TEXT("InputTouchHandlerDynamicSignature__DelegateSignature");
+			InputTouchEvent->EventSignatureClass = UInputComponent::StaticClass();
+			InputTouchEvent->bInternalEvent = true;
+			InputTouchEvent->AllocateDefaultPins();
 
-		// Move the original event variable connections to the intermediate nodes
-		CompilerContext.MovePinLinksToIntermediate(*GetLocationPin(), *TouchLocationVar->GetVariablePin());
-		CompilerContext.MovePinLinksToIntermediate(*GetFingerIndexPin(), *TouchFingerVar->GetVariablePin());
+			// Create assignment nodes to assign the location
+			UK2Node_AssignmentStatement* TouchLocationInitialize = CompilerContext.SpawnIntermediateNode<UK2Node_AssignmentStatement>(this, SourceGraph);
+			TouchLocationInitialize->AllocateDefaultPins();
+			Schema->TryCreateConnection(TouchLocationVar->GetVariablePin(), TouchLocationInitialize->GetVariablePin());
+			Schema->TryCreateConnection(TouchLocationInitialize->GetValuePin(), InputTouchEvent->FindPinChecked(TEXT("Location")));
+			// Connect the events to the assign location nodes
+			Schema->TryCreateConnection(Schema->FindExecutionPin(*InputTouchEvent, EGPD_Output), TouchLocationInitialize->GetExecPin());
+
+			// Create assignment nodes to assign the finger index
+			UK2Node_AssignmentStatement* TouchFingerInitialize = CompilerContext.SpawnIntermediateNode<UK2Node_AssignmentStatement>(this, SourceGraph);
+			TouchFingerInitialize->AllocateDefaultPins();
+			Schema->TryCreateConnection(TouchFingerVar->GetVariablePin(), TouchFingerInitialize->GetVariablePin());
+			Schema->TryCreateConnection(TouchFingerInitialize->GetValuePin(), InputTouchEvent->FindPinChecked(TEXT("FingerIndex")));
+			// Connect the assign location to the assign finger index nodes
+			Schema->TryCreateConnection(TouchLocationInitialize->GetThenPin(), TouchFingerInitialize->GetExecPin());			
+
+			// Move the original event connections to the then pin of the finger index assign
+			CompilerContext.MovePinLinksToIntermediate(*EachPin, *TouchFingerInitialize->GetThenPin());
+			
+			// Move the original event variable connections to the intermediate nodes
+			CompilerContext.MovePinLinksToIntermediate(*GetLocationPin(), *TouchLocationVar->GetVariablePin());
+			CompilerContext.MovePinLinksToIntermediate(*GetFingerIndexPin(), *TouchFingerVar->GetVariablePin());
+		}	
 	}
-	else
+	else if( ActivePins.Num() == 1 )
 	{
-		UEdGraphPin* InputTouchPin; 
-		EInputEvent InputEvent;
-		if (InputTouchPressedPin->LinkedTo.Num() > 0)
-		{
-			InputTouchPin = InputTouchPressedPin;
-			InputEvent = IE_Pressed;
-		}
-		else
-		{
-			InputTouchPin = InputTouchReleasedPin;
-			InputEvent = IE_Released;
-		}
-
+		UEdGraphPin* InputTouchPin = ActivePins[0].Pin;
+		EInputEvent InputEvent = ActivePins[0].EventType;
+		
 		if (InputTouchPin->LinkedTo.Num() > 0)
 		{
 			UK2Node_InputTouchEvent* InputTouchEvent = CompilerContext.SpawnIntermediateNode<UK2Node_InputTouchEvent>(this, SourceGraph);
