Fixing some bugs with MakeArray caused by recent changes and issues with local variable defaults being an array.

diff --git a/Engine/Source/Editor/BlueprintGraph/Private/K2Node_FunctionEntry.cpp b/Engine/Source/Editor/BlueprintGraph/Private/K2Node_FunctionEntry.cpp
--- a/Engine/Source/Editor/BlueprintGraph/Private/K2Node_FunctionEntry.cpp
+++ b/Engine/Source/Editor/BlueprintGraph/Private/K2Node_FunctionEntry.cpp
@@ -329,7 +329,7 @@ void UK2Node_FunctionEntry::ExpandNode(class FKismetCompilerContext& CompilerCon
 
 									// Retrieve the element's default value
 									FString DefaultValue;
-									ArrayProperty->Inner->ExportTextItem( DefaultValue, PropData, PropDefault, NULL, PPF_Delimited, NULL );
+									FBlueprintEditorUtils::PropertyValueToString(ArrayProperty->Inner, PropData, DefaultValue);
 
 									if(ArrayIndex > 0)
 									{
diff --git a/Engine/Source/Editor/BlueprintGraph/Private/K2Node_MakeArray.cpp b/Engine/Source/Editor/BlueprintGraph/Private/K2Node_MakeArray.cpp
--- a/Engine/Source/Editor/BlueprintGraph/Private/K2Node_MakeArray.cpp
+++ b/Engine/Source/Editor/BlueprintGraph/Private/K2Node_MakeArray.cpp
@@ -308,7 +308,8 @@ void UK2Node_MakeArray::AddInputPin()
 	Modify();
 
 	++NumInputs;
-	CreatePin(EGPD_Input, GetOutputPin()->PinType.PinCategory, TEXT(""), NULL, false, false, *FString::Printf(TEXT("[%d]"), (NumInputs-1)));
+	FEdGraphPinType OutputPinType = GetOutputPin()->PinType;
+	CreatePin(EGPD_Input, OutputPinType.PinCategory, OutputPinType.PinSubCategory, OutputPinType.PinSubCategoryObject.Get(), false, false, *FString::Printf(TEXT("[%d]"), (NumInputs-1)));
 	
 	const bool bIsCompiling = GetBlueprint()->bBeingCompiled;
 	if( !bIsCompiling )
