[UE-6362] Fixing it up so you can use interface inputs into object params (+option to auto cast object connections).

diff --git a/Engine/Source/Editor/BlueprintGraph/Classes/K2Node_ClassDynamicCast.h b/Engine/Source/Editor/BlueprintGraph/Classes/K2Node_ClassDynamicCast.h
--- a/Engine/Source/Editor/BlueprintGraph/Classes/K2Node_ClassDynamicCast.h
+++ b/Engine/Source/Editor/BlueprintGraph/Classes/K2Node_ClassDynamicCast.h
@@ -18,6 +18,7 @@ class UK2Node_ClassDynamicCast : public UK2Node_DynamicCast
 
 	// UK2Node interface
 	virtual class FNodeHandlingFunctor* CreateNodeHandler(class FKismetCompilerContext& CompilerContext) const override;
+	virtual bool IsConnectionDisallowed(const UEdGraphPin* MyPin, const UEdGraphPin* OtherPin, FString& OutReason) const override { return false; }
 	// End of UK2Node interface
 
 	// UK2Node_DynamicCast interface
diff --git a/Engine/Source/Editor/BlueprintGraph/Classes/K2Node_DynamicCast.h b/Engine/Source/Editor/BlueprintGraph/Classes/K2Node_DynamicCast.h
--- a/Engine/Source/Editor/BlueprintGraph/Classes/K2Node_DynamicCast.h
+++ b/Engine/Source/Editor/BlueprintGraph/Classes/K2Node_DynamicCast.h
@@ -21,6 +21,7 @@ class UK2Node_DynamicCast : public UK2Node
 	virtual FText GetNodeTitle(ENodeTitleType::Type TitleType) const override;
 	virtual FName GetPaletteIcon(FLinearColor& OutColor) const override{ return TEXT("GraphEditor.Cast_16x"); }
 	virtual void GetContextMenuActions(const FGraphNodeContextMenuBuilder& Context) const override;
+	virtual void PostReconstructNode() override;
 	// End UEdGraphNode interface
 
 	// UK2Node interface
@@ -30,6 +31,8 @@ class UK2Node_DynamicCast : public UK2Node
 	virtual FText GetMenuCategory() const override;
 	virtual FBlueprintNodeSignature GetSignature() const override;
 	virtual bool IsNodePure() const override { return bIsPureCast; }
+	virtual bool IsConnectionDisallowed(const UEdGraphPin* MyPin, const UEdGraphPin* OtherPin, FString& OutReason) const override;
+	virtual void NotifyPinConnectionListChanged(UEdGraphPin* Pin) override;
 	// End of UK2Node interface
 
 	/** Get the 'valid cast' exec pin */
diff --git a/Engine/Source/Editor/BlueprintGraph/Private/BlueprintEditorSettings.cpp b/Engine/Source/Editor/BlueprintGraph/Private/BlueprintEditorSettings.cpp
--- a/Engine/Source/Editor/BlueprintGraph/Private/BlueprintEditorSettings.cpp
+++ b/Engine/Source/Editor/BlueprintGraph/Private/BlueprintEditorSettings.cpp
@@ -16,6 +16,7 @@ UBlueprintEditorSettings::UBlueprintEditorSettings(const FObjectInitializer& Obj
 	, bShowContextualFavorites(false)
 	, bFlattenFavoritesMenus(true)
 	, bFavorPureCastNodes(false)
+	, bAutoCastObjectConnections(false)
 	, bUseLegacyMenuingSystem(false)
 	// Compiler Settings
 	, SaveOnCompile(SoC_Never)
diff --git a/Engine/Source/Editor/BlueprintGraph/Private/EdGraphSchema_K2.cpp b/Engine/Source/Editor/BlueprintGraph/Private/EdGraphSchema_K2.cpp
--- a/Engine/Source/Editor/BlueprintGraph/Private/EdGraphSchema_K2.cpp
+++ b/Engine/Source/Editor/BlueprintGraph/Private/EdGraphSchema_K2.cpp
@@ -9,7 +9,7 @@
 #include "GraphEditorSettings.h"
 #include "ScopedTransaction.h"
 #include "ComponentAssetBroker.h"
-
+#include "BlueprintEditorSettings.h"
 #include "Kismet2/KismetEditorUtilities.h"
 #include "Kismet2/KismetDebugUtilities.h"
 #include "KismetCompiler.h"
@@ -2236,11 +2236,14 @@ bool UEdGraphSchema_K2::FindSpecializedConversionNode(const UEdGraphPin* OutputP
 
 	if (!bCanConvert && InputPin && OutputPin)
 	{
+		FEdGraphPinType const& InputType  = InputPin->PinType;
+		FEdGraphPinType const& OutputType = OutputPin->PinType;
+
 		// CHECK BYTE TO ENUM CAST
-		UEnum* Enum = Cast<UEnum>(InputPin->PinType.PinSubCategoryObject.Get());
-		const bool bInoputMatch = !InputPin->PinType.bIsArray && (PC_Byte == InputPin->PinType.PinCategory) && Enum;
-		const bool bOutputMatch = !OutputPin->PinType.bIsArray && (PC_Byte == OutputPin->PinType.PinCategory);
-		if(bOutputMatch && bInoputMatch)
+		UEnum* Enum = Cast<UEnum>(InputType.PinSubCategoryObject.Get());
+		const bool bInputIsEnum = !InputType.bIsArray && (PC_Byte == InputType.PinCategory) && Enum;
+		const bool bOutputIsByte = !OutputType.bIsArray && (PC_Byte == OutputType.PinCategory);
+		if (bInputIsEnum && bOutputIsByte)
 		{
 			bCanConvert = true;
 			if(bCreateNode)
@@ -2251,6 +2254,32 @@ bool UEdGraphSchema_K2::FindSpecializedConversionNode(const UEdGraphPin* OutputP
 				TargetNode = CastByteToEnum;
 			}
 		}
+		else
+		{
+			UClass* InputClass  = Cast<UClass>(InputType.PinSubCategoryObject.Get());
+			UClass* OutputClass = Cast<UClass>(InputType.PinSubCategoryObject.Get());
+
+			if ((OutputType.PinCategory == PC_Interface) && (InputType.PinCategory == PC_Object))
+			{
+				bCanConvert = (InputClass && OutputClass) && (InputClass->ImplementsInterface(OutputClass) || OutputClass->IsChildOf(InputClass));
+			}
+			else if (OutputType.PinCategory == PC_Object)
+			{
+				UBlueprintEditorSettings const* BlueprintSettings = GetDefault<UBlueprintEditorSettings>();
+				if ((InputType.PinCategory == PC_Object) && BlueprintSettings->bAutoCastObjectConnections)
+				{
+					bCanConvert = (InputClass && OutputClass) && InputClass->IsChildOf(OutputClass);
+				}
+			}
+
+			if (bCanConvert && bCreateNode)
+			{
+				UK2Node_DynamicCast* DynCastNode = NewObject<UK2Node_DynamicCast>();
+				DynCastNode->TargetType = InputClass;
+				DynCastNode->SetPurity(true);
+				TargetNode = DynCastNode;
+			}
+		}
 	}
 
 	return bCanConvert;
@@ -2337,6 +2366,7 @@ FString UEdGraphSchema_K2::IsPinDefaultValid(const UEdGraphPin* Pin, const FStri
 	const bool bIsArray = Pin->PinType.bIsArray;
 	const bool bIsReference = Pin->PinType.bIsReference;
 	const bool bIsAutoCreateRefTerm = IsAutoCreateRefTerm(Pin);
+	const bool bIsInterfacePin = (Pin->PinType.PinCategory == PC_Interface);
 
 	FFormatNamedArguments MessageArgs;
 	MessageArgs.Add(TEXT("PinName"), Pin->GetDisplayName());
@@ -2352,9 +2382,14 @@ FString UEdGraphSchema_K2::IsPinDefaultValid(const UEdGraphPin* Pin, const FStri
 			}
 			else if( bIsReference )
 			{
-				FText MsgFormat = LOCTEXT("BadRefDefaultVal", "'{PinName}' must have an input wired into it (\"by ref\" params expect a valid input to operate on).");
+				FText MsgFormat = LOCTEXT("BadRefDefaultVal", "'{PinNd ame}' must have an input wired into it (\"by ref\" params expect a valid input to operate on).");
 				return FText::Format(MsgFormat, MessageArgs).ToString();
 			}
+// 			else if (bIsInterfacePin)
+// 			{
+// 				FText MsgFormat = LOCTEXT("BadDefaultInterfaceVal", "'{PinName}' must have an input wired into it.");
+// 				return FText::Format(MsgFormat, MessageArgs).ToString();
+// 			}
 		}
 	}
 
@@ -3435,18 +3470,6 @@ bool UEdGraphSchema_K2::ArePinTypesCompatible(const FEdGraphPinType& Output, con
 
 		return true;
 	}
-	else if ((Output.PinCategory == PC_Interface) && (Input.PinCategory == PC_Object))
-	{
-		UClass const* InterfaceClass = Cast<UClass const>(Output.PinSubCategoryObject.Get());
-		UClass const* InputClass     = Cast<UClass const>(Input.PinSubCategoryObject.Get());
-
-		if ((InputClass == nullptr) && (Input.PinSubCategory == PSC_Self))
-		{
-			InputClass = CallingContext;
-		}
-
-		return InputClass && (InputClass->ImplementsInterface(InterfaceClass) || InterfaceClass->IsChildOf(InputClass));
-	}
 	else if ((Output.PinCategory == PC_Object) && (Input.PinCategory == PC_Interface))
 	{
 		UClass const* OutputClass    = Cast<UClass const>(Output.PinSubCategoryObject.Get());
diff --git a/Engine/Source/Editor/BlueprintGraph/Private/K2Node_DynamicCast.cpp b/Engine/Source/Editor/BlueprintGraph/Private/K2Node_DynamicCast.cpp
--- a/Engine/Source/Editor/BlueprintGraph/Private/K2Node_DynamicCast.cpp
+++ b/Engine/Source/Editor/BlueprintGraph/Private/K2Node_DynamicCast.cpp
@@ -45,7 +45,7 @@ void UK2Node_DynamicCast::AllocateDefaultPins()
 	}
 
 	// Input - Source type Pin
-	CreatePin(EGPD_Input, K2Schema->PC_Object, TEXT(""), UObject::StaticClass(), false, false, K2Schema->PN_ObjectToCast);
+	CreatePin(EGPD_Input, K2Schema->PC_Wildcard, TEXT(""), UObject::StaticClass(), false, false, K2Schema->PN_ObjectToCast);
 
 	// Output - Data Pin
 	if (TargetType != NULL)
@@ -145,6 +145,13 @@ void UK2Node_DynamicCast::GetContextMenuActions(const FGraphNodeContextMenuBuild
 	Context.MenuBuilder->EndSection();
 }
 
+void UK2Node_DynamicCast::PostReconstructNode()
+{
+	Super::PostReconstructNode();
+	// update the pin name (to "Interface" if an interface is connected)
+	NotifyPinConnectionListChanged(GetCastSourcePin());
+}
+
 UEdGraphPin* UK2Node_DynamicCast::GetValidCastPin() const
 {
 	const UEdGraphSchema_K2* K2Schema = GetDefault<UEdGraphSchema_K2>();
@@ -274,4 +281,99 @@ FBlueprintNodeSignature UK2Node_DynamicCast::GetSignature() const
 	return NodeSignature;
 }
 
+bool UK2Node_DynamicCast::IsConnectionDisallowed(const UEdGraphPin* MyPin, const UEdGraphPin* OtherPin, FString& OutReason) const
+{
+	bool bIsDisallowed = Super::IsConnectionDisallowed(MyPin, OtherPin, OutReason);
+
+	if (MyPin == GetCastSourcePin())
+	{
+		const FEdGraphPinType& OtherPinType = OtherPin->PinType;
+		const FText OtherPinName = OtherPin->PinFriendlyName.IsEmpty() ? FText::FromString(OtherPin->PinName) : OtherPin->PinFriendlyName;
+
+		if (OtherPinType.bIsArray)
+		{
+			bIsDisallowed = true;
+			OutReason = LOCTEXT("CannotArrayCast", "You cannot cast arrays of objects.").ToString();
+		}
+		else if ((OtherPinType.PinCategory == UEdGraphSchema_K2::PC_Interface) || TargetType->HasAnyClassFlags(CLASS_Interface))
+		{
+			// allow all interface casts
+		}
+		else if (OtherPinType.PinCategory == UEdGraphSchema_K2::PC_Object)
+		{
+			UClass* ObjectClass = Cast<UClass>(OtherPinType.PinSubCategoryObject.Get());
+			if ((ObjectClass == nullptr) && (OtherPinType.PinSubCategory == UEdGraphSchema_K2::PSC_Self))
+			{
+				if (UK2Node* K2Node = Cast<UK2Node>(OtherPin->GetOwningNode()))
+				{
+					ObjectClass = K2Node->GetBlueprint()->GeneratedClass;
+				}
+			}
+			// if the ObjectClass is still null, assume it is a UObject, which 
+			// will work with everything (so don't disallow it)
+			
+			if (ObjectClass != nullptr)
+			{
+				if (ObjectClass == TargetType)
+				{
+					bIsDisallowed = true;
+					OutReason = FText::Format(LOCTEXT("EqualObjectCast", "'{0}' is already a '{1}', you don't need the cast."),
+						OtherPinName, TargetType->GetDisplayNameText()).ToString();
+				}
+				else if (ObjectClass->IsChildOf(TargetType))
+				{
+					bIsDisallowed = true;
+					OutReason = FText::Format(LOCTEXT("UnneededObjectCast", "'{0}' is already a '{1}' (which inherits from '{2}'), so you don't need the cast."),
+						OtherPinName, ObjectClass->GetDisplayNameText(), TargetType->GetDisplayNameText()).ToString();
+				}
+				else if (!TargetType->IsChildOf(ObjectClass))
+				{
+					bIsDisallowed = true;
+					OutReason = FText::Format(LOCTEXT("DisallowedObjectCast", "'{0}' does not inherit from '{1}' (the cast would always fail)."),
+						TargetType->GetDisplayNameText(), ObjectClass->GetDisplayNameText()).ToString();
+				}
+			}
+		}
+		else
+		{
+			bIsDisallowed = true;
+			OutReason = LOCTEXT("NonObjectCast", "You can only cast objects/interfaces.").ToString();
+		}
+	}
+	return bIsDisallowed;
+}
+
+void UK2Node_DynamicCast::NotifyPinConnectionListChanged(UEdGraphPin* Pin)
+{
+	Super::NotifyPinConnectionListChanged(Pin);
+
+	if (Pin == GetCastSourcePin())
+	{
+		Pin->PinFriendlyName = FText::GetEmpty();
+
+		FEdGraphPinType& InputPinType = Pin->PinType;
+		if (Pin->LinkedTo.Num() == 0)
+		{
+			InputPinType.PinCategory = UEdGraphSchema_K2::PC_Wildcard;
+			InputPinType.PinSubCategory.Empty();
+			InputPinType.PinSubCategoryObject = nullptr;
+		}
+		else
+		{
+			const FEdGraphPinType& ConnectedPinType = Pin->LinkedTo[0]->PinType;
+			if (ConnectedPinType.PinCategory == UEdGraphSchema_K2::PC_Interface)
+			{
+				Pin->PinFriendlyName = LOCTEXT("InterfaceInputName", "Interface");
+				InputPinType.PinCategory = UEdGraphSchema_K2::PC_Interface;
+				InputPinType.PinSubCategoryObject = ConnectedPinType.PinSubCategoryObject;
+			}
+			else if (ConnectedPinType.PinCategory == UEdGraphSchema_K2::PC_Object)
+			{
+				InputPinType.PinCategory = UEdGraphSchema_K2::PC_Object;
+				InputPinType.PinSubCategoryObject = UObject::StaticClass();
+			}
+		}
+	}
+}
+
 #undef LOCTEXT_NAMESPACE
diff --git a/Engine/Source/Editor/BlueprintGraph/Public/BlueprintEditorSettings.h b/Engine/Source/Editor/BlueprintGraph/Public/BlueprintEditorSettings.h
--- a/Engine/Source/Editor/BlueprintGraph/Public/BlueprintEditorSettings.h
+++ b/Engine/Source/Editor/BlueprintGraph/Public/BlueprintEditorSettings.h
@@ -51,6 +51,10 @@ public:
 	UPROPERTY(EditAnywhere, config, Category=Workflow, meta=(DisplayName="Default to Using Pure Cast Nodes"))
 	bool bFavorPureCastNodes;
 
+	/** If enabled, then you'll be able to directly connect arbitrary object pins together (a pure cast node will be injected automatically). */
+	UPROPERTY(EditAnywhere, config, Category=Workflow)
+	bool bAutoCastObjectConnections;
+
 	/** If set, then the new refactored menu system will be replaced with the old (legacy) system (as a fallback, in case the new system has unforeseen problems)*/
 	UPROPERTY(EditAnywhere, AdvancedDisplay, config, Category=Workflow)
 	bool bUseLegacyMenuingSystem;
