[UE-6367] Using the "Default to Pure Cast" option no longer modifies old cast nodes.

diff --git a/Engine/Source/Editor/BlueprintGraph/Classes/K2Node_DynamicCast.h b/Engine/Source/Editor/BlueprintGraph/Classes/K2Node_DynamicCast.h
--- a/Engine/Source/Editor/BlueprintGraph/Classes/K2Node_DynamicCast.h
+++ b/Engine/Source/Editor/BlueprintGraph/Classes/K2Node_DynamicCast.h
@@ -22,6 +22,7 @@ class UK2Node_DynamicCast : public UK2Node
 	virtual FName GetPaletteIcon(FLinearColor& OutColor) const override{ return TEXT("GraphEditor.Cast_16x"); }
 	virtual void GetContextMenuActions(const FGraphNodeContextMenuBuilder& Context) const override;
 	virtual void PostReconstructNode() override;
+	virtual void PostPlacedNewNode() override;
 	// End UEdGraphNode interface
 
 	// UK2Node interface
diff --git a/Engine/Source/Editor/BlueprintGraph/Private/K2Node_DynamicCast.cpp b/Engine/Source/Editor/BlueprintGraph/Private/K2Node_DynamicCast.cpp
--- a/Engine/Source/Editor/BlueprintGraph/Private/K2Node_DynamicCast.cpp
+++ b/Engine/Source/Editor/BlueprintGraph/Private/K2Node_DynamicCast.cpp
@@ -17,8 +17,6 @@ UK2Node_DynamicCast::UK2Node_DynamicCast(const FObjectInitializer& ObjectInitial
 	: Super(ObjectInitializer)
 	, bIsPureCast(false)
 {
-	const UBlueprintEditorSettings* BlueprintSettings = GetDefault<UBlueprintEditorSettings>();
-	bIsPureCast = BlueprintSettings->bFavorPureCastNodes;
 }
 
 void UK2Node_DynamicCast::AllocateDefaultPins()
@@ -152,6 +150,14 @@ void UK2Node_DynamicCast::PostReconstructNode()
 	NotifyPinConnectionListChanged(GetCastSourcePin());
 }
 
+void UK2Node_DynamicCast::PostPlacedNewNode()
+{
+	Super::PostPlacedNewNode();
+
+	const UBlueprintEditorSettings* BlueprintSettings = GetDefault<UBlueprintEditorSettings>();
+	SetPurity(BlueprintSettings->bFavorPureCastNodes);
+}
+
 UEdGraphPin* UK2Node_DynamicCast::GetValidCastPin() const
 {
 	const UEdGraphSchema_K2* K2Schema = GetDefault<UEdGraphSchema_K2>();
