Fixing cook crash the results from malformed cast nodes (still need to figure out how the bad cast was generated in the first place).

diff --git a/Engine/Source/Editor/BlueprintGraph/Private/K2Node_DynamicCast.cpp b/Engine/Source/Editor/BlueprintGraph/Private/K2Node_DynamicCast.cpp
--- a/Engine/Source/Editor/BlueprintGraph/Private/K2Node_DynamicCast.cpp
+++ b/Engine/Source/Editor/BlueprintGraph/Private/K2Node_DynamicCast.cpp
@@ -301,6 +301,11 @@ bool UK2Node_DynamicCast::IsConnectionDisallowed(const UEdGraphPin* MyPin, const
 			bIsDisallowed = true;
 			OutReason = LOCTEXT("CannotArrayCast", "You cannot cast arrays of objects.").ToString();
 		}
+		else if (TargetType == nullptr)
+		{
+			bIsDisallowed = true;
+			OutReason = LOCTEXT("BadCastNode", "This cast has an invalid target type (was the class deleted without a redirect?).").ToString();
+		}
 		else if ((OtherPinType.PinCategory == UEdGraphSchema_K2::PC_Interface) || TargetType->HasAnyClassFlags(CLASS_Interface))
 		{
 			// allow all interface casts
@@ -320,6 +325,7 @@ bool UK2Node_DynamicCast::IsConnectionDisallowed(const UEdGraphPin* MyPin, const
 			
 			if (ObjectClass != nullptr)
 			{
+				
 				if (ObjectClass == TargetType)
 				{
 					bIsDisallowed = true;
