Fixing up issues when defaulting to pure casts (moved setting to experimental until we can iron out all the issues).

diff --git a/Engine/Source/Editor/BlueprintGraph/Private/K2Node_Message.cpp b/Engine/Source/Editor/BlueprintGraph/Private/K2Node_Message.cpp
--- a/Engine/Source/Editor/BlueprintGraph/Private/K2Node_Message.cpp
+++ b/Engine/Source/Editor/BlueprintGraph/Private/K2Node_Message.cpp
@@ -133,6 +133,7 @@ void UK2Node_Message::ExpandNode(class FKismetCompilerContext& CompilerContext,
 		// First, create an internal cast-to-interface node
 		UK2Node_DynamicCast* CastToInterfaceNode = CompilerContext.SpawnIntermediateNode<UK2Node_DynamicCast>(this, SourceGraph);
 		CastToInterfaceNode->TargetType = MessageNodeFunction->GetOuterUClass();
+		CastToInterfaceNode->SetPurity(false);
 		CastToInterfaceNode->AllocateDefaultPins();
 
 		UEdGraphPin* CastToInterfaceResultPin = CastToInterfaceNode->GetCastResultPin();
diff --git a/Engine/Source/Editor/BlueprintGraph/Public/BlueprintEditorSettings.h b/Engine/Source/Editor/BlueprintGraph/Public/BlueprintEditorSettings.h
--- a/Engine/Source/Editor/BlueprintGraph/Public/BlueprintEditorSettings.h
+++ b/Engine/Source/Editor/BlueprintGraph/Public/BlueprintEditorSettings.h
@@ -48,7 +48,7 @@ public:
 	bool bFlattenFavoritesMenus;
 
 	/** If enabled, then placed cast nodes will default to their "pure" form (meaning: without execution pins). */
-	UPROPERTY(EditAnywhere, config, Category=Workflow, meta=(DisplayName="Default to Using Pure Cast Nodes"))
+	UPROPERTY(EditAnywhere, config, AdvancedDisplay, Category=Experimental, meta=(DisplayName="Default to Using Pure Cast Nodes"))
 	bool bFavorPureCastNodes;
 
 	/** If set, then the new refactored menu system will be replaced with the old (legacy) system (as a fallback, in case the new system has unforeseen problems)*/
