UE-6487 Right-clicking in the Event Graph of a second Blueprint created using Create Blueprint Using This... will crash the editor

diff --git a/Engine/Source/Editor/BlueprintGraph/Private/BlueprintVariableNodeSpawner.cpp b/Engine/Source/Editor/BlueprintGraph/Private/BlueprintVariableNodeSpawner.cpp
--- a/Engine/Source/Editor/BlueprintGraph/Private/BlueprintVariableNodeSpawner.cpp
+++ b/Engine/Source/Editor/BlueprintGraph/Private/BlueprintVariableNodeSpawner.cpp
@@ -11,6 +11,7 @@
 #include "BlueprintActionFilter.h"	// for FBlueprintActionContext
 #include "K2Node_VariableGet.h"
 #include "K2Node_VariableSet.h"
+#include "Engine/BlueprintGeneratedClass.h"
 
 #define LOCTEXT_NAMESPACE "BlueprintVariableNodeSpawner"
 
@@ -181,8 +182,10 @@ FBlueprintActionUiSpec UBlueprintVariableNodeSpawner::GetUiSpec(FBlueprintAction
 			}
 		}
 
-		UClass const* VariableClass = WrappedVariable->GetOwnerClass()->GetAuthoritativeClass();
-		if (!TargetClass->IsChildOf(VariableClass))
+		auto OwnerClass = WrappedVariable->GetOwnerClass();
+		const bool bIsOwneClassValid = OwnerClass && (!Cast<const UBlueprintGeneratedClass>(OwnerClass) || OwnerClass->ClassGeneratedBy); //todo: more general validation
+		UClass const* VariableClass = bIsOwneClassValid ? OwnerClass->GetAuthoritativeClass() : NULL;
+		if (VariableClass && !TargetClass->IsChildOf(VariableClass))
 		{
 			MenuSignature.Category = FEditorCategoryUtils::BuildCategoryString(FCommonEditorCategory::Class,
 				FText::FromString(VariableClass->GetDisplayNameText().ToString()));
