[UE-6638] EditAnywhere/DefaultsOnly/InstanceOnly does not determine native blueprint properties' accessibility.

diff --git a/Engine/Source/Editor/BlueprintGraph/Private/BlueprintActionFilter.cpp b/Engine/Source/Editor/BlueprintGraph/Private/BlueprintActionFilter.cpp
--- a/Engine/Source/Editor/BlueprintGraph/Private/BlueprintActionFilter.cpp
+++ b/Engine/Source/Editor/BlueprintGraph/Private/BlueprintActionFilter.cpp
@@ -517,29 +517,41 @@ static bool BlueprintActionFilterImpl::IsFieldInaccessible(FBlueprintActionFilte
 		bool bIsPublic = !bIsPrivate && !bIsProtected;
 		if (UProperty const* Property = Cast<UProperty>(Field))
 		{
-			// CPF_DisableEditOnInstance corresponds to the eyeball/editable 
-			// checkbox available in the blueprint editor. It is poorly named. 
-			bIsPublic = !Property->HasAnyPropertyFlags(CPF_DisableEditOnInstance);
-
-			// If we've disabled editing on the instance, it's entirely possible that the variable has
-			// the flags, EditDefaultsOnly, BlueprintReadOnly, which should still be public, but only
-			// for getting the value.  This is for Native properties.
-			bIsPublic = bIsPublic || Property->HasAnyPropertyFlags(CPF_BlueprintReadOnly);
+			// default to assuming that this property is a native one (the
+			// accessibility there is a little more lax)
+			bool bIsNativeProperty = true;
+			if (UClass* PropertyClass = Property->GetOwnerClass())
+			{
+				bIsNativeProperty = (Cast<UBlueprintGeneratedClass>(PropertyClass) == nullptr);
+			}
 
-			// When CPF_DisableEditOnInstance is *not* set, the user has tagged 
-			// the variable as 'editable' and it is treated as public (delegate 
-			// properties don't have the visibility check box, so they default
-			// to public):
-			bIsPublic = bIsPublic || (Cast<UMulticastDelegateProperty>(Property) != nullptr);
+			// if this is a native property, then the user only should have to 
+			// specify BlueprintReadWrite or BlueprintReadOnly for it to be 
+			// accessible
+			if (bIsNativeProperty)
+			{
+				bIsPublic &= Property->HasAnyPropertyFlags(CPF_BlueprintVisible);
+			}
+			else
+			{
+				// however, if this is a blueprint variable, then the
+				// CPF_DisableEditOnInstance corresponds to the eyeball/editable 
+				// checkbox in the editor, and we don't want the variable to be 
+				// public until the user exposes it (delegates don't have an
+				// 'editable' toggle, so they should default to public)
+				bIsPublic &= !Property->HasAnyPropertyFlags(CPF_DisableEditOnInstance) ||
+					(Cast<UMulticastDelegateProperty>(Property) != nullptr);
+
+				// @TODO: allow users to choose "read-only" for blueprint 
+				//        variables ("private" might have provided this in the 
+				//        past, but that is just wrong, or needs to be renamed)
+			}
 
 			// If the variable is not public, and is not tagged as private, it should be protected
 			// by default. This branch handles variables declared in blueprints:
-			if( !bIsPublic )
+			if (!bIsNativeProperty && !bIsPublic && !bIsPrivate)
 			{
-				if( !bIsPrivate )
-				{
-					bIsProtected = true; 
-				}
+				bIsProtected = true;
 			}
 		}
 
