UE-6396 Event Dispatcher not being called in 4.6

diff --git a/Engine/Source/Editor/BlueprintGraph/Private/K2Node_MCDelegate.cpp b/Engine/Source/Editor/BlueprintGraph/Private/K2Node_MCDelegate.cpp
--- a/Engine/Source/Editor/BlueprintGraph/Private/K2Node_MCDelegate.cpp
+++ b/Engine/Source/Editor/BlueprintGraph/Private/K2Node_MCDelegate.cpp
@@ -97,11 +97,12 @@ UFunction* UK2Node_BaseMCDelegate::GetDelegateSignature(bool bForceNotFromSkelCl
 	}
 	else if (UBlueprintGeneratedClass* BpClassOwner = Cast<UBlueprintGeneratedClass>(OwnerClass))
 	{
-		UBlueprint* DelegateBlueprint = CastChecked<UBlueprint>(BpClassOwner->ClassGeneratedBy);
+		UBlueprint* DelegateBlueprint = Cast<UBlueprint>(BpClassOwner->ClassGeneratedBy);
 		// favor the skeleton class, because the generated class may not 
 		// have the delegate yet (hasn't been compiled with it), or it could 
 		// be out of date
-		OwnerClass = DelegateBlueprint->SkeletonGeneratedClass;
+		UClass* SkeletonClass = DelegateBlueprint ? DelegateBlueprint->SkeletonGeneratedClass : nullptr;
+		OwnerClass = SkeletonClass ? SkeletonClass : OwnerClass;
 	}
 
 	FMemberReference ReferenceToUse;
