Improvements to "Spawn Actor From Class" node to prevent the input pins created for the class from disconnecting when changing classes (if the pins remain).

diff --git a/Engine/Source/Editor/BlueprintGraph/Classes/K2Node.h b/Engine/Source/Editor/BlueprintGraph/Classes/K2Node.h
--- a/Engine/Source/Editor/BlueprintGraph/Classes/K2Node.h
+++ b/Engine/Source/Editor/BlueprintGraph/Classes/K2Node.h
@@ -381,6 +381,12 @@ protected:
 	// Handles the actual reconstruction (copying data, links, name, etc...) from two pins that have already been matched together
 	BLUEPRINTGRAPH_API void ReconstructSinglePin(UEdGraphPin* NewPin, UEdGraphPin* OldPin, ERedirectType RedirectType);
 
+	// Helper function to rewire old pins to new pins during node reconstruction (or other regeneration of pins)
+	BLUEPRINTGRAPH_API void RewireOldPinsToNewPins(TArray<UEdGraphPin*>& InOldPins, TArray<UEdGraphPin*>& InNewPins);
+
+	// Helper function to properly destroy a set of pins
+	BLUEPRINTGRAPH_API void DestroyPinList(TArray<UEdGraphPin*>& InPins);
+
 	/** Allows the custom transformation of a param's value when redirecting a matched pin; called only when DoPinsMatchForReconstruction returns ERedirectType_Custom **/
 	BLUEPRINTGRAPH_API virtual void CustomMapParamValue(UEdGraphPin& Pin);
 
diff --git a/Engine/Source/Editor/BlueprintGraph/Classes/K2Node_SpawnActorFromClass.h b/Engine/Source/Editor/BlueprintGraph/Classes/K2Node_SpawnActorFromClass.h
--- a/Engine/Source/Editor/BlueprintGraph/Classes/K2Node_SpawnActorFromClass.h
+++ b/Engine/Source/Editor/BlueprintGraph/Classes/K2Node_SpawnActorFromClass.h
@@ -22,6 +22,7 @@ class UK2Node_SpawnActorFromClass : public UK2Node
 	virtual bool HasExternalBlueprintDependencies(TArray<class UStruct*>* OptionalOutput) const override;
 	virtual FName GetPaletteIcon(FLinearColor& OutColor) const override{ return TEXT("GraphEditor.SpawnActor_16x"); }
 	virtual bool IsCompatibleWithGraph(const UEdGraph* TargetGraph) const override;
+	virtual void PostPlacedNewNode() override;
 	// End UEdGraphNode interface.
 
 	// Begin UK2Node interface
@@ -34,8 +35,13 @@ class UK2Node_SpawnActorFromClass : public UK2Node
 	// End UK2Node interface
 
 
-	/** Create new pins to show properties on archetype */
-	void CreatePinsForClass(UClass* InClass);
+	/**
+	 * Create new pins to show properties on archetype
+	 *
+	 * @param InClass		Class to make pins for
+	 * @param OutClassPins	Any pins created through this function will be added to thist list
+	 */
+	void CreatePinsForClass(UClass* InClass, TArray<UEdGraphPin*>& OutClassPins);
 	/** See if this is a spawn variable pin, or a 'default' pin */
 	BLUEPRINTGRAPH_API bool IsSpawnVarPin(UEdGraphPin* Pin);
 
diff --git a/Engine/Source/Editor/BlueprintGraph/Private/K2Node.cpp b/Engine/Source/Editor/BlueprintGraph/Private/K2Node.cpp
--- a/Engine/Source/Editor/BlueprintGraph/Private/K2Node.cpp
+++ b/Engine/Source/Editor/BlueprintGraph/Private/K2Node.cpp
@@ -366,47 +366,13 @@ void UK2Node::ReconstructNode()
 	}
 	else
 	{
-		// Rewire any connection to pins that are matched by name (O(N^2) right now)
-		//@TODO: Can do moderately smart things here if only one pin changes name by looking at it's relative position, etc...,
-		// rather than just failing to map it and breaking the links
-		for (int32 OldPinIndex = 0; OldPinIndex < OldPins.Num(); ++OldPinIndex)
-		{
-			UEdGraphPin* OldPin = OldPins[OldPinIndex];
-
-			for (int32 NewPinIndex = 0; NewPinIndex < Pins.Num(); ++NewPinIndex)
-			{
-				UEdGraphPin* NewPin = Pins[NewPinIndex];
-
-				const ERedirectType RedirectType = DoPinsMatchForReconstruction(NewPin, NewPinIndex, OldPin, OldPinIndex);
-				if (RedirectType != ERedirectType_None)
-				{
-					ReconstructSinglePin(NewPin, OldPin, RedirectType);
-					break;
-				}
-			}
-		}
+		RewireOldPinsToNewPins(OldPins, Pins);
 	}
 
 
 	if (bDestroyOldPins)
 	{
-		// Throw away the original pins
-		for (int32 OldPinIndex = 0; OldPinIndex < OldPins.Num(); ++OldPinIndex)
-		{
-			UEdGraphPin* OldPin = OldPins[OldPinIndex];
-			OldPin->Modify();
-			OldPin->BreakAllPinLinks();
-			
-			// just in case this pin was set to watch (don't want to save PinWatches with dead pins)
-			Blueprint->PinWatches.Remove(OldPin);
-#if 0
-			UEdGraphNode::ReturnPinToPool(OldPin);
-#else
-			OldPin->Rename(NULL, GetTransientPackage(), (Blueprint->bIsRegeneratingOnLoad ? REN_ForceNoResetLoaders : REN_None));
-			OldPin->RemoveFromRoot();
-			OldPin->MarkPendingKill();
-#endif
-		}
+		DestroyPinList(OldPins);
 	}
 
 	// Let subclasses do any additional work
@@ -624,6 +590,50 @@ void UK2Node::ReconstructSinglePin(UEdGraphPin* NewPin, UEdGraphPin* OldPin, ERe
 	OldPin->Rename(NULL, GetTransientPackage(), (REN_DontCreateRedirectors|(Blueprint->bIsRegeneratingOnLoad ? REN_ForceNoResetLoaders : REN_None)));
 }
 
+void UK2Node::RewireOldPinsToNewPins(TArray<UEdGraphPin*>& InOldPins, TArray<UEdGraphPin*>& InNewPins)
+{
+	// Rewire any connection to pins that are matched by name (O(N^2) right now)
+	//@TODO: Can do moderately smart things here if only one pin changes name by looking at it's relative position, etc...,
+	// rather than just failing to map it and breaking the links
+	for (int32 OldPinIndex = 0; OldPinIndex < InOldPins.Num(); ++OldPinIndex)
+	{
+		UEdGraphPin* OldPin = InOldPins[OldPinIndex];
+
+		for (int32 NewPinIndex = 0; NewPinIndex < InNewPins.Num(); ++NewPinIndex)
+		{
+			UEdGraphPin* NewPin = InNewPins[NewPinIndex];
+
+			const ERedirectType RedirectType = DoPinsMatchForReconstruction(NewPin, NewPinIndex, OldPin, OldPinIndex);
+			if (RedirectType != ERedirectType_None)
+			{
+				ReconstructSinglePin(NewPin, OldPin, RedirectType);
+				break;
+			}
+		}
+	}
+}
+
+void UK2Node::DestroyPinList(TArray<UEdGraphPin*>& InPins)
+{
+	UBlueprint* Blueprint = GetBlueprint();
+	// Throw away the original pins
+	for (UEdGraphPin* Pin : InPins)
+	{
+		Pin->Modify();
+		Pin->BreakAllPinLinks();
+
+		// just in case this pin was set to watch (don't want to save PinWatches with dead pins)
+		Blueprint->PinWatches.Remove(Pin);
+#if 0
+		UEdGraphNode::ReturnPinToPool(Pin);
+#else
+		Pin->Rename(NULL, GetTransientPackage(), (Blueprint->bIsRegeneratingOnLoad ? REN_ForceNoResetLoaders : REN_None));
+		Pin->RemoveFromRoot();
+		Pin->MarkPendingKill();
+#endif
+	}
+}
+
 bool UK2Node::AllowSplitPins() const
 {
 	return true;
diff --git a/Engine/Source/Editor/BlueprintGraph/Private/K2Node_SpawnActorFromClass.cpp b/Engine/Source/Editor/BlueprintGraph/Private/K2Node_SpawnActorFromClass.cpp
--- a/Engine/Source/Editor/BlueprintGraph/Private/K2Node_SpawnActorFromClass.cpp
+++ b/Engine/Source/Editor/BlueprintGraph/Private/K2Node_SpawnActorFromClass.cpp
@@ -62,7 +62,7 @@ void UK2Node_SpawnActorFromClass::AllocateDefaultPins()
 	Super::AllocateDefaultPins();
 }
 
-void UK2Node_SpawnActorFromClass::CreatePinsForClass(UClass* InClass)
+void UK2Node_SpawnActorFromClass::CreatePinsForClass(UClass* InClass, TArray<UEdGraphPin*>& OutClassPins)
 {
 	check(InClass != NULL);
 
@@ -87,7 +87,8 @@ void UK2Node_SpawnActorFromClass::CreatePinsForClass(UClass* InClass)
 		{
 			UEdGraphPin* Pin = CreatePin(EGPD_Input, TEXT(""), TEXT(""), NULL, false, false, Property->GetName());
 			const bool bPinGood = (Pin != NULL) && K2Schema->ConvertPropertyToPinType(Property, /*out*/ Pin->PinType);
-			
+			OutClassPins.Add(Pin);
+
 			if( ClassDefaultObject && K2Schema->PinDefaultValueIsEditable(*Pin))
 			{
 				FString DefaultValueAsString;
@@ -111,7 +112,7 @@ void UK2Node_SpawnActorFromClass::CreatePinsForClass(UClass* InClass)
 
 UClass* UK2Node_SpawnActorFromClass::GetClassToSpawn(const TArray<UEdGraphPin*>* InPinsToSearch /*=NULL*/) const
 {
-	UClass* UseSpawnClass = NULL;
+	UClass* UseSpawnClass = nullptr;
 	const TArray<UEdGraphPin*>* PinsToSearch = InPinsToSearch ? InPinsToSearch : &Pins;
 
 	UEdGraphPin* ClassPin = GetClassPin(PinsToSearch);
@@ -122,7 +123,7 @@ UClass* UK2Node_SpawnActorFromClass::GetClassToSpawn(const TArray<UEdGraphPin*>*
 	else if (ClassPin && (1 == ClassPin->LinkedTo.Num()))
 	{
 		auto SourcePin = ClassPin->LinkedTo[0];
-		UseSpawnClass = SourcePin ? Cast<UClass>(SourcePin->PinType.PinSubCategoryObject.Get()) : NULL;
+		UseSpawnClass = SourcePin ? Cast<UClass>(SourcePin->PinType.PinSubCategoryObject.Get()) : nullptr;
 	}
 
 	return UseSpawnClass;
@@ -135,9 +136,23 @@ void UK2Node_SpawnActorFromClass::ReallocatePinsDuringReconstruction(TArray<UEdG
 	UClass* UseSpawnClass = GetClassToSpawn(&OldPins);
 	if( UseSpawnClass != NULL )
 	{
-		CreatePinsForClass(UseSpawnClass);
+		TArray<UEdGraphPin*> ClassPins;
+		CreatePinsForClass(UseSpawnClass, ClassPins);
+	}
+}
+
+void UK2Node_SpawnActorFromClass::PostPlacedNewNode()
+{
+	Super::PostPlacedNewNode();
+
+	UClass* UseSpawnClass = GetClassToSpawn();
+	if( UseSpawnClass != NULL )
+	{
+		TArray<UEdGraphPin*> ClassPins;
+		CreatePinsForClass(UseSpawnClass, ClassPins);
 	}
 }
+
 bool UK2Node_SpawnActorFromClass::IsSpawnVarPin(UEdGraphPin* Pin)
 {
 	const UEdGraphSchema_K2* K2Schema = GetDefault<UEdGraphSchema_K2>();
@@ -163,7 +178,7 @@ bool UK2Node_SpawnActorFromClass::IsSpawnVarPin(UEdGraphPin* Pin)
 
 void UK2Node_SpawnActorFromClass::OnClassPinChanged()
 {
-	const UEdGraphSchema_K2* K2Schema = GetDefault<UEdGraphSchema_K2>();
+ 	const UEdGraphSchema_K2* K2Schema = GetDefault<UEdGraphSchema_K2>();
 
 	// Because the archetype has changed, we break the output link as the output pin type will change
 	UEdGraphPin* ResultPin = GetResultPin();
@@ -171,25 +186,34 @@ void UK2Node_SpawnActorFromClass::OnClassPinChanged()
 
 	// Remove all pins related to archetype variables
 	TArray<UEdGraphPin*> OldPins = Pins;
+	TArray<UEdGraphPin*> OldClassPins;
+
 	for (int32 i = 0; i < OldPins.Num(); i++)
 	{
 		UEdGraphPin* OldPin = OldPins[i];
 		if (IsSpawnVarPin(OldPin))
 		{
-			OldPin->BreakAllPinLinks();
 			Pins.Remove(OldPin);
+			OldClassPins.Add(OldPin);
 		}
 	}
 
 	CachedNodeTitle.MarkDirty();
 
 	UClass* UseSpawnClass = GetClassToSpawn();
+	TArray<UEdGraphPin*> NewClassPins;
 	if (UseSpawnClass != NULL)
 	{
-		CreatePinsForClass(UseSpawnClass);
+		CreatePinsForClass(UseSpawnClass, NewClassPins);
 	}
 	K2Schema->ConstructBasicPinTooltip(*ResultPin, LOCTEXT("ResultPinDescription", "The spawned Actor"), ResultPin->PinToolTip);
 
+	// Rewire the old pins to the new pins so connections are maintained if possible
+	RewireOldPinsToNewPins(OldClassPins, NewClassPins);
+
+	// Destroy the old pins
+	DestroyPinList(OldClassPins);
+
 	// Refresh the UI for the graph so the pin changes show up
 	UEdGraph* Graph = GetGraph();
 	Graph->NotifyGraphChanged();
diff --git a/Engine/Source/Editor/GraphEditor/Private/KismetNodes/SGraphNodeSpawnActorFromClass.cpp b/Engine/Source/Editor/GraphEditor/Private/KismetNodes/SGraphNodeSpawnActorFromClass.cpp
--- a/Engine/Source/Editor/GraphEditor/Private/KismetNodes/SGraphNodeSpawnActorFromClass.cpp
+++ b/Engine/Source/Editor/GraphEditor/Private/KismetNodes/SGraphNodeSpawnActorFromClass.cpp
@@ -25,10 +25,8 @@ class SGraphPinActorBasedClass : public SGraphPinObject
 	{
 		AssetPickerAnchor->SetIsOpen(false);
 
-		if(InChosenClass && GraphPinObj)
+		if(GraphPinObj)
 		{
-			check(InChosenClass);
-			check(InChosenClass->IsChildOf(AActor::StaticClass()));
 			if(const UEdGraphSchema* Schema = GraphPinObj->GetSchema())
 			{
 				Schema->TrySetDefaultObject(*GraphPinObj, InChosenClass);
