[UE-6023] No longer respecting the "disable edit on instance" flag for the BP context menu (was causing too much confussion, throwing up to many hurdles for users).

diff --git a/Engine/Source/Editor/BlueprintGraph/Private/BlueprintActionFilter.cpp b/Engine/Source/Editor/BlueprintGraph/Private/BlueprintActionFilter.cpp
--- a/Engine/Source/Editor/BlueprintGraph/Private/BlueprintActionFilter.cpp
+++ b/Engine/Source/Editor/BlueprintGraph/Private/BlueprintActionFilter.cpp
@@ -511,49 +511,53 @@ static bool BlueprintActionFilterImpl::IsFieldInaccessible(FBlueprintActionFilte
 
 	if (bIsMemberAction)
 	{
-		bool bIsProtected = Field->HasMetaData(FBlueprintMetadata::MD_Protected);
-		bool bIsPrivate = Field->HasMetaData(FBlueprintMetadata::MD_Private);
-
-		bool bIsPublic = !bIsPrivate && !bIsProtected;
-		if (UProperty const* Property = Cast<UProperty>(Field))
-		{
-			// default to assuming that this property is a native one (the
-			// accessibility there is a little more lax)
-			bool bIsNativeProperty = true;
-			if (UClass* PropertyClass = Property->GetOwnerClass())
-			{
-				bIsNativeProperty = (Cast<UBlueprintGeneratedClass>(PropertyClass) == nullptr);
-			}
-
-			// if this is a native property, then the user only should have to 
-			// specify BlueprintReadWrite or BlueprintReadOnly for it to be 
-			// accessible
-			if (bIsNativeProperty)
-			{
-				bIsPublic &= Property->HasAnyPropertyFlags(CPF_BlueprintVisible);
-			}
-			else
-			{
-				// however, if this is a blueprint variable, then the
-				// CPF_DisableEditOnInstance corresponds to the eyeball/editable 
-				// checkbox in the editor, and we don't want the variable to be 
-				// public until the user exposes it (delegates don't have an
-				// 'editable' toggle, so they should default to public)
-				bIsPublic &= !Property->HasAnyPropertyFlags(CPF_DisableEditOnInstance) ||
-					(Cast<UMulticastDelegateProperty>(Property) != nullptr);
-
-				// @TODO: allow users to choose "read-only" for blueprint 
-				//        variables ("private" might have provided this in the 
-				//        past, but that is just wrong, or needs to be renamed)
-			}
-
-			// If the variable is not public, and is not tagged as private, it should be protected
-			// by default. This branch handles variables declared in blueprints:
-			if (!bIsNativeProperty && !bIsPublic && !bIsPrivate)
-			{
-				bIsProtected = true;
-			}
-		}
+		bool const bIsProtected = Field->HasMetaData(FBlueprintMetadata::MD_Protected);
+		bool const bIsPrivate   = Field->HasMetaData(FBlueprintMetadata::MD_Private);
+		bool const bIsPublic    = !bIsPrivate && !bIsProtected;
+
+		// @TODO: Trying to respect the "editable"/DisableEditOnInstance toggle
+		//        was a bad idea that lead to confusion amongst users (also this 
+		//        created a discrepancy between native and blueprint variables),
+		//        until we make this concept more understandable: hold off
+// 		if (UProperty const* Property = Cast<UProperty>(Field))
+// 		{
+// 			// default to assuming that this property is a native one (the
+// 			// accessibility there is a little more lax)
+// 			bool bIsNativeProperty = true;
+// 			if (UClass* PropertyClass = Property->GetOwnerClass())
+// 			{
+// 				bIsNativeProperty = (Cast<UBlueprintGeneratedClass>(PropertyClass) == nullptr);
+// 			}
+// 
+// 			// if this is a native property, then the user only should have to 
+// 			// specify BlueprintReadWrite or BlueprintReadOnly for it to be 
+// 			// accessible
+// 			if (bIsNativeProperty)
+// 			{
+// 				bIsPublic &= Property->HasAnyPropertyFlags(CPF_BlueprintVisible);
+// 			}
+// 			else
+// 			{
+// 				// however, if this is a blueprint variable, then the
+// 				// CPF_DisableEditOnInstance corresponds to the eyeball/editable 
+// 				// checkbox in the editor, and we don't want the variable to be 
+// 				// public until the user exposes it (delegates don't have an
+// 				// 'editable' toggle, so they should default to public)
+// 				bIsPublic &= !Property->HasAnyPropertyFlags(CPF_DisableEditOnInstance) ||
+// 					(Cast<UMulticastDelegateProperty>(Property) != nullptr);
+// 
+// 				// @TODO: allow users to choose "read-only" for blueprint 
+// 				//        variables ("private" might have provided this in the 
+// 				//        past, but that is just wrong, or needs to be renamed)
+// 			}
+// 
+// 			// If the variable is not public, and is not tagged as private, it should be protected
+// 			// by default. This branch handles variables declared in blueprints:
+// 			if (!bIsNativeProperty && !bIsPublic && !bIsPrivate)
+// 			{
+// 				bIsProtected = true;
+// 			}
+// 		}
 
 		if( !bIsPublic )
 		{
