Fixed faulty assertion when adding a text array to a user defined structure.

diff --git a/Engine/Source/Editor/BlueprintGraph/Private/K2Node_MakeStruct.cpp b/Engine/Source/Editor/BlueprintGraph/Private/K2Node_MakeStruct.cpp
--- a/Engine/Source/Editor/BlueprintGraph/Private/K2Node_MakeStruct.cpp
+++ b/Engine/Source/Editor/BlueprintGraph/Private/K2Node_MakeStruct.cpp
@@ -1,4 +1,4 @@
-// Copyright 1998-2015 Epic Games, Inc. All Rights Reserved.
+﻿// Copyright 1998-2015 Epic Games, Inc. All Rights Reserved.
 
 #include "BlueprintGraphPrivatePCH.h"
 #include "KismetCompiler.h"
@@ -42,7 +42,8 @@ void UK2Node_MakeStruct::FMakeStructPinManager::CustomizePinData(UEdGraphPin* Pi
 	{
 		const UEdGraphSchema_K2* Schema = GetDefault<UEdGraphSchema_K2>();
 		check(Schema);
-		const bool bIsText = Property->IsA<UTextProperty>();
+		UArrayProperty* const ArrayProperty = Cast<UArrayProperty>(Property);
+		const bool bIsText = Property->IsA<UTextProperty>() || (ArrayProperty && ArrayProperty->Inner && ArrayProperty->Inner->IsA<UTextProperty>());
 		checkSlow(bIsText == (Schema->PC_Text == Pin->PinType.PinCategory));
 
 		const FString& MetadataDefaultValue = Property->GetMetaData(TEXT("MakeStructureDefaultValue"));
