#github PR682, fix for crash caused by unresolveable ClassGeneratedBy (probably only for BlueprintGeneratedClasses) from slonopoatmus

just null checking for now, looking for good repro

diff --git a/Engine/Source/Editor/BlueprintGraph/Private/K2Node_CallFunction.cpp b/Engine/Source/Editor/BlueprintGraph/Private/K2Node_CallFunction.cpp
--- a/Engine/Source/Editor/BlueprintGraph/Private/K2Node_CallFunction.cpp
+++ b/Engine/Source/Editor/BlueprintGraph/Private/K2Node_CallFunction.cpp
@@ -509,10 +509,13 @@ void UK2Node_CallFunction::AllocateDefaultPins()
 			// class (the blueprint has not be compiled with it yet), so let's 
 			// check the skeleton class as well, see if we can pull pin data 
 			// from there...
-			UBlueprint* FunctionBlueprint = CastChecked<UBlueprint>(BpClassOwner->ClassGeneratedBy);
-			if (UFunction* SkelFunction = FindField<UFunction>(FunctionBlueprint->SkeletonGeneratedClass, FunctionReference.GetMemberName()))
+			UBlueprint* FunctionBlueprint = CastChecked<UBlueprint>(BpClassOwner->ClassGeneratedBy, ECastCheckedType::NullAllowed);
+			if (FunctionBlueprint)
 			{
-				Function = SkelFunction;
+				if (UFunction* SkelFunction = FindField<UFunction>(FunctionBlueprint->SkeletonGeneratedClass, FunctionReference.GetMemberName()))
+				{
+					Function = SkelFunction;
+				}
 			}
 		}
 	}
@@ -622,7 +625,10 @@ UEdGraphPin* UK2Node_CallFunction::CreateSelfPin(const UFunction* Function)
 	// that represent the same type)... this here could lead to a compiler 
 	// warning (the GeneratedClass could not have the function yet), but in
 	// that, the user would be reminded to compile the other blueprint
-	FunctionClass = FunctionClass->GetAuthoritativeClass();
+	if (FunctionClass->ClassGeneratedBy)
+	{
+		FunctionClass = FunctionClass->GetAuthoritativeClass();
+	}
 
 	UEdGraphPin* SelfPin = NULL;
 	if (FunctionClass == GetBlueprint()->GeneratedClass)
diff --git a/Engine/Source/Editor/BlueprintGraph/Private/K2Node_Variable.cpp b/Engine/Source/Editor/BlueprintGraph/Private/K2Node_Variable.cpp
--- a/Engine/Source/Editor/BlueprintGraph/Private/K2Node_Variable.cpp
+++ b/Engine/Source/Editor/BlueprintGraph/Private/K2Node_Variable.cpp
@@ -61,10 +61,13 @@ bool UK2Node_Variable::CreatePinForVariable(EEdGraphPinDirection Direction, FStr
 			// class (the blueprint has not be compiled with it yet), so let's 
 			// check the skeleton class as well, see if we can pull pin data 
 			// from there...
-			UBlueprint* VariableBlueprint = CastChecked<UBlueprint>(BpClassOwner->ClassGeneratedBy);
-			if (UProperty* SkelProperty = FindField<UProperty>(VariableBlueprint->SkeletonGeneratedClass, VariableReference.GetMemberName()))
+			UBlueprint* VariableBlueprint = CastChecked<UBlueprint>(BpClassOwner->ClassGeneratedBy, ECastCheckedType::NullAllowed);
+			if (VariableBlueprint)
 			{
-				VariableProperty = SkelProperty;
+				if (UProperty* SkelProperty = FindField<UProperty>(VariableBlueprint->SkeletonGeneratedClass, VariableReference.GetMemberName()))
+				{
+					VariableProperty = SkelProperty;
+				}
 			}
 		}
 	}
@@ -100,7 +103,7 @@ void UK2Node_Variable::CreatePinForSelf()
 		{
 			bool bSelfTarget = VariableReference.IsSelfContext() && (ESelfContextInfo::NotSelfContext != SelfContextInfo);
 			UClass* MemberParentClass = VariableReference.GetMemberParentClass(this);
-			UClass* TargetClass = nullptr;
+			UClass* TargetClass = MemberParentClass;
 			
 			// Self Target pins should always make the class be the owning class of the property,
 			// so if the node is from a Macro Blueprint, it will hook up as self in any placed Blueprint
@@ -115,9 +118,9 @@ void UK2Node_Variable::CreatePinForSelf()
 					TargetClass = GetBlueprint()->SkeletonGeneratedClass->GetAuthoritativeClass();
 				}
 			}
-			else
+			else if(MemberParentClass && MemberParentClass->ClassGeneratedBy)
 			{
-				TargetClass = MemberParentClass ? MemberParentClass->GetAuthoritativeClass() : NULL;
+				TargetClass = MemberParentClass->GetAuthoritativeClass();
 			}
 
 			UEdGraphPin* TargetPin = CreatePin(EGPD_Input, K2Schema->PC_Object, TEXT(""), TargetClass, false, false, K2Schema->PN_Self);
