BP, make struct node: Fixed faulty assertion when adding a text array to a user defined structure.

diff --git a/Engine/Source/Editor/BlueprintGraph/Private/K2Node_MakeStruct.cpp b/Engine/Source/Editor/BlueprintGraph/Private/K2Node_MakeStruct.cpp
--- a/Engine/Source/Editor/BlueprintGraph/Private/K2Node_MakeStruct.cpp
+++ b/Engine/Source/Editor/BlueprintGraph/Private/K2Node_MakeStruct.cpp
@@ -42,9 +42,10 @@ void UK2Node_MakeStruct::FMakeStructPinManager::CustomizePinData(UEdGraphPin* Pi
 	{
 		const UEdGraphSchema_K2* Schema = GetDefault<UEdGraphSchema_K2>();
 		check(Schema);
-		UArrayProperty* const ArrayProperty = Cast<UArrayProperty>(Property);
-		const bool bIsText = Property->IsA<UTextProperty>() || (ArrayProperty && ArrayProperty->Inner && ArrayProperty->Inner->IsA<UTextProperty>());
-		checkSlow(bIsText == (Schema->PC_Text == Pin->PinType.PinCategory));
+
+		// Should pin default value be filled as FText?
+		const bool bIsText = Property->IsA<UTextProperty>();
+		checkSlow(bIsText == ((Schema->PC_Text == Pin->PinType.PinCategory) && !Pin->PinType.bIsArray));
 
 		const FString& MetadataDefaultValue = Property->GetMetaData(TEXT("MakeStructureDefaultValue"));
 		if (!MetadataDefaultValue.IsEmpty())
