Pasting local variable nodes in graphs outside of scope will no longer show them as being valid.

Pasting local or variable nodes into graphs with local variables of the same name will have them reference the local variable automatically.

diff --git a/Engine/Source/Editor/BlueprintGraph/Classes/K2Node.h b/Engine/Source/Editor/BlueprintGraph/Classes/K2Node.h
--- a/Engine/Source/Editor/BlueprintGraph/Classes/K2Node.h
+++ b/Engine/Source/Editor/BlueprintGraph/Classes/K2Node.h
@@ -532,8 +532,8 @@ public:
 	/** Only intended for backwards compat! */
 	BLUEPRINTGRAPH_API void SetDirect(const FName InMemberName, const FGuid InMemberGuid, TSubclassOf<class UObject> InMemberParentClass, bool bIsConsideredSelfContext);
 
-	/** Invalidate the current MemberParentClass, if this is a self scope.  Intended for PostDuplication fixups */
-	BLUEPRINTGRAPH_API void InvalidateSelfScope();
+	/** Invalidate the current MemberParentClass, if this is a self scope, or the MemberScope if it is not (and set).  Intended for PostDuplication fixups */
+	BLUEPRINTGRAPH_API void InvalidateScope();
 
 	/** Get the name of this member */
 	FName GetMemberName() const
diff --git a/Engine/Source/Editor/BlueprintGraph/Classes/K2Node_Variable.h b/Engine/Source/Editor/BlueprintGraph/Classes/K2Node_Variable.h
--- a/Engine/Source/Editor/BlueprintGraph/Classes/K2Node_Variable.h
+++ b/Engine/Source/Editor/BlueprintGraph/Classes/K2Node_Variable.h
@@ -57,6 +57,7 @@ public:
 	virtual FName GetPaletteIcon(FLinearColor& OutColor) const override;
 	virtual void AutowireNewNode(UEdGraphPin* FromPin) override;
 	virtual bool CanPasteHere(const UEdGraph* TargetGraph) const override;
+	virtual void PostPasteNode() override;
 	// End UEdGraphNode interface
 
 	// Begin K2Node interface
diff --git a/Engine/Source/Editor/BlueprintGraph/Private/K2Node.cpp b/Engine/Source/Editor/BlueprintGraph/Private/K2Node.cpp
--- a/Engine/Source/Editor/BlueprintGraph/Private/K2Node.cpp
+++ b/Engine/Source/Editor/BlueprintGraph/Private/K2Node.cpp
@@ -1342,12 +1342,19 @@ void FMemberReference::SetLocalMember(FName InMemberName, FString InScopeName, c
 	bSelfContext = false;
 }
 
-void FMemberReference::InvalidateSelfScope()
+void FMemberReference::InvalidateScope()
 {
 	if( IsSelfContext() )
 	{
 		MemberParentClass = NULL;
 	}
+	else if(IsLocalScope())
+	{
+		MemberScope.Empty();
+
+		// Make it into a member reference since we are clearing the local context
+		bSelfContext = true;
+	}
 }
 
 UClass* FMemberReference::GetBlueprintClassFromNode(const UK2Node* Node)
diff --git a/Engine/Source/Editor/BlueprintGraph/Private/K2Node_CallFunction.cpp b/Engine/Source/Editor/BlueprintGraph/Private/K2Node_CallFunction.cpp
--- a/Engine/Source/Editor/BlueprintGraph/Private/K2Node_CallFunction.cpp
+++ b/Engine/Source/Editor/BlueprintGraph/Private/K2Node_CallFunction.cpp
@@ -1567,7 +1567,7 @@ void UK2Node_CallFunction::PostDuplicate(bool bDuplicateForPIE)
 	Super::PostDuplicate(bDuplicateForPIE);
 	if (!bDuplicateForPIE && (!this->HasAnyFlags(RF_Transient)))
 	{
-		FunctionReference.InvalidateSelfScope();
+		FunctionReference.InvalidateScope();
 		EnsureFunctionIsInBlueprint();
 	}
 }
diff --git a/Engine/Source/Editor/BlueprintGraph/Private/K2Node_Variable.cpp b/Engine/Source/Editor/BlueprintGraph/Private/K2Node_Variable.cpp
--- a/Engine/Source/Editor/BlueprintGraph/Private/K2Node_Variable.cpp
+++ b/Engine/Source/Editor/BlueprintGraph/Private/K2Node_Variable.cpp
@@ -651,4 +651,44 @@ bool UK2Node_Variable::CanPasteHere(const UEdGraph* TargetGraph) const
 	return true;
 }
 
+void UK2Node_Variable::PostPasteNode()
+{
+	Super::PostPasteNode();
+
+	UBlueprint* Blueprint = GetBlueprint();
+	bool bInvalidateVariable = false;
+
+	if (VariableReference.ResolveMember<UProperty>(Blueprint) == nullptr)
+	{
+		bInvalidateVariable = true;
+	}
+	else if (VariableReference.IsLocalScope())
+	{
+		// Local scoped variables should always validate whether they are being placed in the same graph as their scope
+		// ResolveMember will not return nullptr when the graph changes but the Blueprint remains the same.
+		UEdGraph* ScopeGraph = FBlueprintEditorUtils::FindScopeGraph(Blueprint, VariableReference.GetMemberScope(this));
+		if(ScopeGraph != GetGraph())
+		{
+			bInvalidateVariable = true;
+		}
+	}
+		
+	if (bInvalidateVariable)
+	{
+		// This invalidates the local scope
+		VariableReference.InvalidateScope();
+
+		// If the current graph is a Function graph, look to see if there is a compatible local variable (same name)
+		if (GetGraph()->GetSchema()->GetGraphType(GetGraph()) == GT_Function)
+		{
+			FBPVariableDescription* VariableDescription = FBlueprintEditorUtils::FindLocalVariable(Blueprint, GetGraph(), VariableReference.GetMemberName());
+			if(VariableDescription)
+			{
+				VariableReference.SetLocalMember(VariableReference.GetMemberName(), GetGraph()->GetName(), VariableReference.GetMemberGuid());
+			}
+		}
+		// If no variable was found, ResolveMember should automatically find a member variable with the same name in the current Blueprint and hook up to it as expected
+	}
+}
+
 #undef LOCTEXT_NAMESPACE
