UE-6784 Crash occurs when loading into the Couch Knights Room map.

(the fix should not be necessary once the circular dependency problems are solved).

diff --git a/Engine/Source/Editor/BlueprintGraph/Private/EdGraphSchema_K2.cpp b/Engine/Source/Editor/BlueprintGraph/Private/EdGraphSchema_K2.cpp
--- a/Engine/Source/Editor/BlueprintGraph/Private/EdGraphSchema_K2.cpp
+++ b/Engine/Source/Editor/BlueprintGraph/Private/EdGraphSchema_K2.cpp
@@ -3427,7 +3427,16 @@ bool UEdGraphSchema_K2::ArePinTypesCompatible(const FEdGraphPinType& Output, con
 		}
 		else if (PC_Delegate == Output.PinCategory || PC_MCDelegate == Output.PinCategory)
 		{
+			auto CanUseFunction = [](const UFunction* Func) -> bool
+			{
+				return Func && (Func->HasAllFlags(RF_LoadCompleted) || !Func->HasAnyFlags(RF_NeedLoad | RF_WasLoaded));
+			};
+
 			const UFunction* OutFunction = FMemberReference::ResolveSimpleMemberReference<UFunction>(Output.PinSubCategoryMemberReference);
+			if (!CanUseFunction(OutFunction))
+			{
+				OutFunction = NULL;
+			}
 			if (!OutFunction && Output.PinSubCategoryMemberReference.MemberParentClass)
 			{
 				const auto ParentClass = Output.PinSubCategoryMemberReference.MemberParentClass;
@@ -3438,6 +3447,10 @@ bool UEdGraphSchema_K2::ArePinTypesCompatible(const FEdGraphPinType& Output, con
 				}
 			}
 			const UFunction* InFunction = FMemberReference::ResolveSimpleMemberReference<UFunction>(Input.PinSubCategoryMemberReference);
+			if (!CanUseFunction(InFunction))
+			{
+				InFunction = NULL;
+			}
 			if (!InFunction && Input.PinSubCategoryMemberReference.MemberParentClass)
 			{
 				const auto ParentClass = Input.PinSubCategoryMemberReference.MemberParentClass;
