[UE-7169] Fixing it so actor-bound event nodes reference the full generated class (instead of the skel class).

diff --git a/Engine/Source/Editor/BlueprintGraph/Private/K2Node_ActorBoundEvent.cpp b/Engine/Source/Editor/BlueprintGraph/Private/K2Node_ActorBoundEvent.cpp
--- a/Engine/Source/Editor/BlueprintGraph/Private/K2Node_ActorBoundEvent.cpp
+++ b/Engine/Source/Editor/BlueprintGraph/Private/K2Node_ActorBoundEvent.cpp
@@ -67,8 +67,10 @@ void UK2Node_ActorBoundEvent::DestroyNode()
 				Delegate.BindUFunction(LSA, CustomFunctionName);
 				
 				// Attempt to remove it from the target's MC delegate
-				FMulticastScriptDelegate* TargetDelegate = GetTargetDelegate();
-				TargetDelegate->Remove(Delegate);
+				if (FMulticastScriptDelegate* TargetDelegate = GetTargetDelegate())
+				{
+					TargetDelegate->Remove(Delegate);
+				}				
 			}
 		}
 		
@@ -138,7 +140,7 @@ void UK2Node_ActorBoundEvent::InitializeActorBoundEventParams(AActor* InEventOwn
 		EventOwner = InEventOwner;
 		DelegatePropertyName = InDelegateProperty->GetFName();
 		EventSignatureName = InDelegateProperty->SignatureFunction->GetFName();
-		EventSignatureClass = CastChecked<UClass>(InDelegateProperty->SignatureFunction->GetOuter());
+		EventSignatureClass = CastChecked<UClass>(InDelegateProperty->SignatureFunction->GetOuter())->GetAuthoritativeClass();
 		CustomFunctionName = FName( *FString::Printf(TEXT("BndEvt__%s_%s_%s"), *InEventOwner->GetName(), *GetName(), *EventSignatureName.ToString()) );
 		bOverrideFunction = false;
 		bInternalEvent = true;
