UK2Node_SpawnActorFromClass can now be inherited from across modules

Integration of GitHub pull request #723 from user iniside

diff --git a/Engine/Source/Editor/BlueprintGraph/Classes/K2Node_SpawnActorFromClass.h b/Engine/Source/Editor/BlueprintGraph/Classes/K2Node_SpawnActorFromClass.h
--- a/Engine/Source/Editor/BlueprintGraph/Classes/K2Node_SpawnActorFromClass.h
+++ b/Engine/Source/Editor/BlueprintGraph/Classes/K2Node_SpawnActorFromClass.h
@@ -6,8 +6,8 @@
 #include "EdGraph/EdGraphNodeUtils.h" // for FNodeTextCache
 #include "K2Node_SpawnActorFromClass.generated.h"
 
-UCLASS(MinimalAPI)
-class UK2Node_SpawnActorFromClass : public UK2Node
+UCLASS()
+class BLUEPRINTGRAPH_API UK2Node_SpawnActorFromClass : public UK2Node
 {
 	GENERATED_UCLASS_BODY()
 
@@ -43,25 +43,25 @@ class UK2Node_SpawnActorFromClass : public UK2Node
 	 */
 	void CreatePinsForClass(UClass* InClass, TArray<UEdGraphPin*>& OutClassPins);
 	/** See if this is a spawn variable pin, or a 'default' pin */
-	BLUEPRINTGRAPH_API bool IsSpawnVarPin(UEdGraphPin* Pin);
+	bool IsSpawnVarPin(UEdGraphPin* Pin);
 
 	/** Get the then output pin */
-	BLUEPRINTGRAPH_API UEdGraphPin* GetThenPin() const;
+	UEdGraphPin* GetThenPin() const;
 	/** Get the blueprint input pin */	
-	BLUEPRINTGRAPH_API UEdGraphPin* GetClassPin(const TArray<UEdGraphPin*>* InPinsToSearch=NULL) const;
+	UEdGraphPin* GetClassPin(const TArray<UEdGraphPin*>* InPinsToSearch=NULL) const;
 	/** Get the world context input pin, can return NULL */	
-	BLUEPRINTGRAPH_API UEdGraphPin* GetWorldContextPin() const;
+	UEdGraphPin* GetWorldContextPin() const;
 	/** Get the spawn transform input pin */	
-	BLUEPRINTGRAPH_API UEdGraphPin* GetSpawnTransformPin() const;
+	UEdGraphPin* GetSpawnTransformPin() const;
 	/** Get the spawn NoCollisionFail input pin */
-	BLUEPRINTGRAPH_API UEdGraphPin* GetNoCollisionFailPin() const;
+	UEdGraphPin* GetNoCollisionFailPin() const;
 	/** Get the result output pin */
-	BLUEPRINTGRAPH_API UEdGraphPin* GetResultPin() const;
+	UEdGraphPin* GetResultPin() const;
 
 	/** Get the class that we are going to spawn, if it's defined as default value */
-	BLUEPRINTGRAPH_API UClass* GetClassToSpawn(const TArray<UEdGraphPin*>* InPinsToSearch=NULL) const;
+	UClass* GetClassToSpawn(const TArray<UEdGraphPin*>* InPinsToSearch=NULL) const;
 
-private:
+protected:
 
 	void OnClassPinChanged();
 
