Allow references to Actors in anonymous blueprints

diff --git a/Engine/Source/Editor/BlueprintGraph/Classes/EdGraphSchema_K2.h b/Engine/Source/Editor/BlueprintGraph/Classes/EdGraphSchema_K2.h
--- a/Engine/Source/Editor/BlueprintGraph/Classes/EdGraphSchema_K2.h
+++ b/Engine/Source/Editor/BlueprintGraph/Classes/EdGraphSchema_K2.h
@@ -636,7 +636,7 @@ public:
 	UFunction* GetCallableParentFunction(UFunction* Function) const;
 
 	/** Whether or not the specified actor is a valid target for bound events and literal references (in the right level, not a builder brush, etc */
-	bool IsActorValidForLevelScriptRefs(const AActor* TestActor, const ULevelScriptBlueprint* Blueprint) const;
+	bool IsActorValidForLevelScriptRefs(const AActor* TestActor, const UBlueprint* Blueprint) const;
 
 	/** 
 	 *	Generate a list of replaceable nodes for context menu based on the editor's current selection 
diff --git a/Engine/Source/Editor/BlueprintGraph/Private/EdGraphSchema_K2.cpp b/Engine/Source/Editor/BlueprintGraph/Private/EdGraphSchema_K2.cpp
--- a/Engine/Source/Editor/BlueprintGraph/Private/EdGraphSchema_K2.cpp
+++ b/Engine/Source/Editor/BlueprintGraph/Private/EdGraphSchema_K2.cpp
@@ -875,12 +875,12 @@ void UEdGraphSchema_K2::ListFunctionsMatchingSignatureAsDelegates(FGraphContextM
 
 }
 
-bool UEdGraphSchema_K2::IsActorValidForLevelScriptRefs(const AActor* TestActor, const ULevelScriptBlueprint* Blueprint) const
+bool UEdGraphSchema_K2::IsActorValidForLevelScriptRefs(const AActor* TestActor, const UBlueprint* Blueprint) const
 {
 	check(Blueprint);
 	
 	return TestActor
-		&& (TestActor->GetLevel() == Blueprint->GetLevel())
+		&& (TestActor->GetLevel() == FBlueprintEditorUtils::GetLevelFromBlueprint(Blueprint))
 		&& FKismetEditorUtilities::IsActorValidForLevelScript(TestActor);
 }
 
@@ -905,30 +905,21 @@ void UEdGraphSchema_K2::ReplaceSelectedNode(UEdGraphNode* SourceNode, AActor* Ta
 
 void UEdGraphSchema_K2::AddSelectedReplaceableNodes( UBlueprint* Blueprint, const UEdGraphNode* InGraphNode, FMenuBuilder* MenuBuilder ) const
 {
-	ULevelScriptBlueprint* LevelBlueprint = Cast<ULevelScriptBlueprint>(Blueprint);
-	
-	if (LevelBlueprint)
+	//Only allow replace object reference functionality for literal nodes
+	const UK2Node_Literal* LiteralNode = Cast<UK2Node_Literal>(InGraphNode);
+	if (LiteralNode)
 	{
-		//Only allow replace object reference functionality for literal nodes
-		if( InGraphNode->IsA( UK2Node_Literal::StaticClass() ) )
+		USelection* SelectedActors = GEditor->GetSelectedActors();
+		for(FSelectionIterator Iter(*SelectedActors); Iter; ++Iter)
 		{
-			UK2Node_Literal* LiteralNode = (UK2Node_Literal*)(InGraphNode);
-
-			if( LiteralNode )
+			// We only care about actors that are referenced in the world for literals, and also in the same level as this blueprint
+			AActor* Actor = Cast<AActor>(*Iter);
+			if( LiteralNode->GetObjectRef() != Actor && IsActorValidForLevelScriptRefs(Actor, Blueprint) )
 			{
-				USelection* SelectedActors = GEditor->GetSelectedActors();
-				for(FSelectionIterator Iter(*SelectedActors); Iter; ++Iter)
-				{
-					// We only care about actors that are referenced in the world for literals, and also in the same level as this blueprint
-					AActor* Actor = Cast<AActor>(*Iter);
-					if( LiteralNode->GetObjectRef() != Actor && IsActorValidForLevelScriptRefs(Actor, LevelBlueprint) )
-					{
-						FText Description = FText::Format( LOCTEXT("ChangeToActorName", "Change to <{0}>"), FText::FromString( Actor->GetActorLabel() ) );
-						FText ToolTip = LOCTEXT("ReplaceNodeReferenceToolTip", "Replace node reference");
-						MenuBuilder->AddMenuEntry( Description, ToolTip, FSlateIcon(), FUIAction(
-							FExecuteAction::CreateUObject((UEdGraphSchema_K2*const)this, &UEdGraphSchema_K2::ReplaceSelectedNode, const_cast< UEdGraphNode* >(InGraphNode), Actor) ) );
-					}
-				}
+				FText Description = FText::Format( LOCTEXT("ChangeToActorName", "Change to <{0}>"), FText::FromString( Actor->GetActorLabel() ) );
+				FText ToolTip = LOCTEXT("ReplaceNodeReferenceToolTip", "Replace node reference");
+				MenuBuilder->AddMenuEntry( Description, ToolTip, FSlateIcon(), FUIAction(
+					FExecuteAction::CreateUObject((UEdGraphSchema_K2*const)this, &UEdGraphSchema_K2::ReplaceSelectedNode, const_cast< UEdGraphNode* >(InGraphNode), Actor) ) );
 			}
 		}
 	}
diff --git a/Engine/Source/Editor/BlueprintGraph/Private/K2ActionMenuBuilder.cpp b/Engine/Source/Editor/BlueprintGraph/Private/K2ActionMenuBuilder.cpp
--- a/Engine/Source/Editor/BlueprintGraph/Private/K2ActionMenuBuilder.cpp
+++ b/Engine/Source/Editor/BlueprintGraph/Private/K2ActionMenuBuilder.cpp
@@ -2141,16 +2141,13 @@ void FK2ActionMenuBuilder::GetEventsForBlueprint(FBlueprintPaletteListBuilder& A
 //------------------------------------------------------------------------------
 void FK2ActionMenuBuilder::GetLiteralsFromActorSelection(FBlueprintGraphActionListBuilder& ContextMenuBuilder) const
 {
-	const ULevelScriptBlueprint* LevelBlueprint = CastChecked<ULevelScriptBlueprint>(ContextMenuBuilder.Blueprint);
-	check(LevelBlueprint);
-
 	USelection* SelectedActors = GEditor->GetSelectedActors();
 
 	TArray<AActor*> ValidActors;
 	for(FSelectionIterator Iter(*SelectedActors); Iter; ++Iter)
 	{
 		AActor* Actor = Cast<AActor>(*Iter);
-		if( K2Schema->IsActorValidForLevelScriptRefs(Actor, LevelBlueprint) )
+		if( K2Schema->IsActorValidForLevelScriptRefs(Actor, ContextMenuBuilder.Blueprint) )
 		{
 			ValidActors.Add( Actor );
 		}
@@ -2186,9 +2183,6 @@ void FK2ActionMenuBuilder::GetLiteralsFromActorSelection(FBlueprintGraphActionLi
 //------------------------------------------------------------------------------
 void FK2ActionMenuBuilder::GetBoundEventsFromActorSelection(FBlueprintGraphActionListBuilder& ContextMenuBuilder) const
 {
-	const ULevelScriptBlueprint* LevelBlueprint = CastChecked<ULevelScriptBlueprint>(ContextMenuBuilder.Blueprint);
-	check(LevelBlueprint != NULL);
-
 	USelection* SelectedActors = GEditor->GetSelectedActors();
 
 	TArray<AActor*> ValidActors;
@@ -2196,7 +2190,7 @@ void FK2ActionMenuBuilder::GetBoundEventsFromActorSelection(FBlueprintGraphActio
 	{
 		// We only care about actors that are referenced in the world for bound events
 		AActor* Actor = Cast<AActor>(*Iter);
-		if( K2Schema->IsActorValidForLevelScriptRefs(Actor, LevelBlueprint) )
+		if( K2Schema->IsActorValidForLevelScriptRefs(Actor, ContextMenuBuilder.Blueprint) )
 		{
 			ValidActors.Add( Actor );
 		}
@@ -2363,16 +2357,12 @@ void FK2ActionMenuBuilder::GetMatineeControllers(FBlueprintGraphActionListBuilde
 //------------------------------------------------------------------------------
 void FK2ActionMenuBuilder::GetFunctionCallsOnSelectedActors(FBlueprintGraphActionListBuilder& ContextMenuBuilder) const
 {
-	const ULevelScriptBlueprint* LevelBlueprint = CastChecked<ULevelScriptBlueprint>(ContextMenuBuilder.Blueprint);
-	ULevel* BlueprintLevel = LevelBlueprint->GetLevel();
-	check(BlueprintLevel != NULL);
-
 	USelection* SelectedActors = GEditor->GetSelectedActors();
 	TArray<AActor*> ValidActors;
 	for(FSelectionIterator Iter(*SelectedActors); Iter; ++Iter)
 	{
 		AActor* Actor = Cast<AActor>(*Iter);
-		if ( Actor != NULL && Actor->GetLevel() == BlueprintLevel && K2Schema->IsActorValidForLevelScriptRefs(Actor, LevelBlueprint) )
+		if ( K2Schema->IsActorValidForLevelScriptRefs(Actor, ContextMenuBuilder.Blueprint) )
 		{
 			ValidActors.Add( Actor );
 		}
diff --git a/Engine/Source/Editor/Kismet/Private/BlueprintActionMenuUtils.cpp b/Engine/Source/Editor/Kismet/Private/BlueprintActionMenuUtils.cpp
--- a/Engine/Source/Editor/Kismet/Private/BlueprintActionMenuUtils.cpp
+++ b/Engine/Source/Editor/Kismet/Private/BlueprintActionMenuUtils.cpp
@@ -440,7 +440,7 @@ void FBlueprintActionMenuUtils::MakeContextMenu(FBlueprintActionContext const& C
 		UClass* BlueprintClass = Blueprint->SkeletonGeneratedClass;
 		if (BlueprintClass != nullptr)
 		{
-			bCanOperateOnLevelActors &= BlueprintClass->IsChildOf<ALevelScriptActor>();
+			bCanOperateOnLevelActors &= (BlueprintClass->IsChildOf<ALevelScriptActor>() || FBlueprintEditorUtils::IsAnonymousBlueprintClass(BlueprintClass));
 			if (bAddTargetContext)
 			{
 				MainMenuFilter.TargetClasses.Add(BlueprintClass);
@@ -471,8 +471,7 @@ void FBlueprintActionMenuUtils::MakeContextMenu(FBlueprintActionContext const& C
 				// Make sure every blueprint is in the same level as this actor
 				for (UBlueprint* Blueprint : Context.Blueprints)
 				{
-					ULevelScriptBlueprint* LevelScriptBlueprint = Cast<ULevelScriptBlueprint>(Blueprint);
-					if (LevelScriptBlueprint && !K2Schema->IsActorValidForLevelScriptRefs(LevelActor, LevelScriptBlueprint))
+					if (!K2Schema->IsActorValidForLevelScriptRefs(LevelActor, Blueprint))
 					{
 						LevelActorsFilter.Context.SelectedObjects.Remove(Selection);
 						break;
@@ -499,8 +498,7 @@ void FBlueprintActionMenuUtils::MakeContextMenu(FBlueprintActionContext const& C
 			// Make sure every blueprint is in the same level as this actor
 			for (UBlueprint* Blueprint : Context.Blueprints)
 			{
-				ULevelScriptBlueprint* LevelScriptBlueprint = Cast<ULevelScriptBlueprint>(Blueprint);
-				if (LevelScriptBlueprint && !K2Schema->IsActorValidForLevelScriptRefs(LevelActor, LevelScriptBlueprint))
+				if (!K2Schema->IsActorValidForLevelScriptRefs(LevelActor, Blueprint))
 				{
 					bAddActor = false;
 					break;
diff --git a/Engine/Source/Editor/Kismet/Private/BlueprintEditor.cpp b/Engine/Source/Editor/Kismet/Private/BlueprintEditor.cpp
--- a/Engine/Source/Editor/Kismet/Private/BlueprintEditor.cpp
+++ b/Engine/Source/Editor/Kismet/Private/BlueprintEditor.cpp
@@ -2499,12 +2499,11 @@ void FBlueprintEditor::OnGraphEditorFocused(const TSharedRef<SGraphEditor>& InGr
 
 void FBlueprintEditor::OnGraphEditorDropActor(const TArray< TWeakObjectPtr<AActor> >& Actors, UEdGraph* Graph, const FVector2D& DropLocation)
 {
-	// First we need to check that the dropped actor is in the right sublevel for the reference
-	ULevelScriptBlueprint* LevelBlueprint = Cast<ULevelScriptBlueprint>(GetBlueprintObj());
-	if (LevelBlueprint != NULL)
-	{
-		ULevel* BlueprintLevel = LevelBlueprint->GetLevel();
+	// We need to check that the dropped actor is in the right sublevel for the reference
+	ULevel* BlueprintLevel = FBlueprintEditorUtils::GetLevelFromBlueprint(GetBlueprintObj());
 
+	if (BlueprintLevel)
+	{
 		FVector2D NodeLocation = DropLocation;
 		for (int32 i = 0; i < Actors.Num(); i++)
 		{
diff --git a/Engine/Source/Editor/UnrealEd/Private/Kismet2/BlueprintEditorUtils.cpp b/Engine/Source/Editor/UnrealEd/Private/Kismet2/BlueprintEditorUtils.cpp
--- a/Engine/Source/Editor/UnrealEd/Private/Kismet2/BlueprintEditorUtils.cpp
+++ b/Engine/Source/Editor/UnrealEd/Private/Kismet2/BlueprintEditorUtils.cpp
@@ -2440,6 +2440,11 @@ bool FBlueprintEditorUtils::IsAnonymousBlueprintClass(const UClass* Class)
 	return (Class->GetOutermost()->ContainsMap());
 }
 
+ULevel* FBlueprintEditorUtils::GetLevelFromBlueprint(const UBlueprint* Blueprint)
+{
+	return Cast<ULevel>(Blueprint->GetOuter());
+}
+
 bool FBlueprintEditorUtils::SupportsConstructionScript(const UBlueprint* Blueprint)
 {
 	return(	!FBlueprintEditorUtils::IsInterfaceBlueprint(Blueprint) && 
diff --git a/Engine/Source/Editor/UnrealEd/Public/Kismet2/BlueprintEditorUtils.h b/Engine/Source/Editor/UnrealEd/Public/Kismet2/BlueprintEditorUtils.h
--- a/Engine/Source/Editor/UnrealEd/Public/Kismet2/BlueprintEditorUtils.h
+++ b/Engine/Source/Editor/UnrealEd/Public/Kismet2/BlueprintEditorUtils.h
@@ -378,6 +378,13 @@ public:
 	 */
 	static bool IsAnonymousBlueprintClass(const UClass* Class);
 
+	/**
+	 * If a blueprint is directly tied to a level (level script and anonymous blueprints), this will return a pointer to that level
+	 *
+	 * @return	The level, if any, tied to this blueprint
+	 */
+	static class ULevel* GetLevelFromBlueprint(const UBlueprint* Blueprint);
+
 	/** Do we support construction scripts */
 	static bool SupportsConstructionScript(const UBlueprint* Blueprint);
 
