#ue4 Changing the scoping of RestrictToClass metadata to restrict the placing context as opposed to the drag context.  Prevents DefaultToSelf crashes (UE-7200)

diff --git a/Engine/Source/Editor/BlueprintGraph/Classes/EdGraphSchema_K2.h b/Engine/Source/Editor/BlueprintGraph/Classes/EdGraphSchema_K2.h
--- a/Engine/Source/Editor/BlueprintGraph/Classes/EdGraphSchema_K2.h
+++ b/Engine/Source/Editor/BlueprintGraph/Classes/EdGraphSchema_K2.h
@@ -56,7 +56,7 @@ public:
 	/** If true, the class will be usable as a base for blueprints */
 	static const FName MD_IsBlueprintBase;
 	
-	/** A listing of classes that this class is accessible from (and only those classes, if present) */
+	/** A listing of classes that this class is accessible from (and only those classes, if present).  Note that this determines the GRAPH CONTEXTS in which the node cannot be placed (e.g. right click menu, palette), and does NOT control menus when dragging off of a context pin (i.e. contextual drag) */
 	static const FName MD_RestrictedToClasses;
 
 	//    function metadata
diff --git a/Engine/Source/Editor/BlueprintGraph/Private/BlueprintActionFilter.cpp b/Engine/Source/Editor/BlueprintGraph/Private/BlueprintActionFilter.cpp
--- a/Engine/Source/Editor/BlueprintGraph/Private/BlueprintActionFilter.cpp
+++ b/Engine/Source/Editor/BlueprintGraph/Private/BlueprintActionFilter.cpp
@@ -596,8 +596,15 @@ static bool BlueprintActionFilterImpl::IsRestrictedClassMember(FBlueprintActionF
 		if (ActionClass->HasMetaData(FBlueprintMetadata::MD_RestrictedToClasses))
 		{
 			FString const& ClassRestrictions = ActionClass->GetMetaData(FBlueprintMetadata::MD_RestrictedToClasses);
-			for (UClass const* TargetClass : Filter.TargetClasses)
+			for (UBlueprint const* TargetContext : FilterContext.Blueprints)
 			{
+				UClass* TargetClass = TargetContext->GeneratedClass;
+				if(!TargetClass)
+				{
+					// Skip possible null classes (e.g. macros, etc)
+					continue;
+				};
+
 				bool bIsClassListed = false;
 				
 				UClass const* QueryClass = TargetClass;
