Increment/Decrement/Negate macros have compact representations

TODO: Change this to use metadata by allowing macros to specify CompactNodeTitle metadata

UE-2620 ROCKET: TASK: Add in increment decrement macros

diff --git a/Engine/Source/Editor/BlueprintGraph/Classes/K2Node_MacroInstance.h b/Engine/Source/Editor/BlueprintGraph/Classes/K2Node_MacroInstance.h
--- a/Engine/Source/Editor/BlueprintGraph/Classes/K2Node_MacroInstance.h
+++ b/Engine/Source/Editor/BlueprintGraph/Classes/K2Node_MacroInstance.h
@@ -61,6 +61,8 @@ public:
 	virtual int32 GetNodeRefreshPriority() const override { return EBaseNodeRefreshPriority::Low_UsesDependentWildcard; }
 	virtual FBlueprintNodeSignature GetSignature() const override;
 	virtual bool IsActionFilteredOut(class FBlueprintActionFilter const& Filter) override;
+	virtual FText GetCompactNodeTitle() const override;
+	virtual bool ShouldDrawCompact() const override;
 	// End UK2Node interface
 
 	void SetMacroGraph(UEdGraph* Graph) { MacroGraphReference.SetGraph(Graph); }
diff --git a/Engine/Source/Editor/BlueprintGraph/Private/K2Node_MacroInstance.cpp b/Engine/Source/Editor/BlueprintGraph/Private/K2Node_MacroInstance.cpp
--- a/Engine/Source/Editor/BlueprintGraph/Private/K2Node_MacroInstance.cpp
+++ b/Engine/Source/Editor/BlueprintGraph/Private/K2Node_MacroInstance.cpp
@@ -356,6 +356,56 @@ FName UK2Node_MacroInstance::GetPaletteIcon(FLinearColor& OutColor) const
 	return TEXT("GraphEditor.Macro_16x");
 }
 
+FText UK2Node_MacroInstance::GetCompactNodeTitle() const
+{
+	// Special case handling for standard macros
+	// @TODO Change this to use metadata by allowing macros to specify CompactNodeTitle metadata
+	UEdGraph* MacroGraph = MacroGraphReference.GetGraph();
+	if(MacroGraph != nullptr && MacroGraph->GetOuter()->GetName() == TEXT("StandardMacros"))
+	{
+		FName MacroName = FName(*MacroGraph->GetName());
+		if(	MacroName == TEXT("IncrementFloat" ) ||
+			MacroName == TEXT("IncrementInt"))
+		{
+			return LOCTEXT("IncrementCompactNodeTitle", "++");
+		}
+		else if( MacroName == TEXT("DecrementFloat") || 
+				 MacroName == TEXT("DecrementInt"))
+		{
+			return LOCTEXT("DecrementCompactNodeTitle", "--");
+		}
+		else if( MacroName == TEXT("NegateFloat") || 
+				 MacroName == TEXT("NegateInt") )
+		{
+			return LOCTEXT("DecrementCompactNodeTitle", "-");
+		}
+	}	
+
+	return Super::GetCompactNodeTitle();
+}
+
+bool UK2Node_MacroInstance::ShouldDrawCompact() const
+{
+	// Special case handling for standard macros
+	// @TODO Change this to use metadata by allowing macros to specify CompactNodeTitle metadata
+	UEdGraph* MacroGraph = MacroGraphReference.GetGraph();
+	if(MacroGraph != nullptr && MacroGraph->GetOuter()->GetName() == TEXT("StandardMacros"))
+	{
+		FName MacroName = FName(*MacroGraph->GetName());
+		if(	MacroName == TEXT("IncrementFloat" ) ||
+			MacroName == TEXT("IncrementInt") || 
+			MacroName == TEXT("DecrementFloat") || 
+			MacroName == TEXT("DecrementInt") || 
+			MacroName == TEXT("NegateFloat") || 
+			MacroName == TEXT("NegateInt") )
+		{
+			return true;
+		}
+	}
+
+	return Super::ShouldDrawCompact();
+}
+
 bool UK2Node_MacroInstance::CanPasteHere(const UEdGraph* TargetGraph) const
 {
 	bool bCanPaste = false;
