Merged pull request 703 (https://github.com/EpicGames/UnrealEngine/pull/703/) to mainline.
Includes
- Several AActor functions became pure (GetDistanceTo, et al)
- K2Node_CallFunction will auto-rewire execution pins when a node converts from impure to pure.

Thanks bozaro!

#github

diff --git a/Engine/Source/Editor/BlueprintGraph/Classes/K2Node_CallFunction.h b/Engine/Source/Editor/BlueprintGraph/Classes/K2Node_CallFunction.h
--- a/Engine/Source/Editor/BlueprintGraph/Classes/K2Node_CallFunction.h
+++ b/Engine/Source/Editor/BlueprintGraph/Classes/K2Node_CallFunction.h
@@ -192,6 +192,11 @@ private:
 	 */
 	void GeneratePinTooltip(UEdGraphPin& Pin) const;
 
+	/**
+	 * Connect Execute and Then pins for functions, which became pure.
+	 */
+	bool ReconnectPureExecPins(TArray<UEdGraphPin*>& OldPins);
+
 protected:
 	/** Helper function to ensure function is called in our context */
 	virtual void EnsureFunctionIsInBlueprint();
diff --git a/Engine/Source/Editor/BlueprintGraph/Private/K2Node_CallFunction.cpp b/Engine/Source/Editor/BlueprintGraph/Private/K2Node_CallFunction.cpp
--- a/Engine/Source/Editor/BlueprintGraph/Private/K2Node_CallFunction.cpp
+++ b/Engine/Source/Editor/BlueprintGraph/Private/K2Node_CallFunction.cpp
@@ -605,6 +605,9 @@ void UK2Node_CallFunction::ReallocatePinsDuringReconstruction(TArray<UEdGraphPin
 	// END TEMP
 
 	Super::ReallocatePinsDuringReconstruction(OldPins);
+
+	// Connect Execute and Then pins for functions, which became pure.
+	ReconnectPureExecPins(OldPins);
 }
 
 UEdGraphPin* UK2Node_CallFunction::CreateSelfPin(const UFunction* Function)
@@ -2034,6 +2037,53 @@ FName UK2Node_CallFunction::GetPaletteIcon(FLinearColor& OutColor) const
 	return GetPaletteIconForFunction(GetTargetFunction(), OutColor);
 }
 
+bool UK2Node_CallFunction::ReconnectPureExecPins(TArray<UEdGraphPin*>& OldPins)
+{
+	if (bIsPureFunc)
+	{
+		// look for an old exec pin
+		const UEdGraphSchema_K2* K2Schema = GetDefault<UEdGraphSchema_K2>();
+		UEdGraphPin* PinExec = nullptr;
+		for (int32 PinIdx = 0; PinIdx < OldPins.Num(); PinIdx++)
+		{
+			if (OldPins[PinIdx]->PinName == K2Schema->PN_Execute)
+			{
+				PinExec = OldPins[PinIdx];
+				break;
+			}
+		}
+		if (PinExec)
+		{
+			// look for old then pin
+			UEdGraphPin* PinThen = nullptr;
+			for (int32 PinIdx = 0; PinIdx < OldPins.Num(); PinIdx++)
+			{
+				if (OldPins[PinIdx]->PinName == K2Schema->PN_Then)
+				{
+					PinThen = OldPins[PinIdx];
+					break;
+				}
+			}
+			if (PinThen)
+			{
+				// reconnect all incoming links to old exec pin to the far end of the old then pin.
+				if (PinThen->LinkedTo.Num() > 0)
+				{
+					UEdGraphPin* PinThenLinked = PinThen->LinkedTo[0];
+					while (PinExec->LinkedTo.Num() > 0)
+					{
+						UEdGraphPin* PinExecLinked = PinExec->LinkedTo[0];
+						PinExecLinked->BreakLinkTo(PinExec);
+						PinExecLinked->MakeLinkTo(PinThenLinked);
+					}
+					return true;
+				}
+			}
+		}
+	}
+	return false;
+}
+
 FText UK2Node_CallFunction::GetToolTipHeading() const
 {
 	FText Heading = Super::GetToolTipHeading();
diff --git a/Engine/Source/Runtime/Engine/Classes/GameFramework/Actor.h b/Engine/Source/Runtime/Engine/Classes/GameFramework/Actor.h
--- a/Engine/Source/Runtime/Engine/Classes/GameFramework/Actor.h
+++ b/Engine/Source/Runtime/Engine/Classes/GameFramework/Actor.h
@@ -689,23 +689,23 @@ public:
 
 	/** Returns the distance from this Actor to OtherActor. */
 	UFUNCTION(BlueprintCallable, Category = "Utilities|Transformation")
-	float GetDistanceTo(AActor* OtherActor);
+	float GetDistanceTo(const AActor* OtherActor) const;
 
 	/** Returns the distance from this Actor to OtherActor, ignoring Z. */
 	UFUNCTION(BlueprintCallable, Category = "Utilities|Transformation")
-	float GetHorizontalDistanceTo(AActor* OtherActor);
+	float GetHorizontalDistanceTo(const AActor* OtherActor) const;
 
 	/** Returns the distance from this Actor to OtherActor, ignoring XY. */
 	UFUNCTION(BlueprintCallable, Category = "Utilities|Transformation")
-	float GetVerticalDistanceTo(AActor* OtherActor);
+	float GetVerticalDistanceTo(const AActor* OtherActor) const;
 
 	/** Returns the dot product from this Actor to OtherActor. Returns -2.0 on failure. Returns 0.0 for coincidental actors. */
 	UFUNCTION(BlueprintCallable, Category = "Utilities|Transformation")
-	float GetDotProductTo(AActor* OtherActor);
+	float GetDotProductTo(const AActor* OtherActor) const;
 
 	/** Returns the dot product from this Actor to OtherActor, ignoring Z. Returns -2.0 on failure. Returns 0.0 for coincidental actors. */
 	UFUNCTION(BlueprintCallable, Category = "Utilities|Transformation")
-	float GetHorizontalDotProductTo(AActor* OtherActor);
+	float GetHorizontalDotProductTo(const AActor* OtherActor) const;
 
 	/**
 	 * Adds a delta to the location of this actor in world space.
@@ -1190,7 +1190,7 @@ public:
 	 * this is a template for no other reason than to delay compilation until USceneComponent is defined
 	 */ 
 	template<class T>
-	static FORCEINLINE FVector GetActorLocation(T* RootComponent)
+	static FORCEINLINE FVector GetActorLocation(const T* RootComponent)
 	{
 		FVector Result(0.f);
 		if( RootComponent != NULL )
diff --git a/Engine/Source/Runtime/Engine/Private/Actor.cpp b/Engine/Source/Runtime/Engine/Private/Actor.cpp
--- a/Engine/Source/Runtime/Engine/Private/Actor.cpp
+++ b/Engine/Source/Runtime/Engine/Private/Actor.cpp
@@ -3478,22 +3478,22 @@ UMaterialInstanceDynamic* AActor::MakeMIDForMaterial(class UMaterialInterface* P
 	return NULL;
 }
 
-float AActor::GetDistanceTo(AActor* OtherActor)
+float AActor::GetDistanceTo(const AActor* OtherActor) const
 {
 	return OtherActor ? (GetActorLocation() - OtherActor->GetActorLocation()).Size() : 0.f;
 }
 
-float AActor::GetHorizontalDistanceTo(AActor* OtherActor)
+float AActor::GetHorizontalDistanceTo(const AActor* OtherActor) const
 {
 	return OtherActor ? (GetActorLocation() - OtherActor->GetActorLocation()).Size2D() : 0.f;
 }
 
-float AActor::GetVerticalDistanceTo(AActor* OtherActor)
+float AActor::GetVerticalDistanceTo(const AActor* OtherActor) const
 {
 	return OtherActor ? FMath::Abs((GetActorLocation().Z - OtherActor->GetActorLocation().Z)) : 0.f;
 }
 
-float AActor::GetDotProductTo(AActor* OtherActor)
+float AActor::GetDotProductTo(const AActor* OtherActor) const
 {
 	if (OtherActor)
 	{
@@ -3505,7 +3505,7 @@ float AActor::GetDotProductTo(AActor* OtherActor)
 	return -2.0;
 }
 
-float AActor::GetHorizontalDotProductTo(AActor* OtherActor)
+float AActor::GetHorizontalDotProductTo(const AActor* OtherActor) const
 {
 	if (OtherActor)
 	{
