When converting a cast node from impure to pure, exec wires will now be automatically fixed to maintain execution flow.

diff --git a/Engine/Source/Editor/BlueprintGraph/Classes/K2Node_DynamicCast.h b/Engine/Source/Editor/BlueprintGraph/Classes/K2Node_DynamicCast.h
--- a/Engine/Source/Editor/BlueprintGraph/Classes/K2Node_DynamicCast.h
+++ b/Engine/Source/Editor/BlueprintGraph/Classes/K2Node_DynamicCast.h
@@ -34,6 +34,7 @@ class UK2Node_DynamicCast : public UK2Node
 	virtual bool IsNodePure() const override { return bIsPureCast; }
 	virtual bool IsConnectionDisallowed(const UEdGraphPin* MyPin, const UEdGraphPin* OtherPin, FString& OutReason) const override;
 	virtual void NotifyPinConnectionListChanged(UEdGraphPin* Pin) override;
+	virtual void ReallocatePinsDuringReconstruction(TArray<UEdGraphPin*>& OldPins) override;
 	// End of UK2Node interface
 
 	/** Get the 'valid cast' exec pin */
@@ -63,6 +64,9 @@ protected:
 	/** Flips the node's purity (adding/removing exec pins as needed). */
 	void TogglePurity();
 
+	/** Update exec pins when converting from impure to pure. */
+	bool ReconnectPureExecPins(TArray<UEdGraphPin*>& OldPins);
+	
 	/** Constructing FText strings can be costly, so we cache the node's title */
 	FNodeTextCache CachedNodeTitle;
 
diff --git a/Engine/Source/Editor/BlueprintGraph/Private/K2Node_DynamicCast.cpp b/Engine/Source/Editor/BlueprintGraph/Private/K2Node_DynamicCast.cpp
--- a/Engine/Source/Editor/BlueprintGraph/Private/K2Node_DynamicCast.cpp
+++ b/Engine/Source/Editor/BlueprintGraph/Private/K2Node_DynamicCast.cpp
@@ -388,4 +388,60 @@ void UK2Node_DynamicCast::NotifyPinConnectionListChanged(UEdGraphPin* Pin)
 	}
 }
 
+void UK2Node_DynamicCast::ReallocatePinsDuringReconstruction(TArray<UEdGraphPin*>& OldPins)
+{
+	Super::ReallocatePinsDuringReconstruction(OldPins);
+
+	// Update exec pins if we converted from impure to pure
+	ReconnectPureExecPins(OldPins);
+}
+
+
+bool UK2Node_DynamicCast::ReconnectPureExecPins(TArray<UEdGraphPin*>& OldPins)
+{
+	if (bIsPureCast)
+	{
+		// look for an old exec pin
+		const UEdGraphSchema_K2* K2Schema = GetDefault<UEdGraphSchema_K2>();
+		UEdGraphPin* PinExec = nullptr;
+		for (UEdGraphPin* Pin : OldPins)
+		{
+			if (Pin->PinName == K2Schema->PN_Execute)
+			{
+				PinExec = Pin;
+				break;
+			}
+		}
+		if (PinExec)
+		{
+			// look for old then pin
+			UEdGraphPin* PinThen = nullptr;
+			for (UEdGraphPin* Pin : OldPins)
+			{
+				if (Pin->PinName == K2Schema->PN_Then)
+				{
+					PinThen = Pin;
+					break;
+				}
+			}
+			if (PinThen)
+			{
+				// reconnect all incoming links to old exec pin to the far end of the old then pin.
+				if (PinThen->LinkedTo.Num() > 0)
+				{
+					UEdGraphPin* PinThenLinked = PinThen->LinkedTo[0];
+					while (PinExec->LinkedTo.Num() > 0)
+					{
+						UEdGraphPin* PinExecLinked = PinExec->LinkedTo[0];
+						PinExecLinked->BreakLinkTo(PinExec);
+						PinExecLinked->MakeLinkTo(PinThenLinked);
+					}
+					return true;
+				}
+			}
+		}
+	}
+	return false;
+}
+
 #undef LOCTEXT_NAMESPACE
