UHT: Allows defining multiple UCLASSes in one header file.

diff --git a/Engine/Source/Developer/Android/AndroidPlatformEditor/Private/AndroidTargetSettingsCustomization.cpp b/Engine/Source/Developer/Android/AndroidPlatformEditor/Private/AndroidTargetSettingsCustomization.cpp
--- a/Engine/Source/Developer/Android/AndroidPlatformEditor/Private/AndroidTargetSettingsCustomization.cpp
+++ b/Engine/Source/Developer/Android/AndroidPlatformEditor/Private/AndroidTargetSettingsCustomization.cpp
@@ -313,9 +313,9 @@ void FAndroidTargetSettingsCustomization::OnOrientationModified()
 
 FString FAndroidTargetSettingsCustomization::OrientationToString(const EAndroidScreenOrientation::Type Orientation)
 {
-	extern ANDROIDRUNTIMESETTINGS_API class UEnum* Z_Construct_UEnum_UAndroidRuntimeSettings_EAndroidScreenOrientation();
+	extern ANDROIDRUNTIMESETTINGS_API class UEnum* Z_Construct_UEnum_AndroidRuntimeSettings_EAndroidScreenOrientation();
 	
-	UEnum* Enum = Z_Construct_UEnum_UAndroidRuntimeSettings_EAndroidScreenOrientation();
+	UEnum* Enum = Z_Construct_UEnum_AndroidRuntimeSettings_EAndroidScreenOrientation();
 	check(Enum != nullptr);
 	
 	return Enum->GetMetaData(TEXT("ManifestValue"), (int32)Orientation);
@@ -349,9 +349,9 @@ void FAndroidTargetSettingsCustomization::OnDepthBufferPreferenceModified()
 
 FString FAndroidTargetSettingsCustomization::DepthBufferPreferenceToString(const EAndroidDepthBufferPreference::Type DepthBufferPreference)
 {
-	extern ANDROIDRUNTIMESETTINGS_API class UEnum* Z_Construct_UEnum_UAndroidRuntimeSettings_EAndroidDepthBufferPreference();
+	extern ANDROIDRUNTIMESETTINGS_API class UEnum* Z_Construct_UEnum_AndroidRuntimeSettings_EAndroidDepthBufferPreference();
 
-	UEnum* Enum = Z_Construct_UEnum_UAndroidRuntimeSettings_EAndroidDepthBufferPreference();
+	UEnum* Enum = Z_Construct_UEnum_AndroidRuntimeSettings_EAndroidDepthBufferPreference();
 	check(Enum != nullptr);
 
 	return Enum->GetMetaData(TEXT("ManifestValue"), (int32)DepthBufferPreference);
diff --git a/Engine/Source/Developer/FunctionalTesting/Private/FunctionalTest.cpp b/Engine/Source/Developer/FunctionalTesting/Private/FunctionalTest.cpp
--- a/Engine/Source/Developer/FunctionalTesting/Private/FunctionalTest.cpp
+++ b/Engine/Source/Developer/FunctionalTesting/Private/FunctionalTest.cpp
@@ -91,7 +91,7 @@ bool AFunctionalTest::StartTest(const TArray<FString>& Params)
 
 void AFunctionalTest::FinishTest(TEnumAsByte<EFunctionalTestResult::Type> TestResult, const FString& Message)
 {
-	const static UEnum* FTestResultTypeEnum = FindObject<UEnum>( NULL, TEXT("FunctionalTesting.FunctionalTest.EFunctionalTestResult") );
+	const static UEnum* FTestResultTypeEnum = FindObject<UEnum>( NULL, TEXT("FunctionalTesting.EFunctionalTestResult") );
 	
 	if (bIsRunning == false)
 	{
diff --git a/Engine/Source/Editor/BlueprintGraph/Classes/K2Node.h b/Engine/Source/Editor/BlueprintGraph/Classes/K2Node.h
--- a/Engine/Source/Editor/BlueprintGraph/Classes/K2Node.h
+++ b/Engine/Source/Editor/BlueprintGraph/Classes/K2Node.h
@@ -688,6 +688,16 @@ public:
 					}
 				}
 			}
+			else
+			{
+				// Native delegate signatures can have null target scope.
+                // Needed because Blueprint pin types are referencing delegate signatures as
+                // class members, which doesn't have to be true.
+				if (TFieldType::StaticClass() == UFunction::StaticClass())
+				{
+					ReturnField = (TFieldType*)FindDelegateSignature(MemberName);
+				}
+			}
 		}
 
 		// Check to see if the member has been deprecated
diff --git a/Engine/Source/Editor/BlueprintGraph/Private/CallFunctionHandler.cpp b/Engine/Source/Editor/BlueprintGraph/Private/CallFunctionHandler.cpp
--- a/Engine/Source/Editor/BlueprintGraph/Private/CallFunctionHandler.cpp
+++ b/Engine/Source/Editor/BlueprintGraph/Private/CallFunctionHandler.cpp
@@ -155,9 +155,9 @@ void FKCHandler_CallFunction::CreateFunctionCallStatement(FKismetFunctionContext
 		int32 LatentTargetParamIndex = INDEX_NONE;
 
 		// Grab the special case structs that use their own literal path
-		UScriptStruct* VectorStruct = FindObjectChecked<UScriptStruct>(UObject::StaticClass(), TEXT("Vector"));
-		UScriptStruct* RotatorStruct = FindObjectChecked<UScriptStruct>(UObject::StaticClass(), TEXT("Rotator"));
-		UScriptStruct* TransformStruct = FindObjectChecked<UScriptStruct>(UObject::StaticClass(), TEXT("Transform"));
+		UScriptStruct* VectorStruct = GetBaseStructure(TEXT("Vector"));
+		UScriptStruct* RotatorStruct = GetBaseStructure(TEXT("Rotator"));
+		UScriptStruct* TransformStruct = GetBaseStructure(TEXT("Transform"));
 
 		// Check each property
 		bool bMatchedAllParams = true;
diff --git a/Engine/Source/Editor/BlueprintGraph/Private/EdGraphSchema_K2.cpp b/Engine/Source/Editor/BlueprintGraph/Private/EdGraphSchema_K2.cpp
--- a/Engine/Source/Editor/BlueprintGraph/Private/EdGraphSchema_K2.cpp
+++ b/Engine/Source/Editor/BlueprintGraph/Private/EdGraphSchema_K2.cpp
@@ -483,11 +483,11 @@ UEdGraphSchema_K2::UEdGraphSchema_K2(const FObjectInitializer& ObjectInitializer
 	// Initialize cached static references to well-known struct types
 	if (VectorStruct == nullptr)
 	{
-		VectorStruct = FindObjectChecked<UScriptStruct>(UObject::StaticClass(), TEXT("Vector"));
-		RotatorStruct = FindObjectChecked<UScriptStruct>(UObject::StaticClass(), TEXT("Rotator"));
-		TransformStruct = FindObjectChecked<UScriptStruct>(UObject::StaticClass(), TEXT("Transform"));
-		LinearColorStruct = FindObjectChecked<UScriptStruct>(UObject::StaticClass(), TEXT("LinearColor"));
-		ColorStruct = FindObjectChecked<UScriptStruct>(UObject::StaticClass(), TEXT("Color"));
+		VectorStruct = GetBaseStructure(TEXT("Vector"));
+		RotatorStruct = GetBaseStructure(TEXT("Rotator"));
+		TransformStruct = GetBaseStructure(TEXT("Transform"));
+		LinearColorStruct = GetBaseStructure(TEXT("LinearColor"));
+		ColorStruct = GetBaseStructure(TEXT("Color"));
 	}
 }
 
@@ -5455,8 +5455,8 @@ void UEdGraphSchema_K2::SplitPin(UEdGraphPin* Pin) const
 	if (Pin->Direction == EGPD_Input)
 	{
 		TArray<FString> OriginalDefaults;
-		if (   StructType == FindObjectChecked<UScriptStruct>(UObject::StaticClass(), TEXT("Vector"))
-			|| StructType == FindObjectChecked<UScriptStruct>(UObject::StaticClass(), TEXT("Rotator")))
+		if (   StructType == GetBaseStructure(TEXT("Vector"))
+			|| StructType == GetBaseStructure(TEXT("Rotator")))
 		{
 			Pin->DefaultValue.ParseIntoArray(&OriginalDefaults, TEXT(","), false);
 			for (FString& Default : OriginalDefaults)
@@ -5469,7 +5469,7 @@ void UEdGraphSchema_K2::SplitPin(UEdGraphPin* Pin) const
 				OriginalDefaults.Add(TEXT("0.0"));
 			}
 		}
-		else if (StructType == FindObjectChecked<UScriptStruct>(UObject::StaticClass(), TEXT("Vector2D")))
+		else if (StructType == GetBaseStructure(TEXT("Vector2D")))
 		{
 			FVector2D V2D;
 			V2D.InitFromString(Pin->DefaultValue);
@@ -5477,7 +5477,7 @@ void UEdGraphSchema_K2::SplitPin(UEdGraphPin* Pin) const
 			OriginalDefaults.Add(FString::SanitizeFloat(V2D.X));
 			OriginalDefaults.Add(FString::SanitizeFloat(V2D.Y));
 		}
-		else if (StructType == FindObjectChecked<UScriptStruct>(UObject::StaticClass(), TEXT("LinearColor")))
+		else if (StructType == GetBaseStructure(TEXT("LinearColor")))
 		{
 			FLinearColor LC;
 			LC.InitFromString(Pin->DefaultValue);
@@ -5532,14 +5532,14 @@ void UEdGraphSchema_K2::RecombinePin(UEdGraphPin* Pin) const
 		UScriptStruct* StructType = CastChecked<UScriptStruct>(ParentPin->PinType.PinSubCategoryObject.Get());
 
 		TArray<FString> OriginalDefaults;
-		if (   StructType == FindObjectChecked<UScriptStruct>(UObject::StaticClass(), TEXT("Vector"))
-			|| StructType == FindObjectChecked<UScriptStruct>(UObject::StaticClass(), TEXT("Rotator")))
+		if (   StructType == GetBaseStructure(TEXT("Vector"))
+			|| StructType == GetBaseStructure(TEXT("Rotator")))
 		{
 			ParentPin->DefaultValue = ParentPin->SubPins[0]->DefaultValue + TEXT(",") 
 									+ ParentPin->SubPins[1]->DefaultValue + TEXT(",")
 									+ ParentPin->SubPins[2]->DefaultValue;
 		}
-		else if (StructType == FindObjectChecked<UScriptStruct>(UObject::StaticClass(), TEXT("Vector2D")))
+		else if (StructType == GetBaseStructure(TEXT("Vector2D")))
 		{
 			FVector2D V2D;
 			V2D.X = FCString::Atof(*ParentPin->SubPins[0]->DefaultValue);
@@ -5547,7 +5547,7 @@ void UEdGraphSchema_K2::RecombinePin(UEdGraphPin* Pin) const
 			
 			ParentPin->DefaultValue = V2D.ToString();
 		}
-		else if (StructType == FindObjectChecked<UScriptStruct>(UObject::StaticClass(), TEXT("LinearColor")))
+		else if (StructType == GetBaseStructure(TEXT("LinearColor")))
 		{
 			FLinearColor LC;
 			LC.R = FCString::Atof(*ParentPin->SubPins[0]->DefaultValue);
diff --git a/Engine/Source/Editor/BlueprintGraph/Private/K2ActionMenuBuilder.cpp b/Engine/Source/Editor/BlueprintGraph/Private/K2ActionMenuBuilder.cpp
--- a/Engine/Source/Editor/BlueprintGraph/Private/K2ActionMenuBuilder.cpp
+++ b/Engine/Source/Editor/BlueprintGraph/Private/K2ActionMenuBuilder.cpp
@@ -1024,19 +1024,19 @@ void FK2ActionMenuBuilder::GetContextAllowedNodeTypes(FBlueprintGraphActionListB
 			TemporaryTypes.Add(FEdGraphPinType(K2Schema->PC_Text, TEXT(""), NULL, true, false));
 
 			FEdGraphPinType VectorPinType = FEdGraphPinType(K2Schema->PC_Struct, TEXT("Vector"), NULL, false, false);
-			VectorPinType.PinSubCategoryObject = FindObjectChecked<UScriptStruct>(UObject::StaticClass(), TEXT("Vector"));
+			VectorPinType.PinSubCategoryObject = GetBaseStructure(TEXT("Vector"));
 			TemporaryTypes.Add(VectorPinType);
 			VectorPinType.bIsArray = true;
 			TemporaryTypes.Add(VectorPinType);
 
 			FEdGraphPinType RotatorPinType = FEdGraphPinType(K2Schema->PC_Struct, TEXT("Rotator"), NULL, false, false);
-			RotatorPinType.PinSubCategoryObject = FindObjectChecked<UScriptStruct>(UObject::StaticClass(), TEXT("Rotator"));
+			RotatorPinType.PinSubCategoryObject = GetBaseStructure(TEXT("Rotator"));
 			TemporaryTypes.Add(RotatorPinType);
 			RotatorPinType.bIsArray = true;
 			TemporaryTypes.Add(RotatorPinType);
 
 			FEdGraphPinType TransformPinType = FEdGraphPinType(K2Schema->PC_Struct, TEXT("Transform"), NULL, false, false);
-			TransformPinType.PinSubCategoryObject = FindObjectChecked<UScriptStruct>(UObject::StaticClass(), TEXT("Transform"));
+			TransformPinType.PinSubCategoryObject = GetBaseStructure(TEXT("Transform"));
 			TemporaryTypes.Add(TransformPinType);
 			TransformPinType.bIsArray = true;
 			TemporaryTypes.Add(TransformPinType);
diff --git a/Engine/Source/Editor/BlueprintGraph/Private/K2Node_Event.cpp b/Engine/Source/Editor/BlueprintGraph/Private/K2Node_Event.cpp
--- a/Engine/Source/Editor/BlueprintGraph/Private/K2Node_Event.cpp
+++ b/Engine/Source/Editor/BlueprintGraph/Private/K2Node_Event.cpp
@@ -235,6 +235,11 @@ UFunction* UK2Node_Event::FindEventSignatureFunction()
 		}
 	}
 
+	if (Function == nullptr)
+	{
+		Function = FindDelegateSignature(EventSignatureName);
+	}
+
 	return Function;
 }
 
diff --git a/Engine/Source/Editor/BlueprintGraph/Private/K2Node_FunctionEntry.cpp b/Engine/Source/Editor/BlueprintGraph/Private/K2Node_FunctionEntry.cpp
--- a/Engine/Source/Editor/BlueprintGraph/Private/K2Node_FunctionEntry.cpp
+++ b/Engine/Source/Editor/BlueprintGraph/Private/K2Node_FunctionEntry.cpp
@@ -143,6 +143,12 @@ void UK2Node_FunctionEntry::AllocateDefaultPins()
 	CreatePin(EGPD_Output, K2Schema->PC_Exec, TEXT(""), NULL, false, false, K2Schema->PN_Then);
 
 	UFunction* Function = FindField<UFunction>(SignatureClass, SignatureName);
+
+	if (Function == nullptr)
+	{
+		Function = FindDelegateSignature(SignatureName);
+	}
+
 	if (Function != NULL)
 	{
 		CreatePinsForFunctionEntryExit(Function, /*bIsFunctionEntry=*/ true);
diff --git a/Engine/Source/Editor/BlueprintGraph/Private/K2Node_InputTouch.cpp b/Engine/Source/Editor/BlueprintGraph/Private/K2Node_InputTouch.cpp
--- a/Engine/Source/Editor/BlueprintGraph/Private/K2Node_InputTouch.cpp
+++ b/Engine/Source/Editor/BlueprintGraph/Private/K2Node_InputTouch.cpp
@@ -47,7 +47,7 @@ void UK2Node_InputTouch::AllocateDefaultPins()
 	CreatePin(EGPD_Output, K2Schema->PC_Exec, TEXT(""), NULL, false, false, TEXT("Released"));
 	CreatePin(EGPD_Output, K2Schema->PC_Exec, TEXT(""), NULL, false, false, TEXT("Moved"));
 	
-	UScriptStruct* VectorStruct = FindObjectChecked<UScriptStruct>(UObject::StaticClass(), TEXT("Vector"));
+	UScriptStruct* VectorStruct = GetBaseStructure(TEXT("Vector"));
 	CreatePin(EGPD_Output, K2Schema->PC_Struct, TEXT(""), VectorStruct, false, false, TEXT("Location"));
 
 	CreatePin(EGPD_Output, K2Schema->PC_Byte, TEXT(""), GetTouchIndexEnum(), false, false, TEXT("FingerIndex"));
@@ -165,7 +165,7 @@ void UK2Node_InputTouch::ExpandNode(FKismetCompilerContext& CompilerContext, UEd
 	if( ActivePins.Num() > 1 )
 	{
 		// Create a temporary variable to copy location in to
-		static UScriptStruct* VectorStruct = FindObjectChecked<UScriptStruct>(UObject::StaticClass(), TEXT("Vector"));
+		static UScriptStruct* VectorStruct = GetBaseStructure(TEXT("Vector"));
 		UK2Node_TemporaryVariable* TouchLocationVar = CompilerContext.SpawnIntermediateNode<UK2Node_TemporaryVariable>(this, SourceGraph);
 		TouchLocationVar->VariableType.PinCategory = Schema->PC_Struct;
 		TouchLocationVar->VariableType.PinSubCategoryObject = VectorStruct;
diff --git a/Engine/Source/Editor/BlueprintGraph/Private/K2Node_SpawnActor.cpp b/Engine/Source/Editor/BlueprintGraph/Private/K2Node_SpawnActor.cpp
--- a/Engine/Source/Editor/BlueprintGraph/Private/K2Node_SpawnActor.cpp
+++ b/Engine/Source/Editor/BlueprintGraph/Private/K2Node_SpawnActor.cpp
@@ -36,7 +36,7 @@ void UK2Node_SpawnActor::AllocateDefaultPins()
 	K2Schema->ConstructBasicPinTooltip(*BlueprintPin, LOCTEXT("BlueprintPinDescription", "The blueprint Actor you want to spawn"), BlueprintPin->PinToolTip);
 
 	// Transform pin
-	UScriptStruct* TransformStruct = FindObjectChecked<UScriptStruct>(UObject::StaticClass(), TEXT("Transform"));
+	UScriptStruct* TransformStruct = GetBaseStructure(TEXT("Transform"));
 	UEdGraphPin* TransformPin = CreatePin(EGPD_Input, K2Schema->PC_Struct, TEXT(""), TransformStruct, false, false, SpawnTransformPinName);
 	K2Schema->ConstructBasicPinTooltip(*TransformPin, LOCTEXT("TransformPinDescription", "The transform to spawn the Actor with"), TransformPin->PinToolTip);
 
diff --git a/Engine/Source/Editor/BlueprintGraph/Private/K2Node_SpawnActorFromClass.cpp b/Engine/Source/Editor/BlueprintGraph/Private/K2Node_SpawnActorFromClass.cpp
--- a/Engine/Source/Editor/BlueprintGraph/Private/K2Node_SpawnActorFromClass.cpp
+++ b/Engine/Source/Editor/BlueprintGraph/Private/K2Node_SpawnActorFromClass.cpp
@@ -47,7 +47,7 @@ void UK2Node_SpawnActorFromClass::AllocateDefaultPins()
 	K2Schema->ConstructBasicPinTooltip(*ClassPin, LOCTEXT("ClassPinDescription", "The Actor class you want to spawn"), ClassPin->PinToolTip);
 
 	// Transform pin
-	UScriptStruct* TransformStruct = FindObjectChecked<UScriptStruct>(UObject::StaticClass(), TEXT("Transform"));
+	UScriptStruct* TransformStruct = GetBaseStructure(TEXT("Transform"));
 	UEdGraphPin* TransformPin = CreatePin(EGPD_Input, K2Schema->PC_Struct, TEXT(""), TransformStruct, false, false, FK2Node_SpawnActorFromClassHelper::SpawnTransformPinName);
 	K2Schema->ConstructBasicPinTooltip(*TransformPin, LOCTEXT("TransformPinDescription", "The transform to spawn the Actor with"), TransformPin->PinToolTip);
 
diff --git a/Engine/Source/Editor/BlueprintGraph/Private/K2Node_TemporaryVariable.cpp b/Engine/Source/Editor/BlueprintGraph/Private/K2Node_TemporaryVariable.cpp
--- a/Engine/Source/Editor/BlueprintGraph/Private/K2Node_TemporaryVariable.cpp
+++ b/Engine/Source/Editor/BlueprintGraph/Private/K2Node_TemporaryVariable.cpp
@@ -180,15 +180,15 @@ void UK2Node_TemporaryVariable::GetMenuActions(FBlueprintActionDatabaseRegistrar
 	ActionRegistrar.AddBlueprintAction(ActionKey, MakeTempVarNodeSpawner(FEdGraphPinType(K2Schema->PC_Text, TEXT(""), nullptr, /*bIsArray =*/false, /*bIsReference =*/false), /*bIsPersistent =*/false));
 	ActionRegistrar.AddBlueprintAction(ActionKey, MakeTempVarNodeSpawner(FEdGraphPinType(K2Schema->PC_Text, TEXT(""), nullptr, /*bIsArray =*/ true, /*bIsReference =*/false), /*bIsPersistent =*/false));
 
-	UScriptStruct* VectorStruct  = FindObjectChecked<UScriptStruct>(UObject::StaticClass(), TEXT("Vector"));
+	UScriptStruct* VectorStruct  = GetBaseStructure(TEXT("Vector"));
 	ActionRegistrar.AddBlueprintAction(ActionKey, MakeTempVarNodeSpawner(FEdGraphPinType(K2Schema->PC_Struct, TEXT("Vector"), VectorStruct, /*bIsArray =*/false, /*bIsReference =*/false), /*bIsPersistent =*/false));
 	ActionRegistrar.AddBlueprintAction(ActionKey, MakeTempVarNodeSpawner(FEdGraphPinType(K2Schema->PC_Struct, TEXT("Vector"), VectorStruct, /*bIsArray =*/ true, /*bIsReference =*/false), /*bIsPersistent =*/false));
 	
-	UScriptStruct* RotatorStruct = FindObjectChecked<UScriptStruct>(UObject::StaticClass(), TEXT("Rotator"));
+	UScriptStruct* RotatorStruct = GetBaseStructure(TEXT("Rotator"));
 	ActionRegistrar.AddBlueprintAction(ActionKey, MakeTempVarNodeSpawner(FEdGraphPinType(K2Schema->PC_Struct, TEXT("Rotator"), RotatorStruct, /*bIsArray =*/false, /*bIsReference =*/false), /*bIsPersistent =*/false));
 	ActionRegistrar.AddBlueprintAction(ActionKey, MakeTempVarNodeSpawner(FEdGraphPinType(K2Schema->PC_Struct, TEXT("Rotator"), RotatorStruct, /*bIsArray =*/ true, /*bIsReference =*/false), /*bIsPersistent =*/false));
 	
-	UScriptStruct* TransformStruct = FindObjectChecked<UScriptStruct>(UObject::StaticClass(), TEXT("Transform"));
+	UScriptStruct* TransformStruct = GetBaseStructure(TEXT("Transform"));
 	ActionRegistrar.AddBlueprintAction(ActionKey, MakeTempVarNodeSpawner(FEdGraphPinType(K2Schema->PC_Struct, TEXT("Transform"), TransformStruct, /*bIsArray =*/false, /*bIsReference =*/false), /*bIsPersistent =*/false));
 	ActionRegistrar.AddBlueprintAction(ActionKey, MakeTempVarNodeSpawner(FEdGraphPinType(K2Schema->PC_Struct, TEXT("Transform"), TransformStruct, /*bIsArray =*/ true, /*bIsReference =*/false), /*bIsPersistent =*/false));
 
diff --git a/Engine/Source/Editor/BlueprintGraph/Private/K2Node_Timeline.cpp b/Engine/Source/Editor/BlueprintGraph/Private/K2Node_Timeline.cpp
--- a/Engine/Source/Editor/BlueprintGraph/Private/K2Node_Timeline.cpp
+++ b/Engine/Source/Editor/BlueprintGraph/Private/K2Node_Timeline.cpp
@@ -81,14 +81,14 @@ void UK2Node_Timeline::AllocateDefaultPins()
 		for(int32 i=0; i<Timeline->VectorTracks.Num(); i++)
 		{
 			FTTVectorTrack& VectorTrack = Timeline->VectorTracks[i];
-			UScriptStruct* VectorStruct = FindObjectChecked<UScriptStruct>(UObject::StaticClass(), TEXT("Vector"));
+			UScriptStruct* VectorStruct = GetBaseStructure(TEXT("Vector"));
 			CreatePin(EGPD_Output, K2Schema->PC_Struct, TEXT(""), VectorStruct, false, false, VectorTrack.TrackName.ToString());			
 		}
 
 		for(int32 i=0; i<Timeline->LinearColorTracks.Num(); i++)
 		{
 			FTTLinearColorTrack& LinearColorTrack = Timeline->LinearColorTracks[i];
-			UScriptStruct* LinearColorStruct = FindObjectChecked<UScriptStruct>(UObject::StaticClass(), TEXT("LinearColor"));
+			UScriptStruct* LinearColorStruct = GetBaseStructure(TEXT("LinearColor"));
 			CreatePin(EGPD_Output, K2Schema->PC_Struct, TEXT(""), LinearColorStruct, false, false, LinearColorTrack.TrackName.ToString());			
 		}
 
diff --git a/Engine/Source/Editor/GameProjectGeneration/Private/GameProjectUtils.cpp b/Engine/Source/Editor/GameProjectGeneration/Private/GameProjectUtils.cpp
--- a/Engine/Source/Editor/GameProjectGeneration/Private/GameProjectUtils.cpp
+++ b/Engine/Source/Editor/GameProjectGeneration/Private/GameProjectUtils.cpp
@@ -1588,10 +1588,10 @@ FString GameProjectUtils::GetHardwareConfigString(const FProjectInformation& InP
 	FString HardwareTargeting;
 	
 	FString TargetHardwareAsString;
-	UEnum::GetValueAsString(TEXT("/Script/HardwareTargeting.HardwareTargetingSettings.EHardwareClass"), InProjectInfo.TargetedHardware, /*out*/ TargetHardwareAsString);
+	UEnum::GetValueAsString(TEXT("/Script/HardwareTargeting.EHardwareClass"), InProjectInfo.TargetedHardware, /*out*/ TargetHardwareAsString);
 
 	FString GraphicsPresetAsString;
-	UEnum::GetValueAsString(TEXT("/Script/HardwareTargeting.HardwareTargetingSettings.EGraphicsPreset"), InProjectInfo.DefaultGraphicsPerformance, /*out*/ GraphicsPresetAsString);
+	UEnum::GetValueAsString(TEXT("/Script/HardwareTargeting.EGraphicsPreset"), InProjectInfo.DefaultGraphicsPerformance, /*out*/ GraphicsPresetAsString);
 
 	HardwareTargeting += TEXT("[/Script/HardwareTargeting.HardwareTargetingSettings]") LINE_TERMINATOR;
 	HardwareTargeting += FString::Printf(TEXT("TargetedHardwareClass=%s") LINE_TERMINATOR, *TargetHardwareAsString);
diff --git a/Engine/Source/Editor/GraphEditor/Private/KismetPins/SGraphPinVector.cpp b/Engine/Source/Editor/GraphEditor/Private/KismetPins/SGraphPinVector.cpp
--- a/Engine/Source/Editor/GraphEditor/Private/KismetPins/SGraphPinVector.cpp
+++ b/Engine/Source/Editor/GraphEditor/Private/KismetPins/SGraphPinVector.cpp
@@ -151,7 +151,7 @@ void SGraphPinVector::Construct(const FArguments& InArgs, UEdGraphPin* InGraphPi
 
 TSharedRef<SWidget>	SGraphPinVector::GetDefaultValueWidget()
 {
-	UScriptStruct* RotatorStruct = FindObjectChecked<UScriptStruct>(UObject::StaticClass(), TEXT("Rotator"));
+	UScriptStruct* RotatorStruct = GetBaseStructure(TEXT("Rotator"));
 	bIsRotator = (GraphPinObj->PinType.PinSubCategoryObject == RotatorStruct) ? true : false;
 	
 	//Create widget
diff --git a/Engine/Source/Editor/GraphEditor/Private/NiagaraPins/SGraphPinVector4.cpp b/Engine/Source/Editor/GraphEditor/Private/NiagaraPins/SGraphPinVector4.cpp
--- a/Engine/Source/Editor/GraphEditor/Private/NiagaraPins/SGraphPinVector4.cpp
+++ b/Engine/Source/Editor/GraphEditor/Private/NiagaraPins/SGraphPinVector4.cpp
@@ -181,7 +181,7 @@ void SGraphPinVector4::Construct(const FArguments& InArgs, UEdGraphPin* InGraphP
 
 TSharedRef<SWidget>	SGraphPinVector4::GetDefaultValueWidget()
 {
-	UScriptStruct* RotatorStruct = FindObjectChecked<UScriptStruct>(UObject::StaticClass(), TEXT("Rotator"));
+	UScriptStruct* RotatorStruct = GetBaseStructure(TEXT("Rotator"));
 
 	//Create widget
 	return	SNew(SVector4TextBox)
diff --git a/Engine/Source/Editor/GraphEditor/Private/NodeFactory.cpp b/Engine/Source/Editor/GraphEditor/Private/NodeFactory.cpp
--- a/Engine/Source/Editor/GraphEditor/Private/NodeFactory.cpp
+++ b/Engine/Source/Editor/GraphEditor/Private/NodeFactory.cpp
@@ -304,10 +304,10 @@ TSharedPtr<SGraphPin> FNodeFactory::CreatePinWidget(UEdGraphPin* InPin)
 		else if (InPin->PinType.PinCategory == K2Schema->PC_Struct)
 		{
 			// If you update this logic you'll probably need to update UEdGraphSchema_K2::ShouldHidePinDefaultValue!
-			UScriptStruct* ColorStruct = FindObjectChecked<UScriptStruct>(UObject::StaticClass(), TEXT("LinearColor"));
-			UScriptStruct* VectorStruct = FindObjectChecked<UScriptStruct>(UObject::StaticClass(), TEXT("Vector"));
-			UScriptStruct* Vector2DStruct = FindObjectChecked<UScriptStruct>(UObject::StaticClass(), TEXT("Vector2D"));
-			UScriptStruct* RotatorStruct = FindObjectChecked<UScriptStruct>(UObject::StaticClass(), TEXT("Rotator"));
+			UScriptStruct* ColorStruct = GetBaseStructure(TEXT("LinearColor"));
+			UScriptStruct* VectorStruct = GetBaseStructure(TEXT("Vector"));
+			UScriptStruct* Vector2DStruct = GetBaseStructure(TEXT("Vector2D"));
+			UScriptStruct* RotatorStruct = GetBaseStructure(TEXT("Rotator"));
 
 			if (InPin->PinType.PinSubCategoryObject == ColorStruct)
 			{
diff --git a/Engine/Source/Editor/Kismet/Private/BlueprintEditor.cpp b/Engine/Source/Editor/Kismet/Private/BlueprintEditor.cpp
--- a/Engine/Source/Editor/Kismet/Private/BlueprintEditor.cpp
+++ b/Engine/Source/Editor/Kismet/Private/BlueprintEditor.cpp
@@ -1743,7 +1743,7 @@ void FBlueprintEditor::OnEditParentClassNativeCodeClicked()
 FText FBlueprintEditor::GetTextForNativeParentClassHeaderLink() const
 {
 	// it could be done using FSourceCodeNavigation, but it could be slow
-	return FText::FromString(FString::Printf(TEXT("%s.h"), *GetParentClassNameText().ToString()));
+	return FText::FromString(*GetParentClassNameText().ToString());
 }
 
 FReply FBlueprintEditor::OnFindParentClassInContentBrowserClicked()
diff --git a/Engine/Source/Editor/KismetCompiler/Private/KismetCompiler.cpp b/Engine/Source/Editor/KismetCompiler/Private/KismetCompiler.cpp
--- a/Engine/Source/Editor/KismetCompiler/Private/KismetCompiler.cpp
+++ b/Engine/Source/Editor/KismetCompiler/Private/KismetCompiler.cpp
@@ -74,11 +74,11 @@ FKismetCompilerContext::FKismetCompilerContext(UBlueprint* SourceSketch, FCompil
 
 	MacroSpawnX = MinimumSpawnX;
 	MacroSpawnY = -2000;
-
-	VectorStruct = FindObjectChecked<UScriptStruct>(UObject::StaticClass(), TEXT("Vector"));
-	RotatorStruct = FindObjectChecked<UScriptStruct>(UObject::StaticClass(), TEXT("Rotator"));
-	TransformStruct = FindObjectChecked<UScriptStruct>(UObject::StaticClass(), TEXT("Transform"));
-	LinearColorStruct = FindObjectChecked<UScriptStruct>(UObject::StaticClass(), TEXT("LinearColor"));
+	
+	VectorStruct = GetBaseStructure(TEXT("Vector"));
+	RotatorStruct = GetBaseStructure(TEXT("Rotator"));
+	TransformStruct = GetBaseStructure(TEXT("Transform"));
+	LinearColorStruct = GetBaseStructure(TEXT("LinearColor"));
 }
 
 FKismetCompilerContext::~FKismetCompilerContext()
diff --git a/Engine/Source/Editor/KismetCompiler/Private/KismetCompilerCppBackend.cpp b/Engine/Source/Editor/KismetCompiler/Private/KismetCompilerCppBackend.cpp
--- a/Engine/Source/Editor/KismetCompiler/Private/KismetCompilerCppBackend.cpp
+++ b/Engine/Source/Editor/KismetCompiler/Private/KismetCompilerCppBackend.cpp
@@ -76,11 +76,11 @@ public:
 		, MessageLog(InContext.MessageLog)
 		, CompilerContext(InContext)
 	{
-		extern UScriptStruct* Z_Construct_UScriptStruct_UObject_FVector();
-		VectorStruct = Z_Construct_UScriptStruct_UObject_FVector();
-		RotatorStruct = FindObjectChecked<UScriptStruct>(UObject::StaticClass(), TEXT("Rotator"));
-		TransformStruct = FindObjectChecked<UScriptStruct>(UObject::StaticClass(), TEXT("Transform"));
-		LinearColorStruct = FindObjectChecked<UScriptStruct>(UObject::StaticClass(), TEXT("LinearColor"));
+		extern UScriptStruct* Z_Construct_UScriptStruct_FVector();
+		VectorStruct = Z_Construct_UScriptStruct_FVector();
+		RotatorStruct = GetBaseStructure(TEXT("Rotator"));
+		TransformStruct = GetBaseStructure(TEXT("Transform"));
+		LinearColorStruct = GetBaseStructure(TEXT("LinearColor"));
 		LatentInfoStruct = FLatentActionInfo::StaticStruct();
 	}
 
diff --git a/Engine/Source/Editor/KismetCompiler/Private/KismetCompilerVMBackend.cpp b/Engine/Source/Editor/KismetCompiler/Private/KismetCompilerVMBackend.cpp
--- a/Engine/Source/Editor/KismetCompiler/Private/KismetCompilerVMBackend.cpp
+++ b/Engine/Source/Editor/KismetCompiler/Private/KismetCompilerVMBackend.cpp
@@ -318,9 +318,9 @@ public:
 		, bIsUbergraph(bInIsUbergraph)
 		, ReturnStatement(InReturnStatement)
 	{
-		VectorStruct = FindObjectChecked<UScriptStruct>(UObject::StaticClass(), TEXT("Vector"));
-		RotatorStruct = FindObjectChecked<UScriptStruct>(UObject::StaticClass(), TEXT("Rotator"));
-		TransformStruct = FindObjectChecked<UScriptStruct>(UObject::StaticClass(), TEXT("Transform"));
+		VectorStruct = GetBaseStructure(TEXT("Vector"));
+		RotatorStruct = GetBaseStructure(TEXT("Rotator"));
+		TransformStruct = GetBaseStructure(TEXT("Transform"));
 		LatentInfoStruct = FLatentActionInfo::StaticStruct();
 	}
 
diff --git a/Engine/Source/Editor/MaterialEditor/Private/MaterialInstanceEditor.cpp b/Engine/Source/Editor/MaterialEditor/Private/MaterialInstanceEditor.cpp
--- a/Engine/Source/Editor/MaterialEditor/Private/MaterialInstanceEditor.cpp
+++ b/Engine/Source/Editor/MaterialEditor/Private/MaterialInstanceEditor.cpp
@@ -733,7 +733,7 @@ void FMaterialInstanceEditor::DrawSamplerWarningStrings(FCanvas* Canvas, int32&
 		{
 			UFont* FontToUse = GEngine->GetTinyFont();
 			const int32 SpacingBetweenLines = 13;
-			UEnum* SamplerTypeEnum = FindObject<UEnum>( NULL, TEXT("/Script/Engine.EngineTypes.EMaterialSamplerType") );
+			UEnum* SamplerTypeEnum = FindObject<UEnum>( NULL, TEXT("/Script/Engine.EMaterialSamplerType") );
 			check( SamplerTypeEnum );
 
 			const int32 GroupCount = MaterialEditorInstance->ParameterGroups.Num();
diff --git a/Engine/Source/Editor/PropertyEditor/Private/UserInterface/PropertyEditor/SPropertyEditorDateTime.cpp b/Engine/Source/Editor/PropertyEditor/Private/UserInterface/PropertyEditor/SPropertyEditorDateTime.cpp
--- a/Engine/Source/Editor/PropertyEditor/Private/UserInterface/PropertyEditor/SPropertyEditorDateTime.cpp
+++ b/Engine/Source/Editor/PropertyEditor/Private/UserInterface/PropertyEditor/SPropertyEditorDateTime.cpp
@@ -40,8 +40,8 @@ bool SPropertyEditorDateTime::Supports( const TSharedRef< FPropertyEditor >& InP
 	if (Property->IsA(UStructProperty::StaticClass()))
 	{
 		const UStructProperty* StructProp = Cast<const UStructProperty>(Property);
-		extern UScriptStruct* Z_Construct_UScriptStruct_UObject_FDateTime();	// It'd be really nice if StaticStruct() worked on types declared in Object.h
-		if (Z_Construct_UScriptStruct_UObject_FDateTime() == StructProp->Struct)
+		extern UScriptStruct* Z_Construct_UScriptStruct_FDateTime();	// It'd be really nice if StaticStruct() worked on types declared in Object.h
+		if (Z_Construct_UScriptStruct_FDateTime() == StructProp->Struct)
 		{
 			return true;
 		}
diff --git a/Engine/Source/Editor/UnrealEd/Private/Kismet2/BlueprintEditorUtils.cpp b/Engine/Source/Editor/UnrealEd/Private/Kismet2/BlueprintEditorUtils.cpp
--- a/Engine/Source/Editor/UnrealEd/Private/Kismet2/BlueprintEditorUtils.cpp
+++ b/Engine/Source/Editor/UnrealEd/Private/Kismet2/BlueprintEditorUtils.cpp
@@ -4459,9 +4459,9 @@ void FBlueprintEditorUtils::ImportKismetDefaultValueToProperty(UEdGraphPin* Sour
 
 	if (UStructProperty* StructProperty = Cast<UStructProperty>(DestinationProperty))
 	{
-		static UScriptStruct* VectorStruct = FindObjectChecked<UScriptStruct>(UObject::StaticClass(), TEXT("Vector"));
-		static UScriptStruct* RotatorStruct = FindObjectChecked<UScriptStruct>(UObject::StaticClass(), TEXT("Rotator"));
-		static UScriptStruct* TransformStruct = FindObjectChecked<UScriptStruct>(UObject::StaticClass(), TEXT("Transform"));
+		static UScriptStruct* VectorStruct = GetBaseStructure(TEXT("Vector"));
+		static UScriptStruct* RotatorStruct = GetBaseStructure(TEXT("Rotator"));
+		static UScriptStruct* TransformStruct = GetBaseStructure(TEXT("Transform"));
 
 		if (StructProperty->Struct == VectorStruct)
 		{
@@ -4495,9 +4495,9 @@ void FBlueprintEditorUtils::ExportPropertyToKismetDefaultValue(UEdGraphPin* Targ
 
 	if (UStructProperty* StructProperty = Cast<UStructProperty>(SourceProperty))
 	{
-		static UScriptStruct* VectorStruct = FindObjectChecked<UScriptStruct>(UObject::StaticClass(), TEXT("Vector"));
-		static UScriptStruct* RotatorStruct = FindObjectChecked<UScriptStruct>(UObject::StaticClass(), TEXT("Rotator"));
-		static UScriptStruct* TransformStruct = FindObjectChecked<UScriptStruct>(UObject::StaticClass(), TEXT("Transform"));
+		static UScriptStruct* VectorStruct = GetBaseStructure(TEXT("Vector"));
+		static UScriptStruct* RotatorStruct = GetBaseStructure(TEXT("Rotator"));
+		static UScriptStruct* TransformStruct = GetBaseStructure(TEXT("Transform"));
 
 		if (StructProperty->Struct == VectorStruct)
 		{
@@ -6579,10 +6579,10 @@ bool FBlueprintEditorUtils::PropertyValueFromString(const UProperty* Property, c
 	}
 	else 
 	{
-		static UScriptStruct* VectorStruct = FindObjectChecked<UScriptStruct>(UObject::StaticClass(), TEXT("Vector"));
-		static UScriptStruct* RotatorStruct = FindObjectChecked<UScriptStruct>(UObject::StaticClass(), TEXT("Rotator"));
-		static UScriptStruct* TransformStruct = FindObjectChecked<UScriptStruct>(UObject::StaticClass(), TEXT("Transform"));
-		static UScriptStruct* LinearColorStruct = FindObjectChecked<UScriptStruct>(UObject::StaticClass(), TEXT("LinearColor"));
+		static UScriptStruct* VectorStruct = GetBaseStructure(TEXT("Vector"));
+		static UScriptStruct* RotatorStruct = GetBaseStructure(TEXT("Rotator"));
+		static UScriptStruct* TransformStruct = GetBaseStructure(TEXT("Transform"));
+		static UScriptStruct* LinearColorStruct = GetBaseStructure(TEXT("LinearColor"));
 
 		const UStructProperty* StructProperty = CastChecked<UStructProperty>(Property);
 
@@ -6636,10 +6636,10 @@ bool FBlueprintEditorUtils::PropertyValueToString(const UProperty* Property, con
 	OutForm.Empty();
 	if (Property->IsA(UStructProperty::StaticClass()))
 	{
-		static UScriptStruct* VectorStruct = FindObjectChecked<UScriptStruct>(UObject::StaticClass(), TEXT("Vector"));
-		static UScriptStruct* RotatorStruct = FindObjectChecked<UScriptStruct>(UObject::StaticClass(), TEXT("Rotator"));
-		static UScriptStruct* TransformStruct = FindObjectChecked<UScriptStruct>(UObject::StaticClass(), TEXT("Transform"));
-		static UScriptStruct* LinearColorStruct = FindObjectChecked<UScriptStruct>(UObject::StaticClass(), TEXT("LinearColor"));
+		static UScriptStruct* VectorStruct = GetBaseStructure(TEXT("Vector"));
+		static UScriptStruct* RotatorStruct = GetBaseStructure(TEXT("Rotator"));
+		static UScriptStruct* TransformStruct = GetBaseStructure(TEXT("Transform"));
+		static UScriptStruct* LinearColorStruct = GetBaseStructure(TEXT("LinearColor"));
 
 		const UStructProperty* StructProperty = CastChecked<UStructProperty>(Property);
 
diff --git a/Engine/Source/Programs/UnrealHeaderTool/Private/ClassMaps.cpp b/Engine/Source/Programs/UnrealHeaderTool/Private/ClassMaps.cpp
--- a/Engine/Source/Programs/UnrealHeaderTool/Private/ClassMaps.cpp
+++ b/Engine/Source/Programs/UnrealHeaderTool/Private/ClassMaps.cpp
@@ -4,18 +4,23 @@
 
 #include "ClassMaps.h"
 
-TMap<UClass*, FString>                   GClassStrippedHeaderTextMap;
-TMap<UClass*, FString>                   GClassSourceFileMap;
-TMap<UClass*, int32>                     GClassDeclarationLineNumber;
-TMap<UClass*, FClassHeaderInfo>          GClassGeneratedFileMap;
-TMap<UClass*, TUniqueObj<TArray<FName>>> GClassDependentOnMap;
-TMap<UClass*, FString>                   GClassHeaderNameWithNoPathMap;
-TMap<UClass*, FString>                   GClassModuleRelativePathMap;
-TMap<UClass*, FString>                   GClassIncludePathMap;
-TSet<UClass*>                            GPublicClassSet;
-TSet<UClass*>                            GExportedClasses;
-TMap<UProperty*, FString>                GArrayDimensions;
-TMap<UPackage*,  const FManifestModule*> GPackageToManifestModuleMap;
-TMap<UField*, uint32>                    GGeneratedCodeCRCs;
-TMap<UEnum*,  EPropertyType>             GEnumUnderlyingTypes;
-TSet<UClass*>                            GTemporaryClasses;
\ No newline at end of file
+TMap<FString, TSharedRef<FUnrealSourceFile> >						GUnrealSourceFilesMap;
+TMap<UField*, TSharedRef<FUnrealTypeDefinitionInfo> >				GTypeDefinitionInfoMap;
+TMap<UClass*, FString>												GClassStrippedHeaderTextMap;
+TMap<UClass*, FString>												GClassHeaderNameWithNoPathMap;
+TSet<UClass*>														GPublicClassSet;
+TSet<FUnrealSourceFile*>											GPublicSourceFileSet;
+TSet<FUnrealSourceFile*>											GExportedSourceFiles;
+TMap<UProperty*, FString>											GArrayDimensions;
+TMap<UPackage*,  const FManifestModule*>							GPackageToManifestModuleMap;
+TMap<UField*, uint32>												GGeneratedCodeCRCs;
+TMap<UEnum*,  EPropertyType>										GEnumUnderlyingTypes;
+
+TSharedRef<FUnrealTypeDefinitionInfo> AddTypeDefinition(FUnrealSourceFile& SourceFile, UStruct* Struct, int32 Line)
+{
+	TSharedRef<FUnrealTypeDefinitionInfo> DefinitionInfo = MakeShareable(new FUnrealTypeDefinitionInfo(SourceFile, Line));
+
+	GTypeDefinitionInfoMap.Add(Struct, DefinitionInfo);
+
+	return DefinitionInfo;
+}
\ No newline at end of file
diff --git a/Engine/Source/Programs/UnrealHeaderTool/Private/ClassMaps.h b/Engine/Source/Programs/UnrealHeaderTool/Private/ClassMaps.h
--- a/Engine/Source/Programs/UnrealHeaderTool/Private/ClassMaps.h
+++ b/Engine/Source/Programs/UnrealHeaderTool/Private/ClassMaps.h
@@ -3,44 +3,25 @@
 #pragma once
 
 #include "UniqueObj.h"
+#include "ParserClass.h"
+#include "Scope.h"
+#include "HeaderProvider.h"
 
-/** Information about class source and generated headers */
-struct FClassHeaderInfo
-{
-	FString SourceFilename;
-	FString GeneratedFilename;
-	bool bHasChanged;
-	
-	FClassHeaderInfo()
-		: bHasChanged(false)
-	{
-	}
-
-	explicit FClassHeaderInfo(const FString& InSourceFilename)
-		: SourceFilename(InSourceFilename)
-		, bHasChanged(false)
-	{
-	}
-};
+#include "UnrealSourceFile.h"
 
 struct FManifestModule;
 
-extern TMap<UClass*, FString>                   GClassStrippedHeaderTextMap;
-extern TMap<UClass*, FString>                   GClassSourceFileMap;
-extern TMap<UClass*, int32>                     GClassDeclarationLineNumber;
-extern TMap<UClass*, FClassHeaderInfo>          GClassGeneratedFileMap;
-extern TMap<UClass*, TUniqueObj<TArray<FName>>> GClassDependentOnMap;
-extern TMap<UClass*, FString>                   GClassHeaderNameWithNoPathMap;
-extern TMap<UClass*, FString>                   GClassModuleRelativePathMap;
-extern TMap<UClass*, FString>                   GClassIncludePathMap;
-extern TSet<UClass*>                            GExportedClasses;
-extern TSet<UClass*>                            GPublicClassSet;
-extern TMap<UProperty*, FString>                GArrayDimensions;
-extern TMap<UPackage*,  const FManifestModule*> GPackageToManifestModuleMap;
-extern TMap<UField*, uint32>                    GGeneratedCodeCRCs;
-extern TMap<UEnum*,  EPropertyType>             GEnumUnderlyingTypes;
-extern TSet<UClass*>                            GTemporaryClasses;
-
+extern TMap<FString, TSharedRef<FUnrealSourceFile> >					GUnrealSourceFilesMap;
+extern TMap<UField*, TSharedRef<FUnrealTypeDefinitionInfo> >			GTypeDefinitionInfoMap;
+extern TMap<UClass*, FString>											GClassStrippedHeaderTextMap;
+extern TMap<UClass*, FString>											GClassHeaderNameWithNoPathMap;
+extern TSet<FUnrealSourceFile*>											GExportedSourceFiles;
+extern TSet<UClass*>													GPublicClassSet;
+extern TSet<FUnrealSourceFile*>											GPublicSourceFileSet;
+extern TMap<UProperty*, FString>										GArrayDimensions;
+extern TMap<UPackage*,  const FManifestModule*>							GPackageToManifestModuleMap;
+extern TMap<UField*, uint32>											GGeneratedCodeCRCs;
+extern TMap<UEnum*,  EPropertyType>										GEnumUnderlyingTypes;
 
 /** Types access specifiers. */
 enum EAccessSpecifier
@@ -51,3 +32,14 @@ enum EAccessSpecifier
 	ACCESS_Protected,
 	ACCESS_Num,
 };
+
+/**
+ * Add type definition info to global map.
+ *
+ * @param SourceFile SourceFile in which type was defined.
+ * @param Struct Defined type.
+ * @param Line Line on which the type was defined.
+ *
+ * @returns Type definition info.
+ */
+TSharedRef<FUnrealTypeDefinitionInfo> AddTypeDefinition(FUnrealSourceFile& SourceFile, UStruct* Struct, int32 Line);
\ No newline at end of file
diff --git a/Engine/Source/Programs/UnrealHeaderTool/Private/Classes.cpp b/Engine/Source/Programs/UnrealHeaderTool/Private/Classes.cpp
--- a/Engine/Source/Programs/UnrealHeaderTool/Private/Classes.cpp
+++ b/Engine/Source/Programs/UnrealHeaderTool/Private/Classes.cpp
@@ -88,49 +88,6 @@ bool FClasses::IsDependentOn(const FClass* Suspect, const FClass* Source, TSet<c
 		VisitedDpendencies.Add(Suspect);
 	}
 
-	// Now consider all dependents of the suspect. If any of them are dependent on the source, the
-	// suspect is too.
-	for (auto It : Suspect->GetDependentNames())
-	{
-		auto DependentName = It.ToString();
-		const FClass* Dependency = FindClass(*DependentName);
-
-		if (!Dependency)
-		{
-			// Check if DependentName is a header file. If so, rerun the check on name without ".h" extension.
-			if (DependentName.EndsWith(TEXT(".h")))
-			{
-				DependentName.RemoveFromEnd(TEXT(".h"));
-				Dependency = FindClass(*DependentName);
-			}
-
-			// If still no class found, we're out of luck.
-			if (!Dependency)
-			{
-				continue;
-			}				
-		}
-
-		// the parser disallows declaring the parent class as a dependency, so the only way this could occur is
-		// if the parent for a native class has been changed (which causes the new parent to be inserted as a dependency),
-		// if this is the case, skip it or we'll go into a loop
-		if (Suspect->GetSuperClass() == Dependency)
-		{
-			continue;
-		}
-
-		// Ignore inter-module dependencies, since modules should be self-contained when they are compiled.
-		if (Dependency->GetOutermost() != Suspect->GetOutermost())
-		{
-			continue;
-		}
-
-		if (Dependency == Source || IsDependentOn(Dependency, Source, VisitedDpendencies))
-		{
-			return true;
-		}
-	}
-
 	return false;
 }
 
