#github PR #754: When creating Select-nodes of the Data-pins new node is connected via Index-node (Contributed by [REDACTED_PERSON])

https://github.com/EpicGames/UnrealEngine/pull/754

diff --git a/Engine/Source/Editor/BlueprintGraph/Classes/K2Node_Select.h b/Engine/Source/Editor/BlueprintGraph/Classes/K2Node_Select.h
--- a/Engine/Source/Editor/BlueprintGraph/Classes/K2Node_Select.h
+++ b/Engine/Source/Editor/BlueprintGraph/Classes/K2Node_Select.h
@@ -42,6 +42,7 @@ class UK2Node_Select : public UK2Node, public INodeDependingOnEnumInterface
 	virtual FText GetNodeTitle(ENodeTitleType::Type TitleType) const override;
 	virtual void NodeConnectionListChanged() override;
 	virtual void PinTypeChanged(UEdGraphPin* Pin) override;
+	virtual void AutowireNewNode(UEdGraphPin* FromPin) override;
 	virtual void PostPasteNode() override;
 	virtual FName GetPaletteIcon(FLinearColor& OutColor) const override{ return TEXT("GraphEditor.Select_16x"); }
 	// End UEdGraphNode interface
diff --git a/Engine/Source/Editor/BlueprintGraph/Private/K2Node_Select.cpp b/Engine/Source/Editor/BlueprintGraph/Private/K2Node_Select.cpp
--- a/Engine/Source/Editor/BlueprintGraph/Private/K2Node_Select.cpp
+++ b/Engine/Source/Editor/BlueprintGraph/Private/K2Node_Select.cpp
@@ -260,6 +260,27 @@ void UK2Node_Select::AllocateDefaultPins()
 	Super::AllocateDefaultPins();
 }
 
+void UK2Node_Select::AutowireNewNode(UEdGraphPin* FromPin)
+{
+	// Attempt to autowire to the index pin as users generally drag off of something intending to use
+	// it as an index in a select statement rather than an arbitrary entry:
+	const UEdGraphSchema_K2* K2Schema = CastChecked<UEdGraphSchema_K2>(GetSchema());
+	UEdGraphPin* IndexPin = GetIndexPin();
+	ECanCreateConnectionResponse ConnectResponse = K2Schema->CanCreateConnection(FromPin, IndexPin).Response;
+	if (ConnectResponse == ECanCreateConnectionResponse::CONNECT_RESPONSE_MAKE)
+	{
+		if (K2Schema->TryCreateConnection(FromPin, IndexPin))
+		{
+			FromPin->GetOwningNode()->NodeConnectionListChanged();
+			this->NodeConnectionListChanged();
+			return;
+		}
+	}
+
+	// No connection made, just use default autowire logic:
+	Super::AutowireNewNode(FromPin);
+}
+
 FText UK2Node_Select::GetTooltipText() const
 {
 	return LOCTEXT("SelectNodeTooltip", "Return the option at Index, (first option is indexed at 0)");
