#ue4 Fix issues with adding triggers to matinee particle tracks
- We now no longer register a new callback delegate for temporary nodes.  this was causing an n^2 growth per keyframe added
- Check for uniqueness of event name, so that matinee nodes will no longer break

diff --git a/Engine/Source/Editor/BlueprintGraph/Private/K2Node_MatineeController.cpp b/Engine/Source/Editor/BlueprintGraph/Private/K2Node_MatineeController.cpp
--- a/Engine/Source/Editor/BlueprintGraph/Private/K2Node_MatineeController.cpp
+++ b/Engine/Source/Editor/BlueprintGraph/Private/K2Node_MatineeController.cpp
@@ -13,9 +13,13 @@
 UK2Node_MatineeController::UK2Node_MatineeController(const FObjectInitializer& ObjectInitializer)
 	: Super(ObjectInitializer)
 {
-	FMatineeDelegates::Get().OnEventKeyframeAdded.AddUObject(this, &UK2Node_MatineeController::OnEventKeyframeAdded);
-	FMatineeDelegates::Get().OnEventKeyframeRenamed.AddUObject(this, &UK2Node_MatineeController::OnEventKeyframeRenamed);
-	FMatineeDelegates::Get().OnEventKeyframeRemoved.AddUObject(this, &UK2Node_MatineeController::OnEventKeyframeRemoved);
+	if (GetOutermost() != GetTransientPackage())
+	{
+		// Register the delegate ONLY if this isn't the transient package (temp node during compilation)
+		FMatineeDelegates::Get().OnEventKeyframeAdded.AddUObject(this, &UK2Node_MatineeController::OnEventKeyframeAdded);
+		FMatineeDelegates::Get().OnEventKeyframeRenamed.AddUObject(this, &UK2Node_MatineeController::OnEventKeyframeRenamed);
+		FMatineeDelegates::Get().OnEventKeyframeRemoved.AddUObject(this, &UK2Node_MatineeController::OnEventKeyframeRemoved);
+	}
 }
 
 void UK2Node_MatineeController::BeginDestroy()
@@ -200,13 +204,17 @@ void UK2Node_MatineeController::OnEventKeyframeAdded(const AMatineeActor* InMati
 {
 	if ( MatineeActor == InMatineeActor )
 	{
-		const UEdGraphSchema_K2* K2Schema = GetDefault<UEdGraphSchema_K2>();
+		// Only add unique event names to the controller node
+		if (!FindPin(InPinName.ToString()))
+		{
+			const UEdGraphSchema_K2* K2Schema = GetDefault<UEdGraphSchema_K2>();
 
-		// Add one to the index as we insert "finished" at index 0
-		CreatePin(EGPD_Output, K2Schema->PC_Exec, TEXT(""), NULL, false, false, InPinName.ToString(), false, InIndex + 1);
+			// Add one to the index as we insert "finished" at index 0
+			CreatePin(EGPD_Output, K2Schema->PC_Exec, TEXT(""), NULL, false, false, InPinName.ToString(), false, InIndex + 1);
 
-		// Update and refresh the blueprint
-		FBlueprintEditorUtils::MarkBlueprintAsStructurallyModified(GetBlueprint());
+			// Update and refresh the blueprint
+			FBlueprintEditorUtils::MarkBlueprintAsStructurallyModified(GetBlueprint());
+		}
 	}
 }
 
