[INTEGRATE] Change 2412854 by [REDACTED_PERSON][REDACTED_HANDLE] on 2015/01/20 16:48:32

	BP - Empty sections now off by default.

diff --git a/Engine/Source/Editor/BlueprintGraph/Private/BlueprintEditorSettings.cpp b/Engine/Source/Editor/BlueprintGraph/Private/BlueprintEditorSettings.cpp
--- a/Engine/Source/Editor/BlueprintGraph/Private/BlueprintEditorSettings.cpp
+++ b/Engine/Source/Editor/BlueprintGraph/Private/BlueprintEditorSettings.cpp
@@ -19,7 +19,7 @@ UBlueprintEditorSettings::UBlueprintEditorSettings(const FObjectInitializer& Obj
 	, bAutoCastObjectConnections(false)
 	, bUseLegacyMenuingSystem(false)
 	, bShowInheritedVariables(false)
-	, bShowEmptySections(true)
+	, bShowEmptySections(false)
 	// Compiler Settings
 	, SaveOnCompile(SoC_Never)
 	, bJumpToNodeErrors(false)
