[INTEGRATE] Change 2415387 by [REDACTED_PERSON][REDACTED_HANDLE] on 2015/01/22 11:27:02

	Inherited variables nodes are now correctly placed without a "Target" pin without first compiling the parent Blueprint.

	Dragging off a "self" node and placing a variable will no longer place a node that will refresh and remove the "Target" pin.

diff --git a/Engine/Source/Editor/BlueprintGraph/Private/BlueprintVariableNodeSpawner.cpp b/Engine/Source/Editor/BlueprintGraph/Private/BlueprintVariableNodeSpawner.cpp
--- a/Engine/Source/Editor/BlueprintGraph/Private/BlueprintVariableNodeSpawner.cpp
+++ b/Engine/Source/Editor/BlueprintGraph/Private/BlueprintVariableNodeSpawner.cpp
@@ -75,6 +75,10 @@ UBlueprintVariableNodeSpawner* UBlueprintVariableNodeSpawner::Create(TSubclassOf
 			UClass* OwnerClass = Property->GetOwnerClass();
 
 			// We need to use a generated class instead of a skeleton class for IsChildOf, so if the OwnerClass has a Blueprint, grab the GeneratedClass
+			if(!UBlueprintGeneratedClass::CompileSkeletonClassesInheritSkeletonClasses() && OwnerClass)
+			{
+				OwnerClass = OwnerClass->GetAuthoritativeClass();
+			}
 			bool const bIsSelfContext = Blueprint->SkeletonGeneratedClass->GetAuthoritativeClass() == OwnerClass || Blueprint->SkeletonGeneratedClass->IsChildOf(OwnerClass);
 
 			UK2Node_Variable* VarNode = CastChecked<UK2Node_Variable>(NewNode);
diff --git a/Engine/Source/Editor/BlueprintGraph/Private/K2Node_Variable.cpp b/Engine/Source/Editor/BlueprintGraph/Private/K2Node_Variable.cpp
--- a/Engine/Source/Editor/BlueprintGraph/Private/K2Node_Variable.cpp
+++ b/Engine/Source/Editor/BlueprintGraph/Private/K2Node_Variable.cpp
@@ -290,7 +290,22 @@ UProperty* UK2Node_Variable::GetPropertyForVariable() const
 	const FName VarName = GetVarName();
 	UEdGraphPin* VariablePin = FindPin(GetVarNameString());
 
-	UProperty* VariableProperty = VariableReference.ResolveMember<UProperty>(this);
+	UProperty* VariableProperty = nullptr;
+
+	// Need to look at parent Blueprint's skeleton classes to see if the variable property can resolve there.
+	UClass* CurrentGeneratedClass = GetBlueprint()->GeneratedClass;
+	while(CurrentGeneratedClass && VariableProperty == nullptr)
+	{
+		if(UBlueprint* CurrentBlueprint = Cast<UBlueprint>(CurrentGeneratedClass->ClassGeneratedBy))
+		{
+			VariableProperty = VariableReference.ResolveMember<UProperty>(CurrentBlueprint->SkeletonGeneratedClass);
+			CurrentGeneratedClass = CurrentBlueprint->ParentClass;
+		}
+		else
+		{
+			break;
+		}
+	}
 
 	// if the variable has been deprecated, don't use it
 	if(VariableProperty != NULL)
@@ -570,9 +585,9 @@ void UK2Node_Variable::AutowireNewNode(UEdGraphPin* FromPin)
 		if(FromPin->Direction == EGPD_Output)
 		{
 			// If the source pin has a valid PinSubCategoryObject, we might be doing BP Comms, so check if it is a class
-			if(FromPin->PinType.PinSubCategoryObject.IsValid() && FromPin->PinType.PinSubCategoryObject->IsA(UClass::StaticClass()))
+			if ((FromPin->PinType.PinSubCategoryObject.IsValid() && FromPin->PinType.PinSubCategoryObject->IsA(UClass::StaticClass())) || FromPin->PinType.PinSubCategory == K2Schema->PSC_Self)
 			{
-				UProperty* VariableProperty = VariableReference.ResolveMember<UProperty>(this);
+				UProperty* VariableProperty = GetPropertyForVariable();
 				if(VariableProperty)
 				{
 					UClass* PropertyOwner = VariableProperty->GetOwnerClass();
@@ -581,8 +596,11 @@ void UK2Node_Variable::AutowireNewNode(UEdGraphPin* FromPin)
 						PropertyOwner = PropertyOwner->GetAuthoritativeClass();
 					}
 
+					// If the pin is a self reference, check if the UProperty is valid for the current Blueprint.
+					const bool bIsSelfReferenceValid = (FromPin->PinType.PinSubCategory == K2Schema->PSC_Self)? GetBlueprint()->GeneratedClass->IsChildOf(PropertyOwner) : false;
+
 					// BP Comms is highly likely at this point, if the source pin's type is a child of the variable's owner class, let's conform the "Target" pin
-					if(FromPin->PinType.PinSubCategoryObject == PropertyOwner || dynamic_cast<UClass*>(FromPin->PinType.PinSubCategoryObject.Get())->IsChildOf(PropertyOwner))
+					if(FromPin->PinType.PinSubCategoryObject == PropertyOwner || dynamic_cast<UClass*>(FromPin->PinType.PinSubCategoryObject.Get())->IsChildOf(PropertyOwner) || bIsSelfReferenceValid)
 					{
 						UEdGraphPin* TargetPin = FindPin(K2Schema->PN_Self);
 						TargetPin->PinType.PinSubCategoryObject = PropertyOwner;
