Static functions (such as those from a function library) will not be filtered from the Blueprint context menu based on the Blueprint's hidden category metadata.

diff --git a/Engine/Source/Editor/BlueprintGraph/Private/BlueprintActionFilter.cpp b/Engine/Source/Editor/BlueprintGraph/Private/BlueprintActionFilter.cpp
--- a/Engine/Source/Editor/BlueprintGraph/Private/BlueprintActionFilter.cpp
+++ b/Engine/Source/Editor/BlueprintGraph/Private/BlueprintActionFilter.cpp
@@ -742,7 +742,8 @@ static bool BlueprintActionFilterImpl::IsFieldCategoryHidden(FBlueprintActionFil
 	{
 		auto IsFunctionHiddenLambda = [](UClass* Class, UFunction const* Function)->bool
 		{
-			return FObjectEditorUtils::IsFunctionHiddenFromClass(Function, Class->GetAuthoritativeClass());
+			// Only hide functions that are not static
+			return !Function->HasAnyFunctionFlags(FUNC_Static) && FObjectEditorUtils::IsFunctionHiddenFromClass(Function, Class->GetAuthoritativeClass());
 		};
 		IsFieldHiddenDelegate = FIsFieldHiddenDelegate::CreateStatic(IsFunctionHiddenLambda, Function);
 	}
