[INTEGRATE] Change 2423646 by [REDACTED_PERSON][REDACTED_HANDLE] on 2015/01/29 08:21:07

	Proper handling of copy/paste for legacy Event Graph nodes that override BPIE class methods declared as 'const'.

	change summary:
	- set up SMyBlueprint to refresh itself whenever the Blueprint it's viewing is modified (including non-structurally); fixes display of "stale" function override state so users can no longer trigger invalid actions via context menu
	- modified UK2Node_Event::CanPasteHere() to also consider implemented function overrides within the Blueprint hierarchy that are also declared as 'const' in native C++ code; prevents injection of legacy event nodes into a Blueprint that has already overridden the function the "new" way via a function graph implementation with a read-only 'self' context somewhere within its inheritance hierarchy

diff --git a/Engine/Source/Editor/BlueprintGraph/Private/K2Node_Event.cpp b/Engine/Source/Editor/BlueprintGraph/Private/K2Node_Event.cpp
--- a/Engine/Source/Editor/BlueprintGraph/Private/K2Node_Event.cpp
+++ b/Engine/Source/Editor/BlueprintGraph/Private/K2Node_Event.cpp
@@ -445,6 +445,24 @@ bool UK2Node_Event::CanPasteHere(const UEdGraph* TargetGraph) const
 								&& ExistingEventNodes[i]->EventSignatureClass == EventSignatureClass;
 						}
 
+						// We need to also check for 'const' BPIE methods that might already be implemented as functions with a read-only 'self' context (these were previously implemented as events)
+						if(!bDisallowPaste)
+						{
+							TArray<UBlueprint*> ParentBPStack;
+							UBlueprint::GetBlueprintHierarchyFromClass(Blueprint->SkeletonGeneratedClass, ParentBPStack);
+							for(auto BPStackIt = ParentBPStack.CreateConstIterator(); BPStackIt && !bDisallowPaste; ++BPStackIt)
+							{
+								TArray<UK2Node_FunctionEntry*> ExistingFunctionEntryNodes;
+								FBlueprintEditorUtils::GetAllNodesOfClass<UK2Node_FunctionEntry>(*BPStackIt, ExistingFunctionEntryNodes);
+								for(auto NodeIt = ExistingFunctionEntryNodes.CreateConstIterator(); NodeIt && !bDisallowPaste; ++NodeIt)
+								{
+									UK2Node_FunctionEntry* ExistingFunctionEntryNode = *NodeIt;
+									bDisallowPaste = ExistingFunctionEntryNode->bEnforceConstCorrectness
+										&& ExistingFunctionEntryNode->SignatureName == EventSignatureName;
+								}
+							}
+						}
+
 						if(!bDisallowPaste)
 						{
 							// If the signature class is not implemented by the Blueprint parent class or an interface, don't paste this event
diff --git a/Engine/Source/Editor/Kismet/Private/SMyBlueprint.cpp b/Engine/Source/Editor/Kismet/Private/SMyBlueprint.cpp
--- a/Engine/Source/Editor/Kismet/Private/SMyBlueprint.cpp
+++ b/Engine/Source/Editor/Kismet/Private/SMyBlueprint.cpp
@@ -131,6 +131,8 @@ public:
 
 void SMyBlueprint::Construct(const FArguments& InArgs, TWeakPtr<FBlueprintEditor> InBlueprintEditor, const UBlueprint* InBlueprint )
 {
+	bNeedsRefresh = false;
+
 	BlueprintEditorPtr = InBlueprintEditor;
 	EdGraph = nullptr;
 	
@@ -404,6 +406,18 @@ void SMyBlueprint::Construct(const FArguments& InArgs, TWeakPtr<FBlueprintEditor
 	ExpandedSections.Add(NodeSectionID::LOCAL_VARIABLE, true);
 
 	GraphActionMenu->SetSectionExpansion(ExpandedSections);
+
+	FCoreUObjectDelegates::OnObjectPropertyChanged.AddRaw(this, &SMyBlueprint::OnObjectPropertyChanged);
+}
+
+void SMyBlueprint::Tick(const FGeometry& AllottedGeometry, const double InCurrentTime, const float InDeltaTime)
+{
+	SCompoundWidget::Tick(AllottedGeometry, InCurrentTime, InDeltaTime);
+
+	if(bNeedsRefresh)
+	{
+		Refresh();
+	}
 }
 
 void SMyBlueprint::OnCategoryNameCommitted(const FText& InNewText, ETextCommit::Type InTextCommit, TWeakPtr< FGraphActionNode > InAction )
@@ -700,6 +714,8 @@ bool SMyBlueprint::CanRequestRenameOnActionNode(TWeakPtr<FGraphActionNode> InSel
 
 void SMyBlueprint::Refresh()
 {
+	bNeedsRefresh = false;
+
 	GraphActionMenu->RefreshAllActions(true);
 	
 	bool bLocalActionsAreVisible = (GetLocalActionsListVisibility() == EVisibility::Visible);
diff --git a/Engine/Source/Editor/Kismet/Private/SMyBlueprint.h b/Engine/Source/Editor/Kismet/Private/SMyBlueprint.h
--- a/Engine/Source/Editor/Kismet/Private/SMyBlueprint.h
+++ b/Engine/Source/Editor/Kismet/Private/SMyBlueprint.h
@@ -57,6 +57,9 @@ public:
 	void Construct(const FArguments& InArgs, TWeakPtr<FBlueprintEditor> InBlueprintEditor, const UBlueprint* InBlueprint = nullptr);
 	void SetInspector( TSharedPtr<SKismetInspector> InInspector ) { Inspector = InInspector ; }
 
+	/* SWidget interface */
+	virtual void Tick(const FGeometry& AllottedGeometry, const double InCurrentTime, const float InDeltaTime);
+
 	/* Reset the last pin type settings to default. */
 	void ResetLastPinType();
 
@@ -217,6 +220,12 @@ private:
 	void OnDeleteDelegate(FEdGraphSchemaAction_K2Delegate* InDelegateAction);
 
 	UEdGraph* GetFocusedGraph() const;
+
+	/** Delegate to hook us into non-structural Blueprint object post-change events */
+	void OnObjectPropertyChanged(UObject* InObject, FPropertyChangedEvent& InPropertyChangedEvent)
+	{
+		bNeedsRefresh = (InObject == Blueprint);
+	}
 private:
 	/** Pointer back to the blueprint editor that owns us */
 	TWeakPtr<FBlueprintEditor> BlueprintEditorPtr;
@@ -251,4 +260,7 @@ private:
 
 	/** The Kismet Inspector used to display properties: */
 	TWeakPtr<SKismetInspector> Inspector;
+
+	/** Flag to indicate whether or not we need to refresh the panel */
+	bool bNeedsRefresh;
 };
