Uncategorized functions will no longer be categorized into a "Class" category in the Blueprint node context menu when a being displayed for a currently selected level Actor.

diff --git a/Engine/Source/Editor/BlueprintGraph/Private/BlueprintFunctionNodeSpawner.cpp b/Engine/Source/Editor/BlueprintGraph/Private/BlueprintFunctionNodeSpawner.cpp
--- a/Engine/Source/Editor/BlueprintGraph/Private/BlueprintFunctionNodeSpawner.cpp
+++ b/Engine/Source/Editor/BlueprintGraph/Private/BlueprintFunctionNodeSpawner.cpp
@@ -329,8 +329,16 @@ FBlueprintActionUiSpec UBlueprintFunctionNodeSpawner::GetUiSpec(FBlueprintAction
 
 		if (!TargetClass->IsChildOf(FunctionClass))
 		{
-			MenuSignature.Category = FEditorCategoryUtils::BuildCategoryString( FCommonEditorCategory::Class, 
-				FText::FromString(FunctionClass->GetDisplayNameText().ToString()) );
+			// When there are no bindings set, functions need to be categorized into a category "Class" to help reduce clutter in the tree root
+			if(Bindings.Num() == 0)
+			{
+				MenuSignature.Category = FEditorCategoryUtils::BuildCategoryString( FCommonEditorCategory::Class, 
+					FText::FromString(FunctionClass->GetDisplayNameText().ToString()) );
+			}
+			else
+			{
+				MenuSignature.Category = FText::FromString(FunctionClass->GetDisplayNameText().ToString());
+			}
 		}
 	}
 
