this changelist fixes an issue with the blueprint action database holding onto object references via cached node templates which in certain circumstances caused problems switching levels.
#UE-8047 - Crashing on level open, opening level blueprint, saving level
#UE-5778 - Switching between levels after opening the level blueprint and saving can result in a crash
#Branch UE4
#Proj BlueprintGraph

#info This removes an issue whereby the an asset could refresh its actions without first clearing any cached node templates which can cause problems with GC, especially when switching levels after saving the previous level under a different name.

diff --git a/Engine/Source/Editor/BlueprintGraph/Private/BlueprintActionDatabase.cpp b/Engine/Source/Editor/BlueprintGraph/Private/BlueprintActionDatabase.cpp
--- a/Engine/Source/Editor/BlueprintGraph/Private/BlueprintActionDatabase.cpp
+++ b/Engine/Source/Editor/BlueprintGraph/Private/BlueprintActionDatabase.cpp
@@ -1147,8 +1147,15 @@ void FBlueprintActionDatabase::RefreshAssetActions(UObject* const AssetObject)
 	using namespace BlueprintActionDatabaseImpl;
 	check(BlueprintActionDatabaseImpl::IsObjectValidForDatabase(AssetObject));
 
-	bool const bHadExistingEntry = ActionRegistry.Contains(AssetObject);
 	FActionList& AssetActionList = ActionRegistry.FindOrAdd(AssetObject);
+	bool const bHadExistingEntry = AssetActionList.Num() > 0;
+	for (UBlueprintNodeSpawner* Action : AssetActionList)
+	{
+		// because some asserts expect everything to be cleaned up in a 
+		// single GC pass, we need to ensure that any previously cached node templates
+		// are cleaned up here before we add any new node spawners.
+		Action->ClearCachedTemplateNode();
+	}
 	AssetActionList.Empty();
 
 	if(const USkeleton* Skeleton = Cast<USkeleton>(AssetObject))
