Fix for assert when refreshing the blueprint action database
#UE-8988 - Editor asserts if you hit compile on a blank level
#Branch UE4
#Proj BlueprintGraph

diff --git a/Engine/Source/Editor/BlueprintGraph/Private/BlueprintActionDatabase.cpp b/Engine/Source/Editor/BlueprintGraph/Private/BlueprintActionDatabase.cpp
--- a/Engine/Source/Editor/BlueprintGraph/Private/BlueprintActionDatabase.cpp
+++ b/Engine/Source/Editor/BlueprintGraph/Private/BlueprintActionDatabase.cpp
@@ -1147,8 +1147,8 @@ void FBlueprintActionDatabase::RefreshAssetActions(UObject* const AssetObject)
 	using namespace BlueprintActionDatabaseImpl;
 	check(BlueprintActionDatabaseImpl::IsObjectValidForDatabase(AssetObject));
 
+	bool const bHadExistingEntry = ActionRegistry.Contains(AssetObject);
 	FActionList& AssetActionList = ActionRegistry.FindOrAdd(AssetObject);
-	bool const bHadExistingEntry = AssetActionList.Num() > 0;
 	for (UBlueprintNodeSpawner* Action : AssetActionList)
 	{
 		// because some asserts expect everything to be cleaned up in a 
