Functions will now appear with their name (instead of display name) in the MyBlueprint window, keeping them in line with other items.

diff --git a/Engine/Source/Editor/BlueprintGraph/Private/EdGraphSchema_K2.cpp b/Engine/Source/Editor/BlueprintGraph/Private/EdGraphSchema_K2.cpp
--- a/Engine/Source/Editor/BlueprintGraph/Private/EdGraphSchema_K2.cpp
+++ b/Engine/Source/Editor/BlueprintGraph/Private/EdGraphSchema_K2.cpp
@@ -4126,7 +4126,7 @@ void UEdGraphSchema_K2::GetGraphDisplayInformation(const UEdGraph& Graph, /*out*
 			// If we found a function from this graph..
 			if (Function)
 			{
-				DisplayInfo.PlainName = FText::FromString(UK2Node_CallFunction::GetUserFacingFunctionName(Function)); // grab friendly function name
+				DisplayInfo.PlainName = FText::FromString(Function->GetName());
 				DisplayInfo.Tooltip = UK2Node_CallFunction::GetDefaultTooltipForFunction(Function); // grab its tooltip
 			}
 			else
@@ -4956,7 +4956,7 @@ UEdGraphNode* UEdGraphSchema_K2::CreateSubstituteNode(UEdGraphNode* Node, const
 			}
 			else if(EventNode->bOverrideFunction)
 			{
-				FunctionName = EventNode->EventSignatureName.ToString();
+				FunctionName = EventNode->EventReference.GetMemberName().ToString();
 			}
 			else
 			{
