Fix a bad change.

diff --git a/Engine/Source/Editor/BlueprintGraph/Private/EdGraphSchema_K2.cpp b/Engine/Source/Editor/BlueprintGraph/Private/EdGraphSchema_K2.cpp
--- a/Engine/Source/Editor/BlueprintGraph/Private/EdGraphSchema_K2.cpp
+++ b/Engine/Source/Editor/BlueprintGraph/Private/EdGraphSchema_K2.cpp
@@ -4956,7 +4956,7 @@ UEdGraphNode* UEdGraphSchema_K2::CreateSubstituteNode(UEdGraphNode* Node, const
 			}
 			else if(EventNode->bOverrideFunction)
 			{
-				FunctionName = EventNode->EventReference.GetMemberName().ToString();
+				FunctionName = EventNode->EventSignatureName.ToString();
 			}
 			else
 			{
