[INTEGRATE] Change 2426005 by [REDACTED_PERSON][REDACTED_HANDLE] on 2015/01/30 12:16:51

	[UE-8524][UE-8524] The add component node now handles blueprintable components properly (without using the SKEL class).

diff --git a/Engine/Source/Editor/BlueprintGraph/Private/BlueprintComponentNodeSpawner.cpp b/Engine/Source/Editor/BlueprintGraph/Private/BlueprintComponentNodeSpawner.cpp
--- a/Engine/Source/Editor/BlueprintGraph/Private/BlueprintComponentNodeSpawner.cpp
+++ b/Engine/Source/Editor/BlueprintGraph/Private/BlueprintComponentNodeSpawner.cpp
@@ -53,20 +53,22 @@ UBlueprintComponentNodeSpawner* UBlueprintComponentNodeSpawner::Create(TSubclass
 		Outer = GetTransientPackage();
 	}
 
+	UClass* const AuthoritativeClass = ComponentClass->GetAuthoritativeClass();
+
 	UBlueprintComponentNodeSpawner* NodeSpawner = NewObject<UBlueprintComponentNodeSpawner>(Outer);
-	NodeSpawner->ComponentClass = ComponentClass;
+	NodeSpawner->ComponentClass = AuthoritativeClass;
 	NodeSpawner->NodeClass      = UK2Node_AddComponent::StaticClass();
 
 	FBlueprintActionUiSpec& MenuSignature = NodeSpawner->DefaultMenuSignature;
-	FText const ComponentTypeName = FText::FromName(ComponentClass->GetFName());
+	FText const ComponentTypeName = AuthoritativeClass->GetDisplayNameText();
 	MenuSignature.MenuName = FText::Format(LOCTEXT("AddComponentMenuName", "Add {0}"), ComponentTypeName);
-	MenuSignature.Category = BlueprintComponentNodeSpawnerImpl::GetDefaultMenuCategory(ComponentClass);
+	MenuSignature.Category = BlueprintComponentNodeSpawnerImpl::GetDefaultMenuCategory(AuthoritativeClass);
 	MenuSignature.Tooltip  = FText::Format(LOCTEXT("AddComponentTooltip", "Spawn a {0}"), ComponentTypeName);
-	MenuSignature.Keywords = ComponentClass->GetMetaData(FBlueprintMetadata::MD_FunctionKeywords);
+	MenuSignature.Keywords = AuthoritativeClass->GetMetaData(FBlueprintMetadata::MD_FunctionKeywords);
 	// add at least one character, so that PrimeDefaultMenuSignature() doesn't 
 	// attempt to query the template node
 	MenuSignature.Keywords.AppendChar(TEXT(' '));
-	MenuSignature.IconName = FClassIconFinder::FindIconNameForClass(ComponentClass);
+	MenuSignature.IconName = FClassIconFinder::FindIconNameForClass(AuthoritativeClass);
 
 	return NodeSpawner;
 }
@@ -118,7 +120,9 @@ UEdGraphNode* UBlueprintComponentNodeSpawner::Invoke(UEdGraph* ParentGraph, FBin
 	if (!bIsTemplateNode)
 	{
 		UBlueprint* Blueprint = NewNode->GetBlueprint();
-		UActorComponent* ComponentTemplate = NewObject<UActorComponent>(Blueprint->GeneratedClass, ComponentClass);
+
+		FName DesiredComponentName = MakeUniqueObjectName(Blueprint->GeneratedClass, ComponentClass, *ComponentClass->GetDisplayNameText().ToString());
+		UActorComponent* ComponentTemplate = NewObject<UActorComponent>(Blueprint->GeneratedClass, ComponentClass, DesiredComponentName);
 		ComponentTemplate->SetFlags(RF_ArchetypeObject);
 
 		Blueprint->ComponentTemplates.Add(ComponentTemplate);
diff --git a/Engine/Source/Editor/BlueprintGraph/Private/K2Node_AddComponent.cpp b/Engine/Source/Editor/BlueprintGraph/Private/K2Node_AddComponent.cpp
--- a/Engine/Source/Editor/BlueprintGraph/Private/K2Node_AddComponent.cpp
+++ b/Engine/Source/Editor/BlueprintGraph/Private/K2Node_AddComponent.cpp
@@ -321,7 +321,7 @@ FText UK2Node_AddComponent::GetNodeTitle(ENodeTitleType::Type TitleType) const
 			else
 			{
 				FFormatNamedArguments Args;
-				Args.Add(TEXT("ClassName"), FText::FromString(SourceTemplate->GetClass()->GetName()));
+				Args.Add(TEXT("ClassName"), SourceTemplate->GetClass()->GetDisplayNameText());
 				CachedNodeTitle = FText::Format(LOCTEXT("AddClass", "Add {ClassName}"), Args);
 			}
 		}
