Original fix by [REDACTED_PERSON]:
"Fixes for multicast delegates (and K2Node_Events in general) to now use FMemberReference for finding their UFunction.

I've change it a bit, so it don't require ObjectVersion.h edition as it's locked. My changes are located in:
* UK2Node_Event::Serialize
* UK2Node_ComponentBoundEvent::Serialize
* UK2Node_ActorBoundEvent::Serialize
* header K2Node_Event.h lines from 16 to 39.

This change works for now as a workaround. We will change that into proper deprecation fix after ObjectVersion.h lock lift.

diff --git a/Engine/Plugins/Experimental/LiveEditor/Source/LiveEditor/Private/K2Node_LiveEditObject.cpp b/Engine/Plugins/Experimental/LiveEditor/Source/LiveEditor/Private/K2Node_LiveEditObject.cpp
--- a/Engine/Plugins/Experimental/LiveEditor/Source/LiveEditor/Private/K2Node_LiveEditObject.cpp
+++ b/Engine/Plugins/Experimental/LiveEditor/Source/LiveEditor/Private/K2Node_LiveEditObject.cpp
@@ -348,8 +348,7 @@ void UK2Node_LiveEditObject::ExpandNode(class FKismetCompilerContext& CompilerCo
 		UFunction *EventMIDISignature = GetEventMIDISignature();
 		UK2Node_Event* EventNode = CompilerContext.SpawnIntermediateNode<UK2Node_Event>(this, SourceGraph);
 		check(EventNode);
-		EventNode->EventSignatureClass = Cast<UClass>(EventMIDISignature->GetOuter());
-		EventNode->EventSignatureName = EventMIDISignature->GetFName();
+		EventNode->EventReference.SetFromField<UFunction>(EventMIDISignature, false);
 		EventNode->CustomFunctionName = *EventNameGuid;
 		EventNode->bInternalEvent = true;
 		EventNode->AllocateDefaultPins();
diff --git a/Engine/Source/Editor/BlueprintGraph/Classes/K2Node.h b/Engine/Source/Editor/BlueprintGraph/Classes/K2Node.h
--- a/Engine/Source/Editor/BlueprintGraph/Classes/K2Node.h
+++ b/Engine/Source/Editor/BlueprintGraph/Classes/K2Node.h
@@ -520,6 +520,9 @@ public:
 	/** Set to a non-'self' member, so must include reference to class owning the member. */
 	BLUEPRINTGRAPH_API void SetExternalMember(FName InMemberName, TSubclassOf<class UObject> InMemberParentClass);
 
+	/** Set to a non-'self' delegate member, this is not self-context but is not given a parent class */
+	BLUEPRINTGRAPH_API void SetExternalDelegateMember(FName InMemberName);
+
 	/** Set up this reference to a 'self' member name */
 	BLUEPRINTGRAPH_API void SetSelfMember(FName InMemberName);
 
diff --git a/Engine/Source/Editor/BlueprintGraph/Classes/K2Node_ActorBoundEvent.h b/Engine/Source/Editor/BlueprintGraph/Classes/K2Node_ActorBoundEvent.h
--- a/Engine/Source/Editor/BlueprintGraph/Classes/K2Node_ActorBoundEvent.h
+++ b/Engine/Source/Editor/BlueprintGraph/Classes/K2Node_ActorBoundEvent.h
@@ -15,10 +15,17 @@ class UK2Node_ActorBoundEvent : public UK2Node_Event
 	UPROPERTY()
 	FName DelegatePropertyName;
 
+	/** Delegate property's owner class that this event is associated with */
+	UPROPERTY(Transient)
+	UClass* DelegateOwnerClass;
+
 	/** The event that this event is bound to */
 	UPROPERTY()
 	class AActor* EventOwner;
 
+	// UObject interface
+	virtual void Serialize(FArchive& Ar) override;
+	// End of UObject interface
 
 	// Begin UEdGraphNode interface
 	virtual void ReconstructNode() override;
diff --git a/Engine/Source/Editor/BlueprintGraph/Classes/K2Node_ComponentBoundEvent.h b/Engine/Source/Editor/BlueprintGraph/Classes/K2Node_ComponentBoundEvent.h
--- a/Engine/Source/Editor/BlueprintGraph/Classes/K2Node_ComponentBoundEvent.h
+++ b/Engine/Source/Editor/BlueprintGraph/Classes/K2Node_ComponentBoundEvent.h
@@ -17,12 +17,17 @@ class UK2Node_ComponentBoundEvent : public UK2Node_Event
 	UPROPERTY()
 	FName DelegatePropertyName;
 
+	/** Delegate property's owner class that this event is associated with */
+	UPROPERTY(Transient)
+	UClass* DelegateOwnerClass;
+
 	/** Name of property in Blueprint class that pointer to component we want to bind to */
 	UPROPERTY()
 	FName ComponentPropertyName;
 
 	// Begin UObject interface
 	virtual bool Modify(bool bAlwaysMarkDirty = true) override;
+	virtual void Serialize(FArchive& Ar) override;
 	// End UObject interface
 
 	// Begin UEdGraphNode interface
diff --git a/Engine/Source/Editor/BlueprintGraph/Classes/K2Node_Event.h b/Engine/Source/Editor/BlueprintGraph/Classes/K2Node_Event.h
--- a/Engine/Source/Editor/BlueprintGraph/Classes/K2Node_Event.h
+++ b/Engine/Source/Editor/BlueprintGraph/Classes/K2Node_Event.h
@@ -2,6 +2,7 @@
 
 
 #pragma once
+#include "Core.h"
 #include "K2Node_EditablePinBase.h"
 #include "EdGraph/EdGraphNodeUtils.h" // for FNodeTextCache
 #include "K2Node_Event.generated.h"
@@ -12,14 +13,31 @@ class UK2Node_Event : public UK2Node_EditablePinBase
 	GENERATED_UCLASS_BODY()
 	BLUEPRINTGRAPH_API static const FString DelegateOutputName;
 
+	/******************************************************************************/
+
+	/**
+	 * Properties EventSignatureName and EventSignatureClass should be deprecated.
+	 * In order to do this we need to edit ObjectVersion.h which is locked for a
+	 * while. After lock lift those properties will be deprecated and eventually
+	 * removed. Please use EventReference instead and it will be translated into
+	 * EventReference accordingly.
+	 */
+
 	/** Name of function signature that this event implements */
+	DEPRECATED(4.7, "This field is deprecated, please use EventReference instead. For details read comment above.")
 	UPROPERTY()
 	FName EventSignatureName;
 
 	/** Class that the function signature is from. */
+	DEPRECATED(4.7, "This field is deprecated, please use EventReference instead. For details read comment above.")
 	UPROPERTY()
 	TSubclassOf<class UObject> EventSignatureClass;
 
+	/** Reference for the function this event is linked to */
+	UPROPERTY(Transient)
+	FMemberReference EventReference;
+	/******************************************************************************/
+
 	/** If true, we are actually overriding this function, not making a new event with a signature that matches */
 	UPROPERTY()
 	uint32 bOverrideFunction:1;
@@ -74,7 +92,7 @@ class UK2Node_Event : public UK2Node_EditablePinBase
 	/** Checks whether the parameters for this event node are compatible with the specified function entry node */
 	BLUEPRINTGRAPH_API virtual bool IsFunctionEntryCompatible(const class UK2Node_FunctionEntry* EntryNode) const;
 	
-	UFunction* FindEventSignatureFunction();
+	BLUEPRINTGRAPH_API UFunction* FindEventSignatureFunction();
 	BLUEPRINTGRAPH_API void UpdateDelegatePin(bool bSilent = false);
 	BLUEPRINTGRAPH_API FName GetFunctionName() const;
 	BLUEPRINTGRAPH_API virtual bool IsUsedByAuthorityOnlyDelegate() const { return false; }
diff --git a/Engine/Source/Editor/BlueprintGraph/Private/BlueprintActionDatabase.cpp b/Engine/Source/Editor/BlueprintGraph/Private/BlueprintActionDatabase.cpp
--- a/Engine/Source/Editor/BlueprintGraph/Private/BlueprintActionDatabase.cpp
+++ b/Engine/Source/Editor/BlueprintGraph/Private/BlueprintActionDatabase.cpp
@@ -248,7 +248,7 @@ static UBlueprintNodeSpawner* FBlueprintNodeSpawnerFactory::MakeAnimOwnedEventSp
 	auto PostSpawnSetupLambda = [](UEdGraphNode* NewNode, bool /*bIsTemplateNode*/)
 	{
 		UK2Node_Event* ActorRefNode = CastChecked<UK2Node_Event>(NewNode);
-		ActorRefNode->EventSignatureClass = UAnimInstance::StaticClass();
+		ActorRefNode->EventReference.SetExternalMember(ActorRefNode->CustomFunctionName, UAnimInstance::StaticClass());
 	};
 
 	UBlueprintNodeSpawner* NodeSpawner = UBlueprintEventNodeSpawner::Create(UK2Node_Event::StaticClass(), SignatureName);
diff --git a/Engine/Source/Editor/BlueprintGraph/Private/BlueprintEventNodeSpawner.cpp b/Engine/Source/Editor/BlueprintGraph/Private/BlueprintEventNodeSpawner.cpp
--- a/Engine/Source/Editor/BlueprintGraph/Private/BlueprintEventNodeSpawner.cpp
+++ b/Engine/Source/Editor/BlueprintGraph/Private/BlueprintEventNodeSpawner.cpp
@@ -157,8 +157,7 @@ UEdGraphNode* UBlueprintEventNodeSpawner::Invoke(UEdGraph* ParentGraph, FBinding
 			UK2Node_Event* EventNode = CastChecked<UK2Node_Event>(NewNode);
 			if (EventFunc != nullptr)
 			{
-				EventNode->EventSignatureName  = EventName;
-				EventNode->EventSignatureClass = EventFunc->GetOuterUClass()->GetAuthoritativeClass();
+				EventNode->EventReference.SetFromField<UFunction>(EventFunc, false);
 				EventNode->bOverrideFunction   = true;
 			}
 			else if (!bIsTemplateNode)
diff --git a/Engine/Source/Editor/BlueprintGraph/Private/EdGraphSchema_K2.cpp b/Engine/Source/Editor/BlueprintGraph/Private/EdGraphSchema_K2.cpp
--- a/Engine/Source/Editor/BlueprintGraph/Private/EdGraphSchema_K2.cpp
+++ b/Engine/Source/Editor/BlueprintGraph/Private/EdGraphSchema_K2.cpp
@@ -714,7 +714,7 @@ bool UEdGraphSchema_K2::CanFunctionBeUsedInGraph(const UClass* InClass, const UF
 
 UFunction* UEdGraphSchema_K2::GetCallableParentFunction(UFunction* Function) const
 {
-	if( Function )
+	if( Function && Cast<UClass>(Function->GetOuter()) )
 	{
 		const FName FunctionName = Function->GetFName();
 
@@ -4959,7 +4959,7 @@ UEdGraphNode* UEdGraphSchema_K2::CreateSubstituteNode(UEdGraphNode* Node, const
 			}
 			else if(EventNode->bOverrideFunction)
 			{
-				FunctionName = EventNode->EventSignatureName.ToString();
+				FunctionName = EventNode->EventReference.GetMemberName().ToString();
 			}
 			else
 			{
diff --git a/Engine/Source/Editor/BlueprintGraph/Private/EdGraphSchema_K2_Actions.cpp b/Engine/Source/Editor/BlueprintGraph/Private/EdGraphSchema_K2_Actions.cpp
--- a/Engine/Source/Editor/BlueprintGraph/Private/EdGraphSchema_K2_Actions.cpp
+++ b/Engine/Source/Editor/BlueprintGraph/Private/EdGraphSchema_K2_Actions.cpp
@@ -400,7 +400,7 @@ bool FEdGraphSchemaAction_K2AddEvent::EventHasAlreadyBeenPlaced(UBlueprint const
 	if (Blueprint != NULL)
 	{
 		UK2Node_Event const* EventTemplate = Cast<UK2Node_Event const>(NodeTemplate);
-		ExistingEvent = FBlueprintEditorUtils::FindOverrideForFunction(Blueprint, EventTemplate->EventSignatureClass, EventTemplate->EventSignatureName);
+		ExistingEvent = FBlueprintEditorUtils::FindOverrideForFunction(Blueprint, EventTemplate->EventReference.GetMemberParentClass(EventTemplate), EventTemplate->EventReference.GetMemberName());
 	}
 
 	if (FoundEventOut != NULL)
diff --git a/Engine/Source/Editor/BlueprintGraph/Private/K2ActionMenuBuilder.cpp b/Engine/Source/Editor/BlueprintGraph/Private/K2ActionMenuBuilder.cpp
--- a/Engine/Source/Editor/BlueprintGraph/Private/K2ActionMenuBuilder.cpp
+++ b/Engine/Source/Editor/BlueprintGraph/Private/K2ActionMenuBuilder.cpp
@@ -609,8 +609,7 @@ void FK2ActionMenuBuilder::AddEventForFunction(FGraphActionListBuilderBase& Acti
 		const FString EventName = FText::Format( LOCTEXT("EventWithSignatureName", "Event {0}"), FText::FromString(SigName)).ToString();
 
 		UK2Node_Event* EventNodeTemplate = ActionMenuBuilder.CreateTemplateNode<UK2Node_Event>();
-		EventNodeTemplate->EventSignatureName = Function->GetFName();
-		EventNodeTemplate->EventSignatureClass = CastChecked<UClass>(Function->GetOuter());
+		EventNodeTemplate->EventReference.SetFromField<UFunction>(Function, false);
 		EventNodeTemplate->bOverrideFunction = true;
 
 		const FString Category = UK2Node_CallFunction::GetDefaultCategoryForFunction(Function, K2ActionCategories::AddEventCategory);
@@ -1699,9 +1698,8 @@ void FK2ActionMenuBuilder::GetAnimNotifyMenuItems(FBlueprintGraphActionListBuild
 
 				UK2Node_Event* EventNode = ContextMenuBuilder.CreateTemplateNode<UK2Node_Event>();
 				FString SignatureName = FString::Printf(TEXT("AnimNotify_%s"), *Label);
-				EventNode->EventSignatureName = FName(*SignatureName);
-				EventNode->EventSignatureClass = UAnimInstance::StaticClass();
-				EventNode->CustomFunctionName = EventNode->EventSignatureName;
+				EventNode->EventReference.SetExternalMember(FName(*SignatureName), UAnimInstance::StaticClass());
+				EventNode->CustomFunctionName = EventNode->EventReference.GetMemberName();
 				Action->NodeTemplate = EventNode;
 			}
 		}
@@ -1720,9 +1718,8 @@ void FK2ActionMenuBuilder::GetAnimNotifyMenuItems(FBlueprintGraphActionListBuild
 
 					UK2Node_Event* EventNode = ContextMenuBuilder.CreateTemplateNode<UK2Node_Event>();
 					FString SignatureName = FString::Printf(TEXT("AnimNotify_%s"), *Label);
-					EventNode->EventSignatureName = FName(*SignatureName);
-					EventNode->EventSignatureClass = UAnimInstance::StaticClass();
-					EventNode->CustomFunctionName = EventNode->EventSignatureName;
+					EventNode->EventReference.SetExternalMember(FName(*SignatureName), UAnimInstance::StaticClass());
+					EventNode->CustomFunctionName = EventNode->EventReference.GetMemberName();
 					Action->NodeTemplate = EventNode;
 				}
 			}
@@ -2065,7 +2062,7 @@ void FK2ActionMenuBuilder::GetEventsForBlueprint(FBlueprintPaletteListBuilder& A
 		FBlueprintEditorUtils::GetAllNodesOfClass<UK2Node_Event>(Blueprint, AllEvents);
 		for (auto It = AllEvents.CreateConstIterator(); It; It++)
 		{
-			ExistingFuncNames.Add((*It)->EventSignatureName);
+			ExistingFuncNames.Add((*It)->EventReference.GetMemberName());
 		}
 
 		for (TFieldIterator<UFunction> FunctionIt(Blueprint->ParentClass, EFieldIteratorFlags::IncludeSuper); FunctionIt; ++FunctionIt)
diff --git a/Engine/Source/Editor/BlueprintGraph/Private/K2Node.cpp b/Engine/Source/Editor/BlueprintGraph/Private/K2Node.cpp
--- a/Engine/Source/Editor/BlueprintGraph/Private/K2Node.cpp
+++ b/Engine/Source/Editor/BlueprintGraph/Private/K2Node.cpp
@@ -1316,6 +1316,11 @@ void FMemberReference::SetExternalMember(FName InMemberName, TSubclassOf<class U
 	bWasDeprecated = false;
 }
 
+void FMemberReference::SetExternalDelegateMember(FName InMemberName)
+{
+	SetExternalMember(InMemberName, nullptr);
+}
+
 void FMemberReference::SetSelfMember(FName InMemberName)
 {
 	MemberName = InMemberName;
diff --git a/Engine/Source/Editor/BlueprintGraph/Private/K2Node_ActorBoundEvent.cpp b/Engine/Source/Editor/BlueprintGraph/Private/K2Node_ActorBoundEvent.cpp
--- a/Engine/Source/Editor/BlueprintGraph/Private/K2Node_ActorBoundEvent.cpp
+++ b/Engine/Source/Editor/BlueprintGraph/Private/K2Node_ActorBoundEvent.cpp
@@ -116,7 +116,7 @@ FText UK2Node_ActorBoundEvent::GetTooltipText() const
 
 FString UK2Node_ActorBoundEvent::GetDocumentationLink() const
 {
-	if (EventSignatureClass)
+	if (UClass* EventSignatureClass = EventReference.GetMemberParentClass(this))
 	{
 		return FString::Printf(TEXT("Shared/GraphNodes/Blueprint/%s%s"), EventSignatureClass->GetPrefixCPP(), *EventSignatureClass->GetName());
 	}
@@ -141,9 +141,9 @@ void UK2Node_ActorBoundEvent::InitializeActorBoundEventParams(AActor* InEventOwn
 	{
 		EventOwner = InEventOwner;
 		DelegatePropertyName = InDelegateProperty->GetFName();
-		EventSignatureName = InDelegateProperty->SignatureFunction->GetFName();
-		EventSignatureClass = CastChecked<UClass>(InDelegateProperty->SignatureFunction->GetOuter())->GetAuthoritativeClass();
-		CustomFunctionName = FName( *FString::Printf(TEXT("BndEvt__%s_%s_%s"), *InEventOwner->GetName(), *GetName(), *EventSignatureName.ToString()) );
+		DelegateOwnerClass = CastChecked<UClass>(InDelegateProperty->GetOuter())->GetAuthoritativeClass();
+		EventReference.SetFromField<UFunction>(InDelegateProperty->SignatureFunction, false);
+		CustomFunctionName = FName( *FString::Printf(TEXT("BndEvt__%s_%s_%s"), *InEventOwner->GetName(), *GetName(), *EventReference.GetMemberName().ToString()) );
 		bOverrideFunction = false;
 		bInternalEvent = true;
 		CachedNodeTitle.MarkDirty();
@@ -152,23 +152,22 @@ void UK2Node_ActorBoundEvent::InitializeActorBoundEventParams(AActor* InEventOwn
 
 UMulticastDelegateProperty* UK2Node_ActorBoundEvent::GetTargetDelegatePropertyConst() const
 {
-	return Cast<UMulticastDelegateProperty>(FindField<UMulticastDelegateProperty>(EventSignatureClass, DelegatePropertyName));
+	return Cast<UMulticastDelegateProperty>(FindField<UMulticastDelegateProperty>(DelegateOwnerClass, DelegatePropertyName));
 }
 
 UMulticastDelegateProperty* UK2Node_ActorBoundEvent::GetTargetDelegateProperty()
 {
-	UMulticastDelegateProperty* TargetDelegateProp = Cast<UMulticastDelegateProperty>(FindField<UMulticastDelegateProperty>(EventSignatureClass, DelegatePropertyName));
+	UMulticastDelegateProperty* TargetDelegateProp = Cast<UMulticastDelegateProperty>(FindField<UMulticastDelegateProperty>(DelegateOwnerClass, DelegatePropertyName));
 
 	// If we couldn't find the target delegate, then try to find it in the property remap table
 	if (!TargetDelegateProp)
 	{
-		UMulticastDelegateProperty* NewProperty = Cast<UMulticastDelegateProperty>(FindRemappedField(EventSignatureClass, DelegatePropertyName));
+		UMulticastDelegateProperty* NewProperty = Cast<UMulticastDelegateProperty>(FindRemappedField(DelegateOwnerClass, DelegatePropertyName));
 		if (NewProperty)
 		{
 			// Found a remapped property, update the node
 			TargetDelegateProp = NewProperty;
 			DelegatePropertyName = NewProperty->GetFName();
-			EventSignatureClass = Cast<UClass>(NewProperty->GetOuter());
 			CachedNodeTitle.MarkDirty();
 		}
 	}
@@ -192,8 +191,40 @@ FMulticastScriptDelegate* UK2Node_ActorBoundEvent::GetTargetDelegate()
 
 bool UK2Node_ActorBoundEvent::IsUsedByAuthorityOnlyDelegate() const
 {
-	const UMulticastDelegateProperty* TargetDelegateProp = Cast<const UMulticastDelegateProperty>(FindField<UMulticastDelegateProperty>( EventSignatureClass, DelegatePropertyName ));
+	const UMulticastDelegateProperty* TargetDelegateProp = Cast<const UMulticastDelegateProperty>(FindField<UMulticastDelegateProperty>( DelegateOwnerClass, DelegatePropertyName ));
 	return (TargetDelegateProp && TargetDelegateProp->HasAnyPropertyFlags(CPF_BlueprintAuthorityOnly));
 }
 
+void UK2Node_ActorBoundEvent::Serialize(FArchive& Ar)
+{
+	// This deprecation warning suppression is due to the workaround. For more
+	// details read comment in K2Node_Event.h near EventSignatureName property
+	// definition.
+
+	PRAGMA_DISABLE_DEPRECATION_WARNINGS
+
+	// Workaround, so fields EventSignatureName and EventSignatureClass work
+	// until ObjectVersion.h lock is lifted.
+	if (Ar.IsSaving())
+	{
+		EventSignatureName = EventReference.GetMemberName();
+		EventSignatureClass = DelegateOwnerClass;
+		Super::Super::Serialize(Ar);
+		return;
+	}
+
+	Super::Serialize(Ar);
+
+	// Fix up legacy nodes that may not yet have a delegate pin
+	if(Ar.IsLoading())
+	{
+		DelegateOwnerClass = EventSignatureClass;
+		UMulticastDelegateProperty* TargetDelegateProp = GetTargetDelegateProperty();
+		check(TargetDelegateProp);
+		EventReference.SetFromField<UFunction>(TargetDelegateProp->SignatureFunction, false);
+	}
+
+	PRAGMA_ENABLE_DEPRECATION_WARNINGS
+}
+
 #undef LOCTEXT_NAMESPACE
diff --git a/Engine/Source/Editor/BlueprintGraph/Private/K2Node_ComponentBoundEvent.cpp b/Engine/Source/Editor/BlueprintGraph/Private/K2Node_ComponentBoundEvent.cpp
--- a/Engine/Source/Editor/BlueprintGraph/Private/K2Node_ComponentBoundEvent.cpp
+++ b/Engine/Source/Editor/BlueprintGraph/Private/K2Node_ComponentBoundEvent.cpp
@@ -38,9 +38,9 @@ void UK2Node_ComponentBoundEvent::InitializeComponentBoundEventParams(UObjectPro
 	{
 		ComponentPropertyName = InComponentProperty->GetFName();
 		DelegatePropertyName = InDelegateProperty->GetFName();
-		EventSignatureName = InDelegateProperty->SignatureFunction->GetFName();
-		EventSignatureClass = CastChecked<UClass>(InDelegateProperty->SignatureFunction->GetOuter());
-		CustomFunctionName = FName( *FString::Printf(TEXT("BndEvt__%s_%s_%s"), *InComponentProperty->GetName(), *GetName(), *EventSignatureName.ToString()) );
+		DelegateOwnerClass = CastChecked<UClass>(InDelegateProperty->GetOuter())->GetAuthoritativeClass();
+		EventReference.SetExternalDelegateMember(InDelegateProperty->SignatureFunction->GetFName());
+		CustomFunctionName = FName( *FString::Printf(TEXT("BndEvt__%s_%s_%s"), *InComponentProperty->GetName(), *GetName(), *EventReference.GetMemberName().ToString()) );
 		bOverrideFunction = false;
 		bInternalEvent = true;
 		CachedNodeTitle.MarkDirty();
@@ -73,7 +73,7 @@ bool UK2Node_ComponentBoundEvent::IsUsedByAuthorityOnlyDelegate() const
 
 UMulticastDelegateProperty* UK2Node_ComponentBoundEvent::GetTargetDelegateProperty() const
 {
-	return Cast<UMulticastDelegateProperty>(FindField<UMulticastDelegateProperty>(EventSignatureClass, DelegatePropertyName));
+	return Cast<UMulticastDelegateProperty>(FindField<UMulticastDelegateProperty>(DelegateOwnerClass, DelegatePropertyName));
 }
 
 
@@ -92,9 +92,9 @@ FText UK2Node_ComponentBoundEvent::GetTooltipText() const
 
 FString UK2Node_ComponentBoundEvent::GetDocumentationLink() const
 {
-	if (EventSignatureClass)
+	if (DelegateOwnerClass)
 	{
-		return FString::Printf(TEXT("Shared/GraphNodes/Blueprint/%s%s"), EventSignatureClass->GetPrefixCPP(), *EventSignatureClass->GetName());
+		return FString::Printf(TEXT("Shared/GraphNodes/Blueprint/%s%s"), DelegateOwnerClass->GetPrefixCPP(), *EventReference.GetMemberName().ToString());
 	}
 
 	return FString();
@@ -105,4 +105,35 @@ FString UK2Node_ComponentBoundEvent::GetDocumentationExcerptName() const
 	return DelegatePropertyName.ToString();
 }
 
+void UK2Node_ComponentBoundEvent::Serialize(FArchive& Ar)
+{
+	// This deprecation warning suppression is due to the workaround. For more
+	// details read comment in K2Node_Event.h near EventSignatureName property
+	// definition.
+
+	PRAGMA_DISABLE_DEPRECATION_WARNINGS
+	// Workaround, so fields EventSignatureName and EventSignatureClass work
+	// until ObjectVersion.h lock is lifted.
+	if (Ar.IsSaving())
+	{
+		EventSignatureName = EventReference.GetMemberName();
+		EventSignatureClass = DelegateOwnerClass;
+		Super::Super::Serialize(Ar);
+		return;
+	}
+
+	Super::Serialize(Ar);
+
+	// Fix up legacy nodes that may not yet have a delegate pin
+	if(Ar.IsLoading())
+	{
+		DelegateOwnerClass = EventSignatureClass;
+		UMulticastDelegateProperty* TargetDelegateProp = GetTargetDelegateProperty();
+		check(TargetDelegateProp);
+		EventReference.SetFromField<UFunction>(TargetDelegateProp->SignatureFunction, false);
+	}
+	PRAGMA_ENABLE_DEPRECATION_WARNINGS
+}
+
+
 #undef LOCTEXT_NAMESPACE
\ No newline at end of file
diff --git a/Engine/Source/Editor/BlueprintGraph/Private/K2Node_Composite.cpp b/Engine/Source/Editor/BlueprintGraph/Private/K2Node_Composite.cpp
--- a/Engine/Source/Editor/BlueprintGraph/Private/K2Node_Composite.cpp
+++ b/Engine/Source/Editor/BlueprintGraph/Private/K2Node_Composite.cpp
@@ -95,7 +95,7 @@ void UK2Node_Composite::PostPasteNode()
 			if(UK2Node_Event* Event = Cast<UK2Node_Event>(Node))
 			{
 				UBlueprint* BP = FBlueprintEditorUtils::FindBlueprintForGraphChecked(BoundGraph);
-				if(FBlueprintEditorUtils::FindOverrideForFunction(BP, Event->EventSignatureClass, Event->EventSignatureName))
+				if(FBlueprintEditorUtils::FindOverrideForFunction(BP, Event->EventReference.GetMemberParentClass(Event), Event->EventReference.GetMemberName()))
 				{
 					FBlueprintEditorUtils::RemoveNode(BP, Node, true);
 					NodeIndex--;
diff --git a/Engine/Source/Editor/BlueprintGraph/Private/K2Node_DelegateSet.cpp b/Engine/Source/Editor/BlueprintGraph/Private/K2Node_DelegateSet.cpp
--- a/Engine/Source/Editor/BlueprintGraph/Private/K2Node_DelegateSet.cpp
+++ b/Engine/Source/Editor/BlueprintGraph/Private/K2Node_DelegateSet.cpp
@@ -272,8 +272,7 @@ void UK2Node_DelegateSet::ExpandNode(class FKismetCompilerContext& CompilerConte
 
 			// First, create an event node matching the delegate signature
 			UK2Node_Event* DelegateEvent = CompilerContext.SpawnIntermediateNode<UK2Node_Event>(this, SourceGraph);
-			DelegateEvent->EventSignatureClass = Cast<UClass>(TargetFunction->GetOuter());
-			DelegateEvent->EventSignatureName = TargetFunction->GetFName();
+			DelegateEvent->EventReference.SetFromField<UFunction>(TargetFunction, false);
 			DelegateEvent->CustomFunctionName = GetDelegateTargetEntryPointName();
 			DelegateEvent->bInternalEvent = true;
 			DelegateEvent->AllocateDefaultPins();
diff --git a/Engine/Source/Editor/BlueprintGraph/Private/K2Node_Event.cpp b/Engine/Source/Editor/BlueprintGraph/Private/K2Node_Event.cpp
--- a/Engine/Source/Editor/BlueprintGraph/Private/K2Node_Event.cpp
+++ b/Engine/Source/Editor/BlueprintGraph/Private/K2Node_Event.cpp
@@ -13,7 +13,7 @@ bool UK2Node_Event::IsCosmeticTickEvent() const
 {
 	// Special case for EventTick/ReceiveTick that is conditionally executed by a separate bool rather than function flag.
 	static const FName EventTickName(TEXT("ReceiveTick"));
-	if (EventSignatureName == EventTickName)
+	if (EventReference.GetMemberName() == EventTickName)
 	{
 		const UBlueprint* Blueprint = GetBlueprint();
 		if (Blueprint)
@@ -41,14 +41,35 @@ UK2Node_Event::UK2Node_Event(const FObjectInitializer& ObjectInitializer)
 
 void UK2Node_Event::Serialize(FArchive& Ar)
 {
+	// This deprecation warning suppression is due to the workaround. For more
+	// details read comment in K2Node_Event.h near EventSignatureName property
+	// definition.
+
+	PRAGMA_DISABLE_DEPRECATION_WARNINGS
+
+	// Workaround, so fields EventSignatureName and EventSignatureClass work
+	// until ObjectVersion.h lock is lifted.
+	if (Ar.IsSaving())
+	{
+		EventSignatureName = EventReference.GetMemberName();
+		EventSignatureClass = EventReference.GetMemberParentClass(this);
+	}
+
 	Super::Serialize(Ar);
 
 	// Fix up legacy nodes that may not yet have a delegate pin
-	if(Ar.IsLoading() && !FindPin(DelegateOutputName))
+	if(Ar.IsLoading())
 	{
-		const UEdGraphSchema_K2* K2Schema = GetDefault<UEdGraphSchema_K2>();
-		CreatePin(EGPD_Output, K2Schema->PC_Delegate, TEXT(""), NULL, false, false, DelegateOutputName);
+		if(!FindPin(DelegateOutputName))
+		{
+			const UEdGraphSchema_K2* K2Schema = GetDefault<UEdGraphSchema_K2>();
+			CreatePin(EGPD_Output, K2Schema->PC_Delegate, TEXT(""), NULL, false, false, DelegateOutputName);
+		}
+
+		EventReference.SetExternalMember(EventSignatureName, EventSignatureClass);
 	}
+
+	PRAGMA_ENABLE_DEPRECATION_WARNINGS
 }
 
 FNodeHandlingFunctor* UK2Node_Event::CreateNodeHandler(FKismetCompilerContext& CompilerContext) const
@@ -66,9 +87,9 @@ FText UK2Node_Event::GetNodeTitle(ENodeTitleType::Type TitleType) const
 {
 	if (bOverrideFunction || (CustomFunctionName == NAME_None))
 	{
-		FString FunctionName = EventSignatureName.ToString(); // If we fail to find the function, still want to write something on the node.
+		FString FunctionName = EventReference.GetMemberName().ToString(); // If we fail to find the function, still want to write something on the node.
 
-		if (UFunction* Function = FindField<UFunction>(EventSignatureClass, EventSignatureName))
+		if (UFunction* Function = EventReference.ResolveMember<UFunction>(this))
 		{
 			FunctionName = UEdGraphSchema_K2::GetFriendlySignatureName(Function);
 		}
@@ -77,9 +98,9 @@ FText UK2Node_Event::GetNodeTitle(ENodeTitleType::Type TitleType) const
 		Args.Add(TEXT("FunctionName"), FText::FromString(FunctionName));
 		FText Title = FText::Format(NSLOCTEXT("K2Node", "Event_Name", "Event {FunctionName}"), Args);
 
-		if(TitleType == ENodeTitleType::FullTitle && EventSignatureClass != NULL && EventSignatureClass->IsChildOf(UInterface::StaticClass()))
+		if(TitleType == ENodeTitleType::FullTitle && EventReference.GetMemberParentClass(this) != nullptr && EventReference.GetMemberParentClass(this)->IsChildOf(UInterface::StaticClass()))
 		{
-			const FText SignatureClassAsText = FBlueprintEditorUtils::GetFriendlyClassDisplayName(EventSignatureClass);
+			const FText SignatureClassAsText = FBlueprintEditorUtils::GetFriendlyClassDisplayName(EventReference.GetMemberParentClass(this));
 
 			FFormatNamedArguments FullTitleArgs;
 			FullTitleArgs.Add(TEXT("Title"), Title);
@@ -98,7 +119,7 @@ FText UK2Node_Event::GetNodeTitle(ENodeTitleType::Type TitleType) const
 
 FText UK2Node_Event::GetTooltipText() const
 {
-	UFunction* Function = FindField<UFunction>(EventSignatureClass, EventSignatureName);
+	UFunction* Function = EventReference.ResolveMember<UFunction>(this);
 	if (CachedTooltip.IsOutOfDate() && (Function != nullptr))
 	{
 		CachedTooltip = FText::FromString(UK2Node_CallFunction::GetDefaultTooltipForFunction(Function));
@@ -137,7 +158,7 @@ FString UK2Node_Event::GetKeywords() const
 {
 	FString Keywords;
 
-	UFunction* Function = FindField<UFunction>(EventSignatureClass, EventSignatureName);
+	UFunction* Function = EventReference.ResolveMember<UFunction>(this);
 	if (Function != NULL)
 	{
 		Keywords = UK2Node_CallFunction::GetKeywordsForFunction( Function );
@@ -148,7 +169,7 @@ FString UK2Node_Event::GetKeywords() const
 
 FString UK2Node_Event::GetDocumentationLink() const
 {
-	if ( EventSignatureClass != NULL )
+	if ( UClass* EventSignatureClass = EventReference.GetMemberParentClass(this))
 	{
 		return FString::Printf(TEXT("Shared/Types/%s%s"), EventSignatureClass->GetPrefixCPP(), *EventSignatureClass->GetName());
 	}
@@ -158,7 +179,7 @@ FString UK2Node_Event::GetDocumentationLink() const
 
 FString UK2Node_Event::GetDocumentationExcerptName() const
 {
-	return EventSignatureName.ToString();
+	return EventReference.GetMemberName().ToString();
 }
 
 void UK2Node_Event::PostReconstructNode()
@@ -180,7 +201,7 @@ void UK2Node_Event::UpdateDelegatePin(bool bSilent)
 	UFunction* NewSignature = NULL;
 	if(bOverrideFunction)
 	{
-		NewSignature = EventSignatureClass->FindFunctionByName(EventSignatureName);
+		NewSignature = EventReference.ResolveMember<UFunction>(this);
 	}
 	else if(UBlueprint* Blueprint = GetBlueprint())
 	{
@@ -210,31 +231,12 @@ void UK2Node_Event::PinConnectionListChanged(UEdGraphPin* Pin)
 
 FName UK2Node_Event::GetFunctionName() const
 {
-	return bOverrideFunction ? EventSignatureName : CustomFunctionName;
+	return bOverrideFunction ? EventReference.GetMemberName() : CustomFunctionName;
 }
 
 UFunction* UK2Node_Event::FindEventSignatureFunction()
 {
-	UFunction* Function = FindField<UFunction>(EventSignatureClass, EventSignatureName);
-
-	// First try remap table
-	if ((Function == NULL) && (EventSignatureClass != NULL))
-	{
-		Function = Cast<UFunction>(FindRemappedField(EventSignatureClass, EventSignatureName));
-		if( Function )
-		{
-			// Found a remapped property, update the node
-			EventSignatureName = Function->GetFName();
-			EventSignatureClass = Cast<UClass>(Function->GetOuter());
-		}
-	}
-
-	if (Function == nullptr)
-	{
-		Function = FindDelegateSignature(EventSignatureName);
-	}
-
-	return Function;
+	return EventReference.ResolveMember<UFunction>(this);
 }
 
 void UK2Node_Event::AllocateDefaultPins()
@@ -260,9 +262,9 @@ void UK2Node_Event::ValidateNodeDuringCompilation(class FCompilerResultsLog& Mes
 	Super::ValidateNodeDuringCompilation(MessageLog);
 
 	// If we are overriding a function, but we can;t find the function we are overriding, that is a compile error
-	if(bOverrideFunction && FindField<UFunction>(EventSignatureClass, EventSignatureName) == NULL)
+	if(bOverrideFunction && EventReference.ResolveMember<UFunction>(this) == nullptr)
 	{
-		MessageLog.Error(*FString::Printf(*NSLOCTEXT("KismetCompiler", "MissingEventSig_Error", "Missing Event '%s' for @@").ToString(), *EventSignatureName.ToString()), this);
+		MessageLog.Error(*FString::Printf(*NSLOCTEXT("KismetCompiler", "MissingEventSig_Error", "Missing Event '%s' for @@").ToString(), *EventReference.GetMemberName().ToString()), this);
 	}
 }
 
@@ -281,11 +283,14 @@ void UK2Node_Event::GetRedirectPinNames(const UEdGraphPin& Pin, TArray<FString>&
 		const FString& OldPinName = RedirectPinNames[0];
 
 		// first add functionname.param
-		RedirectPinNames.Add(FString::Printf(TEXT("%s.%s"), *EventSignatureName.ToString(), *OldPinName));
+		RedirectPinNames.Add(FString::Printf(TEXT("%s.%s"), *EventReference.GetMemberName().ToString(), *OldPinName));
 		// if there is class, also add an option for class.functionname.param
-		if ( EventSignatureClass!=NULL )
+		if ( EventReference.ResolveMember<UFunction>(this) != nullptr )
 		{
-			RedirectPinNames.Add(FString::Printf(TEXT("%s.%s.%s"), *EventSignatureClass->GetName(), *EventSignatureName.ToString(), *OldPinName));
+			if(const UClass* EventSignatureClass = EventReference.GetMemberParentClass(this))
+			{
+				RedirectPinNames.Add(FString::Printf(TEXT("%s.%s.%s"), *EventSignatureClass->GetName(), *EventReference.GetMemberName().ToString(), *OldPinName));
+			}
 		}
 	}
 }
@@ -427,16 +432,16 @@ bool UK2Node_Event::CanPasteHere(const UEdGraph* TargetGraph) const
 				if(bOverrideFunction)
 				{
 					// If the function name is hidden by the parent class, don't paste this event
-					bDisallowPaste = EventSignatureClass == Blueprint->ParentClass
-						&& ExcludedEventNames.Contains(EventSignatureName.ToString());
+					bDisallowPaste = EventReference.GetMemberParentClass(this) == Blueprint->ParentClass
+						&& ExcludedEventNames.Contains(EventReference.GetMemberName().ToString());
 					if(!bDisallowPaste)
 					{
 						// If the event function is already handled in this Blueprint, don't paste this event
 						for(int32 i = 0; i < ExistingEventNodes.Num() && !bDisallowPaste; ++i)
 						{
 							bDisallowPaste = ExistingEventNodes[i]->bOverrideFunction
-								&& ExistingEventNodes[i]->EventSignatureName == EventSignatureName
-								&& ExistingEventNodes[i]->EventSignatureClass == EventSignatureClass;
+								&& ExistingEventNodes[i]->EventReference.GetMemberName() == EventReference.GetMemberName()
+								&& ExistingEventNodes[i]->EventReference.GetMemberParentClass(this) == EventReference.GetMemberParentClass(this);
 						}
 
 						// We need to also check for 'const' BPIE methods that might already be implemented as functions with a read-only 'self' context (these were previously implemented as events)
@@ -452,7 +457,7 @@ bool UK2Node_Event::CanPasteHere(const UEdGraph* TargetGraph) const
 								{
 									UK2Node_FunctionEntry* ExistingFunctionEntryNode = *NodeIt;
 									bDisallowPaste = ExistingFunctionEntryNode->bEnforceConstCorrectness
-										&& ExistingFunctionEntryNode->SignatureName == EventSignatureName;
+										&& ExistingFunctionEntryNode->SignatureName == EventReference.GetMemberName();
 								}
 							}
 						}
@@ -460,21 +465,21 @@ bool UK2Node_Event::CanPasteHere(const UEdGraph* TargetGraph) const
 						if(!bDisallowPaste)
 						{
 							// If the signature class is not implemented by the Blueprint parent class or an interface, don't paste this event
-							bDisallowPaste = !Blueprint->ParentClass->IsChildOf(EventSignatureClass)
-								&& !ImplementedInterfaceClasses.Contains(EventSignatureClass);
+							bDisallowPaste = !Blueprint->ParentClass->IsChildOf(EventReference.GetMemberParentClass(this))
+								&& !ImplementedInterfaceClasses.Contains(EventReference.GetMemberParentClass(this));
 							if(bDisallowPaste)
 							{
-								UE_LOG(LogBlueprint, Log, TEXT("Cannot paste event node (%s) directly because the event signature class (%s) is incompatible with this Blueprint."), *GetFName().ToString(), EventSignatureClass ? *EventSignatureClass->GetFName().ToString() : TEXT("NONE"));
+								UE_LOG(LogBlueprint, Log, TEXT("Cannot paste event node (%s) directly because the event signature class (%s) is incompatible with this Blueprint."), *GetFName().ToString(), EventReference.GetMemberParentClass(this) ? *EventReference.GetMemberParentClass(this)->GetFName().ToString() : TEXT("NONE"));
 							}
 						}
 						else
 						{
-							UE_LOG(LogBlueprint, Log, TEXT("Cannot paste event node (%s) directly because the event function (%s) is already handled."), *GetFName().ToString(), *EventSignatureName.ToString());
+							UE_LOG(LogBlueprint, Log, TEXT("Cannot paste event node (%s) directly because the event function (%s) is already handled."), *GetFName().ToString(), *EventReference.GetMemberName().ToString());
 						}
 					}
 					else
 					{
-						UE_LOG(LogBlueprint, Log, TEXT("Cannot paste event node (%s) directly because the event function (%s) is hidden by the Blueprint parent class (%s)."), *GetFName().ToString(), *EventSignatureName.ToString(), EventSignatureClass ? *EventSignatureClass->GetFName().ToString() : TEXT("NONE"));
+						UE_LOG(LogBlueprint, Log, TEXT("Cannot paste event node (%s) directly because the event function (%s) is hidden by the Blueprint parent class (%s)."), *GetFName().ToString(), *EventReference.GetMemberName().ToString(), EventReference.GetMemberParentClass(this) ? *EventReference.GetMemberParentClass(this)->GetFName().ToString() : TEXT("NONE"));
 					}
 				}
 				else if(CustomFunctionName != NAME_None)
@@ -485,14 +490,14 @@ bool UK2Node_Event::CanPasteHere(const UEdGraph* TargetGraph) const
 					if(!bDisallowPaste)
 					{
 						// Handle events that have a custom function name with an actual signature name/class that is not an override (e.g. AnimNotify events)
-						if(EventSignatureName != NAME_None)
+						if(EventReference.GetMemberName() != NAME_None)
 						{
 							// If the signature class is not implemented by the Blueprint parent class or an interface, don't paste this event
-							bDisallowPaste = !Blueprint->ParentClass->IsChildOf(EventSignatureClass)
-								&& !ImplementedInterfaceClasses.Contains(EventSignatureClass);
+							bDisallowPaste = !Blueprint->ParentClass->IsChildOf(EventReference.GetMemberParentClass(this))
+								&& !ImplementedInterfaceClasses.Contains(EventReference.GetMemberParentClass(this));
 							if(bDisallowPaste)
 							{
-								UE_LOG(LogBlueprint, Log, TEXT("Cannot paste event node (%s) directly because the custom event function (%s) with event signature name (%s) has an event signature class (%s) that is incompatible with this Blueprint."), *GetFName().ToString(), *CustomFunctionName.ToString(), *EventSignatureName.ToString(), EventSignatureClass ? *EventSignatureClass->GetFName().ToString() : TEXT("NONE"));
+								UE_LOG(LogBlueprint, Log, TEXT("Cannot paste event node (%s) directly because the custom event function (%s) with event signature name (%s) has an event signature class (%s) that is incompatible with this Blueprint."), *GetFName().ToString(), *CustomFunctionName.ToString(), *EventReference.GetMemberName().ToString(), EventReference.GetMemberParentClass(this) ? *EventReference.GetMemberParentClass(this)->GetFName().ToString() : TEXT("NONE"));
 							}
 						}
 					}
@@ -503,7 +508,7 @@ bool UK2Node_Event::CanPasteHere(const UEdGraph* TargetGraph) const
 				}
 				else
 				{
-					UE_LOG(LogBlueprint, Log, TEXT("Cannot paste event node (%s) directly because the event configuration is not specifically handled (EventSignatureName=%s, EventSignatureClass=%s)."), *GetFName().ToString(), *EventSignatureName.ToString(), EventSignatureClass ? *EventSignatureClass->GetFName().ToString() : TEXT("NONE"));
+					UE_LOG(LogBlueprint, Log, TEXT("Cannot paste event node (%s) directly because the event configuration is not specifically handled (EventSignatureName=%s, EventSignatureClass=%s)."), *GetFName().ToString(), *EventReference.GetMemberName().ToString(), EventReference.GetMemberParentClass(this) ? *EventReference.GetMemberParentClass(this)->GetFName().ToString() : TEXT("NONE"));
 				}
 			}
 			else
@@ -595,7 +600,7 @@ void UK2Node_Event::ExpandNode(class FKismetCompilerContext& CompilerContext, UE
 
 FName UK2Node_Event::GetCornerIcon() const
 {
-	if (UFunction* Function = FindField<UFunction>(EventSignatureClass, EventSignatureName))
+	if (UFunction* Function = EventReference.ResolveMember<UFunction>(this))
 	{
 		if (bOverrideFunction || (CustomFunctionName == NAME_None))
 		{
@@ -616,7 +621,7 @@ FName UK2Node_Event::GetCornerIcon() const
 		return TEXT("Graph.Replication.AuthorityOnly");
 	}
 
-	if(EventSignatureClass != NULL && EventSignatureClass->IsChildOf(UInterface::StaticClass()))
+	if(EventReference.GetMemberParentClass(this) != nullptr && EventReference.GetMemberParentClass(this)->IsChildOf(UInterface::StaticClass()))
 	{
 		return TEXT("Graph.Event.InterfaceEventIcon");
 	}
@@ -629,7 +634,7 @@ FText UK2Node_Event::GetToolTipHeading() const
 	FText Heading = Super::GetToolTipHeading();
 
 	FText EventHeading = FText::GetEmpty();
-	if (UFunction* Function = FindField<UFunction>(EventSignatureClass, EventSignatureName))
+	if (UFunction* Function = EventReference.ResolveMember<UFunction>(this))
 	{
 		if (bOverrideFunction || (CustomFunctionName == NAME_None))
 		{
@@ -648,7 +653,7 @@ FText UK2Node_Event::GetToolTipHeading() const
 	{
 		EventHeading = LOCTEXT("ServerOnlyEvent", "Server Only");
 	}
-	else if(EventHeading.IsEmpty() && (EventSignatureClass != NULL) && EventSignatureClass->IsChildOf(UInterface::StaticClass()))
+	else if(EventHeading.IsEmpty() && (EventReference.GetMemberParentClass(this) != nullptr) && EventReference.GetMemberParentClass(this)->IsChildOf(UInterface::StaticClass()))
 	{
 		EventHeading = LOCTEXT("InterfaceEvent", "Interface Event");
 	}
@@ -675,7 +680,7 @@ void UK2Node_Event::GetNodeAttributes( TArray<TKeyValuePair<FString, FString>>&
 FText UK2Node_Event::GetMenuCategory() const
 {
 	FString FuncSubCategory;
-	if (UFunction* Function = FindField<UFunction>(EventSignatureClass, EventSignatureName))
+	if (UFunction* Function = EventReference.ResolveMember<UFunction>(this))
 	{
 		FuncSubCategory = UK2Node_CallFunction::GetDefaultCategoryForFunction(Function, TEXT(""));
 	}
@@ -686,7 +691,7 @@ FText UK2Node_Event::GetMenuCategory() const
 
 bool UK2Node_Event::IsDeprecated() const
 {
-	if (UFunction* Function = FindField<UFunction>(EventSignatureClass, EventSignatureName))
+	if (UFunction* Function = EventReference.ResolveMember<UFunction>(this))
 	{
 		return Function->HasMetaData(FBlueprintMetadata::MD_DeprecatedFunction);
 	}
@@ -696,7 +701,7 @@ bool UK2Node_Event::IsDeprecated() const
 
 FString UK2Node_Event::GetDeprecationMessage() const
 {
-	if (UFunction* Function = FindField<UFunction>(EventSignatureClass, EventSignatureName))
+	if (UFunction* Function = EventReference.ResolveMember<UFunction>(this))
 	{
 		if (Function->HasMetaData(FBlueprintMetadata::MD_DeprecationMessage))
 		{
@@ -709,14 +714,14 @@ FString UK2Node_Event::GetDeprecationMessage() const
 
 UObject* UK2Node_Event::GetJumpTargetForDoubleClick() const
 {
-	if(EventSignatureClass != NULL && EventSignatureClass->ClassGeneratedBy != NULL && EventSignatureClass->ClassGeneratedBy->IsA(UBlueprint::StaticClass()))
+	if(EventReference.GetMemberParentClass(this) != NULL && EventReference.GetMemberParentClass(this)->ClassGeneratedBy != NULL && EventReference.GetMemberParentClass(this)->ClassGeneratedBy->IsA(UBlueprint::StaticClass()))
 	{
-		UBlueprint* Blueprint = CastChecked<UBlueprint>(EventSignatureClass->ClassGeneratedBy);
+		UBlueprint* Blueprint = CastChecked<UBlueprint>(EventReference.GetMemberParentClass(this)->ClassGeneratedBy);
 		TArray<UEdGraph*> Graphs;
 		Blueprint->GetAllGraphs(Graphs);
 		for(auto It(Graphs.CreateConstIterator()); It; It++)
 		{
-			if((*It)->GetFName() == EventSignatureName)
+			if((*It)->GetFName() == EventReference.GetMemberName())
 			{
 				return *It;
 			}
diff --git a/Engine/Source/Editor/BlueprintGraph/Private/K2Node_InputAction.cpp b/Engine/Source/Editor/BlueprintGraph/Private/K2Node_InputAction.cpp
--- a/Engine/Source/Editor/BlueprintGraph/Private/K2Node_InputAction.cpp
+++ b/Engine/Source/Editor/BlueprintGraph/Private/K2Node_InputAction.cpp
@@ -118,8 +118,7 @@ void UK2Node_InputAction::CreateInputActionEvent(FKismetCompilerContext& Compile
 		InputActionEvent->bExecuteWhenPaused = bExecuteWhenPaused;
 		InputActionEvent->bOverrideParentBinding = bOverrideParentBinding;
 		InputActionEvent->InputKeyEvent = InputKeyEvent;
-		InputActionEvent->EventSignatureName = TEXT("InputActionHandlerDynamicSignature__DelegateSignature");
-		InputActionEvent->EventSignatureClass = UInputComponent::StaticClass();
+		InputActionEvent->EventReference.SetExternalDelegateMember(FName(TEXT("InputActionHandlerDynamicSignature__DelegateSignature")));
 		InputActionEvent->bInternalEvent = true;
 		InputActionEvent->AllocateDefaultPins();
 
diff --git a/Engine/Source/Editor/BlueprintGraph/Private/K2Node_InputAxisEvent.cpp b/Engine/Source/Editor/BlueprintGraph/Private/K2Node_InputAxisEvent.cpp
--- a/Engine/Source/Editor/BlueprintGraph/Private/K2Node_InputAxisEvent.cpp
+++ b/Engine/Source/Editor/BlueprintGraph/Private/K2Node_InputAxisEvent.cpp
@@ -17,8 +17,7 @@ UK2Node_InputAxisEvent::UK2Node_InputAxisEvent(const FObjectInitializer& ObjectI
 	bOverrideParentBinding = true;
 	bInternalEvent = true;
 
-	EventSignatureName = TEXT("InputAxisHandlerDynamicSignature__DelegateSignature");
-	EventSignatureClass = UInputComponent::StaticClass();
+	EventReference.SetExternalDelegateMember(FName(TEXT("InputAxisHandlerDynamicSignature__DelegateSignature")));
 }
 
 void UK2Node_InputAxisEvent::PostLoad()
diff --git a/Engine/Source/Editor/BlueprintGraph/Private/K2Node_InputAxisKeyEvent.cpp b/Engine/Source/Editor/BlueprintGraph/Private/K2Node_InputAxisKeyEvent.cpp
--- a/Engine/Source/Editor/BlueprintGraph/Private/K2Node_InputAxisKeyEvent.cpp
+++ b/Engine/Source/Editor/BlueprintGraph/Private/K2Node_InputAxisKeyEvent.cpp
@@ -19,8 +19,7 @@ UK2Node_InputAxisKeyEvent::UK2Node_InputAxisKeyEvent(const FObjectInitializer& O
 	bOverrideParentBinding = true;
 	bInternalEvent = true;
 
-	EventSignatureName = TEXT("InputAxisHandlerDynamicSignature__DelegateSignature");
-	EventSignatureClass = UInputComponent::StaticClass();
+	EventReference.SetExternalDelegateMember(FName(TEXT("InputAxisHandlerDynamicSignature__DelegateSignature")));
 }
 
 void UK2Node_InputAxisKeyEvent::Initialize(const FKey InAxisKey)
diff --git a/Engine/Source/Editor/BlueprintGraph/Private/K2Node_InputKey.cpp b/Engine/Source/Editor/BlueprintGraph/Private/K2Node_InputKey.cpp
--- a/Engine/Source/Editor/BlueprintGraph/Private/K2Node_InputKey.cpp
+++ b/Engine/Source/Editor/BlueprintGraph/Private/K2Node_InputKey.cpp
@@ -268,8 +268,7 @@ void UK2Node_InputKey::CreateInputKeyEvent(FKismetCompilerContext& CompilerConte
 		InputKeyEvent->bExecuteWhenPaused = bExecuteWhenPaused;
 		InputKeyEvent->bOverrideParentBinding = bOverrideParentBinding;
 		InputKeyEvent->InputKeyEvent = KeyEvent;
-		InputKeyEvent->EventSignatureName = TEXT("InputActionHandlerDynamicSignature__DelegateSignature");
-		InputKeyEvent->EventSignatureClass = UInputComponent::StaticClass();
+		InputKeyEvent->EventReference.SetExternalDelegateMember(FName(TEXT("InputActionHandlerDynamicSignature__DelegateSignature")));
 		InputKeyEvent->bInternalEvent = true;
 		InputKeyEvent->AllocateDefaultPins();
 
diff --git a/Engine/Source/Editor/BlueprintGraph/Private/K2Node_InputTouch.cpp b/Engine/Source/Editor/BlueprintGraph/Private/K2Node_InputTouch.cpp
--- a/Engine/Source/Editor/BlueprintGraph/Private/K2Node_InputTouch.cpp
+++ b/Engine/Source/Editor/BlueprintGraph/Private/K2Node_InputTouch.cpp
@@ -187,8 +187,7 @@ void UK2Node_InputTouch::ExpandNode(FKismetCompilerContext& CompilerContext, UEd
 			InputTouchEvent->bExecuteWhenPaused = bExecuteWhenPaused;
 			InputTouchEvent->bOverrideParentBinding = bOverrideParentBinding;
 			InputTouchEvent->InputKeyEvent = (*PinIt).EventType;
-			InputTouchEvent->EventSignatureName = TEXT("InputTouchHandlerDynamicSignature__DelegateSignature");
-			InputTouchEvent->EventSignatureClass = UInputComponent::StaticClass();
+			InputTouchEvent->EventReference.SetExternalDelegateMember(FName(TEXT("InputTouchHandlerDynamicSignature__DelegateSignature")));
 			InputTouchEvent->bInternalEvent = true;
 			InputTouchEvent->AllocateDefaultPins();
 
@@ -229,8 +228,7 @@ void UK2Node_InputTouch::ExpandNode(FKismetCompilerContext& CompilerContext, UEd
 			InputTouchEvent->bConsumeInput = bConsumeInput;
 			InputTouchEvent->bExecuteWhenPaused = bExecuteWhenPaused;
 			InputTouchEvent->bOverrideParentBinding = bOverrideParentBinding;
-			InputTouchEvent->EventSignatureName = TEXT("InputTouchHandlerDynamicSignature__DelegateSignature");
-			InputTouchEvent->EventSignatureClass = UInputComponent::StaticClass();
+			InputTouchEvent->EventReference.SetExternalDelegateMember(FName(TEXT("InputTouchHandlerDynamicSignature__DelegateSignature")));
 			InputTouchEvent->bInternalEvent = true;
 			InputTouchEvent->AllocateDefaultPins();
 
diff --git a/Engine/Source/Editor/BlueprintGraph/Private/K2Node_InputVectorAxisEvent.cpp b/Engine/Source/Editor/BlueprintGraph/Private/K2Node_InputVectorAxisEvent.cpp
--- a/Engine/Source/Editor/BlueprintGraph/Private/K2Node_InputVectorAxisEvent.cpp
+++ b/Engine/Source/Editor/BlueprintGraph/Private/K2Node_InputVectorAxisEvent.cpp
@@ -10,7 +10,7 @@
 UK2Node_InputVectorAxisEvent::UK2Node_InputVectorAxisEvent(const FObjectInitializer& ObjectInitializer)
 : Super(ObjectInitializer)
 {
-	EventSignatureName = TEXT("InputVectorAxisHandlerDynamicSignature__DelegateSignature");
+	EventReference.SetExternalDelegateMember(FName(TEXT("InputAxisHandlerDynamicSignature__DelegateSignature")));
 }
 
 void UK2Node_InputVectorAxisEvent::ValidateNodeDuringCompilation(class FCompilerResultsLog& MessageLog) const
diff --git a/Engine/Source/Editor/BlueprintGraph/Private/K2Node_MatineeController.cpp b/Engine/Source/Editor/BlueprintGraph/Private/K2Node_MatineeController.cpp
--- a/Engine/Source/Editor/BlueprintGraph/Private/K2Node_MatineeController.cpp
+++ b/Engine/Source/Editor/BlueprintGraph/Private/K2Node_MatineeController.cpp
@@ -117,8 +117,7 @@ void UK2Node_MatineeController::ExpandNode(FKismetCompilerContext& CompilerConte
 				FName EventFuncName = MatineeActor->GetFunctionNameForEvent( FName(*(MatineePin->PinName)) );
 
 				UK2Node_Event* MatineeEventNode = CompilerContext.SpawnIntermediateNode<UK2Node_Event>(this, SourceGraph);
-				MatineeEventNode->EventSignatureName = MatineeEventSig->GetFName();
-				MatineeEventNode->EventSignatureClass = AMatineeActor::StaticClass();
+				MatineeEventNode->EventReference.SetFromField<UFunction>(MatineeEventSig, false);
 				MatineeEventNode->CustomFunctionName = EventFuncName;
 				MatineeEventNode->bInternalEvent = true;
 				MatineeEventNode->AllocateDefaultPins();
diff --git a/Engine/Source/Editor/GameplayAbilitiesEditor/Private/K2Node_GameplayCueEvent.cpp b/Engine/Source/Editor/GameplayAbilitiesEditor/Private/K2Node_GameplayCueEvent.cpp
--- a/Engine/Source/Editor/GameplayAbilitiesEditor/Private/K2Node_GameplayCueEvent.cpp
+++ b/Engine/Source/Editor/GameplayAbilitiesEditor/Private/K2Node_GameplayCueEvent.cpp
@@ -16,8 +16,7 @@
 UK2Node_GameplayCueEvent::UK2Node_GameplayCueEvent(const FObjectInitializer& ObjectInitializer)
 	: Super(ObjectInitializer)
 {
-	EventSignatureName = GAMEPLAYABILITIES_BlueprintCustomHandler;
-	EventSignatureClass = UGameplayCueInterface::StaticClass();
+	EventReference.SetExternalMember(GAMEPLAYABILITIES_BlueprintCustomHandler, UGameplayCueInterface::StaticClass());
 }
 
 FText UK2Node_GameplayCueEvent::GetTooltipText() const
diff --git a/Engine/Source/Editor/Kismet/Private/BlueprintPaletteFavorites.cpp b/Engine/Source/Editor/Kismet/Private/BlueprintPaletteFavorites.cpp
--- a/Engine/Source/Editor/Kismet/Private/BlueprintPaletteFavorites.cpp
+++ b/Engine/Source/Editor/Kismet/Private/BlueprintPaletteFavorites.cpp
@@ -79,11 +79,8 @@ static FBlueprintNodeSignature BlueprintPaletteFavoritesImpl::ConstructLegacySig
 		}
 		else if (UK2Node_Event const* EventNode = Cast<UK2Node_Event const>(NodeTemplate))
 		{
-			if (EventNode->EventSignatureClass != nullptr)
-			{
-				SignatureSubObject = FindField<UFunction>(EventNode->EventSignatureClass, EventNode->EventSignatureName);
-				bIsSupported = (SignatureSubObject != nullptr);
-			}
+			SignatureSubObject = EventNode->EventReference.ResolveMember<UFunction>(EventNode);
+			bIsSupported = (SignatureSubObject != nullptr);
 		}
 		else if (UK2Node_MacroInstance const* MacroNode = Cast<UK2Node_MacroInstance const>(NodeTemplate))
 		{
diff --git a/Engine/Source/Editor/KismetCompiler/Private/KismetCompiler.cpp b/Engine/Source/Editor/KismetCompiler/Private/KismetCompiler.cpp
--- a/Engine/Source/Editor/KismetCompiler/Private/KismetCompiler.cpp
+++ b/Engine/Source/Editor/KismetCompiler/Private/KismetCompiler.cpp
@@ -152,7 +152,7 @@ void FKismetCompilerContext::CleanAndSanitizeClass(UBlueprintGeneratedClass* Cla
 	
 	UClass* ParentClass = Blueprint->ParentClass;
 
-	if(FKismetEditorUtilities::IsClassABlueprintSkeleton(ClassToClean))
+	if(CompileOptions.CompileType == EKismetCompileType::SkeletonOnly)
 	{
 		if(UBlueprint* BlueprintParent = Cast<UBlueprint>(Blueprint->ParentClass->ClassGeneratedBy))
 		{
@@ -1845,8 +1845,7 @@ void FKismetCompilerContext::BuildDynamicBindingObjects(UBlueprintGeneratedClass
 void FKismetCompilerContext::CreatePinEventNodeForTimelineFunction(UK2Node_Timeline* TimelineNode, UEdGraph* SourceGraph, FName FunctionName, const FString& PinName, FName ExecFuncName)
 {
 	UK2Node_Event* TimelineEventNode = SpawnIntermediateNode<UK2Node_Event>(TimelineNode, SourceGraph);
-	TimelineEventNode->EventSignatureName = ExecFuncName;
-	TimelineEventNode->EventSignatureClass = UTimelineComponent::StaticClass();
+	TimelineEventNode->EventReference.SetExternalMember(FunctionName, UTimelineComponent::StaticClass());
 	TimelineEventNode->CustomFunctionName = FunctionName; // Make sure we name this function the thing we are expecting
 	TimelineEventNode->bInternalEvent = true;
 	TimelineEventNode->AllocateDefaultPins();
@@ -2304,7 +2303,7 @@ FName FKismetCompilerContext::GetEventStubFunctionName(UK2Node_Event* SrcEventNo
 	// If we are overriding a function, we use the exact name for the event node
 	if (SrcEventNode->bOverrideFunction)
 	{
-		EventNodeName = SrcEventNode->EventSignatureName;
+		EventNodeName = SrcEventNode->EventReference.GetMemberName();
 	}
 	else
 	{
@@ -2366,8 +2365,8 @@ void FKismetCompilerContext::CreateFunctionStubForEvent(UK2Node_Event* SrcEventN
 	// Create an entry point
 	UK2Node_FunctionEntry* EntryNode = SpawnIntermediateNode<UK2Node_FunctionEntry>(SrcEventNode, ChildStubGraph);
 	EntryNode->NodePosX = -200;
-	EntryNode->SignatureClass = SrcEventNode->EventSignatureClass;
-	EntryNode->SignatureName = SrcEventNode->EventSignatureName;
+	EntryNode->SignatureClass = SrcEventNode->EventReference.GetMemberParentClass(SrcEventNode);
+	EntryNode->SignatureName = SrcEventNode->EventReference.GetMemberName();
 	EntryNode->CustomGeneratedFunctionName = EventNodeName;
 
 	if(!SrcEventNode->bOverrideFunction && SrcEventNode->IsUsedByAuthorityOnlyDelegate())
@@ -2563,8 +2562,8 @@ void FKismetCompilerContext::VerifyValidOverrideEvent(const UEdGraph* Graph)
 			{
 				const UK2Node_Event* EventNode = EntryPoints[EntryPointsIdx];
 				if( EventNode && EventNode->bOverrideFunction &&
-					(EventNode->EventSignatureClass == FuncClass) &&
-					(EventNode->EventSignatureName == FuncName))
+					(EventNode->EventReference.GetMemberParentClass(EventNode) == FuncClass) &&
+					(EventNode->EventReference.GetMemberName() == FuncName))
 				{
 					if (EventNode->IsDeprecated())
 					{
@@ -2654,7 +2653,7 @@ void FKismetCompilerContext::CreateAndProcessUbergraph()
 				for (int32 EntryIndex = 0; EntryIndex < EntryPoints.Num(); ++EntryIndex)
 				{
 					const UK2Node_Event* EventNode = EntryPoints[EntryIndex];
-					if( EventNode && (EventNode->EventSignatureName == FunctionName) )
+					if( EventNode && (EventNode->EventReference.GetMemberName() == FunctionName) )
 					{
 						bFoundEntry = true;
 						break;
@@ -2665,8 +2664,7 @@ void FKismetCompilerContext::CreateAndProcessUbergraph()
 				{
 					// Create an entry node stub, so that we have a entry point for interfaces to call to
 					UK2Node_Event* EventNode = SpawnIntermediateNode<UK2Node_Event>(NULL, ConsolidatedEventGraph);
-					EventNode->EventSignatureName = FunctionName;
-					EventNode->EventSignatureClass = InterfaceDesc.Interface;
+					EventNode->EventReference.SetExternalMember(FunctionName, InterfaceDesc.Interface);
 					EventNode->bOverrideFunction = true;
 					EventNode->AllocateDefaultPins();
 				}
diff --git a/Engine/Source/Editor/UnrealEd/Private/Kismet2/BlueprintEditorUtils.cpp b/Engine/Source/Editor/UnrealEd/Private/Kismet2/BlueprintEditorUtils.cpp
--- a/Engine/Source/Editor/UnrealEd/Private/Kismet2/BlueprintEditorUtils.cpp
+++ b/Engine/Source/Editor/UnrealEd/Private/Kismet2/BlueprintEditorUtils.cpp
@@ -1771,7 +1771,7 @@ bool FBlueprintEditorUtils::IsGraphNameUnique(UBlueprint* Blueprint, const FName
 				check(EventNode);
 
 				if( EventNode->CustomFunctionName == InName
-					|| EventNode->EventSignatureName == InName )
+					|| EventNode->EventReference.GetMemberName() == InName )
 				{
 					return false;
 				}
@@ -2234,8 +2234,8 @@ UK2Node_Event* FBlueprintEditorUtils::FindOverrideForFunction(const UBlueprint*
 		UK2Node_Event* EventNode = AllEvents[i];
 		check(EventNode);
 		if(	EventNode->bOverrideFunction == true &&
-			EventNode->EventSignatureClass == SignatureClass &&
-			EventNode->EventSignatureName == SignatureName )
+			EventNode->EventReference.GetMemberParentClass(EventNode) == SignatureClass &&
+			EventNode->EventReference.GetMemberName() == SignatureName )
 		{
 			return EventNode;
 		}
@@ -4017,7 +4017,7 @@ UFunction* FFunctionFromNodeHelper::FunctionFromNode(UK2Node* Node)
 	{
 		if (UK2Node_Event* EventNode = Cast<UK2Node_Event>(Node))
 		{
-			Function = SearchScope->FindFunctionByName(EventNode->EventSignatureName);
+			Function = EventNode->EventReference.ResolveMember<UFunction>((UClass*)(SearchScope));
 		}
 		else if (UK2Node_FunctionEntry* FunctionNode = Cast<UK2Node_FunctionEntry>(Node))
 		{
@@ -4708,14 +4708,14 @@ void FBlueprintEditorUtils::RemoveInterface(UBlueprint* Blueprint, const FName&
 		for(TArray<UK2Node_Event*>::TIterator NodeIt(AllEvents); NodeIt; ++NodeIt)
 		{
 			UK2Node_Event* EventNode = *NodeIt;
-			if( EventNode->EventSignatureClass == InterfaceClass )
+			if( EventNode->EventReference.GetMemberParentClass(EventNode) == InterfaceClass )
 			{
 				if(bPreserveFunctions)
 				{
 					// Create a custom event with the same name and signature
 					const FVector2D PreviousNodePos = FVector2D(EventNode->NodePosX, EventNode->NodePosY);
-					const FString PreviousNodeName = EventNode->EventSignatureName.ToString();
-					const UFunction* PreviousSignatureFunction = EventNode->EventSignatureClass->FindFunctionByName(EventNode->EventSignatureName);
+					const FString PreviousNodeName = EventNode->EventReference.GetMemberName().ToString();
+					const UFunction* PreviousSignatureFunction = EventNode->FindEventSignatureFunction();
 					check(PreviousSignatureFunction);
 					
 					UK2Node_CustomEvent* NewEvent = UK2Node_CustomEvent::CreateFromFunction(PreviousNodePos, EventNode->GetGraph(), PreviousNodeName, PreviousSignatureFunction, false);
@@ -4929,15 +4929,15 @@ void FBlueprintEditorUtils::ConformImplementedEvents(UBlueprint* Blueprint)
 				// If the event is loaded and is not a custom event
 				if(!EventNode->HasAnyFlags(RF_NeedLoad|RF_NeedPostLoad) && EventNode->bOverrideFunction)
 				{
-					const bool bEventNodeUsedByInterface = ImplementedInterfaceClasses.Contains(EventNode->EventSignatureClass);
+					const bool bEventNodeUsedByInterface = ImplementedInterfaceClasses.Contains(EventNode->EventReference.GetMemberParentClass(EventNode));
 					if (Blueprint->GeneratedClass && !bEventNodeUsedByInterface)
 					{
 						// See if the generated class implements an event with the given function signature
-						const UFunction* TargetFunction = Blueprint->GeneratedClass->FindFunctionByName(EventNode->EventSignatureName);
-						if (TargetFunction || EventGraphNames.Contains(EventNode->EventSignatureName))
+						const UFunction* TargetFunction = EventNode->EventReference.ResolveMember<UFunction>(Blueprint->GeneratedClass);
+						if (TargetFunction || EventGraphNames.Contains(EventNode->EventReference.GetMemberName()))
 						{
 							// The generated class implements the event but the function signature is not up-to-date
-							if (!Blueprint->GeneratedClass->IsChildOf(EventNode->EventSignatureClass))
+							if (!Blueprint->GeneratedClass->IsChildOf(EventNode->EventReference.GetMemberParentClass(EventNode)))
 							{
 								FFormatNamedArguments Args;
 								Args.Add(TEXT("NodeTitle"), EventNode->GetNodeTitle(ENodeTitleType::ListView));
@@ -4947,7 +4947,7 @@ void FBlueprintEditorUtils::ConformImplementedEvents(UBlueprint* Blueprint)
 								Blueprint->Message_Note( FText::Format(LOCTEXT("EventSignatureFixed_Note", "{NodeTitle} ({EventNodeName}) had an invalid function signature - it has now been fixed."), Args).ToString() );
 
 								// Fix up the event signature
-								EventNode->EventSignatureClass = Blueprint->GeneratedClass;
+								EventNode->EventReference.SetFromField<UFunction>(TargetFunction, false);
 							}
 						}
 						else
@@ -5035,12 +5035,12 @@ static void ConformInterfaceByName(UBlueprint* Blueprint, FBPInterfaceDescriptio
 		for (UK2Node_Event* EventNode : ImplementedEvents)
 		{
 			// if this event belongs to something other than this interface
-			if (EventNode->EventSignatureClass != CurrentInterfaceDesc.Interface)
+			if (EventNode->EventReference.GetMemberParentClass(EventNode) != CurrentInterfaceDesc.Interface)
 			{
 				continue;
 			}
 
-			UFunction* InterfaceFunction = CurrentInterfaceDesc.Interface->FindFunctionByName(EventNode->EventSignatureName);
+			UFunction* InterfaceFunction = EventNode->EventReference.ResolveMember<UFunction>(CurrentInterfaceDesc.Interface);
 			// if the function is still ok as an event, no need to try and fix it up
 			if (UEdGraphSchema_K2::FunctionCanBePlacedAsEvent(InterfaceFunction))
 			{
@@ -5056,7 +5056,7 @@ static void ConformInterfaceByName(UBlueprint* Blueprint, FBPInterfaceDescriptio
 			}
 
 			// grab the function's name before we delete the node
-			FName const FunctionName = EventNode->EventSignatureName;
+			FName const FunctionName = EventNode->EventReference.GetMemberName();
 			// destroy the old event node (this will also break all pin links and remove it from the graph)
 			EventNode->DestroyNode();
 
diff --git a/Engine/Source/Editor/UnrealEd/Private/Kismet2/Kismet2.cpp b/Engine/Source/Editor/UnrealEd/Private/Kismet2/Kismet2.cpp
--- a/Engine/Source/Editor/UnrealEd/Private/Kismet2/Kismet2.cpp
+++ b/Engine/Source/Editor/UnrealEd/Private/Kismet2/Kismet2.cpp
@@ -470,8 +470,7 @@ UBlueprint* FKismetEditorUtilities::CreateBlueprint(UClass* ParentClass, UObject
 		{
 			// add update event graph
 			UK2Node_Event* EventNode = NewObject<UK2Node_Event>(Graph);
-			EventNode->EventSignatureClass = UAnimInstance::StaticClass();
-			EventNode->EventSignatureName = FName(TEXT("BlueprintUpdateAnimation"));
+			EventNode->EventReference.SetExternalMember(FName(TEXT("BlueprintUpdateAnimation")), UAnimInstance::StaticClass());
 			EventNode->bOverrideFunction=true;
 			EventNode->CreateNewGuid();
 			EventNode->PostPlacedNewNode();
diff --git a/Engine/Source/Editor/UnrealEd/Private/Tests/EditorAutomationTests.cpp b/Engine/Source/Editor/UnrealEd/Private/Tests/EditorAutomationTests.cpp
--- a/Engine/Source/Editor/UnrealEd/Private/Tests/EditorAutomationTests.cpp
+++ b/Engine/Source/Editor/UnrealEd/Private/Tests/EditorAutomationTests.cpp
@@ -1792,7 +1792,7 @@ bool FTraceAllTimelinesAutomationTest::RunTest(const FString& BlueprintName)
 			for (int32 EventIndex = 0; EventIndex < AllEventNodes.Num(); ++EventIndex)
 			{
 				UK2Node_Event* Event = AllEventNodes[EventIndex];
-				UFunction* Function = FindField<UFunction>(Event->EventSignatureClass, Event->EventSignatureName);
+				UFunction* Function = Event->FindEventSignatureFunction();
 
 				const bool bIsCosmeticEvent = (Function && Function->HasAllFunctionFlags(FUNC_BlueprintCosmetic)) || Event->IsCosmeticTickEvent();
 				const bool bIsCosmeticChain = UncheckedEventNodes.Contains(Event);
diff --git a/Engine/Source/Editor/UnrealEd/Private/Tests/EditorBuildPromotionTests.cpp b/Engine/Source/Editor/UnrealEd/Private/Tests/EditorBuildPromotionTests.cpp
--- a/Engine/Source/Editor/UnrealEd/Private/Tests/EditorBuildPromotionTests.cpp
+++ b/Engine/Source/Editor/UnrealEd/Private/Tests/EditorBuildPromotionTests.cpp
@@ -706,12 +706,11 @@ namespace EditorBuildPromotionTestUtils
 
 		// Make an add component node
 		UK2Node_Event* NewEventNode = NewObject<UK2Node_Event>(TempOuter);
-		NewEventNode->EventSignatureName = "ReceiveBeginPlay";
-		NewEventNode->EventSignatureClass = AActor::StaticClass();
+		NewEventNode->EventReference.SetExternalMember(FName(TEXT("ReceiveBeginPlay")), AActor::StaticClass());
 		NewEventNode->bOverrideFunction = true;
 
 		//Check for existing events
-		UK2Node_Event* ExistingEvent = FBlueprintEditorUtils::FindOverrideForFunction(InBlueprint, NewEventNode->EventSignatureClass, NewEventNode->EventSignatureName);
+		UK2Node_Event* ExistingEvent = FBlueprintEditorUtils::FindOverrideForFunction(InBlueprint, NewEventNode->EventReference.GetMemberParentClass(NewEventNode), NewEventNode->EventReference.GetMemberName());
 
 		if (!ExistingEvent)
 		{
