[INTEGRATE] Change 2427168 by [REDACTED_PERSON][REDACTED_HANDLE] on 2015/01/31 15:37:45

	UE-8780 Add tooltip to K2Node_InputTouch

diff --git a/Engine/Source/Editor/BlueprintGraph/Classes/K2Node_InputTouch.h b/Engine/Source/Editor/BlueprintGraph/Classes/K2Node_InputTouch.h
--- a/Engine/Source/Editor/BlueprintGraph/Classes/K2Node_InputTouch.h
+++ b/Engine/Source/Editor/BlueprintGraph/Classes/K2Node_InputTouch.h
@@ -35,6 +35,7 @@ class UK2Node_InputTouch : public UK2Node
 	virtual void AllocateDefaultPins() override;
 	virtual FLinearColor GetNodeTitleColor() const override;
 	virtual FText GetNodeTitle(ENodeTitleType::Type TitleType) const override;
+	virtual FText GetTooltipText() const override;
 	virtual FName GetPaletteIcon(FLinearColor& OutColor) const override{ return TEXT("GraphEditor.TouchEvent_16x"); }
 	virtual void GetMenuActions(FBlueprintActionDatabaseRegistrar& ActionRegistrar) const override;
 	virtual FText GetMenuCategory() const override;
diff --git a/Engine/Source/Editor/BlueprintGraph/Private/K2Node_InputTouch.cpp b/Engine/Source/Editor/BlueprintGraph/Private/K2Node_InputTouch.cpp
--- a/Engine/Source/Editor/BlueprintGraph/Private/K2Node_InputTouch.cpp
+++ b/Engine/Source/Editor/BlueprintGraph/Private/K2Node_InputTouch.cpp
@@ -70,6 +70,11 @@ FText UK2Node_InputTouch::GetNodeTitle(ENodeTitleType::Type TitleType) const
 	return Title;
 }
 
+FText UK2Node_InputTouch::GetTooltipText() const
+{
+	return NSLOCTEXT("K2Node", "InputTouch_Tooltip", "Event for when a finger presses, releases or is moved on a touch device.");
+}
+
 void UK2Node_InputTouch::GetMenuActions(FBlueprintActionDatabaseRegistrar& ActionRegistrar) const
 {
 	// actions get registered under specific object-keys; the idea is that 
