Fix for some invalid BP nodes after previous delegates change.

diff --git a/Engine/Source/Editor/BlueprintGraph/Private/K2Node_Event.cpp b/Engine/Source/Editor/BlueprintGraph/Private/K2Node_Event.cpp
--- a/Engine/Source/Editor/BlueprintGraph/Private/K2Node_Event.cpp
+++ b/Engine/Source/Editor/BlueprintGraph/Private/K2Node_Event.cpp
@@ -55,6 +55,14 @@ void UK2Node_Event::Serialize(FArchive& Ar)
 		EventSignatureClass = EventReference.GetMemberParentClass(this);
 	}
 
+	FName OldEventSignatureName;
+	TSubclassOf<class UObject> OldEventSignatureClass;
+	if (Ar.IsLoading())
+	{
+		OldEventSignatureName = EventSignatureName;
+		OldEventSignatureClass = EventSignatureClass;
+	}
+
 	Super::Serialize(Ar);
 
 	// Fix up legacy nodes that may not yet have a delegate pin
@@ -66,7 +74,10 @@ void UK2Node_Event::Serialize(FArchive& Ar)
 			CreatePin(EGPD_Output, K2Schema->PC_Delegate, TEXT(""), NULL, false, false, DelegateOutputName);
 		}
 
-		EventReference.SetExternalMember(EventSignatureName, EventSignatureClass);
+		if (EventSignatureName != OldEventSignatureName || EventSignatureClass != OldEventSignatureClass)
+		{
+			EventReference.SetExternalMember(EventSignatureName, EventSignatureClass);
+		}
 	}
 
 	PRAGMA_ENABLE_DEPRECATION_WARNINGS
