Another fix for some invalid component bound BP event nodes after previous delegates change.

diff --git a/Engine/Source/Editor/BlueprintGraph/Private/K2Node_ComponentBoundEvent.cpp b/Engine/Source/Editor/BlueprintGraph/Private/K2Node_ComponentBoundEvent.cpp
--- a/Engine/Source/Editor/BlueprintGraph/Private/K2Node_ComponentBoundEvent.cpp
+++ b/Engine/Source/Editor/BlueprintGraph/Private/K2Node_ComponentBoundEvent.cpp
@@ -130,7 +130,7 @@ void UK2Node_ComponentBoundEvent::Serialize(FArchive& Ar)
 		DelegateOwnerClass = EventSignatureClass;
 		UMulticastDelegateProperty* TargetDelegateProp = GetTargetDelegateProperty();
 		check(TargetDelegateProp);
-		EventReference.SetFromField<UFunction>(TargetDelegateProp->SignatureFunction, false);
+		EventReference.SetExternalDelegateMember(TargetDelegateProp->SignatureFunction->GetFName());
 	}
 	PRAGMA_ENABLE_DEPRECATION_WARNINGS
 }
