[INTEGRATE] Change 2430574 by [REDACTED_PERSON][REDACTED_HANDLE] on 2015/02/03 16:57:55

	AddComponent component templates are also created as RF_Public

diff --git a/Engine/Source/Editor/BlueprintGraph/Private/K2Node_AddComponent.cpp b/Engine/Source/Editor/BlueprintGraph/Private/K2Node_AddComponent.cpp
--- a/Engine/Source/Editor/BlueprintGraph/Private/K2Node_AddComponent.cpp
+++ b/Engine/Source/Editor/BlueprintGraph/Private/K2Node_AddComponent.cpp
@@ -239,8 +239,7 @@ void UK2Node_AddComponent::PostPasteNode()
 		{
 			ensure(NULL != Cast<UBlueprintGeneratedClass>(Blueprint->GeneratedClass));
 			// Create a new template object and update the template pin to point to it
-			UActorComponent* NewTemplate = NewObject<UActorComponent>(Blueprint->GeneratedClass, ComponentClass);
-			NewTemplate->SetFlags(RF_ArchetypeObject);
+			UActorComponent* NewTemplate = NewObject<UActorComponent>(Blueprint->GeneratedClass, ComponentClass, NAME_None, RF_ArchetypeObject|RF_Public);
 			Blueprint->ComponentTemplates.Add(NewTemplate);
 
 			TemplateNamePin->DefaultValue = NewTemplate->GetName();
