[INTEGRATE] Change 2431124 by [REDACTED_PERSON][REDACTED_HANDLE] on 2015/02/03 22:54:42

	Blueprints: Hide RelativeTransform and bManualAttachment pins on non-USceneComponent AddComponent nodes [UE-8924]

diff --git a/Engine/Source/Editor/BlueprintGraph/Private/K2Node_AddComponent.cpp b/Engine/Source/Editor/BlueprintGraph/Private/K2Node_AddComponent.cpp
--- a/Engine/Source/Editor/BlueprintGraph/Private/K2Node_AddComponent.cpp
+++ b/Engine/Source/Editor/BlueprintGraph/Private/K2Node_AddComponent.cpp
@@ -40,11 +40,11 @@ void UK2Node_AddComponent::AllocatePinsForExposedVariables()
 	const UEdGraphSchema_K2* K2Schema = GetDefault<UEdGraphSchema_K2>();
 
 	const UActorComponent* TemplateComponent = GetTemplateFromNode();
-	const UClass* ComponentClass = TemplateComponent ? TemplateComponent->GetClass() : NULL;
+	const UClass* ComponentClass = TemplateComponent ? TemplateComponent->GetClass() : nullptr;
 
-	if(ComponentClass)
+	if (ComponentClass != nullptr)
 	{
-		const UObject* ClassDefaultObject = ComponentClass ? ComponentClass->ClassDefaultObject : NULL;
+		const UObject* ClassDefaultObject = ComponentClass ? ComponentClass->ClassDefaultObject : nullptr;
 
 		for (TFieldIterator<UProperty> PropertyIt(ComponentClass, EFieldIteratorFlags::IncludeSuper); PropertyIt; ++PropertyIt)
 		{
@@ -58,13 +58,13 @@ void UK2Node_AddComponent::AllocatePinsForExposedVariables()
 				FEdGraphPinType PinType;
 				K2Schema->ConvertPropertyToPinType(Property, /*out*/ PinType);	
 				const bool bIsUnique = (NULL == FindPin(Property->GetName()));
-				if(K2Schema->FindSetVariableByNameFunction(PinType) && bIsUnique)
+				if (K2Schema->FindSetVariableByNameFunction(PinType) && bIsUnique)
 				{
 					UEdGraphPin* Pin = CreatePin(EGPD_Input, TEXT(""), TEXT(""), NULL, false, false, Property->GetName());
 					Pin->PinType = PinType;
 					bHasExposedVariable = true;
 
-					if (ClassDefaultObject && K2Schema->PinDefaultValueIsEditable(*Pin))
+					if ((ClassDefaultObject != nullptr) && K2Schema->PinDefaultValueIsEditable(*Pin))
 					{
 						FString DefaultValueAsString;
 						const bool bDefaultValueSet = FBlueprintEditorUtils::PropertyValueToString(Property, reinterpret_cast<const uint8*>(ClassDefaultObject), DefaultValueAsString);
@@ -78,6 +78,15 @@ void UK2Node_AddComponent::AllocatePinsForExposedVariables()
 			}
 		}
 	}
+
+	// Hide transform and attachment pins if it is not a scene component
+	const bool bHideTransformPins = (ComponentClass != nullptr) ? !ComponentClass->IsChildOf(USceneComponent::StaticClass()) : false;
+
+	UEdGraphPin* ManualAttachmentPin = GetManualAttachmentPin();
+	ManualAttachmentPin->bHidden = bHideTransformPins;
+
+	UEdGraphPin* TransformPin = GetRelativeTransformPin();
+	TransformPin->bHidden = bHideTransformPins;
 }
 
 const UClass* UK2Node_AddComponent::GetSpawnedType() const
