Removing a check when serializing ActorBoundEvent and ComponentBoundEvent nodes that cannot be guranteed during node duplication and replacing it with an if-statement.

diff --git a/Engine/Source/Editor/BlueprintGraph/Private/K2Node_ActorBoundEvent.cpp b/Engine/Source/Editor/BlueprintGraph/Private/K2Node_ActorBoundEvent.cpp
--- a/Engine/Source/Editor/BlueprintGraph/Private/K2Node_ActorBoundEvent.cpp
+++ b/Engine/Source/Editor/BlueprintGraph/Private/K2Node_ActorBoundEvent.cpp
@@ -220,8 +220,10 @@ void UK2Node_ActorBoundEvent::Serialize(FArchive& Ar)
 	{
 		DelegateOwnerClass = EventSignatureClass;
 		UMulticastDelegateProperty* TargetDelegateProp = GetTargetDelegateProperty();
-		check(TargetDelegateProp);
-		EventReference.SetFromField<UFunction>(TargetDelegateProp->SignatureFunction, false);
+		if(TargetDelegateProp)
+		{
+			EventReference.SetFromField<UFunction>(TargetDelegateProp->SignatureFunction, false);
+		}
 	}
 
 	PRAGMA_ENABLE_DEPRECATION_WARNINGS
diff --git a/Engine/Source/Editor/BlueprintGraph/Private/K2Node_ComponentBoundEvent.cpp b/Engine/Source/Editor/BlueprintGraph/Private/K2Node_ComponentBoundEvent.cpp
--- a/Engine/Source/Editor/BlueprintGraph/Private/K2Node_ComponentBoundEvent.cpp
+++ b/Engine/Source/Editor/BlueprintGraph/Private/K2Node_ComponentBoundEvent.cpp
@@ -129,8 +129,10 @@ void UK2Node_ComponentBoundEvent::Serialize(FArchive& Ar)
 	{
 		DelegateOwnerClass = EventSignatureClass;
 		UMulticastDelegateProperty* TargetDelegateProp = GetTargetDelegateProperty();
-		check(TargetDelegateProp);
-		EventReference.SetExternalDelegateMember(TargetDelegateProp->SignatureFunction->GetFName());
+		if (TargetDelegateProp)
+		{
+			EventReference.SetExternalDelegateMember(TargetDelegateProp->SignatureFunction->GetFName());
+		}
 	}
 	PRAGMA_ENABLE_DEPRECATION_WARNINGS
 }
