[INTEGRATE] Change 2440029 by [REDACTED_PERSON][REDACTED_HANDLE]

[UE-9490] Fixing it so impure K2Nodes that derive from CallFunction nodes, don't mistakenly have their exec pins disconnected.

diff --git a/Engine/Source/Editor/BlueprintGraph/Private/K2Node_CallFunction.cpp b/Engine/Source/Editor/BlueprintGraph/Private/K2Node_CallFunction.cpp
--- a/Engine/Source/Editor/BlueprintGraph/Private/K2Node_CallFunction.cpp
+++ b/Engine/Source/Editor/BlueprintGraph/Private/K2Node_CallFunction.cpp
@@ -2034,7 +2034,7 @@ FName UK2Node_CallFunction::GetPaletteIcon(FLinearColor& OutColor) const
 
 bool UK2Node_CallFunction::ReconnectPureExecPins(TArray<UEdGraphPin*>& OldPins)
 {
-	if (bIsPureFunc)
+	if (IsNodePure())
 	{
 		// look for an old exec pin
 		const UEdGraphSchema_K2* K2Schema = GetDefault<UEdGraphSchema_K2>();
