[INTEGRATE] Change 2443539 by [REDACTED_PERSON][REDACTED_HANDLE] on 2015/02/12 11:22

	Add Component nodes created by dragging in to the graph now properly mark their templates as RF_Public

diff --git a/Engine/Source/Editor/BlueprintGraph/Private/EdGraphSchema_K2_Actions.cpp b/Engine/Source/Editor/BlueprintGraph/Private/EdGraphSchema_K2_Actions.cpp
--- a/Engine/Source/Editor/BlueprintGraph/Private/EdGraphSchema_K2_Actions.cpp
+++ b/Engine/Source/Editor/BlueprintGraph/Private/EdGraphSchema_K2_Actions.cpp
@@ -303,8 +303,7 @@ UEdGraphNode* FEdGraphSchemaAction_K2AddComponent::PerformAction(class UEdGraph*
 
 		ensure(NULL != Cast<UBlueprintGeneratedClass>(Blueprint->GeneratedClass));
 		// Then create a new template object, and add to array in
-		UActorComponent* NewTemplate = NewObject<UActorComponent>(Blueprint->GeneratedClass, ComponentClass);
-		NewTemplate->SetFlags(RF_ArchetypeObject);
+		UActorComponent* NewTemplate = NewObject<UActorComponent>(Blueprint->GeneratedClass, ComponentClass, NAME_None, RF_ArchetypeObject | RF_Public);
 		Blueprint->ComponentTemplates.Add(NewTemplate);
 
 		// Set the name of the template as the default for the TemplateName param
