Format Text nodes are no longer considered (internally) as in-pure, and cannot have breakpoints set.

diff --git a/Engine/Source/Editor/BlueprintGraph/Classes/K2Node_FormatText.h b/Engine/Source/Editor/BlueprintGraph/Classes/K2Node_FormatText.h
--- a/Engine/Source/Editor/BlueprintGraph/Classes/K2Node_FormatText.h
+++ b/Engine/Source/Editor/BlueprintGraph/Classes/K2Node_FormatText.h
@@ -24,7 +24,7 @@ class UK2Node_FormatText : public UK2Node
 	// End UEdGraphNode interface.
 
 	// Begin UK2Node interface.
-	//virtual bool IsNodePure() const override { return true; }
+	virtual bool IsNodePure() const override { return true; }
 	virtual bool NodeCausesStructuralBlueprintChange() const override { return true; }
 	virtual void ExpandNode(class FKismetCompilerContext& CompilerContext, UEdGraph* SourceGraph) override;
 	virtual ERedirectType DoPinsMatchForReconstruction(const UEdGraphPin* NewPin, int32 NewPinIndex, const UEdGraphPin* OldPin, int32 OldPinIndex) const override;
