[INTEGRATE] Change 2434174 by [REDACTED_PERSON][REDACTED_HANDLE] on 2015/02/05 16:17:50

	[UE-8766] Add component nodes show up for blueprint defined components that have not been loaded

diff --git a/Engine/Source/Editor/BlueprintGraph/Classes/BlueprintComponentNodeSpawner.h b/Engine/Source/Editor/BlueprintGraph/Classes/BlueprintComponentNodeSpawner.h
--- a/Engine/Source/Editor/BlueprintGraph/Classes/BlueprintComponentNodeSpawner.h
+++ b/Engine/Source/Editor/BlueprintGraph/Classes/BlueprintComponentNodeSpawner.h
@@ -26,10 +26,9 @@ public:
 	 * viable as a spawnable component (do that before calling this).
 	 *
 	 * @param  ComponentClass	The component type you want spawned nodes to spawn.
-	 * @param  Outer			Optional outer for the new spawner (if left null, the transient package will be used).
 	 * @return A newly allocated instance of this class.
 	 */
-	static UBlueprintComponentNodeSpawner* Create(TSubclassOf<UActorComponent> const ComponentClass, UObject* Outer = nullptr);
+	static UBlueprintComponentNodeSpawner* Create(const struct FComponentTypeEntry& Entry);
 
 	// UBlueprintNodeSpawner interface
 	virtual FBlueprintNodeSignature GetSpawnerSignature() const override;
@@ -53,4 +52,12 @@ private:
 	/** The component class to configure new nodes with. */
 	UPROPERTY()
 	TSubclassOf<UActorComponent> ComponentClass;
+
+	/** The name of the component class to configure new nodes with. */
+	UPROPERTY()
+	FString ComponentName;
+	
+	/** The name of the asset name that needs to be loaded */
+	UPROPERTY()
+	FString ComponentAssetName;
 };
diff --git a/Engine/Source/Editor/BlueprintGraph/Private/BlueprintActionDatabase.cpp b/Engine/Source/Editor/BlueprintGraph/Private/BlueprintActionDatabase.cpp
--- a/Engine/Source/Editor/BlueprintGraph/Private/BlueprintActionDatabase.cpp
+++ b/Engine/Source/Editor/BlueprintGraph/Private/BlueprintActionDatabase.cpp
@@ -2,6 +2,7 @@
 
 #include "BlueprintGraphPrivatePCH.h"
 #include "BlueprintActionDatabase.h"
+#include "ComponentTypeRegistry.h"
 #include "EdGraphSchema_K2.h"       // for CanUserKismetCallFunction()
 #include "EditorCategoryUtils.h"
 #include "KismetEditorUtilities.h"	// for IsClassABlueprintSkeleton(), IsClassABlueprintInterface(), GetBoundsForSelectedNodes(), etc.
@@ -340,16 +341,6 @@ namespace BlueprintActionDatabaseImpl
 	 */
 	static void AddClassCastActions(UClass* const Class, FActionList& ActionListOut);
 
-	/**
-	 * Evolved from K2ActionMenuBuilder's GetAddComponentClasses(). If the
-	 * specified class is a component type (and can be spawned), then a 
-	 * UK2Node_AddComponent spawner is created and added to ActionListOut.
-	 *
-	 * @param  ComponentClass	The class who you want a spawner for (the component class).
-	 * @param  ActionListOut	The list you want populated with the new spawner.
-	 */
-	static void AddComponentClassActions(TSubclassOf<UActorComponent> const ComponentClass, FActionList& ActionListOut);
-
 	/**
 	 * Adds custom actions to operate on the provided skeleton. Used primarily
 	 * to find AnimNotify event vocabulary
@@ -543,11 +534,6 @@ static void BlueprintActionDatabaseImpl::GetClassMemberActions(UClass* const Cla
 	}
 
 	AddClassCastActions(Class, ActionListOut);
-
-	if (Class->IsChildOf<UActorComponent>())
-	{
-		AddComponentClassActions(Class, ActionListOut);
-	}
 }
 
 //------------------------------------------------------------------------------
@@ -687,18 +673,6 @@ static void BlueprintActionDatabaseImpl::AddClassCastActions(UClass* Class, FAct
 	}
 }
 
-//------------------------------------------------------------------------------
-static void BlueprintActionDatabaseImpl::AddComponentClassActions(TSubclassOf<UActorComponent> const ComponentClass, FActionList& ActionListOut)
-{
-	if ((ComponentClass != nullptr) && !ComponentClass->HasAnyClassFlags(CLASS_Abstract) && ComponentClass->HasMetaData(FBlueprintMetadata::MD_BlueprintSpawnableComponent))
-	{
-		if (UBlueprintComponentNodeSpawner* NodeSpawner = UBlueprintComponentNodeSpawner::Create(ComponentClass))
-		{
-			ActionListOut.Add(NodeSpawner);
-		}
-	}
-}
-
 //------------------------------------------------------------------------------
 static void BlueprintActionDatabaseImpl::AddSkeletonActions(const USkeleton& Skeleton, FActionList& ActionListOut)
 {
@@ -1049,6 +1023,10 @@ void FBlueprintActionDatabase::RefreshAll()
 		FActionList& ClassActionList = ActionRegistry.FindOrAdd(*SkeletonIt);
 		BlueprintActionDatabaseImpl::AddSkeletonActions(**SkeletonIt, ClassActionList);
 	}
+
+	// this handles creating entries for components that were loaded before the database was alive:
+	FComponentTypeRegistry::Get().SubscribeToComponentList(ComponentTypes).AddRaw(this, &FBlueprintActionDatabase::RefreshComponentActions);
+	RefreshComponentActions();
 }
 
 //------------------------------------------------------------------------------
@@ -1218,6 +1196,21 @@ void FBlueprintActionDatabase::RefreshAssetActions(UObject* const AssetObject)
 	}
 }
 
+//------------------------------------------------------------------------------
+void FBlueprintActionDatabase::RefreshComponentActions()
+{
+	check(ComponentTypes);
+	FActionList& ClassActionList = ActionRegistry.FindOrAdd(UBlueprintComponentNodeSpawner::StaticClass());
+	ClassActionList.Empty(ComponentTypes->Num());
+	for (const auto& ComponentType : *ComponentTypes)
+	{
+		if (UBlueprintComponentNodeSpawner* NodeSpawner = UBlueprintComponentNodeSpawner::Create(ComponentType))
+		{
+			ClassActionList.Add(NodeSpawner);
+		}
+	}
+}
+
 //------------------------------------------------------------------------------
 void FBlueprintActionDatabase::ClearAssetActions(UObject* const AssetObject)
 {
diff --git a/Engine/Source/Editor/BlueprintGraph/Private/BlueprintComponentNodeSpawner.cpp b/Engine/Source/Editor/BlueprintGraph/Private/BlueprintComponentNodeSpawner.cpp
--- a/Engine/Source/Editor/BlueprintGraph/Private/BlueprintComponentNodeSpawner.cpp
+++ b/Engine/Source/Editor/BlueprintGraph/Private/BlueprintComponentNodeSpawner.cpp
@@ -44,18 +44,41 @@ static FText BlueprintComponentNodeSpawnerImpl::GetDefaultMenuCategory(TSubclass
  ******************************************************************************/
 
 //------------------------------------------------------------------------------
-UBlueprintComponentNodeSpawner* UBlueprintComponentNodeSpawner::Create(TSubclassOf<UActorComponent> const ComponentClass, UObject* Outer/* = nullptr*/)
+UBlueprintComponentNodeSpawner* UBlueprintComponentNodeSpawner::Create(const FComponentTypeEntry& Entry)
 {
-	check(ComponentClass != nullptr);
+	UClass* ComponentClass = Entry.ComponentClass;
+	if (ComponentClass == nullptr )
+	{
+		// unloaded class, must be blueprint created. Create an entry. We'll load the class when we spawn the node:
+
+		UBlueprintComponentNodeSpawner* NodeSpawner = NewObject<UBlueprintComponentNodeSpawner>(GetTransientPackage());
+		NodeSpawner->ComponentClass = nullptr;
+		NodeSpawner->NodeClass = UK2Node_AddComponent::StaticClass();
+		NodeSpawner->ComponentName = Entry.ComponentName;
+		NodeSpawner->ComponentAssetName = Entry.ComponentAssetName;
+
+		FBlueprintActionUiSpec& MenuSignature = NodeSpawner->DefaultMenuSignature;
+		FText const ComponentTypeName = FText::FromString(Entry.ComponentName);
+		MenuSignature.MenuName = FText::Format(LOCTEXT("AddComponentMenuName", "Add {0}"), ComponentTypeName);
+		MenuSignature.Category = LOCTEXT("BlueprintComponentCategory", "Custom");
+		MenuSignature.Tooltip = FText::Format(LOCTEXT("AddComponentTooltip", "Spawn a {0}"), ComponentTypeName);
+		// add at least one character, so that PrimeDefaultMenuSignature() doesn't 
+		// attempt to query the template node
+		MenuSignature.Keywords.AppendChar(TEXT(' '));
+		MenuSignature.IconName = FClassIconFinder::FindIconNameForClass(nullptr);
+
+		return NodeSpawner;
+	}
 
-	if (Outer == nullptr)
+	if (ComponentClass->HasAnyClassFlags(CLASS_Abstract) || !ComponentClass->HasMetaData(FBlueprintMetadata::MD_BlueprintSpawnableComponent))
 	{
-		Outer = GetTransientPackage();
+		// loaded class that is marked as abstract or not spawnable, don't create an entry:
+		return nullptr;
 	}
 
 	UClass* const AuthoritativeClass = ComponentClass->GetAuthoritativeClass();
 
-	UBlueprintComponentNodeSpawner* NodeSpawner = NewObject<UBlueprintComponentNodeSpawner>(Outer);
+	UBlueprintComponentNodeSpawner* NodeSpawner = NewObject<UBlueprintComponentNodeSpawner>(GetTransientPackage());
 	NodeSpawner->ComponentClass = AuthoritativeClass;
 	NodeSpawner->NodeClass      = UK2Node_AddComponent::StaticClass();
 
@@ -92,8 +115,6 @@ FBlueprintNodeSignature UBlueprintComponentNodeSpawner::GetSpawnerSignature() co
 // and FEdGraphSchemaAction_K2AddComponent::PerformAction().
 UEdGraphNode* UBlueprintComponentNodeSpawner::Invoke(UEdGraph* ParentGraph, FBindingSet const& Bindings, FVector2D const Location) const
 {
-	check(ComponentClass != nullptr);
-	
 	auto PostSpawnLambda = [](UEdGraphNode* NewNode, bool bIsTemplateNode, FCustomizeNodeDelegate UserDelegate)
 	{		
 		UK2Node_AddComponent* AddCompNode = CastChecked<UK2Node_AddComponent>(NewNode);
@@ -113,16 +134,33 @@ UEdGraphNode* UBlueprintComponentNodeSpawner::Invoke(UEdGraph* ParentGraph, FBin
 	UEdGraphPin* ReturnPin = NewNode->GetReturnValuePin();
 	if (ReturnPin != nullptr)
 	{
-		ReturnPin->PinType.PinSubCategoryObject = *ComponentClass;
+		ReturnPin->PinType.PinSubCategoryObject = *(ComponentClass ? ComponentClass : UActorComponent::StaticClass());
 	}
 
 	bool const bIsTemplateNode = FBlueprintNodeTemplateCache::IsTemplateOuter(ParentGraph);
 	if (!bIsTemplateNode)
 	{
+		TSubclassOf<UActorComponent> Class = ComponentClass;
+		if (Class == nullptr)
+		{
+			const ELoadFlags LoadFlags = LOAD_None;
+			UBlueprint* LoadedObject = LoadObject<UBlueprint>(NULL, *ComponentAssetName, NULL, LoadFlags, NULL);
+			if (LoadedObject == nullptr)
+			{
+				return nullptr;
+			}
+
+			Class = TSubclassOf<UActorComponent>(Cast<UBlueprintGeneratedClass>(LoadedObject->GeneratedClass));
+			if (Class == nullptr)
+			{
+				return nullptr;
+			}
+		}
+
 		UBlueprint* Blueprint = NewNode->GetBlueprint();
 
-		FName DesiredComponentName = MakeUniqueObjectName(Blueprint->GeneratedClass, ComponentClass, *ComponentClass->GetDisplayNameText().ToString());
-		UActorComponent* ComponentTemplate = NewObject<UActorComponent>(Blueprint->GeneratedClass, ComponentClass, DesiredComponentName);
+		FName DesiredComponentName = MakeUniqueObjectName(Blueprint->GeneratedClass, Class, *Class->GetDisplayNameText().ToString());
+		UActorComponent* ComponentTemplate = NewObject<UActorComponent>(Blueprint->GeneratedClass, Class, DesiredComponentName);
 		ComponentTemplate->SetFlags(RF_ArchetypeObject);
 
 		Blueprint->ComponentTemplates.Add(ComponentTemplate);
@@ -133,6 +171,7 @@ UEdGraphNode* UBlueprintComponentNodeSpawner::Invoke(UEdGraph* ParentGraph, FBin
 		{
 			TemplateNamePin->DefaultValue = ComponentTemplate->GetName();
 		}
+		NewNode->ReconstructNode();
 	}
 
 	// apply bindings, after we've setup the template pin
diff --git a/Engine/Source/Editor/BlueprintGraph/Public/BlueprintActionDatabase.h b/Engine/Source/Editor/BlueprintGraph/Public/BlueprintActionDatabase.h
--- a/Engine/Source/Editor/BlueprintGraph/Public/BlueprintActionDatabase.h
+++ b/Engine/Source/Editor/BlueprintGraph/Public/BlueprintActionDatabase.h
@@ -86,6 +86,11 @@ public:
 	 */
 	void RefreshAssetActions(UObject* const AssetObject);
 
+	/**
+	 * Updates all component related actions
+	 */
+	void RefreshComponentActions();
+
 	/**
 	 * Finds the database entry for the specified class and wipes it. The entry 
 	 * won't be rebuilt, unless RefreshAssetActions() is explicitly called after.
@@ -160,4 +165,7 @@ private:
 
 	/** Handle to the registered OnBlueprintChanged delegate. */
 	FDelegateHandle OnBlueprintChangedDelegateHandle;
+
+	/** Pointer to the shared list of currently existing component types */
+	const TArray<struct FComponentTypeEntry>* ComponentTypes;
 };
diff --git a/Engine/Source/Editor/UnrealEd/Private/ComponentTypeRegistry.cpp b/Engine/Source/Editor/UnrealEd/Private/ComponentTypeRegistry.cpp
--- a/Engine/Source/Editor/UnrealEd/Private/ComponentTypeRegistry.cpp
+++ b/Engine/Source/Editor/UnrealEd/Private/ComponentTypeRegistry.cpp
@@ -22,6 +22,7 @@ struct FComponentTypeRegistryData
 	/** End implementation of FTickableEditorObject */
 
 	TArray<FComponentClassComboEntryPtr> ComponentClassList;
+	TArray<FComponentTypeEntry> ComponentTypeList;
 	TArray<FAssetData> PendingAssetData;
 	FOnComponentTypeListChanged ComponentListChanged;
 };
@@ -109,6 +110,7 @@ FComponentTypeRegistryData::FComponentTypeRegistryData()
 void FComponentTypeRegistryData::RefreshComponentList()
 {
 	ComponentClassList.Empty();
+	ComponentTypeList.Empty();
 
 	struct SortComboEntry
 	{
@@ -187,6 +189,9 @@ void FComponentTypeRegistryData::RefreshComponentList()
 					SortedClassList.Add(NewEntry);
 				}
 			}
+
+			FComponentTypeEntry Entry = { Class->GetName(), FString(), Class };
+			ComponentTypeList.Add(Entry);
 		}
 	}
 
@@ -225,6 +230,9 @@ void FComponentTypeRegistryData::RefreshComponentList()
 				}
 			}
 
+			FComponentTypeEntry Entry = { FixedString, AssetPath.ToString(), nullptr };
+			ComponentTypeList.Add(Entry);
+
 			FComponentClassComboEntryPtr NewEntry(new FComponentClassComboEntry(BlueprintComponents, FixedString, AssetPath, bIncludeInFilter));
 			SortedClassList.Add(NewEntry);
 		}
@@ -312,6 +320,13 @@ FOnComponentTypeListChanged& FComponentTypeRegistry::SubscribeToComponentList(TA
 	return Data->ComponentListChanged;
 }
 
+FOnComponentTypeListChanged& FComponentTypeRegistry::SubscribeToComponentList(const TArray<FComponentTypeEntry>*& OutComponentList)
+{
+	check(Data);
+	OutComponentList = &Data->ComponentTypeList;
+	return Data->ComponentListChanged;
+}
+
 FOnComponentTypeListChanged& FComponentTypeRegistry::GetOnComponentTypeListChanged()
 {
 	check(Data);
diff --git a/Engine/Source/Editor/UnrealEd/Public/ComponentTypeRegistry.h b/Engine/Source/Editor/UnrealEd/Public/ComponentTypeRegistry.h
--- a/Engine/Source/Editor/UnrealEd/Public/ComponentTypeRegistry.h
+++ b/Engine/Source/Editor/UnrealEd/Public/ComponentTypeRegistry.h
@@ -8,7 +8,19 @@ DECLARE_MULTICAST_DELEGATE(FOnComponentTypeListChanged);
 
 typedef TSharedPtr<class FComponentClassComboEntry> FComponentClassComboEntryPtr;
 
-struct FComponentTypeRegistry
+struct FComponentTypeEntry
+{
+	/** Name of the component, as typed by the user */
+	FString ComponentName;
+
+	/** Name of the component, corresponds to asset name for blueprint components */
+	FString ComponentAssetName;
+
+	/** Optional pointer to the UClass, will be nullptr for blueprint components that aren't loaded */
+	class UClass* ComponentClass;
+};
+
+struct UNREALED_API FComponentTypeRegistry
 {
 	static FComponentTypeRegistry& Get();
 
@@ -19,6 +31,7 @@ struct FComponentTypeRegistry
 	*	added or removed from the component type list
 	*/
 	FOnComponentTypeListChanged& SubscribeToComponentList(TArray<FComponentClassComboEntryPtr>*& OutComponentList);
+	FOnComponentTypeListChanged& SubscribeToComponentList(const TArray<FComponentTypeEntry>*& OutComponentList);
 	FOnComponentTypeListChanged& GetOnComponentTypeListChanged();
 
 private:
