[INTEGRATE] Change 2434611 by [REDACTED_PERSON][REDACTED_HANDLE] on 2015/02/05 20:25:56

	Fix for clang build, common issue with ternary operator + other operators creating ambiguity

diff --git a/Engine/Source/Editor/BlueprintGraph/Private/BlueprintComponentNodeSpawner.cpp b/Engine/Source/Editor/BlueprintGraph/Private/BlueprintComponentNodeSpawner.cpp
--- a/Engine/Source/Editor/BlueprintGraph/Private/BlueprintComponentNodeSpawner.cpp
+++ b/Engine/Source/Editor/BlueprintGraph/Private/BlueprintComponentNodeSpawner.cpp
@@ -134,7 +134,14 @@ UEdGraphNode* UBlueprintComponentNodeSpawner::Invoke(UEdGraph* ParentGraph, FBin
 	UEdGraphPin* ReturnPin = NewNode->GetReturnValuePin();
 	if (ReturnPin != nullptr)
 	{
-		ReturnPin->PinType.PinSubCategoryObject = *(ComponentClass ? ComponentClass : UActorComponent::StaticClass());
+		if (ComponentClass != nullptr)
+		{
+			ReturnPin->PinType.PinSubCategoryObject = *ComponentClass;
+		}
+		else
+		{
+			ReturnPin->PinType.PinSubCategoryObject = UActorComponent::StaticClass();
+		}
 	}
 
 	bool const bIsTemplateNode = FBlueprintNodeTemplateCache::IsTemplateOuter(ParentGraph);
