[INTEGRATE] Change 2434845 by [REDACTED_PERSON][REDACTED_HANDLE] on 2015/02/06 00:34:04

	Editor: Fixing another non-unity build issue

diff --git a/Engine/Source/Editor/BlueprintGraph/Private/BlueprintComponentNodeSpawner.cpp b/Engine/Source/Editor/BlueprintGraph/Private/BlueprintComponentNodeSpawner.cpp
--- a/Engine/Source/Editor/BlueprintGraph/Private/BlueprintComponentNodeSpawner.cpp
+++ b/Engine/Source/Editor/BlueprintGraph/Private/BlueprintComponentNodeSpawner.cpp
@@ -7,6 +7,7 @@
 #include "BlueprintNodeTemplateCache.h" // for IsTemplateOuter()
 #include "ComponentAssetBroker.h" // for GetComponentsForAsset()/AssignAssetToComponent()
 #include "BlueprintActionFilter.h"	// for FBlueprintActionContext
+#include "ComponentTypeRegistry.h"
 
 #define LOCTEXT_NAMESPACE "BlueprintComponenetNodeSpawner"
 
