[INTEGRATE] Change 2439375 by [REDACTED_PERSON][REDACTED_HANDLE] on 2015/02/10 00:25:00

	Blueprints: Various improvements to the node title of Add Component nodes
	- Generalize the 'Add ComponentType AssetName' display code for Add Component to use AdditionalStatObject instead of an if-else tree, allowing it to work with any component class
	- Now displays asset information on a second line when asked for the full title
	- Fixed an issue where the component title was latched and didn't update when the template was modified (minor pessimization)

diff --git a/Engine/Source/Editor/BlueprintGraph/Classes/K2Node_AddComponent.h b/Engine/Source/Editor/BlueprintGraph/Classes/K2Node_AddComponent.h
--- a/Engine/Source/Editor/BlueprintGraph/Classes/K2Node_AddComponent.h
+++ b/Engine/Source/Editor/BlueprintGraph/Classes/K2Node_AddComponent.h
@@ -70,9 +70,6 @@ private:
 	}
 
 	const UClass* GetSpawnedType() const;
-
-	/** Constructing FText strings can be costly, so we cache the node's title */
-	FNodeTextCache CachedNodeTitle;
 };
 
 
diff --git a/Engine/Source/Editor/BlueprintGraph/Private/K2Node_AddComponent.cpp b/Engine/Source/Editor/BlueprintGraph/Private/K2Node_AddComponent.cpp
--- a/Engine/Source/Editor/BlueprintGraph/Private/K2Node_AddComponent.cpp
+++ b/Engine/Source/Editor/BlueprintGraph/Private/K2Node_AddComponent.cpp
@@ -289,56 +289,58 @@ void UK2Node_AddComponent::PostPasteNode()
 
 FText UK2Node_AddComponent::GetNodeTitle(ENodeTitleType::Type TitleType) const
 {
+	FText CachedAssetTitle;
+	FText CachedNodeTitle;
 	UEdGraphPin* TemplateNamePin = GetTemplateNamePin();
-	if (CachedNodeTitle.IsOutOfDate() && (TemplateNamePin != NULL))
+	if (TemplateNamePin != nullptr)
 	{
 		FString TemplateName = TemplateNamePin->DefaultValue;
 		UBlueprint* Blueprint = GetBlueprint();
-		UActorComponent* SourceTemplate = Blueprint->FindTemplateByName(FName(*TemplateName));
-		if(SourceTemplate)
+
+		if (UActorComponent* SourceTemplate = Blueprint->FindTemplateByName(FName(*TemplateName)))
 		{
-			UStaticMeshComponent* StaticMeshComp = Cast<UStaticMeshComponent>(SourceTemplate);
-			USkeletalMeshComponent* SkelMeshComp = Cast<USkeletalMeshComponent>(SourceTemplate);
-			UParticleSystemComponent* PSysComp = Cast<UParticleSystemComponent>(SourceTemplate);
+			CachedNodeTitle = SourceTemplate->GetClass()->GetDisplayNameText();
+
+			FFormatNamedArguments Args;
+			Args.Add(TEXT("ComponentType"), SourceTemplate->GetClass()->GetDisplayNameText());
+			CachedNodeTitle = FText::Format(LOCTEXT("AddClass", "Add {ComponentType}"), Args);
+
 			UChildActorComponent* SubActorComp = Cast<UChildActorComponent>(SourceTemplate);
 
-			if(StaticMeshComp != NULL && StaticMeshComp->StaticMesh != NULL)
+			if (const UObject* AssociatedAsset = SourceTemplate->AdditionalStatObject())
 			{
 				FFormatNamedArguments Args;
-				Args.Add(TEXT("StaticMeshName"), FText::FromString(StaticMeshComp->StaticMesh->GetName()));
-				CachedNodeTitle = FText::Format(LOCTEXT("AddStaticMesh", "Add StaticMesh {StaticMeshName}"), Args);
+				Args.Add(TEXT("AssetType"), AssociatedAsset->GetClass()->GetDisplayNameText());
+				Args.Add(TEXT("AssetName"), FText::FromString(AssociatedAsset->GetName()));
+				CachedAssetTitle = FText::Format(LOCTEXT("AddComponentAssetDescription", "{AssetType} {AssetName}"), Args);
 			}
-			else if(SkelMeshComp != NULL && SkelMeshComp->SkeletalMesh != NULL)
+			else if ((SubActorComp != nullptr) && (SubActorComp->ChildActorClass != nullptr))
 			{
 				FFormatNamedArguments Args;
-				Args.Add(TEXT("SkeletalMeshName"), FText::FromString(SkelMeshComp->SkeletalMesh->GetName()));
-				CachedNodeTitle = FText::Format(LOCTEXT("AddSkeletalMesh", "Add SkeletalMesh {SkeletalMeshName}"), Args);
-			}
-			else if(PSysComp != NULL && PSysComp->Template != NULL)
-			{
-				FFormatNamedArguments Args;
-				Args.Add(TEXT("ParticleSystemName"), FText::FromString(PSysComp->Template->GetName()));
-				CachedNodeTitle = FText::Format(LOCTEXT("AddParticleSystem", "Add ParticleSystem {ParticleSystemName}"), Args);
-			}
-			else if (SubActorComp && SubActorComp->ChildActorClass)
-			{
-				FFormatNamedArguments Args;
-				Args.Add(TEXT("ComponentClassName"), FText::FromString(SubActorComp->ChildActorClass->GetName()));
-				CachedNodeTitle = FText::Format(LOCTEXT("AddChildActorComponent", "Add ChildActorComponent {ComponentClassName}"), Args);
+				Args.Add(TEXT("ComponentClassName"), SubActorComp->ChildActorClass->GetDisplayNameText());
+				CachedAssetTitle = FText::Format(LOCTEXT("AddChildActorComponent", "Actor Class {ComponentClassName}"), Args);
 			}
 			else
 			{
-				FFormatNamedArguments Args;
-				Args.Add(TEXT("ClassName"), SourceTemplate->GetClass()->GetDisplayNameText());
-				CachedNodeTitle = FText::Format(LOCTEXT("AddClass", "Add {ClassName}"), Args);
+				CachedAssetTitle = FText::GetEmpty();
 			}
 		}
 	}
 
-	// FText::Format() is slow, so we cache the title to save on performance
-	if (!CachedNodeTitle.IsOutOfDate())
+	if (!CachedNodeTitle.IsEmpty())
 	{
-		return CachedNodeTitle;
+		if (TitleType == ENodeTitleType::FullTitle)
+		{
+			return FText::Format(LOCTEXT("FullAddComponentTitle", "{0}\n{1}"), CachedNodeTitle, CachedAssetTitle);
+		}
+		else if (!CachedAssetTitle.IsEmpty())
+		{
+			return FText::Format(LOCTEXT("ShortAddComponentTitle", "{0} [{1}]"), CachedNodeTitle, CachedAssetTitle);
+		}
+		else
+		{
+			return CachedNodeTitle;
+		}
 	}
 	return Super::GetNodeTitle(TitleType);
 }
