PR #800: Removed spurious changes when compiling bytecode (Contributed by [REDACTED_PERSON])

https://github.com/EpicGames/UnrealEngine/pull/800

Modified slightly so that we look up the real source graph name, rather than just using the package name.

diff --git a/Engine/Source/Editor/BlueprintGraph/Private/K2Node_Select.cpp b/Engine/Source/Editor/BlueprintGraph/Private/K2Node_Select.cpp
--- a/Engine/Source/Editor/BlueprintGraph/Private/K2Node_Select.cpp
+++ b/Engine/Source/Editor/BlueprintGraph/Private/K2Node_Select.cpp
@@ -160,9 +160,11 @@ public:
 							SelectionNodeType = IndexPin->PinType.PinCategory;
 						}
 					}
+
+					const UEdGraph* OwningGraph = Context.MessageLog.FindSourceObjectTypeChecked<UEdGraph>( SelectNode->GetGraph() );
 					LiteralStringTerm->Name =
 						FString::Printf(*LOCTEXT("SelectNodeIndexWarning", "Graph %s: Selection Node of type %s failed! Out of bounds indexing of the options. There are only %d options available.").ToString(),
-						(SelectNode && SelectNode->GetGraph()) ? *SelectNode->GetGraph()->GetFullName() : TEXT("NONE"),
+						(OwningGraph) ? *OwningGraph->GetFullName() : TEXT("NONE"),
 						*SelectionNodeType,
 						OptionPins.Num());
 					PrintStatement.RHS.Add(LiteralStringTerm);
