[INTEGRATE] Change 2440547 by [REDACTED_PERSON][REDACTED_HANDLE] on 2015/02/10 17:50:03

	Enable CanPasteHere to check in the skeleton class for functions to prevent renamed/new functions being filtered out of the action menu.
	#UE-9087 - BP variable not showing up in context menu for anonymous BP
	#Branch UE4-Releases
	#Proj BlueprintGraph

diff --git a/Engine/Source/Editor/BlueprintGraph/Classes/K2Node_CallFunction.h b/Engine/Source/Editor/BlueprintGraph/Classes/K2Node_CallFunction.h
--- a/Engine/Source/Editor/BlueprintGraph/Classes/K2Node_CallFunction.h
+++ b/Engine/Source/Editor/BlueprintGraph/Classes/K2Node_CallFunction.h
@@ -200,5 +200,8 @@ private:
 protected:
 	/** Helper function to ensure function is called in our context */
 	virtual void EnsureFunctionIsInBlueprint();
+
+	/** Helper function to find UFunction entries from the skeleton class, use with caution.. */
+	UFunction* GetTargetFunctionFromSkeletonClass() const;
 };
 
diff --git a/Engine/Source/Editor/BlueprintGraph/Private/K2Node_CallFunction.cpp b/Engine/Source/Editor/BlueprintGraph/Private/K2Node_CallFunction.cpp
--- a/Engine/Source/Editor/BlueprintGraph/Private/K2Node_CallFunction.cpp
+++ b/Engine/Source/Editor/BlueprintGraph/Private/K2Node_CallFunction.cpp
@@ -966,6 +966,18 @@ UFunction* UK2Node_CallFunction::GetTargetFunction() const
 	return Function;
 }
 
+UFunction* UK2Node_CallFunction::GetTargetFunctionFromSkeletonClass() const
+{
+	UFunction* TargetFunction = nullptr;
+	UClass* ParentClass = FunctionReference.GetMemberParentClass( this );
+	UBlueprint* OwningBP = ParentClass ? Cast<UBlueprint>( ParentClass->ClassGeneratedBy ) : nullptr;
+	if( UClass* SkeletonClass = OwningBP ? OwningBP->SkeletonGeneratedClass : nullptr )
+	{
+		TargetFunction = SkeletonClass->FindFunctionByName( FunctionReference.GetMemberName() );
+	}
+	return TargetFunction;
+}
+
 UEdGraphPin* UK2Node_CallFunction::GetThenPin() const
 {
 	const UEdGraphSchema_K2* K2Schema = GetDefault<UEdGraphSchema_K2>();
@@ -1031,8 +1043,12 @@ bool UK2Node_CallFunction::CanPasteHere(const UEdGraph* TargetGraph) const
 		{
 			AllowedFunctionTypes |= UEdGraphSchema_K2::EFunctionType::FT_Imperative;
 		}
-
-		bCanPaste = K2Schema->CanFunctionBeUsedInGraph(FBlueprintEditorUtils::FindBlueprintForGraphChecked(TargetGraph)->GeneratedClass, GetTargetFunction(), TargetGraph, AllowedFunctionTypes, false, FFunctionTargetInfo());
+		UFunction* TargetFunction = GetTargetFunction();
+		if( !TargetFunction )
+		{
+			TargetFunction = GetTargetFunctionFromSkeletonClass();
+		}
+		bCanPaste = K2Schema->CanFunctionBeUsedInGraph(FBlueprintEditorUtils::FindBlueprintForGraphChecked(TargetGraph)->GeneratedClass, TargetFunction, TargetGraph, AllowedFunctionTypes, false, FFunctionTargetInfo());
 	}
 	
 	return bCanPaste;
