[INTEGRATE] Change 2441937 by [REDACTED_PERSON][REDACTED_HANDLE] on 2015/02/11 13:20:49

	Fix for timeline nodes not displaying curve pins when a blueprint is reloaded.
	#UE-9544 - Timeline nodes lose track pins when reloaded
	#Branch UE4-Releases
	#Proj BlueprintGraph

diff --git a/Engine/Source/Editor/BlueprintGraph/Private/K2Node_Timeline.cpp b/Engine/Source/Editor/BlueprintGraph/Private/K2Node_Timeline.cpp
--- a/Engine/Source/Editor/BlueprintGraph/Private/K2Node_Timeline.cpp
+++ b/Engine/Source/Editor/BlueprintGraph/Private/K2Node_Timeline.cpp
@@ -72,6 +72,9 @@ void UK2Node_Timeline::AllocateDefaultPins()
 	UTimelineTemplate* Timeline = Blueprint->FindTimelineTemplateByVariableName(TimelineName);
 	if(Timeline)
 	{
+		// Ensure the timeline template is fully loaded or the node representation will be wrong.
+		PreloadObject(Timeline);
+
 		for(int32 i=0; i<Timeline->FloatTracks.Num(); i++)
 		{
 			FTTFloatTrack& FloatTrack = Timeline->FloatTracks[i];
