Blueprints: Fix default text not displaying for PrintText node, and made PrintString show up first in the filter list for "print"

diff --git a/Engine/Source/Editor/BlueprintGraph/Private/EdGraphSchema_K2.cpp b/Engine/Source/Editor/BlueprintGraph/Private/EdGraphSchema_K2.cpp
--- a/Engine/Source/Editor/BlueprintGraph/Private/EdGraphSchema_K2.cpp
+++ b/Engine/Source/Editor/BlueprintGraph/Private/EdGraphSchema_K2.cpp
@@ -3715,10 +3715,10 @@ void UEdGraphSchema_K2::TrySetDefaultText(UEdGraphPin& InPin, const FText& InNew
 				}
 				else
 				{
-				InPin.DefaultTextValue = FText::ChangeKey(TEXT(""), InPin.GetOwningNode()->NodeGuid.ToString() + TEXT("_") + InPin.PinName + FString::FromInt(InPin.GetOwningNode()->Pins.Find(&InPin)), InNewDefaultText);
+					InPin.DefaultTextValue = FText::ChangeKey(TEXT(""), InPin.GetOwningNode()->NodeGuid.ToString() + TEXT("_") + InPin.PinName + FString::FromInt(InPin.GetOwningNode()->Pins.Find(&InPin)), InNewDefaultText);
+				}
 			}
 		}
-		}
 
 		UEdGraphNode* Node = InPin.GetOwningNode();
 		check(Node);
@@ -3796,21 +3796,37 @@ bool UEdGraphSchema_K2::ShouldShowAssetPickerForPin(UEdGraphPin* Pin) const
 
 void UEdGraphSchema_K2::SetPinDefaultValue(UEdGraphPin* Pin, const UFunction* Function, const UProperty* Param) const
 {
-	if (Function && Param)
+	if ((Function != nullptr) && (Param != nullptr))
 	{
+		bool bHasAutomaticValue = false;
+
 		const FString MetadataDefaultValue = Function->GetMetaData(*Param->GetName());
 		if (!MetadataDefaultValue.IsEmpty())
 		{
 			// Specified default value in the metadata
-			Pin->DefaultValue = Pin->AutogeneratedDefaultValue = MetadataDefaultValue;
+			Pin->AutogeneratedDefaultValue = MetadataDefaultValue;
+			bHasAutomaticValue = true;
 		}
 		else
 		{
 			const FName MetadataCppDefaultValueKey( *(FString(TEXT("CPP_Default_")) + Param->GetName()) );
 			const FString MetadataCppDefaultValue = Function->GetMetaData(MetadataCppDefaultValueKey);
-			if(!MetadataCppDefaultValue.IsEmpty())
+			if (!MetadataCppDefaultValue.IsEmpty())
+			{
+				Pin->AutogeneratedDefaultValue = MetadataCppDefaultValue;
+				bHasAutomaticValue = true;
+			}
+		}
+
+		if (bHasAutomaticValue)
+		{
+			if (Pin->PinType.PinCategory == PC_Text)
+			{
+				Pin->DefaultTextValue = FText::AsCultureInvariant(Pin->AutogeneratedDefaultValue);
+			}
+			else
 			{
-				Pin->DefaultValue = Pin->AutogeneratedDefaultValue = MetadataCppDefaultValue;
+				Pin->DefaultValue = Pin->AutogeneratedDefaultValue;
 			}
 		}
 	}
diff --git a/Engine/Source/Runtime/Engine/Classes/Kismet/KismetSystemLibrary.h b/Engine/Source/Runtime/Engine/Classes/Kismet/KismetSystemLibrary.h
--- a/Engine/Source/Runtime/Engine/Classes/Kismet/KismetSystemLibrary.h
+++ b/Engine/Source/Runtime/Engine/Classes/Kismet/KismetSystemLibrary.h
@@ -195,7 +195,7 @@ class UKismetSystemLibrary : public UBlueprintFunctionLibrary
 	 * @param	bPrintToConsole	Whether or not to print the output to the console
 	 * @param	TextColor		Whether or not to print the output to the console
 	 */
-	UFUNCTION(BlueprintCallable, meta=(WorldContext="WorldContextObject", CallableWithoutWorldContext, Keywords = "log", AdvancedDisplay = "2"), Category="Utilities|String")
+	UFUNCTION(BlueprintCallable, meta=(WorldContext="WorldContextObject", CallableWithoutWorldContext, Keywords = "log print", AdvancedDisplay = "2"), Category="Utilities|String")
 	static void PrintString(UObject* WorldContextObject, const FString& InString = FString(TEXT("Hello")), bool bPrintToScreen = true, bool bPrintToLog = true, FLinearColor TextColor = FLinearColor(0.0,0.66,1.0));
 
 	/**
@@ -209,7 +209,7 @@ class UKismetSystemLibrary : public UBlueprintFunctionLibrary
 	 * @param	TextColor		Whether or not to print the output to the console
 	 */
 	UFUNCTION(BlueprintCallable, meta=(WorldContext="WorldContextObject", CallableWithoutWorldContext, Keywords = "log", AdvancedDisplay = "2"), Category="Utilities|Text")
-	static void PrintText(UObject* WorldContextObject, const FText& InText = FText::FromString(TEXT("Hello")), bool bPrintToScreen = true, bool bPrintToLog = true, FLinearColor TextColor = FLinearColor(0.0,0.66,1.0));
+	static void PrintText(UObject* WorldContextObject, const FText InText = FText::FromString(TEXT("Hello")), bool bPrintToScreen = true, bool bPrintToLog = true, FLinearColor TextColor = FLinearColor(0.0,0.66,1.0));
 
 	/**
 	 * Prints a warning string to the log and the screen. Meant to be used as a way to inform the user that they misused the node.
diff --git a/Engine/Source/Runtime/Engine/Private/KismetSystemLibrary.cpp b/Engine/Source/Runtime/Engine/Private/KismetSystemLibrary.cpp
--- a/Engine/Source/Runtime/Engine/Private/KismetSystemLibrary.cpp
+++ b/Engine/Source/Runtime/Engine/Private/KismetSystemLibrary.cpp
@@ -220,7 +220,7 @@ void UKismetSystemLibrary::PrintString(UObject* WorldContextObject, const FStrin
 #endif
 }
 
-void UKismetSystemLibrary::PrintText(UObject* WorldContextObject, const FText& InText, bool bPrintToScreen, bool bPrintToLog, FLinearColor TextColor)
+void UKismetSystemLibrary::PrintText(UObject* WorldContextObject, const FText InText, bool bPrintToScreen, bool bPrintToLog, FLinearColor TextColor)
 {
 #if !(UE_BUILD_SHIPPING || UE_BUILD_TEST) // Do not Print in Shipping or Test
 	PrintString(WorldContextObject, InText.ToString(), bPrintToScreen, bPrintToLog, TextColor);
