[UE-9879] No longer spamming "Files need check-out!" when compiling a blueprint, we avoid marking the blueprint as structurely modified simply because generated data was updated

diff --git a/Engine/Source/Editor/BlueprintGraph/Private/CallFunctionHandler.cpp b/Engine/Source/Editor/BlueprintGraph/Private/CallFunctionHandler.cpp
--- a/Engine/Source/Editor/BlueprintGraph/Private/CallFunctionHandler.cpp
+++ b/Engine/Source/Editor/BlueprintGraph/Private/CallFunctionHandler.cpp
@@ -672,8 +672,7 @@ void FKCHandler_CallFunction::Transform(FKismetFunctionContext& Context, UEdGrap
 		if ((OldOutPin != NULL) && (OldOutPin->LinkedTo.Num() > 0))
 		{
 			// Create a dummy execution sequence that will be the target of the return call from the latent action
-			UK2Node_ExecutionSequence* DummyNode = CallFuncNode->GetGraph()->CreateBlankNode<UK2Node_ExecutionSequence>();
-			CompilerContext.MessageLog.NotifyIntermediateObjectCreation(DummyNode, CallFuncNode);
+			UK2Node_ExecutionSequence* DummyNode = CompilerContext.SpawnIntermediateNode<UK2Node_ExecutionSequence>(CallFuncNode);
 			DummyNode->AllocateDefaultPins();
 
 			// Wire in the dummy node
diff --git a/Engine/Source/Editor/BlueprintGraph/Private/EdGraphSchema_K2.cpp b/Engine/Source/Editor/BlueprintGraph/Private/EdGraphSchema_K2.cpp
--- a/Engine/Source/Editor/BlueprintGraph/Private/EdGraphSchema_K2.cpp
+++ b/Engine/Source/Editor/BlueprintGraph/Private/EdGraphSchema_K2.cpp
@@ -1923,8 +1923,9 @@ bool UEdGraphSchema_K2::TryCreateConnection(UEdGraphPin* PinA, UEdGraphPin* PinB
 
 	bool bModified = UEdGraphSchema::TryCreateConnection(PinA, PinB);
 
-	if (bModified)
+	if (bModified && !PinA->HasAnyFlags(RF_Transient))
 	{
+		check(!PinB->HasAnyFlags(RF_Transient));
 		FBlueprintEditorUtils::MarkBlueprintAsModified(Blueprint);
 	}
 
diff --git a/Engine/Source/Editor/KismetCompiler/Private/KismetCompiler.cpp b/Engine/Source/Editor/KismetCompiler/Private/KismetCompiler.cpp
--- a/Engine/Source/Editor/KismetCompiler/Private/KismetCompiler.cpp
+++ b/Engine/Source/Editor/KismetCompiler/Private/KismetCompiler.cpp
@@ -170,7 +170,7 @@ void FKismetCompilerContext::CleanAndSanitizeClass(UBlueprintGeneratedClass* Cla
 	NewClass = ClassToClean;
 	OldCDO = ClassToClean->ClassDefaultObject; // we don't need to create the CDO at this point
 	
-	const ERenameFlags RenFlags = REN_DontCreateRedirectors | (bRecompilingOnLoad ? REN_ForceNoResetLoaders : 0);
+	const ERenameFlags RenFlags = REN_DontCreateRedirectors | (bRecompilingOnLoad ? REN_ForceNoResetLoaders : 0) | REN_NonTransactional | REN_DoNotDirty;
 
 	if( OldCDO )
 	{
diff --git a/Engine/Source/Editor/KismetCompiler/Public/KismetCompiler.h b/Engine/Source/Editor/KismetCompiler/Public/KismetCompiler.h
--- a/Engine/Source/Editor/KismetCompiler/Public/KismetCompiler.h
+++ b/Engine/Source/Editor/KismetCompiler/Public/KismetCompiler.h
@@ -103,7 +103,7 @@ public:
 		NodeType* Result = ParentGraph->CreateBlankNode<NodeType>();
 		MessageLog.NotifyIntermediateObjectCreation(Result, SourceNode);
 		Result->CreateNewGuid();
-
+		Result->SetFlags(RF_Transient);
 		AutoAssignNodePosition(Result);
 
 		return Result;
diff --git a/Engine/Source/Runtime/Engine/Private/EdGraph/EdGraphNode.cpp b/Engine/Source/Runtime/Engine/Private/EdGraph/EdGraphNode.cpp
--- a/Engine/Source/Runtime/Engine/Private/EdGraph/EdGraphNode.cpp
+++ b/Engine/Source/Runtime/Engine/Private/EdGraph/EdGraphNode.cpp
@@ -70,6 +70,11 @@ UEdGraphPin* UEdGraphNode::CreatePin(EEdGraphPinDirection Dir, const FString& Pi
 	NewPin->PinType.bIsConst = bIsConst;
 	NewPin->SetFlags(RF_Transactional);
 
+	if (HasAnyFlags(RF_Transient))
+	{
+		NewPin->SetFlags(RF_Transient);
+	}
+
 	Modify(false);
 	if ( Pins.IsValidIndex( Index ) )
 	{
