Fixed some blueprint node pin naming logic for splitting pins that wasn't using localization-aware formatting patterns.

diff --git a/Engine/Source/Editor/BlueprintGraph/Private/EdGraphSchema_K2.cpp b/Engine/Source/Editor/BlueprintGraph/Private/EdGraphSchema_K2.cpp
--- a/Engine/Source/Editor/BlueprintGraph/Private/EdGraphSchema_K2.cpp
+++ b/Engine/Source/Editor/BlueprintGraph/Private/EdGraphSchema_K2.cpp
@@ -5446,20 +5446,16 @@ void UEdGraphSchema_K2::SplitPin(UEdGraphPin* Pin) const
 			const FEdGraphPinType& ProtoPinType = ProtoPin->PinType;
 			UEdGraphPin* SubPin = GraphNode->CreatePin(Pin->Direction, ProtoPinType.PinCategory, ProtoPinType.PinSubCategory, ProtoPinType.PinSubCategoryObject.Get(), ProtoPinType.bIsArray, false, PinName);
 
-			if (K2Node != nullptr && K2Node->ShouldDrawCompact())
+			if (K2Node != nullptr && K2Node->ShouldDrawCompact() && !Pin->ParentPin)
 			{
-				if (Pin->ParentPin)
-				{
-					SubPin->PinFriendlyName = FText::FromString(FString::Printf(TEXT("%s %s"), *Pin->GetDisplayName().ToString(), *ProtoPin->GetDisplayName().ToString()));
-				}
-				else
-				{
-					SubPin->PinFriendlyName = ProtoPin->GetDisplayName();
-				}
+				SubPin->PinFriendlyName = ProtoPin->GetDisplayName();
 			}
 			else
 			{
-				SubPin->PinFriendlyName = FText::FromString(FString::Printf(TEXT("%s %s"), *Pin->GetDisplayName().ToString(), *ProtoPin->GetDisplayName().ToString()));
+				FFormatNamedArguments Arguments;
+				Arguments.Add(TEXT("PinDisplayName"), Pin->GetDisplayName());
+				Arguments.Add(TEXT("PinDisplayName"), ProtoPin->GetDisplayName());
+				SubPin->PinFriendlyName = FText::Format(LOCTEXT("SplitPinFriendlyNameFormat", "{PinDisplayName} {ProtoPinDisplayName}"), Arguments);
 			}
 
 			SubPin->DefaultValue = ProtoPin->DefaultValue;
