Removing this assert, the existing behavior is not harmful

diff --git a/Engine/Source/Editor/BlueprintGraph/Private/EdGraphSchema_K2.cpp b/Engine/Source/Editor/BlueprintGraph/Private/EdGraphSchema_K2.cpp
--- a/Engine/Source/Editor/BlueprintGraph/Private/EdGraphSchema_K2.cpp
+++ b/Engine/Source/Editor/BlueprintGraph/Private/EdGraphSchema_K2.cpp
@@ -1925,7 +1925,6 @@ bool UEdGraphSchema_K2::TryCreateConnection(UEdGraphPin* PinA, UEdGraphPin* PinB
 
 	if (bModified && !PinA->HasAnyFlags(RF_Transient))
 	{
-		check(!PinB->HasAnyFlags(RF_Transient));
 		FBlueprintEditorUtils::MarkBlueprintAsModified(Blueprint);
 	}
 
