Encapsulate ChildActorComponent.ChildActorClass so that setting it will destroy and recreate an Actor of the correct type
Expose SetChildActorClass to blueprints

diff --git a/Engine/Source/Editor/BlueprintGraph/Private/K2Node_AddComponent.cpp b/Engine/Source/Editor/BlueprintGraph/Private/K2Node_AddComponent.cpp
--- a/Engine/Source/Editor/BlueprintGraph/Private/K2Node_AddComponent.cpp
+++ b/Engine/Source/Editor/BlueprintGraph/Private/K2Node_AddComponent.cpp
@@ -150,7 +150,7 @@ void UK2Node_AddComponent::ValidateNodeDuringCompilation(FCompilerResultsLog& Me
 		{
 			UBlueprint const* Blueprint = GetBlueprint();
 
-			UClass const* ChildActorClass = ChildActorComponent->ChildActorClass;
+			UClass const* ChildActorClass = ChildActorComponent->GetChildActorClass();
 			if (ChildActorClass == Blueprint->GeneratedClass)
 			{
 				UEdGraph const* ParentGraph = GetGraph();
@@ -314,10 +314,10 @@ FText UK2Node_AddComponent::GetNodeTitle(ENodeTitleType::Type TitleType) const
 				Args.Add(TEXT("AssetName"), FText::FromString(AssociatedAsset->GetName()));
 				CachedAssetTitle = FText::Format(LOCTEXT("AddComponentAssetDescription", "{AssetType} {AssetName}"), Args);
 			}
-			else if ((SubActorComp != nullptr) && (SubActorComp->ChildActorClass != nullptr))
+			else if ((SubActorComp != nullptr) && (SubActorComp->GetChildActorClass() != nullptr))
 			{
 				FFormatNamedArguments Args;
-				Args.Add(TEXT("ComponentClassName"), SubActorComp->ChildActorClass->GetDisplayNameText());
+				Args.Add(TEXT("ComponentClassName"), SubActorComp->GetChildActorClass()->GetDisplayNameText());
 				CachedAssetTitle = FText::Format(LOCTEXT("AddChildActorComponent", "Actor Class {ComponentClassName}"), Args);
 			}
 			else
diff --git a/Engine/Source/Editor/UnrealEd/Private/ComponentAssetBroker.cpp b/Engine/Source/Editor/UnrealEd/Private/ComponentAssetBroker.cpp
--- a/Engine/Source/Editor/UnrealEd/Private/ComponentAssetBroker.cpp
+++ b/Engine/Source/Editor/UnrealEd/Private/ComponentAssetBroker.cpp
@@ -225,7 +225,7 @@ public:
 			}
 			if (Class && Class->IsChildOf<AActor>())
 			{
-				ChildActorComp->ChildActorClass = Class;
+				ChildActorComp->SetChildActorClass(Class);
 				return true;
 			}
 		}
@@ -237,7 +237,7 @@ public:
 	{
 		if (UChildActorComponent* ChildActorComp = Cast<UChildActorComponent>(InComponent))
 		{
-			return UBlueprint::GetBlueprintFromClass(*(ChildActorComp->ChildActorClass));
+			return UBlueprint::GetBlueprintFromClass(*(ChildActorComp->GetChildActorClass()));
 		}
 		return NULL;
 	}
diff --git a/Engine/Source/Runtime/Engine/Classes/Components/ChildActorComponent.h b/Engine/Source/Runtime/Engine/Classes/Components/ChildActorComponent.h
--- a/Engine/Source/Runtime/Engine/Classes/Components/ChildActorComponent.h
+++ b/Engine/Source/Runtime/Engine/Classes/Components/ChildActorComponent.h
@@ -13,10 +13,19 @@ class UChildActorComponent : public USceneComponent
 {
 	GENERATED_UCLASS_BODY()
 
+	UFUNCTION(BlueprintCallable, Category=ChildActorComponent)
+	ENGINE_API void SetChildActorClass(TSubclassOf<AActor> Class);
+
+	TSubclassOf<AActor> GetChildActorClass() const { return ChildActorClass; }
+
+	friend class FChildActorComponentDetails;
+
+private:
 	/** The class of Actor to spawn */
-	UPROPERTY(EditAnywhere, BlueprintReadOnly, Category=ChildActorComponent, meta=(OnlyPlaceable))
+	UPROPERTY(EditAnywhere, BlueprintReadOnly, Category=ChildActorComponent, meta=(OnlyPlaceable, AllowPrivateAccess="true"))
 	TSubclassOf<AActor>	ChildActorClass;
 
+public:
 	/** The actor that we spawned and own */
 	UPROPERTY(BlueprintReadOnly, Category=ChildActorComponent, TextExportTransient, NonPIEDuplicateTransient)
 	AActor*	ChildActor;
diff --git a/Engine/Source/Runtime/Engine/Private/ChildActorComponent.cpp b/Engine/Source/Runtime/Engine/Private/ChildActorComponent.cpp
--- a/Engine/Source/Runtime/Engine/Private/ChildActorComponent.cpp
+++ b/Engine/Source/Runtime/Engine/Private/ChildActorComponent.cpp
@@ -186,6 +186,16 @@ void UChildActorComponent::ApplyComponentInstanceData(FChildActorComponentInstan
 	}
 }
 
+void UChildActorComponent::SetChildActorClass(TSubclassOf<AActor> Class)
+{
+	ChildActorClass = Class;
+	if (IsRegistered())
+	{
+		DestroyChildActor();
+		CreateChildActor();
+	}
+}
+
 void UChildActorComponent::CreateChildActor()
 {
 	// Kill spawned actor if we have one
