Copy and pasting custom events that have pins with spaces in their name will no longer create individual pins for each word.

diff --git a/Engine/Source/Editor/BlueprintGraph/Private/EdGraphSchema_K2.cpp b/Engine/Source/Editor/BlueprintGraph/Private/EdGraphSchema_K2.cpp
--- a/Engine/Source/Editor/BlueprintGraph/Private/EdGraphSchema_K2.cpp
+++ b/Engine/Source/Editor/BlueprintGraph/Private/EdGraphSchema_K2.cpp
@@ -1925,7 +1925,6 @@ bool UEdGraphSchema_K2::TryCreateConnection(UEdGraphPin* PinA, UEdGraphPin* PinB
 
 	if (bModified && !PinA->HasAnyFlags(RF_Transient))
 	{
-		check(!PinB->HasAnyFlags(RF_Transient));
 		FBlueprintEditorUtils::MarkBlueprintAsModified(Blueprint);
 	}
 
@@ -5021,6 +5020,8 @@ UEdGraphNode* UEdGraphSchema_K2::CreateSubstituteNode(UEdGraphNode* Node, const
 			EventNode->Pins.Empty();
 			EventNode->UserDefinedPins.Empty();
 
+			bool bOriginalWasCustomEvent = Cast<UK2Node_CustomEvent>(Node) != nullptr;
+
 			// Fixup pins
 			for(int32 PinIndex = 0; PinIndex < CustomEventNode->Pins.Num(); ++PinIndex)
 			{
@@ -5031,7 +5032,7 @@ UEdGraphNode* UEdGraphSchema_K2::CreateSubstituteNode(UEdGraphNode* Node, const
 				Pin->Rename(*Pin->GetName(), CustomEventNode, RenameFlags);
 
 				// Don't include execution or delegate output pins as user-defined pins
-				if(!IsExecPin(*Pin) && !IsDelegateCategory(Pin->PinType.PinCategory))
+				if(!bOriginalWasCustomEvent && !IsExecPin(*Pin) && !IsDelegateCategory(Pin->PinType.PinCategory))
 				{
 					// Check to see if this pin already exists as a user-defined pin on the custom event node
 					bool bFoundUserDefinedPin = false;
diff --git a/Engine/Source/Editor/BlueprintGraph/Private/K2Node_EditablePinBase.cpp b/Engine/Source/Editor/BlueprintGraph/Private/K2Node_EditablePinBase.cpp
--- a/Engine/Source/Editor/BlueprintGraph/Private/K2Node_EditablePinBase.cpp
+++ b/Engine/Source/Editor/BlueprintGraph/Private/K2Node_EditablePinBase.cpp
@@ -63,7 +63,7 @@ void UK2Node_EditablePinBase::ExportCustomProperties(FOutputDevice& Out, uint32
 		const FUserPinInfo& PinInfo = *UserDefinedPins[PinIndex].Get();
 
 		Out.Logf( TEXT("%sCustomProperties UserDefinedPin "), FCString::Spc(Indent));
-		Out.Logf( TEXT("Name=%s "), *PinInfo.PinName);
+		Out.Logf( TEXT("Name=\"%s\" "), *PinInfo.PinName);
 		Out.Logf( TEXT("IsArray=%s "), (PinInfo.PinType.bIsArray ? TEXT("1") : TEXT("0")));
 		Out.Logf( TEXT("IsReference=%s "), (PinInfo.PinType.bIsReference ? TEXT("1") : TEXT("0")));
 		
