UE-9695 UK2Node_CreateDelegate excludes all functions that cannot be used by [REDACTED_PERSON].

diff --git a/Engine/Source/Editor/BlueprintGraph/Classes/EdGraphSchema_K2.h b/Engine/Source/Editor/BlueprintGraph/Classes/EdGraphSchema_K2.h
--- a/Engine/Source/Editor/BlueprintGraph/Classes/EdGraphSchema_K2.h
+++ b/Engine/Source/Editor/BlueprintGraph/Classes/EdGraphSchema_K2.h
@@ -682,6 +682,9 @@ public:
 	/** Can this function be called by kismet code */
 	static bool CanUserKismetCallFunction(const UFunction* Function);
 
+	/** Returns if hunction has output parameter(s) */
+	static bool HasFunctionAnyOutputParameter(const UFunction* Function);
+
 	enum EDelegateFilterMode
 	{
 		CannotBeDelegate,
diff --git a/Engine/Source/Editor/BlueprintGraph/Private/DelegateNodeHandlers.cpp b/Engine/Source/Editor/BlueprintGraph/Private/DelegateNodeHandlers.cpp
--- a/Engine/Source/Editor/BlueprintGraph/Private/DelegateNodeHandlers.cpp
+++ b/Engine/Source/Editor/BlueprintGraph/Private/DelegateNodeHandlers.cpp
@@ -216,11 +216,7 @@ void FKCHandler_CreateDelegate::RegisterNets(FKismetFunctionContext& Context, UE
 		return;
 	}
 
-	if(UFunction* SignatureFunc = DelegateNode->GetDelegateSignature())
-	{
-		FKCHandlerDelegateHelper::CheckOutputsParametersInDelegateSignature(SignatureFunc, DelegateNode, CompilerContext.MessageLog);
-	}
-	else
+	if(!DelegateNode->GetDelegateSignature())
 	{
 		const FString ErrorStr = FString::Printf(
 			*LOCTEXT("NoDelegateFunction", "No delegate function '%' @@").ToString(),
diff --git a/Engine/Source/Editor/BlueprintGraph/Private/EdGraphSchema_K2.cpp b/Engine/Source/Editor/BlueprintGraph/Private/EdGraphSchema_K2.cpp
--- a/Engine/Source/Editor/BlueprintGraph/Private/EdGraphSchema_K2.cpp
+++ b/Engine/Source/Editor/BlueprintGraph/Private/EdGraphSchema_K2.cpp
@@ -744,23 +744,20 @@ bool UEdGraphSchema_K2::CanKismetOverrideFunction(const UFunction* Function)
 		(Function->HasAllFunctionFlags(FUNC_BlueprintEvent) && !Function->HasAllFunctionFlags(FUNC_Delegate) && !Function->GetBoolMetaData(FBlueprintMetadata::MD_BlueprintInternalUseOnly) && !Function->HasMetaData(FBlueprintMetadata::MD_DeprecatedFunction));
 }
 
-struct FNoOutputParametersHelper 
+bool UEdGraphSchema_K2::HasFunctionAnyOutputParameter(const UFunction* InFunction)
 {
-	static bool Check(const UFunction* InFunction)
+	check(InFunction);
+	for (TFieldIterator<UProperty> PropIt(InFunction); PropIt && (PropIt->PropertyFlags & CPF_Parm); ++PropIt)
 	{
-		check(InFunction);
-		for (TFieldIterator<UProperty> PropIt(InFunction); PropIt && (PropIt->PropertyFlags & CPF_Parm); ++PropIt)
+		UProperty* FuncParam = *PropIt;
+		if (FuncParam->HasAnyPropertyFlags(CPF_ReturnParm) || (FuncParam->HasAnyPropertyFlags(CPF_OutParm) && !FuncParam->HasAnyPropertyFlags(CPF_ReferenceParm) && !FuncParam->HasAnyPropertyFlags(CPF_ConstParm)))
 		{
-			UProperty* FuncParam = *PropIt;
-			if(FuncParam->HasAnyPropertyFlags(CPF_ReturnParm) || (FuncParam->HasAnyPropertyFlags(CPF_OutParm) && !FuncParam->HasAnyPropertyFlags(CPF_ReferenceParm) && !FuncParam->HasAnyPropertyFlags(CPF_ConstParm)))
-			{
-				return false;
-			}
+			return true;
 		}
-
-		return true;
 	}
-};
+
+	return false;
+}
 
 bool UEdGraphSchema_K2::FunctionCanBePlacedAsEvent(const UFunction* InFunction)
 {
@@ -771,7 +768,7 @@ bool UEdGraphSchema_K2::FunctionCanBePlacedAsEvent(const UFunction* InFunction)
 	}
 
 	// Then look to see if we have any output, return, or reference params
-	return FNoOutputParametersHelper::Check(InFunction);
+	return !HasFunctionAnyOutputParameter(InFunction);
 }
 
 bool UEdGraphSchema_K2::FunctionCanBeUsedInDelegate(const UFunction* InFunction)
@@ -784,7 +781,7 @@ bool UEdGraphSchema_K2::FunctionCanBeUsedInDelegate(const UFunction* InFunction)
 		return false;
 	}
 
-	return FNoOutputParametersHelper::Check(InFunction);
+	return true;
 }
 
 FString UEdGraphSchema_K2::GetFriendlySignatureName(const UFunction* Function)
diff --git a/Engine/Source/Editor/Kismet/Private/BlueprintDetailsCustomization.cpp b/Engine/Source/Editor/Kismet/Private/BlueprintDetailsCustomization.cpp
--- a/Engine/Source/Editor/Kismet/Private/BlueprintDetailsCustomization.cpp
+++ b/Engine/Source/Editor/Kismet/Private/BlueprintDetailsCustomization.cpp
@@ -3070,7 +3070,7 @@ void FBlueprintDelegateActionDetails::CollectAvailibleSignatures()
 			for(TFieldIterator<UFunction> It(ScopeClass, EFieldIteratorFlags::IncludeSuper); It; ++It)
 			{
 				UFunction* Func = *It;
-				if (UEdGraphSchema_K2::FunctionCanBeUsedInDelegate(Func))
+				if (UEdGraphSchema_K2::FunctionCanBeUsedInDelegate(Func) && !UEdGraphSchema_K2::HasFunctionAnyOutputParameter(Func))
 				{
 					TSharedPtr<FString> ItemData = MakeShareable(new FString(Func->GetName()));
 					FunctionsToCopySignatureFrom.Add(ItemData);
