Blueprints: Add GetDisplayName as an acceptable autocast from object to string

diff --git a/Engine/Source/Editor/BlueprintGraph/Private/EdGraphSchema_K2.cpp b/Engine/Source/Editor/BlueprintGraph/Private/EdGraphSchema_K2.cpp
--- a/Engine/Source/Editor/BlueprintGraph/Private/EdGraphSchema_K2.cpp
+++ b/Engine/Source/Editor/BlueprintGraph/Private/EdGraphSchema_K2.cpp
@@ -2157,6 +2157,10 @@ bool UEdGraphSchema_K2::SearchForAutocastFunction(const UEdGraphPin* OutputPin,
 				TargetFunction = TEXT("GetObjectClass");
 			}
 		}
+		else if (InputPin->PinType.PinCategory == PC_String)
+		{
+			TargetFunction = TEXT("GetDisplayName");
+		}
 	}
 
 	return TargetFunction != NAME_None;
