[INTEGRATE] Change 2451844 by [REDACTED_PERSON][REDACTED_HANDLE] on 2015/02/19 12:25:16

	Fix additional cases where RF_Public was not set on Components created as a template in a Blueprint

diff --git a/Engine/Source/Editor/BlueprintGraph/Private/BlueprintComponentNodeSpawner.cpp b/Engine/Source/Editor/BlueprintGraph/Private/BlueprintComponentNodeSpawner.cpp
--- a/Engine/Source/Editor/BlueprintGraph/Private/BlueprintComponentNodeSpawner.cpp
+++ b/Engine/Source/Editor/BlueprintGraph/Private/BlueprintComponentNodeSpawner.cpp
@@ -168,8 +168,7 @@ UEdGraphNode* UBlueprintComponentNodeSpawner::Invoke(UEdGraph* ParentGraph, FBin
 		UBlueprint* Blueprint = NewNode->GetBlueprint();
 
 		FName DesiredComponentName = MakeUniqueObjectName(Blueprint->GeneratedClass, Class, *Class->GetDisplayNameText().ToString());
-		UActorComponent* ComponentTemplate = NewObject<UActorComponent>(Blueprint->GeneratedClass, Class, DesiredComponentName);
-		ComponentTemplate->SetFlags(RF_ArchetypeObject);
+		UActorComponent* ComponentTemplate = NewObject<UActorComponent>(Blueprint->GeneratedClass, Class, DesiredComponentName, RF_ArchetypeObject | RF_Public);
 
 		Blueprint->ComponentTemplates.Add(ComponentTemplate);
 
diff --git a/Engine/Source/Editor/UnrealEd/Private/Tests/EditorBuildPromotionTests.cpp b/Engine/Source/Editor/UnrealEd/Private/Tests/EditorBuildPromotionTests.cpp
--- a/Engine/Source/Editor/UnrealEd/Private/Tests/EditorBuildPromotionTests.cpp
+++ b/Engine/Source/Editor/UnrealEd/Private/Tests/EditorBuildPromotionTests.cpp
@@ -600,8 +600,7 @@ namespace EditorBuildPromotionTestUtils
 
 			ensure(NULL != Cast<UBlueprintGeneratedClass>(InBlueprint->GeneratedClass));
 			// Then create a new template object, and add to array in
-			UActorComponent* NewTemplate = NewObject<UActorComponent>(InBlueprint->GeneratedClass, ComponentClass);
-			NewTemplate->SetFlags(RF_ArchetypeObject);
+			UActorComponent* NewTemplate = NewObject<UActorComponent>(InBlueprint->GeneratedClass, ComponentClass, NAME_None, RF_ArchetypeObject|RF_Public);
 			InBlueprint->ComponentTemplates.Add(NewTemplate);
 
 			// Set the name of the template as the default for the TemplateName param
