Duplicating Blueprint variables will no longer result in two instances when searching using Find-in-Blueprint.

diff --git a/Engine/Source/Editor/BlueprintGraph/Classes/EdGraphSchema_K2_Actions.h b/Engine/Source/Editor/BlueprintGraph/Classes/EdGraphSchema_K2_Actions.h
--- a/Engine/Source/Editor/BlueprintGraph/Classes/EdGraphSchema_K2_Actions.h
+++ b/Engine/Source/Editor/BlueprintGraph/Classes/EdGraphSchema_K2_Actions.h
@@ -514,10 +514,14 @@ private:
 	/** The class that owns this item */
 	TWeakObjectPtr<UClass>		OwningClass;
 
+	/** TRUE if the variable's type is boolean */
+	bool bIsVarBool;
+
 public:
-	void SetVariableInfo(const FName& InVarName, const UClass* InOwningClass)
+	void SetVariableInfo(const FName& InVarName, const UClass* InOwningClass, bool bInIsVarBool)
 	{
 		VarName = InVarName;
+		bIsVarBool = bInIsVarBool;
 
 		check(InOwningClass);
 		OwningClass = InOwningClass;
@@ -540,6 +544,11 @@ public:
 		return VarName;
 	}
 
+	FString GetFriendlyVariableName() const
+	{
+		return FName::NameToDisplayString( VarName.ToString(), bIsVarBool );
+	}
+
 	UClass* GetVariableClass() const
 	{
 		return OwningClass.Get();
@@ -572,10 +581,14 @@ private:
 	/** The struct that owns this item */
 	TWeakObjectPtr<UStruct>		VariableScope;
 
+	/** TRUE if the variable's type is boolean */
+	bool bIsVarBool;
+
 public:
-	void SetVariableInfo(const FName& InVarName, const UStruct* InVariableScope)
+	void SetVariableInfo(const FName& InVarName, const UStruct* InVariableScope, bool bInIsVarBool)
 	{
 		VarName = InVarName;
+		bIsVarBool = bInIsVarBool;
 
 		check(InVariableScope);
 		VariableScope = InVariableScope;
@@ -598,6 +611,11 @@ public:
 		return VarName;
 	}
 
+	FString GetFriendlyVariableName() const
+	{
+		return FName::NameToDisplayString( VarName.ToString(), bIsVarBool );
+	}
+
 	UStruct* GetVariableScope() const
 	{
 		return VariableScope.Get();
diff --git a/Engine/Source/Editor/BlueprintGraph/Private/K2ActionMenuBuilder.cpp b/Engine/Source/Editor/BlueprintGraph/Private/K2ActionMenuBuilder.cpp
--- a/Engine/Source/Editor/BlueprintGraph/Private/K2ActionMenuBuilder.cpp
+++ b/Engine/Source/Editor/BlueprintGraph/Private/K2ActionMenuBuilder.cpp
@@ -1581,7 +1581,9 @@ void FK2ActionMenuBuilder::GetAllActionsForClass(FBlueprintPaletteListBuilder& P
 					const FString PropertyCategory = FObjectEditorUtils::GetCategory(Property);
 
 					TSharedPtr<FEdGraphSchemaAction_K2Var> NewVarAction = TSharedPtr<FEdGraphSchemaAction_K2Var>(new FEdGraphSchemaAction_K2Var(PropertyCategory, FText::FromString(PropertyDesc), PropertyTooltip, 0));
-					NewVarAction->SetVariableInfo(Property->GetFName(), InClass);
+					const UArrayProperty* ArrayProperty = Cast<const UArrayProperty>(Property);
+					const UProperty* TestProperty = ArrayProperty ? ArrayProperty->Inner : Property;
+					NewVarAction->SetVariableInfo(Property->GetFName(), InClass, Cast<UBoolProperty>(TestProperty) != nullptr);
 
 					PaletteBuilder.AddAction(NewVarAction);
 				}
diff --git a/Engine/Source/Editor/UnrealEd/Private/Kismet2/BlueprintEditorUtils.cpp b/Engine/Source/Editor/UnrealEd/Private/Kismet2/BlueprintEditorUtils.cpp
--- a/Engine/Source/Editor/UnrealEd/Private/Kismet2/BlueprintEditorUtils.cpp
+++ b/Engine/Source/Editor/UnrealEd/Private/Kismet2/BlueprintEditorUtils.cpp
@@ -3966,6 +3966,7 @@ FBPVariableDescription FBlueprintEditorUtils::DuplicateVariableDescription(UBlue
 	// Now create new variable
 	FBPVariableDescription NewVar = InVariableDescription;
 	NewVar.VarName = DuplicatedVariableName;
+	NewVar.FriendlyName = FName::NameToDisplayString( NewVar.VarName.ToString(), NewVar.VarType.PinCategory == GetDefault<UEdGraphSchema_K2>()->PC_Boolean);
 	NewVar.VarGuid = FGuid::NewGuid();
 
 	return NewVar;
