Copy and pasting of Blueprint Events works again.

Version bumped to deprecate old variables for event handling and to upgrade Blueprints to the new system.

diff --git a/Engine/Source/Editor/BlueprintGraph/Classes/K2Node_Event.h b/Engine/Source/Editor/BlueprintGraph/Classes/K2Node_Event.h
--- a/Engine/Source/Editor/BlueprintGraph/Classes/K2Node_Event.h
+++ b/Engine/Source/Editor/BlueprintGraph/Classes/K2Node_Event.h
@@ -13,30 +13,17 @@ class UK2Node_Event : public UK2Node_EditablePinBase
 	GENERATED_UCLASS_BODY()
 	BLUEPRINTGRAPH_API static const FString DelegateOutputName;
 
-	/******************************************************************************/
-
-	/**
-	 * Properties EventSignatureName and EventSignatureClass should be deprecated.
-	 * In order to do this we need to edit ObjectVersion.h which is locked for a
-	 * while. After lock lift those properties will be deprecated and eventually
-	 * removed. Please use EventReference instead and it will be translated into
-	 * EventReference accordingly.
-	 */
-
 	/** Name of function signature that this event implements */
-	DEPRECATED(4.7, "This field is deprecated, please use EventReference instead. For details read comment above.")
 	UPROPERTY()
-	FName EventSignatureName;
+	FName EventSignatureName_DEPRECATED;
 
 	/** Class that the function signature is from. */
-	DEPRECATED(4.7, "This field is deprecated, please use EventReference instead. For details read comment above.")
 	UPROPERTY()
-	TSubclassOf<class UObject> EventSignatureClass;
+	TSubclassOf<class UObject> EventSignatureClass_DEPRECATED;
 
 	/** Reference for the function this event is linked to */
-	UPROPERTY(Transient)
+	UPROPERTY()
 	FMemberReference EventReference;
-	/******************************************************************************/
 
 	/** If true, we are actually overriding this function, not making a new event with a signature that matches */
 	UPROPERTY()
diff --git a/Engine/Source/Editor/BlueprintGraph/Private/K2Node_ActorBoundEvent.cpp b/Engine/Source/Editor/BlueprintGraph/Private/K2Node_ActorBoundEvent.cpp
--- a/Engine/Source/Editor/BlueprintGraph/Private/K2Node_ActorBoundEvent.cpp
+++ b/Engine/Source/Editor/BlueprintGraph/Private/K2Node_ActorBoundEvent.cpp
@@ -197,36 +197,19 @@ bool UK2Node_ActorBoundEvent::IsUsedByAuthorityOnlyDelegate() const
 
 void UK2Node_ActorBoundEvent::Serialize(FArchive& Ar)
 {
-	// This deprecation warning suppression is due to the workaround. For more
-	// details read comment in K2Node_Event.h near EventSignatureName property
-	// definition.
-
-	PRAGMA_DISABLE_DEPRECATION_WARNINGS
-
-	// Workaround, so fields EventSignatureName and EventSignatureClass work
-	// until ObjectVersion.h lock is lifted.
-	if (Ar.IsSaving())
-	{
-		EventSignatureName = EventReference.GetMemberName();
-		EventSignatureClass = DelegateOwnerClass;
-		Super::Super::Serialize(Ar);
-		return;
-	}
-
 	Super::Serialize(Ar);
 
 	// Fix up legacy nodes that may not yet have a delegate pin
 	if(Ar.IsLoading())
 	{
-		DelegateOwnerClass = EventSignatureClass;
-		UMulticastDelegateProperty* TargetDelegateProp = GetTargetDelegateProperty();
-		if(TargetDelegateProp)
+		if(Ar.UE4Ver() < VER_UE4_K2NODE_EVENT_MEMBER_REFERENCE)
 		{
+			DelegateOwnerClass = EventSignatureClass_DEPRECATED;
+			UMulticastDelegateProperty* TargetDelegateProp = GetTargetDelegateProperty();
+			check(TargetDelegateProp);
 			EventReference.SetFromField<UFunction>(TargetDelegateProp->SignatureFunction, false);
 		}
 	}
-
-	PRAGMA_ENABLE_DEPRECATION_WARNINGS
 }
 
 #undef LOCTEXT_NAMESPACE
diff --git a/Engine/Source/Editor/BlueprintGraph/Private/K2Node_ComponentBoundEvent.cpp b/Engine/Source/Editor/BlueprintGraph/Private/K2Node_ComponentBoundEvent.cpp
--- a/Engine/Source/Editor/BlueprintGraph/Private/K2Node_ComponentBoundEvent.cpp
+++ b/Engine/Source/Editor/BlueprintGraph/Private/K2Node_ComponentBoundEvent.cpp
@@ -107,34 +107,21 @@ FString UK2Node_ComponentBoundEvent::GetDocumentationExcerptName() const
 
 void UK2Node_ComponentBoundEvent::Serialize(FArchive& Ar)
 {
-	// This deprecation warning suppression is due to the workaround. For more
-	// details read comment in K2Node_Event.h near EventSignatureName property
-	// definition.
-
-	PRAGMA_DISABLE_DEPRECATION_WARNINGS
-	// Workaround, so fields EventSignatureName and EventSignatureClass work
-	// until ObjectVersion.h lock is lifted.
-	if (Ar.IsSaving())
-	{
-		EventSignatureName = EventReference.GetMemberName();
-		EventSignatureClass = DelegateOwnerClass;
-		Super::Super::Serialize(Ar);
-		return;
-	}
-
 	Super::Serialize(Ar);
 
 	// Fix up legacy nodes that may not yet have a delegate pin
 	if(Ar.IsLoading())
 	{
-		DelegateOwnerClass = EventSignatureClass;
-		UMulticastDelegateProperty* TargetDelegateProp = GetTargetDelegateProperty();
-		if (TargetDelegateProp)
+		if(Ar.UE4Ver() < VER_UE4_K2NODE_EVENT_MEMBER_REFERENCE)
 		{
-			EventReference.SetExternalDelegateMember(TargetDelegateProp->SignatureFunction->GetFName());
+			DelegateOwnerClass = EventSignatureClass_DEPRECATED;
+			UMulticastDelegateProperty* TargetDelegateProp = GetTargetDelegateProperty();
+			if (TargetDelegateProp)
+			{
+				EventReference.SetExternalDelegateMember(TargetDelegateProp->SignatureFunction->GetFName());
+			}
 		}
 	}
-	PRAGMA_ENABLE_DEPRECATION_WARNINGS
 }
 
 
diff --git a/Engine/Source/Editor/BlueprintGraph/Private/K2Node_Event.cpp b/Engine/Source/Editor/BlueprintGraph/Private/K2Node_Event.cpp
--- a/Engine/Source/Editor/BlueprintGraph/Private/K2Node_Event.cpp
+++ b/Engine/Source/Editor/BlueprintGraph/Private/K2Node_Event.cpp
@@ -41,28 +41,6 @@ UK2Node_Event::UK2Node_Event(const FObjectInitializer& ObjectInitializer)
 
 void UK2Node_Event::Serialize(FArchive& Ar)
 {
-	// This deprecation warning suppression is due to the workaround. For more
-	// details read comment in K2Node_Event.h near EventSignatureName property
-	// definition.
-
-	PRAGMA_DISABLE_DEPRECATION_WARNINGS
-
-	// Workaround, so fields EventSignatureName and EventSignatureClass work
-	// until ObjectVersion.h lock is lifted.
-	if (Ar.IsSaving())
-	{
-		EventSignatureName = EventReference.GetMemberName();
-		EventSignatureClass = EventReference.GetMemberParentClass(this);
-	}
-
-	FName OldEventSignatureName;
-	TSubclassOf<class UObject> OldEventSignatureClass;
-	if (Ar.IsLoading())
-	{
-		OldEventSignatureName = EventSignatureName;
-		OldEventSignatureClass = EventSignatureClass;
-	}
-
 	Super::Serialize(Ar);
 
 	// Fix up legacy nodes that may not yet have a delegate pin
@@ -74,13 +52,11 @@ void UK2Node_Event::Serialize(FArchive& Ar)
 			CreatePin(EGPD_Output, K2Schema->PC_Delegate, TEXT(""), NULL, false, false, DelegateOutputName);
 		}
 
-		if (EventSignatureName != OldEventSignatureName || EventSignatureClass != OldEventSignatureClass)
+		if(Ar.UE4Ver() < VER_UE4_K2NODE_EVENT_MEMBER_REFERENCE)
 		{
-			EventReference.SetExternalMember(EventSignatureName, EventSignatureClass);
+			EventReference.SetExternalMember(EventSignatureName_DEPRECATED, EventSignatureClass_DEPRECATED);
 		}
 	}
-
-	PRAGMA_ENABLE_DEPRECATION_WARNINGS
 }
 
 FNodeHandlingFunctor* UK2Node_Event::CreateNodeHandler(FKismetCompilerContext& CompilerContext) const
diff --git a/Engine/Source/Editor/GameplayAbilitiesEditor/Classes/K2Node_GameplayCueEvent.h b/Engine/Source/Editor/GameplayAbilitiesEditor/Classes/K2Node_GameplayCueEvent.h
--- a/Engine/Source/Editor/GameplayAbilitiesEditor/Classes/K2Node_GameplayCueEvent.h
+++ b/Engine/Source/Editor/GameplayAbilitiesEditor/Classes/K2Node_GameplayCueEvent.h
@@ -9,6 +9,10 @@ class UK2Node_GameplayCueEvent : public UK2Node_Event
 {
 	GENERATED_UCLASS_BODY()
 
+	// UObject interface
+	virtual void Serialize(FArchive& Ar) override;
+	// End of UObject interface
+
 	// Begin UEdGraphNode interface
 	virtual FText GetTooltipText() const override;
 	virtual FText GetNodeTitle(ENodeTitleType::Type TitleType) const override;
diff --git a/Engine/Source/Editor/GameplayAbilitiesEditor/Private/K2Node_GameplayCueEvent.cpp b/Engine/Source/Editor/GameplayAbilitiesEditor/Private/K2Node_GameplayCueEvent.cpp
--- a/Engine/Source/Editor/GameplayAbilitiesEditor/Private/K2Node_GameplayCueEvent.cpp
+++ b/Engine/Source/Editor/GameplayAbilitiesEditor/Private/K2Node_GameplayCueEvent.cpp
@@ -19,6 +19,19 @@ UK2Node_GameplayCueEvent::UK2Node_GameplayCueEvent(const FObjectInitializer& Obj
 	EventReference.SetExternalMember(GAMEPLAYABILITIES_BlueprintCustomHandler, UGameplayCueInterface::StaticClass());
 }
 
+void UK2Node_GameplayCueEvent::Serialize(FArchive& Ar)
+{
+	Super::Serialize(Ar);
+
+	if(Ar.IsLoading())
+	{
+		if(Ar.UE4Ver() < VER_UE4_K2NODE_EVENT_MEMBER_REFERENCE && EventSignatureName_DEPRECATED.IsNone() && EventSignatureClass_DEPRECATED == nullptr)
+		{
+			EventReference.SetExternalMember(GAMEPLAYABILITIES_BlueprintCustomHandler, UGameplayCueInterface::StaticClass());
+		}
+	}
+}
+
 FText UK2Node_GameplayCueEvent::GetTooltipText() const
 {
 	return LOCTEXT("GameplayCueEvent_Tooltip", "Handle GameplayCue Event");
diff --git a/Engine/Source/Runtime/Core/Public/UObject/ObjectVersion.h b/Engine/Source/Runtime/Core/Public/UObject/ObjectVersion.h
--- a/Engine/Source/Runtime/Core/Public/UObject/ObjectVersion.h
+++ b/Engine/Source/Runtime/Core/Public/UObject/ObjectVersion.h
@@ -474,6 +474,8 @@ enum EUnrealEngineObjectUE4Version
 	VER_UE4_BLUEPRINT_GENERATED_CLASS_COMPONENT_TEMPLATES_PUBLIC,
 	// Split out creation method on ActorComponents to distinguish between native, instance, and simple or user construction script
 	VER_UE4_ACTOR_COMPONENT_CREATION_METHOD,
+	// K2Node_Event now uses FMemberReference for handling references
+	VER_UE4_K2NODE_EVENT_MEMBER_REFERENCE,
 
 	// -----<new versions can be added before this line>-------------------------------------------------
 	// - this needs to be the last line (see note below)
