Fixes for fallout from CL# 2460807 and the version bump. Input events should work properly again.

diff --git a/Engine/Source/Editor/BlueprintGraph/Classes/K2Node_InputAxisEvent.h b/Engine/Source/Editor/BlueprintGraph/Classes/K2Node_InputAxisEvent.h
--- a/Engine/Source/Editor/BlueprintGraph/Classes/K2Node_InputAxisEvent.h
+++ b/Engine/Source/Editor/BlueprintGraph/Classes/K2Node_InputAxisEvent.h
@@ -27,7 +27,8 @@ class UK2Node_InputAxisEvent : public UK2Node_Event
 
 	// Begin UObject interface
 	virtual void PostLoad() override;
-	// End UObject interface
+	virtual void Serialize(FArchive& Ar) override;
+	// End of UObject interface
 
 	// Begin EdGraphNode interface
 	virtual FText GetNodeTitle(ENodeTitleType::Type TitleType) const override;
diff --git a/Engine/Source/Editor/BlueprintGraph/Classes/K2Node_InputAxisKeyEvent.h b/Engine/Source/Editor/BlueprintGraph/Classes/K2Node_InputAxisKeyEvent.h
--- a/Engine/Source/Editor/BlueprintGraph/Classes/K2Node_InputAxisKeyEvent.h
+++ b/Engine/Source/Editor/BlueprintGraph/Classes/K2Node_InputAxisKeyEvent.h
@@ -24,6 +24,10 @@ class UK2Node_InputAxisKeyEvent : public UK2Node_Event
 	UPROPERTY(EditAnywhere, Category = "Input")
 	uint32 bOverrideParentBinding : 1;
 
+	// UObject interface
+	virtual void Serialize(FArchive& Ar) override;
+	// End of UObject interface
+
 	// Begin EdGraphNode interface
 	virtual FText GetNodeTitle(ENodeTitleType::Type TitleType) const override;
 	virtual FText GetTooltipText() const override;
diff --git a/Engine/Source/Editor/BlueprintGraph/Classes/K2Node_InputVectorAxisEvent.h b/Engine/Source/Editor/BlueprintGraph/Classes/K2Node_InputVectorAxisEvent.h
--- a/Engine/Source/Editor/BlueprintGraph/Classes/K2Node_InputVectorAxisEvent.h
+++ b/Engine/Source/Editor/BlueprintGraph/Classes/K2Node_InputVectorAxisEvent.h
@@ -9,6 +9,10 @@ class UK2Node_InputVectorAxisEvent : public UK2Node_InputAxisKeyEvent
 {
 	GENERATED_UCLASS_BODY()
 
+	// UObject interface
+	virtual void Serialize(FArchive& Ar) override;
+	// End of UObject interface
+
 	// Begin UK2Node interface
 	virtual void ValidateNodeDuringCompilation(class FCompilerResultsLog& MessageLog) const override;
 	virtual UClass* GetDynamicBindingClass() const override;
diff --git a/Engine/Source/Editor/BlueprintGraph/Private/K2Node_InputAxisEvent.cpp b/Engine/Source/Editor/BlueprintGraph/Private/K2Node_InputAxisEvent.cpp
--- a/Engine/Source/Editor/BlueprintGraph/Private/K2Node_InputAxisEvent.cpp
+++ b/Engine/Source/Editor/BlueprintGraph/Private/K2Node_InputAxisEvent.cpp
@@ -20,6 +20,19 @@ UK2Node_InputAxisEvent::UK2Node_InputAxisEvent(const FObjectInitializer& ObjectI
 	EventReference.SetExternalDelegateMember(FName(TEXT("InputAxisHandlerDynamicSignature__DelegateSignature")));
 }
 
+void UK2Node_InputAxisEvent::Serialize(FArchive& Ar)
+{
+	Super::Serialize(Ar);
+
+	if(Ar.IsLoading())
+	{
+		if(Ar.UE4Ver() < VER_UE4_K2NODE_EVENT_MEMBER_REFERENCE && EventSignatureName_DEPRECATED.IsNone() && EventSignatureClass_DEPRECATED == nullptr)
+		{
+			EventReference.SetExternalDelegateMember(TEXT("InputAxisHandlerDynamicSignature__DelegateSignature"));
+		}
+	}
+}
+
 void UK2Node_InputAxisEvent::PostLoad()
 {
 	Super::PostLoad();
diff --git a/Engine/Source/Editor/BlueprintGraph/Private/K2Node_InputAxisKeyEvent.cpp b/Engine/Source/Editor/BlueprintGraph/Private/K2Node_InputAxisKeyEvent.cpp
--- a/Engine/Source/Editor/BlueprintGraph/Private/K2Node_InputAxisKeyEvent.cpp
+++ b/Engine/Source/Editor/BlueprintGraph/Private/K2Node_InputAxisKeyEvent.cpp
@@ -22,6 +22,19 @@ UK2Node_InputAxisKeyEvent::UK2Node_InputAxisKeyEvent(const FObjectInitializer& O
 	EventReference.SetExternalDelegateMember(FName(TEXT("InputAxisHandlerDynamicSignature__DelegateSignature")));
 }
 
+void UK2Node_InputAxisKeyEvent::Serialize(FArchive& Ar)
+{
+	Super::Serialize(Ar);
+
+	if(Ar.IsLoading())
+	{
+		if(Ar.UE4Ver() < VER_UE4_K2NODE_EVENT_MEMBER_REFERENCE && EventSignatureName_DEPRECATED.IsNone() && EventSignatureClass_DEPRECATED == nullptr)
+		{
+			EventReference.SetExternalDelegateMember(FName(TEXT("InputAxisHandlerDynamicSignature__DelegateSignature")));
+		}
+	}
+}
+
 void UK2Node_InputAxisKeyEvent::Initialize(const FKey InAxisKey)
 {
 	AxisKey = InAxisKey;
diff --git a/Engine/Source/Editor/BlueprintGraph/Private/K2Node_InputVectorAxisEvent.cpp b/Engine/Source/Editor/BlueprintGraph/Private/K2Node_InputVectorAxisEvent.cpp
--- a/Engine/Source/Editor/BlueprintGraph/Private/K2Node_InputVectorAxisEvent.cpp
+++ b/Engine/Source/Editor/BlueprintGraph/Private/K2Node_InputVectorAxisEvent.cpp
@@ -13,6 +13,19 @@ UK2Node_InputVectorAxisEvent::UK2Node_InputVectorAxisEvent(const FObjectInitiali
 	EventReference.SetExternalDelegateMember(FName(TEXT("InputAxisHandlerDynamicSignature__DelegateSignature")));
 }
 
+void UK2Node_InputVectorAxisEvent::Serialize(FArchive& Ar)
+{
+	Super::Serialize(Ar);
+
+	if(Ar.IsLoading())
+	{
+		if(Ar.UE4Ver() < VER_UE4_K2NODE_EVENT_MEMBER_REFERENCE && EventSignatureName_DEPRECATED.IsNone() && EventSignatureClass_DEPRECATED == nullptr)
+		{
+			EventReference.SetExternalDelegateMember(FName(TEXT("InputAxisHandlerDynamicSignature__DelegateSignature")));
+		}
+	}
+}
+
 void UK2Node_InputVectorAxisEvent::ValidateNodeDuringCompilation(class FCompilerResultsLog& MessageLog) const
 {
 	// Skip AxisKeyEvent validation
