[UE-9032] IWCE - Hide UCS-added component instances from the Actor details view by default. Also prevent reattachment of SCS components to component instances added outside of the CDO.

change summary:
- added UBlueprintEditorSettings::bHideConstructionScriptComponentsInDetailsView to control UCS-added component visibility (default == true)
- modified SSCS_RowWidget::GetIntroducedInToolTipText() to include native/BP class name for component instances added outside of the CDO
- modified SSCS_RowWidget::OnDragEnter() to prevent reattachment to component instances added outside of the CDO
- modified SSCSEditor::Tick() to not count UCS-added components when checking consistency of component instances vs. number of nodes in the tree view when bHideConstructionScriptComponentsInDetailsView is true
- modified SSCSEditor::UpdateTree() to skip UCS-added components when bHideConstructionScriptComponentsInDetailsView is true

diff --git a/Engine/Source/Editor/BlueprintGraph/Public/BlueprintEditorSettings.h b/Engine/Source/Editor/BlueprintGraph/Public/BlueprintEditorSettings.h
--- a/Engine/Source/Editor/BlueprintGraph/Public/BlueprintEditorSettings.h
+++ b/Engine/Source/Editor/BlueprintGraph/Public/BlueprintEditorSettings.h
@@ -74,6 +74,10 @@ public:
 	/** If set will spawn default nodes in new Blueprints */
 	UPROPERTY(EditAnywhere, config, Category=Workflow)
 	bool bSpawnDefaultBlueprintNodes;
+
+	/** If set will exclude components added in a Blueprint class Construction Script from the component details view */
+	UPROPERTY(EditAnywhere, config, Category=Workflow)
+	bool bHideConstructionScriptComponentsInDetailsView;
 // Compiler Settings
 public:
 	/** Determines when to save Blueprints post-compile */
