Converting a cast node to/from pure will add an undo/redo transaction.

diff --git a/Engine/Source/Editor/BlueprintGraph/Private/K2Node_DynamicCast.cpp b/Engine/Source/Editor/BlueprintGraph/Private/K2Node_DynamicCast.cpp
--- a/Engine/Source/Editor/BlueprintGraph/Private/K2Node_DynamicCast.cpp
+++ b/Engine/Source/Editor/BlueprintGraph/Private/K2Node_DynamicCast.cpp
@@ -5,6 +5,7 @@
 #include "DynamicCastHandler.h"
 #include "EditorCategoryUtils.h"
 #include "BlueprintEditorSettings.h"
+#include "ScopedTransaction.h"
 
 #define LOCTEXT_NAMESPACE "K2Node_DynamicCast"
 
@@ -226,6 +227,18 @@ void UK2Node_DynamicCast::SetPurity(bool bNewPurity)
 
 void UK2Node_DynamicCast::TogglePurity()
 {
+	FText TransactionTitle;
+	if(bIsPureCast)
+	{
+		TransactionTitle = LOCTEXT("TogglePure", "Convert to Pure Cast");
+	}
+	else
+	{
+		TransactionTitle = LOCTEXT("ToggleImpure", "Convert to Impure Cast");
+	}
+	const FScopedTransaction Transaction( TransactionTitle );
+	Modify();
+
 	SetPurity(!bIsPureCast);
 }
 
