[UE-11121] Hide incompatible macro actions that prevent child Blueprints from saving.

change summary:
- modified UK2Node_MacroInstance::CanPasteHere() to exclude "local" macro instance nodes if the macro is not part of a macro library Blueprint and does not match the target Blueprint
- added IsIncompatibleMacroInstance() to the Blueprint action filter list as a "fast path" to filtering out local macro instances

diff --git a/Engine/Source/Editor/BlueprintGraph/Private/BlueprintActionFilter.cpp b/Engine/Source/Editor/BlueprintGraph/Private/BlueprintActionFilter.cpp
--- a/Engine/Source/Editor/BlueprintGraph/Private/BlueprintActionFilter.cpp
+++ b/Engine/Source/Editor/BlueprintGraph/Private/BlueprintActionFilter.cpp
@@ -328,6 +328,15 @@ namespace BlueprintActionFilterImpl
 	 * @return 
 	 */
 	static bool IsExtraneousInterfaceCall(FBlueprintActionFilter const& Filter, FBlueprintActionInfo& BlueprintAction);
+
+	/**
+	 * Rejection test that checks if a macro instance is incompatible with the current graph context
+	 * 
+	 * @param  Filter			Holds the graph context for this test.
+	 * @param  BlueprintAction	The action you wish to query.
+	 * @return true if the macro instance is incompatible with the current graph context
+	 */
+	static bool IsIncompatibleMacroInstance(FBlueprintActionFilter const& Filter, FBlueprintActionInfo& BlueprintAction);
 };
 
 //------------------------------------------------------------------------------
@@ -1368,6 +1377,32 @@ static bool BlueprintActionFilterImpl::IsExtraneousInterfaceCall(FBlueprintActio
 	return bIsFilteredOut;
 }
 
+//------------------------------------------------------------------------------
+static bool BlueprintActionFilterImpl::IsIncompatibleMacroInstance(FBlueprintActionFilter const& Filter, FBlueprintActionInfo& BlueprintAction)
+{
+	bool bIsFilteredOut = false;
+
+	if(BlueprintAction.GetNodeClass()->IsChildOf<UK2Node_MacroInstance>())
+	{
+		if(const UBlueprint* MacroBP = Cast<const UBlueprint>(BlueprintAction.GetActionOwner()))
+		{
+			check(MacroBP->ParentClass != nullptr);
+
+			for(auto BlueprintIt = Filter.Context.Blueprints.CreateConstIterator(); BlueprintIt && !bIsFilteredOut; ++BlueprintIt)
+			{
+				const UBlueprint* Blueprint = *BlueprintIt;
+				check(Blueprint != nullptr && Blueprint->ParentClass != nullptr)
+
+				bIsFilteredOut = (Blueprint != MacroBP) && (MacroBP->BlueprintType != BPTYPE_MacroLibrary || !Blueprint->ParentClass->IsChildOf(MacroBP->ParentClass));
+			}
+
+			// Note: The rest is handled by IsNodeTemplateSelfFiltered() - the check above is a "fast path" in that we don't have to instance the node template (see UK2Node_MacroInstance::IsActionFilteredOut())
+		}
+	}
+
+	return bIsFilteredOut;
+}
+
 /*******************************************************************************
  * FBlueprintActionInfo
  ******************************************************************************/
@@ -1555,6 +1590,7 @@ FBlueprintActionFilter::FBlueprintActionFilter(uint32 Flags/*= 0x00*/)
 	AddRejectionTest(FRejectionTestDelegate::CreateStatic(IsIncompatibleWithGraphType));
 	AddRejectionTest(FRejectionTestDelegate::CreateStatic(IsSchemaIncompatible));
 	AddRejectionTest(FRejectionTestDelegate::CreateStatic(IsExtraneousInterfaceCall));
+	AddRejectionTest(FRejectionTestDelegate::CreateStatic(IsIncompatibleMacroInstance));
 
 	if (!(Flags & BPFILTER_PermitDeprecated))
 	{
diff --git a/Engine/Source/Editor/BlueprintGraph/Private/K2Node_MacroInstance.cpp b/Engine/Source/Editor/BlueprintGraph/Private/K2Node_MacroInstance.cpp
--- a/Engine/Source/Editor/BlueprintGraph/Private/K2Node_MacroInstance.cpp
+++ b/Engine/Source/Editor/BlueprintGraph/Private/K2Node_MacroInstance.cpp
@@ -426,14 +426,9 @@ bool UK2Node_MacroInstance::CanPasteHere(const UEdGraph* TargetGraph) const
 
 	if ((MacroBlueprint != nullptr) && (TargetBlueprint != nullptr))
 	{
-		UClass* TargetClass = (TargetBlueprint->GeneratedClass != nullptr) ? TargetBlueprint->GeneratedClass : TargetBlueprint->ParentClass;
-		UClass* MacroBaseClass = ((MacroBlueprint->BlueprintType == BPTYPE_MacroLibrary) || (MacroBlueprint->GeneratedClass == nullptr)) ? 
-			MacroBlueprint->ParentClass : MacroBlueprint->GeneratedClass;
-		
-		if ((MacroBaseClass != nullptr) && (TargetClass != nullptr))
-		{
-			bCanPaste = TargetClass->IsChildOf(MacroBaseClass);
-		}
+		// Only allow "local" macro instances or instances from a macro library blueprint with the same parent class
+		check(MacroBlueprint->ParentClass != nullptr && TargetBlueprint->ParentClass != nullptr);
+		bCanPaste = (MacroBlueprint == TargetBlueprint) || (MacroBlueprint->BlueprintType == BPTYPE_MacroLibrary && TargetBlueprint->ParentClass->IsChildOf(MacroBlueprint->ParentClass));
 	}
 
 	// Macro Instances are not allowed in it's own graph
