Removed transient from UK2Node_ActorBoundEvent::DelegateOwnerClass, was losing the information during serialization.

diff --git a/Engine/Source/Editor/BlueprintGraph/Classes/K2Node_ActorBoundEvent.h b/Engine/Source/Editor/BlueprintGraph/Classes/K2Node_ActorBoundEvent.h
--- a/Engine/Source/Editor/BlueprintGraph/Classes/K2Node_ActorBoundEvent.h
+++ b/Engine/Source/Editor/BlueprintGraph/Classes/K2Node_ActorBoundEvent.h
@@ -16,7 +16,7 @@ class UK2Node_ActorBoundEvent : public UK2Node_Event
 	FName DelegatePropertyName;
 
 	/** Delegate property's owner class that this event is associated with */
-	UPROPERTY(Transient)
+	UPROPERTY()
 	UClass* DelegateOwnerClass;
 
 	/** The event that this event is bound to */
