Bug in FKCHandler_PureAssignmentStatement::RegisterNets
https://udn.unrealengine.com/questions/237845/bug-in-fkchandler-pureassignmentstatementregistern.html

diff --git a/Engine/Source/Editor/BlueprintGraph/Private/K2Node_PureAssignmentStatement.cpp b/Engine/Source/Editor/BlueprintGraph/Private/K2Node_PureAssignmentStatement.cpp
--- a/Engine/Source/Editor/BlueprintGraph/Private/K2Node_PureAssignmentStatement.cpp
+++ b/Engine/Source/Editor/BlueprintGraph/Private/K2Node_PureAssignmentStatement.cpp
@@ -40,13 +40,14 @@ public:
 			}
 
 			// At the moment, a term for variable should be already registered
-			auto VariableTerm = Context.NetMap.Find(FEdGraphUtilities::GetNetFromPin(VariablePin));
-			if (!VariableTerm || !*VariableTerm)
+			auto VariableTermPtr = Context.NetMap.Find(FEdGraphUtilities::GetNetFromPin(VariablePin));
+			if (!VariableTermPtr || !*VariableTermPtr)
 			{
 				CompilerContext.MessageLog.Error(*LOCTEXT("NoVarriableTerm_Error", "ICE: no variable term found in @@").ToString(), Node);
 				return;
 			}
-			Context.NetMap.Add(OutputPin, *VariableTerm);
+			FBPTerminal* VariableTerm = *VariableTermPtr; // we must take a copy here because Add takes a reference and the map might be resized
+			Context.NetMap.Add(OutputPin, VariableTerm);
 		}
 
 		auto ValuePin = PureAssignmentNode->GetValuePin();
