Changing a custom event's replication will immediately update the node's title.

Added a function to EdGraphNode to help force refresh of titles. This needs to be overriden on a per node basis.

diff --git a/Engine/Source/Editor/BlueprintGraph/Classes/K2Node_CustomEvent.h b/Engine/Source/Editor/BlueprintGraph/Classes/K2Node_CustomEvent.h
--- a/Engine/Source/Editor/BlueprintGraph/Classes/K2Node_CustomEvent.h
+++ b/Engine/Source/Editor/BlueprintGraph/Classes/K2Node_CustomEvent.h
@@ -17,9 +17,14 @@ class UK2Node_CustomEvent : public UK2Node_Event
 
 	virtual bool IsEditable() const override;
 
+	// UObject interface
+	virtual void Serialize(FArchive& Ar) override;
+	// End of UObject interface
+
 	// Begin UEdGraphNode interface
 	virtual void ReconstructNode() override;
 	virtual FText GetNodeTitle(ENodeTitleType::Type TitleType) const override;
+	virtual void RefreshNodeTitle() override;
 	virtual void OnRenameNode(const FString& NewName) override;
 	virtual TSharedPtr<class INameValidatorInterface> MakeNameValidator() const override;
 	virtual FText GetTooltipText() const override;
diff --git a/Engine/Source/Editor/BlueprintGraph/Private/K2Node_CustomEvent.cpp b/Engine/Source/Editor/BlueprintGraph/Private/K2Node_CustomEvent.cpp
--- a/Engine/Source/Editor/BlueprintGraph/Private/K2Node_CustomEvent.cpp
+++ b/Engine/Source/Editor/BlueprintGraph/Private/K2Node_CustomEvent.cpp
@@ -106,6 +106,16 @@ UK2Node_CustomEvent::UK2Node_CustomEvent(const FObjectInitializer& ObjectInitial
 	bCallInEditor = false;
 }
 
+void UK2Node_CustomEvent::Serialize(FArchive& Ar)
+{
+	Super::Serialize(Ar);
+
+	if (Ar.IsLoading())
+	{
+		CachedNodeTitle.MarkDirty();
+	}
+}
+
 FText UK2Node_CustomEvent::GetNodeTitle(ENodeTitleType::Type TitleType) const
 {
 	if (TitleType != ENodeTitleType::FullTitle)
@@ -127,6 +137,11 @@ FText UK2Node_CustomEvent::GetNodeTitle(ENodeTitleType::Type TitleType) const
 	return CachedNodeTitle;
 }
 
+void UK2Node_CustomEvent::RefreshNodeTitle()
+{
+	CachedNodeTitle.MarkDirty();
+}
+
 UEdGraphPin* UK2Node_CustomEvent::CreatePinFromUserDefinition(const TSharedPtr<FUserPinInfo> NewPinInfo)
 {
 	UEdGraphPin* NewPin = CreatePin(
@@ -269,6 +284,8 @@ void UK2Node_CustomEvent::GetMenuActions(FBlueprintActionDatabaseRegistrar& Acti
 
 void UK2Node_CustomEvent::ReconstructNode()
 {
+	CachedNodeTitle.MarkDirty();
+
 	const UEdGraphPin* DelegateOutPin = FindPin(DelegateOutputName);
 
 	const UK2Node_BaseMCDelegate* OtherNode = (DelegateOutPin && DelegateOutPin->LinkedTo.Num() && DelegateOutPin->LinkedTo[0]) ?
diff --git a/Engine/Source/Editor/Kismet/Private/BlueprintDetailsCustomization.cpp b/Engine/Source/Editor/Kismet/Private/BlueprintDetailsCustomization.cpp
--- a/Engine/Source/Editor/Kismet/Private/BlueprintDetailsCustomization.cpp
+++ b/Engine/Source/Editor/Kismet/Private/BlueprintDetailsCustomization.cpp
@@ -2663,6 +2663,9 @@ void FBlueprintGraphActionDetails::SetNetFlags( TWeakObjectPtr<UK2Node_EditableP
 		// Clear all net flags before setting
 		if( FlagsToSet != FlagsToClear )
 		{
+			const FScopedTransaction Transaction( LOCTEXT("GraphSetNetFlags", "Change Replication") );
+			FunctionEntryNode->Modify();
+
 			bool bBlueprintModified = false;
 
 			if (UK2Node_FunctionEntry* TypedEntryNode = Cast<UK2Node_FunctionEntry>(FunctionEntryNode.Get()))
@@ -2677,6 +2680,10 @@ void FBlueprintGraphActionDetails::SetNetFlags( TWeakObjectPtr<UK2Node_EditableP
 				CustomEventNode->FunctionFlags |= FlagsToSet;
 				bBlueprintModified = true;
 			}
+
+			// The node title needs to refresh
+			FunctionEntryNode->RefreshNodeTitle();
+
 			if( bBlueprintModified )
 			{
 				FBlueprintEditorUtils::MarkBlueprintAsStructurallyModified( FunctionEntryNode->GetBlueprint() );
diff --git a/Engine/Source/Runtime/Engine/Classes/EdGraph/EdGraphNode.h b/Engine/Source/Runtime/Engine/Classes/EdGraph/EdGraphNode.h
--- a/Engine/Source/Runtime/Engine/Classes/EdGraph/EdGraphNode.h
+++ b/Engine/Source/Runtime/Engine/Classes/EdGraph/EdGraphNode.h
@@ -286,6 +286,9 @@ public:
 	/** Gets the name of this node, shown in title bar */
 	virtual FText GetNodeTitle(ENodeTitleType::Type TitleType) const;
 	
+	/** Forces a refresh of the node title if it's cached */
+	virtual void RefreshNodeTitle() {};
+
 	/** 
 	 * Gets the draw color of a node's title bar
 	 */
