[UE-11185] Fix inability to save a Blueprint Function Library asset containing a function that references another function in the same library after it's been loaded.

notes:
- this issue was introduced (by me) in CL# 2368059

change summary:
- modified UK2Node_CallFunction::CreatePinsForFunctionCall() to revert a previous change that incorrectly assigned the BPGC's CDO to the self pin's DefaultObject property for static functions owned by the BPGC. while that works, it creates a problem in that the CDO's class will change during compile-on-load, leaving the reference invalid and in the transient package after the load has completed.
- modified UK2Node_CallFunction::PostReconstructNode() to ensure that the self pin's DefaultObject reference is cleared if the function is owned by a BPGC in the current hierarchy

diff --git a/Engine/Source/Editor/BlueprintGraph/Private/K2Node_CallFunction.cpp b/Engine/Source/Editor/BlueprintGraph/Private/K2Node_CallFunction.cpp
--- a/Engine/Source/Editor/BlueprintGraph/Private/K2Node_CallFunction.cpp
+++ b/Engine/Source/Editor/BlueprintGraph/Private/K2Node_CallFunction.cpp
@@ -790,7 +790,7 @@ bool UK2Node_CallFunction::CreatePinsForFunctionCall(const UFunction* Function)
 	ensure(BP);
 	if (BP != nullptr)
 	{
-		const bool bIsFunctionCompatibleWithSelf = !K2Schema->IsStaticFunctionGraph(GetGraph()) && BP->SkeletonGeneratedClass->IsChildOf(FunctionOwnerClass);
+		const bool bIsFunctionCompatibleWithSelf = BP->SkeletonGeneratedClass->IsChildOf(FunctionOwnerClass);
 
 		if (bIsStaticFunc)
 		{
@@ -900,6 +900,11 @@ void UK2Node_CallFunction::PostReconstructNode()
 			{
 				SelfPin->DefaultObject = FunctionOwnerClass->GetDefaultObject();
 			}
+			else
+			{
+				// In case a non-NULL reference was previously serialized on load, ensure that it's set to NULL here to match what a new node's self pin would be initialized as (see CreatePinsForFunctionCall).
+				SelfPin->DefaultObject = nullptr;
+			}
 		}
 	}
 
