UE-10151 Performing a Hot Reload after changing a custom component code class can make the component unusable in Blueprints, or crash the Editor.

Native classes based on SceneComponent are not considered valid after hot reload.

diff --git a/Engine/Source/Editor/BlueprintGraph/Private/BlueprintActionDatabase.cpp b/Engine/Source/Editor/BlueprintGraph/Private/BlueprintActionDatabase.cpp
--- a/Engine/Source/Editor/BlueprintGraph/Private/BlueprintActionDatabase.cpp
+++ b/Engine/Source/Editor/BlueprintGraph/Private/BlueprintActionDatabase.cpp
@@ -44,6 +44,8 @@
 #include "EdGraph/EdGraphNode_Documentation.h"
 #include "K2Node_FunctionEntry.h"
 
+#include "HotReloadInterface.h"
+
 #define LOCTEXT_NAMESPACE "BlueprintActionDatabase"
 
 /*******************************************************************************
@@ -461,6 +463,11 @@ namespace BlueprintActionDatabaseImpl
 	 */
 	static bool IsObjectValidForDatabase(UObject const* Object);
 
+	/**
+	 * Refreshes database after project was hot-reloaded.
+	 */
+	static void OnProjectHotReloaded(bool bWasTriggeredAutomatically);
+
 	/** 
 	 * Assets that we cleared from the database (to remove references, and make 
 	 * way for a delete), but in-case the class wasn't deleted we need them 
@@ -472,6 +479,13 @@ namespace BlueprintActionDatabaseImpl
 	bool bIsInitializing = false;
 }
 
+//------------------------------------------------------------------------------
+static void BlueprintActionDatabaseImpl::OnProjectHotReloaded(bool bWasTriggeredAutomatically)
+{
+	FBlueprintActionDatabase& ActionDatabase = FBlueprintActionDatabase::Get();
+	ActionDatabase.RefreshAll();
+}
+
 //------------------------------------------------------------------------------
 static bool BlueprintActionDatabaseImpl::IsPropertyBlueprintVisible(UProperty const* const Property)
 {
@@ -932,6 +946,9 @@ FBlueprintActionDatabase::FBlueprintActionDatabase()
 	FKismetEditorUtilities::OnBlueprintUnloaded.AddStatic(&BlueprintActionDatabaseImpl::OnBlueprintUnloaded);
 
 	GEngine->OnWorldDestroyed().AddStatic(&BlueprintActionDatabaseImpl::OnWorldDestroyed);
+
+	IHotReloadInterface& HotReloadSupport = FModuleManager::LoadModuleChecked<IHotReloadInterface>("HotReload");
+	HotReloadSupport.OnHotReload().AddStatic(&BlueprintActionDatabaseImpl::OnProjectHotReloaded);
 }
 
 //------------------------------------------------------------------------------
diff --git a/Engine/Source/Editor/UnrealEd/Private/ComponentTypeRegistry.cpp b/Engine/Source/Editor/UnrealEd/Private/ComponentTypeRegistry.cpp
--- a/Engine/Source/Editor/UnrealEd/Private/ComponentTypeRegistry.cpp
+++ b/Engine/Source/Editor/UnrealEd/Private/ComponentTypeRegistry.cpp
@@ -214,7 +214,10 @@ void FComponentTypeRegistryData::ForceRefreshComponentList()
 		if (Class->IsChildOf(UActorComponent::StaticClass()))
 		{
 			InMemoryClasses.Push(Class->GetFName());
-			if (!Class->HasAnyClassFlags(CLASS_Abstract) && Class->HasMetaData(FBlueprintMetadata::MD_BlueprintSpawnableComponent) && !FKismetEditorUtilities::IsClassABlueprintSkeleton(Class)) //@TODO: Fold this logic together with the one in UEdGraphSchema_K2::GetAddComponentClasses
+
+			bool const bOutOfDateClass = Class->HasAnyClassFlags(CLASS_NewerVersionExists);
+			bool const bBlueprintSkeletonClass = FKismetEditorUtilities::IsClassABlueprintSkeleton(Class);
+			if (!Class->HasAnyClassFlags(CLASS_Abstract) && Class->HasMetaData(FBlueprintMetadata::MD_BlueprintSpawnableComponent) && !bOutOfDateClass && !bBlueprintSkeletonClass) //@TODO: Fold this logic together with the one in UEdGraphSchema_K2::GetAddComponentClasses
 			{
 				TArray<FString> ClassGroupNames;
 				Class->GetClassGroupNames(ClassGroupNames);
@@ -234,9 +237,8 @@ void FComponentTypeRegistryData::ForceRefreshComponentList()
 					SortedClassList.Add(NewEntry);
 				}
 			}
-
-			bool const bOutOfDateClass = Class->HasAnyClassFlags(CLASS_NewerVersionExists);
-			if (!bOutOfDateClass && !FKismetEditorUtilities::IsClassABlueprintSkeleton(Class))
+			
+			if (!bOutOfDateClass && !bBlueprintSkeletonClass)
 			{
 				FComponentTypeEntry Entry = { Class->GetName(), FString(), Class };
 				ComponentTypeList.Add(Entry);
