[UE-11377] MakeArray now checks to see if there are subpins and doesn't mutate the type on subpins, only makes sure parent pin is still the correct type.

diff --git a/Engine/Source/Editor/BlueprintGraph/Private/K2Node_MakeArray.cpp b/Engine/Source/Editor/BlueprintGraph/Private/K2Node_MakeArray.cpp
--- a/Engine/Source/Editor/BlueprintGraph/Private/K2Node_MakeArray.cpp
+++ b/Engine/Source/Editor/BlueprintGraph/Private/K2Node_MakeArray.cpp
@@ -217,10 +217,28 @@ void UK2Node_MakeArray::PropagatePinType()
 
 			if (CurrentPin != OutputPin)
 			{
-				bWantRefresh = true;
-				CurrentPin->PinType.PinCategory = OutputPin->PinType.PinCategory;
-				CurrentPin->PinType.PinSubCategory = OutputPin->PinType.PinSubCategory;
-				CurrentPin->PinType.PinSubCategoryObject = OutputPin->PinType.PinSubCategoryObject;
+				// sub pins will be updated by their parent pin, so if we have a parent pin just do nothing
+				if (CurrentPin->ParentPin == nullptr)
+				{
+					bWantRefresh = true;
+
+					// if we've reset to wild card or the parentpin no longer matches we need to collapse the split pin(s)
+					// otherwise everything should be OK:
+					if (CurrentPin->SubPins.Num() != 0 &&
+						(	CurrentPin->PinType.PinCategory != OutputPin->PinType.PinCategory ||
+							CurrentPin->PinType.PinSubCategory != OutputPin->PinType.PinSubCategory ||
+							CurrentPin->PinType.PinSubCategoryObject != OutputPin->PinType.PinSubCategoryObject )
+						)
+					{
+						// this is a little dicey, but should be fine.. relies on the fact that RecombinePin will only remove pins that
+						// are placed after CurrentPin in the Pins member:
+						Schema->RecombinePin(CurrentPin->SubPins[0]);
+					}
+
+					CurrentPin->PinType.PinCategory = OutputPin->PinType.PinCategory;
+					CurrentPin->PinType.PinSubCategory = OutputPin->PinType.PinSubCategory;
+					CurrentPin->PinType.PinSubCategoryObject = OutputPin->PinType.PinSubCategoryObject;
+				}
 
 				if (CurrentPin->PinType.PinCategory == UEdGraphSchema_K2::PC_Wildcard)
 				{
