PR #863: Add UK2Node_CallParentFunction node on Blueprint function override (Contributed by [REDACTED_PERSON])

https://github.com/EpicGames/UnrealEngine/pull/863

diff --git a/Engine/Source/Editor/BlueprintGraph/Classes/EdGraphSchema_K2.h b/Engine/Source/Editor/BlueprintGraph/Classes/EdGraphSchema_K2.h
--- a/Engine/Source/Editor/BlueprintGraph/Classes/EdGraphSchema_K2.h
+++ b/Engine/Source/Editor/BlueprintGraph/Classes/EdGraphSchema_K2.h
@@ -969,6 +969,11 @@ public:
 	 */
 	void CombineTwoPinNetsAndRemoveOldPins(UEdGraphPin* InPinA, UEdGraphPin* InPinB) const;
 
+	/**
+	 * Make links from all data pins from InOutputNode output to InInputNode input.
+	 */
+	void LinkDataPinFromOutputToInput(UEdGraphNode* InOutputNode, UEdGraphNode* InInputNode) const;
+
 	/** Function that returns _all_ nodes we could place */
 	static void GetAllActions(struct FBlueprintPaletteListBuilder& PaletteBuilder);
 
diff --git a/Engine/Source/Editor/BlueprintGraph/Private/EdGraphSchema_K2.cpp b/Engine/Source/Editor/BlueprintGraph/Private/EdGraphSchema_K2.cpp
--- a/Engine/Source/Editor/BlueprintGraph/Private/EdGraphSchema_K2.cpp
+++ b/Engine/Source/Editor/BlueprintGraph/Private/EdGraphSchema_K2.cpp
@@ -2850,7 +2850,20 @@ void UEdGraphSchema_K2::CreateMacroGraphTerminators(UEdGraph& Graph, UClass* Cla
 		SetNodeMetaData(ExitNode, FNodeMetadata::DefaultGraphNode);
 	}
 }
-		
+
+void UEdGraphSchema_K2::LinkDataPinFromOutputToInput(UEdGraphNode* InOutputNode, UEdGraphNode* InInputNode) const
+{
+	for (auto PinIter = InOutputNode->Pins.CreateIterator(); PinIter; ++PinIter)
+	{
+		UEdGraphPin* const OutputPin = *PinIter;
+		if ((OutputPin->Direction == EGPD_Output) && (!IsExecPin(*OutputPin)))
+		{
+			UEdGraphPin* const InputPin = InInputNode->FindPinChecked(OutputPin->PinName);
+			OutputPin->MakeLinkTo(InputPin);
+		}
+	}
+}
+
 void UEdGraphSchema_K2::CreateFunctionGraphTerminators(UEdGraph& Graph, UClass* Class) const
 {
 	const FName GraphName = Graph.GetFName();
@@ -2873,6 +2886,29 @@ void UEdGraphSchema_K2::CreateFunctionGraphTerminators(UEdGraph& Graph, UClass*
 		// Add modifier flags from the declaration
 		EntryNode->ExtraFlags |= InterfaceToImplement->FunctionFlags & (FUNC_Const | FUNC_Static | FUNC_BlueprintPure);
 
+		UK2Node* NextNode = EntryNode;
+		bool bHasParentNode = false;
+		// Create node for call parent function
+		if (((Class->GetClassFlags() & CLASS_Interface) == 0)  &&
+			(InterfaceToImplement->FunctionFlags & FUNC_BlueprintCallable))
+		{
+			FGraphNodeCreator<UK2Node_CallParentFunction> FunctionParentCreator(Graph);
+			UK2Node_CallParentFunction* ParentNode = FunctionParentCreator.CreateNode();
+			ParentNode->SetFromFunction(InterfaceToImplement);
+			ParentNode->NodePosX = EntryNode->NodePosX + EntryNode->NodeWidth + 256;
+			ParentNode->NodePosY = EntryNode->NodePosY;
+			FunctionParentCreator.Finalize();
+
+			UEdGraphPin* EntryNodeExec = FindExecutionPin(*EntryNode, EGPD_Output);
+			UEdGraphPin* ParentNodeExec = FindExecutionPin(*ParentNode, EGPD_Input);
+			EntryNodeExec->MakeLinkTo(ParentNodeExec);
+
+			LinkDataPinFromOutputToInput(EntryNode, ParentNode);
+
+			NextNode = ParentNode;
+			bHasParentNode = true;
+		}
+
 		// See if any function params are marked as out
 		bool bHasOutParam =  false;
 		for( TFieldIterator<UProperty> It(InterfaceToImplement); It && (It->PropertyFlags & CPF_Parm); ++It )
@@ -2890,15 +2926,20 @@ void UEdGraphSchema_K2::CreateFunctionGraphTerminators(UEdGraph& Graph, UClass*
 			UK2Node_FunctionResult* ReturnNode = NodeCreator.CreateNode();
 			ReturnNode->SignatureClass = Class;
 			ReturnNode->SignatureName = GraphName;
-			ReturnNode->NodePosX = EntryNode->NodePosX + EntryNode->NodeWidth + 256;
+			ReturnNode->NodePosX = NextNode->NodePosX + NextNode->NodeWidth + 256;
 			ReturnNode->NodePosY = EntryNode->NodePosY;
 			NodeCreator.Finalize();
 			SetNodeMetaData(ReturnNode, FNodeMetadata::DefaultGraphNode);
 
 			// Auto-connect the pins for entry and exit, so that by default the signature is properly generated
-			UEdGraphPin* EntryNodeExec = FindExecutionPin(*EntryNode, EGPD_Output);
+			UEdGraphPin* EntryNodeExec = FindExecutionPin(*NextNode, EGPD_Output);
 			UEdGraphPin* ResultNodeExec = FindExecutionPin(*ReturnNode, EGPD_Input);
 			EntryNodeExec->MakeLinkTo(ResultNodeExec);
+
+			if (bHasParentNode)
+			{
+				LinkDataPinFromOutputToInput(NextNode, ReturnNode);
+			}
 		}
 	}
 }
