Move member reference out of K2Node and make it usable elsewhere

diff --git a/Engine/Plugins/Experimental/LiveEditor/Source/LiveEditor/Private/K2Node_LiveEditObject.cpp b/Engine/Plugins/Experimental/LiveEditor/Source/LiveEditor/Private/K2Node_LiveEditObject.cpp
--- a/Engine/Plugins/Experimental/LiveEditor/Source/LiveEditor/Private/K2Node_LiveEditObject.cpp
+++ b/Engine/Plugins/Experimental/LiveEditor/Source/LiveEditor/Private/K2Node_LiveEditObject.cpp
@@ -630,7 +630,7 @@ UFunction *UK2Node_LiveEditObject::GetEventMIDISignature() const
 	UBlueprint::GetGuidFromClassByFieldName<UFunction>(ULiveEditorBroadcaster::StaticClass(), TEXT("OnEventMIDI"), DelegateGuid);
 	ReferenceToUse.SetDirect( TEXT("OnEventMIDI"), DelegateGuid, ULiveEditorBroadcaster::StaticClass(), false );
 
-	UMulticastDelegateProperty* DelegateProperty = ReferenceToUse.ResolveMember<UMulticastDelegateProperty>(this);
+	UMulticastDelegateProperty* DelegateProperty = ReferenceToUse.ResolveMember<UMulticastDelegateProperty>(GetBlueprintClassFromNode());
 	if (DelegateProperty != NULL)
 	{
 		return DelegateProperty->SignatureFunction;
diff --git a/Engine/Source/Editor/BlueprintGraph/Classes/K2Node.h b/Engine/Source/Editor/BlueprintGraph/Classes/K2Node.h
--- a/Engine/Source/Editor/BlueprintGraph/Classes/K2Node.h
+++ b/Engine/Source/Editor/BlueprintGraph/Classes/K2Node.h
@@ -10,50 +10,6 @@ class UBlueprintNodeSpawner;
 class FBlueprintActionDatabaseRegistrar;
 class UDynamicBlueprintBinding;
 
-/** Helper struct to allow us to redirect properties and functions through renames and additionally between classes if necessary */
-struct FFieldRemapInfo
-{
-	/** The new name of the field after being renamed */
-	FName FieldName;
-
-	/** The new name of the field's outer class if different from its original location, or NAME_None if it hasn't moved */
-	FName FieldClass;
-
-	bool operator==(const FFieldRemapInfo& Other) const
-	{
-		return FieldName == Other.FieldName && FieldClass == Other.FieldClass;
-	}
-
-	friend uint32 GetTypeHash(const FFieldRemapInfo& RemapInfo)
-	{
-		return GetTypeHash(RemapInfo.FieldName) + GetTypeHash(RemapInfo.FieldClass) * 23;
-	}
-
-	FFieldRemapInfo()
-		: FieldName(NAME_None)
-		, FieldClass(NAME_None)
-	{
-	}
-};
-
-/** Helper struct to allow us to redirect pin name for node class */
-struct FParamRemapInfo
-{
-	bool	bCustomValueMapping;
-	FName	OldParam;
-	FName	NewParam;
-	FName	NodeTitle;
-	TMap<FString, FString> ParamValueMap;
-
-	// constructor
-	FParamRemapInfo()
-		: OldParam(NAME_None)
-		, NewParam(NAME_None)
-		, NodeTitle(NAME_None)
-	{
-	}
-};
-
 USTRUCT()
 struct FOptionalPinFromProperty
 {
@@ -260,16 +216,6 @@ class UK2Node : public UEdGraphNode
 	// Can this node be created under the specified schema
 	BLUEPRINTGRAPH_API virtual bool CanCreateUnderSpecifiedSchema(const UEdGraphSchema* DesiredSchema) const override;
 
-	/**
-	 * Searches the field redirect map for the specified named field in the scope, and returns the remapped field if found
-	 *
-	 * @param	InitialScope	The scope the field was initially defined in.  The function will search up into parent scopes to attempt to find remappings
-	 * @param	InitialName		The name of the field to attempt to find a redirector for
-	 * @param	bInitialScopeMustBeOwnerOfField		if true the InitialScope must be Child of the field's owner
-	 * @return	The remapped field, if one exists
-	 */
-	BLUEPRINTGRAPH_API static UField* FindRemappedField(UClass* InitialScope, FName InitialName, bool bInitialScopeMustBeOwnerOfField = false);
-
 	// Renames an existing pin on the node.
 	BLUEPRINTGRAPH_API virtual ERenamePinResult RenameUserDefinedPin(const FString& OldName, const FString& NewName, bool bTest = false);
 
@@ -345,30 +291,6 @@ class UK2Node : public UEdGraphNode
 
 	BLUEPRINTGRAPH_API virtual bool DoesInputWildcardPinAcceptArray(const UEdGraphPin* Pin) const { return true; }
 protected:
-	/** 
-	 * A mapping from old property and function names to new ones.  Get primed from INI files, and should contain entries for properties, functions, and delegates that get moved, so they can be fixed up
-	 */
-	static TMap<FFieldRemapInfo, FFieldRemapInfo> FieldRedirectMap;
-	/** 
-	 * A mapping from old pin name to new pin name for each K2 node.  Get primed from INI files, and should contain entries for node class, and old param name and new param name
-	 */
-	static TMultiMap<UClass*, FParamRemapInfo> ParamRedirectMap;
-
-	/** Has the field map been intialized this run */
-	static bool bFieldRedirectMapInitialized;
-
-	/** Init the field redirect map (if not already done) from .ini file entries */
-	static void InitFieldRedirectMap();
-
-	/** 
-	 * Searches the field replacement map for an appropriately named field in the specified scope, and returns an updated name if this property has been listed in the remapping table
-	 *
-	 * @param	Class		The class scope to search for a field remap on
-	 * @param	FieldName	The field name (function name or property name) to search for
-	 * @param	RemapInfo	(out) Struct containing updated location info for the field
-	 * @return	Whether or not a remap was found in the specified scope
-	 */
-	static bool FindReplacementFieldName(UClass* Class, FName FieldName, FFieldRemapInfo& RemapInfo);
 
 	enum ERedirectType
 	{
@@ -426,330 +348,9 @@ protected:
 	}
 
 	void FixupPinDefaultValues();
-};
-
-
-USTRUCT()
-struct FMemberReference
-{
-	GENERATED_USTRUCT_BODY()
 
-protected:
-	/** Class that this member is defined in. Should be NULL if bSelfContext is true.  */
-	UPROPERTY()
-	mutable TSubclassOf<class UObject> MemberParentClass;
-
-	/** Class that this member is defined in. Should be NULL if bSelfContext is true.  */
-	UPROPERTY()
-	mutable FString MemberScope;
-
-	/** Name of variable */
-	UPROPERTY()
-	mutable FName MemberName;
-
-	/** The Guid of the variable */
-	UPROPERTY()
-	mutable FGuid MemberGuid;
-
-	/** Whether or not this should be a "self" context */
-	UPROPERTY()
-	mutable bool bSelfContext;
-
-	/** Whether or not this property has been deprecated */
-	UPROPERTY()
-	mutable bool bWasDeprecated;
-	
 public:
-	FMemberReference()
-		: MemberParentClass(NULL)
-		, MemberName(NAME_None)
-		, bSelfContext(false)
-	{
-	}
-
-	/** Set up this reference from a supplied field */
-	template<class TFieldType>
-	void SetFromField(const UField* InField, const bool bIsConsideredSelfContext)
-	{
-		MemberParentClass = bIsConsideredSelfContext ? NULL : InField->GetOwnerClass();
-		MemberName = InField->GetFName();
-		bSelfContext = bIsConsideredSelfContext;
-		bWasDeprecated = false;
-
-		if (MemberParentClass != nullptr)
-		{
-			MemberParentClass = MemberParentClass->GetAuthoritativeClass();
-		}
-
-		MemberGuid.Invalidate();
-		if (InField->GetOwnerClass())
-		{
-			UBlueprint::GetGuidFromClassByFieldName<TFieldType>(InField->GetOwnerClass(), InField->GetFName(), MemberGuid);
-		}
-	}
-
-	template<class TFieldType>
-	void SetFromField(const UField* InField, const UK2Node* SelfNode)
-	{
-		FGuid FieldGuid;
-		if (InField->GetOwnerClass())
-		{
-			UBlueprint::GetGuidFromClassByFieldName<TFieldType>(InField->GetOwnerClass(), InField->GetFName(), FieldGuid);
-		}
-
-		SetGivenSelfScope(InField->GetFName(), FieldGuid, InField->GetOwnerClass(), GetBlueprintClassFromNode(SelfNode));
-	}
-
-	/** Update given a new self */
-	template<class TFieldType>
-	void RefreshGivenNewSelfScope(const UK2Node* SelfNode)
-	{
-		UClass* SelfScope = GetBlueprintClassFromNode(SelfNode);
-
-		if ((MemberParentClass != NULL) && (SelfScope != NULL))
-		{
-			UBlueprint::GetGuidFromClassByFieldName<TFieldType>((MemberParentClass ? *MemberParentClass : SelfScope), MemberName, MemberGuid);
-			SetGivenSelfScope(MemberName, MemberGuid, MemberParentClass, SelfScope);
-		}
-		else
-		{
-			// We no longer have enough information to known if we've done the right thing, and just have to hope...
-		}
-	}
-
-	/** Set to a non-'self' member, so must include reference to class owning the member. */
-	BLUEPRINTGRAPH_API void SetExternalMember(FName InMemberName, TSubclassOf<class UObject> InMemberParentClass);
-
-	/** Set to a non-'self' delegate member, this is not self-context but is not given a parent class */
-	BLUEPRINTGRAPH_API void SetExternalDelegateMember(FName InMemberName);
-
-	/** Set up this reference to a 'self' member name */
-	BLUEPRINTGRAPH_API void SetSelfMember(FName InMemberName);
-
-	/** Set up this reference to a 'self' member name, scoped to a struct */
-	BLUEPRINTGRAPH_API void SetLocalMember(FName InMemberName, UStruct* InScope, const FGuid InMemberGuid);
-
-	/** Set up this reference to a 'self' member name, scoped to a struct name */
-	BLUEPRINTGRAPH_API void SetLocalMember(FName InMemberName, FString InScopeName, const FGuid InMemberGuid);
 
-	/** Only intended for backwards compat! */
-	BLUEPRINTGRAPH_API void SetDirect(const FName InMemberName, const FGuid InMemberGuid, TSubclassOf<class UObject> InMemberParentClass, bool bIsConsideredSelfContext);
-
-	/** Invalidate the current MemberParentClass, if this is a self scope, or the MemberScope if it is not (and set).  Intended for PostDuplication fixups */
-	BLUEPRINTGRAPH_API void InvalidateScope();
-
-	/** Get the name of this member */
-	FName GetMemberName() const
-	{
-		return MemberName;
-	}
-
-	FGuid GetMemberGuid() const
-	{
-		return MemberGuid;
-	}
-
-	/** Returns if this is a 'self' context. */
-	bool IsSelfContext() const
-	{
-		return bSelfContext;
-	}
-
-	/** Returns if this is a local scope. */
-	bool IsLocalScope() const
-	{
-		return !MemberScope.IsEmpty();
-	}
-private:
 	/** Util to get the generated class from a node. */
-	BLUEPRINTGRAPH_API static UClass* GetBlueprintClassFromNode(const UK2Node* Node);
-
-	/**
-	 * Refreshes a local variable reference name if it has changed
-	 *
-	 * @param InSelfScope		Scope to lookup the variable in, to see if it has changed
-	 */
-	BLUEPRINTGRAPH_API FName RefreshLocalVariableName(UClass* InSelfScope) const;
-
-protected:
-	/** Only intended for backwards compat! */
-	BLUEPRINTGRAPH_API void SetGivenSelfScope(const FName InMemberName, const FGuid InMemberGuid, TSubclassOf<class UObject> InMemberParentClass, TSubclassOf<class UObject> SelfScope) const;
-
-public:
-	/** Get the class that owns this member */
-	UClass* GetMemberParentClass(const UK2Node* SelfNode) const
-	{
-		return GetMemberParentClass( GetBlueprintClassFromNode(SelfNode) );
-	}
-
-	/** Get the class that owns this member */
-	UClass* GetMemberParentClass(UClass* SelfScope) const
-	{
-		// Local variables with a MemberScope act much the same as being SelfContext, their parent class is SelfScope.
-		return (bSelfContext || !MemberScope.IsEmpty())? SelfScope : (UClass*)MemberParentClass;
-	}
-
-	/** Get the scope of this member, using a node to derive the class */
-	UStruct* GetMemberScope(const UK2Node* SelfNode) const
-	{
-		return GetMemberScope( GetBlueprintClassFromNode(SelfNode) );
-	}
-
-	/** Get the scope of this member */
-	UStruct* GetMemberScope(UClass* InMemberParentClass) const
-	{
-		return FindField<UStruct>(InMemberParentClass, *MemberScope);
-	}
-
-	/** Get the name of the scope of this member */
-	FString GetMemberScopeName() const
-	{
-		return MemberScope;
-	}
-
-	/** Returns whether or not the variable has been deprecated */
-	bool IsDeprecated() const
-	{
-		return bWasDeprecated;
-	}
-
-	/** 
-	 *	Returns the member UProperty/UFunction this reference is pointing to, or NULL if it no longer exists 
-	 *	Derives 'self' scope from supplied Blueprint node if required
-	 *	Will check for redirects and fix itself up if one is found.
-	 */
-	template<class TFieldType>
-	TFieldType* ResolveMember(UClass* SelfScope) const
-	{
-		TFieldType* ReturnField = NULL;
-
-		if(bSelfContext && SelfScope == NULL)
-		{
-			UE_LOG(LogBlueprint, Warning, TEXT("FMemberReference::ResolveMember (%s) bSelfContext == true, but no scope supplied!"), *MemberName.ToString() );
-		}
-
-		// Check if the member reference is function scoped
-		if(IsLocalScope())
-		{
-			UStruct* MemberScopeStruct = FindField<UStruct>(SelfScope, *MemberScope);
-
-			// Find in target scope
-			ReturnField = FindField<TFieldType>(MemberScopeStruct, MemberName);
-
-			if(ReturnField == NULL)
-			{
-				// If the property was not found, refresh the local variable name and try again
-				const FName RenamedMemberName = RefreshLocalVariableName(SelfScope);
-				if (RenamedMemberName != NAME_None)
-				{
-					ReturnField = FindField<TFieldType>(MemberScopeStruct, MemberName);
-				}
-			}
-		}
-		else
-		{
-			// Look for remapped member
-			UClass* TargetScope = bSelfContext ? SelfScope : (UClass*)MemberParentClass;
-			if( TargetScope != NULL &&  !GIsSavingPackage )
-			{
-				ReturnField = Cast<TFieldType>(UK2Node::FindRemappedField(TargetScope, MemberName, true));
-			}
-
-			if(ReturnField != NULL)
-			{
-				// Fix up this struct, we found a redirect
-				MemberName = ReturnField->GetFName();
-				MemberParentClass = Cast<UClass>(ReturnField->GetOuter());
-
-				MemberGuid.Invalidate();
-				UBlueprint::GetGuidFromClassByFieldName<TFieldType>(TargetScope, MemberName, MemberGuid);
-
-				if (MemberParentClass != nullptr)
-				{
-					MemberParentClass = MemberParentClass->GetAuthoritativeClass();
-					// Re-evaluate self-ness against the redirect if we were given a valid SelfScope
-					if (SelfScope != NULL)
-					{
-						SetGivenSelfScope(MemberName, MemberGuid, MemberParentClass, SelfScope);
-					}
-				}	
-			}
-			else if(TargetScope != NULL)
-			{
-				// Find in target scope
-				ReturnField = FindField<TFieldType>(TargetScope, MemberName);
-
-				// If we have a GUID find the reference variable and make sure the name is up to date and find the field again
-				// For now only variable references will have valid GUIDs.  Will have to deal with finding other names subsequently
-				if (ReturnField == NULL && MemberGuid.IsValid())
-				{
-					const FName RenamedMemberName = UBlueprint::GetFieldNameFromClassByGuid<TFieldType>(TargetScope, MemberGuid);
-					if (RenamedMemberName != NAME_None)
-					{
-						MemberName = RenamedMemberName;
-						ReturnField = FindField<TFieldType>(TargetScope, MemberName);
-					}
-				}
-			}
-			else
-			{
-				// Native delegate signatures can have null target scope.
-                // Needed because Blueprint pin types are referencing delegate signatures as
-                // class members, which doesn't have to be true.
-				if (TFieldType::StaticClass() == UFunction::StaticClass())
-				{
-					ReturnField = (TFieldType*)FindDelegateSignature(MemberName);
-				}
-			}
-		}
-
-		// Check to see if the member has been deprecated
-		if (UProperty* Property = Cast<UProperty>(ReturnField))
-		{
-			bWasDeprecated = Property->HasAnyPropertyFlags(CPF_Deprecated);
-		}
-
-		return ReturnField;
-	}
-
-	template<class TFieldType>
-	TFieldType* ResolveMember(UBlueprint* SelfScope)
-	{
-		return ResolveMember<TFieldType>(SelfScope->SkeletonGeneratedClass);
-	}
-
-	/** 
-	 *	Returns the member UProperty/UFunction this reference is pointing to, or NULL if it no longer exists 
-	 *	Derives 'self' scope from supplied Blueprint node if required
-	 *	Will check for redirects and fix itself up if one is found.
-	 */
-	template<class TFieldType>
-	TFieldType* ResolveMember(const UK2Node* SelfNode) const
-	{
-		return ResolveMember<TFieldType>( GetBlueprintClassFromNode(SelfNode) );
-	}
-
-	template<class TFieldType>
-	static void FillSimpleMemberReference(const UField* InField, FSimpleMemberReference& OutReference)
-	{
-		OutReference.Reset();
-
-		if (InField)
-		{
-			FMemberReference TempMemberReference;
-			TempMemberReference.SetFromField<TFieldType>(InField, false);
-
-			OutReference.MemberName = TempMemberReference.MemberName;
-			OutReference.MemberParentClass = TempMemberReference.MemberParentClass;
-			OutReference.MemberGuid = TempMemberReference.MemberGuid;
-		}
-	}
-
-	template<class TFieldType>
-	static TFieldType* ResolveSimpleMemberReference(const FSimpleMemberReference& Reference)
-	{
-		FMemberReference TempMemberReference;
-		TempMemberReference.SetDirect(Reference.MemberName, Reference.MemberGuid, Reference.MemberParentClass, false);
-		return TempMemberReference.ResolveMember<TFieldType>((UClass*)NULL);
-	}
-};
+	BLUEPRINTGRAPH_API UClass* GetBlueprintClassFromNode() const;
+};
\ No newline at end of file
diff --git a/Engine/Source/Editor/BlueprintGraph/Classes/K2Node_BaseMCDelegate.h b/Engine/Source/Editor/BlueprintGraph/Classes/K2Node_BaseMCDelegate.h
--- a/Engine/Source/Editor/BlueprintGraph/Classes/K2Node_BaseMCDelegate.h
+++ b/Engine/Source/Editor/BlueprintGraph/Classes/K2Node_BaseMCDelegate.h
@@ -2,6 +2,7 @@
 #pragma once
 #include "K2Node.h"
 #include "EdGraph/EdGraphNodeUtils.h" // for FNodeTextCache
+#include "Engine/MemberReference.h"
 #include "K2Node_BaseMCDelegate.generated.h"
 
 UCLASS(MinimalAPI, abstract)
@@ -40,7 +41,7 @@ public:
 
 	BLUEPRINTGRAPH_API UProperty* GetProperty() const
 	{
-		return DelegateReference.ResolveMember<UProperty>(this);
+		return DelegateReference.ResolveMember<UProperty>(GetBlueprintClassFromNode());
 	}
 
 	BLUEPRINTGRAPH_API FName GetPropertyName() const
diff --git a/Engine/Source/Editor/BlueprintGraph/Classes/K2Node_CallFunction.h b/Engine/Source/Editor/BlueprintGraph/Classes/K2Node_CallFunction.h
--- a/Engine/Source/Editor/BlueprintGraph/Classes/K2Node_CallFunction.h
+++ b/Engine/Source/Editor/BlueprintGraph/Classes/K2Node_CallFunction.h
@@ -4,6 +4,7 @@
 
 #include "K2Node.h"
 #include "EdGraph/EdGraphNodeUtils.h" // for FNodeTextCache
+#include "Engine/MemberReference.h"
 #include "K2Node_CallFunction.generated.h"
 
 UCLASS()
diff --git a/Engine/Source/Editor/BlueprintGraph/Classes/K2Node_Event.h b/Engine/Source/Editor/BlueprintGraph/Classes/K2Node_Event.h
--- a/Engine/Source/Editor/BlueprintGraph/Classes/K2Node_Event.h
+++ b/Engine/Source/Editor/BlueprintGraph/Classes/K2Node_Event.h
@@ -3,6 +3,7 @@
 
 #pragma once
 #include "Core.h"
+#include "Engine/MemberReference.h"
 #include "K2Node_EditablePinBase.h"
 #include "EdGraph/EdGraphNodeUtils.h" // for FNodeTextCache
 #include "K2Node_Event.generated.h"
diff --git a/Engine/Source/Editor/BlueprintGraph/Classes/K2Node_Variable.h b/Engine/Source/Editor/BlueprintGraph/Classes/K2Node_Variable.h
--- a/Engine/Source/Editor/BlueprintGraph/Classes/K2Node_Variable.h
+++ b/Engine/Source/Editor/BlueprintGraph/Classes/K2Node_Variable.h
@@ -3,6 +3,7 @@
 
 #pragma once
 #include "K2Node.h"
+#include "Engine/MemberReference.h"
 #include "K2Node_Variable.generated.h"
 
 // Forward Declarations
diff --git a/Engine/Source/Editor/BlueprintGraph/Private/CallFunctionHandler.cpp b/Engine/Source/Editor/BlueprintGraph/Private/CallFunctionHandler.cpp
--- a/Engine/Source/Editor/BlueprintGraph/Private/CallFunctionHandler.cpp
+++ b/Engine/Source/Editor/BlueprintGraph/Private/CallFunctionHandler.cpp
@@ -502,7 +502,7 @@ UClass* FKCHandler_CallFunction::GetCallingContext(FKismetFunctionContext& Conte
 		else
 		{
 			// Final functions need the call context to be the specified class, so don't bother checking for the self pin.   The schema should enforce this.
-			return CallFuncNode->FunctionReference.GetMemberParentClass(CallFuncNode);
+			return CallFuncNode->FunctionReference.GetMemberParentClass(CallFuncNode->GetBlueprintClassFromNode());
 		}
 	}
 	else
diff --git a/Engine/Source/Editor/BlueprintGraph/Private/DelegateNodeHandlers.cpp b/Engine/Source/Editor/BlueprintGraph/Private/DelegateNodeHandlers.cpp
--- a/Engine/Source/Editor/BlueprintGraph/Private/DelegateNodeHandlers.cpp
+++ b/Engine/Source/Editor/BlueprintGraph/Private/DelegateNodeHandlers.cpp
@@ -74,7 +74,7 @@ struct FKCHandlerDelegateHelper
 			const UFunction* OrgSignature = DelegateNode->GetDelegateSignature();
 			if(const UEdGraphPin* DelegatePin = DelegateNode->GetDelegatePin())
 			{
-				const UFunction* PinSignature = FMemberReference::ResolveSimpleMemberReference<UFunction>(DelegatePin->PinType.PinSubCategoryMemberReference);
+				const UFunction* PinSignature = DelegatePin->PinType.PinSubCategoryMemberReference.ResolveSimpleMemberReference<UFunction>();
 				if (!OrgSignature || !PinSignature || !OrgSignature->IsSignatureCompatibleWith(PinSignature))
 				{
 					MessageLog.Error(*LOCTEXT("WrongDelegate", "Wrong Event Dispatcher. Refresh node @@").ToString(), DelegateNode);
@@ -259,11 +259,11 @@ void FKCHandler_CreateDelegate::RegisterNets(FKismetFunctionContext& Context, UE
 		if(NULL == OutDelegateTerm)
 		{
 			OutDelegateTerm = Context.CreateLocalTerminalFromPinAutoChooseScope(Net, Context.NetNameMap->MakeValidName(Net));
-			if (NULL == FMemberReference::ResolveSimpleMemberReference<UFunction>(OutDelegateTerm->Type.PinSubCategoryMemberReference))
+			if (NULL == OutDelegateTerm->Type.PinSubCategoryMemberReference.ResolveSimpleMemberReference<UFunction>())
 			{
-				FMemberReference::FillSimpleMemberReference<UFunction>(DelegateNode->GetDelegateSignature(), OutDelegateTerm->Type.PinSubCategoryMemberReference);
+				OutDelegateTerm->Type.PinSubCategoryMemberReference.FillSimpleMemberReference<UFunction>(DelegateNode->GetDelegateSignature());
 			}
-			if (NULL == FMemberReference::ResolveSimpleMemberReference<UFunction>(OutDelegateTerm->Type.PinSubCategoryMemberReference))
+			if (NULL == OutDelegateTerm->Type.PinSubCategoryMemberReference.ResolveSimpleMemberReference<UFunction>())
 			{
 				CompilerContext.MessageLog.Error(*LOCTEXT("UnconnectedDelegateSig", "Event Dispatcher has no signature @@").ToString(), OutPin);
 				return;
diff --git a/Engine/Source/Editor/BlueprintGraph/Private/EdGraphSchema_K2.cpp b/Engine/Source/Editor/BlueprintGraph/Private/EdGraphSchema_K2.cpp
--- a/Engine/Source/Editor/BlueprintGraph/Private/EdGraphSchema_K2.cpp
+++ b/Engine/Source/Editor/BlueprintGraph/Private/EdGraphSchema_K2.cpp
@@ -1101,8 +1101,8 @@ static bool IsUsingNonExistantVariable(const UEdGraphNode* InGraphNode, UBluepri
 			else if(Variable->VariableReference.IsLocalScope())
 			{
 				// If there is no member scope, or we can't find the local variable in the member scope, then it's non-existant
-				if(!Variable->VariableReference.GetMemberScope(Variable) 
-					||  (Variable->VariableReference.GetMemberScope(Variable) && !FBlueprintEditorUtils::FindLocalVariable(OwnerBlueprint, Variable->VariableReference.GetMemberScope(Variable), Variable->GetVarName())))
+				UStruct* MemberScope = Variable->VariableReference.GetMemberScope(Variable->GetBlueprintClassFromNode());
+				if (MemberScope == nullptr || !FBlueprintEditorUtils::FindLocalVariable(OwnerBlueprint, MemberScope, Variable->GetVarName()))
 				{
 					bNonExistantVariable = true;
 				}
@@ -2182,7 +2182,7 @@ bool UEdGraphSchema_K2::FindSpecializedConversionNode(const UEdGraphPin* OutputP
 		UK2Node_CallFunction* CallFunctionNode = (UK2Node_CallFunction*)(InputPin->GetOwningNode());
 		UClass* OutputPinClass = Cast<UClass>(OutputPin->PinType.PinSubCategoryObject.Get());
 
-		UClass* FunctionClass = CallFunctionNode->FunctionReference.GetMemberParentClass(CallFunctionNode);
+		UClass* FunctionClass = CallFunctionNode->FunctionReference.GetMemberParentClass(CallFunctionNode->GetBlueprintClassFromNode());
 		if(FunctionClass != NULL && OutputPinClass != NULL)
 		{
 			// Iterate over object properties..
@@ -3078,12 +3078,12 @@ bool UEdGraphSchema_K2::ConvertPropertyToPinType(const UProperty* Property, /*ou
 	else if (const UMulticastDelegateProperty* MulticastDelegateProperty = Cast<const UMulticastDelegateProperty>(TestProperty))
 	{
 		TypeOut.PinCategory = PC_MCDelegate;
-		FMemberReference::FillSimpleMemberReference<UFunction>(MulticastDelegateProperty->SignatureFunction, TypeOut.PinSubCategoryMemberReference);
+		TypeOut.PinSubCategoryMemberReference.FillSimpleMemberReference<UFunction>(MulticastDelegateProperty->SignatureFunction);
 	}
 	else if (const UDelegateProperty* DelegateProperty = Cast<const UDelegateProperty>(TestProperty))
 	{
 		TypeOut.PinCategory = PC_Delegate;
-		FMemberReference::FillSimpleMemberReference<UFunction>(DelegateProperty->SignatureFunction, TypeOut.PinSubCategoryMemberReference);
+		TypeOut.PinSubCategoryMemberReference.FillSimpleMemberReference<UFunction>(DelegateProperty->SignatureFunction);
 	}
 	else
 	{
@@ -3469,7 +3469,7 @@ bool UEdGraphSchema_K2::ArePinTypesCompatible(const FEdGraphPinType& Output, con
 				return Func && (Func->HasAllFlags(RF_LoadCompleted) || !Func->HasAnyFlags(RF_NeedLoad | RF_WasLoaded));
 			};
 
-			const UFunction* OutFunction = FMemberReference::ResolveSimpleMemberReference<UFunction>(Output.PinSubCategoryMemberReference);
+			const UFunction* OutFunction = Output.PinSubCategoryMemberReference.ResolveSimpleMemberReference<UFunction>();
 			if (!CanUseFunction(OutFunction))
 			{
 				OutFunction = NULL;
@@ -3483,7 +3483,7 @@ bool UEdGraphSchema_K2::ArePinTypesCompatible(const FEdGraphPinType& Output, con
 					OutFunction = BPOwner->SkeletonGeneratedClass->FindFunctionByName(Output.PinSubCategoryMemberReference.MemberName);
 				}
 			}
-			const UFunction* InFunction = FMemberReference::ResolveSimpleMemberReference<UFunction>(Input.PinSubCategoryMemberReference);
+			const UFunction* InFunction = Input.PinSubCategoryMemberReference.ResolveSimpleMemberReference<UFunction>();
 			if (!CanUseFunction(InFunction))
 			{
 				InFunction = NULL;
diff --git a/Engine/Source/Editor/BlueprintGraph/Private/EdGraphSchema_K2_Actions.cpp b/Engine/Source/Editor/BlueprintGraph/Private/EdGraphSchema_K2_Actions.cpp
--- a/Engine/Source/Editor/BlueprintGraph/Private/EdGraphSchema_K2_Actions.cpp
+++ b/Engine/Source/Editor/BlueprintGraph/Private/EdGraphSchema_K2_Actions.cpp
@@ -399,7 +399,7 @@ bool FEdGraphSchemaAction_K2AddEvent::EventHasAlreadyBeenPlaced(UBlueprint const
 	if (Blueprint != NULL)
 	{
 		UK2Node_Event const* EventTemplate = Cast<UK2Node_Event const>(NodeTemplate);
-		ExistingEvent = FBlueprintEditorUtils::FindOverrideForFunction(Blueprint, EventTemplate->EventReference.GetMemberParentClass(EventTemplate), EventTemplate->EventReference.GetMemberName());
+		ExistingEvent = FBlueprintEditorUtils::FindOverrideForFunction(Blueprint, EventTemplate->EventReference.GetMemberParentClass(EventTemplate->GetBlueprintClassFromNode()), EventTemplate->EventReference.GetMemberName());
 	}
 
 	if (FoundEventOut != NULL)
diff --git a/Engine/Source/Editor/BlueprintGraph/Private/K2ActionMenuBuilder.cpp b/Engine/Source/Editor/BlueprintGraph/Private/K2ActionMenuBuilder.cpp
--- a/Engine/Source/Editor/BlueprintGraph/Private/K2ActionMenuBuilder.cpp
+++ b/Engine/Source/Editor/BlueprintGraph/Private/K2ActionMenuBuilder.cpp
@@ -1416,7 +1416,7 @@ void FK2ActionMenuBuilder::GetPinAllowedNodeTypes(FBlueprintGraphActionListBuild
 				TSharedPtr<FEdGraphSchemaAction_K2NewNode> Action = FK2ActionMenuBuilder::AddNewNodeAction(ContextMenuBuilder, FString(), NodeTemplate->GetNodeTitle(ENodeTitleType::ListView), NodeTemplate->GetTooltipText().ToString(), 0, NodeTemplate->GetKeywords());
 				Action->NodeTemplate = NodeTemplate;
 
-				UFunction* Signature = FMemberReference::ResolveSimpleMemberReference<UFunction>(FromPin.PinType.PinSubCategoryMemberReference);
+				UFunction* Signature = FromPin.PinType.PinSubCategoryMemberReference.ResolveSimpleMemberReference<UFunction>();
 				const bool bSupportsEventGraphs = (Blueprint && FBlueprintEditorUtils::DoesSupportEventGraphs(Blueprint));
 				const bool bAllowEvents = (GraphType == GT_Ubergraph);
 				if(Signature && bSupportsEventGraphs && bAllowEvents)
diff --git a/Engine/Source/Editor/BlueprintGraph/Private/K2Node.cpp b/Engine/Source/Editor/BlueprintGraph/Private/K2Node.cpp
--- a/Engine/Source/Editor/BlueprintGraph/Private/K2Node.cpp
+++ b/Engine/Source/Editor/BlueprintGraph/Private/K2Node.cpp
@@ -411,7 +411,7 @@ UK2Node::ERedirectType UK2Node::ShouldRedirectParam(const TArray<FString>& OldPi
 {
 	if ( ensure(NewPinNode) )
 	{
-		InitFieldRedirectMap();
+		const TMultiMap<UClass*, FParamRemapInfo>& ParamRedirectMap = FMemberReference::GetParamRedirectMap();
 
 		if ( ParamRedirectMap.Num() > 0 )
 		{
@@ -577,7 +577,7 @@ void UK2Node::ReconstructSinglePin(UEdGraphPin* NewPin, UEdGraphPin* OldPin, ERe
 			}
 
 			// go through for the NewPinNode
-			for(TMultiMap<UClass*, FParamRemapInfo>::TConstKeyIterator ParamIter(ParamRedirectMap, Cast<UK2Node>(NewPin->GetOwningNode())->GetClass()); ParamIter; ++ParamIter)
+			for(TMultiMap<UClass*, FParamRemapInfo>::TConstKeyIterator ParamIter(FMemberReference::GetParamRedirectMap(), Cast<UK2Node>(NewPin->GetOwningNode())->GetClass()); ParamIter; ++ParamIter)
 			{
 				const FParamRemapInfo& ParamRemap = ParamIter.Value();
 
@@ -735,220 +735,6 @@ FLinearColor UK2Node::GetNodeTitleColor() const
 	return GetDefault<UGraphEditorSettings>()->FunctionCallNodeTitleColor;
 }
 
-
-TMap<FFieldRemapInfo, FFieldRemapInfo> UK2Node::FieldRedirectMap;
-TMultiMap<UClass*, FParamRemapInfo> UK2Node::ParamRedirectMap;
-
-bool UK2Node::bFieldRedirectMapInitialized = false;
-
-bool UK2Node::FindReplacementFieldName(UClass* Class, FName FieldName, FFieldRemapInfo& RemapInfo)
-{
-	check(Class != NULL);
-	InitFieldRedirectMap();
-
-	// Reset the property remap info
-	RemapInfo = FFieldRemapInfo();
-
-	FFieldRemapInfo OldField;
-	OldField.FieldClass = Class->GetFName();
-	OldField.FieldName = FieldName;
-
-	FFieldRemapInfo* NewFieldInfoPtr = FieldRedirectMap.Find(OldField);
-	if (NewFieldInfoPtr != NULL)
-	{
-		RemapInfo = *NewFieldInfoPtr;
-		return true;
-	}
-	else
-	{
-		return false;
-	}
-}
-
-void UK2Node::InitFieldRedirectMap()
-{
-	if (!bFieldRedirectMapInitialized)
-	{
-		if (GConfig)
-		{
-			FConfigSection* PackageRedirects = GConfig->GetSectionPrivate( TEXT("/Script/Engine.Engine"), false, true, GEngineIni );
-			for (FConfigSection::TIterator It(*PackageRedirects); It; ++It)
-			{
-				if (It.Key() == TEXT("K2FieldRedirects"))
-				{
-					FString OldFieldPathString;
-					FString NewFieldPathString;
-
-					FParse::Value( *It.Value(), TEXT("OldFieldName="), OldFieldPathString );
-					FParse::Value( *It.Value(), TEXT("NewFieldName="), NewFieldPathString );
-
-					// Handle both cases of just a field being renamed (just one FName), as well as a class and field name (ClassName.FieldName)
-					FFieldRemapInfo OldFieldRemap;
-					{
-						TArray<FString> OldFieldPath;
-						OldFieldPathString.ParseIntoArray(OldFieldPath, TEXT("."), true);
-
-						if (OldFieldPath.Num() == 1)
-						{
-							// Only the new property name is specified
-							OldFieldRemap.FieldName = FName(*OldFieldPath[0]);
-						}
-						else if (OldFieldPath.Num() == 2)
-						{
-							// Property name and new class are specified
-							OldFieldRemap.FieldClass = FName(*OldFieldPath[0]);
-							OldFieldRemap.FieldName = FName(*OldFieldPath[1]);
-						}
-					}
-
-					// Handle both cases of just a field being renamed (just one FName), as well as a class and field name (ClassName.FieldName)
-					FFieldRemapInfo NewFieldRemap;
-					{
-						TArray<FString> NewFieldPath;
-						NewFieldPathString.ParseIntoArray(NewFieldPath, TEXT("."), true);
-
-						if( NewFieldPath.Num() == 1 )
-						{
-							// Only the new property name is specified
-							NewFieldRemap.FieldName = FName(*NewFieldPath[0]);
-						}
-						else if( NewFieldPath.Num() == 2 )
-						{
-							// Property name and new class are specified
-							NewFieldRemap.FieldClass = FName(*NewFieldPath[0]);
-							NewFieldRemap.FieldName = FName(*NewFieldPath[1]);
-						}
-					}
-
-					FieldRedirectMap.Add(OldFieldRemap, NewFieldRemap);
-				}			
-				if (It.Key() == TEXT("K2ParamRedirects"))
-				{
-					FName NodeName = NAME_None;
-					FName OldParam = NAME_None;
-					FName NewParam = NAME_None;
-					FName NodeTitle = NAME_None;
-
-					FString OldParamValues;
-					FString NewParamValues;
-					FString CustomValueMapping;
-
-					FParse::Value( *It.Value(), TEXT("NodeName="), NodeName );
-					FParse::Value( *It.Value(), TEXT("NodeTitle="), NodeTitle );
-					FParse::Value( *It.Value(), TEXT("OldParamName="), OldParam );
-					FParse::Value( *It.Value(), TEXT("NewParamName="), NewParam );
-					FParse::Value( *It.Value(), TEXT("OldParamValues="), OldParamValues );
-					FParse::Value( *It.Value(), TEXT("NewParamValues="), NewParamValues );
-					FParse::Value( *It.Value(), TEXT("CustomValueMapping="), CustomValueMapping );
-
-					FParamRemapInfo ParamRemap;
-					ParamRemap.NewParam = NewParam;
-					ParamRemap.OldParam = OldParam;
-					ParamRemap.NodeTitle = NodeTitle;
-					ParamRemap.bCustomValueMapping = (FCString::Stricmp(*CustomValueMapping,TEXT("true")) == 0);
-
-					if (CustomValueMapping.Len() > 0 && !ParamRemap.bCustomValueMapping)
-					{
-						UE_LOG(LogBlueprint, Warning, TEXT("Value other than true specified for CustomValueMapping for '%s' node param redirect '%s' to '%s'."), *(NodeName.ToString()), *(OldParam.ToString()), *(NewParam.ToString()));
-					}
-
-					TArray<FString> OldParamValuesList;
-					TArray<FString> NewParamValuesList;
-					OldParamValues.ParseIntoArray(OldParamValuesList, TEXT(";"), false);
-					NewParamValues.ParseIntoArray(NewParamValuesList, TEXT(";"), false);
-
-					if (OldParamValuesList.Num() != NewParamValuesList.Num())
-					{
-						UE_LOG(LogBlueprint, Warning, TEXT("Unequal lengths for old and new param values for '%s' node param redirect '%s' to '%s'."), *(NodeName.ToString()), *(OldParam.ToString()), *(NewParam.ToString()));
-					}
-
-					if (CustomValueMapping.Len() > 0 && (OldParamValuesList.Num() > 0 || NewParamValuesList.Num() > 0))
-					{
-						UE_LOG(LogBlueprint, Warning, TEXT("Both Custom and Automatic param value remapping specified for '%s' node param redirect '%s' to '%s'.  Only Custom will be applied."), *(NodeName.ToString()), *(OldParam.ToString()), *(NewParam.ToString()));
-					}
-
-					for (int32 i = FMath::Min(OldParamValuesList.Num(), NewParamValuesList.Num()) - 1; i >= 0; --i)
-					{
-						int32 CurSize = ParamRemap.ParamValueMap.Num();
-						ParamRemap.ParamValueMap.Add(OldParamValuesList[i], NewParamValuesList[i]);
-						if (CurSize == ParamRemap.ParamValueMap.Num())
-						{
-							UE_LOG(LogBlueprint, Warning, TEXT("Duplicate old param value '%s' for '%s' node param redirect '%s' to '%s'."), *(OldParamValuesList[i]), *(NodeName.ToString()), *(OldParam.ToString()), *(NewParam.ToString()));
-						}
-					}
-
-					// load class
-					UClass* NodeClass = LoadClass<UK2Node>(NULL, *NodeName.ToString(), NULL, LOAD_None, NULL);
-					if (NodeClass )
-					{
-						ParamRedirectMap.Add(NodeClass, ParamRemap);
-					}
-				}			
-			}
-
-			bFieldRedirectMapInitialized = true;
-		}
-	}
-}
-
-UField* UK2Node::FindRemappedField(UClass* InitialScope, FName InitialName, bool bInitialScopeMustBeOwnerOfField)
-{
-	FFieldRemapInfo NewFieldInfo;
-
-	bool bFoundReplacement = false;
-	
-	// Step up the class chain to check if us or any of our parents specify a redirect
-	UClass* TestRemapClass = InitialScope;
-	while( TestRemapClass != NULL )
-	{
-		if( FindReplacementFieldName(TestRemapClass, InitialName, NewFieldInfo) )
-		{
-			// Found it, stop our search
-			bFoundReplacement = true;
-			break;
-		}
-
-		TestRemapClass = TestRemapClass->GetSuperClass();
-	}
-
-	// In the case of a bifurcation of a variable (e.g. moved from a parent into certain children), verify that we don't also define the variable in the current scope first
-	if( bFoundReplacement && (FindField<UField>(InitialScope, InitialName) != nullptr))
-	{
-		bFoundReplacement = false;		
-	}
-
-	if( bFoundReplacement )
-	{
-		const FName NewFieldName = NewFieldInfo.FieldName;
-		UClass* SearchClass = (NewFieldInfo.FieldClass != NAME_None) ? (UClass*)StaticFindObject(UClass::StaticClass(), ANY_PACKAGE, *NewFieldInfo.FieldClass.ToString()) : (UClass*)TestRemapClass;
-
-		// Find the actual field specified by the redirector, so we can return it and update the node that uses it
-		UField* NewField = FindField<UField>(SearchClass, NewFieldInfo.FieldName);
-		if( NewField != NULL )
-		{
-			if (bInitialScopeMustBeOwnerOfField && !InitialScope->IsChildOf(SearchClass))
-			{
-				UE_LOG(LogBlueprint, Log, TEXT("UK2Node:  Unable to update field. Remapped field '%s' in not owned by given scope. Scope: '%s', Owner: '%s'."), *InitialName.ToString(), *InitialScope->GetName(), *NewFieldInfo.FieldClass.ToString());
-			}
-			else
-			{
-				UE_LOG(LogBlueprint, Log, TEXT("UK2Node:  Fixed up old field '%s' to new name '%s' on class '%s'."), *InitialName.ToString(), *NewFieldInfo.FieldName.ToString(), *SearchClass->GetName());
-				return NewField;
-			}
-		}
-		else if (SearchClass != NULL)
-		{
-			UE_LOG(LogBlueprint, Log, TEXT("UK2Node:  Unable to find updated field name for '%s' on class '%s'."), *InitialName.ToString(), *SearchClass->GetName());
-		}
-		else
-		{
-			UE_LOG(LogBlueprint, Log, TEXT("UK2Node:  Unable to find updated field name for '%s' on unknown class '%s'."), *InitialName.ToString(), *NewFieldInfo.FieldClass.ToString());
-		}
-	}
-
-	return NULL;
-}
-
 ERenamePinResult UK2Node::RenameUserDefinedPin(const FString& OldName, const FString& NewName, bool bTest)
 {
 	UEdGraphPin* Pin = NULL;
@@ -1306,111 +1092,14 @@ void UK2Node::GetPinHoverText(const UEdGraphPin& Pin, FString& HoverTextOut) con
 	}
 }
 
-//////////////////////////////////////////////////////////////////////////
-// FMemberReference
-
-void FMemberReference::SetExternalMember(FName InMemberName, TSubclassOf<class UObject> InMemberParentClass)
-{
-	MemberName = InMemberName;
-	MemberParentClass = (InMemberParentClass != nullptr) ? InMemberParentClass->GetAuthoritativeClass() : nullptr;
-	MemberScope.Empty();
-	bSelfContext = false;
-	bWasDeprecated = false;
-}
-
-void FMemberReference::SetExternalDelegateMember(FName InMemberName)
-{
-	SetExternalMember(InMemberName, nullptr);
-}
-
-void FMemberReference::SetSelfMember(FName InMemberName)
-{
-	MemberName = InMemberName;
-	MemberParentClass = NULL;
-	MemberScope.Empty();
-	bSelfContext = true;
-	bWasDeprecated = false;
-}
-
-void FMemberReference::SetDirect(const FName InMemberName, const FGuid InMemberGuid, TSubclassOf<class UObject> InMemberParentClass, bool bIsConsideredSelfContext)
-{
-	MemberName = InMemberName;
-	MemberGuid = InMemberGuid;
-	bSelfContext = bIsConsideredSelfContext;
-	bWasDeprecated = false;
-	MemberParentClass = InMemberParentClass;
-	MemberScope.Empty();
-}
-
-void FMemberReference::SetGivenSelfScope(const FName InMemberName, const FGuid InMemberGuid, TSubclassOf<class UObject> InMemberParentClass, TSubclassOf<class UObject> SelfScope) const
-{
-	MemberName = InMemberName;
-	MemberGuid = InMemberGuid;
-	MemberParentClass = (InMemberParentClass != nullptr) ? InMemberParentClass->GetAuthoritativeClass() : nullptr;
-	MemberScope.Empty();
-	bSelfContext = (SelfScope->IsChildOf(InMemberParentClass)) || (SelfScope->ClassGeneratedBy == InMemberParentClass->ClassGeneratedBy);
-	bWasDeprecated = false;
-
-	if (bSelfContext)
-	{
-		MemberParentClass = NULL;
-	}
-}
-
-void FMemberReference::SetLocalMember(FName InMemberName, UStruct* InScope, const FGuid InMemberGuid)
-{
-	SetLocalMember(InMemberName, InScope->GetName(), InMemberGuid);
-}
-
-void FMemberReference::SetLocalMember(FName InMemberName, FString InScopeName, const FGuid InMemberGuid)
+UClass* UK2Node::GetBlueprintClassFromNode() const
 {
-	MemberName = InMemberName;
-	MemberScope = InScopeName;
-	MemberGuid = InMemberGuid;
-	bSelfContext = false;
-}
-
-void FMemberReference::InvalidateScope()
-{
-	if( IsSelfContext() )
+	UClass* BPClass = nullptr;
+	if (HasValidBlueprint())
 	{
-		MemberParentClass = NULL;
-	}
-	else if(IsLocalScope())
-	{
-		MemberScope.Empty();
-
-		// Make it into a member reference since we are clearing the local context
-		bSelfContext = true;
-	}
-}
-
-UClass* FMemberReference::GetBlueprintClassFromNode(const UK2Node* Node)
-{
-	UClass* BPClass = NULL;
-	if(Node != NULL && Node->HasValidBlueprint())
-	{
-		BPClass =  Node->GetBlueprint()->SkeletonGeneratedClass;
+		BPClass = GetBlueprint()->SkeletonGeneratedClass;
 	}
 	return BPClass;
 }
 
-FName FMemberReference::RefreshLocalVariableName(UClass* InSelfScope) const
-{
-	TArray<UBlueprint*> Blueprints;
-	UBlueprint::GetBlueprintHierarchyFromClass(InSelfScope, Blueprints);
-
-	FName RenamedMemberName = NAME_None;
-	for (int32 BPIndex = 0; BPIndex < Blueprints.Num(); ++BPIndex)
-	{
-		RenamedMemberName = FBlueprintEditorUtils::FindLocalVariableNameByGuid(Blueprints[BPIndex], MemberGuid);
-		if (RenamedMemberName != NAME_None)
-		{
-			MemberName = RenamedMemberName;
-			break;
-		}
-	}
-	return RenamedMemberName;
-}
-
 #undef LOCTEXT_NAMESPACE
diff --git a/Engine/Source/Editor/BlueprintGraph/Private/K2Node_ActorBoundEvent.cpp b/Engine/Source/Editor/BlueprintGraph/Private/K2Node_ActorBoundEvent.cpp
--- a/Engine/Source/Editor/BlueprintGraph/Private/K2Node_ActorBoundEvent.cpp
+++ b/Engine/Source/Editor/BlueprintGraph/Private/K2Node_ActorBoundEvent.cpp
@@ -116,7 +116,7 @@ FText UK2Node_ActorBoundEvent::GetTooltipText() const
 
 FString UK2Node_ActorBoundEvent::GetDocumentationLink() const
 {
-	if (UClass* EventSignatureClass = EventReference.GetMemberParentClass(this))
+	if (UClass* EventSignatureClass = EventReference.GetMemberParentClass(GetBlueprintClassFromNode()))
 	{
 		return FString::Printf(TEXT("Shared/GraphNodes/Blueprint/%s%s"), EventSignatureClass->GetPrefixCPP(), *EventSignatureClass->GetName());
 	}
@@ -162,7 +162,7 @@ UMulticastDelegateProperty* UK2Node_ActorBoundEvent::GetTargetDelegateProperty()
 	// If we couldn't find the target delegate, then try to find it in the property remap table
 	if (!TargetDelegateProp)
 	{
-		UMulticastDelegateProperty* NewProperty = Cast<UMulticastDelegateProperty>(FindRemappedField(DelegateOwnerClass, DelegatePropertyName));
+		UMulticastDelegateProperty* NewProperty = Cast<UMulticastDelegateProperty>(FMemberReference::FindRemappedField(DelegateOwnerClass, DelegatePropertyName));
 		if (NewProperty)
 		{
 			// Found a remapped property, update the node
diff --git a/Engine/Source/Editor/BlueprintGraph/Private/K2Node_CallFunction.cpp b/Engine/Source/Editor/BlueprintGraph/Private/K2Node_CallFunction.cpp
--- a/Engine/Source/Editor/BlueprintGraph/Private/K2Node_CallFunction.cpp
+++ b/Engine/Source/Editor/BlueprintGraph/Private/K2Node_CallFunction.cpp
@@ -518,11 +518,11 @@ void UK2Node_CallFunction::AllocateDefaultPins()
 	// First try remap table
 	if (Function == NULL)
 	{
-		UClass* ParentClass = FunctionReference.GetMemberParentClass(this);
+		UClass* ParentClass = FunctionReference.GetMemberParentClass(GetBlueprintClassFromNode());
 
 		if (ParentClass != NULL)
 		{
-			if (UFunction* NewFunction = Cast<UFunction>(FindRemappedField(ParentClass, FunctionReference.GetMemberName())))
+			if (UFunction* NewFunction = Cast<UFunction>(FMemberReference::FindRemappedField(ParentClass, FunctionReference.GetMemberName())))
 			{
 				// Found a remapped property, update the node
 				Function = NewFunction;
@@ -543,7 +543,7 @@ void UK2Node_CallFunction::AllocateDefaultPins()
 				Function = FindField<UFunction>(TestClass, FunctionReference.GetMemberName());
 				if (Function != NULL)
 				{
-					UClass* OldClass = FunctionReference.GetMemberParentClass(this);
+					UClass* OldClass = FunctionReference.GetMemberParentClass(GetBlueprintClassFromNode());
 					Message_Note( FString::Printf(*LOCTEXT("FixedUpFunctionInLibrary", "UK2Node_CallFunction: Fixed up function '%s', originally in '%s', now in library '%s'.").ToString(),
 						*FunctionReference.GetMemberName().ToString(),
 						 (OldClass != NULL) ? *OldClass->GetName() : TEXT("(null)"), *TestClass->GetName()) );
@@ -967,14 +967,14 @@ void UK2Node_CallFunction::PinDefaultValueChanged(UEdGraphPin* Pin)
 
 UFunction* UK2Node_CallFunction::GetTargetFunction() const
 {
-	UFunction* Function = FunctionReference.ResolveMember<UFunction>(this);
+	UFunction* Function = FunctionReference.ResolveMember<UFunction>(GetBlueprintClassFromNode());
 	return Function;
 }
 
 UFunction* UK2Node_CallFunction::GetTargetFunctionFromSkeletonClass() const
 {
 	UFunction* TargetFunction = nullptr;
-	UClass* ParentClass = FunctionReference.GetMemberParentClass( this );
+	UClass* ParentClass = FunctionReference.GetMemberParentClass( GetBlueprintClassFromNode() );
 	UBlueprint* OwningBP = ParentClass ? Cast<UBlueprint>( ParentClass->ClassGeneratedBy ) : nullptr;
 	if( UClass* SkeletonClass = OwningBP ? OwningBP->SkeletonGeneratedClass : nullptr )
 	{
@@ -1368,7 +1368,7 @@ void UK2Node_CallFunction::SetFromFunction(const UFunction* Function)
 		bIsConstFunc = Function->HasAnyFunctionFlags(FUNC_Const);
 		DetermineWantsEnumToExecExpansion(Function);
 
-		FunctionReference.SetFromField<UFunction>(Function, this);
+		FunctionReference.SetFromField<UFunction>(Function, GetBlueprintClassFromNode());
 	}
 }
 
@@ -1388,7 +1388,7 @@ FString UK2Node_CallFunction::GetDocumentationLink() const
 	}
 	else 
 	{
-		ParentClass = FunctionReference.GetMemberParentClass(this);
+		ParentClass = FunctionReference.GetMemberParentClass(GetBlueprintClassFromNode());
 	}
 	
 	if (ParentClass != NULL)
@@ -1496,7 +1496,7 @@ void UK2Node_CallFunction::GetRedirectPinNames(const UEdGraphPin& Pin, TArray<FS
 		RedirectPinNames.Add(FString::Printf(TEXT("%s.%s"), *FunctionReference.GetMemberName().ToString(), *OldPinName));
 
 		// if there is class, also add an option for class.functionname.param
-		UClass* FunctionClass = FunctionReference.GetMemberParentClass(this);
+		UClass* FunctionClass = FunctionReference.GetMemberParentClass(GetBlueprintClassFromNode());
 		while (FunctionClass)
 		{
 			RedirectPinNames.Add(FString::Printf(TEXT("%s.%s.%s"), *FunctionClass->GetName(), *FunctionReference.GetMemberName().ToString(), *OldPinName));
@@ -1509,7 +1509,7 @@ bool UK2Node_CallFunction::IsSelfPinCompatibleWithBlueprintContext(UEdGraphPin *
 {
 	check(BlueprintObj);
 
-	UClass* FunctionClass = FunctionReference.GetMemberParentClass(this);
+	UClass* FunctionClass = FunctionReference.GetMemberParentClass(GetBlueprintClassFromNode());
 
 	bool bIsCompatible = (SelfPin != NULL) ? SelfPin->bHidden : true;
 	if (!bIsCompatible && (BlueprintObj->GeneratedClass != NULL))
@@ -1705,7 +1705,7 @@ void UK2Node_CallFunction::PostPlacedNewNode()
 	Super::PostPlacedNewNode();
 
 	// Try re-setting the function given our new parent scope, in case it turns an external to an internal, or vis versa
-	FunctionReference.RefreshGivenNewSelfScope<UFunction>(this);
+	FunctionReference.RefreshGivenNewSelfScope<UFunction>(GetBlueprintClassFromNode());
 }
 
 FNodeHandlingFunctor* UK2Node_CallFunction::CreateNodeHandler(FKismetCompilerContext& CompilerContext) const
@@ -2164,7 +2164,7 @@ FText UK2Node_CallFunction::GetMenuCategory() const
 
 bool UK2Node_CallFunction::HasExternalBlueprintDependencies(TArray<class UStruct*>* OptionalOutput) const
 {
-	const UClass* SourceClass = FunctionReference.GetMemberParentClass(this);
+	const UClass* SourceClass = FunctionReference.GetMemberParentClass(GetBlueprintClassFromNode());
 	const UBlueprint* SourceBlueprint = GetBlueprint();
 	const bool bResult = (SourceClass != NULL) && (SourceClass->ClassGeneratedBy != NULL) && (SourceClass->ClassGeneratedBy != SourceBlueprint);
 	if (bResult && OptionalOutput)
@@ -2179,7 +2179,7 @@ UEdGraph* UK2Node_CallFunction::GetFunctionGraph(const UEdGraphNode*& OutGraphNo
 	OutGraphNode = NULL;
 
 	// Search for the Blueprint owner of the function graph, climbing up through the Blueprint hierarchy
-	UClass* MemberParentClass = FunctionReference.GetMemberParentClass(this);
+	UClass* MemberParentClass = FunctionReference.GetMemberParentClass(GetBlueprintClassFromNode());
 	if(MemberParentClass != NULL)
 	{
 		UBlueprintGeneratedClass* ParentClass = Cast<UBlueprintGeneratedClass>(MemberParentClass);
diff --git a/Engine/Source/Editor/BlueprintGraph/Private/K2Node_CallFunctionOnMember.cpp b/Engine/Source/Editor/BlueprintGraph/Private/K2Node_CallFunctionOnMember.cpp
--- a/Engine/Source/Editor/BlueprintGraph/Private/K2Node_CallFunctionOnMember.cpp
+++ b/Engine/Source/Editor/BlueprintGraph/Private/K2Node_CallFunctionOnMember.cpp
@@ -25,7 +25,7 @@ UEdGraphPin* UK2Node_CallFunctionOnMember::CreateSelfPin(const UFunction* Functi
 	else
 	{
 		// This means that the function is declared in an external class, and should reference that class
-		SelfPin = CreatePin(EGPD_Input, K2Schema->PC_Object, TEXT(""), MemberVariableToCallOn.GetMemberParentClass(this), false, false, K2Schema->PN_Self);
+		SelfPin = CreatePin(EGPD_Input, K2Schema->PC_Object, TEXT(""), MemberVariableToCallOn.GetMemberParentClass(GetBlueprintClassFromNode()), false, false, K2Schema->PN_Self);
 	}
 	check(SelfPin != NULL);
 
@@ -34,7 +34,7 @@ UEdGraphPin* UK2Node_CallFunctionOnMember::CreateSelfPin(const UFunction* Functi
 
 FString UK2Node_CallFunctionOnMember::GetFunctionContextString() const
 {
-	UClass* MemberVarClass = MemberVariableToCallOn.GetMemberParentClass(this);
+	UClass* MemberVarClass = MemberVariableToCallOn.GetMemberParentClass(GetBlueprintClassFromNode());
 	FString CallFunctionClassName = (MemberVarClass != NULL) ? MemberVarClass->GetName() : TEXT("NULL");
 	return FString(TEXT("\n")) + FString::Printf(TEXT("%s %s (%s)"), *LOCTEXT("CallFunctionOnMemberDifferentContext", "Target is").ToString(), *CallFunctionClassName, *MemberVariableToCallOn.GetMemberName().ToString());
 }
@@ -124,7 +124,7 @@ void UK2Node_CallFunctionOnMember::ExpandNode(class FKismetCompilerContext& Comp
 
 bool UK2Node_CallFunctionOnMember::HasExternalBlueprintDependencies(TArray<class UStruct*>* OptionalOutput) const
 {
-	UClass* SourceClass = MemberVariableToCallOn.GetMemberParentClass(this);
+	UClass* SourceClass = MemberVariableToCallOn.GetMemberParentClass(GetBlueprintClassFromNode());
 	const UBlueprint* SourceBlueprint = GetBlueprint();
 	const bool bResult = (SourceClass != NULL) && (SourceClass->ClassGeneratedBy != NULL) && (SourceClass->ClassGeneratedBy != SourceBlueprint);
 	if (bResult && OptionalOutput)
diff --git a/Engine/Source/Editor/BlueprintGraph/Private/K2Node_CallParentFunction.cpp b/Engine/Source/Editor/BlueprintGraph/Private/K2Node_CallParentFunction.cpp
--- a/Engine/Source/Editor/BlueprintGraph/Private/K2Node_CallParentFunction.cpp
+++ b/Engine/Source/Editor/BlueprintGraph/Private/K2Node_CallParentFunction.cpp
@@ -15,7 +15,7 @@ UK2Node_CallParentFunction::UK2Node_CallParentFunction(const FObjectInitializer&
 
 FText UK2Node_CallParentFunction::GetNodeTitle(ENodeTitleType::Type TitleType) const
 {
-	UFunction* Function = FunctionReference.ResolveMember<UFunction>(this);
+	UFunction* Function = FunctionReference.ResolveMember<UFunction>(GetBlueprintClassFromNode());
 	FString FunctionName = Function ? GetUserFacingFunctionName( Function ) : FunctionReference.GetMemberName().ToString();
 
 	if( GEditor && GetDefault<UEditorStyleSettings>()->bShowFriendlyNames )
diff --git a/Engine/Source/Editor/BlueprintGraph/Private/K2Node_CommutativeAssociativeBinaryOperator.cpp b/Engine/Source/Editor/BlueprintGraph/Private/K2Node_CommutativeAssociativeBinaryOperator.cpp
--- a/Engine/Source/Editor/BlueprintGraph/Private/K2Node_CommutativeAssociativeBinaryOperator.cpp
+++ b/Engine/Source/Editor/BlueprintGraph/Private/K2Node_CommutativeAssociativeBinaryOperator.cpp
@@ -136,7 +136,7 @@ void UK2Node_CommutativeAssociativeBinaryOperator::AllocateDefaultPins()
 	}
 	else
 	{
-		const UClass* FunctionParentClass = FunctionReference.GetMemberParentClass(this);
+		const UClass* FunctionParentClass = FunctionReference.GetMemberParentClass(GetBlueprintClassFromNode());
 		Message_Error(FString::Printf(
 			*LOCTEXT("NoFunction_Error", "CommutativeAssociativeBinaryOperator has no function: '%s' class: '%s'").ToString(), 
 			*FunctionReference.GetMemberName().ToString(),
diff --git a/Engine/Source/Editor/BlueprintGraph/Private/K2Node_Composite.cpp b/Engine/Source/Editor/BlueprintGraph/Private/K2Node_Composite.cpp
--- a/Engine/Source/Editor/BlueprintGraph/Private/K2Node_Composite.cpp
+++ b/Engine/Source/Editor/BlueprintGraph/Private/K2Node_Composite.cpp
@@ -95,7 +95,7 @@ void UK2Node_Composite::PostPasteNode()
 			if(UK2Node_Event* Event = Cast<UK2Node_Event>(Node))
 			{
 				UBlueprint* BP = FBlueprintEditorUtils::FindBlueprintForGraphChecked(BoundGraph);
-				if(FBlueprintEditorUtils::FindOverrideForFunction(BP, Event->EventReference.GetMemberParentClass(Event), Event->EventReference.GetMemberName()))
+				if(FBlueprintEditorUtils::FindOverrideForFunction(BP, Event->EventReference.GetMemberParentClass(Event->GetBlueprintClassFromNode()), Event->EventReference.GetMemberName()))
 				{
 					FBlueprintEditorUtils::RemoveNode(BP, Node, true);
 					NodeIndex--;
diff --git a/Engine/Source/Editor/BlueprintGraph/Private/K2Node_CreateDelegate.cpp b/Engine/Source/Editor/BlueprintGraph/Private/K2Node_CreateDelegate.cpp
--- a/Engine/Source/Editor/BlueprintGraph/Private/K2Node_CreateDelegate.cpp
+++ b/Engine/Source/Editor/BlueprintGraph/Private/K2Node_CreateDelegate.cpp
@@ -72,7 +72,7 @@ bool UK2Node_CreateDelegate::IsValid(FString* OutMsg, bool bDontUseSkeletalClass
 	{
 		const UEdGraphPin* OtherPin = DelegatePin->LinkedTo[PinIter];
 		const UFunction* OtherSignature = OtherPin ? 
-			FMemberReference::ResolveSimpleMemberReference<UFunction>(OtherPin->PinType.PinSubCategoryMemberReference) : NULL;
+			OtherPin->PinType.PinSubCategoryMemberReference.ResolveSimpleMemberReference<UFunction>() : NULL;
 		if(!OtherSignature || !Signature->IsSignatureCompatibleWith(OtherSignature))
 		{
 			if (OutMsg)
@@ -272,7 +272,7 @@ UFunction* UK2Node_CreateDelegate::GetDelegateSignature() const
 		if(UEdGraphPin* ResultPin = Pin->LinkedTo[0])
 		{
 			ensure(K2Schema->PC_Delegate == ResultPin->PinType.PinCategory);
-			return FMemberReference::ResolveSimpleMemberReference<UFunction>(ResultPin->PinType.PinSubCategoryMemberReference);
+			return ResultPin->PinType.PinSubCategoryMemberReference.ResolveSimpleMemberReference<UFunction>();
 		}
 	}
 	return NULL;
diff --git a/Engine/Source/Editor/BlueprintGraph/Private/K2Node_DelegateSet.cpp b/Engine/Source/Editor/BlueprintGraph/Private/K2Node_DelegateSet.cpp
--- a/Engine/Source/Editor/BlueprintGraph/Private/K2Node_DelegateSet.cpp
+++ b/Engine/Source/Editor/BlueprintGraph/Private/K2Node_DelegateSet.cpp
@@ -30,7 +30,7 @@ public:
 			// Create a term to store the locally created delegate that we'll use to add to the MC delegate
 			FBPTerminal* DelegateTerm = Context.CreateLocalTerminal();
 			DelegateTerm->Type.PinCategory = CompilerContext.GetSchema()->PC_Delegate;
-			FMemberReference::FillSimpleMemberReference<UFunction>(DelegateNode->GetDelegateSignature(), DelegateTerm->Type.PinSubCategoryMemberReference);
+			DelegateTerm->Type.PinSubCategoryMemberReference.FillSimpleMemberReference<UFunction>(DelegateNode->GetDelegateSignature());
 			DelegateTerm->Source = Node;
 			DelegateTerm->Name = Context.NetNameMap->MakeValidName(Node) + TEXT("_TempBindingDelegate");
 			LocalDelegateMap.Add(Node, DelegateTerm);
@@ -194,7 +194,7 @@ UFunction* UK2Node_DelegateSet::GetDelegateSignature()
 	if( !DelegateProperty )
 	{
 		// Attempt to find a remapped delegate property
-		UMulticastDelegateProperty* NewProperty = Cast<UMulticastDelegateProperty>(FindRemappedField(DelegatePropertyClass, DelegatePropertyName));
+		UMulticastDelegateProperty* NewProperty = Cast<UMulticastDelegateProperty>(FMemberReference::FindRemappedField(DelegatePropertyClass, DelegatePropertyName));
 		if( NewProperty )
 		{
 			// Found a remapped property, update the node
@@ -214,7 +214,7 @@ UFunction* UK2Node_DelegateSet::GetDelegateSignature() const
 	if( !DelegateProperty )
 	{
 		// Attempt to find a remapped delegate property
-		DelegateProperty = Cast<UMulticastDelegateProperty>(FindRemappedField(DelegatePropertyClass, DelegatePropertyName));
+		DelegateProperty = Cast<UMulticastDelegateProperty>(FMemberReference::FindRemappedField(DelegatePropertyClass, DelegatePropertyName));
 	}
 
 	return (DelegateProperty != NULL) ? DelegateProperty->SignatureFunction : NULL;
diff --git a/Engine/Source/Editor/BlueprintGraph/Private/K2Node_Event.cpp b/Engine/Source/Editor/BlueprintGraph/Private/K2Node_Event.cpp
--- a/Engine/Source/Editor/BlueprintGraph/Private/K2Node_Event.cpp
+++ b/Engine/Source/Editor/BlueprintGraph/Private/K2Node_Event.cpp
@@ -76,7 +76,7 @@ FText UK2Node_Event::GetNodeTitle(ENodeTitleType::Type TitleType) const
 	{
 		FString FunctionName = EventReference.GetMemberName().ToString(); // If we fail to find the function, still want to write something on the node.
 
-		if (UFunction* Function = EventReference.ResolveMember<UFunction>(this))
+		if (UFunction* Function = EventReference.ResolveMember<UFunction>(GetBlueprintClassFromNode()))
 		{
 			FunctionName = UEdGraphSchema_K2::GetFriendlySignatureName(Function);
 		}
@@ -85,9 +85,9 @@ FText UK2Node_Event::GetNodeTitle(ENodeTitleType::Type TitleType) const
 		Args.Add(TEXT("FunctionName"), FText::FromString(FunctionName));
 		FText Title = FText::Format(NSLOCTEXT("K2Node", "Event_Name", "Event {FunctionName}"), Args);
 
-		if(TitleType == ENodeTitleType::FullTitle && EventReference.GetMemberParentClass(this) != nullptr && EventReference.GetMemberParentClass(this)->IsChildOf(UInterface::StaticClass()))
+		if(TitleType == ENodeTitleType::FullTitle && EventReference.GetMemberParentClass(GetBlueprintClassFromNode()) != nullptr && EventReference.GetMemberParentClass(GetBlueprintClassFromNode())->IsChildOf(UInterface::StaticClass()))
 		{
-			const FText SignatureClassAsText = FBlueprintEditorUtils::GetFriendlyClassDisplayName(EventReference.GetMemberParentClass(this));
+			const FText SignatureClassAsText = FBlueprintEditorUtils::GetFriendlyClassDisplayName(EventReference.GetMemberParentClass(GetBlueprintClassFromNode()));
 
 			FFormatNamedArguments FullTitleArgs;
 			FullTitleArgs.Add(TEXT("Title"), Title);
@@ -106,7 +106,7 @@ FText UK2Node_Event::GetNodeTitle(ENodeTitleType::Type TitleType) const
 
 FText UK2Node_Event::GetTooltipText() const
 {
-	UFunction* Function = EventReference.ResolveMember<UFunction>(this);
+	UFunction* Function = EventReference.ResolveMember<UFunction>(GetBlueprintClassFromNode());
 	if (CachedTooltip.IsOutOfDate() && (Function != nullptr))
 	{
 		CachedTooltip = FText::FromString(UK2Node_CallFunction::GetDefaultTooltipForFunction(Function));
@@ -145,7 +145,7 @@ FString UK2Node_Event::GetKeywords() const
 {
 	FString Keywords;
 
-	UFunction* Function = EventReference.ResolveMember<UFunction>(this);
+	UFunction* Function = EventReference.ResolveMember<UFunction>(GetBlueprintClassFromNode());
 	if (Function != NULL)
 	{
 		Keywords = UK2Node_CallFunction::GetKeywordsForFunction( Function );
@@ -156,7 +156,7 @@ FString UK2Node_Event::GetKeywords() const
 
 FString UK2Node_Event::GetDocumentationLink() const
 {
-	if ( UClass* EventSignatureClass = EventReference.GetMemberParentClass(this))
+	if ( UClass* EventSignatureClass = EventReference.GetMemberParentClass(GetBlueprintClassFromNode()))
 	{
 		return FString::Printf(TEXT("Shared/Types/%s%s"), EventSignatureClass->GetPrefixCPP(), *EventSignatureClass->GetName());
 	}
@@ -179,7 +179,7 @@ void UK2Node_Event::UpdateDelegatePin(bool bSilent)
 	UEdGraphPin* Pin = FindPinChecked(DelegateOutputName);
 	checkSlow(EGPD_Output == Pin->Direction);
 
-	const UObject* OldSignature = FMemberReference::ResolveSimpleMemberReference<UFunction>(Pin->PinType.PinSubCategoryMemberReference);
+	const UObject* OldSignature = Pin->PinType.PinSubCategoryMemberReference.ResolveSimpleMemberReference<UFunction>();
 	if (!OldSignature)
 	{
 		OldSignature = Pin->PinType.PinSubCategoryObject.Get();
@@ -188,7 +188,7 @@ void UK2Node_Event::UpdateDelegatePin(bool bSilent)
 	UFunction* NewSignature = NULL;
 	if(bOverrideFunction)
 	{
-		NewSignature = EventReference.ResolveMember<UFunction>(this);
+		NewSignature = EventReference.ResolveMember<UFunction>(GetBlueprintClassFromNode());
 	}
 	else if(UBlueprint* Blueprint = GetBlueprint())
 	{
@@ -198,7 +198,7 @@ void UK2Node_Event::UpdateDelegatePin(bool bSilent)
 	}
 
 	Pin->PinType.PinSubCategoryObject = NULL;
-	FMemberReference::FillSimpleMemberReference<UFunction>(NewSignature, Pin->PinType.PinSubCategoryMemberReference);
+	Pin->PinType.PinSubCategoryMemberReference.FillSimpleMemberReference<UFunction>(NewSignature);
 
 	if ((OldSignature != NewSignature) && !bSilent)
 	{
@@ -223,7 +223,7 @@ FName UK2Node_Event::GetFunctionName() const
 
 UFunction* UK2Node_Event::FindEventSignatureFunction()
 {
-	return EventReference.ResolveMember<UFunction>(this);
+	return EventReference.ResolveMember<UFunction>(GetBlueprintClassFromNode());
 }
 
 void UK2Node_Event::AllocateDefaultPins()
@@ -249,7 +249,7 @@ void UK2Node_Event::ValidateNodeDuringCompilation(class FCompilerResultsLog& Mes
 	Super::ValidateNodeDuringCompilation(MessageLog);
 
 	// If we are overriding a function, but we can;t find the function we are overriding, that is a compile error
-	if(bOverrideFunction && EventReference.ResolveMember<UFunction>(this) == nullptr)
+	if(bOverrideFunction && EventReference.ResolveMember<UFunction>(GetBlueprintClassFromNode()) == nullptr)
 	{
 		MessageLog.Error(*FString::Printf(*NSLOCTEXT("KismetCompiler", "MissingEventSig_Error", "Missing Event '%s' for @@").ToString(), *EventReference.GetMemberName().ToString()), this);
 	}
@@ -272,9 +272,9 @@ void UK2Node_Event::GetRedirectPinNames(const UEdGraphPin& Pin, TArray<FString>&
 		// first add functionname.param
 		RedirectPinNames.Add(FString::Printf(TEXT("%s.%s"), *EventReference.GetMemberName().ToString(), *OldPinName));
 		// if there is class, also add an option for class.functionname.param
-		if ( EventReference.ResolveMember<UFunction>(this) != nullptr )
+		if ( EventReference.ResolveMember<UFunction>(GetBlueprintClassFromNode()) != nullptr )
 		{
-			if(const UClass* EventSignatureClass = EventReference.GetMemberParentClass(this))
+			if(const UClass* EventSignatureClass = EventReference.GetMemberParentClass(GetBlueprintClassFromNode()))
 			{
 				RedirectPinNames.Add(FString::Printf(TEXT("%s.%s.%s"), *EventSignatureClass->GetName(), *EventReference.GetMemberName().ToString(), *OldPinName));
 			}
@@ -419,7 +419,7 @@ bool UK2Node_Event::CanPasteHere(const UEdGraph* TargetGraph) const
 				if(bOverrideFunction)
 				{
 					// If the function name is hidden by the parent class, don't paste this event
-					bDisallowPaste = EventReference.GetMemberParentClass(this) == Blueprint->ParentClass
+					bDisallowPaste = EventReference.GetMemberParentClass(GetBlueprintClassFromNode()) == Blueprint->ParentClass
 						&& ExcludedEventNames.Contains(EventReference.GetMemberName().ToString());
 					if(!bDisallowPaste)
 					{
@@ -428,7 +428,7 @@ bool UK2Node_Event::CanPasteHere(const UEdGraph* TargetGraph) const
 						{
 							bDisallowPaste = ExistingEventNodes[i]->bOverrideFunction
 								&& ExistingEventNodes[i]->EventReference.GetMemberName() == EventReference.GetMemberName()
-								&& ExistingEventNodes[i]->EventReference.GetMemberParentClass(this) == EventReference.GetMemberParentClass(this);
+								&& ExistingEventNodes[i]->EventReference.GetMemberParentClass(GetBlueprintClassFromNode()) == EventReference.GetMemberParentClass(GetBlueprintClassFromNode());
 						}
 
 						// We need to also check for 'const' BPIE methods that might already be implemented as functions with a read-only 'self' context (these were previously implemented as events)
@@ -452,11 +452,11 @@ bool UK2Node_Event::CanPasteHere(const UEdGraph* TargetGraph) const
 						if(!bDisallowPaste)
 						{
 							// If the signature class is not implemented by the Blueprint parent class or an interface, don't paste this event
-							bDisallowPaste = !Blueprint->ParentClass->IsChildOf(EventReference.GetMemberParentClass(this))
-								&& !ImplementedInterfaceClasses.Contains(EventReference.GetMemberParentClass(this));
+							bDisallowPaste = !Blueprint->ParentClass->IsChildOf(EventReference.GetMemberParentClass(GetBlueprintClassFromNode()))
+								&& !ImplementedInterfaceClasses.Contains(EventReference.GetMemberParentClass(GetBlueprintClassFromNode()));
 							if(bDisallowPaste)
 							{
-								UE_LOG(LogBlueprint, Log, TEXT("Cannot paste event node (%s) directly because the event signature class (%s) is incompatible with this Blueprint."), *GetFName().ToString(), EventReference.GetMemberParentClass(this) ? *EventReference.GetMemberParentClass(this)->GetFName().ToString() : TEXT("NONE"));
+								UE_LOG(LogBlueprint, Log, TEXT("Cannot paste event node (%s) directly because the event signature class (%s) is incompatible with this Blueprint."), *GetFName().ToString(), EventReference.GetMemberParentClass(GetBlueprintClassFromNode()) ? *EventReference.GetMemberParentClass(GetBlueprintClassFromNode())->GetFName().ToString() : TEXT("NONE"));
 							}
 						}
 						else
@@ -466,7 +466,7 @@ bool UK2Node_Event::CanPasteHere(const UEdGraph* TargetGraph) const
 					}
 					else
 					{
-						UE_LOG(LogBlueprint, Log, TEXT("Cannot paste event node (%s) directly because the event function (%s) is hidden by the Blueprint parent class (%s)."), *GetFName().ToString(), *EventReference.GetMemberName().ToString(), EventReference.GetMemberParentClass(this) ? *EventReference.GetMemberParentClass(this)->GetFName().ToString() : TEXT("NONE"));
+						UE_LOG(LogBlueprint, Log, TEXT("Cannot paste event node (%s) directly because the event function (%s) is hidden by the Blueprint parent class (%s)."), *GetFName().ToString(), *EventReference.GetMemberName().ToString(), EventReference.GetMemberParentClass(GetBlueprintClassFromNode()) ? *EventReference.GetMemberParentClass(GetBlueprintClassFromNode())->GetFName().ToString() : TEXT("NONE"));
 					}
 				}
 				else if(CustomFunctionName != NAME_None)
@@ -480,11 +480,11 @@ bool UK2Node_Event::CanPasteHere(const UEdGraph* TargetGraph) const
 						if(EventReference.GetMemberName() != NAME_None)
 						{
 							// If the signature class is not implemented by the Blueprint parent class or an interface, don't paste this event
-							bDisallowPaste = !Blueprint->ParentClass->IsChildOf(EventReference.GetMemberParentClass(this))
-								&& !ImplementedInterfaceClasses.Contains(EventReference.GetMemberParentClass(this));
+							bDisallowPaste = !Blueprint->ParentClass->IsChildOf(EventReference.GetMemberParentClass(GetBlueprintClassFromNode()))
+								&& !ImplementedInterfaceClasses.Contains(EventReference.GetMemberParentClass(GetBlueprintClassFromNode()));
 							if(bDisallowPaste)
 							{
-								UE_LOG(LogBlueprint, Log, TEXT("Cannot paste event node (%s) directly because the custom event function (%s) with event signature name (%s) has an event signature class (%s) that is incompatible with this Blueprint."), *GetFName().ToString(), *CustomFunctionName.ToString(), *EventReference.GetMemberName().ToString(), EventReference.GetMemberParentClass(this) ? *EventReference.GetMemberParentClass(this)->GetFName().ToString() : TEXT("NONE"));
+								UE_LOG(LogBlueprint, Log, TEXT("Cannot paste event node (%s) directly because the custom event function (%s) with event signature name (%s) has an event signature class (%s) that is incompatible with this Blueprint."), *GetFName().ToString(), *CustomFunctionName.ToString(), *EventReference.GetMemberName().ToString(), EventReference.GetMemberParentClass(GetBlueprintClassFromNode()) ? *EventReference.GetMemberParentClass(GetBlueprintClassFromNode())->GetFName().ToString() : TEXT("NONE"));
 							}
 						}
 					}
@@ -495,7 +495,7 @@ bool UK2Node_Event::CanPasteHere(const UEdGraph* TargetGraph) const
 				}
 				else
 				{
-					UE_LOG(LogBlueprint, Log, TEXT("Cannot paste event node (%s) directly because the event configuration is not specifically handled (EventSignatureName=%s, EventSignatureClass=%s)."), *GetFName().ToString(), *EventReference.GetMemberName().ToString(), EventReference.GetMemberParentClass(this) ? *EventReference.GetMemberParentClass(this)->GetFName().ToString() : TEXT("NONE"));
+					UE_LOG(LogBlueprint, Log, TEXT("Cannot paste event node (%s) directly because the event configuration is not specifically handled (EventSignatureName=%s, EventSignatureClass=%s)."), *GetFName().ToString(), *EventReference.GetMemberName().ToString(), EventReference.GetMemberParentClass(GetBlueprintClassFromNode()) ? *EventReference.GetMemberParentClass(GetBlueprintClassFromNode())->GetFName().ToString() : TEXT("NONE"));
 				}
 			}
 			else
@@ -587,7 +587,7 @@ void UK2Node_Event::ExpandNode(class FKismetCompilerContext& CompilerContext, UE
 
 FName UK2Node_Event::GetCornerIcon() const
 {
-	if (UFunction* Function = EventReference.ResolveMember<UFunction>(this))
+	if (UFunction* Function = EventReference.ResolveMember<UFunction>(GetBlueprintClassFromNode()))
 	{
 		if (bOverrideFunction || (CustomFunctionName == NAME_None))
 		{
@@ -608,7 +608,7 @@ FName UK2Node_Event::GetCornerIcon() const
 		return TEXT("Graph.Replication.AuthorityOnly");
 	}
 
-	if(EventReference.GetMemberParentClass(this) != nullptr && EventReference.GetMemberParentClass(this)->IsChildOf(UInterface::StaticClass()))
+	if(EventReference.GetMemberParentClass(GetBlueprintClassFromNode()) != nullptr && EventReference.GetMemberParentClass(GetBlueprintClassFromNode())->IsChildOf(UInterface::StaticClass()))
 	{
 		return TEXT("Graph.Event.InterfaceEventIcon");
 	}
@@ -621,7 +621,7 @@ FText UK2Node_Event::GetToolTipHeading() const
 	FText Heading = Super::GetToolTipHeading();
 
 	FText EventHeading = FText::GetEmpty();
-	if (UFunction* Function = EventReference.ResolveMember<UFunction>(this))
+	if (UFunction* Function = EventReference.ResolveMember<UFunction>(GetBlueprintClassFromNode()))
 	{
 		if (bOverrideFunction || (CustomFunctionName == NAME_None))
 		{
@@ -640,7 +640,7 @@ FText UK2Node_Event::GetToolTipHeading() const
 	{
 		EventHeading = LOCTEXT("ServerOnlyEvent", "Server Only");
 	}
-	else if(EventHeading.IsEmpty() && (EventReference.GetMemberParentClass(this) != nullptr) && EventReference.GetMemberParentClass(this)->IsChildOf(UInterface::StaticClass()))
+	else if(EventHeading.IsEmpty() && (EventReference.GetMemberParentClass(GetBlueprintClassFromNode()) != nullptr) && EventReference.GetMemberParentClass(GetBlueprintClassFromNode())->IsChildOf(UInterface::StaticClass()))
 	{
 		EventHeading = LOCTEXT("InterfaceEvent", "Interface Event");
 	}
@@ -667,7 +667,7 @@ void UK2Node_Event::GetNodeAttributes( TArray<TKeyValuePair<FString, FString>>&
 FText UK2Node_Event::GetMenuCategory() const
 {
 	FString FuncSubCategory;
-	if (UFunction* Function = EventReference.ResolveMember<UFunction>(this))
+	if (UFunction* Function = EventReference.ResolveMember<UFunction>(GetBlueprintClassFromNode()))
 	{
 		FuncSubCategory = UK2Node_CallFunction::GetDefaultCategoryForFunction(Function, TEXT(""));
 	}
@@ -678,7 +678,7 @@ FText UK2Node_Event::GetMenuCategory() const
 
 bool UK2Node_Event::IsDeprecated() const
 {
-	if (UFunction* Function = EventReference.ResolveMember<UFunction>(this))
+	if (UFunction* Function = EventReference.ResolveMember<UFunction>(GetBlueprintClassFromNode()))
 	{
 		return Function->HasMetaData(FBlueprintMetadata::MD_DeprecatedFunction);
 	}
@@ -688,7 +688,7 @@ bool UK2Node_Event::IsDeprecated() const
 
 FString UK2Node_Event::GetDeprecationMessage() const
 {
-	if (UFunction* Function = EventReference.ResolveMember<UFunction>(this))
+	if (UFunction* Function = EventReference.ResolveMember<UFunction>(GetBlueprintClassFromNode()))
 	{
 		if (Function->HasMetaData(FBlueprintMetadata::MD_DeprecationMessage))
 		{
@@ -701,9 +701,9 @@ FString UK2Node_Event::GetDeprecationMessage() const
 
 UObject* UK2Node_Event::GetJumpTargetForDoubleClick() const
 {
-	if(EventReference.GetMemberParentClass(this) != NULL && EventReference.GetMemberParentClass(this)->ClassGeneratedBy != NULL && EventReference.GetMemberParentClass(this)->ClassGeneratedBy->IsA(UBlueprint::StaticClass()))
+	if(EventReference.GetMemberParentClass(GetBlueprintClassFromNode()) != NULL && EventReference.GetMemberParentClass(GetBlueprintClassFromNode())->ClassGeneratedBy != NULL && EventReference.GetMemberParentClass(GetBlueprintClassFromNode())->ClassGeneratedBy->IsA(UBlueprint::StaticClass()))
 	{
-		UBlueprint* Blueprint = CastChecked<UBlueprint>(EventReference.GetMemberParentClass(this)->ClassGeneratedBy);
+		UBlueprint* Blueprint = CastChecked<UBlueprint>(EventReference.GetMemberParentClass(GetBlueprintClassFromNode())->ClassGeneratedBy);
 		TArray<UEdGraph*> Graphs;
 		Blueprint->GetAllGraphs(Graphs);
 		for(auto It(Graphs.CreateConstIterator()); It; It++)
diff --git a/Engine/Source/Editor/BlueprintGraph/Private/K2Node_MCDelegate.cpp b/Engine/Source/Editor/BlueprintGraph/Private/K2Node_MCDelegate.cpp
--- a/Engine/Source/Editor/BlueprintGraph/Private/K2Node_MCDelegate.cpp
+++ b/Engine/Source/Editor/BlueprintGraph/Private/K2Node_MCDelegate.cpp
@@ -63,7 +63,7 @@ void UK2Node_BaseMCDelegate::AllocateDefaultPins()
 	CreatePin(EGPD_Input, K2Schema->PC_Exec, TEXT(""), NULL, false, false, K2Schema->PN_Execute);
 	CreatePin(EGPD_Output, K2Schema->PC_Exec, TEXT(""), NULL, false, false, K2Schema->PN_Then);
 
-	UClass* PropertyOwnerClass = DelegateReference.GetMemberParentClass(this);
+	UClass* PropertyOwnerClass = DelegateReference.GetMemberParentClass(GetBlueprintClassFromNode());
 	if (PropertyOwnerClass != nullptr)
 	{
 		PropertyOwnerClass = PropertyOwnerClass->GetAuthoritativeClass();
@@ -90,7 +90,7 @@ void UK2Node_BaseMCDelegate::AllocateDefaultPins()
 
 UFunction* UK2Node_BaseMCDelegate::GetDelegateSignature(bool bForceNotFromSkelClass) const
 {
-	UClass* OwnerClass = DelegateReference.GetMemberParentClass(this);
+	UClass* OwnerClass = DelegateReference.GetMemberParentClass(GetBlueprintClassFromNode());
 	if (bForceNotFromSkelClass)
 	{
 		OwnerClass = (OwnerClass != nullptr) ? OwnerClass->GetAuthoritativeClass() : nullptr;
@@ -114,7 +114,7 @@ UFunction* UK2Node_BaseMCDelegate::GetDelegateSignature(bool bForceNotFromSkelCl
 	}
 	ReferenceToUse.SetDirect(DelegateReference.GetMemberName(), DelegateGuid, OwnerClass, /*bIsConsideredSelfContext =*/false);
 
-	UMulticastDelegateProperty* DelegateProperty = ReferenceToUse.ResolveMember<UMulticastDelegateProperty>(this);
+	UMulticastDelegateProperty* DelegateProperty = ReferenceToUse.ResolveMember<UMulticastDelegateProperty>(GetBlueprintClassFromNode());
 	return (DelegateProperty != NULL) ? DelegateProperty->SignatureFunction : NULL;
 }
 
@@ -139,7 +139,7 @@ FString UK2Node_BaseMCDelegate::GetDocumentationLink() const
 	}
 	else 
 	{
-		ParentClass = DelegateReference.GetMemberParentClass(this);
+		ParentClass = DelegateReference.GetMemberParentClass(GetBlueprintClassFromNode());
 	}
 
 	if ( ParentClass != NULL )
@@ -187,14 +187,14 @@ void UK2Node_BaseMCDelegate::ExpandNode(class FKismetCompilerContext& CompilerCo
 
 bool UK2Node_BaseMCDelegate::IsAuthorityOnly() const
 {
-	const UMulticastDelegateProperty* DelegateProperty = DelegateReference.ResolveMember<UMulticastDelegateProperty>(this);
+	const UMulticastDelegateProperty* DelegateProperty = DelegateReference.ResolveMember<UMulticastDelegateProperty>(GetBlueprintClassFromNode());
 	return DelegateProperty && DelegateProperty->HasAnyPropertyFlags(CPF_BlueprintAuthorityOnly);
 
 }
 
 bool UK2Node_BaseMCDelegate::HasExternalBlueprintDependencies(TArray<class UStruct*>* OptionalOutput) const
 {
-	const UClass* SourceClass = DelegateReference.GetMemberParentClass(this);
+	const UClass* SourceClass = DelegateReference.GetMemberParentClass(GetBlueprintClassFromNode());
 	const UBlueprint* SourceBlueprint = GetBlueprint();
 	const bool bResult = (SourceClass != NULL) && (SourceClass->ClassGeneratedBy != NULL) && (SourceClass->ClassGeneratedBy != SourceBlueprint);
 	if (bResult && OptionalOutput)
@@ -225,7 +225,7 @@ void UK2Node_BaseMCDelegate::AutowireNewNode(UEdGraphPin* FromPin)
 		{
 			if (FromPin->PinType.PinSubCategoryObject.IsValid() && FromPin->PinType.PinSubCategoryObject->IsA(UClass::StaticClass()))
 			{
-				UProperty* DelegateProperty = DelegateReference.ResolveMember<UProperty>(this);
+				UProperty* DelegateProperty = DelegateReference.ResolveMember<UProperty>(GetBlueprintClassFromNode());
 				if (DelegateProperty)
 				{
 					UClass* DelegateOwner = DelegateProperty->GetOwnerClass();
@@ -275,7 +275,7 @@ void UK2Node_AddDelegate::AllocateDefaultPins()
 
 	if(UEdGraphPin* DelegatePin = CreatePin(EGPD_Input, K2Schema->PC_Delegate, TEXT(""), NULL, false, true, FK2Node_BaseMCDelegateHelper::DelegatePinName, true))
 	{
-		FMemberReference::FillSimpleMemberReference<UFunction>(GetDelegateSignature(), DelegatePin->PinType.PinSubCategoryMemberReference);
+		DelegatePin->PinType.PinSubCategoryMemberReference.FillSimpleMemberReference<UFunction>(GetDelegateSignature());
 		DelegatePin->PinFriendlyName = NSLOCTEXT("K2Node", "PinFriendlyDelegatetName", "Event");
 	}
 }
@@ -343,7 +343,7 @@ void UK2Node_RemoveDelegate::AllocateDefaultPins()
 
 	if(UEdGraphPin* DelegatePin = CreatePin(EGPD_Input, K2Schema->PC_Delegate, TEXT(""), NULL, false, true, FK2Node_BaseMCDelegateHelper::DelegatePinName, true))
 	{
-		FMemberReference::FillSimpleMemberReference<UFunction>(GetDelegateSignature(), DelegatePin->PinType.PinSubCategoryMemberReference);
+		DelegatePin->PinType.PinSubCategoryMemberReference.FillSimpleMemberReference<UFunction>(GetDelegateSignature());
 		DelegatePin->PinFriendlyName = NSLOCTEXT("K2Node", "PinFriendlyDelegatetName", "Event");
 	}
 }
diff --git a/Engine/Source/Editor/BlueprintGraph/Private/K2Node_Message.cpp b/Engine/Source/Editor/BlueprintGraph/Private/K2Node_Message.cpp
--- a/Engine/Source/Editor/BlueprintGraph/Private/K2Node_Message.cpp
+++ b/Engine/Source/Editor/BlueprintGraph/Private/K2Node_Message.cpp
@@ -103,7 +103,7 @@ void UK2Node_Message::ExpandNode(class FKismetCompilerContext& CompilerContext,
 	if (bExecPinConnected)
 	{
 		// Make sure our interface is valid
-		if (FunctionReference.GetMemberParentClass(this) == NULL)
+		if (FunctionReference.GetMemberParentClass(GetBlueprintClassFromNode()) == NULL)
 		{
 			CompilerContext.MessageLog.Error(*FString::Printf(*LOCTEXT("MessageNodeInvalid_Error", "Message node @@ has an invalid interface.").ToString()), this);
 			return;
@@ -113,7 +113,7 @@ void UK2Node_Message::ExpandNode(class FKismetCompilerContext& CompilerContext,
 		if (MessageNodeFunction == NULL)
 		{
 			//@TODO: Why do this here in teh compiler, it's already done on AllocateDefaultPins() during on-load node reconstruction
-			MessageNodeFunction = Cast<UFunction>(UK2Node::FindRemappedField(FunctionReference.GetMemberParentClass(this), FunctionReference.GetMemberName()));
+			MessageNodeFunction = Cast<UFunction>(FMemberReference::FindRemappedField(FunctionReference.GetMemberParentClass(GetBlueprintClassFromNode()), FunctionReference.GetMemberName()));
 		}
 
 		if (MessageNodeFunction == NULL)
diff --git a/Engine/Source/Editor/BlueprintGraph/Private/K2Node_Variable.cpp b/Engine/Source/Editor/BlueprintGraph/Private/K2Node_Variable.cpp
--- a/Engine/Source/Editor/BlueprintGraph/Private/K2Node_Variable.cpp
+++ b/Engine/Source/Editor/BlueprintGraph/Private/K2Node_Variable.cpp
@@ -103,14 +103,14 @@ void UK2Node_Variable::CreatePinForSelf()
 		if( !VariableReference.IsLocalScope() )
 		{
 			bool bSelfTarget = VariableReference.IsSelfContext() && (ESelfContextInfo::NotSelfContext != SelfContextInfo);
-			UClass* MemberParentClass = VariableReference.GetMemberParentClass(this);
+			UClass* MemberParentClass = VariableReference.GetMemberParentClass(GetBlueprintClassFromNode());
 			UClass* TargetClass = MemberParentClass;
 			
 			// Self Target pins should always make the class be the owning class of the property,
 			// so if the node is from a Macro Blueprint, it will hook up as self in any placed Blueprint
 			if(bSelfTarget)
 			{
-				if(UProperty* Property = VariableReference.ResolveMember<UProperty>(this))
+				if(UProperty* Property = VariableReference.ResolveMember<UProperty>(GetBlueprintClassFromNode()))
 				{
 					TargetClass = Property->GetOwnerClass()->GetAuthoritativeClass();
 				}
@@ -282,7 +282,7 @@ UK2Node::ERedirectType UK2Node_Variable::DoPinsMatchForReconstruction( const UEd
 
 UClass* UK2Node_Variable::GetVariableSourceClass() const
 {
-	UClass* Result = VariableReference.GetMemberParentClass(this);
+	UClass* Result = VariableReference.GetMemberParentClass(GetBlueprintClassFromNode());
 	return Result;
 }
 
@@ -291,7 +291,7 @@ UProperty* UK2Node_Variable::GetPropertyForVariable() const
 	const FName VarName = GetVarName();
 	UEdGraphPin* VariablePin = FindPin(GetVarNameString());
 
-	UProperty* VariableProperty = VariableReference.ResolveMember<UProperty>(this);
+	UProperty* VariableProperty = VariableReference.ResolveMember<UProperty>(GetBlueprintClassFromNode());
 
 	// if the variable has been deprecated, don't use it
 	if(VariableProperty != NULL)
@@ -362,7 +362,7 @@ FName UK2Node_Variable::GetPaletteIcon(FLinearColor& ColorOut) const
 
 	if(VariableReference.IsLocalScope())
 	{
-		ReturnIconName = GetVariableIconAndColor(VariableReference.GetMemberScope(this), GetVarName(), ColorOut);
+		ReturnIconName = GetVariableIconAndColor(VariableReference.GetMemberScope(GetBlueprintClassFromNode()), GetVarName(), ColorOut);
 	}
 	else
 	{
@@ -398,7 +398,7 @@ FText UK2Node_Variable::GetToolTipHeading() const
 {
 	FText Heading = Super::GetToolTipHeading();
 
-	UProperty const* VariableProperty = VariableReference.ResolveMember<UProperty>(this);
+	UProperty const* VariableProperty = VariableReference.ResolveMember<UProperty>(GetBlueprintClassFromNode());
 	if (VariableProperty && VariableProperty->HasAllPropertyFlags(CPF_Net))
 	{
 		FText ReplicatedTag = LOCTEXT("ReplicatedVar", "Replicated");
@@ -449,7 +449,7 @@ FName UK2Node_Variable::GetVariableIconAndColor(const UStruct* VarScope, FName V
 
 void UK2Node_Variable::CheckForErrors(const UEdGraphSchema_K2* Schema, FCompilerResultsLog& MessageLog)
 {
-	if(!VariableReference.IsSelfContext() && VariableReference.GetMemberParentClass(this) != NULL)
+	if(!VariableReference.IsSelfContext() && VariableReference.GetMemberParentClass(GetBlueprintClassFromNode()) != NULL)
 	{
 		// Check to see if we're not a self context, if we have a valid context.  It may have been purged because of a dead execution chain
 		UEdGraphPin* ContextPin = Schema->FindSelfPin(*this, EGPD_Input);
@@ -576,7 +576,7 @@ bool UK2Node_Variable::RemapRestrictedLinkReference(FName OldVariableName, FName
 
 FName UK2Node_Variable::GetCornerIcon() const
 {
-	const UProperty* VariableProperty = VariableReference.ResolveMember<UProperty>(this);
+	const UProperty* VariableProperty = VariableReference.ResolveMember<UProperty>(GetBlueprintClassFromNode());
 	if (VariableProperty && VariableProperty->HasAllPropertyFlags(CPF_Net))
 	{
 		return TEXT("Graph.Replication.Replicated");
@@ -634,7 +634,7 @@ void UK2Node_Variable::AutowireNewNode(UEdGraphPin* FromPin)
 			// If the source pin has a valid PinSubCategoryObject, we might be doing BP Comms, so check if it is a class
 			if(FromPin->PinType.PinSubCategoryObject.IsValid() && FromPin->PinType.PinSubCategoryObject->IsA(UClass::StaticClass()))
 			{
-				UProperty* VariableProperty = VariableReference.ResolveMember<UProperty>(this);
+				UProperty* VariableProperty = VariableReference.ResolveMember<UProperty>(GetBlueprintClassFromNode());
 				if(VariableProperty)
 				{
 					UClass* PropertyOwner = VariableProperty->GetOwnerClass();
@@ -674,7 +674,7 @@ void UK2Node_Variable::AutowireNewNode(UEdGraphPin* FromPin)
 FBPVariableDescription const* UK2Node_Variable::GetBlueprintVarDescription() const
 {
 	FName const& VarName = VariableReference.GetMemberName();
-	UStruct const* VariableScope = VariableReference.GetMemberScope(this);
+	UStruct const* VariableScope = VariableReference.GetMemberScope(GetBlueprintClassFromNode());
 
 	bool const bIsLocalVariable = (VariableScope != nullptr);
 	if (bIsLocalVariable)
@@ -701,7 +701,7 @@ bool UK2Node_Variable::CanPasteHere(const UEdGraph* TargetGraph) const
 	if ( FBlueprintEditorUtils::FindBlueprintForGraph(TargetGraph)->BlueprintType == BPTYPE_MacroLibrary && VariableReference.IsSelfContext() )
 	{
 		// Self variables must be from a parent class to the macro BP
-		if(UProperty* Property = VariableReference.ResolveMember<UProperty>(this))
+		if(UProperty* Property = VariableReference.ResolveMember<UProperty>(GetBlueprintClassFromNode()))
 		{
 			const UClass* CurrentClass = GetBlueprint()->SkeletonGeneratedClass->GetAuthoritativeClass();
 			const UClass* PropertyClass = Property->GetOwnerClass()->GetAuthoritativeClass();
@@ -728,7 +728,7 @@ void UK2Node_Variable::PostPasteNode()
 	{
 		// Local scoped variables should always validate whether they are being placed in the same graph as their scope
 		// ResolveMember will not return nullptr when the graph changes but the Blueprint remains the same.
-		UEdGraph* ScopeGraph = FBlueprintEditorUtils::FindScopeGraph(Blueprint, VariableReference.GetMemberScope(this));
+		UEdGraph* ScopeGraph = FBlueprintEditorUtils::FindScopeGraph(Blueprint, VariableReference.GetMemberScope(GetBlueprintClassFromNode()));
 		if(ScopeGraph != GetGraph())
 		{
 			bInvalidateVariable = true;
diff --git a/Engine/Source/Editor/Kismet/Private/BlueprintDetailsCustomization.cpp b/Engine/Source/Editor/Kismet/Private/BlueprintDetailsCustomization.cpp
--- a/Engine/Source/Editor/Kismet/Private/BlueprintDetailsCustomization.cpp
+++ b/Engine/Source/Editor/Kismet/Private/BlueprintDetailsCustomization.cpp
@@ -3175,7 +3175,7 @@ struct FBasePinChangeHelper
 					UBlueprint* CallSiteBlueprint = FBlueprintEditorUtils::FindBlueprintForNode(CallSite);
 
 					const bool bNameMatches = (CallSite->FunctionReference.GetMemberName() == FunctionDefNode->SignatureName);
-					const UClass* MemberParentClass = CallSite->FunctionReference.GetMemberParentClass(CallSite);
+					const UClass* MemberParentClass = CallSite->FunctionReference.GetMemberParentClass(CallSite->GetBlueprintClassFromNode());
 					const bool bClassMatchesEasy = (MemberParentClass != NULL) && (MemberParentClass->IsChildOf(FunctionDefNode->SignatureClass));
 					const bool bClassMatchesHard = (CallSiteBlueprint != NULL) && (CallSite->FunctionReference.IsSelfContext()) && (FunctionDefNode->SignatureClass == NULL) && (CallSiteBlueprint == BlueprintGraphActionDetails.GetBlueprintObj());
 					const bool bValidSchema = CallSite->GetSchema() != NULL;
diff --git a/Engine/Source/Editor/Kismet/Private/BlueprintEditor.cpp b/Engine/Source/Editor/Kismet/Private/BlueprintEditor.cpp
--- a/Engine/Source/Editor/Kismet/Private/BlueprintEditor.cpp
+++ b/Engine/Source/Editor/Kismet/Private/BlueprintEditor.cpp
@@ -7092,7 +7092,7 @@ bool FBlueprintEditor::IsSelectionNativeFunction()
 
 			if( FunctionNode && SelectedNodes.Num() == 1 )
 			{
-				UFunction* FunctionPtr = FunctionNode->FunctionReference.ResolveMember<UFunction>( FunctionNode );
+				UFunction* FunctionPtr = FunctionNode->FunctionReference.ResolveMember<UFunction>( FunctionNode->GetBlueprintClassFromNode() );
 				UClass* OwningClass = FunctionPtr ? FunctionPtr->GetOuterUClass() : NULL;
 
 				if( OwningClass && OwningClass->HasAllClassFlags( CLASS_Native ))
@@ -7116,7 +7116,7 @@ void FBlueprintEditor::GotoNativeVariableDefinition()
 
 		if( VarNode )
 		{
-			UProperty* VariableProperty = VarNode->VariableReference.ResolveMember<UProperty>( VarNode );
+			UProperty* VariableProperty = VarNode->VariableReference.ResolveMember<UProperty>( VarNode->GetBlueprintClassFromNode() );
 
 			if( VariableProperty )
 			{
@@ -7147,7 +7147,7 @@ bool FBlueprintEditor::IsSelectionNativeVariable()
 
 			if( VarNode && SelectedNodes.Num() == 1 )
 			{
-				UProperty* VariableProperty = VarNode->VariableReference.ResolveMember<UProperty>( VarNode );
+				UProperty* VariableProperty = VarNode->VariableReference.ResolveMember<UProperty>( VarNode->GetBlueprintClassFromNode() );
 
 				if( VariableProperty && VariableProperty->HasAllFlags( RF_Native ))
 				{
diff --git a/Engine/Source/Editor/Kismet/Private/BlueprintPaletteFavorites.cpp b/Engine/Source/Editor/Kismet/Private/BlueprintPaletteFavorites.cpp
--- a/Engine/Source/Editor/Kismet/Private/BlueprintPaletteFavorites.cpp
+++ b/Engine/Source/Editor/Kismet/Private/BlueprintPaletteFavorites.cpp
@@ -69,7 +69,7 @@ static FBlueprintNodeSignature BlueprintPaletteFavoritesImpl::ConstructLegacySig
 		// to discern between them
 		if (UK2Node_CallFunction const* CallFuncNode = Cast<UK2Node_CallFunction const>(NodeTemplate))
 		{
-			SignatureSubObject = CallFuncNode->FunctionReference.ResolveMember<UFunction>(CallFuncNode);
+			SignatureSubObject = CallFuncNode->FunctionReference.ResolveMember<UFunction>(CallFuncNode->GetBlueprintClassFromNode());
 			bIsSupported = (SignatureSubObject != nullptr);
 		}
 		else if (UK2Node_InputAxisEvent const* InputAxisEventNode = Cast<UK2Node_InputAxisEvent const>(NodeTemplate))
@@ -79,7 +79,7 @@ static FBlueprintNodeSignature BlueprintPaletteFavoritesImpl::ConstructLegacySig
 		}
 		else if (UK2Node_Event const* EventNode = Cast<UK2Node_Event const>(NodeTemplate))
 		{
-			SignatureSubObject = EventNode->EventReference.ResolveMember<UFunction>(EventNode);
+			SignatureSubObject = EventNode->EventReference.ResolveMember<UFunction>(EventNode->GetBlueprintClassFromNode());
 			bIsSupported = (SignatureSubObject != nullptr);
 		}
 		else if (UK2Node_MacroInstance const* MacroNode = Cast<UK2Node_MacroInstance const>(NodeTemplate))
diff --git a/Engine/Source/Editor/Kismet/Private/SBlueprintPalette.cpp b/Engine/Source/Editor/Kismet/Private/SBlueprintPalette.cpp
--- a/Engine/Source/Editor/Kismet/Private/SBlueprintPalette.cpp
+++ b/Engine/Source/Editor/Kismet/Private/SBlueprintPalette.cpp
@@ -1342,7 +1342,7 @@ FText SBlueprintPaletteItem::GetToolTipText() const
 
 			if (UK2Node_CallFunction const* CallFuncNode = Cast<UK2Node_CallFunction const>(NodeTemplate))
 			{			
-				if(UClass* ParentClass = CallFuncNode->FunctionReference.GetMemberParentClass(CallFuncNode))
+				if(UClass* ParentClass = CallFuncNode->FunctionReference.GetMemberParentClass(CallFuncNode->GetBlueprintClassFromNode()))
 				{
 					UBlueprint* BlueprintObj = UBlueprint::GetBlueprintFromClass(ParentClass);
 					ClassDisplayName = (BlueprintObj != NULL) ? BlueprintObj->GetName() : ParentClass->GetName();
diff --git a/Engine/Source/Editor/Kismet/Private/SBlueprintSubPalette.cpp b/Engine/Source/Editor/Kismet/Private/SBlueprintSubPalette.cpp
--- a/Engine/Source/Editor/Kismet/Private/SBlueprintSubPalette.cpp
+++ b/Engine/Source/Editor/Kismet/Private/SBlueprintSubPalette.cpp
@@ -70,7 +70,7 @@ static bool CanPaletteItemBePlaced(TSharedPtr<FEdGraphSchemaAction> DropActionIn
 		{
 			FName const FuncName = CallFuncNode->FunctionReference.GetMemberName();
 			check(FuncName != NAME_None);
-			UClass const* const FuncOwner = CallFuncNode->FunctionReference.GetMemberParentClass(CallFuncNode);
+			UClass const* const FuncOwner = CallFuncNode->FunctionReference.GetMemberParentClass(CallFuncNode->GetBlueprintClassFromNode());
 			check(FuncOwner != nullptr);
 
 			UFunction* const Function = FindField<UFunction>(FuncOwner, FuncName);
diff --git a/Engine/Source/Editor/KismetCompiler/Private/KismetCompiler.cpp b/Engine/Source/Editor/KismetCompiler/Private/KismetCompiler.cpp
--- a/Engine/Source/Editor/KismetCompiler/Private/KismetCompiler.cpp
+++ b/Engine/Source/Editor/KismetCompiler/Private/KismetCompiler.cpp
@@ -2391,7 +2391,7 @@ void FKismetCompilerContext::CreateFunctionStubForEvent(UK2Node_Event* SrcEventN
 	// Create an entry point
 	UK2Node_FunctionEntry* EntryNode = SpawnIntermediateNode<UK2Node_FunctionEntry>(SrcEventNode, ChildStubGraph);
 	EntryNode->NodePosX = -200;
-	EntryNode->SignatureClass = SrcEventNode->EventReference.GetMemberParentClass(SrcEventNode);
+	EntryNode->SignatureClass = SrcEventNode->EventReference.GetMemberParentClass(SrcEventNode->GetBlueprintClassFromNode());
 	EntryNode->SignatureName = SrcEventNode->EventReference.GetMemberName();
 	EntryNode->CustomGeneratedFunctionName = EventNodeName;
 
@@ -2588,7 +2588,7 @@ void FKismetCompilerContext::VerifyValidOverrideEvent(const UEdGraph* Graph)
 			{
 				const UK2Node_Event* EventNode = EntryPoints[EntryPointsIdx];
 				if( EventNode && EventNode->bOverrideFunction &&
-					(EventNode->EventReference.GetMemberParentClass(EventNode) == FuncClass) &&
+					(EventNode->EventReference.GetMemberParentClass(EventNode->GetBlueprintClassFromNode()) == FuncClass) &&
 					(EventNode->EventReference.GetMemberName() == FuncName))
 				{
 					if (EventNode->IsDeprecated())
diff --git a/Engine/Source/Editor/KismetCompiler/Private/KismetCompilerMisc.cpp b/Engine/Source/Editor/KismetCompiler/Private/KismetCompilerMisc.cpp
--- a/Engine/Source/Editor/KismetCompiler/Private/KismetCompilerMisc.cpp
+++ b/Engine/Source/Editor/KismetCompiler/Private/KismetCompilerMisc.cpp
@@ -158,7 +158,7 @@ bool FKismetCompilerUtilities::IsTypeCompatibleWithProperty(UEdGraphPin* SourceP
 	}
 	else if (PinCategory == Schema->PC_Delegate)
 	{
-		const UFunction* SignatureFunction = FMemberReference::ResolveSimpleMemberReference<UFunction>(Type.PinSubCategoryMemberReference);
+		const UFunction* SignatureFunction = Type.PinSubCategoryMemberReference.ResolveSimpleMemberReference<UFunction>();
 		const UDelegateProperty* PropertyDelegate = Cast<const UDelegateProperty>(TestProperty);
 		bTypeMismatch = !(SignatureFunction 
 			&& PropertyDelegate 
@@ -851,7 +851,7 @@ UProperty* FKismetCompilerUtilities::CreatePropertyOnScope(UStruct* Scope, const
 	}
 	else if (Type.PinCategory == Schema->PC_Delegate)
 	{
-		if (UFunction* SignatureFunction = FMemberReference::ResolveSimpleMemberReference<UFunction>(Type.PinSubCategoryMemberReference))
+		if (UFunction* SignatureFunction = Type.PinSubCategoryMemberReference.ResolveSimpleMemberReference<UFunction>())
 		{
 			UDelegateProperty* NewPropertyDelegate = NewObject<UDelegateProperty>(PropertyScope, ValidatedPropertyName, ObjectFlags);
 			NewPropertyDelegate->SignatureFunction = SignatureFunction;
@@ -860,7 +860,7 @@ UProperty* FKismetCompilerUtilities::CreatePropertyOnScope(UStruct* Scope, const
 	}
 	else if (Type.PinCategory == Schema->PC_MCDelegate)
 	{
-		UFunction* const SignatureFunction = FMemberReference::ResolveSimpleMemberReference<UFunction>(Type.PinSubCategoryMemberReference);
+		UFunction* const SignatureFunction = Type.PinSubCategoryMemberReference.ResolveSimpleMemberReference<UFunction>();
 		UMulticastDelegateProperty* NewPropertyDelegate = NewObject<UMulticastDelegateProperty>(PropertyScope, ValidatedPropertyName, ObjectFlags);
 		NewPropertyDelegate->SignatureFunction = SignatureFunction;
 		NewProperty = NewPropertyDelegate;
diff --git a/Engine/Source/Editor/UnrealEd/Private/Kismet2/BlueprintEditorUtils.cpp b/Engine/Source/Editor/UnrealEd/Private/Kismet2/BlueprintEditorUtils.cpp
--- a/Engine/Source/Editor/UnrealEd/Private/Kismet2/BlueprintEditorUtils.cpp
+++ b/Engine/Source/Editor/UnrealEd/Private/Kismet2/BlueprintEditorUtils.cpp
@@ -841,7 +841,7 @@ struct FRegenerationHelper
 							// Thus, we handle that here in order to ensure that all 'Inner' fields are also valid before class regeneration.
 							if (auto VariableNode = Cast<UK2Node_Variable>(Node))
 							{
-								UArrayProperty* ArrayProperty = Cast<UArrayProperty>(VariableNode->VariableReference.ResolveMember<UProperty>(Node));
+								UArrayProperty* ArrayProperty = Cast<UArrayProperty>(VariableNode->VariableReference.ResolveMember<UProperty>(Node->GetBlueprintClassFromNode()));
 								if (ArrayProperty != nullptr && ArrayProperty->Inner != nullptr && ArrayProperty->Inner->HasAnyFlags(RF_NeedLoad|RF_WasLoaded))
 								{
 									ForcedLoad(ArrayProperty->Inner);
@@ -1573,7 +1573,7 @@ void FBlueprintEditorUtils::PostDuplicateBlueprint(UBlueprint* Blueprint)
 						{
 							UClass* TargetClass = nullptr;
 
-							if(UProperty* Property = VariableNode->VariableReference.ResolveMember<UProperty>(VariableNode))
+							if(UProperty* Property = VariableNode->VariableReference.ResolveMember<UProperty>(VariableNode->GetBlueprintClassFromNode()))
 							{
 								TargetClass = Property->GetOwnerClass()->GetAuthoritativeClass();
 							}
@@ -2244,7 +2244,7 @@ UK2Node_Event* FBlueprintEditorUtils::FindOverrideForFunction(const UBlueprint*
 		UK2Node_Event* EventNode = AllEvents[i];
 		check(EventNode);
 		if(	EventNode->bOverrideFunction == true &&
-			EventNode->EventReference.GetMemberParentClass(EventNode) == SignatureClass &&
+			EventNode->EventReference.GetMemberParentClass(EventNode->GetBlueprintClassFromNode()) == SignatureClass &&
 			EventNode->EventReference.GetMemberName() == SignatureName )
 		{
 			return EventNode;
@@ -4759,7 +4759,7 @@ void FBlueprintEditorUtils::RemoveInterface(UBlueprint* Blueprint, const FName&
 		for(TArray<UK2Node_Event*>::TIterator NodeIt(AllEvents); NodeIt; ++NodeIt)
 		{
 			UK2Node_Event* EventNode = *NodeIt;
-			if( EventNode->EventReference.GetMemberParentClass(EventNode) == InterfaceClass )
+			if( EventNode->EventReference.GetMemberParentClass(EventNode->GetBlueprintClassFromNode()) == InterfaceClass )
 			{
 				if(bPreserveFunctions)
 				{
@@ -4980,7 +4980,7 @@ void FBlueprintEditorUtils::ConformImplementedEvents(UBlueprint* Blueprint)
 				// If the event is loaded and is not a custom event
 				if(!EventNode->HasAnyFlags(RF_NeedLoad|RF_NeedPostLoad) && EventNode->bOverrideFunction)
 				{
-					const bool bEventNodeUsedByInterface = ImplementedInterfaceClasses.Contains(EventNode->EventReference.GetMemberParentClass(EventNode));
+					const bool bEventNodeUsedByInterface = ImplementedInterfaceClasses.Contains(EventNode->EventReference.GetMemberParentClass(EventNode->GetBlueprintClassFromNode()));
 					if (Blueprint->GeneratedClass && !bEventNodeUsedByInterface)
 					{
 						// See if the generated class implements an event with the given function signature
@@ -4988,7 +4988,7 @@ void FBlueprintEditorUtils::ConformImplementedEvents(UBlueprint* Blueprint)
 						if (TargetFunction || EventGraphNames.Contains(EventNode->EventReference.GetMemberName()))
 						{
 							// The generated class implements the event but the function signature is not up-to-date
-							if (!Blueprint->GeneratedClass->IsChildOf(EventNode->EventReference.GetMemberParentClass(EventNode)))
+							if (!Blueprint->GeneratedClass->IsChildOf(EventNode->EventReference.GetMemberParentClass(EventNode->GetBlueprintClassFromNode())))
 							{
 								FFormatNamedArguments Args;
 								Args.Add(TEXT("NodeTitle"), EventNode->GetNodeTitle(ENodeTitleType::ListView));
@@ -5086,7 +5086,7 @@ static void ConformInterfaceByName(UBlueprint* Blueprint, FBPInterfaceDescriptio
 		for (UK2Node_Event* EventNode : ImplementedEvents)
 		{
 			// if this event belongs to something other than this interface
-			if (EventNode->EventReference.GetMemberParentClass(EventNode) != CurrentInterfaceDesc.Interface)
+			if (EventNode->EventReference.GetMemberParentClass(EventNode->GetBlueprintClassFromNode()) != CurrentInterfaceDesc.Interface)
 			{
 				continue;
 			}
diff --git a/Engine/Source/Editor/UnrealEd/Private/Tests/EditorBuildPromotionTests.cpp b/Engine/Source/Editor/UnrealEd/Private/Tests/EditorBuildPromotionTests.cpp
--- a/Engine/Source/Editor/UnrealEd/Private/Tests/EditorBuildPromotionTests.cpp
+++ b/Engine/Source/Editor/UnrealEd/Private/Tests/EditorBuildPromotionTests.cpp
@@ -709,7 +709,7 @@ namespace EditorBuildPromotionTestUtils
 		NewEventNode->bOverrideFunction = true;
 
 		//Check for existing events
-		UK2Node_Event* ExistingEvent = FBlueprintEditorUtils::FindOverrideForFunction(InBlueprint, NewEventNode->EventReference.GetMemberParentClass(NewEventNode), NewEventNode->EventReference.GetMemberName());
+		UK2Node_Event* ExistingEvent = FBlueprintEditorUtils::FindOverrideForFunction(InBlueprint, NewEventNode->EventReference.GetMemberParentClass(NewEventNode->GetBlueprintClassFromNode()), NewEventNode->EventReference.GetMemberName());
 
 		if (!ExistingEvent)
 		{
diff --git a/Engine/Source/Runtime/Engine/Classes/EdGraph/EdGraphPin.h b/Engine/Source/Runtime/Engine/Classes/EdGraph/EdGraphPin.h
--- a/Engine/Source/Runtime/Engine/Classes/EdGraph/EdGraphPin.h
+++ b/Engine/Source/Runtime/Engine/Classes/EdGraph/EdGraphPin.h
@@ -3,38 +3,9 @@
 #pragma once
 
 #include "EdGraph/EdGraphNode.h"
+#include "Engine/MemberReference.h"
 #include "EdGraphPin.generated.h"
 
-USTRUCT()
-struct FSimpleMemberReference
-{
-	GENERATED_USTRUCT_BODY()
-
-	/** Class that this member is defined in. */
-	UPROPERTY()
-	TSubclassOf<class UObject> MemberParentClass;
-
-	/** Name of the member */
-	UPROPERTY()
-	FName MemberName;
-
-	/** The Guid of the member */
-	UPROPERTY()
-	FGuid MemberGuid;
-
-	void Reset()
-	{
-		operator=(FSimpleMemberReference());
-	}
-
-	bool operator==(const FSimpleMemberReference& Other) const
-	{
-		return (MemberParentClass == Other.MemberParentClass)
-			&& (MemberName == Other.MemberName)
-			&& (MemberGuid == Other.MemberGuid);
-	}
-};
-
 FORCEINLINE FArchive& operator<<(FArchive& Ar, FSimpleMemberReference& Data)
 {
 	Ar << Data.MemberParentClass;
diff --git a/Engine/Source/Runtime/Engine/Classes/Engine/MemberReference.h b/Engine/Source/Runtime/Engine/Classes/Engine/MemberReference.h
new file mode 100644
--- /dev/null
+++ b/Engine/Source/Runtime/Engine/Classes/Engine/MemberReference.h
@@ -0,0 +1,428 @@
+// Copyright 1998-2015 Epic Games, Inc. All Rights Reserved.
+
+#pragma once
+
+#include "MemberReference.generated.h"
+
+/** Helper struct to allow us to redirect properties and functions through renames and additionally between classes if necessary */
+struct FFieldRemapInfo
+{
+	/** The new name of the field after being renamed */
+	FName FieldName;
+
+	/** The new name of the field's outer class if different from its original location, or NAME_None if it hasn't moved */
+	FName FieldClass;
+
+	bool operator==(const FFieldRemapInfo& Other) const
+	{
+		return FieldName == Other.FieldName && FieldClass == Other.FieldClass;
+	}
+
+	friend uint32 GetTypeHash(const FFieldRemapInfo& RemapInfo)
+	{
+		return GetTypeHash(RemapInfo.FieldName) + GetTypeHash(RemapInfo.FieldClass) * 23;
+	}
+
+	FFieldRemapInfo()
+		: FieldName(NAME_None)
+		, FieldClass(NAME_None)
+	{
+	}
+};
+
+/** Helper struct to allow us to redirect pin name for node class */
+struct FParamRemapInfo
+{
+	bool	bCustomValueMapping;
+	FName	OldParam;
+	FName	NewParam;
+	FName	NodeTitle;
+	TMap<FString, FString> ParamValueMap;
+
+	// constructor
+	FParamRemapInfo()
+		: OldParam(NAME_None)
+		, NewParam(NAME_None)
+		, NodeTitle(NAME_None)
+	{
+	}
+};
+
+USTRUCT()
+struct FMemberReference
+{
+	GENERATED_USTRUCT_BODY()
+
+protected:
+	/** Class that this member is defined in. Should be NULL if bSelfContext is true.  */
+	UPROPERTY()
+	mutable TSubclassOf<class UObject> MemberParentClass;
+
+	/** Class that this member is defined in. Should be NULL if bSelfContext is true.  */
+	UPROPERTY()
+	mutable FString MemberScope;
+
+	/** Name of variable */
+	UPROPERTY()
+	mutable FName MemberName;
+
+	/** The Guid of the variable */
+	UPROPERTY()
+	mutable FGuid MemberGuid;
+
+	/** Whether or not this should be a "self" context */
+	UPROPERTY()
+	mutable bool bSelfContext;
+
+	/** Whether or not this property has been deprecated */
+	UPROPERTY()
+	mutable bool bWasDeprecated;
+	
+public:
+	FMemberReference()
+		: MemberParentClass(NULL)
+		, MemberName(NAME_None)
+		, bSelfContext(false)
+	{
+	}
+
+	/** Set up this reference from a supplied field */
+	template<class TFieldType>
+	void SetFromField(const UField* InField, const bool bIsConsideredSelfContext)
+	{
+		MemberParentClass = bIsConsideredSelfContext ? NULL : InField->GetOwnerClass();
+		MemberName = InField->GetFName();
+		bSelfContext = bIsConsideredSelfContext;
+		bWasDeprecated = false;
+
+		if (MemberParentClass != nullptr)
+		{
+			MemberParentClass = MemberParentClass->GetAuthoritativeClass();
+		}
+
+		MemberGuid.Invalidate();
+		if (InField->GetOwnerClass())
+		{
+			UBlueprint::GetGuidFromClassByFieldName<TFieldType>(InField->GetOwnerClass(), InField->GetFName(), MemberGuid);
+		}
+	}
+
+	template<class TFieldType>
+	void SetFromField(const UField* InField, UClass* SelfScope)
+	{
+		FGuid FieldGuid;
+		if (InField->GetOwnerClass())
+		{
+			UBlueprint::GetGuidFromClassByFieldName<TFieldType>(InField->GetOwnerClass(), InField->GetFName(), FieldGuid);
+		}
+
+		SetGivenSelfScope(InField->GetFName(), FieldGuid, InField->GetOwnerClass(), SelfScope);
+	}
+
+	/** Update given a new self */
+	template<class TFieldType>
+	void RefreshGivenNewSelfScope(UClass* SelfScope)
+	{
+		if ((MemberParentClass != NULL) && (SelfScope != NULL))
+		{
+			UBlueprint::GetGuidFromClassByFieldName<TFieldType>((MemberParentClass ? *MemberParentClass : SelfScope), MemberName, MemberGuid);
+			SetGivenSelfScope(MemberName, MemberGuid, MemberParentClass, SelfScope);
+		}
+		else
+		{
+			// We no longer have enough information to known if we've done the right thing, and just have to hope...
+		}
+	}
+
+	/** Set to a non-'self' member, so must include reference to class owning the member. */
+	ENGINE_API void SetExternalMember(FName InMemberName, TSubclassOf<class UObject> InMemberParentClass);
+
+	/** Set to a non-'self' delegate member, this is not self-context but is not given a parent class */
+	ENGINE_API void SetExternalDelegateMember(FName InMemberName);
+
+	/** Set up this reference to a 'self' member name */
+	ENGINE_API void SetSelfMember(FName InMemberName);
+
+	/** Set up this reference to a 'self' member name, scoped to a struct */
+	ENGINE_API void SetLocalMember(FName InMemberName, UStruct* InScope, const FGuid InMemberGuid);
+
+	/** Set up this reference to a 'self' member name, scoped to a struct name */
+	ENGINE_API void SetLocalMember(FName InMemberName, FString InScopeName, const FGuid InMemberGuid);
+
+	/** Only intended for backwards compat! */
+	ENGINE_API void SetDirect(const FName InMemberName, const FGuid InMemberGuid, TSubclassOf<class UObject> InMemberParentClass, bool bIsConsideredSelfContext);
+
+	/** Invalidate the current MemberParentClass, if this is a self scope, or the MemberScope if it is not (and set).  Intended for PostDuplication fixups */
+	ENGINE_API void InvalidateScope();
+
+	/** Get the name of this member */
+	FName GetMemberName() const
+	{
+		return MemberName;
+	}
+
+	FGuid GetMemberGuid() const
+	{
+		return MemberGuid;
+	}
+
+	UClass* GetMemberParentClass()
+	{
+		return *MemberParentClass;
+	}
+
+	/** Returns if this is a 'self' context. */
+	bool IsSelfContext() const
+	{
+		return bSelfContext;
+	}
+
+	/** Returns if this is a local scope. */
+	bool IsLocalScope() const
+	{
+		return !MemberScope.IsEmpty();
+	}
+private:
+	/**
+	 * Refreshes a local variable reference name if it has changed
+	 *
+	 * @param InSelfScope		Scope to lookup the variable in, to see if it has changed
+	 */
+	ENGINE_API FName RefreshLocalVariableName(UClass* InSelfScope) const;
+
+protected:
+	/** Only intended for backwards compat! */
+	ENGINE_API void SetGivenSelfScope(const FName InMemberName, const FGuid InMemberGuid, TSubclassOf<class UObject> InMemberParentClass, TSubclassOf<class UObject> SelfScope) const;
+
+public:
+
+	/** Get the class that owns this member */
+	UClass* GetMemberParentClass(UClass* SelfScope) const
+	{
+		// Local variables with a MemberScope act much the same as being SelfContext, their parent class is SelfScope.
+		return (bSelfContext || !MemberScope.IsEmpty())? SelfScope : *MemberParentClass;
+	}
+
+	/** Get the scope of this member */
+	UStruct* GetMemberScope(UClass* InMemberParentClass) const
+	{
+		return FindField<UStruct>(InMemberParentClass, *MemberScope);
+	}
+
+	/** Get the name of the scope of this member */
+	FString GetMemberScopeName() const
+	{
+		return MemberScope;
+	}
+
+	/** Returns whether or not the variable has been deprecated */
+	bool IsDeprecated() const
+	{
+		return bWasDeprecated;
+	}
+
+	/** 
+	 *	Returns the member UProperty/UFunction this reference is pointing to, or NULL if it no longer exists 
+	 *	Derives 'self' scope from supplied Blueprint node if required
+	 *	Will check for redirects and fix itself up if one is found.
+	 */
+	template<class TFieldType>
+	TFieldType* ResolveMember(UClass* SelfScope) const
+	{
+		TFieldType* ReturnField = NULL;
+
+		if(bSelfContext && SelfScope == NULL)
+		{
+			UE_LOG(LogBlueprint, Warning, TEXT("FMemberReference::ResolveMember (%s) bSelfContext == true, but no scope supplied!"), *MemberName.ToString() );
+		}
+
+		// Check if the member reference is function scoped
+		if(IsLocalScope())
+		{
+			UStruct* MemberScopeStruct = FindField<UStruct>(SelfScope, *MemberScope);
+
+			// Find in target scope
+			ReturnField = FindField<TFieldType>(MemberScopeStruct, MemberName);
+
+			if(ReturnField == NULL)
+			{
+				// If the property was not found, refresh the local variable name and try again
+				const FName RenamedMemberName = RefreshLocalVariableName(SelfScope);
+				if (RenamedMemberName != NAME_None)
+				{
+					ReturnField = FindField<TFieldType>(MemberScopeStruct, MemberName);
+				}
+			}
+		}
+		else
+		{
+			// Look for remapped member
+			UClass* TargetScope = bSelfContext ? SelfScope : (UClass*)MemberParentClass;
+			if( TargetScope != NULL &&  !GIsSavingPackage )
+			{
+				ReturnField = Cast<TFieldType>(FindRemappedField(TargetScope, MemberName, true));
+			}
+
+			if(ReturnField != NULL)
+			{
+				// Fix up this struct, we found a redirect
+				MemberName = ReturnField->GetFName();
+				MemberParentClass = Cast<UClass>(ReturnField->GetOuter());
+
+				MemberGuid.Invalidate();
+				UBlueprint::GetGuidFromClassByFieldName<TFieldType>(TargetScope, MemberName, MemberGuid);
+
+				if (MemberParentClass != nullptr)
+				{
+					MemberParentClass = MemberParentClass->GetAuthoritativeClass();
+					// Re-evaluate self-ness against the redirect if we were given a valid SelfScope
+					if (SelfScope != NULL)
+					{
+						SetGivenSelfScope(MemberName, MemberGuid, MemberParentClass, SelfScope);
+					}
+				}	
+			}
+			else if(TargetScope != NULL)
+			{
+				// Find in target scope
+				ReturnField = FindField<TFieldType>(TargetScope, MemberName);
+
+				// If we have a GUID find the reference variable and make sure the name is up to date and find the field again
+				// For now only variable references will have valid GUIDs.  Will have to deal with finding other names subsequently
+				if (ReturnField == NULL && MemberGuid.IsValid())
+				{
+					const FName RenamedMemberName = UBlueprint::GetFieldNameFromClassByGuid<TFieldType>(TargetScope, MemberGuid);
+					if (RenamedMemberName != NAME_None)
+					{
+						MemberName = RenamedMemberName;
+						ReturnField = FindField<TFieldType>(TargetScope, MemberName);
+					}
+				}
+			}
+			else
+			{
+				// Native delegate signatures can have null target scope.
+                // Needed because Blueprint pin types are referencing delegate signatures as
+                // class members, which doesn't have to be true.
+				if (TFieldType::StaticClass() == UFunction::StaticClass())
+				{
+					ReturnField = (TFieldType*)FindDelegateSignature(MemberName);
+				}
+			}
+		}
+
+		// Check to see if the member has been deprecated
+		if (UProperty* Property = Cast<UProperty>(ReturnField))
+		{
+			bWasDeprecated = Property->HasAnyPropertyFlags(CPF_Deprecated);
+		}
+
+		return ReturnField;
+	}
+
+	template<class TFieldType>
+	TFieldType* ResolveMember(UBlueprint* SelfScope)
+	{
+		return ResolveMember<TFieldType>(SelfScope->SkeletonGeneratedClass);
+	}
+
+	/**
+	 * Searches the field redirect map for the specified named field in the scope, and returns the remapped field if found
+	 *
+	 * @param	InitialScope	The scope the field was initially defined in.  The function will search up into parent scopes to attempt to find remappings
+	 * @param	InitialName		The name of the field to attempt to find a redirector for
+	 * @param	bInitialScopeMustBeOwnerOfField		if true the InitialScope must be Child of the field's owner
+	 * @return	The remapped field, if one exists
+	 */
+	ENGINE_API static UField* FindRemappedField(UClass* InitialScope, FName InitialName, bool bInitialScopeMustBeOwnerOfField = false);
+
+	ENGINE_API static const TMap<FFieldRemapInfo, FFieldRemapInfo>& GetFieldRedirectMap() { InitFieldRedirectMap(); return FieldRedirectMap; }
+	ENGINE_API static const TMultiMap<UClass*, FParamRemapInfo>& GetParamRedirectMap() { InitFieldRedirectMap(); return ParamRedirectMap; }
+
+protected:
+	/** 
+	 * A mapping from old property and function names to new ones.  Get primed from INI files, and should contain entries for properties, functions, and delegates that get moved, so they can be fixed up
+	 */
+	static TMap<FFieldRemapInfo, FFieldRemapInfo> FieldRedirectMap;
+	/** 
+	 * A mapping from old pin name to new pin name for each K2 node.  Get primed from INI files, and should contain entries for node class, and old param name and new param name
+	 */
+	static TMultiMap<UClass*, FParamRemapInfo> ParamRedirectMap;
+
+	/** Has the field map been intialized this run */
+	static bool bFieldRedirectMapInitialized;
+
+	/** Init the field redirect map (if not already done) from .ini file entries */
+	static void InitFieldRedirectMap();
+
+	/** 
+	 * Searches the field replacement map for an appropriately named field in the specified scope, and returns an updated name if this property has been listed in the remapping table
+	 *
+	 * @param	Class		The class scope to search for a field remap on
+	 * @param	FieldName	The field name (function name or property name) to search for
+	 * @param	RemapInfo	(out) Struct containing updated location info for the field
+	 * @return	Whether or not a remap was found in the specified scope
+	 */
+	static bool FindReplacementFieldName(UClass* Class, FName FieldName, FFieldRemapInfo& RemapInfo);
+
+};
+
+USTRUCT()
+struct FSimpleMemberReference
+{
+	GENERATED_USTRUCT_BODY()
+
+	FSimpleMemberReference()
+	{
+	}
+
+	/** Class that this member is defined in. */
+	UPROPERTY()
+	TSubclassOf<class UObject> MemberParentClass;
+
+	/** Name of the member */
+	UPROPERTY()
+	FName MemberName;
+
+	/** The Guid of the member */
+	UPROPERTY()
+	FGuid MemberGuid;
+
+	void Reset()
+	{
+		operator=(FSimpleMemberReference());
+	}
+
+	bool operator==(const FSimpleMemberReference& Other) const
+	{
+		return (MemberParentClass == Other.MemberParentClass)
+			&& (MemberName == Other.MemberName)
+			&& (MemberGuid == Other.MemberGuid);
+	}
+
+	template<class TFieldType>
+	void FillSimpleMemberReference(const UField* InField)
+	{
+		if (InField)
+		{
+			FMemberReference TempMemberReference;
+			TempMemberReference.SetFromField<TFieldType>(InField, false);
+
+			MemberName = TempMemberReference.GetMemberName();
+			MemberParentClass = TempMemberReference.GetMemberParentClass();
+			MemberGuid = TempMemberReference.GetMemberGuid();
+		}
+		else
+		{
+			Reset();
+		}
+	}
+
+	template<class TFieldType>
+	TFieldType* ResolveSimpleMemberReference() const
+	{
+		FMemberReference TempMemberReference;
+		TempMemberReference.SetDirect(MemberName, MemberGuid, MemberParentClass, false);
+		return TempMemberReference.ResolveMember<TFieldType>((UClass*)NULL);
+	}
+};
\ No newline at end of file
diff --git a/Engine/Source/Runtime/Engine/Private/MemberReference.cpp b/Engine/Source/Runtime/Engine/Private/MemberReference.cpp
new file mode 100644
--- /dev/null
+++ b/Engine/Source/Runtime/Engine/Private/MemberReference.cpp
@@ -0,0 +1,320 @@
+// Copyright 1998-2015 Epic Games, Inc. All Rights Reserved.
+
+#include "EnginePrivate.h"
+
+#include "Engine/MemberReference.h"
+#include "Kismet2/BlueprintEditorUtils.h"
+
+//////////////////////////////////////////////////////////////////////////
+// FMemberReference
+
+void FMemberReference::SetExternalMember(FName InMemberName, TSubclassOf<class UObject> InMemberParentClass)
+{
+	MemberName = InMemberName;
+	MemberParentClass = (InMemberParentClass != nullptr) ? InMemberParentClass->GetAuthoritativeClass() : nullptr;
+	MemberScope.Empty();
+	bSelfContext = false;
+	bWasDeprecated = false;
+}
+
+void FMemberReference::SetExternalDelegateMember(FName InMemberName)
+{
+	SetExternalMember(InMemberName, nullptr);
+}
+
+void FMemberReference::SetSelfMember(FName InMemberName)
+{
+	MemberName = InMemberName;
+	MemberParentClass = NULL;
+	MemberScope.Empty();
+	bSelfContext = true;
+	bWasDeprecated = false;
+}
+
+void FMemberReference::SetDirect(const FName InMemberName, const FGuid InMemberGuid, TSubclassOf<class UObject> InMemberParentClass, bool bIsConsideredSelfContext)
+{
+	MemberName = InMemberName;
+	MemberGuid = InMemberGuid;
+	bSelfContext = bIsConsideredSelfContext;
+	bWasDeprecated = false;
+	MemberParentClass = InMemberParentClass;
+	MemberScope.Empty();
+}
+
+void FMemberReference::SetGivenSelfScope(const FName InMemberName, const FGuid InMemberGuid, TSubclassOf<class UObject> InMemberParentClass, TSubclassOf<class UObject> SelfScope) const
+{
+	MemberName = InMemberName;
+	MemberGuid = InMemberGuid;
+	MemberParentClass = (InMemberParentClass != nullptr) ? InMemberParentClass->GetAuthoritativeClass() : nullptr;
+	MemberScope.Empty();
+	bSelfContext = (SelfScope->IsChildOf(InMemberParentClass)) || (SelfScope->ClassGeneratedBy == InMemberParentClass->ClassGeneratedBy);
+	bWasDeprecated = false;
+
+	if (bSelfContext)
+	{
+		MemberParentClass = NULL;
+	}
+}
+
+void FMemberReference::SetLocalMember(FName InMemberName, UStruct* InScope, const FGuid InMemberGuid)
+{
+	SetLocalMember(InMemberName, InScope->GetName(), InMemberGuid);
+}
+
+void FMemberReference::SetLocalMember(FName InMemberName, FString InScopeName, const FGuid InMemberGuid)
+{
+	MemberName = InMemberName;
+	MemberScope = InScopeName;
+	MemberGuid = InMemberGuid;
+	bSelfContext = false;
+}
+
+void FMemberReference::InvalidateScope()
+{
+	if( IsSelfContext() )
+	{
+		MemberParentClass = NULL;
+	}
+	else if(IsLocalScope())
+	{
+		MemberScope.Empty();
+
+		// Make it into a member reference since we are clearing the local context
+		bSelfContext = true;
+	}
+}
+
+FName FMemberReference::RefreshLocalVariableName(UClass* InSelfScope) const
+{
+	TArray<UBlueprint*> Blueprints;
+	UBlueprint::GetBlueprintHierarchyFromClass(InSelfScope, Blueprints);
+
+	FName RenamedMemberName = NAME_None;
+	for (int32 BPIndex = 0; BPIndex < Blueprints.Num(); ++BPIndex)
+	{
+		RenamedMemberName = FBlueprintEditorUtils::FindLocalVariableNameByGuid(Blueprints[BPIndex], MemberGuid);
+		if (RenamedMemberName != NAME_None)
+		{
+			MemberName = RenamedMemberName;
+			break;
+		}
+	}
+	return RenamedMemberName;
+}
+
+TMap<FFieldRemapInfo, FFieldRemapInfo> FMemberReference::FieldRedirectMap;
+TMultiMap<UClass*, FParamRemapInfo> FMemberReference::ParamRedirectMap;
+
+bool FMemberReference::bFieldRedirectMapInitialized = false;
+
+bool FMemberReference::FindReplacementFieldName(UClass* Class, FName FieldName, FFieldRemapInfo& RemapInfo)
+{
+	check(Class != NULL);
+	InitFieldRedirectMap();
+
+	// Reset the property remap info
+	RemapInfo = FFieldRemapInfo();
+
+	FFieldRemapInfo OldField;
+	OldField.FieldClass = Class->GetFName();
+	OldField.FieldName = FieldName;
+
+	FFieldRemapInfo* NewFieldInfoPtr = FieldRedirectMap.Find(OldField);
+	if (NewFieldInfoPtr != NULL)
+	{
+		RemapInfo = *NewFieldInfoPtr;
+		return true;
+	}
+	else
+	{
+		return false;
+	}
+}
+
+void FMemberReference::InitFieldRedirectMap()
+{
+	if (!bFieldRedirectMapInitialized)
+	{
+		if (GConfig)
+		{
+			FConfigSection* PackageRedirects = GConfig->GetSectionPrivate( TEXT("/Script/Engine.Engine"), false, true, GEngineIni );
+			for (FConfigSection::TIterator It(*PackageRedirects); It; ++It)
+			{
+				if (It.Key() == TEXT("K2FieldRedirects"))
+				{
+					FString OldFieldPathString;
+					FString NewFieldPathString;
+
+					FParse::Value( *It.Value(), TEXT("OldFieldName="), OldFieldPathString );
+					FParse::Value( *It.Value(), TEXT("NewFieldName="), NewFieldPathString );
+
+					// Handle both cases of just a field being renamed (just one FName), as well as a class and field name (ClassName.FieldName)
+					FFieldRemapInfo OldFieldRemap;
+					{
+						TArray<FString> OldFieldPath;
+						OldFieldPathString.ParseIntoArray(OldFieldPath, TEXT("."), true);
+
+						if (OldFieldPath.Num() == 1)
+						{
+							// Only the new property name is specified
+							OldFieldRemap.FieldName = FName(*OldFieldPath[0]);
+						}
+						else if (OldFieldPath.Num() == 2)
+						{
+							// Property name and new class are specified
+							OldFieldRemap.FieldClass = FName(*OldFieldPath[0]);
+							OldFieldRemap.FieldName = FName(*OldFieldPath[1]);
+						}
+					}
+
+					// Handle both cases of just a field being renamed (just one FName), as well as a class and field name (ClassName.FieldName)
+					FFieldRemapInfo NewFieldRemap;
+					{
+						TArray<FString> NewFieldPath;
+						NewFieldPathString.ParseIntoArray(NewFieldPath, TEXT("."), true);
+
+						if( NewFieldPath.Num() == 1 )
+						{
+							// Only the new property name is specified
+							NewFieldRemap.FieldName = FName(*NewFieldPath[0]);
+						}
+						else if( NewFieldPath.Num() == 2 )
+						{
+							// Property name and new class are specified
+							NewFieldRemap.FieldClass = FName(*NewFieldPath[0]);
+							NewFieldRemap.FieldName = FName(*NewFieldPath[1]);
+						}
+					}
+
+					FieldRedirectMap.Add(OldFieldRemap, NewFieldRemap);
+				}			
+				if (It.Key() == TEXT("K2ParamRedirects"))
+				{
+					FName NodeName = NAME_None;
+					FName OldParam = NAME_None;
+					FName NewParam = NAME_None;
+					FName NodeTitle = NAME_None;
+
+					FString OldParamValues;
+					FString NewParamValues;
+					FString CustomValueMapping;
+
+					FParse::Value( *It.Value(), TEXT("NodeName="), NodeName );
+					FParse::Value( *It.Value(), TEXT("NodeTitle="), NodeTitle );
+					FParse::Value( *It.Value(), TEXT("OldParamName="), OldParam );
+					FParse::Value( *It.Value(), TEXT("NewParamName="), NewParam );
+					FParse::Value( *It.Value(), TEXT("OldParamValues="), OldParamValues );
+					FParse::Value( *It.Value(), TEXT("NewParamValues="), NewParamValues );
+					FParse::Value( *It.Value(), TEXT("CustomValueMapping="), CustomValueMapping );
+
+					FParamRemapInfo ParamRemap;
+					ParamRemap.NewParam = NewParam;
+					ParamRemap.OldParam = OldParam;
+					ParamRemap.NodeTitle = NodeTitle;
+					ParamRemap.bCustomValueMapping = (FCString::Stricmp(*CustomValueMapping,TEXT("true")) == 0);
+
+					if (CustomValueMapping.Len() > 0 && !ParamRemap.bCustomValueMapping)
+					{
+						UE_LOG(LogBlueprint, Warning, TEXT("Value other than true specified for CustomValueMapping for '%s' node param redirect '%s' to '%s'."), *(NodeName.ToString()), *(OldParam.ToString()), *(NewParam.ToString()));
+					}
+
+					TArray<FString> OldParamValuesList;
+					TArray<FString> NewParamValuesList;
+					OldParamValues.ParseIntoArray(OldParamValuesList, TEXT(";"), false);
+					NewParamValues.ParseIntoArray(NewParamValuesList, TEXT(";"), false);
+
+					if (OldParamValuesList.Num() != NewParamValuesList.Num())
+					{
+						UE_LOG(LogBlueprint, Warning, TEXT("Unequal lengths for old and new param values for '%s' node param redirect '%s' to '%s'."), *(NodeName.ToString()), *(OldParam.ToString()), *(NewParam.ToString()));
+					}
+
+					if (CustomValueMapping.Len() > 0 && (OldParamValuesList.Num() > 0 || NewParamValuesList.Num() > 0))
+					{
+						UE_LOG(LogBlueprint, Warning, TEXT("Both Custom and Automatic param value remapping specified for '%s' node param redirect '%s' to '%s'.  Only Custom will be applied."), *(NodeName.ToString()), *(OldParam.ToString()), *(NewParam.ToString()));
+					}
+
+					for (int32 i = FMath::Min(OldParamValuesList.Num(), NewParamValuesList.Num()) - 1; i >= 0; --i)
+					{
+						int32 CurSize = ParamRemap.ParamValueMap.Num();
+						ParamRemap.ParamValueMap.Add(OldParamValuesList[i], NewParamValuesList[i]);
+						if (CurSize == ParamRemap.ParamValueMap.Num())
+						{
+							UE_LOG(LogBlueprint, Warning, TEXT("Duplicate old param value '%s' for '%s' node param redirect '%s' to '%s'."), *(OldParamValuesList[i]), *(NodeName.ToString()), *(OldParam.ToString()), *(NewParam.ToString()));
+						}
+					}
+
+					// load class
+					static UClass* K2NodeClass = StaticLoadClass(UObject::StaticClass(), nullptr, TEXT("BlueprintGraph.K2Node"));
+					if (K2NodeClass)
+					{
+						UClass* NodeClass = StaticLoadClass(K2NodeClass, nullptr, *NodeName.ToString());
+						if (NodeClass )
+						{
+							ParamRedirectMap.Add(NodeClass, ParamRemap);
+						}
+					}
+				}			
+			}
+
+			bFieldRedirectMapInitialized = true;
+		}
+	}
+}
+
+UField* FMemberReference::FindRemappedField(UClass* InitialScope, FName InitialName, bool bInitialScopeMustBeOwnerOfField)
+{
+	FFieldRemapInfo NewFieldInfo;
+
+	bool bFoundReplacement = false;
+	
+	// Step up the class chain to check if us or any of our parents specify a redirect
+	UClass* TestRemapClass = InitialScope;
+	while( TestRemapClass != NULL )
+	{
+		if( FindReplacementFieldName(TestRemapClass, InitialName, NewFieldInfo) )
+		{
+			// Found it, stop our search
+			bFoundReplacement = true;
+			break;
+		}
+
+		TestRemapClass = TestRemapClass->GetSuperClass();
+	}
+
+	// In the case of a bifurcation of a variable (e.g. moved from a parent into certain children), verify that we don't also define the variable in the current scope first
+	if( bFoundReplacement && (FindField<UField>(InitialScope, InitialName) != nullptr))
+	{
+		bFoundReplacement = false;		
+	}
+
+	if( bFoundReplacement )
+	{
+		const FName NewFieldName = NewFieldInfo.FieldName;
+		UClass* SearchClass = (NewFieldInfo.FieldClass != NAME_None) ? (UClass*)StaticFindObject(UClass::StaticClass(), ANY_PACKAGE, *NewFieldInfo.FieldClass.ToString()) : (UClass*)TestRemapClass;
+
+		// Find the actual field specified by the redirector, so we can return it and update the node that uses it
+		UField* NewField = FindField<UField>(SearchClass, NewFieldInfo.FieldName);
+		if( NewField != NULL )
+		{
+			if (bInitialScopeMustBeOwnerOfField && !InitialScope->IsChildOf(SearchClass))
+			{
+				UE_LOG(LogBlueprint, Log, TEXT("UK2Node:  Unable to update field. Remapped field '%s' in not owned by given scope. Scope: '%s', Owner: '%s'."), *InitialName.ToString(), *InitialScope->GetName(), *NewFieldInfo.FieldClass.ToString());
+			}
+			else
+			{
+				UE_LOG(LogBlueprint, Log, TEXT("UK2Node:  Fixed up old field '%s' to new name '%s' on class '%s'."), *InitialName.ToString(), *NewFieldInfo.FieldName.ToString(), *SearchClass->GetName());
+				return NewField;
+			}
+		}
+		else if (SearchClass != NULL)
+		{
+			UE_LOG(LogBlueprint, Log, TEXT("UK2Node:  Unable to find updated field name for '%s' on class '%s'."), *InitialName.ToString(), *SearchClass->GetName());
+		}
+		else
+		{
+			UE_LOG(LogBlueprint, Log, TEXT("UK2Node:  Unable to find updated field name for '%s' on unknown class '%s'."), *InitialName.ToString(), *NewFieldInfo.FieldClass.ToString());
+		}
+	}
+
+	return NULL;
+}
\ No newline at end of file
