UE4, BP: Crash in BlueprintActionDatabase after HotReload

diff --git a/Engine/Source/Editor/BlueprintGraph/Private/BlueprintActionDatabase.cpp b/Engine/Source/Editor/BlueprintGraph/Private/BlueprintActionDatabase.cpp
--- a/Engine/Source/Editor/BlueprintGraph/Private/BlueprintActionDatabase.cpp
+++ b/Engine/Source/Editor/BlueprintGraph/Private/BlueprintActionDatabase.cpp
@@ -1028,6 +1028,12 @@ void FBlueprintActionDatabase::RefreshAll()
 {
 	TGuardValue<bool> ScopedInitialization(BlueprintActionDatabaseImpl::bIsInitializing, true);
 
+	// Remove callbacks from blueprints
+	for (TObjectIterator<UBlueprint> BlueprintIt; BlueprintIt; ++BlueprintIt)
+	{
+		ClearAssetActions(*BlueprintIt);
+	}
+
 	ActionRegistry.Empty();
 	for (TObjectIterator<UClass> ClassIt; ClassIt; ++ClassIt)
 	{
@@ -1041,6 +1047,8 @@ void FBlueprintActionDatabase::RefreshAll()
 		BlueprintActionDatabaseImpl::AddSkeletonActions(**SkeletonIt, ClassActionList);
 	}
 
+	FComponentTypeRegistry::Get().SubscribeToComponentList(ComponentTypes).RemoveAll(this);
+
 	// this handles creating entries for components that were loaded before the database was alive:
 	FComponentTypeRegistry::Get().SubscribeToComponentList(ComponentTypes).AddRaw(this, &FBlueprintActionDatabase::RefreshComponentActions);
 	RefreshComponentActions();
