UE-10718 Set Member node should not set newly created member variables by default

In "Set member in .." node pins for properties are disabled by default.

diff --git a/Engine/Source/Editor/BlueprintGraph/Classes/K2Node_SetFieldsInStruct.h b/Engine/Source/Editor/BlueprintGraph/Classes/K2Node_SetFieldsInStruct.h
--- a/Engine/Source/Editor/BlueprintGraph/Classes/K2Node_SetFieldsInStruct.h
+++ b/Engine/Source/Editor/BlueprintGraph/Classes/K2Node_SetFieldsInStruct.h
@@ -35,6 +35,16 @@ class UK2Node_SetFieldsInStruct : public UK2Node_MakeStruct
 	BLUEPRINTGRAPH_API bool AllPinsAreShown() const;
 	BLUEPRINTGRAPH_API void RestoreAllPins();
 
+protected:
+	struct FSetFieldsInStructPinManager : public FMakeStructPinManager
+	{
+	public:
+		FSetFieldsInStructPinManager(const uint8* InSampleStructMemory) : FMakeStructPinManager(InSampleStructMemory) 
+		{}
+
+		virtual void GetRecordDefaults(UProperty* TestProperty, FOptionalPinFromProperty& Record) const override;
+	};
+
 private:
 	/** Constructing FText strings can be costly, so we cache the node's title/tooltip */
 	FNodeTextCache CachedTooltip;
diff --git a/Engine/Source/Editor/BlueprintGraph/Private/K2Node_SetFieldsInStruct.cpp b/Engine/Source/Editor/BlueprintGraph/Private/K2Node_SetFieldsInStruct.cpp
--- a/Engine/Source/Editor/BlueprintGraph/Private/K2Node_SetFieldsInStruct.cpp
+++ b/Engine/Source/Editor/BlueprintGraph/Private/K2Node_SetFieldsInStruct.cpp
@@ -46,7 +46,7 @@ void UK2Node_SetFieldsInStruct::AllocateDefaultPins()
 
 		{
 			FStructOnScope StructOnScope(Cast<UScriptStruct>(StructType));
-			FMakeStructPinManager OptionalPinManager(StructOnScope.GetStructMemory());
+			FSetFieldsInStructPinManager OptionalPinManager(StructOnScope.GetStructMemory());
 			OptionalPinManager.RebuildPropertyList(ShowPinForProperties, StructType);
 			OptionalPinManager.CreateVisiblePins(ShowPinForProperties, StructType, EGPD_Input, this);
 		}
@@ -184,4 +184,11 @@ void UK2Node_SetFieldsInStruct::RestoreAllPins()
 	}
 }
 
+void UK2Node_SetFieldsInStruct::FSetFieldsInStructPinManager::GetRecordDefaults(UProperty* TestProperty, FOptionalPinFromProperty& Record) const
+{
+	FMakeStructPinManager::GetRecordDefaults(TestProperty, Record);
+
+	Record.bShowPin = false;
+}
+
 #undef LOCTEXT_NAMESPACE
\ No newline at end of file
