Fix for PR #863 - crash when adding override to pure function

diff --git a/Engine/Source/Editor/BlueprintGraph/Private/EdGraphSchema_K2.cpp b/Engine/Source/Editor/BlueprintGraph/Private/EdGraphSchema_K2.cpp
--- a/Engine/Source/Editor/BlueprintGraph/Private/EdGraphSchema_K2.cpp
+++ b/Engine/Source/Editor/BlueprintGraph/Private/EdGraphSchema_K2.cpp
@@ -2887,6 +2887,7 @@ void UEdGraphSchema_K2::CreateFunctionGraphTerminators(UEdGraph& Graph, UClass*
 		EntryNode->ExtraFlags |= InterfaceToImplement->FunctionFlags & (FUNC_Const | FUNC_Static | FUNC_BlueprintPure);
 
 		UK2Node* NextNode = EntryNode;
+		UEdGraphPin* NextExec = FindExecutionPin(*EntryNode, EGPD_Output);
 		bool bHasParentNode = false;
 		// Create node for call parent function
 		if (((Class->GetClassFlags() & CLASS_Interface) == 0)  &&
@@ -2899,11 +2900,16 @@ void UEdGraphSchema_K2::CreateFunctionGraphTerminators(UEdGraph& Graph, UClass*
 			ParentNode->NodePosY = EntryNode->NodePosY;
 			FunctionParentCreator.Finalize();
 
-			UEdGraphPin* EntryNodeExec = FindExecutionPin(*EntryNode, EGPD_Output);
-			UEdGraphPin* ParentNodeExec = FindExecutionPin(*ParentNode, EGPD_Input);
-			EntryNodeExec->MakeLinkTo(ParentNodeExec);
+			UEdGraphPin* ParentNodeExec = FindExecutionPin(*ParentNode, EGPD_Input); 
 
-			LinkDataPinFromOutputToInput(EntryNode, ParentNode);
+			// If the parent node has an execution pin, then we should as well (we're overriding them, after all)
+			// but perhaps this assumption is not valid in the case where a function becomes pure after being
+			// initially declared impure - for that reason I'm checking for validity on both ParentNodeExec and NextExec
+			if (ParentNodeExec && NextExec)
+			{
+				NextExec->MakeLinkTo(ParentNodeExec);
+				NextExec = FindExecutionPin(*ParentNode, EGPD_Output);
+			}
 
 			NextNode = ParentNode;
 			bHasParentNode = true;
@@ -2932,9 +2938,11 @@ void UEdGraphSchema_K2::CreateFunctionGraphTerminators(UEdGraph& Graph, UClass*
 			SetNodeMetaData(ReturnNode, FNodeMetadata::DefaultGraphNode);
 
 			// Auto-connect the pins for entry and exit, so that by default the signature is properly generated
-			UEdGraphPin* EntryNodeExec = FindExecutionPin(*NextNode, EGPD_Output);
 			UEdGraphPin* ResultNodeExec = FindExecutionPin(*ReturnNode, EGPD_Input);
-			EntryNodeExec->MakeLinkTo(ResultNodeExec);
+			if (ResultNodeExec && NextExec)
+			{
+				NextExec->MakeLinkTo(ResultNodeExec);
+			}
 
 			if (bHasParentNode)
 			{
