Enabled UHT to digest GENERATED_BODY instead of GENERATED_UCLASS_BODY, GENERATED_USTRUCT_BODY, GENERATED_UINTERFACE_BODY or GENERATED_IINTERFACE_BODY, changed every occurence to the new syntax and fix

diff --git a/Engine/Source/Developer/SessionMessages/Classes/SessionServiceMessages.h b/Engine/Source/Developer/SessionMessages/Classes/SessionServiceMessages.h
--- a/Engine/Source/Developer/SessionMessages/Classes/SessionServiceMessages.h
+++ b/Engine/Source/Developer/SessionMessages/Classes/SessionServiceMessages.h
@@ -14,7 +14,7 @@
 USTRUCT()
 struct FSessionServicePing
 {
-	GENERATED_USTRUCT_BODY()
+	GENERATED_BODY()
 
 };
 
@@ -34,7 +34,7 @@ struct TStructOpsTypeTraits<FSessionServicePing> : public TStructOpsTypeTraitsBa
 USTRUCT()
 struct FSessionServicePong
 {
-	GENERATED_USTRUCT_BODY()
+	GENERATED_BODY()
 
 	/** Holds the application's build date. */
 	UPROPERTY()
@@ -96,7 +96,7 @@ struct TStructOpsTypeTraits<FSessionServicePong> : public TStructOpsTypeTraitsBa
 USTRUCT()
 struct FSessionServiceLog
 {
-	GENERATED_USTRUCT_BODY()
+	GENERATED_BODY()
 
 	/** Holds the log message category. */
 	UPROPERTY()
@@ -153,7 +153,7 @@ struct TStructOpsTypeTraits<FSessionServiceLog> : public TStructOpsTypeTraitsBas
 USTRUCT()
 struct FSessionServiceLogSubscribe
 {
-	GENERATED_USTRUCT_BODY()
+	GENERATED_BODY()
 };
 
 template<>
@@ -172,7 +172,7 @@ struct TStructOpsTypeTraits<FSessionServiceLogSubscribe> : public TStructOpsType
 USTRUCT()
 struct FSessionServiceLogUnsubscribe
 {
-	GENERATED_USTRUCT_BODY()
+	GENERATED_BODY()
 };
 
 template<>
diff --git a/Engine/Source/Developer/SourceCodeAccess/Private/SourceCodeAccessSettings.h b/Engine/Source/Developer/SourceCodeAccess/Private/SourceCodeAccessSettings.h
--- a/Engine/Source/Developer/SourceCodeAccess/Private/SourceCodeAccessSettings.h
+++ b/Engine/Source/Developer/SourceCodeAccess/Private/SourceCodeAccessSettings.h
@@ -7,9 +7,11 @@
 UCLASS(config=EditorUserSettings)
 class USourceCodeAccessSettings : public UObject
 {
-	GENERATED_UCLASS_BODY()
+	GENERATED_BODY()
+public:
+	USourceCodeAccessSettings(const FObjectInitializer& ObjectInitializer = FObjectInitializer::Get());
 
 	/** The source code editor we prefer to use. */
 	UPROPERTY(Config, EditAnywhere, Category="Source Code Editor", meta=(DisplayName="Source Code Editor"))
 	FString PreferredAccessor;
-};
\ No newline at end of file
+};
diff --git a/Engine/Source/Developer/TargetDeviceServices/Classes/TargetDeviceServiceMessages.h b/Engine/Source/Developer/TargetDeviceServices/Classes/TargetDeviceServiceMessages.h
--- a/Engine/Source/Developer/TargetDeviceServices/Classes/TargetDeviceServiceMessages.h
+++ b/Engine/Source/Developer/TargetDeviceServices/Classes/TargetDeviceServiceMessages.h
@@ -16,7 +16,7 @@
 USTRUCT()
 struct FTargetDeviceServiceDeployCommit
 {
-	GENERATED_USTRUCT_BODY()
+	GENERATED_BODY()
 
 	/** Holds the variant identifier of the target device to use. */
 	UPROPERTY()
@@ -57,7 +57,7 @@ struct TStructOpsTypeTraits<FTargetDeviceServiceDeployCommit> : public TStructOp
 USTRUCT()
 struct FTargetDeviceServiceDeployFile
 {
-	GENERATED_USTRUCT_BODY()
+	GENERATED_BODY()
 
 	/** Holds the name and path of the file as it will be stored on the target device. */
 	UPROPERTY()
@@ -96,7 +96,7 @@ struct TStructOpsTypeTraits<FTargetDeviceServiceDeployFile> : public TStructOpsT
 USTRUCT()
 struct FTargetDeviceServiceDeployFinished
 {
-	GENERATED_USTRUCT_BODY()
+	GENERATED_BODY()
 
 	/** Holds the variant identifier of the target device to use. */
 	UPROPERTY()
@@ -157,7 +157,7 @@ struct TStructOpsTypeTraits<FTargetDeviceServiceDeployFinished> : public TStruct
 USTRUCT()
 struct FTargetDeviceServiceLaunchApp
 {
-	GENERATED_USTRUCT_BODY()
+	GENERATED_BODY()
 
 	/** Holds the variant identifier of the target device to use. */
 	UPROPERTY()
@@ -213,7 +213,7 @@ struct TStructOpsTypeTraits<FTargetDeviceServiceLaunchApp> : public TStructOpsTy
 USTRUCT()
 struct FTargetDeviceServiceLaunchFinished
 {
-	GENERATED_USTRUCT_BODY()
+	GENERATED_BODY()
 
 	/**
 	 * Holds the identifier of the launched application.
@@ -268,7 +268,7 @@ struct TStructOpsTypeTraits<FTargetDeviceServiceLaunchFinished> : public TStruct
 USTRUCT()
 struct FTargetDeviceClaimDenied
 {
-	GENERATED_USTRUCT_BODY()
+	GENERATED_BODY()
 
 	/** Holds the identifier of the device that is already claimed. */
 	UPROPERTY()
@@ -313,7 +313,7 @@ struct TStructOpsTypeTraits<FTargetDeviceClaimDenied> : public TStructOpsTypeTra
 USTRUCT()
 struct FTargetDeviceClaimed
 {
-	GENERATED_USTRUCT_BODY()
+	GENERATED_BODY()
 
 	/** Holds the identifier of the device that is being claimed. */
 	UPROPERTY()
@@ -357,7 +357,7 @@ struct TStructOpsTypeTraits<FTargetDeviceClaimed> : public TStructOpsTypeTraitsB
 USTRUCT()
 struct FTargetDeviceUnclaimed
 {
-	GENERATED_USTRUCT_BODY()
+	GENERATED_BODY()
 
 	/** Holds the identifier of the device that is no longer claimed. */
 	UPROPERTY()
@@ -402,7 +402,7 @@ struct TStructOpsTypeTraits<FTargetDeviceUnclaimed> : public TStructOpsTypeTrait
 USTRUCT()
 struct FTargetDeviceServicePing
 {
-	GENERATED_USTRUCT_BODY()
+	GENERATED_BODY()
 
 	/** Holds the name of the user who generated the ping. */
 	UPROPERTY()
@@ -434,7 +434,7 @@ struct TStructOpsTypeTraits<FTargetDeviceServicePing> : public TStructOpsTypeTra
 USTRUCT()
 struct FTargetDeviceVariant
 {
-	GENERATED_USTRUCT_BODY()
+	GENERATED_BODY()
 
 	UPROPERTY()
 	FString DeviceID;
@@ -471,7 +471,7 @@ struct TStructOpsTypeTraits<FTargetDeviceVariant> : public TStructOpsTypeTraitsB
 USTRUCT()
 struct FTargetDeviceServicePong
 {
-	GENERATED_USTRUCT_BODY()
+	GENERATED_BODY()
 
 	/** Holds a flag indicating whether the device is currently connected. */
 	UPROPERTY()
@@ -561,7 +561,7 @@ struct TStructOpsTypeTraits<FTargetDeviceServicePong> : public TStructOpsTypeTra
 USTRUCT()
 struct FTargetDeviceServicePowerOff
 {
-	GENERATED_USTRUCT_BODY()
+	GENERATED_BODY()
 
 	/**
 	 * Holds a flag indicating whether the power-off should be enforced.
@@ -602,7 +602,7 @@ struct TStructOpsTypeTraits<FTargetDeviceServicePowerOff> : public TStructOpsTyp
 USTRUCT()
 struct FTargetDeviceServicePowerOn
 {
-	GENERATED_USTRUCT_BODY()
+	GENERATED_BODY()
 
 	/** Holds the name of the user that wishes to power on the device. */
 	UPROPERTY()
@@ -634,7 +634,7 @@ struct TStructOpsTypeTraits<FTargetDeviceServicePowerOn> : public TStructOpsType
 USTRUCT()
 struct FTargetDeviceServiceReboot
 {
-	GENERATED_USTRUCT_BODY()
+	GENERATED_BODY()
 
 	/** Holds the name of the user that wishes to reboot the device. */
 	UPROPERTY()
@@ -670,7 +670,7 @@ struct TStructOpsTypeTraits<FTargetDeviceServiceReboot> : public TStructOpsTypeT
 USTRUCT()
 struct FTargetDeviceServiceRunExecutable
 {
-	GENERATED_USTRUCT_BODY()
+	GENERATED_BODY()
 
 	/** Holds the variant identifier of the target device to use for execution. */
 	UPROPERTY()
@@ -714,7 +714,7 @@ struct TStructOpsTypeTraits<FTargetDeviceServiceRunExecutable> : public TStructO
 USTRUCT()
 struct FTargetDeviceServiceRunFinished
 {
-	GENERATED_USTRUCT_BODY()
+	GENERATED_BODY()
 
 	/** Holds the variant identifier of the target device to use. */
 	UPROPERTY()
diff --git a/Engine/Source/Developer/WidgetCarousel/Public/WidgetCarouselStyle.h b/Engine/Source/Developer/WidgetCarousel/Public/WidgetCarouselStyle.h
--- a/Engine/Source/Developer/WidgetCarousel/Public/WidgetCarouselStyle.h
+++ b/Engine/Source/Developer/WidgetCarousel/Public/WidgetCarouselStyle.h
@@ -8,7 +8,7 @@
 USTRUCT()
 struct WIDGETCAROUSEL_API FWidgetCarouselNavigationButtonStyle : public FSlateWidgetStyle
 {
-	GENERATED_USTRUCT_BODY()
+	GENERATED_BODY()
 
 	FWidgetCarouselNavigationButtonStyle();
 
@@ -37,7 +37,7 @@ struct WIDGETCAROUSEL_API FWidgetCarouselNavigationButtonStyle : public FSlateWi
 USTRUCT()
 struct WIDGETCAROUSEL_API FWidgetCarouselNavigationBarStyle : public FSlateWidgetStyle
 {
-	GENERATED_USTRUCT_BODY()
+	GENERATED_BODY()
 
 	FWidgetCarouselNavigationBarStyle();
 
@@ -91,4 +91,4 @@ private:
 private:
 
 	static TSharedPtr< class FSlateStyleSet > WidgetCarouselStyleInstance;
-};
\ No newline at end of file
+};
diff --git a/Engine/Source/Developer/Windows/WindowsTargetPlatform/Classes/WindowsTargetSettings.h b/Engine/Source/Developer/Windows/WindowsTargetPlatform/Classes/WindowsTargetSettings.h
--- a/Engine/Source/Developer/Windows/WindowsTargetPlatform/Classes/WindowsTargetSettings.h
+++ b/Engine/Source/Developer/Windows/WindowsTargetPlatform/Classes/WindowsTargetSettings.h
@@ -22,9 +22,11 @@ UCLASS(config=Engine, defaultconfig)
 class WINDOWSTARGETPLATFORM_API UWindowsTargetSettings
 	: public UObject
 {
+	GENERATED_BODY()
+public:
+	UWindowsTargetSettings(const FObjectInitializer& ObjectInitializer = FObjectInitializer::Get());
 public:
 
-	GENERATED_UCLASS_BODY()
 
 	/** 
 	 * The collection of RHI's we want to support on this platform.
diff --git a/Engine/Source/Editor/AIGraph/Classes/AIGraph.h b/Engine/Source/Editor/AIGraph/Classes/AIGraph.h
--- a/Engine/Source/Editor/AIGraph/Classes/AIGraph.h
+++ b/Engine/Source/Editor/AIGraph/Classes/AIGraph.h
@@ -6,7 +6,9 @@
 UCLASS()
 class AIGRAPH_API UAIGraph : public UEdGraph
 {
-	GENERATED_UCLASS_BODY()
+	GENERATED_BODY()
+public:
+	UAIGraph(const FObjectInitializer& ObjectInitializer = FObjectInitializer::Get());
 
 	UPROPERTY()
 	int32 GraphVersion;
diff --git a/Engine/Source/Editor/AIGraph/Classes/AIGraphNode.h b/Engine/Source/Editor/AIGraph/Classes/AIGraphNode.h
--- a/Engine/Source/Editor/AIGraph/Classes/AIGraphNode.h
+++ b/Engine/Source/Editor/AIGraph/Classes/AIGraphNode.h
@@ -9,7 +9,9 @@
 UCLASS()
 class AIGRAPH_API UAIGraphNode : public UEdGraphNode
 {
-	GENERATED_UCLASS_BODY()
+	GENERATED_BODY()
+public:
+	UAIGraphNode(const FObjectInitializer& ObjectInitializer = FObjectInitializer::Get());
 
 	/** instance class */
 	UPROPERTY()
diff --git a/Engine/Source/Editor/AIGraph/Classes/AIGraphSchema.h b/Engine/Source/Editor/AIGraph/Classes/AIGraphSchema.h
--- a/Engine/Source/Editor/AIGraph/Classes/AIGraphSchema.h
+++ b/Engine/Source/Editor/AIGraph/Classes/AIGraphSchema.h
@@ -9,7 +9,7 @@
 USTRUCT()
 struct AIGRAPH_API FAISchemaAction_NewNode : public FEdGraphSchemaAction
 {
-	GENERATED_USTRUCT_BODY();
+	GENERATED_BODY()
 
 	/** Template of node we want to create */
 	UPROPERTY()
@@ -45,7 +45,7 @@ struct AIGRAPH_API FAISchemaAction_NewNode : public FEdGraphSchemaAction
 USTRUCT()
 struct AIGRAPH_API FAISchemaAction_NewSubNode : public FEdGraphSchemaAction
 {
-	GENERATED_USTRUCT_BODY();
+	GENERATED_BODY()
 
 	/** Template of node we want to create */
 	UPROPERTY()
@@ -77,7 +77,9 @@ struct AIGRAPH_API FAISchemaAction_NewSubNode : public FEdGraphSchemaAction
 UCLASS()
 class AIGRAPH_API UAIGraphSchema : public UEdGraphSchema
 {
-	GENERATED_UCLASS_BODY()
+	GENERATED_BODY()
+public:
+	UAIGraphSchema(const FObjectInitializer& ObjectInitializer = FObjectInitializer::Get());
 
 	void GetBreakLinkToSubMenuActions(class FMenuBuilder& MenuBuilder, class UEdGraphPin* InGraphPin);
 
diff --git a/Engine/Source/Editor/AIGraph/Classes/AIGraphTypes.h b/Engine/Source/Editor/AIGraph/Classes/AIGraphTypes.h
--- a/Engine/Source/Editor/AIGraph/Classes/AIGraphTypes.h
+++ b/Engine/Source/Editor/AIGraph/Classes/AIGraphTypes.h
@@ -6,7 +6,7 @@
 USTRUCT()
 struct AIGRAPH_API FGraphNodeClassData
 {
-	GENERATED_USTRUCT_BODY()
+	GENERATED_BODY()
 
 	FGraphNodeClassData() {}
 	FGraphNodeClassData(UClass* InClass, const FString& InDeprecatedMessage);
diff --git a/Engine/Source/Editor/AnimGraph/Classes/AnimGraphNode_ApplyAdditive.h b/Engine/Source/Editor/AnimGraph/Classes/AnimGraphNode_ApplyAdditive.h
--- a/Engine/Source/Editor/AnimGraph/Classes/AnimGraphNode_ApplyAdditive.h
+++ b/Engine/Source/Editor/AnimGraph/Classes/AnimGraphNode_ApplyAdditive.h
@@ -8,7 +8,9 @@
 UCLASS(MinimalAPI)
 class UAnimGraphNode_ApplyAdditive : public UAnimGraphNode_Base
 {
-	GENERATED_UCLASS_BODY()
+	GENERATED_BODY()
+public:
+	ANIMGRAPH_API UAnimGraphNode_ApplyAdditive(const FObjectInitializer& ObjectInitializer = FObjectInitializer::Get());
 
 	UPROPERTY(EditAnywhere, Category=Settings)
 	FAnimNode_ApplyAdditive Node;
diff --git a/Engine/Source/Editor/AnimGraph/Classes/AnimGraphNode_Base.h b/Engine/Source/Editor/AnimGraph/Classes/AnimGraphNode_Base.h
--- a/Engine/Source/Editor/AnimGraph/Classes/AnimGraphNode_Base.h
+++ b/Engine/Source/Editor/AnimGraph/Classes/AnimGraphNode_Base.h
@@ -88,7 +88,9 @@ protected:
 UCLASS(Abstract, MinimalAPI)
 class UAnimGraphNode_Base : public UK2Node
 {
-	GENERATED_UCLASS_BODY()
+	GENERATED_BODY()
+public:
+	ANIMGRAPH_API UAnimGraphNode_Base(const FObjectInitializer& ObjectInitializer = FObjectInitializer::Get());
 
 	UPROPERTY(EditAnywhere, BlueprintReadWrite, Category=PinOptions, EditFixedSize)
 	TArray<FOptionalPinFromProperty> ShowPinForProperties;
diff --git a/Engine/Source/Editor/AnimGraph/Classes/AnimGraphNode_BlendListBase.h b/Engine/Source/Editor/AnimGraph/Classes/AnimGraphNode_BlendListBase.h
--- a/Engine/Source/Editor/AnimGraph/Classes/AnimGraphNode_BlendListBase.h
+++ b/Engine/Source/Editor/AnimGraph/Classes/AnimGraphNode_BlendListBase.h
@@ -7,7 +7,9 @@
 UCLASS(MinimalAPI, Abstract)
 class UAnimGraphNode_BlendListBase : public UAnimGraphNode_Base
 {
-	GENERATED_UCLASS_BODY()
+	GENERATED_BODY()
+public:
+	ANIMGRAPH_API UAnimGraphNode_BlendListBase(const FObjectInitializer& ObjectInitializer = FObjectInitializer::Get());
 
 	// UEdGraphNode interface
 	ANIMGRAPH_API virtual FLinearColor GetNodeTitleColor() const override;
diff --git a/Engine/Source/Editor/AnimGraph/Classes/AnimGraphNode_BlendListByBool.h b/Engine/Source/Editor/AnimGraph/Classes/AnimGraphNode_BlendListByBool.h
--- a/Engine/Source/Editor/AnimGraph/Classes/AnimGraphNode_BlendListByBool.h
+++ b/Engine/Source/Editor/AnimGraph/Classes/AnimGraphNode_BlendListByBool.h
@@ -8,7 +8,9 @@
 UCLASS(MinimalAPI)
 class UAnimGraphNode_BlendListByBool : public UAnimGraphNode_BlendListBase
 {
-	GENERATED_UCLASS_BODY()
+	GENERATED_BODY()
+public:
+	ANIMGRAPH_API UAnimGraphNode_BlendListByBool(const FObjectInitializer& ObjectInitializer = FObjectInitializer::Get());
 
 	UPROPERTY(EditAnywhere, Category=Settings)
 	FAnimNode_BlendListByBool Node;
diff --git a/Engine/Source/Editor/AnimGraph/Classes/AnimGraphNode_BlendListByEnum.h b/Engine/Source/Editor/AnimGraph/Classes/AnimGraphNode_BlendListByEnum.h
--- a/Engine/Source/Editor/AnimGraph/Classes/AnimGraphNode_BlendListByEnum.h
+++ b/Engine/Source/Editor/AnimGraph/Classes/AnimGraphNode_BlendListByEnum.h
@@ -9,7 +9,9 @@
 UCLASS(MinimalAPI)
 class UAnimGraphNode_BlendListByEnum : public UAnimGraphNode_BlendListBase, public INodeDependingOnEnumInterface
 {
-	GENERATED_UCLASS_BODY()
+	GENERATED_BODY()
+public:
+	ANIMGRAPH_API UAnimGraphNode_BlendListByEnum(const FObjectInitializer& ObjectInitializer = FObjectInitializer::Get());
 
 	UPROPERTY(EditAnywhere, Category=Settings)
 	FAnimNode_BlendListByEnum Node;
diff --git a/Engine/Source/Editor/AnimGraph/Classes/AnimGraphNode_BlendListByInt.h b/Engine/Source/Editor/AnimGraph/Classes/AnimGraphNode_BlendListByInt.h
--- a/Engine/Source/Editor/AnimGraph/Classes/AnimGraphNode_BlendListByInt.h
+++ b/Engine/Source/Editor/AnimGraph/Classes/AnimGraphNode_BlendListByInt.h
@@ -8,7 +8,9 @@
 UCLASS(MinimalAPI)
 class UAnimGraphNode_BlendListByInt : public UAnimGraphNode_BlendListBase
 {
-	GENERATED_UCLASS_BODY()
+	GENERATED_BODY()
+public:
+	ANIMGRAPH_API UAnimGraphNode_BlendListByInt(const FObjectInitializer& ObjectInitializer = FObjectInitializer::Get());
 
 	UPROPERTY(EditAnywhere, Category=Settings)
 	FAnimNode_BlendListByInt Node;
diff --git a/Engine/Source/Editor/AnimGraph/Classes/AnimGraphNode_BlendSpaceBase.h b/Engine/Source/Editor/AnimGraph/Classes/AnimGraphNode_BlendSpaceBase.h
--- a/Engine/Source/Editor/AnimGraph/Classes/AnimGraphNode_BlendSpaceBase.h
+++ b/Engine/Source/Editor/AnimGraph/Classes/AnimGraphNode_BlendSpaceBase.h
@@ -8,7 +8,9 @@
 UCLASS(Abstract, MinimalAPI)
 class UAnimGraphNode_BlendSpaceBase : public UAnimGraphNode_Base
 {
-	GENERATED_UCLASS_BODY()
+	GENERATED_BODY()
+public:
+	ANIMGRAPH_API UAnimGraphNode_BlendSpaceBase(const FObjectInitializer& ObjectInitializer = FObjectInitializer::Get());
 
 	// UEdGraphNode interface
 	ANIMGRAPH_API virtual FLinearColor GetNodeTitleColor() const override;
diff --git a/Engine/Source/Editor/AnimGraph/Classes/AnimGraphNode_BlendSpaceEvaluator.h b/Engine/Source/Editor/AnimGraph/Classes/AnimGraphNode_BlendSpaceEvaluator.h
--- a/Engine/Source/Editor/AnimGraph/Classes/AnimGraphNode_BlendSpaceEvaluator.h
+++ b/Engine/Source/Editor/AnimGraph/Classes/AnimGraphNode_BlendSpaceEvaluator.h
@@ -9,7 +9,9 @@
 UCLASS(MinimalAPI)
 class UAnimGraphNode_BlendSpaceEvaluator : public UAnimGraphNode_BlendSpaceBase
 {
-	GENERATED_UCLASS_BODY()
+	GENERATED_BODY()
+public:
+	ANIMGRAPH_API UAnimGraphNode_BlendSpaceEvaluator(const FObjectInitializer& ObjectInitializer = FObjectInitializer::Get());
 
 	UPROPERTY(EditAnywhere, Category=Settings)
 	FAnimNode_BlendSpaceEvaluator Node;
diff --git a/Engine/Source/Editor/AnimGraph/Classes/AnimGraphNode_BlendSpacePlayer.h b/Engine/Source/Editor/AnimGraph/Classes/AnimGraphNode_BlendSpacePlayer.h
--- a/Engine/Source/Editor/AnimGraph/Classes/AnimGraphNode_BlendSpacePlayer.h
+++ b/Engine/Source/Editor/AnimGraph/Classes/AnimGraphNode_BlendSpacePlayer.h
@@ -9,7 +9,9 @@
 UCLASS(MinimalAPI)
 class UAnimGraphNode_BlendSpacePlayer : public UAnimGraphNode_BlendSpaceBase
 {
-	GENERATED_UCLASS_BODY()
+	GENERATED_BODY()
+public:
+	ANIMGRAPH_API UAnimGraphNode_BlendSpacePlayer(const FObjectInitializer& ObjectInitializer = FObjectInitializer::Get());
 
 	UPROPERTY(EditAnywhere, Category=Settings)
 	FAnimNode_BlendSpacePlayer Node;
diff --git a/Engine/Source/Editor/AnimGraph/Classes/AnimGraphNode_BoneDrivenController.h b/Engine/Source/Editor/AnimGraph/Classes/AnimGraphNode_BoneDrivenController.h
--- a/Engine/Source/Editor/AnimGraph/Classes/AnimGraphNode_BoneDrivenController.h
+++ b/Engine/Source/Editor/AnimGraph/Classes/AnimGraphNode_BoneDrivenController.h
@@ -10,7 +10,9 @@
 UCLASS(MinimalAPI)
 class UAnimGraphNode_BoneDrivenController : public UAnimGraphNode_SkeletalControlBase
 {
-	GENERATED_UCLASS_BODY()
+	GENERATED_BODY()
+public:
+	ANIMGRAPH_API UAnimGraphNode_BoneDrivenController(const FObjectInitializer& ObjectInitializer = FObjectInitializer::Get());
 
 	UPROPERTY(EditAnywhere, Category = Settings)
 	FAnimNode_BoneDrivenController Node;
@@ -36,4 +38,4 @@ protected:
 private:
 	/** Constructing FText strings can be costly, so we cache the node's title */
 	FNodeTitleTextTable CachedNodeTitles;
-};
\ No newline at end of file
+};
diff --git a/Engine/Source/Editor/AnimGraph/Classes/AnimGraphNode_ComponentToLocalSpace.h b/Engine/Source/Editor/AnimGraph/Classes/AnimGraphNode_ComponentToLocalSpace.h
--- a/Engine/Source/Editor/AnimGraph/Classes/AnimGraphNode_ComponentToLocalSpace.h
+++ b/Engine/Source/Editor/AnimGraph/Classes/AnimGraphNode_ComponentToLocalSpace.h
@@ -8,7 +8,9 @@
 UCLASS(MinimalAPI)
 class UAnimGraphNode_ComponentToLocalSpace : public UAnimGraphNode_Base
 {
-	GENERATED_UCLASS_BODY()
+	GENERATED_BODY()
+public:
+	ANIMGRAPH_API UAnimGraphNode_ComponentToLocalSpace(const FObjectInitializer& ObjectInitializer = FObjectInitializer::Get());
 
 	UPROPERTY(EditAnywhere, Category=Settings)
 	FAnimNode_ConvertComponentToLocalSpace Node;
diff --git a/Engine/Source/Editor/AnimGraph/Classes/AnimGraphNode_CopyBone.h b/Engine/Source/Editor/AnimGraph/Classes/AnimGraphNode_CopyBone.h
--- a/Engine/Source/Editor/AnimGraph/Classes/AnimGraphNode_CopyBone.h
+++ b/Engine/Source/Editor/AnimGraph/Classes/AnimGraphNode_CopyBone.h
@@ -9,7 +9,9 @@
 UCLASS(MinimalAPI)
 class UAnimGraphNode_CopyBone : public UAnimGraphNode_SkeletalControlBase
 {
-	GENERATED_UCLASS_BODY()
+	GENERATED_BODY()
+public:
+	ANIMGRAPH_API UAnimGraphNode_CopyBone(const FObjectInitializer& ObjectInitializer = FObjectInitializer::Get());
 
 	UPROPERTY(EditAnywhere, Category=Settings)
 	FAnimNode_CopyBone Node;
diff --git a/Engine/Source/Editor/AnimGraph/Classes/AnimGraphNode_CustomTransitionResult.h b/Engine/Source/Editor/AnimGraph/Classes/AnimGraphNode_CustomTransitionResult.h
--- a/Engine/Source/Editor/AnimGraph/Classes/AnimGraphNode_CustomTransitionResult.h
+++ b/Engine/Source/Editor/AnimGraph/Classes/AnimGraphNode_CustomTransitionResult.h
@@ -7,7 +7,9 @@
 UCLASS(MinimalAPI)
 class UAnimGraphNode_CustomTransitionResult : public UAnimGraphNode_Root
 {
-	GENERATED_UCLASS_BODY()
+	GENERATED_BODY()
+public:
+	ANIMGRAPH_API UAnimGraphNode_CustomTransitionResult(const FObjectInitializer& ObjectInitializer = FObjectInitializer::Get());
 
 	// UEdGraphNode interface
 	virtual FLinearColor GetNodeTitleColor() const override;
diff --git a/Engine/Source/Editor/AnimGraph/Classes/AnimGraphNode_Fabrik.h b/Engine/Source/Editor/AnimGraph/Classes/AnimGraphNode_Fabrik.h
--- a/Engine/Source/Editor/AnimGraph/Classes/AnimGraphNode_Fabrik.h
+++ b/Engine/Source/Editor/AnimGraph/Classes/AnimGraphNode_Fabrik.h
@@ -9,7 +9,9 @@
 UCLASS(MinimalAPI)
 class UAnimGraphNode_Fabrik : public UAnimGraphNode_SkeletalControlBase
 {
-	GENERATED_UCLASS_BODY()
+	GENERATED_BODY()
+public:
+	ANIMGRAPH_API UAnimGraphNode_Fabrik(const FObjectInitializer& ObjectInitializer = FObjectInitializer::Get());
 
 	UPROPERTY(EditAnywhere, Category = Settings)
 	FAnimNode_Fabrik Node;
@@ -34,4 +36,4 @@ protected:
 	// UAnimGraphNode_SkeletalControlBase interface
 	virtual FText GetControllerDescription() const override;
 	// End of UAnimGraphNode_SkeletalControlBase interface
-};
\ No newline at end of file
+};
diff --git a/Engine/Source/Editor/AnimGraph/Classes/AnimGraphNode_HandIKRetargeting.h b/Engine/Source/Editor/AnimGraph/Classes/AnimGraphNode_HandIKRetargeting.h
--- a/Engine/Source/Editor/AnimGraph/Classes/AnimGraphNode_HandIKRetargeting.h
+++ b/Engine/Source/Editor/AnimGraph/Classes/AnimGraphNode_HandIKRetargeting.h
@@ -8,7 +8,9 @@
 UCLASS(MinimalAPI)
 class UAnimGraphNode_HandIKRetargeting : public UAnimGraphNode_SkeletalControlBase
 {
-	GENERATED_UCLASS_BODY()
+	GENERATED_BODY()
+public:
+	ANIMGRAPH_API UAnimGraphNode_HandIKRetargeting(const FObjectInitializer& ObjectInitializer = FObjectInitializer::Get());
 
 	UPROPERTY(EditAnywhere, Category = Settings)
 	FAnimNode_HandIKRetargeting Node;
diff --git a/Engine/Source/Editor/AnimGraph/Classes/AnimGraphNode_IdentityPose.h b/Engine/Source/Editor/AnimGraph/Classes/AnimGraphNode_IdentityPose.h
--- a/Engine/Source/Editor/AnimGraph/Classes/AnimGraphNode_IdentityPose.h
+++ b/Engine/Source/Editor/AnimGraph/Classes/AnimGraphNode_IdentityPose.h
@@ -7,7 +7,9 @@
 UCLASS(MinimalAPI)
 class UAnimGraphNode_IdentityPose : public UAnimGraphNode_RefPoseBase
 {
-	GENERATED_UCLASS_BODY()
+	GENERATED_BODY()
+public:
+	ANIMGRAPH_API UAnimGraphNode_IdentityPose(const FObjectInitializer& ObjectInitializer = FObjectInitializer::Get());
 
 	// UEdGraphNode interface
 	virtual FText GetTooltipText() const override;
diff --git a/Engine/Source/Editor/AnimGraph/Classes/AnimGraphNode_LayeredBoneBlend.h b/Engine/Source/Editor/AnimGraph/Classes/AnimGraphNode_LayeredBoneBlend.h
--- a/Engine/Source/Editor/AnimGraph/Classes/AnimGraphNode_LayeredBoneBlend.h
+++ b/Engine/Source/Editor/AnimGraph/Classes/AnimGraphNode_LayeredBoneBlend.h
@@ -8,7 +8,9 @@
 UCLASS(MinimalAPI)
 class UAnimGraphNode_LayeredBoneBlend : public UAnimGraphNode_BlendListBase
 {
-	GENERATED_UCLASS_BODY()
+	GENERATED_BODY()
+public:
+	ANIMGRAPH_API UAnimGraphNode_LayeredBoneBlend(const FObjectInitializer& ObjectInitializer = FObjectInitializer::Get());
 
 	UPROPERTY(EditAnywhere, Category=Settings)
 	FAnimNode_LayeredBoneBlend Node;
diff --git a/Engine/Source/Editor/AnimGraph/Classes/AnimGraphNode_LocalRefPose.h b/Engine/Source/Editor/AnimGraph/Classes/AnimGraphNode_LocalRefPose.h
--- a/Engine/Source/Editor/AnimGraph/Classes/AnimGraphNode_LocalRefPose.h
+++ b/Engine/Source/Editor/AnimGraph/Classes/AnimGraphNode_LocalRefPose.h
@@ -7,7 +7,9 @@
 UCLASS(MinimalAPI)
 class UAnimGraphNode_LocalRefPose : public UAnimGraphNode_RefPoseBase
 {
-	GENERATED_UCLASS_BODY()
+	GENERATED_BODY()
+public:
+	ANIMGRAPH_API UAnimGraphNode_LocalRefPose(const FObjectInitializer& ObjectInitializer = FObjectInitializer::Get());
 
 	// UEdGraphNode interface
 	virtual FText GetTooltipText() const override;
diff --git a/Engine/Source/Editor/AnimGraph/Classes/AnimGraphNode_LocalToComponentSpace.h b/Engine/Source/Editor/AnimGraph/Classes/AnimGraphNode_LocalToComponentSpace.h
--- a/Engine/Source/Editor/AnimGraph/Classes/AnimGraphNode_LocalToComponentSpace.h
+++ b/Engine/Source/Editor/AnimGraph/Classes/AnimGraphNode_LocalToComponentSpace.h
@@ -8,7 +8,9 @@
 UCLASS(MinimalAPI)
 class UAnimGraphNode_LocalToComponentSpace : public UAnimGraphNode_Base
 {
-	GENERATED_UCLASS_BODY()
+	GENERATED_BODY()
+public:
+	ANIMGRAPH_API UAnimGraphNode_LocalToComponentSpace(const FObjectInitializer& ObjectInitializer = FObjectInitializer::Get());
 
 	UPROPERTY(EditAnywhere, Category=Settings)
 	FAnimNode_ConvertLocalToComponentSpace Node;
diff --git a/Engine/Source/Editor/AnimGraph/Classes/AnimGraphNode_LookAt.h b/Engine/Source/Editor/AnimGraph/Classes/AnimGraphNode_LookAt.h
--- a/Engine/Source/Editor/AnimGraph/Classes/AnimGraphNode_LookAt.h
+++ b/Engine/Source/Editor/AnimGraph/Classes/AnimGraphNode_LookAt.h
@@ -10,7 +10,9 @@
 UCLASS(MinimalAPI)
 class UAnimGraphNode_LookAt : public UAnimGraphNode_SkeletalControlBase
 {
-	GENERATED_UCLASS_BODY()
+	GENERATED_BODY()
+public:
+	ANIMGRAPH_API UAnimGraphNode_LookAt(const FObjectInitializer& ObjectInitializer = FObjectInitializer::Get());
 
 	UPROPERTY(EditAnywhere, Category=Settings)
 	FAnimNode_LookAt Node;
diff --git a/Engine/Source/Editor/AnimGraph/Classes/AnimGraphNode_MeshRefPose.h b/Engine/Source/Editor/AnimGraph/Classes/AnimGraphNode_MeshRefPose.h
--- a/Engine/Source/Editor/AnimGraph/Classes/AnimGraphNode_MeshRefPose.h
+++ b/Engine/Source/Editor/AnimGraph/Classes/AnimGraphNode_MeshRefPose.h
@@ -8,7 +8,9 @@
 UCLASS(MinimalAPI)
 class UAnimGraphNode_MeshRefPose : public UAnimGraphNode_Base
 {
-	GENERATED_UCLASS_BODY()
+	GENERATED_BODY()
+public:
+	ANIMGRAPH_API UAnimGraphNode_MeshRefPose(const FObjectInitializer& ObjectInitializer = FObjectInitializer::Get());
 
 	UPROPERTY(EditAnywhere, Category=Settings)
 	FAnimNode_MeshSpaceRefPose Node;
diff --git a/Engine/Source/Editor/AnimGraph/Classes/AnimGraphNode_ModifyBone.h b/Engine/Source/Editor/AnimGraph/Classes/AnimGraphNode_ModifyBone.h
--- a/Engine/Source/Editor/AnimGraph/Classes/AnimGraphNode_ModifyBone.h
+++ b/Engine/Source/Editor/AnimGraph/Classes/AnimGraphNode_ModifyBone.h
@@ -10,7 +10,9 @@
 UCLASS(MinimalAPI)
 class UAnimGraphNode_ModifyBone : public UAnimGraphNode_SkeletalControlBase
 {
-	GENERATED_UCLASS_BODY()
+	GENERATED_BODY()
+public:
+	ANIMGRAPH_API UAnimGraphNode_ModifyBone(const FObjectInitializer& ObjectInitializer = FObjectInitializer::Get());
 
 	UPROPERTY(EditAnywhere, Category=Settings)
 	FAnimNode_ModifyBone Node;
diff --git a/Engine/Source/Editor/AnimGraph/Classes/AnimGraphNode_RefPoseBase.h b/Engine/Source/Editor/AnimGraph/Classes/AnimGraphNode_RefPoseBase.h
--- a/Engine/Source/Editor/AnimGraph/Classes/AnimGraphNode_RefPoseBase.h
+++ b/Engine/Source/Editor/AnimGraph/Classes/AnimGraphNode_RefPoseBase.h
@@ -8,7 +8,9 @@
 UCLASS(MinimalAPI, abstract)
 class UAnimGraphNode_RefPoseBase : public UAnimGraphNode_Base
 {
-	GENERATED_UCLASS_BODY()
+	GENERATED_BODY()
+public:
+	ANIMGRAPH_API UAnimGraphNode_RefPoseBase(const FObjectInitializer& ObjectInitializer = FObjectInitializer::Get());
 
 	UPROPERTY(EditAnywhere, Category=Settings)
 	FAnimNode_RefPose Node;
diff --git a/Engine/Source/Editor/AnimGraph/Classes/AnimGraphNode_Root.h b/Engine/Source/Editor/AnimGraph/Classes/AnimGraphNode_Root.h
--- a/Engine/Source/Editor/AnimGraph/Classes/AnimGraphNode_Root.h
+++ b/Engine/Source/Editor/AnimGraph/Classes/AnimGraphNode_Root.h
@@ -8,7 +8,9 @@
 UCLASS(MinimalAPI)
 class UAnimGraphNode_Root : public UAnimGraphNode_Base
 {
-	GENERATED_UCLASS_BODY()
+	GENERATED_BODY()
+public:
+	ANIMGRAPH_API UAnimGraphNode_Root(const FObjectInitializer& ObjectInitializer = FObjectInitializer::Get());
 
 	UPROPERTY(EditAnywhere, Category=Settings)
 	FAnimNode_Root Node;
diff --git a/Engine/Source/Editor/AnimGraph/Classes/AnimGraphNode_RotateRootBone.h b/Engine/Source/Editor/AnimGraph/Classes/AnimGraphNode_RotateRootBone.h
--- a/Engine/Source/Editor/AnimGraph/Classes/AnimGraphNode_RotateRootBone.h
+++ b/Engine/Source/Editor/AnimGraph/Classes/AnimGraphNode_RotateRootBone.h
@@ -9,7 +9,9 @@
 UCLASS(MinimalAPI)
 class UAnimGraphNode_RotateRootBone : public UAnimGraphNode_Base
 {
-	GENERATED_UCLASS_BODY()
+	GENERATED_BODY()
+public:
+	ANIMGRAPH_API UAnimGraphNode_RotateRootBone(const FObjectInitializer& ObjectInitializer = FObjectInitializer::Get());
 
 	UPROPERTY(EditAnywhere, Category=Settings)
 	FAnimNode_RotateRootBone Node;
diff --git a/Engine/Source/Editor/AnimGraph/Classes/AnimGraphNode_RotationMultiplier.h b/Engine/Source/Editor/AnimGraph/Classes/AnimGraphNode_RotationMultiplier.h
--- a/Engine/Source/Editor/AnimGraph/Classes/AnimGraphNode_RotationMultiplier.h
+++ b/Engine/Source/Editor/AnimGraph/Classes/AnimGraphNode_RotationMultiplier.h
@@ -9,7 +9,9 @@
 UCLASS(MinimalAPI)
 class UAnimGraphNode_RotationMultiplier : public UAnimGraphNode_SkeletalControlBase
 {
-	GENERATED_UCLASS_BODY()
+	GENERATED_BODY()
+public:
+	ANIMGRAPH_API UAnimGraphNode_RotationMultiplier(const FObjectInitializer& ObjectInitializer = FObjectInitializer::Get());
 
 	UPROPERTY(EditAnywhere, Category=Settings)
 	FAnimNode_RotationMultiplier Node;
diff --git a/Engine/Source/Editor/AnimGraph/Classes/AnimGraphNode_RotationOffsetBlendSpace.h b/Engine/Source/Editor/AnimGraph/Classes/AnimGraphNode_RotationOffsetBlendSpace.h
--- a/Engine/Source/Editor/AnimGraph/Classes/AnimGraphNode_RotationOffsetBlendSpace.h
+++ b/Engine/Source/Editor/AnimGraph/Classes/AnimGraphNode_RotationOffsetBlendSpace.h
@@ -9,7 +9,9 @@
 UCLASS(MinimalAPI)
 class UAnimGraphNode_RotationOffsetBlendSpace: public UAnimGraphNode_BlendSpaceBase
 {
-	GENERATED_UCLASS_BODY()
+	GENERATED_BODY()
+public:
+	ANIMGRAPH_API UAnimGraphNode_RotationOffsetBlendSpace(const FObjectInitializer& ObjectInitializer = FObjectInitializer::Get());
 
 	UPROPERTY(EditAnywhere, Category=Settings)
 	FAnimNode_RotationOffsetBlendSpace Node;
diff --git a/Engine/Source/Editor/AnimGraph/Classes/AnimGraphNode_SaveCachedPose.h b/Engine/Source/Editor/AnimGraph/Classes/AnimGraphNode_SaveCachedPose.h
--- a/Engine/Source/Editor/AnimGraph/Classes/AnimGraphNode_SaveCachedPose.h
+++ b/Engine/Source/Editor/AnimGraph/Classes/AnimGraphNode_SaveCachedPose.h
@@ -9,7 +9,9 @@
 UCLASS(MinimalAPI)
 class UAnimGraphNode_SaveCachedPose : public UAnimGraphNode_Base
 {
-	GENERATED_UCLASS_BODY()
+	GENERATED_BODY()
+public:
+	ANIMGRAPH_API UAnimGraphNode_SaveCachedPose(const FObjectInitializer& ObjectInitializer = FObjectInitializer::Get());
 
 	UPROPERTY(EditAnywhere, Category=Settings)
 	FAnimNode_SaveCachedPose Node;
diff --git a/Engine/Source/Editor/AnimGraph/Classes/AnimGraphNode_SequenceEvaluator.h b/Engine/Source/Editor/AnimGraph/Classes/AnimGraphNode_SequenceEvaluator.h
--- a/Engine/Source/Editor/AnimGraph/Classes/AnimGraphNode_SequenceEvaluator.h
+++ b/Engine/Source/Editor/AnimGraph/Classes/AnimGraphNode_SequenceEvaluator.h
@@ -9,7 +9,9 @@
 UCLASS(MinimalAPI)
 class UAnimGraphNode_SequenceEvaluator : public UAnimGraphNode_Base
 {
-	GENERATED_UCLASS_BODY()
+	GENERATED_BODY()
+public:
+	ANIMGRAPH_API UAnimGraphNode_SequenceEvaluator(const FObjectInitializer& ObjectInitializer = FObjectInitializer::Get());
 
 	UPROPERTY(EditAnywhere, Category=Settings)
 	FAnimNode_SequenceEvaluator Node;
diff --git a/Engine/Source/Editor/AnimGraph/Classes/AnimGraphNode_SequencePlayer.h b/Engine/Source/Editor/AnimGraph/Classes/AnimGraphNode_SequencePlayer.h
--- a/Engine/Source/Editor/AnimGraph/Classes/AnimGraphNode_SequencePlayer.h
+++ b/Engine/Source/Editor/AnimGraph/Classes/AnimGraphNode_SequencePlayer.h
@@ -9,7 +9,9 @@
 UCLASS(MinimalAPI)
 class UAnimGraphNode_SequencePlayer : public UAnimGraphNode_Base
 {
-	GENERATED_UCLASS_BODY()
+	GENERATED_BODY()
+public:
+	ANIMGRAPH_API UAnimGraphNode_SequencePlayer(const FObjectInitializer& ObjectInitializer = FObjectInitializer::Get());
 
 	UPROPERTY(EditAnywhere, Category=Settings)
 	FAnimNode_SequencePlayer Node;
diff --git a/Engine/Source/Editor/AnimGraph/Classes/AnimGraphNode_SkeletalControlBase.h b/Engine/Source/Editor/AnimGraph/Classes/AnimGraphNode_SkeletalControlBase.h
--- a/Engine/Source/Editor/AnimGraph/Classes/AnimGraphNode_SkeletalControlBase.h
+++ b/Engine/Source/Editor/AnimGraph/Classes/AnimGraphNode_SkeletalControlBase.h
@@ -8,7 +8,9 @@
 UCLASS(Abstract, MinimalAPI)
 class UAnimGraphNode_SkeletalControlBase : public UAnimGraphNode_Base
 {
-	GENERATED_UCLASS_BODY()
+	GENERATED_BODY()
+public:
+	ANIMGRAPH_API UAnimGraphNode_SkeletalControlBase(const FObjectInitializer& ObjectInitializer = FObjectInitializer::Get());
 
 public:
 	// UEdGraphNode interface
diff --git a/Engine/Source/Editor/AnimGraph/Classes/AnimGraphNode_Slot.h b/Engine/Source/Editor/AnimGraph/Classes/AnimGraphNode_Slot.h
--- a/Engine/Source/Editor/AnimGraph/Classes/AnimGraphNode_Slot.h
+++ b/Engine/Source/Editor/AnimGraph/Classes/AnimGraphNode_Slot.h
@@ -9,7 +9,9 @@
 UCLASS(MinimalAPI)
 class UAnimGraphNode_Slot : public UAnimGraphNode_Base
 {
-	GENERATED_UCLASS_BODY()
+	GENERATED_BODY()
+public:
+	ANIMGRAPH_API UAnimGraphNode_Slot(const FObjectInitializer& ObjectInitializer = FObjectInitializer::Get());
 
 	UPROPERTY(EditAnywhere, Category=Settings)
 	FAnimNode_Slot Node;
diff --git a/Engine/Source/Editor/AnimGraph/Classes/AnimGraphNode_SpringBone.h b/Engine/Source/Editor/AnimGraph/Classes/AnimGraphNode_SpringBone.h
--- a/Engine/Source/Editor/AnimGraph/Classes/AnimGraphNode_SpringBone.h
+++ b/Engine/Source/Editor/AnimGraph/Classes/AnimGraphNode_SpringBone.h
@@ -9,7 +9,9 @@
 UCLASS(MinimalAPI)
 class UAnimGraphNode_SpringBone : public UAnimGraphNode_SkeletalControlBase
 {
-	GENERATED_UCLASS_BODY()
+	GENERATED_BODY()
+public:
+	ANIMGRAPH_API UAnimGraphNode_SpringBone(const FObjectInitializer& ObjectInitializer = FObjectInitializer::Get());
 
 	UPROPERTY(EditAnywhere, Category=Settings)
 	FAnimNode_SpringBone Node;
diff --git a/Engine/Source/Editor/AnimGraph/Classes/AnimGraphNode_StateMachine.h b/Engine/Source/Editor/AnimGraph/Classes/AnimGraphNode_StateMachine.h
--- a/Engine/Source/Editor/AnimGraph/Classes/AnimGraphNode_StateMachine.h
+++ b/Engine/Source/Editor/AnimGraph/Classes/AnimGraphNode_StateMachine.h
@@ -7,7 +7,9 @@
 UCLASS(MinimalAPI)
 class UAnimGraphNode_StateMachine : public UAnimGraphNode_StateMachineBase
 {
-	GENERATED_UCLASS_BODY()
+	GENERATED_BODY()
+public:
+	ANIMGRAPH_API UAnimGraphNode_StateMachine(const FObjectInitializer& ObjectInitializer = FObjectInitializer::Get());
 
 	// Runtime state machine representation (empty; populated during compilation)
 	UPROPERTY(EditAnywhere, Category=Settings)
diff --git a/Engine/Source/Editor/AnimGraph/Classes/AnimGraphNode_StateMachineBase.h b/Engine/Source/Editor/AnimGraph/Classes/AnimGraphNode_StateMachineBase.h
--- a/Engine/Source/Editor/AnimGraph/Classes/AnimGraphNode_StateMachineBase.h
+++ b/Engine/Source/Editor/AnimGraph/Classes/AnimGraphNode_StateMachineBase.h
@@ -9,7 +9,9 @@
 UCLASS(Abstract)
 class ANIMGRAPH_API UAnimGraphNode_StateMachineBase : public UAnimGraphNode_Base
 {
-	GENERATED_UCLASS_BODY()
+	GENERATED_BODY()
+public:
+	UAnimGraphNode_StateMachineBase(const FObjectInitializer& ObjectInitializer = FObjectInitializer::Get());
 
 	// Editor state machine representation
 	UPROPERTY()
diff --git a/Engine/Source/Editor/AnimGraph/Classes/AnimGraphNode_StateResult.h b/Engine/Source/Editor/AnimGraph/Classes/AnimGraphNode_StateResult.h
--- a/Engine/Source/Editor/AnimGraph/Classes/AnimGraphNode_StateResult.h
+++ b/Engine/Source/Editor/AnimGraph/Classes/AnimGraphNode_StateResult.h
@@ -8,5 +8,7 @@
 UCLASS(MinimalAPI)
 class UAnimGraphNode_StateResult : public UAnimGraphNode_Root
 {
-	GENERATED_UCLASS_BODY()
+	GENERATED_BODY()
+public:
+	ANIMGRAPH_API UAnimGraphNode_StateResult(const FObjectInitializer& ObjectInitializer = FObjectInitializer::Get());
 };
diff --git a/Engine/Source/Editor/AnimGraph/Classes/AnimGraphNode_Trail.h b/Engine/Source/Editor/AnimGraph/Classes/AnimGraphNode_Trail.h
--- a/Engine/Source/Editor/AnimGraph/Classes/AnimGraphNode_Trail.h
+++ b/Engine/Source/Editor/AnimGraph/Classes/AnimGraphNode_Trail.h
@@ -9,7 +9,9 @@
 UCLASS(MinimalAPI)
 class UAnimGraphNode_Trail : public UAnimGraphNode_SkeletalControlBase
 {
-	GENERATED_UCLASS_BODY()
+	GENERATED_BODY()
+public:
+	ANIMGRAPH_API UAnimGraphNode_Trail(const FObjectInitializer& ObjectInitializer = FObjectInitializer::Get());
 
 		UPROPERTY(EditAnywhere, Category = Settings)
 		FAnimNode_Trail Node;
diff --git a/Engine/Source/Editor/AnimGraph/Classes/AnimGraphNode_TransitionPoseEvaluator.h b/Engine/Source/Editor/AnimGraph/Classes/AnimGraphNode_TransitionPoseEvaluator.h
--- a/Engine/Source/Editor/AnimGraph/Classes/AnimGraphNode_TransitionPoseEvaluator.h
+++ b/Engine/Source/Editor/AnimGraph/Classes/AnimGraphNode_TransitionPoseEvaluator.h
@@ -8,7 +8,9 @@
 UCLASS(MinimalAPI)
 class UAnimGraphNode_TransitionPoseEvaluator : public UAnimGraphNode_Base
 {
-	GENERATED_UCLASS_BODY()
+	GENERATED_BODY()
+public:
+	ANIMGRAPH_API UAnimGraphNode_TransitionPoseEvaluator(const FObjectInitializer& ObjectInitializer = FObjectInitializer::Get());
 
 	UPROPERTY(EditAnywhere, BlueprintReadWrite, Category=Settings)
 	FAnimNode_TransitionPoseEvaluator Node;
@@ -30,4 +32,4 @@ class UAnimGraphNode_TransitionPoseEvaluator : public UAnimGraphNode_Base
 	// UK2Node interface.
 	virtual void GetMenuActions(FBlueprintActionDatabaseRegistrar& ActionRegistrar) const override;
 	// End of UK2Node interface
-};
\ No newline at end of file
+};
diff --git a/Engine/Source/Editor/AnimGraph/Classes/AnimGraphNode_TransitionResult.h b/Engine/Source/Editor/AnimGraph/Classes/AnimGraphNode_TransitionResult.h
--- a/Engine/Source/Editor/AnimGraph/Classes/AnimGraphNode_TransitionResult.h
+++ b/Engine/Source/Editor/AnimGraph/Classes/AnimGraphNode_TransitionResult.h
@@ -8,7 +8,9 @@
 UCLASS(MinimalAPI)
 class UAnimGraphNode_TransitionResult : public UAnimGraphNode_Base
 {
-	GENERATED_UCLASS_BODY()
+	GENERATED_BODY()
+public:
+	ANIMGRAPH_API UAnimGraphNode_TransitionResult(const FObjectInitializer& ObjectInitializer = FObjectInitializer::Get());
 
 	UPROPERTY(EditAnywhere, Category=Settings)
 	FAnimNode_TransitionResult Node;
diff --git a/Engine/Source/Editor/AnimGraph/Classes/AnimGraphNode_TwoBoneIK.h b/Engine/Source/Editor/AnimGraph/Classes/AnimGraphNode_TwoBoneIK.h
--- a/Engine/Source/Editor/AnimGraph/Classes/AnimGraphNode_TwoBoneIK.h
+++ b/Engine/Source/Editor/AnimGraph/Classes/AnimGraphNode_TwoBoneIK.h
@@ -12,7 +12,9 @@
 UCLASS(MinimalAPI)
 class UAnimGraphNode_TwoBoneIK : public UAnimGraphNode_SkeletalControlBase
 {
-	GENERATED_UCLASS_BODY()
+	GENERATED_BODY()
+public:
+	ANIMGRAPH_API UAnimGraphNode_TwoBoneIK(const FObjectInitializer& ObjectInitializer = FObjectInitializer::Get());
 
 	UPROPERTY(EditAnywhere, Category=Settings)
 	FAnimNode_TwoBoneIK Node;
diff --git a/Engine/Source/Editor/AnimGraph/Classes/AnimGraphNode_TwoWayBlend.h b/Engine/Source/Editor/AnimGraph/Classes/AnimGraphNode_TwoWayBlend.h
--- a/Engine/Source/Editor/AnimGraph/Classes/AnimGraphNode_TwoWayBlend.h
+++ b/Engine/Source/Editor/AnimGraph/Classes/AnimGraphNode_TwoWayBlend.h
@@ -8,7 +8,9 @@
 UCLASS(MinimalAPI)
 class UAnimGraphNode_TwoWayBlend : public UAnimGraphNode_Base
 {
-	GENERATED_UCLASS_BODY()
+	GENERATED_BODY()
+public:
+	ANIMGRAPH_API UAnimGraphNode_TwoWayBlend(const FObjectInitializer& ObjectInitializer = FObjectInitializer::Get());
 
 	UPROPERTY(EditAnywhere, Category=Settings)
 	FAnimationNode_TwoWayBlend BlendNode;
diff --git a/Engine/Source/Editor/AnimGraph/Classes/AnimGraphNode_UseCachedPose.h b/Engine/Source/Editor/AnimGraph/Classes/AnimGraphNode_UseCachedPose.h
--- a/Engine/Source/Editor/AnimGraph/Classes/AnimGraphNode_UseCachedPose.h
+++ b/Engine/Source/Editor/AnimGraph/Classes/AnimGraphNode_UseCachedPose.h
@@ -9,7 +9,9 @@
 UCLASS(MinimalAPI)
 class UAnimGraphNode_UseCachedPose : public UAnimGraphNode_Base
 {
-	GENERATED_UCLASS_BODY()
+	GENERATED_BODY()
+public:
+	ANIMGRAPH_API UAnimGraphNode_UseCachedPose(const FObjectInitializer& ObjectInitializer = FObjectInitializer::Get());
 
 public:
 	UPROPERTY()
diff --git a/Engine/Source/Editor/AnimGraph/Classes/AnimGraphNode_WheelHandler.h b/Engine/Source/Editor/AnimGraph/Classes/AnimGraphNode_WheelHandler.h
--- a/Engine/Source/Editor/AnimGraph/Classes/AnimGraphNode_WheelHandler.h
+++ b/Engine/Source/Editor/AnimGraph/Classes/AnimGraphNode_WheelHandler.h
@@ -9,7 +9,9 @@
 UCLASS(MinimalAPI)
 class UAnimGraphNode_WheelHandler : public UAnimGraphNode_SkeletalControlBase
 {
-	GENERATED_UCLASS_BODY()
+	GENERATED_BODY()
+public:
+	ANIMGRAPH_API UAnimGraphNode_WheelHandler(const FObjectInitializer& ObjectInitializer = FObjectInitializer::Get());
 
 	UPROPERTY(EditAnywhere, Category=Settings)
 	FAnimNode_WheelHandler Node;
diff --git a/Engine/Source/Editor/AnimGraph/Classes/AnimPreviewInstance.h b/Engine/Source/Editor/AnimGraph/Classes/AnimPreviewInstance.h
--- a/Engine/Source/Editor/AnimGraph/Classes/AnimPreviewInstance.h
+++ b/Engine/Source/Editor/AnimGraph/Classes/AnimPreviewInstance.h
@@ -27,7 +27,9 @@ enum EMontagePreviewType
 UCLASS(transient, NotBlueprintable, noteditinlinenew)
 class ANIMGRAPH_API UAnimPreviewInstance : public UAnimSingleNodeInstance
 {
-	GENERATED_UCLASS_BODY()
+	GENERATED_BODY()
+public:
+	UAnimPreviewInstance(const FObjectInitializer& ObjectInitializer = FObjectInitializer::Get());
 
 	/** Controllers for individual bones */
 	UPROPERTY(transient)
diff --git a/Engine/Source/Editor/AnimGraph/Classes/AnimStateConduitNode.h b/Engine/Source/Editor/AnimGraph/Classes/AnimStateConduitNode.h
--- a/Engine/Source/Editor/AnimGraph/Classes/AnimStateConduitNode.h
+++ b/Engine/Source/Editor/AnimGraph/Classes/AnimStateConduitNode.h
@@ -7,7 +7,9 @@
 UCLASS(MinimalAPI)
 class UAnimStateConduitNode : public UAnimStateNodeBase
 {
-	GENERATED_UCLASS_BODY()
+	GENERATED_BODY()
+public:
+	ANIMGRAPH_API UAnimStateConduitNode(const FObjectInitializer& ObjectInitializer = FObjectInitializer::Get());
 public:
 
 	// The transition graph for this conduit; it's a logic graph, not an animation graph
diff --git a/Engine/Source/Editor/AnimGraph/Classes/AnimStateEntryNode.h b/Engine/Source/Editor/AnimGraph/Classes/AnimStateEntryNode.h
--- a/Engine/Source/Editor/AnimGraph/Classes/AnimStateEntryNode.h
+++ b/Engine/Source/Editor/AnimGraph/Classes/AnimStateEntryNode.h
@@ -7,7 +7,9 @@
 UCLASS(MinimalAPI)
 class UAnimStateEntryNode : public UEdGraphNode
 {
-	GENERATED_UCLASS_BODY()
+	GENERATED_BODY()
+public:
+	ANIMGRAPH_API UAnimStateEntryNode(const FObjectInitializer& ObjectInitializer = FObjectInitializer::Get());
 
 
 	// Begin UEdGraphNode interface
diff --git a/Engine/Source/Editor/AnimGraph/Classes/AnimStateNode.h b/Engine/Source/Editor/AnimGraph/Classes/AnimStateNode.h
--- a/Engine/Source/Editor/AnimGraph/Classes/AnimStateNode.h
+++ b/Engine/Source/Editor/AnimGraph/Classes/AnimStateNode.h
@@ -15,7 +15,9 @@ enum EAnimStateType
 UCLASS(MinimalAPI)
 class UAnimStateNode : public UAnimStateNodeBase
 {
-	GENERATED_UCLASS_BODY()
+	GENERATED_BODY()
+public:
+	ANIMGRAPH_API UAnimStateNode(const FObjectInitializer& ObjectInitializer = FObjectInitializer::Get());
 public:
 
 	// The animation graph for this state
diff --git a/Engine/Source/Editor/AnimGraph/Classes/AnimStateNodeBase.h b/Engine/Source/Editor/AnimGraph/Classes/AnimStateNodeBase.h
--- a/Engine/Source/Editor/AnimGraph/Classes/AnimStateNodeBase.h
+++ b/Engine/Source/Editor/AnimGraph/Classes/AnimStateNodeBase.h
@@ -7,7 +7,9 @@
 UCLASS(MinimalAPI, abstract)
 class UAnimStateNodeBase : public UEdGraphNode
 {
-	GENERATED_UCLASS_BODY()
+	GENERATED_BODY()
+public:
+	ANIMGRAPH_API UAnimStateNodeBase(const FObjectInitializer& ObjectInitializer = FObjectInitializer::Get());
 
 	// UEdGraphNode interface
 	virtual void PostPasteNode() override;
diff --git a/Engine/Source/Editor/AnimGraph/Classes/AnimStateTransitionNode.h b/Engine/Source/Editor/AnimGraph/Classes/AnimStateTransitionNode.h
--- a/Engine/Source/Editor/AnimGraph/Classes/AnimStateTransitionNode.h
+++ b/Engine/Source/Editor/AnimGraph/Classes/AnimStateTransitionNode.h
@@ -11,7 +11,9 @@
 UCLASS(MinimalAPI, config=Editor)
 class UAnimStateTransitionNode : public UAnimStateNodeBase
 {
-	GENERATED_UCLASS_BODY()
+	GENERATED_BODY()
+public:
+	ANIMGRAPH_API UAnimStateTransitionNode(const FObjectInitializer& ObjectInitializer = FObjectInitializer::Get());
 
 	// The transition logic graph for this transition (returning a boolean)
 	UPROPERTY()
diff --git a/Engine/Source/Editor/AnimGraph/Classes/AnimationConduitGraphSchema.h b/Engine/Source/Editor/AnimGraph/Classes/AnimationConduitGraphSchema.h
--- a/Engine/Source/Editor/AnimGraph/Classes/AnimationConduitGraphSchema.h
+++ b/Engine/Source/Editor/AnimGraph/Classes/AnimationConduitGraphSchema.h
@@ -8,7 +8,9 @@
 UCLASS(MinimalAPI)
 class UAnimationConduitGraphSchema : public UEdGraphSchema_K2
 {
-	GENERATED_UCLASS_BODY()
+	GENERATED_BODY()
+public:
+	ANIMGRAPH_API UAnimationConduitGraphSchema(const FObjectInitializer& ObjectInitializer = FObjectInitializer::Get());
 
 	// UEdGraphSchema interface
 	virtual void CreateDefaultNodesForGraph(UEdGraph& Graph) const override;
diff --git a/Engine/Source/Editor/AnimGraph/Classes/AnimationCustomTransitionGraph.h b/Engine/Source/Editor/AnimGraph/Classes/AnimationCustomTransitionGraph.h
--- a/Engine/Source/Editor/AnimGraph/Classes/AnimationCustomTransitionGraph.h
+++ b/Engine/Source/Editor/AnimGraph/Classes/AnimationCustomTransitionGraph.h
@@ -7,7 +7,9 @@
 UCLASS(MinimalAPI)
 class UAnimationCustomTransitionGraph : public UAnimationGraph
 {
-	GENERATED_UCLASS_BODY()
+	GENERATED_BODY()
+public:
+	ANIMGRAPH_API UAnimationCustomTransitionGraph(const FObjectInitializer& ObjectInitializer = FObjectInitializer::Get());
 
 	// Result node within the state's animation graph
 	UPROPERTY()
diff --git a/Engine/Source/Editor/AnimGraph/Classes/AnimationCustomTransitionSchema.h b/Engine/Source/Editor/AnimGraph/Classes/AnimationCustomTransitionSchema.h
--- a/Engine/Source/Editor/AnimGraph/Classes/AnimationCustomTransitionSchema.h
+++ b/Engine/Source/Editor/AnimGraph/Classes/AnimationCustomTransitionSchema.h
@@ -7,7 +7,9 @@
 UCLASS(MinimalAPI)
 class UAnimationCustomTransitionSchema : public UAnimationGraphSchema
 {
-	GENERATED_UCLASS_BODY()
+	GENERATED_BODY()
+public:
+	ANIMGRAPH_API UAnimationCustomTransitionSchema(const FObjectInitializer& ObjectInitializer = FObjectInitializer::Get());
 
 	// Begin UEdGraphSchema interface.
 	virtual void CreateDefaultNodesForGraph(UEdGraph& Graph) const override;
diff --git a/Engine/Source/Editor/AnimGraph/Classes/AnimationGraph.h b/Engine/Source/Editor/AnimGraph/Classes/AnimationGraph.h
--- a/Engine/Source/Editor/AnimGraph/Classes/AnimationGraph.h
+++ b/Engine/Source/Editor/AnimGraph/Classes/AnimationGraph.h
@@ -7,6 +7,8 @@
 UCLASS(MinimalAPI)
 class UAnimationGraph : public UEdGraph
 {
-	GENERATED_UCLASS_BODY()
+	GENERATED_BODY()
+public:
+	ANIMGRAPH_API UAnimationGraph(const FObjectInitializer& ObjectInitializer = FObjectInitializer::Get());
 };
 
diff --git a/Engine/Source/Editor/AnimGraph/Classes/AnimationGraphSchema.h b/Engine/Source/Editor/AnimGraph/Classes/AnimationGraphSchema.h
--- a/Engine/Source/Editor/AnimGraph/Classes/AnimationGraphSchema.h
+++ b/Engine/Source/Editor/AnimGraph/Classes/AnimationGraphSchema.h
@@ -8,7 +8,9 @@
 UCLASS(MinimalAPI)
 class UAnimationGraphSchema : public UEdGraphSchema_K2
 {
-	GENERATED_UCLASS_BODY()
+	GENERATED_BODY()
+public:
+	ANIMGRAPH_API UAnimationGraphSchema(const FObjectInitializer& ObjectInitializer = FObjectInitializer::Get());
 
 	// Common PinNames
 	UPROPERTY()
diff --git a/Engine/Source/Editor/AnimGraph/Classes/AnimationStateGraph.h b/Engine/Source/Editor/AnimGraph/Classes/AnimationStateGraph.h
--- a/Engine/Source/Editor/AnimGraph/Classes/AnimationStateGraph.h
+++ b/Engine/Source/Editor/AnimGraph/Classes/AnimationStateGraph.h
@@ -7,7 +7,9 @@
 UCLASS(MinimalAPI)
 class UAnimationStateGraph : public UAnimationGraph
 {
-	GENERATED_UCLASS_BODY()
+	GENERATED_BODY()
+public:
+	ANIMGRAPH_API UAnimationStateGraph(const FObjectInitializer& ObjectInitializer = FObjectInitializer::Get());
 
 	// Result node within the state's animation graph
 	UPROPERTY()
diff --git a/Engine/Source/Editor/AnimGraph/Classes/AnimationStateGraphSchema.h b/Engine/Source/Editor/AnimGraph/Classes/AnimationStateGraphSchema.h
--- a/Engine/Source/Editor/AnimGraph/Classes/AnimationStateGraphSchema.h
+++ b/Engine/Source/Editor/AnimGraph/Classes/AnimationStateGraphSchema.h
@@ -7,7 +7,9 @@
 UCLASS(MinimalAPI)
 class UAnimationStateGraphSchema : public UAnimationGraphSchema
 {
-	GENERATED_UCLASS_BODY()
+	GENERATED_BODY()
+public:
+	ANIMGRAPH_API UAnimationStateGraphSchema(const FObjectInitializer& ObjectInitializer = FObjectInitializer::Get());
 
 	// Begin UEdGraphSchema interface.
 	virtual void CreateDefaultNodesForGraph(UEdGraph& Graph) const override;
diff --git a/Engine/Source/Editor/AnimGraph/Classes/AnimationStateMachineGraph.h b/Engine/Source/Editor/AnimGraph/Classes/AnimationStateMachineGraph.h
--- a/Engine/Source/Editor/AnimGraph/Classes/AnimationStateMachineGraph.h
+++ b/Engine/Source/Editor/AnimGraph/Classes/AnimationStateMachineGraph.h
@@ -7,7 +7,9 @@
 UCLASS(MinimalAPI)
 class UAnimationStateMachineGraph : public UEdGraph
 {
-	GENERATED_UCLASS_BODY()
+	GENERATED_BODY()
+public:
+	ANIMGRAPH_API UAnimationStateMachineGraph(const FObjectInitializer& ObjectInitializer = FObjectInitializer::Get());
 
 	// Entry node within the state machine
 	UPROPERTY()
diff --git a/Engine/Source/Editor/AnimGraph/Classes/AnimationStateMachineSchema.h b/Engine/Source/Editor/AnimGraph/Classes/AnimationStateMachineSchema.h
--- a/Engine/Source/Editor/AnimGraph/Classes/AnimationStateMachineSchema.h
+++ b/Engine/Source/Editor/AnimGraph/Classes/AnimationStateMachineSchema.h
@@ -8,7 +8,7 @@
 USTRUCT()
 struct ANIMGRAPH_API FEdGraphSchemaAction_NewStateNode : public FEdGraphSchemaAction
 {
-	GENERATED_USTRUCT_BODY()
+	GENERATED_BODY()
 
 	UEdGraphNode* NodeTemplate;
 
@@ -39,7 +39,7 @@ struct ANIMGRAPH_API FEdGraphSchemaAction_NewStateNode : public FEdGraphSchemaAc
 USTRUCT()
 struct ANIMGRAPH_API FEdGraphSchemaAction_NewStateComment : public FEdGraphSchemaAction
 {
-	GENERATED_USTRUCT_BODY();
+	GENERATED_BODY()
 
 	FEdGraphSchemaAction_NewStateComment() 
 		: FEdGraphSchemaAction()
@@ -58,7 +58,9 @@ struct ANIMGRAPH_API FEdGraphSchemaAction_NewStateComment : public FEdGraphSchem
 UCLASS(MinimalAPI)
 class UAnimationStateMachineSchema : public UEdGraphSchema
 {
-	GENERATED_UCLASS_BODY()
+	GENERATED_BODY()
+public:
+	ANIMGRAPH_API UAnimationStateMachineSchema(const FObjectInitializer& ObjectInitializer = FObjectInitializer::Get());
 
 	UPROPERTY()
 	FString PC_Exec;
diff --git a/Engine/Source/Editor/AnimGraph/Classes/AnimationTransitionGraph.h b/Engine/Source/Editor/AnimGraph/Classes/AnimationTransitionGraph.h
--- a/Engine/Source/Editor/AnimGraph/Classes/AnimationTransitionGraph.h
+++ b/Engine/Source/Editor/AnimGraph/Classes/AnimationTransitionGraph.h
@@ -7,7 +7,9 @@
 UCLASS(MinimalAPI)
 class UAnimationTransitionGraph : public UAnimationGraph
 {
-	GENERATED_UCLASS_BODY()
+	GENERATED_BODY()
+public:
+	ANIMGRAPH_API UAnimationTransitionGraph(const FObjectInitializer& ObjectInitializer = FObjectInitializer::Get());
 
 	UPROPERTY()
 	class UAnimGraphNode_TransitionResult* MyResultNode;
diff --git a/Engine/Source/Editor/AnimGraph/Classes/AnimationTransitionSchema.h b/Engine/Source/Editor/AnimGraph/Classes/AnimationTransitionSchema.h
--- a/Engine/Source/Editor/AnimGraph/Classes/AnimationTransitionSchema.h
+++ b/Engine/Source/Editor/AnimGraph/Classes/AnimationTransitionSchema.h
@@ -7,7 +7,9 @@
 UCLASS(MinimalAPI)
 class UAnimationTransitionSchema : public UEdGraphSchema_K2
 {
-	GENERATED_UCLASS_BODY()
+	GENERATED_BODY()
+public:
+	ANIMGRAPH_API UAnimationTransitionSchema(const FObjectInitializer& ObjectInitializer = FObjectInitializer::Get());
 
 	// Begin UEdGraphSchema interface.
 	virtual void GetGraphContextActions(FGraphContextMenuBuilder& ContextMenuBuilder) const override;
diff --git a/Engine/Source/Editor/AnimGraph/Classes/K2Node_TransitionRuleGetter.h b/Engine/Source/Editor/AnimGraph/Classes/K2Node_TransitionRuleGetter.h
--- a/Engine/Source/Editor/AnimGraph/Classes/K2Node_TransitionRuleGetter.h
+++ b/Engine/Source/Editor/AnimGraph/Classes/K2Node_TransitionRuleGetter.h
@@ -25,7 +25,9 @@ namespace ETransitionGetter
 UCLASS(MinimalAPI)
 class UK2Node_TransitionRuleGetter : public UK2Node
 {
-	GENERATED_UCLASS_BODY()
+	GENERATED_BODY()
+public:
+	ANIMGRAPH_API UK2Node_TransitionRuleGetter(const FObjectInitializer& ObjectInitializer = FObjectInitializer::Get());
 
 	UPROPERTY()
 	TEnumAsByte<ETransitionGetter::Type> GetterType;
diff --git a/Engine/Source/Editor/BehaviorTreeEditor/Classes/BehaviorTreeDecoratorGraph.h b/Engine/Source/Editor/BehaviorTreeEditor/Classes/BehaviorTreeDecoratorGraph.h
--- a/Engine/Source/Editor/BehaviorTreeEditor/Classes/BehaviorTreeDecoratorGraph.h
+++ b/Engine/Source/Editor/BehaviorTreeEditor/Classes/BehaviorTreeDecoratorGraph.h
@@ -7,7 +7,9 @@
 UCLASS()
 class UBehaviorTreeDecoratorGraph : public UEdGraph
 {
-	GENERATED_UCLASS_BODY()
+	GENERATED_BODY()
+public:
+	UBehaviorTreeDecoratorGraph(const FObjectInitializer& ObjectInitializer = FObjectInitializer::Get());
 
 	void CollectDecoratorData(TArray<class UBTDecorator*>& DecoratorInstances, TArray<struct FBTDecoratorLogic>& DecoratorOperations) const;
 	int32 SpawnMissingNodes(const TArray<class UBTDecorator*>& NodeInstances, const TArray<struct FBTDecoratorLogic>& Operations, int32 StartIndex);
diff --git a/Engine/Source/Editor/BehaviorTreeEditor/Classes/BehaviorTreeDecoratorGraphNode.h b/Engine/Source/Editor/BehaviorTreeEditor/Classes/BehaviorTreeDecoratorGraphNode.h
--- a/Engine/Source/Editor/BehaviorTreeEditor/Classes/BehaviorTreeDecoratorGraphNode.h
+++ b/Engine/Source/Editor/BehaviorTreeEditor/Classes/BehaviorTreeDecoratorGraphNode.h
@@ -10,7 +10,9 @@
 UCLASS()
 class UBehaviorTreeDecoratorGraphNode : public UEdGraphNode
 {
-	GENERATED_UCLASS_BODY()
+	GENERATED_BODY()
+public:
+	UBehaviorTreeDecoratorGraphNode(const FObjectInitializer& ObjectInitializer = FObjectInitializer::Get());
 
 	/** Get the BT graph that owns this node */
 	virtual class UBehaviorTreeDecoratorGraph* GetDecoratorGraph();
diff --git a/Engine/Source/Editor/BehaviorTreeEditor/Classes/BehaviorTreeDecoratorGraphNode_Decorator.h b/Engine/Source/Editor/BehaviorTreeEditor/Classes/BehaviorTreeDecoratorGraphNode_Decorator.h
--- a/Engine/Source/Editor/BehaviorTreeEditor/Classes/BehaviorTreeDecoratorGraphNode_Decorator.h
+++ b/Engine/Source/Editor/BehaviorTreeEditor/Classes/BehaviorTreeDecoratorGraphNode_Decorator.h
@@ -8,7 +8,9 @@
 UCLASS()
 class UBehaviorTreeDecoratorGraphNode_Decorator : public UBehaviorTreeDecoratorGraphNode 
 {
-	GENERATED_UCLASS_BODY()
+	GENERATED_BODY()
+public:
+	UBehaviorTreeDecoratorGraphNode_Decorator(const FObjectInitializer& ObjectInitializer = FObjectInitializer::Get());
 
 	UPROPERTY()
 	UObject* NodeInstance;
diff --git a/Engine/Source/Editor/BehaviorTreeEditor/Classes/BehaviorTreeDecoratorGraphNode_Logic.h b/Engine/Source/Editor/BehaviorTreeEditor/Classes/BehaviorTreeDecoratorGraphNode_Logic.h
--- a/Engine/Source/Editor/BehaviorTreeEditor/Classes/BehaviorTreeDecoratorGraphNode_Logic.h
+++ b/Engine/Source/Editor/BehaviorTreeEditor/Classes/BehaviorTreeDecoratorGraphNode_Logic.h
@@ -18,7 +18,9 @@ namespace EDecoratorLogicMode
 UCLASS()
 class UBehaviorTreeDecoratorGraphNode_Logic : public UBehaviorTreeDecoratorGraphNode
 {
-	GENERATED_UCLASS_BODY()
+	GENERATED_BODY()
+public:
+	UBehaviorTreeDecoratorGraphNode_Logic(const FObjectInitializer& ObjectInitializer = FObjectInitializer::Get());
 
 	UPROPERTY()
 	TEnumAsByte<EDecoratorLogicMode::Type> LogicMode;
diff --git a/Engine/Source/Editor/BehaviorTreeEditor/Classes/BehaviorTreeEditorTypes.h b/Engine/Source/Editor/BehaviorTreeEditor/Classes/BehaviorTreeEditorTypes.h
--- a/Engine/Source/Editor/BehaviorTreeEditor/Classes/BehaviorTreeEditorTypes.h
+++ b/Engine/Source/Editor/BehaviorTreeEditor/Classes/BehaviorTreeEditorTypes.h
@@ -46,7 +46,9 @@ struct FNodeBounds
 UCLASS()
 class UBehaviorTreeEditorTypes : public UObject
 {
-	GENERATED_UCLASS_BODY()
+	GENERATED_BODY()
+public:
+	UBehaviorTreeEditorTypes(const FObjectInitializer& ObjectInitializer = FObjectInitializer::Get());
 
 	static const FString PinCategory_MultipleNodes;
 	static const FString PinCategory_SingleComposite;
diff --git a/Engine/Source/Editor/BehaviorTreeEditor/Classes/BehaviorTreeFactory.h b/Engine/Source/Editor/BehaviorTreeEditor/Classes/BehaviorTreeFactory.h
--- a/Engine/Source/Editor/BehaviorTreeEditor/Classes/BehaviorTreeFactory.h
+++ b/Engine/Source/Editor/BehaviorTreeEditor/Classes/BehaviorTreeFactory.h
@@ -6,7 +6,9 @@
 UCLASS()
 class UBehaviorTreeFactory : public UFactory
 {
-	GENERATED_UCLASS_BODY()
+	GENERATED_BODY()
+public:
+	UBehaviorTreeFactory(const FObjectInitializer& ObjectInitializer = FObjectInitializer::Get());
 
 	// UFactory interface
 	virtual UObject* FactoryCreateNew(UClass* Class,UObject* InParent,FName Name,EObjectFlags Flags,UObject* Context,FFeedbackContext* Warn) override;
diff --git a/Engine/Source/Editor/BehaviorTreeEditor/Classes/BehaviorTreeGraph.h b/Engine/Source/Editor/BehaviorTreeEditor/Classes/BehaviorTreeGraph.h
--- a/Engine/Source/Editor/BehaviorTreeEditor/Classes/BehaviorTreeGraph.h
+++ b/Engine/Source/Editor/BehaviorTreeEditor/Classes/BehaviorTreeGraph.h
@@ -8,7 +8,9 @@
 UCLASS()
 class UBehaviorTreeGraph : public UAIGraph
 {
-	GENERATED_UCLASS_BODY()
+	GENERATED_BODY()
+public:
+	UBehaviorTreeGraph(const FObjectInitializer& ObjectInitializer = FObjectInitializer::Get());
 
 	enum EUpdateFlags
 	{
diff --git a/Engine/Source/Editor/BehaviorTreeEditor/Classes/BehaviorTreeGraphNode.h b/Engine/Source/Editor/BehaviorTreeEditor/Classes/BehaviorTreeGraphNode.h
--- a/Engine/Source/Editor/BehaviorTreeEditor/Classes/BehaviorTreeGraphNode.h
+++ b/Engine/Source/Editor/BehaviorTreeEditor/Classes/BehaviorTreeGraphNode.h
@@ -9,7 +9,9 @@
 UCLASS()
 class UBehaviorTreeGraphNode : public UAIGraphNode
 {
-	GENERATED_UCLASS_BODY()
+	GENERATED_BODY()
+public:
+	UBehaviorTreeGraphNode(const FObjectInitializer& ObjectInitializer = FObjectInitializer::Get());
 
 	/** only some of behavior tree nodes support decorators */
 	UPROPERTY()
diff --git a/Engine/Source/Editor/BehaviorTreeEditor/Classes/BehaviorTreeGraphNode_Composite.h b/Engine/Source/Editor/BehaviorTreeEditor/Classes/BehaviorTreeGraphNode_Composite.h
--- a/Engine/Source/Editor/BehaviorTreeEditor/Classes/BehaviorTreeGraphNode_Composite.h
+++ b/Engine/Source/Editor/BehaviorTreeEditor/Classes/BehaviorTreeGraphNode_Composite.h
@@ -6,7 +6,9 @@
 UCLASS()
 class UBehaviorTreeGraphNode_Composite : public UBehaviorTreeGraphNode
 {
-	GENERATED_UCLASS_BODY()
+	GENERATED_BODY()
+public:
+	UBehaviorTreeGraphNode_Composite(const FObjectInitializer& ObjectInitializer = FObjectInitializer::Get());
 	
 	virtual FText GetNodeTitle(ENodeTitleType::Type TitleType) const override;
 	virtual bool RefreshNodeClass() override{ return false; }
diff --git a/Engine/Source/Editor/BehaviorTreeEditor/Classes/BehaviorTreeGraphNode_CompositeDecorator.h b/Engine/Source/Editor/BehaviorTreeEditor/Classes/BehaviorTreeGraphNode_CompositeDecorator.h
--- a/Engine/Source/Editor/BehaviorTreeEditor/Classes/BehaviorTreeGraphNode_CompositeDecorator.h
+++ b/Engine/Source/Editor/BehaviorTreeEditor/Classes/BehaviorTreeGraphNode_CompositeDecorator.h
@@ -6,7 +6,9 @@
 UCLASS()
 class UBehaviorTreeGraphNode_CompositeDecorator : public UBehaviorTreeGraphNode
 {
-	GENERATED_UCLASS_BODY()
+	GENERATED_BODY()
+public:
+	UBehaviorTreeGraphNode_CompositeDecorator(const FObjectInitializer& ObjectInitializer = FObjectInitializer::Get());
 
 	// The logic graph for this decorator (returning a boolean)
 	UPROPERTY()
diff --git a/Engine/Source/Editor/BehaviorTreeEditor/Classes/BehaviorTreeGraphNode_Decorator.h b/Engine/Source/Editor/BehaviorTreeEditor/Classes/BehaviorTreeGraphNode_Decorator.h
--- a/Engine/Source/Editor/BehaviorTreeEditor/Classes/BehaviorTreeGraphNode_Decorator.h
+++ b/Engine/Source/Editor/BehaviorTreeEditor/Classes/BehaviorTreeGraphNode_Decorator.h
@@ -6,7 +6,9 @@
 UCLASS()
 class UBehaviorTreeGraphNode_Decorator : public UBehaviorTreeGraphNode
 {
-	GENERATED_UCLASS_BODY()
+	GENERATED_BODY()
+public:
+	UBehaviorTreeGraphNode_Decorator(const FObjectInitializer& ObjectInitializer = FObjectInitializer::Get());
 
 	virtual FText GetNodeTitle(ENodeTitleType::Type TitleType) const override;
 	virtual void AllocateDefaultPins() override;
diff --git a/Engine/Source/Editor/BehaviorTreeEditor/Classes/BehaviorTreeGraphNode_Root.h b/Engine/Source/Editor/BehaviorTreeEditor/Classes/BehaviorTreeGraphNode_Root.h
--- a/Engine/Source/Editor/BehaviorTreeEditor/Classes/BehaviorTreeGraphNode_Root.h
+++ b/Engine/Source/Editor/BehaviorTreeEditor/Classes/BehaviorTreeGraphNode_Root.h
@@ -7,7 +7,9 @@
 UCLASS()
 class UBehaviorTreeGraphNode_Root : public UBehaviorTreeGraphNode
 {
-	GENERATED_UCLASS_BODY()
+	GENERATED_BODY()
+public:
+	UBehaviorTreeGraphNode_Root(const FObjectInitializer& ObjectInitializer = FObjectInitializer::Get());
 
 	UPROPERTY(EditAnywhere, Category=BehaviorTree)
 	class UBlackboardData* BlackboardAsset;
diff --git a/Engine/Source/Editor/BehaviorTreeEditor/Classes/BehaviorTreeGraphNode_Service.h b/Engine/Source/Editor/BehaviorTreeEditor/Classes/BehaviorTreeGraphNode_Service.h
--- a/Engine/Source/Editor/BehaviorTreeEditor/Classes/BehaviorTreeGraphNode_Service.h
+++ b/Engine/Source/Editor/BehaviorTreeEditor/Classes/BehaviorTreeGraphNode_Service.h
@@ -6,7 +6,9 @@
 UCLASS()
 class UBehaviorTreeGraphNode_Service : public UBehaviorTreeGraphNode
 {
-	GENERATED_UCLASS_BODY()
+	GENERATED_BODY()
+public:
+	UBehaviorTreeGraphNode_Service(const FObjectInitializer& ObjectInitializer = FObjectInitializer::Get());
 
 	virtual FText GetNodeTitle(ENodeTitleType::Type TitleType) const override;
 	virtual void AllocateDefaultPins() override;
diff --git a/Engine/Source/Editor/BehaviorTreeEditor/Classes/BehaviorTreeGraphNode_SimpleParallel.h b/Engine/Source/Editor/BehaviorTreeEditor/Classes/BehaviorTreeGraphNode_SimpleParallel.h
--- a/Engine/Source/Editor/BehaviorTreeEditor/Classes/BehaviorTreeGraphNode_SimpleParallel.h
+++ b/Engine/Source/Editor/BehaviorTreeEditor/Classes/BehaviorTreeGraphNode_SimpleParallel.h
@@ -6,7 +6,9 @@
 UCLASS()
 class UBehaviorTreeGraphNode_SimpleParallel : public UBehaviorTreeGraphNode_Composite
 {
-	GENERATED_UCLASS_BODY()
+	GENERATED_BODY()
+public:
+	UBehaviorTreeGraphNode_SimpleParallel(const FObjectInitializer& ObjectInitializer = FObjectInitializer::Get());
 
 	virtual void AllocateDefaultPins() override;
 	virtual void GetPinHoverText(const UEdGraphPin& Pin, FString& HoverTextOut) const override;
diff --git a/Engine/Source/Editor/BehaviorTreeEditor/Classes/BehaviorTreeGraphNode_SubtreeTask.h b/Engine/Source/Editor/BehaviorTreeEditor/Classes/BehaviorTreeGraphNode_SubtreeTask.h
--- a/Engine/Source/Editor/BehaviorTreeEditor/Classes/BehaviorTreeGraphNode_SubtreeTask.h
+++ b/Engine/Source/Editor/BehaviorTreeEditor/Classes/BehaviorTreeGraphNode_SubtreeTask.h
@@ -6,7 +6,9 @@
 UCLASS()
 class UBehaviorTreeGraphNode_SubtreeTask : public UBehaviorTreeGraphNode_Task
 {
-	GENERATED_UCLASS_BODY()
+	GENERATED_BODY()
+public:
+	UBehaviorTreeGraphNode_SubtreeTask(const FObjectInitializer& ObjectInitializer = FObjectInitializer::Get());
 
 	/** UBehaviorTreeGraph.UpdateCounter value of subtree graph */
 	int32 SubtreeVersion;
diff --git a/Engine/Source/Editor/BehaviorTreeEditor/Classes/BehaviorTreeGraphNode_Task.h b/Engine/Source/Editor/BehaviorTreeEditor/Classes/BehaviorTreeGraphNode_Task.h
--- a/Engine/Source/Editor/BehaviorTreeEditor/Classes/BehaviorTreeGraphNode_Task.h
+++ b/Engine/Source/Editor/BehaviorTreeEditor/Classes/BehaviorTreeGraphNode_Task.h
@@ -7,7 +7,9 @@
 UCLASS()
 class UBehaviorTreeGraphNode_Task : public UBehaviorTreeGraphNode
 {
-	GENERATED_UCLASS_BODY()
+	GENERATED_BODY()
+public:
+	UBehaviorTreeGraphNode_Task(const FObjectInitializer& ObjectInitializer = FObjectInitializer::Get());
 
 	virtual void AllocateDefaultPins() override;
 	virtual FText GetNodeTitle(ENodeTitleType::Type TitleType) const override;
diff --git a/Engine/Source/Editor/BehaviorTreeEditor/Classes/BlackboardDataFactory.h b/Engine/Source/Editor/BehaviorTreeEditor/Classes/BlackboardDataFactory.h
--- a/Engine/Source/Editor/BehaviorTreeEditor/Classes/BlackboardDataFactory.h
+++ b/Engine/Source/Editor/BehaviorTreeEditor/Classes/BlackboardDataFactory.h
@@ -6,7 +6,9 @@
 UCLASS()
 class UBlackboardDataFactory : public UFactory
 {
-	GENERATED_UCLASS_BODY()
+	GENERATED_BODY()
+public:
+	UBlackboardDataFactory(const FObjectInitializer& ObjectInitializer = FObjectInitializer::Get());
 
 	// UFactory interface
 	virtual UObject* FactoryCreateNew(UClass* Class,UObject* InParent,FName Name,EObjectFlags Flags,UObject* Context,FFeedbackContext* Warn) override;
diff --git a/Engine/Source/Editor/BehaviorTreeEditor/Classes/EdGraphSchema_BehaviorTree.h b/Engine/Source/Editor/BehaviorTreeEditor/Classes/EdGraphSchema_BehaviorTree.h
--- a/Engine/Source/Editor/BehaviorTreeEditor/Classes/EdGraphSchema_BehaviorTree.h
+++ b/Engine/Source/Editor/BehaviorTreeEditor/Classes/EdGraphSchema_BehaviorTree.h
@@ -9,7 +9,7 @@
 USTRUCT()
 struct FBehaviorTreeSchemaAction_AutoArrange : public FEdGraphSchemaAction
 {
-	GENERATED_USTRUCT_BODY();
+	GENERATED_BODY()
 
 	FBehaviorTreeSchemaAction_AutoArrange() 
 		: FEdGraphSchemaAction() {}
@@ -26,7 +26,9 @@ struct FBehaviorTreeSchemaAction_AutoArrange : public FEdGraphSchemaAction
 UCLASS(MinimalAPI)
 class UEdGraphSchema_BehaviorTree : public UAIGraphSchema
 {
-	GENERATED_UCLASS_BODY()
+	GENERATED_BODY()
+public:
+	BEHAVIORTREEEDITOR_API UEdGraphSchema_BehaviorTree(const FObjectInitializer& ObjectInitializer = FObjectInitializer::Get());
 
 	// Begin EdGraphSchema interface
 	virtual void CreateDefaultNodesForGraph(UEdGraph& Graph) const override;
diff --git a/Engine/Source/Editor/BehaviorTreeEditor/Classes/EdGraphSchema_BehaviorTreeDecorator.h b/Engine/Source/Editor/BehaviorTreeEditor/Classes/EdGraphSchema_BehaviorTreeDecorator.h
--- a/Engine/Source/Editor/BehaviorTreeEditor/Classes/EdGraphSchema_BehaviorTreeDecorator.h
+++ b/Engine/Source/Editor/BehaviorTreeEditor/Classes/EdGraphSchema_BehaviorTreeDecorator.h
@@ -7,7 +7,7 @@
 USTRUCT()
 struct FDecoratorSchemaAction_NewNode : public FEdGraphSchemaAction
 {
-	GENERATED_USTRUCT_BODY();
+	GENERATED_BODY()
 
 	/** Template of node we want to create */
 	UPROPERTY()
@@ -42,7 +42,9 @@ struct FDecoratorSchemaAction_NewNode : public FEdGraphSchemaAction
 UCLASS(MinimalAPI)
 class UEdGraphSchema_BehaviorTreeDecorator : public UEdGraphSchema
 {
-	GENERATED_UCLASS_BODY()
+	GENERATED_BODY()
+public:
+	BEHAVIORTREEEDITOR_API UEdGraphSchema_BehaviorTreeDecorator(const FObjectInitializer& ObjectInitializer = FObjectInitializer::Get());
 
 	UPROPERTY()
 	FString PC_Boolean;
diff --git a/Engine/Source/Editor/BlueprintGraph/Classes/BlueprintBoundEventNodeSpawner.h b/Engine/Source/Editor/BlueprintGraph/Classes/BlueprintBoundEventNodeSpawner.h
--- a/Engine/Source/Editor/BlueprintGraph/Classes/BlueprintBoundEventNodeSpawner.h
+++ b/Engine/Source/Editor/BlueprintGraph/Classes/BlueprintBoundEventNodeSpawner.h
@@ -17,7 +17,9 @@ class UK2Node_Event;
 UCLASS(Transient)
 class BLUEPRINTGRAPH_API UBlueprintBoundEventNodeSpawner : public UBlueprintEventNodeSpawner
 {
-	GENERATED_UCLASS_BODY()
+	GENERATED_BODY()
+public:
+	UBlueprintBoundEventNodeSpawner(const FObjectInitializer& ObjectInitializer = FObjectInitializer::Get());
 
 public:
 	/**
diff --git a/Engine/Source/Editor/BlueprintGraph/Classes/BlueprintBoundNodeSpawner.h b/Engine/Source/Editor/BlueprintGraph/Classes/BlueprintBoundNodeSpawner.h
--- a/Engine/Source/Editor/BlueprintGraph/Classes/BlueprintBoundNodeSpawner.h
+++ b/Engine/Source/Editor/BlueprintGraph/Classes/BlueprintBoundNodeSpawner.h
@@ -12,7 +12,9 @@
 UCLASS(Transient)
 class BLUEPRINTGRAPH_API UBlueprintBoundNodeSpawner : public UBlueprintNodeSpawner
 {
-	GENERATED_UCLASS_BODY()
+	GENERATED_BODY()
+public:
+	UBlueprintBoundNodeSpawner(const FObjectInitializer& ObjectInitializer = FObjectInitializer::Get());
 	DECLARE_DELEGATE_RetVal_OneParam(bool, FCanBindObjectDelegate, UObject const*);
 	DECLARE_DELEGATE_RetVal_TwoParams(bool, FOnBindObjectDelegate, UEdGraphNode*, UObject*);
 	DECLARE_DELEGATE_RetVal_TwoParams(UEdGraphNode*, FFindPreExistingNodeDelegate, const UBlueprint*, IBlueprintNodeBinder::FBindingSet const& );
diff --git a/Engine/Source/Editor/BlueprintGraph/Classes/BlueprintComponentNodeSpawner.h b/Engine/Source/Editor/BlueprintGraph/Classes/BlueprintComponentNodeSpawner.h
--- a/Engine/Source/Editor/BlueprintGraph/Classes/BlueprintComponentNodeSpawner.h
+++ b/Engine/Source/Editor/BlueprintGraph/Classes/BlueprintComponentNodeSpawner.h
@@ -17,7 +17,9 @@ class UActorComponent;
 UCLASS(Transient)
 class BLUEPRINTGRAPH_API UBlueprintComponentNodeSpawner : public UBlueprintNodeSpawner
 {
-	GENERATED_UCLASS_BODY()
+	GENERATED_BODY()
+public:
+	UBlueprintComponentNodeSpawner(const FObjectInitializer& ObjectInitializer = FObjectInitializer::Get());
 
 public:
 	/**
diff --git a/Engine/Source/Editor/BlueprintGraph/Classes/BlueprintDelegateNodeSpawner.h b/Engine/Source/Editor/BlueprintGraph/Classes/BlueprintDelegateNodeSpawner.h
--- a/Engine/Source/Editor/BlueprintGraph/Classes/BlueprintDelegateNodeSpawner.h
+++ b/Engine/Source/Editor/BlueprintGraph/Classes/BlueprintDelegateNodeSpawner.h
@@ -16,7 +16,9 @@ class UK2Node_BaseMCDelegate;
 UCLASS(Transient)
 class BLUEPRINTGRAPH_API UBlueprintDelegateNodeSpawner : public UBlueprintFieldNodeSpawner
 {
-	GENERATED_UCLASS_BODY()
+	GENERATED_BODY()
+public:
+	UBlueprintDelegateNodeSpawner(const FObjectInitializer& ObjectInitializer = FObjectInitializer::Get());
 
 public:
 	/**
diff --git a/Engine/Source/Editor/BlueprintGraph/Classes/BlueprintEventNodeSpawner.h b/Engine/Source/Editor/BlueprintGraph/Classes/BlueprintEventNodeSpawner.h
--- a/Engine/Source/Editor/BlueprintGraph/Classes/BlueprintEventNodeSpawner.h
+++ b/Engine/Source/Editor/BlueprintGraph/Classes/BlueprintEventNodeSpawner.h
@@ -18,7 +18,9 @@ class UK2Node_Event;
 UCLASS(Transient)
 class BLUEPRINTGRAPH_API UBlueprintEventNodeSpawner : public UBlueprintNodeSpawner
 {
-	GENERATED_UCLASS_BODY()
+	GENERATED_BODY()
+public:
+	UBlueprintEventNodeSpawner(const FObjectInitializer& ObjectInitializer = FObjectInitializer::Get());
 
 public:
 	/**
diff --git a/Engine/Source/Editor/BlueprintGraph/Classes/BlueprintFieldNodeSpawner.h b/Engine/Source/Editor/BlueprintGraph/Classes/BlueprintFieldNodeSpawner.h
--- a/Engine/Source/Editor/BlueprintGraph/Classes/BlueprintFieldNodeSpawner.h
+++ b/Engine/Source/Editor/BlueprintGraph/Classes/BlueprintFieldNodeSpawner.h
@@ -16,7 +16,9 @@ class UK2Node;
 UCLASS(Transient)
 class BLUEPRINTGRAPH_API UBlueprintFieldNodeSpawner : public UBlueprintNodeSpawner
 {
-	GENERATED_UCLASS_BODY()
+	GENERATED_BODY()
+public:
+	UBlueprintFieldNodeSpawner(const FObjectInitializer& ObjectInitializer = FObjectInitializer::Get());
 	DECLARE_DELEGATE_TwoParams(FSetNodeFieldDelegate, UEdGraphNode*, UField const*);
 
 public:
diff --git a/Engine/Source/Editor/BlueprintGraph/Classes/BlueprintFunctionNodeSpawner.h b/Engine/Source/Editor/BlueprintGraph/Classes/BlueprintFunctionNodeSpawner.h
--- a/Engine/Source/Editor/BlueprintGraph/Classes/BlueprintFunctionNodeSpawner.h
+++ b/Engine/Source/Editor/BlueprintGraph/Classes/BlueprintFunctionNodeSpawner.h
@@ -15,7 +15,9 @@ class UK2Node_CallFunction;
 UCLASS(Transient)
 class BLUEPRINTGRAPH_API UBlueprintFunctionNodeSpawner : public UBlueprintFieldNodeSpawner
 {
-	GENERATED_UCLASS_BODY()
+	GENERATED_BODY()
+public:
+	UBlueprintFunctionNodeSpawner(const FObjectInitializer& ObjectInitializer = FObjectInitializer::Get());
 
 public:
 	/**
diff --git a/Engine/Source/Editor/BlueprintGraph/Classes/BlueprintVariableNodeSpawner.h b/Engine/Source/Editor/BlueprintGraph/Classes/BlueprintVariableNodeSpawner.h
--- a/Engine/Source/Editor/BlueprintGraph/Classes/BlueprintVariableNodeSpawner.h
+++ b/Engine/Source/Editor/BlueprintGraph/Classes/BlueprintVariableNodeSpawner.h
@@ -19,7 +19,9 @@ class UK2Node_Variable;
 UCLASS(Transient)
 class BLUEPRINTGRAPH_API UBlueprintVariableNodeSpawner : public UBlueprintFieldNodeSpawner
 {
-	GENERATED_UCLASS_BODY()
+	GENERATED_BODY()
+public:
+	UBlueprintVariableNodeSpawner(const FObjectInitializer& ObjectInitializer = FObjectInitializer::Get());
 
 public:
 	/**
@@ -103,4 +105,4 @@ private:
 	/** */
 	UPROPERTY()
 	FBPVariableDescription LocalVarDesc;
-};
\ No newline at end of file
+};
diff --git a/Engine/Source/Editor/BlueprintGraph/Classes/EdGraphSchema_K2.h b/Engine/Source/Editor/BlueprintGraph/Classes/EdGraphSchema_K2.h
--- a/Engine/Source/Editor/BlueprintGraph/Classes/EdGraphSchema_K2.h
+++ b/Engine/Source/Editor/BlueprintGraph/Classes/EdGraphSchema_K2.h
@@ -8,7 +8,7 @@
 USTRUCT()
 struct BLUEPRINTGRAPH_API FEdGraphSchemaAction_K2Struct : public FEdGraphSchemaAction
 {
-	GENERATED_USTRUCT_BODY()
+	GENERATED_BODY()
 
 	// Simple type info
 	static FName StaticGetTypeId() {static FName Type("FEdGraphSchemaAction_K2Struct"); return Type;}
@@ -214,7 +214,7 @@ USTRUCT()
 // into their never versions (with class param instead of blueprint).
 struct FBlueprintCallableFunctionRedirect
 {
-	GENERATED_USTRUCT_BODY()
+	GENERATED_BODY()
 
 	UPROPERTY()
 	FString ClassName;
@@ -235,7 +235,9 @@ struct FBlueprintCallableFunctionRedirect
 UCLASS(config=Editor)
 class BLUEPRINTGRAPH_API UEdGraphSchema_K2 : public UEdGraphSchema
 {
-	GENERATED_UCLASS_BODY()
+	GENERATED_BODY()
+public:
+	UEdGraphSchema_K2(const FObjectInitializer& ObjectInitializer = FObjectInitializer::Get());
 
 	// Allowable PinType.PinCategory values
 	static const FString PC_Exec;
diff --git a/Engine/Source/Editor/BlueprintGraph/Classes/EdGraphSchema_K2_Actions.h b/Engine/Source/Editor/BlueprintGraph/Classes/EdGraphSchema_K2_Actions.h
--- a/Engine/Source/Editor/BlueprintGraph/Classes/EdGraphSchema_K2_Actions.h
+++ b/Engine/Source/Editor/BlueprintGraph/Classes/EdGraphSchema_K2_Actions.h
@@ -12,7 +12,7 @@
 USTRUCT()
 struct BLUEPRINTGRAPH_API FEdGraphSchemaAction_K2NewNode : public FEdGraphSchemaAction
 {
-	GENERATED_USTRUCT_BODY()
+	GENERATED_BODY()
 
 	// Simple type info
 	static FName StaticGetTypeId() {static FName Type("FEdGraphSchemaAction_K2NewNode"); return Type;}
@@ -63,7 +63,7 @@ struct BLUEPRINTGRAPH_API FEdGraphSchemaAction_K2NewNode : public FEdGraphSchema
 USTRUCT()
 struct BLUEPRINTGRAPH_API FEdGraphSchemaAction_K2ViewNode : public FEdGraphSchemaAction
 {
-	GENERATED_USTRUCT_BODY()
+	GENERATED_BODY()
 
 	// Simple type info
 	static FName StaticGetTypeId() {static FName Type("FEdGraphSchemaAction_K2ViewNode"); return Type;}
@@ -98,7 +98,7 @@ struct BLUEPRINTGRAPH_API FEdGraphSchemaAction_K2ViewNode : public FEdGraphSchem
 USTRUCT()
 struct BLUEPRINTGRAPH_API FEdGraphSchemaAction_K2AssignDelegate : public FEdGraphSchemaAction_K2NewNode
 {
-	GENERATED_USTRUCT_BODY()
+	GENERATED_BODY()
 
 	static FName StaticGetTypeId() {static FName Type("FEdGraphSchemaAction_K2AssignDelegate"); return Type;}
 	virtual FName GetTypeId() const override { return StaticGetTypeId(); } 
@@ -126,7 +126,7 @@ struct BLUEPRINTGRAPH_API FEdGraphSchemaAction_K2AssignDelegate : public FEdGrap
 USTRUCT()
 struct BLUEPRINTGRAPH_API FEdGraphSchemaAction_EventFromFunction : public FEdGraphSchemaAction
 {
-	GENERATED_USTRUCT_BODY()
+	GENERATED_BODY()
 
 	static FName StaticGetTypeId() {static FName Type("FEdGraphSchemaAction_EventFromFunction"); return Type;}
 	virtual FName GetTypeId() const override { return StaticGetTypeId(); } 
@@ -159,7 +159,7 @@ struct BLUEPRINTGRAPH_API FEdGraphSchemaAction_EventFromFunction : public FEdGra
 USTRUCT()
 struct FEdGraphSchemaAction_K2AddComponent : public FEdGraphSchemaAction_K2NewNode
 {
-	GENERATED_USTRUCT_BODY()
+	GENERATED_BODY()
 
 	// Simple type info
 	static FName StaticGetTypeId() {static FName Type("FEdGraphSchemaAction_K2AddComponent"); return Type;}
@@ -199,7 +199,7 @@ struct FEdGraphSchemaAction_K2AddComponent : public FEdGraphSchemaAction_K2NewNo
 USTRUCT()
 struct FEdGraphSchemaAction_K2AddTimeline : public FEdGraphSchemaAction_K2NewNode
 {
-	GENERATED_USTRUCT_BODY()
+	GENERATED_BODY()
 
 	// Simple type info
 	static FName StaticGetTypeId() {static FName Type("FEdGraphSchemaAction_K2AddTimeline"); return Type;}
@@ -226,7 +226,7 @@ struct FEdGraphSchemaAction_K2AddTimeline : public FEdGraphSchemaAction_K2NewNod
 USTRUCT()
 struct FEdGraphSchemaAction_K2AddEvent : public FEdGraphSchemaAction_K2NewNode
 {
-	GENERATED_USTRUCT_BODY()
+	GENERATED_BODY()
 
 	// Simple type info
 	static FName StaticGetTypeId() {static FName Type("FEdGraphSchemaAction_K2AddEvent"); return Type;}
@@ -256,7 +256,7 @@ struct FEdGraphSchemaAction_K2AddEvent : public FEdGraphSchemaAction_K2NewNode
 USTRUCT()
 struct FEdGraphSchemaAction_K2AddCustomEvent : public FEdGraphSchemaAction_K2NewNode
 {
-	GENERATED_USTRUCT_BODY()
+	GENERATED_BODY()
 
 	// Simple type info
 	static FName StaticGetTypeId() {static FName Type("FEdGraphSchemaAction_K2AddCustomEvent"); return Type;}
@@ -283,7 +283,7 @@ struct FEdGraphSchemaAction_K2AddCustomEvent : public FEdGraphSchemaAction_K2New
 USTRUCT()
 struct FEdGraphSchemaAction_K2AddCallOnActor : public FEdGraphSchemaAction_K2NewNode
 {
-	GENERATED_USTRUCT_BODY()
+	GENERATED_BODY()
 
 	// Simple type info
 	static FName StaticGetTypeId() {static FName Type("FEdGraphSchemaAction_K2AddCallOnActor"); return Type;}
@@ -315,7 +315,7 @@ struct FEdGraphSchemaAction_K2AddCallOnActor : public FEdGraphSchemaAction_K2New
 USTRUCT()
 struct FEdGraphSchemaAction_K2AddCallOnVariable : public FEdGraphSchemaAction_K2NewNode
 {
-	GENERATED_USTRUCT_BODY()
+	GENERATED_BODY()
 
 	/** Property name we want to call the function on */
 	UPROPERTY()
@@ -342,7 +342,7 @@ struct FEdGraphSchemaAction_K2AddCallOnVariable : public FEdGraphSchemaAction_K2
 USTRUCT()
 struct FEdGraphSchemaAction_K2AddComment : public FEdGraphSchemaAction
 {
-	GENERATED_USTRUCT_BODY()
+	GENERATED_BODY()
 
 	// Simple type info
 	static FName StaticGetTypeId() {static FName Type("FEdGraphSchemaAction_K2AddComment"); return Type;}
@@ -377,7 +377,7 @@ struct FEdGraphSchemaAction_K2AddComment : public FEdGraphSchemaAction
 USTRUCT()
 struct FEdGraphSchemaAction_K2AddDocumentation : public FEdGraphSchemaAction
 {
-	GENERATED_USTRUCT_BODY()
+	GENERATED_BODY()
 
 	// Simple type info
 	static FName StaticGetTypeId() {static FName Type("FEdGraphSchemaAction_K2AddDocumentation"); return Type;}
@@ -412,7 +412,7 @@ struct FEdGraphSchemaAction_K2AddDocumentation : public FEdGraphSchemaAction
 USTRUCT()
 struct BLUEPRINTGRAPH_API FEdGraphSchemaAction_K2TargetNode : public FEdGraphSchemaAction_K2NewNode
 {
-	GENERATED_USTRUCT_BODY()
+	GENERATED_BODY()
 
 
 	// Simple type info
@@ -440,7 +440,7 @@ struct BLUEPRINTGRAPH_API FEdGraphSchemaAction_K2TargetNode : public FEdGraphSch
 USTRUCT()
 struct BLUEPRINTGRAPH_API FEdGraphSchemaAction_K2PasteHere : public FEdGraphSchemaAction_K2NewNode
 {
-	GENERATED_USTRUCT_BODY()
+	GENERATED_BODY()
 
 	// Simple type info
 	static FName StaticGetTypeId() {static FName Type("FEdGraphSchemaAction_K2PasteHere"); return Type;}
@@ -467,7 +467,7 @@ struct BLUEPRINTGRAPH_API FEdGraphSchemaAction_K2PasteHere : public FEdGraphSche
 USTRUCT()
 struct BLUEPRINTGRAPH_API FEdGraphSchemaAction_K2Enum : public FEdGraphSchemaAction
 {
-	GENERATED_USTRUCT_BODY()
+	GENERATED_BODY()
 
 	// Simple type info
 	static FName StaticGetTypeId() {static FName Type("FEdGraphSchemaAction_K2Enum"); return Type;}
@@ -505,7 +505,7 @@ struct BLUEPRINTGRAPH_API FEdGraphSchemaAction_K2Enum : public FEdGraphSchemaAct
 USTRUCT()
 struct BLUEPRINTGRAPH_API FEdGraphSchemaAction_K2Var : public FEdGraphSchemaAction
 {
-	GENERATED_USTRUCT_BODY()
+	GENERATED_BODY()
 
 private:
 	/** Name of function or class */
@@ -572,7 +572,7 @@ public:
 USTRUCT()
 struct BLUEPRINTGRAPH_API FEdGraphSchemaAction_K2LocalVar : public FEdGraphSchemaAction
 {
-	GENERATED_USTRUCT_BODY()
+	GENERATED_BODY()
 
 private:
 	/** Name of function or class */
@@ -654,7 +654,7 @@ namespace EEdGraphSchemaAction_K2Graph
 USTRUCT()
 struct BLUEPRINTGRAPH_API FEdGraphSchemaAction_K2Graph : public FEdGraphSchemaAction
 {
-	GENERATED_USTRUCT_BODY()
+	GENERATED_BODY()
 
 	// Simple type info
 	static FName StaticGetTypeId() {static FName Type("FEdGraphSchemaAction_K2Graph"); return Type;}
@@ -693,7 +693,7 @@ struct BLUEPRINTGRAPH_API FEdGraphSchemaAction_K2Graph : public FEdGraphSchemaAc
 USTRUCT()
 struct BLUEPRINTGRAPH_API FEdGraphSchemaAction_K2Event : public FEdGraphSchemaAction_K2TargetNode
 {
-	GENERATED_USTRUCT_BODY()
+	GENERATED_BODY()
 
 	/**
 	 * An empty default constructor (stubbed out for arrays and other containers 
@@ -743,7 +743,7 @@ struct BLUEPRINTGRAPH_API FEdGraphSchemaAction_K2Event : public FEdGraphSchemaAc
 USTRUCT()
 struct BLUEPRINTGRAPH_API FEdGraphSchemaAction_K2InputAction : public FEdGraphSchemaAction_K2TargetNode
 {
-	GENERATED_USTRUCT_BODY()
+	GENERATED_BODY()
 
 	/**
 	 * An empty default constructor (stubbed out for arrays and other containers
@@ -790,7 +790,7 @@ struct BLUEPRINTGRAPH_API FEdGraphSchemaAction_K2InputAction : public FEdGraphSc
 USTRUCT()
 struct BLUEPRINTGRAPH_API FEdGraphSchemaAction_K2Delegate : public FEdGraphSchemaAction
 {
-	GENERATED_USTRUCT_BODY()
+	GENERATED_BODY()
 
 private:
 	/** Name of the property */
diff --git a/Engine/Source/Editor/BlueprintGraph/Classes/K2Node.h b/Engine/Source/Editor/BlueprintGraph/Classes/K2Node.h
--- a/Engine/Source/Editor/BlueprintGraph/Classes/K2Node.h
+++ b/Engine/Source/Editor/BlueprintGraph/Classes/K2Node.h
@@ -13,7 +13,7 @@ class UDynamicBlueprintBinding;
 USTRUCT()
 struct FOptionalPinFromProperty
 {
-	GENERATED_USTRUCT_BODY()
+	GENERATED_BODY()
 
 	UPROPERTY(EditAnywhere, Category=Hi, BlueprintReadOnly)
 	FName PropertyName;
@@ -87,7 +87,9 @@ enum ERenamePinResult
 UCLASS(abstract, MinimalAPI)
 class UK2Node : public UEdGraphNode
 {
-	GENERATED_UCLASS_BODY()
+	GENERATED_BODY()
+public:
+	BLUEPRINTGRAPH_API UK2Node(const FObjectInitializer& ObjectInitializer = FObjectInitializer::Get());
 
 	// UObject interface
 	BLUEPRINTGRAPH_API virtual void PostLoad() override;
diff --git a/Engine/Source/Editor/BlueprintGraph/Classes/K2Node_AIMoveTo.h b/Engine/Source/Editor/BlueprintGraph/Classes/K2Node_AIMoveTo.h
--- a/Engine/Source/Editor/BlueprintGraph/Classes/K2Node_AIMoveTo.h
+++ b/Engine/Source/Editor/BlueprintGraph/Classes/K2Node_AIMoveTo.h
@@ -9,7 +9,9 @@
 UCLASS()
 class BLUEPRINTGRAPH_API UK2Node_AIMoveTo : public UK2Node_BaseAsyncTask
 {
-	GENERATED_UCLASS_BODY()
+	GENERATED_BODY()
+public:
+	UK2Node_AIMoveTo(const FObjectInitializer& ObjectInitializer = FObjectInitializer::Get());
 
 	// Begin UEdGraphNode interface
 	virtual FText GetTooltipText() const override;
diff --git a/Engine/Source/Editor/BlueprintGraph/Classes/K2Node_ActorBoundEvent.h b/Engine/Source/Editor/BlueprintGraph/Classes/K2Node_ActorBoundEvent.h
--- a/Engine/Source/Editor/BlueprintGraph/Classes/K2Node_ActorBoundEvent.h
+++ b/Engine/Source/Editor/BlueprintGraph/Classes/K2Node_ActorBoundEvent.h
@@ -9,7 +9,9 @@
 UCLASS(MinimalAPI)
 class UK2Node_ActorBoundEvent : public UK2Node_Event
 {
-	GENERATED_UCLASS_BODY()
+	GENERATED_BODY()
+public:
+	BLUEPRINTGRAPH_API UK2Node_ActorBoundEvent(const FObjectInitializer& ObjectInitializer = FObjectInitializer::Get());
 
 	/** Delegate property name that this event is associated with */
 	UPROPERTY()
diff --git a/Engine/Source/Editor/BlueprintGraph/Classes/K2Node_AddComponent.h b/Engine/Source/Editor/BlueprintGraph/Classes/K2Node_AddComponent.h
--- a/Engine/Source/Editor/BlueprintGraph/Classes/K2Node_AddComponent.h
+++ b/Engine/Source/Editor/BlueprintGraph/Classes/K2Node_AddComponent.h
@@ -9,7 +9,9 @@
 UCLASS(MinimalAPI)
 class UK2Node_AddComponent : public UK2Node_CallFunction
 {
-	GENERATED_UCLASS_BODY()
+	GENERATED_BODY()
+public:
+	BLUEPRINTGRAPH_API UK2Node_AddComponent(const FObjectInitializer& ObjectInitializer = FObjectInitializer::Get());
 
 	UPROPERTY()
 	uint32 bHasExposedVariable:1;
diff --git a/Engine/Source/Editor/BlueprintGraph/Classes/K2Node_AddDelegate.h b/Engine/Source/Editor/BlueprintGraph/Classes/K2Node_AddDelegate.h
--- a/Engine/Source/Editor/BlueprintGraph/Classes/K2Node_AddDelegate.h
+++ b/Engine/Source/Editor/BlueprintGraph/Classes/K2Node_AddDelegate.h
@@ -6,7 +6,9 @@
 UCLASS(MinimalAPI)
 class UK2Node_AddDelegate : public UK2Node_BaseMCDelegate
 {
-	GENERATED_UCLASS_BODY()
+	GENERATED_BODY()
+public:
+	BLUEPRINTGRAPH_API UK2Node_AddDelegate(const FObjectInitializer& ObjectInitializer = FObjectInitializer::Get());
 
 public:
 	// UEdGraphNode interface
diff --git a/Engine/Source/Editor/BlueprintGraph/Classes/K2Node_AssignDelegate.h b/Engine/Source/Editor/BlueprintGraph/Classes/K2Node_AssignDelegate.h
--- a/Engine/Source/Editor/BlueprintGraph/Classes/K2Node_AssignDelegate.h
+++ b/Engine/Source/Editor/BlueprintGraph/Classes/K2Node_AssignDelegate.h
@@ -13,7 +13,9 @@
 UCLASS(MinimalAPI)
 class UK2Node_AssignDelegate : public UK2Node_AddDelegate
 {
-	GENERATED_UCLASS_BODY()
+	GENERATED_BODY()
+public:
+	BLUEPRINTGRAPH_API UK2Node_AssignDelegate(const FObjectInitializer& ObjectInitializer = FObjectInitializer::Get());
 
 public:
 	// UEdGraphNode interface
diff --git a/Engine/Source/Editor/BlueprintGraph/Classes/K2Node_AssignmentStatement.h b/Engine/Source/Editor/BlueprintGraph/Classes/K2Node_AssignmentStatement.h
--- a/Engine/Source/Editor/BlueprintGraph/Classes/K2Node_AssignmentStatement.h
+++ b/Engine/Source/Editor/BlueprintGraph/Classes/K2Node_AssignmentStatement.h
@@ -7,7 +7,9 @@
 UCLASS(MinimalAPI)
 class UK2Node_AssignmentStatement : public UK2Node
 {
-	GENERATED_UCLASS_BODY()
+	GENERATED_BODY()
+public:
+	BLUEPRINTGRAPH_API UK2Node_AssignmentStatement(const FObjectInitializer& ObjectInitializer = FObjectInitializer::Get());
 
 
 	// Name of the Variable pin for this node
diff --git a/Engine/Source/Editor/BlueprintGraph/Classes/K2Node_BaseAsyncTask.h b/Engine/Source/Editor/BlueprintGraph/Classes/K2Node_BaseAsyncTask.h
--- a/Engine/Source/Editor/BlueprintGraph/Classes/K2Node_BaseAsyncTask.h
+++ b/Engine/Source/Editor/BlueprintGraph/Classes/K2Node_BaseAsyncTask.h
@@ -13,7 +13,9 @@ class UK2Node_TemporaryVariable;
 UCLASS(Abstract)
 class BLUEPRINTGRAPH_API UK2Node_BaseAsyncTask : public UK2Node
 {
-	GENERATED_UCLASS_BODY()
+	GENERATED_BODY()
+public:
+	UK2Node_BaseAsyncTask(const FObjectInitializer& ObjectInitializer = FObjectInitializer::Get());
 
 	// UEdGraphNode interface
 	virtual void AllocateDefaultPins() override;
diff --git a/Engine/Source/Editor/BlueprintGraph/Classes/K2Node_BaseMCDelegate.h b/Engine/Source/Editor/BlueprintGraph/Classes/K2Node_BaseMCDelegate.h
--- a/Engine/Source/Editor/BlueprintGraph/Classes/K2Node_BaseMCDelegate.h
+++ b/Engine/Source/Editor/BlueprintGraph/Classes/K2Node_BaseMCDelegate.h
@@ -8,7 +8,9 @@
 UCLASS(MinimalAPI, abstract)
 class UK2Node_BaseMCDelegate : public UK2Node
 {
-	GENERATED_UCLASS_BODY()
+	GENERATED_BODY()
+public:
+	BLUEPRINTGRAPH_API UK2Node_BaseMCDelegate(const FObjectInitializer& ObjectInitializer = FObjectInitializer::Get());
 
 	/** Reference to delegate */
 	UPROPERTY()
diff --git a/Engine/Source/Editor/BlueprintGraph/Classes/K2Node_BreakStruct.h b/Engine/Source/Editor/BlueprintGraph/Classes/K2Node_BreakStruct.h
--- a/Engine/Source/Editor/BlueprintGraph/Classes/K2Node_BreakStruct.h
+++ b/Engine/Source/Editor/BlueprintGraph/Classes/K2Node_BreakStruct.h
@@ -8,7 +8,9 @@
 UCLASS(MinimalAPI)
 class UK2Node_BreakStruct : public UK2Node_StructMemberGet
 {
-	GENERATED_UCLASS_BODY()
+	GENERATED_BODY()
+public:
+	BLUEPRINTGRAPH_API UK2Node_BreakStruct(const FObjectInitializer& ObjectInitializer = FObjectInitializer::Get());
 
 	static bool CanBeBroken(const UScriptStruct* Struct);
 	static bool CanCreatePinForProperty(const UProperty* Property);
diff --git a/Engine/Source/Editor/BlueprintGraph/Classes/K2Node_CallArrayFunction.h b/Engine/Source/Editor/BlueprintGraph/Classes/K2Node_CallArrayFunction.h
--- a/Engine/Source/Editor/BlueprintGraph/Classes/K2Node_CallArrayFunction.h
+++ b/Engine/Source/Editor/BlueprintGraph/Classes/K2Node_CallArrayFunction.h
@@ -19,7 +19,9 @@ struct FArrayPropertyPinCombo
 UCLASS(MinimalAPI)
 class UK2Node_CallArrayFunction : public UK2Node_CallFunction
 {
-	GENERATED_UCLASS_BODY()
+	GENERATED_BODY()
+public:
+	BLUEPRINTGRAPH_API UK2Node_CallArrayFunction(const FObjectInitializer& ObjectInitializer = FObjectInitializer::Get());
 
 	// Begin UEdGraphNode interface
 	virtual void AllocateDefaultPins() override;
diff --git a/Engine/Source/Editor/BlueprintGraph/Classes/K2Node_CallDataTableFunction.h b/Engine/Source/Editor/BlueprintGraph/Classes/K2Node_CallDataTableFunction.h
--- a/Engine/Source/Editor/BlueprintGraph/Classes/K2Node_CallDataTableFunction.h
+++ b/Engine/Source/Editor/BlueprintGraph/Classes/K2Node_CallDataTableFunction.h
@@ -7,10 +7,12 @@
 UCLASS(MinimalAPI)
 class UK2Node_CallDataTableFunction : public UK2Node_CallFunction
 {
-	GENERATED_UCLASS_BODY()
+	GENERATED_BODY()
+public:
+	BLUEPRINTGRAPH_API UK2Node_CallDataTableFunction(const FObjectInitializer& ObjectInitializer = FObjectInitializer::Get());
 
 	// Begin EdGraphNode interface
 	virtual void PinDefaultValueChanged(UEdGraphPin* Pin) override;
 	virtual void NotifyPinConnectionListChanged(UEdGraphPin* Pin) override;
 	// End EdGraphNode interface
-};
\ No newline at end of file
+};
diff --git a/Engine/Source/Editor/BlueprintGraph/Classes/K2Node_CallDelegate.h b/Engine/Source/Editor/BlueprintGraph/Classes/K2Node_CallDelegate.h
--- a/Engine/Source/Editor/BlueprintGraph/Classes/K2Node_CallDelegate.h
+++ b/Engine/Source/Editor/BlueprintGraph/Classes/K2Node_CallDelegate.h
@@ -6,7 +6,9 @@
 UCLASS(MinimalAPI)
 class UK2Node_CallDelegate : public UK2Node_BaseMCDelegate
 {
-	GENERATED_UCLASS_BODY()
+	GENERATED_BODY()
+public:
+	BLUEPRINTGRAPH_API UK2Node_CallDelegate(const FObjectInitializer& ObjectInitializer = FObjectInitializer::Get());
 
 public:
 
@@ -22,4 +24,4 @@ public:
 
 	virtual void ValidateNodeDuringCompilation(class FCompilerResultsLog& MessageLog) const override;
 	bool CreatePinsForFunctionInputs(const UFunction* Function);
-};
\ No newline at end of file
+};
diff --git a/Engine/Source/Editor/BlueprintGraph/Classes/K2Node_CallFunction.h b/Engine/Source/Editor/BlueprintGraph/Classes/K2Node_CallFunction.h
--- a/Engine/Source/Editor/BlueprintGraph/Classes/K2Node_CallFunction.h
+++ b/Engine/Source/Editor/BlueprintGraph/Classes/K2Node_CallFunction.h
@@ -10,7 +10,9 @@
 UCLASS()
 class BLUEPRINTGRAPH_API UK2Node_CallFunction : public UK2Node
 {
-	GENERATED_UCLASS_BODY()
+	GENERATED_BODY()
+public:
+	UK2Node_CallFunction(const FObjectInitializer& ObjectInitializer = FObjectInitializer::Get());
 
 	/** Indicates that this is a call to a pure function */
 	UPROPERTY()
diff --git a/Engine/Source/Editor/BlueprintGraph/Classes/K2Node_CallFunctionOnMember.h b/Engine/Source/Editor/BlueprintGraph/Classes/K2Node_CallFunctionOnMember.h
--- a/Engine/Source/Editor/BlueprintGraph/Classes/K2Node_CallFunctionOnMember.h
+++ b/Engine/Source/Editor/BlueprintGraph/Classes/K2Node_CallFunctionOnMember.h
@@ -9,7 +9,9 @@
 UCLASS(MinimalAPI)
 class UK2Node_CallFunctionOnMember : public UK2Node_CallFunction
 {
-	GENERATED_UCLASS_BODY()
+	GENERATED_BODY()
+public:
+	BLUEPRINTGRAPH_API UK2Node_CallFunctionOnMember(const FObjectInitializer& ObjectInitializer = FObjectInitializer::Get());
 
 	/** Reference to member variable to call function on */
 	UPROPERTY()
diff --git a/Engine/Source/Editor/BlueprintGraph/Classes/K2Node_CallMaterialParameterCollectionFunction.h b/Engine/Source/Editor/BlueprintGraph/Classes/K2Node_CallMaterialParameterCollectionFunction.h
--- a/Engine/Source/Editor/BlueprintGraph/Classes/K2Node_CallMaterialParameterCollectionFunction.h
+++ b/Engine/Source/Editor/BlueprintGraph/Classes/K2Node_CallMaterialParameterCollectionFunction.h
@@ -7,9 +7,11 @@
 UCLASS(MinimalAPI)
 class UK2Node_CallMaterialParameterCollectionFunction : public UK2Node_CallFunction
 {
-	GENERATED_UCLASS_BODY()
+	GENERATED_BODY()
+public:
+	BLUEPRINTGRAPH_API UK2Node_CallMaterialParameterCollectionFunction(const FObjectInitializer& ObjectInitializer = FObjectInitializer::Get());
 
 	// Begin EdGraphNode interface
 	virtual void PinDefaultValueChanged(UEdGraphPin* Pin) override;
 	// End EdGraphNode interface
-};
\ No newline at end of file
+};
diff --git a/Engine/Source/Editor/BlueprintGraph/Classes/K2Node_CallParentFunction.h b/Engine/Source/Editor/BlueprintGraph/Classes/K2Node_CallParentFunction.h
--- a/Engine/Source/Editor/BlueprintGraph/Classes/K2Node_CallParentFunction.h
+++ b/Engine/Source/Editor/BlueprintGraph/Classes/K2Node_CallParentFunction.h
@@ -8,7 +8,9 @@
 UCLASS(MinimalAPI)
 class UK2Node_CallParentFunction : public UK2Node_CallFunction
 {
-	GENERATED_UCLASS_BODY()
+	GENERATED_BODY()
+public:
+	BLUEPRINTGRAPH_API UK2Node_CallParentFunction(const FObjectInitializer& ObjectInitializer = FObjectInitializer::Get());
 
 	// Begin EdGraphNode interface
 	virtual FText GetNodeTitle(ENodeTitleType::Type TitleType) const override;
diff --git a/Engine/Source/Editor/BlueprintGraph/Classes/K2Node_CastByteToEnum.h b/Engine/Source/Editor/BlueprintGraph/Classes/K2Node_CastByteToEnum.h
--- a/Engine/Source/Editor/BlueprintGraph/Classes/K2Node_CastByteToEnum.h
+++ b/Engine/Source/Editor/BlueprintGraph/Classes/K2Node_CastByteToEnum.h
@@ -8,7 +8,9 @@
 UCLASS(MinimalAPI)
 class UK2Node_CastByteToEnum : public UK2Node
 {
-	GENERATED_UCLASS_BODY()
+	GENERATED_BODY()
+public:
+	BLUEPRINTGRAPH_API UK2Node_CastByteToEnum(const FObjectInitializer& ObjectInitializer = FObjectInitializer::Get());
 
 	UPROPERTY()
 	UEnum* Enum;
diff --git a/Engine/Source/Editor/BlueprintGraph/Classes/K2Node_ClassDynamicCast.h b/Engine/Source/Editor/BlueprintGraph/Classes/K2Node_ClassDynamicCast.h
--- a/Engine/Source/Editor/BlueprintGraph/Classes/K2Node_ClassDynamicCast.h
+++ b/Engine/Source/Editor/BlueprintGraph/Classes/K2Node_ClassDynamicCast.h
@@ -8,7 +8,9 @@
 UCLASS(MinimalAPI)
 class UK2Node_ClassDynamicCast : public UK2Node_DynamicCast
 {
-	GENERATED_UCLASS_BODY()
+	GENERATED_BODY()
+public:
+	BLUEPRINTGRAPH_API UK2Node_ClassDynamicCast(const FObjectInitializer& ObjectInitializer = FObjectInitializer::Get());
 
 	// Begin UEdGraphNode interface
 	virtual void AllocateDefaultPins() override;
diff --git a/Engine/Source/Editor/BlueprintGraph/Classes/K2Node_ClearDelegate.h b/Engine/Source/Editor/BlueprintGraph/Classes/K2Node_ClearDelegate.h
--- a/Engine/Source/Editor/BlueprintGraph/Classes/K2Node_ClearDelegate.h
+++ b/Engine/Source/Editor/BlueprintGraph/Classes/K2Node_ClearDelegate.h
@@ -7,7 +7,9 @@
 UCLASS(MinimalAPI)
 class UK2Node_ClearDelegate : public UK2Node_BaseMCDelegate
 {
-	GENERATED_UCLASS_BODY()
+	GENERATED_BODY()
+public:
+	BLUEPRINTGRAPH_API UK2Node_ClearDelegate(const FObjectInitializer& ObjectInitializer = FObjectInitializer::Get());
 
 public:
 	// UEdGraphNode interface
diff --git a/Engine/Source/Editor/BlueprintGraph/Classes/K2Node_CommutativeAssociativeBinaryOperator.h b/Engine/Source/Editor/BlueprintGraph/Classes/K2Node_CommutativeAssociativeBinaryOperator.h
--- a/Engine/Source/Editor/BlueprintGraph/Classes/K2Node_CommutativeAssociativeBinaryOperator.h
+++ b/Engine/Source/Editor/BlueprintGraph/Classes/K2Node_CommutativeAssociativeBinaryOperator.h
@@ -7,7 +7,9 @@
 UCLASS(MinimalAPI)
 class UK2Node_CommutativeAssociativeBinaryOperator : public UK2Node_CallFunction
 {
-	GENERATED_UCLASS_BODY()
+	GENERATED_BODY()
+public:
+	BLUEPRINTGRAPH_API UK2Node_CommutativeAssociativeBinaryOperator(const FObjectInitializer& ObjectInitializer = FObjectInitializer::Get());
 
 	/** The number of additional input pins to generate for this node (2 base pins are not included) */
 	UPROPERTY()
@@ -42,4 +44,4 @@ public:
 	virtual void GetContextMenuActions(const FGraphNodeContextMenuBuilder& Context) const override;
 	virtual void ExpandNode(class FKismetCompilerContext& CompilerContext, UEdGraph* SourceGraph) override;
 	// End of UK2Node interface
-};
\ No newline at end of file
+};
diff --git a/Engine/Source/Editor/BlueprintGraph/Classes/K2Node_ComponentBoundEvent.h b/Engine/Source/Editor/BlueprintGraph/Classes/K2Node_ComponentBoundEvent.h
--- a/Engine/Source/Editor/BlueprintGraph/Classes/K2Node_ComponentBoundEvent.h
+++ b/Engine/Source/Editor/BlueprintGraph/Classes/K2Node_ComponentBoundEvent.h
@@ -11,7 +11,9 @@ class UDynamicBlueprintBinding;
 UCLASS(MinimalAPI)
 class UK2Node_ComponentBoundEvent : public UK2Node_Event
 {
-	GENERATED_UCLASS_BODY()
+	GENERATED_BODY()
+public:
+	BLUEPRINTGRAPH_API UK2Node_ComponentBoundEvent(const FObjectInitializer& ObjectInitializer = FObjectInitializer::Get());
 
 	/** Delegate property name that this event is associated with */
 	UPROPERTY()
diff --git a/Engine/Source/Editor/BlueprintGraph/Classes/K2Node_Composite.h b/Engine/Source/Editor/BlueprintGraph/Classes/K2Node_Composite.h
--- a/Engine/Source/Editor/BlueprintGraph/Classes/K2Node_Composite.h
+++ b/Engine/Source/Editor/BlueprintGraph/Classes/K2Node_Composite.h
@@ -7,7 +7,9 @@
 UCLASS(MinimalAPI)
 class UK2Node_Composite : public UK2Node_Tunnel
 {
-	GENERATED_UCLASS_BODY()
+	GENERATED_BODY()
+public:
+	BLUEPRINTGRAPH_API UK2Node_Composite(const FObjectInitializer& ObjectInitializer = FObjectInitializer::Get());
 
 	// The graph that this composite node is representing
 	UPROPERTY()
diff --git a/Engine/Source/Editor/BlueprintGraph/Classes/K2Node_ConstructObjectFromClass.h b/Engine/Source/Editor/BlueprintGraph/Classes/K2Node_ConstructObjectFromClass.h
--- a/Engine/Source/Editor/BlueprintGraph/Classes/K2Node_ConstructObjectFromClass.h
+++ b/Engine/Source/Editor/BlueprintGraph/Classes/K2Node_ConstructObjectFromClass.h
@@ -9,7 +9,9 @@
 UCLASS(abstract)
 class BLUEPRINTGRAPH_API UK2Node_ConstructObjectFromClass : public UK2Node
 {
-	GENERATED_UCLASS_BODY()
+	GENERATED_BODY()
+public:
+	UK2Node_ConstructObjectFromClass(const FObjectInitializer& ObjectInitializer = FObjectInitializer::Get());
 
 	// Begin UEdGraphNode interface.
 	virtual void AllocateDefaultPins() override;
diff --git a/Engine/Source/Editor/BlueprintGraph/Classes/K2Node_CreateDelegate.h b/Engine/Source/Editor/BlueprintGraph/Classes/K2Node_CreateDelegate.h
--- a/Engine/Source/Editor/BlueprintGraph/Classes/K2Node_CreateDelegate.h
+++ b/Engine/Source/Editor/BlueprintGraph/Classes/K2Node_CreateDelegate.h
@@ -6,7 +6,9 @@
 UCLASS(MinimalAPI)
 class UK2Node_CreateDelegate : public UK2Node
 {
-	GENERATED_UCLASS_BODY()
+	GENERATED_BODY()
+public:
+	BLUEPRINTGRAPH_API UK2Node_CreateDelegate(const FObjectInitializer& ObjectInitializer = FObjectInitializer::Get());
 
 	UPROPERTY()
 	FName SelectedFunctionName;
diff --git a/Engine/Source/Editor/BlueprintGraph/Classes/K2Node_CustomEvent.h b/Engine/Source/Editor/BlueprintGraph/Classes/K2Node_CustomEvent.h
--- a/Engine/Source/Editor/BlueprintGraph/Classes/K2Node_CustomEvent.h
+++ b/Engine/Source/Editor/BlueprintGraph/Classes/K2Node_CustomEvent.h
@@ -9,7 +9,9 @@
 UCLASS(MinimalAPI)
 class UK2Node_CustomEvent : public UK2Node_Event
 {
-	GENERATED_UCLASS_BODY()
+	GENERATED_BODY()
+public:
+	BLUEPRINTGRAPH_API UK2Node_CustomEvent(const FObjectInitializer& ObjectInitializer = FObjectInitializer::Get());
 
 	/** Specifies that the event can be triggered in Editor */
 	UPROPERTY()
diff --git a/Engine/Source/Editor/BlueprintGraph/Classes/K2Node_DeadClass.h b/Engine/Source/Editor/BlueprintGraph/Classes/K2Node_DeadClass.h
--- a/Engine/Source/Editor/BlueprintGraph/Classes/K2Node_DeadClass.h
+++ b/Engine/Source/Editor/BlueprintGraph/Classes/K2Node_DeadClass.h
@@ -7,5 +7,7 @@
 UCLASS(MinimalAPI)
 class UK2Node_DeadClass : public UK2Node
 {
-	GENERATED_UCLASS_BODY()
+	GENERATED_BODY()
+public:
+	BLUEPRINTGRAPH_API UK2Node_DeadClass(const FObjectInitializer& ObjectInitializer = FObjectInitializer::Get());
 };
