Back out changelist 2481333

diff --git a/Engine/Source/Editor/AnimGraph/Classes/AnimGraphNode_BlendSpaceBase.h b/Engine/Source/Editor/AnimGraph/Classes/AnimGraphNode_BlendSpaceBase.h
--- a/Engine/Source/Editor/AnimGraph/Classes/AnimGraphNode_BlendSpaceBase.h
+++ b/Engine/Source/Editor/AnimGraph/Classes/AnimGraphNode_BlendSpaceBase.h
@@ -8,9 +8,7 @@
 UCLASS(Abstract, MinimalAPI)
 class UAnimGraphNode_BlendSpaceBase : public UAnimGraphNode_Base
 {
-	GENERATED_BODY()
-public:
-	ANIMGRAPH_API UAnimGraphNode_BlendSpaceBase(const FObjectInitializer& ObjectInitializer = FObjectInitializer::Get());
+	GENERATED_UCLASS_BODY()
 
 	// UEdGraphNode interface
 	ANIMGRAPH_API virtual FLinearColor GetNodeTitleColor() const override;
diff --git a/Engine/Source/Editor/AnimGraph/Classes/AnimGraphNode_BlendSpaceEvaluator.h b/Engine/Source/Editor/AnimGraph/Classes/AnimGraphNode_BlendSpaceEvaluator.h
--- a/Engine/Source/Editor/AnimGraph/Classes/AnimGraphNode_BlendSpaceEvaluator.h
+++ b/Engine/Source/Editor/AnimGraph/Classes/AnimGraphNode_BlendSpaceEvaluator.h
@@ -9,9 +9,7 @@
 UCLASS(MinimalAPI)
 class UAnimGraphNode_BlendSpaceEvaluator : public UAnimGraphNode_BlendSpaceBase
 {
-	GENERATED_BODY()
-public:
-	ANIMGRAPH_API UAnimGraphNode_BlendSpaceEvaluator(const FObjectInitializer& ObjectInitializer = FObjectInitializer::Get());
+	GENERATED_UCLASS_BODY()
 
 	UPROPERTY(EditAnywhere, Category=Settings)
 	FAnimNode_BlendSpaceEvaluator Node;
diff --git a/Engine/Source/Editor/AnimGraph/Classes/AnimGraphNode_BlendSpacePlayer.h b/Engine/Source/Editor/AnimGraph/Classes/AnimGraphNode_BlendSpacePlayer.h
--- a/Engine/Source/Editor/AnimGraph/Classes/AnimGraphNode_BlendSpacePlayer.h
+++ b/Engine/Source/Editor/AnimGraph/Classes/AnimGraphNode_BlendSpacePlayer.h
@@ -9,9 +9,7 @@
 UCLASS(MinimalAPI)
 class UAnimGraphNode_BlendSpacePlayer : public UAnimGraphNode_BlendSpaceBase
 {
-	GENERATED_BODY()
-public:
-	ANIMGRAPH_API UAnimGraphNode_BlendSpacePlayer(const FObjectInitializer& ObjectInitializer = FObjectInitializer::Get());
+	GENERATED_UCLASS_BODY()
 
 	UPROPERTY(EditAnywhere, Category=Settings)
 	FAnimNode_BlendSpacePlayer Node;
diff --git a/Engine/Source/Editor/AnimGraph/Classes/AnimGraphNode_BoneDrivenController.h b/Engine/Source/Editor/AnimGraph/Classes/AnimGraphNode_BoneDrivenController.h
--- a/Engine/Source/Editor/AnimGraph/Classes/AnimGraphNode_BoneDrivenController.h
+++ b/Engine/Source/Editor/AnimGraph/Classes/AnimGraphNode_BoneDrivenController.h
@@ -10,9 +10,7 @@
 UCLASS(MinimalAPI)
 class UAnimGraphNode_BoneDrivenController : public UAnimGraphNode_SkeletalControlBase
 {
-	GENERATED_BODY()
-public:
-	ANIMGRAPH_API UAnimGraphNode_BoneDrivenController(const FObjectInitializer& ObjectInitializer = FObjectInitializer::Get());
+	GENERATED_UCLASS_BODY()
 
 	UPROPERTY(EditAnywhere, Category = Settings)
 	FAnimNode_BoneDrivenController Node;
@@ -38,4 +36,4 @@ protected:
 private:
 	/** Constructing FText strings can be costly, so we cache the node's title */
 	FNodeTitleTextTable CachedNodeTitles;
-};
+};
\ No newline at end of file
diff --git a/Engine/Source/Editor/AnimGraph/Classes/AnimGraphNode_ComponentToLocalSpace.h b/Engine/Source/Editor/AnimGraph/Classes/AnimGraphNode_ComponentToLocalSpace.h
--- a/Engine/Source/Editor/AnimGraph/Classes/AnimGraphNode_ComponentToLocalSpace.h
+++ b/Engine/Source/Editor/AnimGraph/Classes/AnimGraphNode_ComponentToLocalSpace.h
@@ -8,9 +8,7 @@
 UCLASS(MinimalAPI)
 class UAnimGraphNode_ComponentToLocalSpace : public UAnimGraphNode_Base
 {
-	GENERATED_BODY()
-public:
-	ANIMGRAPH_API UAnimGraphNode_ComponentToLocalSpace(const FObjectInitializer& ObjectInitializer = FObjectInitializer::Get());
+	GENERATED_UCLASS_BODY()
 
 	UPROPERTY(EditAnywhere, Category=Settings)
 	FAnimNode_ConvertComponentToLocalSpace Node;
diff --git a/Engine/Source/Editor/AnimGraph/Classes/AnimGraphNode_CopyBone.h b/Engine/Source/Editor/AnimGraph/Classes/AnimGraphNode_CopyBone.h
--- a/Engine/Source/Editor/AnimGraph/Classes/AnimGraphNode_CopyBone.h
+++ b/Engine/Source/Editor/AnimGraph/Classes/AnimGraphNode_CopyBone.h
@@ -9,9 +9,7 @@
 UCLASS(MinimalAPI)
 class UAnimGraphNode_CopyBone : public UAnimGraphNode_SkeletalControlBase
 {
-	GENERATED_BODY()
-public:
-	ANIMGRAPH_API UAnimGraphNode_CopyBone(const FObjectInitializer& ObjectInitializer = FObjectInitializer::Get());
+	GENERATED_UCLASS_BODY()
 
 	UPROPERTY(EditAnywhere, Category=Settings)
 	FAnimNode_CopyBone Node;
diff --git a/Engine/Source/Editor/AnimGraph/Classes/AnimGraphNode_Fabrik.h b/Engine/Source/Editor/AnimGraph/Classes/AnimGraphNode_Fabrik.h
--- a/Engine/Source/Editor/AnimGraph/Classes/AnimGraphNode_Fabrik.h
+++ b/Engine/Source/Editor/AnimGraph/Classes/AnimGraphNode_Fabrik.h
@@ -9,9 +9,7 @@
 UCLASS(MinimalAPI)
 class UAnimGraphNode_Fabrik : public UAnimGraphNode_SkeletalControlBase
 {
-	GENERATED_BODY()
-public:
-	ANIMGRAPH_API UAnimGraphNode_Fabrik(const FObjectInitializer& ObjectInitializer = FObjectInitializer::Get());
+	GENERATED_UCLASS_BODY()
 
 	UPROPERTY(EditAnywhere, Category = Settings)
 	FAnimNode_Fabrik Node;
@@ -36,4 +34,4 @@ protected:
 	// UAnimGraphNode_SkeletalControlBase interface
 	virtual FText GetControllerDescription() const override;
 	// End of UAnimGraphNode_SkeletalControlBase interface
-};
+};
\ No newline at end of file
diff --git a/Engine/Source/Editor/AnimGraph/Classes/AnimGraphNode_HandIKRetargeting.h b/Engine/Source/Editor/AnimGraph/Classes/AnimGraphNode_HandIKRetargeting.h
--- a/Engine/Source/Editor/AnimGraph/Classes/AnimGraphNode_HandIKRetargeting.h
+++ b/Engine/Source/Editor/AnimGraph/Classes/AnimGraphNode_HandIKRetargeting.h
@@ -8,9 +8,7 @@
 UCLASS(MinimalAPI)
 class UAnimGraphNode_HandIKRetargeting : public UAnimGraphNode_SkeletalControlBase
 {
-	GENERATED_BODY()
-public:
-	ANIMGRAPH_API UAnimGraphNode_HandIKRetargeting(const FObjectInitializer& ObjectInitializer = FObjectInitializer::Get());
+	GENERATED_UCLASS_BODY()
 
 	UPROPERTY(EditAnywhere, Category = Settings)
 	FAnimNode_HandIKRetargeting Node;
diff --git a/Engine/Source/Editor/AnimGraph/Classes/AnimGraphNode_IdentityPose.h b/Engine/Source/Editor/AnimGraph/Classes/AnimGraphNode_IdentityPose.h
--- a/Engine/Source/Editor/AnimGraph/Classes/AnimGraphNode_IdentityPose.h
+++ b/Engine/Source/Editor/AnimGraph/Classes/AnimGraphNode_IdentityPose.h
@@ -7,9 +7,7 @@
 UCLASS(MinimalAPI)
 class UAnimGraphNode_IdentityPose : public UAnimGraphNode_RefPoseBase
 {
-	GENERATED_BODY()
-public:
-	ANIMGRAPH_API UAnimGraphNode_IdentityPose(const FObjectInitializer& ObjectInitializer = FObjectInitializer::Get());
+	GENERATED_UCLASS_BODY()
 
 	// UEdGraphNode interface
 	virtual FText GetTooltipText() const override;
diff --git a/Engine/Source/Editor/AnimGraph/Classes/AnimGraphNode_LayeredBoneBlend.h b/Engine/Source/Editor/AnimGraph/Classes/AnimGraphNode_LayeredBoneBlend.h
--- a/Engine/Source/Editor/AnimGraph/Classes/AnimGraphNode_LayeredBoneBlend.h
+++ b/Engine/Source/Editor/AnimGraph/Classes/AnimGraphNode_LayeredBoneBlend.h
@@ -8,9 +8,7 @@
 UCLASS(MinimalAPI)
 class UAnimGraphNode_LayeredBoneBlend : public UAnimGraphNode_BlendListBase
 {
-	GENERATED_BODY()
-public:
-	ANIMGRAPH_API UAnimGraphNode_LayeredBoneBlend(const FObjectInitializer& ObjectInitializer = FObjectInitializer::Get());
+	GENERATED_UCLASS_BODY()
 
 	UPROPERTY(EditAnywhere, Category=Settings)
 	FAnimNode_LayeredBoneBlend Node;
diff --git a/Engine/Source/Editor/AnimGraph/Classes/AnimGraphNode_LocalRefPose.h b/Engine/Source/Editor/AnimGraph/Classes/AnimGraphNode_LocalRefPose.h
--- a/Engine/Source/Editor/AnimGraph/Classes/AnimGraphNode_LocalRefPose.h
+++ b/Engine/Source/Editor/AnimGraph/Classes/AnimGraphNode_LocalRefPose.h
@@ -7,9 +7,7 @@
 UCLASS(MinimalAPI)
 class UAnimGraphNode_LocalRefPose : public UAnimGraphNode_RefPoseBase
 {
-	GENERATED_BODY()
-public:
-	ANIMGRAPH_API UAnimGraphNode_LocalRefPose(const FObjectInitializer& ObjectInitializer = FObjectInitializer::Get());
+	GENERATED_UCLASS_BODY()
 
 	// UEdGraphNode interface
 	virtual FText GetTooltipText() const override;
diff --git a/Engine/Source/Editor/AnimGraph/Classes/AnimGraphNode_LocalToComponentSpace.h b/Engine/Source/Editor/AnimGraph/Classes/AnimGraphNode_LocalToComponentSpace.h
--- a/Engine/Source/Editor/AnimGraph/Classes/AnimGraphNode_LocalToComponentSpace.h
+++ b/Engine/Source/Editor/AnimGraph/Classes/AnimGraphNode_LocalToComponentSpace.h
@@ -8,9 +8,7 @@
 UCLASS(MinimalAPI)
 class UAnimGraphNode_LocalToComponentSpace : public UAnimGraphNode_Base
 {
-	GENERATED_BODY()
-public:
-	ANIMGRAPH_API UAnimGraphNode_LocalToComponentSpace(const FObjectInitializer& ObjectInitializer = FObjectInitializer::Get());
+	GENERATED_UCLASS_BODY()
 
 	UPROPERTY(EditAnywhere, Category=Settings)
 	FAnimNode_ConvertLocalToComponentSpace Node;
diff --git a/Engine/Source/Editor/AnimGraph/Classes/AnimGraphNode_LookAt.h b/Engine/Source/Editor/AnimGraph/Classes/AnimGraphNode_LookAt.h
--- a/Engine/Source/Editor/AnimGraph/Classes/AnimGraphNode_LookAt.h
+++ b/Engine/Source/Editor/AnimGraph/Classes/AnimGraphNode_LookAt.h
@@ -10,9 +10,7 @@
 UCLASS(MinimalAPI)
 class UAnimGraphNode_LookAt : public UAnimGraphNode_SkeletalControlBase
 {
-	GENERATED_BODY()
-public:
-	ANIMGRAPH_API UAnimGraphNode_LookAt(const FObjectInitializer& ObjectInitializer = FObjectInitializer::Get());
+	GENERATED_UCLASS_BODY()
 
 	UPROPERTY(EditAnywhere, Category=Settings)
 	FAnimNode_LookAt Node;
diff --git a/Engine/Source/Editor/AnimGraph/Classes/AnimGraphNode_MeshRefPose.h b/Engine/Source/Editor/AnimGraph/Classes/AnimGraphNode_MeshRefPose.h
--- a/Engine/Source/Editor/AnimGraph/Classes/AnimGraphNode_MeshRefPose.h
+++ b/Engine/Source/Editor/AnimGraph/Classes/AnimGraphNode_MeshRefPose.h
@@ -8,9 +8,7 @@
 UCLASS(MinimalAPI)
 class UAnimGraphNode_MeshRefPose : public UAnimGraphNode_Base
 {
-	GENERATED_BODY()
-public:
-	ANIMGRAPH_API UAnimGraphNode_MeshRefPose(const FObjectInitializer& ObjectInitializer = FObjectInitializer::Get());
+	GENERATED_UCLASS_BODY()
 
 	UPROPERTY(EditAnywhere, Category=Settings)
 	FAnimNode_MeshSpaceRefPose Node;
diff --git a/Engine/Source/Editor/AnimGraph/Classes/AnimGraphNode_ModifyBone.h b/Engine/Source/Editor/AnimGraph/Classes/AnimGraphNode_ModifyBone.h
--- a/Engine/Source/Editor/AnimGraph/Classes/AnimGraphNode_ModifyBone.h
+++ b/Engine/Source/Editor/AnimGraph/Classes/AnimGraphNode_ModifyBone.h
@@ -10,9 +10,7 @@
 UCLASS(MinimalAPI)
 class UAnimGraphNode_ModifyBone : public UAnimGraphNode_SkeletalControlBase
 {
-	GENERATED_BODY()
-public:
-	ANIMGRAPH_API UAnimGraphNode_ModifyBone(const FObjectInitializer& ObjectInitializer = FObjectInitializer::Get());
+	GENERATED_UCLASS_BODY()
 
 	UPROPERTY(EditAnywhere, Category=Settings)
 	FAnimNode_ModifyBone Node;
diff --git a/Engine/Source/Editor/AnimGraph/Classes/AnimGraphNode_RefPoseBase.h b/Engine/Source/Editor/AnimGraph/Classes/AnimGraphNode_RefPoseBase.h
--- a/Engine/Source/Editor/AnimGraph/Classes/AnimGraphNode_RefPoseBase.h
+++ b/Engine/Source/Editor/AnimGraph/Classes/AnimGraphNode_RefPoseBase.h
@@ -8,9 +8,7 @@
 UCLASS(MinimalAPI, abstract)
 class UAnimGraphNode_RefPoseBase : public UAnimGraphNode_Base
 {
-	GENERATED_BODY()
-public:
-	ANIMGRAPH_API UAnimGraphNode_RefPoseBase(const FObjectInitializer& ObjectInitializer = FObjectInitializer::Get());
+	GENERATED_UCLASS_BODY()
 
 	UPROPERTY(EditAnywhere, Category=Settings)
 	FAnimNode_RefPose Node;
diff --git a/Engine/Source/Editor/AnimGraph/Classes/AnimGraphNode_Root.h b/Engine/Source/Editor/AnimGraph/Classes/AnimGraphNode_Root.h
--- a/Engine/Source/Editor/AnimGraph/Classes/AnimGraphNode_Root.h
+++ b/Engine/Source/Editor/AnimGraph/Classes/AnimGraphNode_Root.h
@@ -8,9 +8,7 @@
 UCLASS(MinimalAPI)
 class UAnimGraphNode_Root : public UAnimGraphNode_Base
 {
-	GENERATED_BODY()
-public:
-	ANIMGRAPH_API UAnimGraphNode_Root(const FObjectInitializer& ObjectInitializer = FObjectInitializer::Get());
+	GENERATED_UCLASS_BODY()
 
 	UPROPERTY(EditAnywhere, Category=Settings)
 	FAnimNode_Root Node;
diff --git a/Engine/Source/Editor/AnimGraph/Classes/AnimGraphNode_RotateRootBone.h b/Engine/Source/Editor/AnimGraph/Classes/AnimGraphNode_RotateRootBone.h
--- a/Engine/Source/Editor/AnimGraph/Classes/AnimGraphNode_RotateRootBone.h
+++ b/Engine/Source/Editor/AnimGraph/Classes/AnimGraphNode_RotateRootBone.h
@@ -9,9 +9,7 @@
 UCLASS(MinimalAPI)
 class UAnimGraphNode_RotateRootBone : public UAnimGraphNode_Base
 {
-	GENERATED_BODY()
-public:
-	ANIMGRAPH_API UAnimGraphNode_RotateRootBone(const FObjectInitializer& ObjectInitializer = FObjectInitializer::Get());
+	GENERATED_UCLASS_BODY()
 
 	UPROPERTY(EditAnywhere, Category=Settings)
 	FAnimNode_RotateRootBone Node;
diff --git a/Engine/Source/Editor/AnimGraph/Classes/AnimGraphNode_RotationMultiplier.h b/Engine/Source/Editor/AnimGraph/Classes/AnimGraphNode_RotationMultiplier.h
--- a/Engine/Source/Editor/AnimGraph/Classes/AnimGraphNode_RotationMultiplier.h
+++ b/Engine/Source/Editor/AnimGraph/Classes/AnimGraphNode_RotationMultiplier.h
@@ -9,9 +9,7 @@
 UCLASS(MinimalAPI)
 class UAnimGraphNode_RotationMultiplier : public UAnimGraphNode_SkeletalControlBase
 {
-	GENERATED_BODY()
-public:
-	ANIMGRAPH_API UAnimGraphNode_RotationMultiplier(const FObjectInitializer& ObjectInitializer = FObjectInitializer::Get());
+	GENERATED_UCLASS_BODY()
 
 	UPROPERTY(EditAnywhere, Category=Settings)
 	FAnimNode_RotationMultiplier Node;
diff --git a/Engine/Source/Editor/AnimGraph/Classes/AnimGraphNode_RotationOffsetBlendSpace.h b/Engine/Source/Editor/AnimGraph/Classes/AnimGraphNode_RotationOffsetBlendSpace.h
--- a/Engine/Source/Editor/AnimGraph/Classes/AnimGraphNode_RotationOffsetBlendSpace.h
+++ b/Engine/Source/Editor/AnimGraph/Classes/AnimGraphNode_RotationOffsetBlendSpace.h
@@ -9,9 +9,7 @@
 UCLASS(MinimalAPI)
 class UAnimGraphNode_RotationOffsetBlendSpace: public UAnimGraphNode_BlendSpaceBase
 {
-	GENERATED_BODY()
-public:
-	ANIMGRAPH_API UAnimGraphNode_RotationOffsetBlendSpace(const FObjectInitializer& ObjectInitializer = FObjectInitializer::Get());
+	GENERATED_UCLASS_BODY()
 
 	UPROPERTY(EditAnywhere, Category=Settings)
 	FAnimNode_RotationOffsetBlendSpace Node;
diff --git a/Engine/Source/Editor/AnimGraph/Classes/AnimGraphNode_SaveCachedPose.h b/Engine/Source/Editor/AnimGraph/Classes/AnimGraphNode_SaveCachedPose.h
--- a/Engine/Source/Editor/AnimGraph/Classes/AnimGraphNode_SaveCachedPose.h
+++ b/Engine/Source/Editor/AnimGraph/Classes/AnimGraphNode_SaveCachedPose.h
@@ -9,9 +9,7 @@
 UCLASS(MinimalAPI)
 class UAnimGraphNode_SaveCachedPose : public UAnimGraphNode_Base
 {
-	GENERATED_BODY()
-public:
-	ANIMGRAPH_API UAnimGraphNode_SaveCachedPose(const FObjectInitializer& ObjectInitializer = FObjectInitializer::Get());
+	GENERATED_UCLASS_BODY()
 
 	UPROPERTY(EditAnywhere, Category=Settings)
 	FAnimNode_SaveCachedPose Node;
diff --git a/Engine/Source/Editor/AnimGraph/Classes/AnimGraphNode_SequenceEvaluator.h b/Engine/Source/Editor/AnimGraph/Classes/AnimGraphNode_SequenceEvaluator.h
--- a/Engine/Source/Editor/AnimGraph/Classes/AnimGraphNode_SequenceEvaluator.h
+++ b/Engine/Source/Editor/AnimGraph/Classes/AnimGraphNode_SequenceEvaluator.h
@@ -9,9 +9,7 @@
 UCLASS(MinimalAPI)
 class UAnimGraphNode_SequenceEvaluator : public UAnimGraphNode_Base
 {
-	GENERATED_BODY()
-public:
-	ANIMGRAPH_API UAnimGraphNode_SequenceEvaluator(const FObjectInitializer& ObjectInitializer = FObjectInitializer::Get());
+	GENERATED_UCLASS_BODY()
 
 	UPROPERTY(EditAnywhere, Category=Settings)
 	FAnimNode_SequenceEvaluator Node;
diff --git a/Engine/Source/Editor/AnimGraph/Classes/AnimGraphNode_SequencePlayer.h b/Engine/Source/Editor/AnimGraph/Classes/AnimGraphNode_SequencePlayer.h
--- a/Engine/Source/Editor/AnimGraph/Classes/AnimGraphNode_SequencePlayer.h
+++ b/Engine/Source/Editor/AnimGraph/Classes/AnimGraphNode_SequencePlayer.h
@@ -9,9 +9,7 @@
 UCLASS(MinimalAPI)
 class UAnimGraphNode_SequencePlayer : public UAnimGraphNode_Base
 {
-	GENERATED_BODY()
-public:
-	ANIMGRAPH_API UAnimGraphNode_SequencePlayer(const FObjectInitializer& ObjectInitializer = FObjectInitializer::Get());
+	GENERATED_UCLASS_BODY()
 
 	UPROPERTY(EditAnywhere, Category=Settings)
 	FAnimNode_SequencePlayer Node;
diff --git a/Engine/Source/Editor/AnimGraph/Classes/AnimGraphNode_Slot.h b/Engine/Source/Editor/AnimGraph/Classes/AnimGraphNode_Slot.h
--- a/Engine/Source/Editor/AnimGraph/Classes/AnimGraphNode_Slot.h
+++ b/Engine/Source/Editor/AnimGraph/Classes/AnimGraphNode_Slot.h
@@ -9,9 +9,7 @@
 UCLASS(MinimalAPI)
 class UAnimGraphNode_Slot : public UAnimGraphNode_Base
 {
-	GENERATED_BODY()
-public:
-	ANIMGRAPH_API UAnimGraphNode_Slot(const FObjectInitializer& ObjectInitializer = FObjectInitializer::Get());
+	GENERATED_UCLASS_BODY()
 
 	UPROPERTY(EditAnywhere, Category=Settings)
 	FAnimNode_Slot Node;
diff --git a/Engine/Source/Editor/AnimGraph/Classes/AnimGraphNode_SpringBone.h b/Engine/Source/Editor/AnimGraph/Classes/AnimGraphNode_SpringBone.h
--- a/Engine/Source/Editor/AnimGraph/Classes/AnimGraphNode_SpringBone.h
+++ b/Engine/Source/Editor/AnimGraph/Classes/AnimGraphNode_SpringBone.h
@@ -9,9 +9,7 @@
 UCLASS(MinimalAPI)
 class UAnimGraphNode_SpringBone : public UAnimGraphNode_SkeletalControlBase
 {
-	GENERATED_BODY()
-public:
-	ANIMGRAPH_API UAnimGraphNode_SpringBone(const FObjectInitializer& ObjectInitializer = FObjectInitializer::Get());
+	GENERATED_UCLASS_BODY()
 
 	UPROPERTY(EditAnywhere, Category=Settings)
 	FAnimNode_SpringBone Node;
diff --git a/Engine/Source/Editor/AnimGraph/Classes/AnimGraphNode_StateMachine.h b/Engine/Source/Editor/AnimGraph/Classes/AnimGraphNode_StateMachine.h
--- a/Engine/Source/Editor/AnimGraph/Classes/AnimGraphNode_StateMachine.h
+++ b/Engine/Source/Editor/AnimGraph/Classes/AnimGraphNode_StateMachine.h
@@ -7,9 +7,7 @@
 UCLASS(MinimalAPI)
 class UAnimGraphNode_StateMachine : public UAnimGraphNode_StateMachineBase
 {
-	GENERATED_BODY()
-public:
-	ANIMGRAPH_API UAnimGraphNode_StateMachine(const FObjectInitializer& ObjectInitializer = FObjectInitializer::Get());
+	GENERATED_UCLASS_BODY()
 
 	// Runtime state machine representation (empty; populated during compilation)
 	UPROPERTY(EditAnywhere, Category=Settings)
diff --git a/Engine/Source/Editor/AnimGraph/Classes/AnimGraphNode_StateMachineBase.h b/Engine/Source/Editor/AnimGraph/Classes/AnimGraphNode_StateMachineBase.h
--- a/Engine/Source/Editor/AnimGraph/Classes/AnimGraphNode_StateMachineBase.h
+++ b/Engine/Source/Editor/AnimGraph/Classes/AnimGraphNode_StateMachineBase.h
@@ -9,9 +9,7 @@
 UCLASS(Abstract)
 class ANIMGRAPH_API UAnimGraphNode_StateMachineBase : public UAnimGraphNode_Base
 {
-	GENERATED_BODY()
-public:
-	UAnimGraphNode_StateMachineBase(const FObjectInitializer& ObjectInitializer = FObjectInitializer::Get());
+	GENERATED_UCLASS_BODY()
 
 	// Editor state machine representation
 	UPROPERTY()
diff --git a/Engine/Source/Editor/AnimGraph/Classes/AnimGraphNode_StateResult.h b/Engine/Source/Editor/AnimGraph/Classes/AnimGraphNode_StateResult.h
--- a/Engine/Source/Editor/AnimGraph/Classes/AnimGraphNode_StateResult.h
+++ b/Engine/Source/Editor/AnimGraph/Classes/AnimGraphNode_StateResult.h
@@ -8,7 +8,5 @@
 UCLASS(MinimalAPI)
 class UAnimGraphNode_StateResult : public UAnimGraphNode_Root
 {
-	GENERATED_BODY()
-public:
-	ANIMGRAPH_API UAnimGraphNode_StateResult(const FObjectInitializer& ObjectInitializer = FObjectInitializer::Get());
+	GENERATED_UCLASS_BODY()
 };
diff --git a/Engine/Source/Editor/AnimGraph/Classes/AnimGraphNode_Trail.h b/Engine/Source/Editor/AnimGraph/Classes/AnimGraphNode_Trail.h
--- a/Engine/Source/Editor/AnimGraph/Classes/AnimGraphNode_Trail.h
+++ b/Engine/Source/Editor/AnimGraph/Classes/AnimGraphNode_Trail.h
@@ -9,9 +9,7 @@
 UCLASS(MinimalAPI)
 class UAnimGraphNode_Trail : public UAnimGraphNode_SkeletalControlBase
 {
-	GENERATED_BODY()
-public:
-	ANIMGRAPH_API UAnimGraphNode_Trail(const FObjectInitializer& ObjectInitializer = FObjectInitializer::Get());
+	GENERATED_UCLASS_BODY()
 
 		UPROPERTY(EditAnywhere, Category = Settings)
 		FAnimNode_Trail Node;
diff --git a/Engine/Source/Editor/AnimGraph/Classes/AnimGraphNode_TransitionPoseEvaluator.h b/Engine/Source/Editor/AnimGraph/Classes/AnimGraphNode_TransitionPoseEvaluator.h
--- a/Engine/Source/Editor/AnimGraph/Classes/AnimGraphNode_TransitionPoseEvaluator.h
+++ b/Engine/Source/Editor/AnimGraph/Classes/AnimGraphNode_TransitionPoseEvaluator.h
@@ -8,9 +8,7 @@
 UCLASS(MinimalAPI)
 class UAnimGraphNode_TransitionPoseEvaluator : public UAnimGraphNode_Base
 {
-	GENERATED_BODY()
-public:
-	ANIMGRAPH_API UAnimGraphNode_TransitionPoseEvaluator(const FObjectInitializer& ObjectInitializer = FObjectInitializer::Get());
+	GENERATED_UCLASS_BODY()
 
 	UPROPERTY(EditAnywhere, BlueprintReadWrite, Category=Settings)
 	FAnimNode_TransitionPoseEvaluator Node;
@@ -32,4 +30,4 @@ public:
 	// UK2Node interface.
 	virtual void GetMenuActions(FBlueprintActionDatabaseRegistrar& ActionRegistrar) const override;
 	// End of UK2Node interface
-};
+};
\ No newline at end of file
diff --git a/Engine/Source/Editor/AnimGraph/Classes/AnimGraphNode_TransitionResult.h b/Engine/Source/Editor/AnimGraph/Classes/AnimGraphNode_TransitionResult.h
--- a/Engine/Source/Editor/AnimGraph/Classes/AnimGraphNode_TransitionResult.h
+++ b/Engine/Source/Editor/AnimGraph/Classes/AnimGraphNode_TransitionResult.h
@@ -8,9 +8,7 @@
 UCLASS(MinimalAPI)
 class UAnimGraphNode_TransitionResult : public UAnimGraphNode_Base
 {
-	GENERATED_BODY()
-public:
-	ANIMGRAPH_API UAnimGraphNode_TransitionResult(const FObjectInitializer& ObjectInitializer = FObjectInitializer::Get());
+	GENERATED_UCLASS_BODY()
 
 	UPROPERTY(EditAnywhere, Category=Settings)
 	FAnimNode_TransitionResult Node;
diff --git a/Engine/Source/Editor/AnimGraph/Classes/AnimGraphNode_TwoBoneIK.h b/Engine/Source/Editor/AnimGraph/Classes/AnimGraphNode_TwoBoneIK.h
--- a/Engine/Source/Editor/AnimGraph/Classes/AnimGraphNode_TwoBoneIK.h
+++ b/Engine/Source/Editor/AnimGraph/Classes/AnimGraphNode_TwoBoneIK.h
@@ -12,9 +12,7 @@
 UCLASS(MinimalAPI)
 class UAnimGraphNode_TwoBoneIK : public UAnimGraphNode_SkeletalControlBase
 {
-	GENERATED_BODY()
-public:
-	ANIMGRAPH_API UAnimGraphNode_TwoBoneIK(const FObjectInitializer& ObjectInitializer = FObjectInitializer::Get());
+	GENERATED_UCLASS_BODY()
 
 	UPROPERTY(EditAnywhere, Category=Settings)
 	FAnimNode_TwoBoneIK Node;
diff --git a/Engine/Source/Editor/AnimGraph/Classes/AnimGraphNode_TwoWayBlend.h b/Engine/Source/Editor/AnimGraph/Classes/AnimGraphNode_TwoWayBlend.h
--- a/Engine/Source/Editor/AnimGraph/Classes/AnimGraphNode_TwoWayBlend.h
+++ b/Engine/Source/Editor/AnimGraph/Classes/AnimGraphNode_TwoWayBlend.h
@@ -8,9 +8,7 @@
 UCLASS(MinimalAPI)
 class UAnimGraphNode_TwoWayBlend : public UAnimGraphNode_Base
 {
-	GENERATED_BODY()
-public:
-	ANIMGRAPH_API UAnimGraphNode_TwoWayBlend(const FObjectInitializer& ObjectInitializer = FObjectInitializer::Get());
+	GENERATED_UCLASS_BODY()
 
 	UPROPERTY(EditAnywhere, Category=Settings)
 	FAnimationNode_TwoWayBlend BlendNode;
diff --git a/Engine/Source/Editor/AnimGraph/Classes/AnimGraphNode_WheelHandler.h b/Engine/Source/Editor/AnimGraph/Classes/AnimGraphNode_WheelHandler.h
--- a/Engine/Source/Editor/AnimGraph/Classes/AnimGraphNode_WheelHandler.h
+++ b/Engine/Source/Editor/AnimGraph/Classes/AnimGraphNode_WheelHandler.h
@@ -9,9 +9,7 @@
 UCLASS(MinimalAPI)
 class UAnimGraphNode_WheelHandler : public UAnimGraphNode_SkeletalControlBase
 {
-	GENERATED_BODY()
-public:
-	ANIMGRAPH_API UAnimGraphNode_WheelHandler(const FObjectInitializer& ObjectInitializer = FObjectInitializer::Get());
+	GENERATED_UCLASS_BODY()
 
 	UPROPERTY(EditAnywhere, Category=Settings)
 	FAnimNode_WheelHandler Node;
diff --git a/Engine/Source/Editor/AnimGraph/Classes/AnimPreviewInstance.h b/Engine/Source/Editor/AnimGraph/Classes/AnimPreviewInstance.h
--- a/Engine/Source/Editor/AnimGraph/Classes/AnimPreviewInstance.h
+++ b/Engine/Source/Editor/AnimGraph/Classes/AnimPreviewInstance.h
@@ -27,9 +27,7 @@ enum EMontagePreviewType
 UCLASS(transient, NotBlueprintable, noteditinlinenew)
 class ANIMGRAPH_API UAnimPreviewInstance : public UAnimSingleNodeInstance
 {
-	GENERATED_BODY()
-public:
-	UAnimPreviewInstance(const FObjectInitializer& ObjectInitializer = FObjectInitializer::Get());
+	GENERATED_UCLASS_BODY()
 
 	/** Controllers for individual bones */
 	UPROPERTY(transient)
diff --git a/Engine/Source/Editor/AnimGraph/Classes/AnimStateConduitNode.h b/Engine/Source/Editor/AnimGraph/Classes/AnimStateConduitNode.h
--- a/Engine/Source/Editor/AnimGraph/Classes/AnimStateConduitNode.h
+++ b/Engine/Source/Editor/AnimGraph/Classes/AnimStateConduitNode.h
@@ -7,9 +7,7 @@
 UCLASS(MinimalAPI)
 class UAnimStateConduitNode : public UAnimStateNodeBase
 {
-	GENERATED_BODY()
-public:
-	ANIMGRAPH_API UAnimStateConduitNode(const FObjectInitializer& ObjectInitializer = FObjectInitializer::Get());
+	GENERATED_UCLASS_BODY()
 public:
 
 	// The transition graph for this conduit; it's a logic graph, not an animation graph
diff --git a/Engine/Source/Editor/AnimGraph/Classes/AnimStateEntryNode.h b/Engine/Source/Editor/AnimGraph/Classes/AnimStateEntryNode.h
--- a/Engine/Source/Editor/AnimGraph/Classes/AnimStateEntryNode.h
+++ b/Engine/Source/Editor/AnimGraph/Classes/AnimStateEntryNode.h
@@ -7,9 +7,7 @@
 UCLASS(MinimalAPI)
 class UAnimStateEntryNode : public UEdGraphNode
 {
-	GENERATED_BODY()
-public:
-	ANIMGRAPH_API UAnimStateEntryNode(const FObjectInitializer& ObjectInitializer = FObjectInitializer::Get());
+	GENERATED_UCLASS_BODY()
 
 
 	// Begin UEdGraphNode interface
diff --git a/Engine/Source/Editor/AnimGraph/Classes/AnimStateNode.h b/Engine/Source/Editor/AnimGraph/Classes/AnimStateNode.h
--- a/Engine/Source/Editor/AnimGraph/Classes/AnimStateNode.h
+++ b/Engine/Source/Editor/AnimGraph/Classes/AnimStateNode.h
@@ -15,9 +15,7 @@ enum EAnimStateType
 UCLASS(MinimalAPI)
 class UAnimStateNode : public UAnimStateNodeBase
 {
-	GENERATED_BODY()
-public:
-	ANIMGRAPH_API UAnimStateNode(const FObjectInitializer& ObjectInitializer = FObjectInitializer::Get());
+	GENERATED_UCLASS_BODY()
 public:
 
 	// The animation graph for this state
diff --git a/Engine/Source/Editor/AnimGraph/Classes/AnimStateNodeBase.h b/Engine/Source/Editor/AnimGraph/Classes/AnimStateNodeBase.h
--- a/Engine/Source/Editor/AnimGraph/Classes/AnimStateNodeBase.h
+++ b/Engine/Source/Editor/AnimGraph/Classes/AnimStateNodeBase.h
@@ -7,9 +7,7 @@
 UCLASS(MinimalAPI, abstract)
 class UAnimStateNodeBase : public UEdGraphNode
 {
-	GENERATED_BODY()
-public:
-	ANIMGRAPH_API UAnimStateNodeBase(const FObjectInitializer& ObjectInitializer = FObjectInitializer::Get());
+	GENERATED_UCLASS_BODY()
 
 	// UEdGraphNode interface
 	virtual void PostPasteNode() override;
diff --git a/Engine/Source/Editor/AnimGraph/Classes/AnimStateTransitionNode.h b/Engine/Source/Editor/AnimGraph/Classes/AnimStateTransitionNode.h
--- a/Engine/Source/Editor/AnimGraph/Classes/AnimStateTransitionNode.h
+++ b/Engine/Source/Editor/AnimGraph/Classes/AnimStateTransitionNode.h
@@ -11,9 +11,7 @@
 UCLASS(MinimalAPI, config=Editor)
 class UAnimStateTransitionNode : public UAnimStateNodeBase
 {
-	GENERATED_BODY()
-public:
-	ANIMGRAPH_API UAnimStateTransitionNode(const FObjectInitializer& ObjectInitializer = FObjectInitializer::Get());
+	GENERATED_UCLASS_BODY()
 
 	// The transition logic graph for this transition (returning a boolean)
 	UPROPERTY()
diff --git a/Engine/Source/Editor/AnimGraph/Classes/AnimationConduitGraphSchema.h b/Engine/Source/Editor/AnimGraph/Classes/AnimationConduitGraphSchema.h
--- a/Engine/Source/Editor/AnimGraph/Classes/AnimationConduitGraphSchema.h
+++ b/Engine/Source/Editor/AnimGraph/Classes/AnimationConduitGraphSchema.h
@@ -8,9 +8,7 @@
 UCLASS(MinimalAPI)
 class UAnimationConduitGraphSchema : public UEdGraphSchema_K2
 {
-	GENERATED_BODY()
-public:
-	ANIMGRAPH_API UAnimationConduitGraphSchema(const FObjectInitializer& ObjectInitializer = FObjectInitializer::Get());
+	GENERATED_UCLASS_BODY()
 
 	// UEdGraphSchema interface
 	virtual void CreateDefaultNodesForGraph(UEdGraph& Graph) const override;
diff --git a/Engine/Source/Editor/AnimGraph/Classes/AnimationCustomTransitionGraph.h b/Engine/Source/Editor/AnimGraph/Classes/AnimationCustomTransitionGraph.h
--- a/Engine/Source/Editor/AnimGraph/Classes/AnimationCustomTransitionGraph.h
+++ b/Engine/Source/Editor/AnimGraph/Classes/AnimationCustomTransitionGraph.h
@@ -7,9 +7,7 @@
 UCLASS(MinimalAPI)
 class UAnimationCustomTransitionGraph : public UAnimationGraph
 {
-	GENERATED_BODY()
-public:
-	ANIMGRAPH_API UAnimationCustomTransitionGraph(const FObjectInitializer& ObjectInitializer = FObjectInitializer::Get());
+	GENERATED_UCLASS_BODY()
 
 	// Result node within the state's animation graph
 	UPROPERTY()
diff --git a/Engine/Source/Editor/AnimGraph/Classes/AnimationCustomTransitionSchema.h b/Engine/Source/Editor/AnimGraph/Classes/AnimationCustomTransitionSchema.h
--- a/Engine/Source/Editor/AnimGraph/Classes/AnimationCustomTransitionSchema.h
+++ b/Engine/Source/Editor/AnimGraph/Classes/AnimationCustomTransitionSchema.h
@@ -7,9 +7,7 @@
 UCLASS(MinimalAPI)
 class UAnimationCustomTransitionSchema : public UAnimationGraphSchema
 {
-	GENERATED_BODY()
-public:
-	ANIMGRAPH_API UAnimationCustomTransitionSchema(const FObjectInitializer& ObjectInitializer = FObjectInitializer::Get());
+	GENERATED_UCLASS_BODY()
 
 	// Begin UEdGraphSchema interface.
 	virtual void CreateDefaultNodesForGraph(UEdGraph& Graph) const override;
diff --git a/Engine/Source/Editor/AnimGraph/Classes/AnimationGraph.h b/Engine/Source/Editor/AnimGraph/Classes/AnimationGraph.h
--- a/Engine/Source/Editor/AnimGraph/Classes/AnimationGraph.h
+++ b/Engine/Source/Editor/AnimGraph/Classes/AnimationGraph.h
@@ -7,8 +7,6 @@
 UCLASS(MinimalAPI)
 class UAnimationGraph : public UEdGraph
 {
-	GENERATED_BODY()
-public:
-	ANIMGRAPH_API UAnimationGraph(const FObjectInitializer& ObjectInitializer = FObjectInitializer::Get());
+	GENERATED_UCLASS_BODY()
 };
 
diff --git a/Engine/Source/Editor/AnimGraph/Classes/AnimationGraphSchema.h b/Engine/Source/Editor/AnimGraph/Classes/AnimationGraphSchema.h
--- a/Engine/Source/Editor/AnimGraph/Classes/AnimationGraphSchema.h
+++ b/Engine/Source/Editor/AnimGraph/Classes/AnimationGraphSchema.h
@@ -8,9 +8,7 @@
 UCLASS(MinimalAPI)
 class UAnimationGraphSchema : public UEdGraphSchema_K2
 {
-	GENERATED_BODY()
-public:
-	ANIMGRAPH_API UAnimationGraphSchema(const FObjectInitializer& ObjectInitializer = FObjectInitializer::Get());
+	GENERATED_UCLASS_BODY()
 
 	// Common PinNames
 	UPROPERTY()
diff --git a/Engine/Source/Editor/AnimGraph/Classes/AnimationStateGraph.h b/Engine/Source/Editor/AnimGraph/Classes/AnimationStateGraph.h
--- a/Engine/Source/Editor/AnimGraph/Classes/AnimationStateGraph.h
+++ b/Engine/Source/Editor/AnimGraph/Classes/AnimationStateGraph.h
@@ -7,9 +7,7 @@
 UCLASS(MinimalAPI)
 class UAnimationStateGraph : public UAnimationGraph
 {
-	GENERATED_BODY()
-public:
-	ANIMGRAPH_API UAnimationStateGraph(const FObjectInitializer& ObjectInitializer = FObjectInitializer::Get());
+	GENERATED_UCLASS_BODY()
 
 	// Result node within the state's animation graph
 	UPROPERTY()
diff --git a/Engine/Source/Editor/AnimGraph/Classes/AnimationStateGraphSchema.h b/Engine/Source/Editor/AnimGraph/Classes/AnimationStateGraphSchema.h
--- a/Engine/Source/Editor/AnimGraph/Classes/AnimationStateGraphSchema.h
+++ b/Engine/Source/Editor/AnimGraph/Classes/AnimationStateGraphSchema.h
@@ -7,9 +7,7 @@
 UCLASS(MinimalAPI)
 class UAnimationStateGraphSchema : public UAnimationGraphSchema
 {
-	GENERATED_BODY()
-public:
-	ANIMGRAPH_API UAnimationStateGraphSchema(const FObjectInitializer& ObjectInitializer = FObjectInitializer::Get());
+	GENERATED_UCLASS_BODY()
 
 	// Begin UEdGraphSchema interface.
 	virtual void CreateDefaultNodesForGraph(UEdGraph& Graph) const override;
diff --git a/Engine/Source/Editor/AnimGraph/Classes/AnimationStateMachineGraph.h b/Engine/Source/Editor/AnimGraph/Classes/AnimationStateMachineGraph.h
--- a/Engine/Source/Editor/AnimGraph/Classes/AnimationStateMachineGraph.h
+++ b/Engine/Source/Editor/AnimGraph/Classes/AnimationStateMachineGraph.h
@@ -7,9 +7,7 @@
 UCLASS(MinimalAPI)
 class UAnimationStateMachineGraph : public UEdGraph
 {
-	GENERATED_BODY()
-public:
-	ANIMGRAPH_API UAnimationStateMachineGraph(const FObjectInitializer& ObjectInitializer = FObjectInitializer::Get());
+	GENERATED_UCLASS_BODY()
 
 	// Entry node within the state machine
 	UPROPERTY()
diff --git a/Engine/Source/Editor/AnimGraph/Classes/AnimationTransitionGraph.h b/Engine/Source/Editor/AnimGraph/Classes/AnimationTransitionGraph.h
--- a/Engine/Source/Editor/AnimGraph/Classes/AnimationTransitionGraph.h
+++ b/Engine/Source/Editor/AnimGraph/Classes/AnimationTransitionGraph.h
@@ -7,9 +7,7 @@
 UCLASS(MinimalAPI)
 class UAnimationTransitionGraph : public UAnimationGraph
 {
-	GENERATED_BODY()
-public:
-	ANIMGRAPH_API UAnimationTransitionGraph(const FObjectInitializer& ObjectInitializer = FObjectInitializer::Get());
+	GENERATED_UCLASS_BODY()
 
 	UPROPERTY()
 	class UAnimGraphNode_TransitionResult* MyResultNode;
diff --git a/Engine/Source/Editor/AnimGraph/Classes/AnimationTransitionSchema.h b/Engine/Source/Editor/AnimGraph/Classes/AnimationTransitionSchema.h
--- a/Engine/Source/Editor/AnimGraph/Classes/AnimationTransitionSchema.h
+++ b/Engine/Source/Editor/AnimGraph/Classes/AnimationTransitionSchema.h
@@ -7,9 +7,7 @@
 UCLASS(MinimalAPI)
 class UAnimationTransitionSchema : public UEdGraphSchema_K2
 {
-	GENERATED_BODY()
-public:
-	ANIMGRAPH_API UAnimationTransitionSchema(const FObjectInitializer& ObjectInitializer = FObjectInitializer::Get());
+	GENERATED_UCLASS_BODY()
 
 	// Begin UEdGraphSchema interface.
 	virtual void GetGraphContextActions(FGraphContextMenuBuilder& ContextMenuBuilder) const override;
diff --git a/Engine/Source/Editor/AnimGraph/Classes/K2Node_TransitionRuleGetter.h b/Engine/Source/Editor/AnimGraph/Classes/K2Node_TransitionRuleGetter.h
--- a/Engine/Source/Editor/AnimGraph/Classes/K2Node_TransitionRuleGetter.h
+++ b/Engine/Source/Editor/AnimGraph/Classes/K2Node_TransitionRuleGetter.h
@@ -25,9 +25,7 @@ namespace ETransitionGetter
 UCLASS(MinimalAPI)
 class UK2Node_TransitionRuleGetter : public UK2Node
 {
-	GENERATED_BODY()
-public:
-	ANIMGRAPH_API UK2Node_TransitionRuleGetter(const FObjectInitializer& ObjectInitializer = FObjectInitializer::Get());
+	GENERATED_UCLASS_BODY()
 
 	UPROPERTY()
 	TEnumAsByte<ETransitionGetter::Type> GetterType;
diff --git a/Engine/Source/Editor/BehaviorTreeEditor/Classes/BehaviorTreeDecoratorGraph.h b/Engine/Source/Editor/BehaviorTreeEditor/Classes/BehaviorTreeDecoratorGraph.h
--- a/Engine/Source/Editor/BehaviorTreeEditor/Classes/BehaviorTreeDecoratorGraph.h
+++ b/Engine/Source/Editor/BehaviorTreeEditor/Classes/BehaviorTreeDecoratorGraph.h
@@ -7,9 +7,7 @@
 UCLASS()
 class UBehaviorTreeDecoratorGraph : public UEdGraph
 {
-	GENERATED_BODY()
-public:
-	UBehaviorTreeDecoratorGraph(const FObjectInitializer& ObjectInitializer = FObjectInitializer::Get());
+	GENERATED_UCLASS_BODY()
 
 	void CollectDecoratorData(TArray<class UBTDecorator*>& DecoratorInstances, TArray<struct FBTDecoratorLogic>& DecoratorOperations) const;
 	int32 SpawnMissingNodes(const TArray<class UBTDecorator*>& NodeInstances, const TArray<struct FBTDecoratorLogic>& Operations, int32 StartIndex);
diff --git a/Engine/Source/Editor/BehaviorTreeEditor/Classes/BehaviorTreeDecoratorGraphNode.h b/Engine/Source/Editor/BehaviorTreeEditor/Classes/BehaviorTreeDecoratorGraphNode.h
--- a/Engine/Source/Editor/BehaviorTreeEditor/Classes/BehaviorTreeDecoratorGraphNode.h
+++ b/Engine/Source/Editor/BehaviorTreeEditor/Classes/BehaviorTreeDecoratorGraphNode.h
@@ -10,9 +10,7 @@
 UCLASS()
 class UBehaviorTreeDecoratorGraphNode : public UEdGraphNode
 {
-	GENERATED_BODY()
-public:
-	UBehaviorTreeDecoratorGraphNode(const FObjectInitializer& ObjectInitializer = FObjectInitializer::Get());
+	GENERATED_UCLASS_BODY()
 
 	/** Get the BT graph that owns this node */
 	virtual class UBehaviorTreeDecoratorGraph* GetDecoratorGraph();
diff --git a/Engine/Source/Editor/BehaviorTreeEditor/Classes/BehaviorTreeDecoratorGraphNode_Logic.h b/Engine/Source/Editor/BehaviorTreeEditor/Classes/BehaviorTreeDecoratorGraphNode_Logic.h
--- a/Engine/Source/Editor/BehaviorTreeEditor/Classes/BehaviorTreeDecoratorGraphNode_Logic.h
+++ b/Engine/Source/Editor/BehaviorTreeEditor/Classes/BehaviorTreeDecoratorGraphNode_Logic.h
@@ -18,9 +18,7 @@ namespace EDecoratorLogicMode
 UCLASS()
 class UBehaviorTreeDecoratorGraphNode_Logic : public UBehaviorTreeDecoratorGraphNode
 {
-	GENERATED_BODY()
-public:
-	UBehaviorTreeDecoratorGraphNode_Logic(const FObjectInitializer& ObjectInitializer = FObjectInitializer::Get());
+	GENERATED_UCLASS_BODY()
 
 	UPROPERTY()
 	TEnumAsByte<EDecoratorLogicMode::Type> LogicMode;
diff --git a/Engine/Source/Editor/BehaviorTreeEditor/Classes/BehaviorTreeEditorTypes.h b/Engine/Source/Editor/BehaviorTreeEditor/Classes/BehaviorTreeEditorTypes.h
--- a/Engine/Source/Editor/BehaviorTreeEditor/Classes/BehaviorTreeEditorTypes.h
+++ b/Engine/Source/Editor/BehaviorTreeEditor/Classes/BehaviorTreeEditorTypes.h
@@ -46,9 +46,7 @@ struct FNodeBounds
 UCLASS()
 class UBehaviorTreeEditorTypes : public UObject
 {
-	GENERATED_BODY()
-public:
-	UBehaviorTreeEditorTypes(const FObjectInitializer& ObjectInitializer = FObjectInitializer::Get());
+	GENERATED_UCLASS_BODY()
 
 	static const FString PinCategory_MultipleNodes;
 	static const FString PinCategory_SingleComposite;
diff --git a/Engine/Source/Editor/BehaviorTreeEditor/Classes/BehaviorTreeGraph.h b/Engine/Source/Editor/BehaviorTreeEditor/Classes/BehaviorTreeGraph.h
--- a/Engine/Source/Editor/BehaviorTreeEditor/Classes/BehaviorTreeGraph.h
+++ b/Engine/Source/Editor/BehaviorTreeEditor/Classes/BehaviorTreeGraph.h
@@ -8,9 +8,7 @@
 UCLASS()
 class UBehaviorTreeGraph : public UAIGraph
 {
-	GENERATED_BODY()
-public:
-	UBehaviorTreeGraph(const FObjectInitializer& ObjectInitializer = FObjectInitializer::Get());
+	GENERATED_UCLASS_BODY()
 
 	enum EUpdateFlags
 	{
diff --git a/Engine/Source/Editor/BehaviorTreeEditor/Classes/BehaviorTreeGraphNode.h b/Engine/Source/Editor/BehaviorTreeEditor/Classes/BehaviorTreeGraphNode.h
--- a/Engine/Source/Editor/BehaviorTreeEditor/Classes/BehaviorTreeGraphNode.h
+++ b/Engine/Source/Editor/BehaviorTreeEditor/Classes/BehaviorTreeGraphNode.h
@@ -9,9 +9,7 @@
 UCLASS()
 class UBehaviorTreeGraphNode : public UAIGraphNode
 {
-	GENERATED_BODY()
-public:
-	UBehaviorTreeGraphNode(const FObjectInitializer& ObjectInitializer = FObjectInitializer::Get());
+	GENERATED_UCLASS_BODY()
 
 	/** only some of behavior tree nodes support decorators */
 	UPROPERTY()
diff --git a/Engine/Source/Editor/BehaviorTreeEditor/Classes/BehaviorTreeGraphNode_Composite.h b/Engine/Source/Editor/BehaviorTreeEditor/Classes/BehaviorTreeGraphNode_Composite.h
--- a/Engine/Source/Editor/BehaviorTreeEditor/Classes/BehaviorTreeGraphNode_Composite.h
+++ b/Engine/Source/Editor/BehaviorTreeEditor/Classes/BehaviorTreeGraphNode_Composite.h
@@ -6,9 +6,7 @@
 UCLASS()
 class UBehaviorTreeGraphNode_Composite : public UBehaviorTreeGraphNode
 {
-	GENERATED_BODY()
-public:
-	UBehaviorTreeGraphNode_Composite(const FObjectInitializer& ObjectInitializer = FObjectInitializer::Get());
+	GENERATED_UCLASS_BODY()
 	
 	virtual FText GetNodeTitle(ENodeTitleType::Type TitleType) const override;
 	virtual bool RefreshNodeClass() override{ return false; }
diff --git a/Engine/Source/Editor/BehaviorTreeEditor/Classes/BehaviorTreeGraphNode_CompositeDecorator.h b/Engine/Source/Editor/BehaviorTreeEditor/Classes/BehaviorTreeGraphNode_CompositeDecorator.h
--- a/Engine/Source/Editor/BehaviorTreeEditor/Classes/BehaviorTreeGraphNode_CompositeDecorator.h
+++ b/Engine/Source/Editor/BehaviorTreeEditor/Classes/BehaviorTreeGraphNode_CompositeDecorator.h
@@ -6,9 +6,7 @@
 UCLASS()
 class UBehaviorTreeGraphNode_CompositeDecorator : public UBehaviorTreeGraphNode
 {
-	GENERATED_BODY()
-public:
-	UBehaviorTreeGraphNode_CompositeDecorator(const FObjectInitializer& ObjectInitializer = FObjectInitializer::Get());
+	GENERATED_UCLASS_BODY()
 
 	// The logic graph for this decorator (returning a boolean)
 	UPROPERTY()
diff --git a/Engine/Source/Editor/BehaviorTreeEditor/Classes/BehaviorTreeGraphNode_Decorator.h b/Engine/Source/Editor/BehaviorTreeEditor/Classes/BehaviorTreeGraphNode_Decorator.h
--- a/Engine/Source/Editor/BehaviorTreeEditor/Classes/BehaviorTreeGraphNode_Decorator.h
+++ b/Engine/Source/Editor/BehaviorTreeEditor/Classes/BehaviorTreeGraphNode_Decorator.h
@@ -6,9 +6,7 @@
 UCLASS()
 class UBehaviorTreeGraphNode_Decorator : public UBehaviorTreeGraphNode
 {
-	GENERATED_BODY()
-public:
-	UBehaviorTreeGraphNode_Decorator(const FObjectInitializer& ObjectInitializer = FObjectInitializer::Get());
+	GENERATED_UCLASS_BODY()
 
 	virtual FText GetNodeTitle(ENodeTitleType::Type TitleType) const override;
 	virtual void AllocateDefaultPins() override;
diff --git a/Engine/Source/Editor/BehaviorTreeEditor/Classes/BehaviorTreeGraphNode_Root.h b/Engine/Source/Editor/BehaviorTreeEditor/Classes/BehaviorTreeGraphNode_Root.h
--- a/Engine/Source/Editor/BehaviorTreeEditor/Classes/BehaviorTreeGraphNode_Root.h
+++ b/Engine/Source/Editor/BehaviorTreeEditor/Classes/BehaviorTreeGraphNode_Root.h
@@ -7,9 +7,7 @@
 UCLASS()
 class UBehaviorTreeGraphNode_Root : public UBehaviorTreeGraphNode
 {
-	GENERATED_BODY()
-public:
-	UBehaviorTreeGraphNode_Root(const FObjectInitializer& ObjectInitializer = FObjectInitializer::Get());
+	GENERATED_UCLASS_BODY()
 
 	UPROPERTY(EditAnywhere, Category=BehaviorTree)
 	class UBlackboardData* BlackboardAsset;
diff --git a/Engine/Source/Editor/BehaviorTreeEditor/Classes/BehaviorTreeGraphNode_Service.h b/Engine/Source/Editor/BehaviorTreeEditor/Classes/BehaviorTreeGraphNode_Service.h
--- a/Engine/Source/Editor/BehaviorTreeEditor/Classes/BehaviorTreeGraphNode_Service.h
+++ b/Engine/Source/Editor/BehaviorTreeEditor/Classes/BehaviorTreeGraphNode_Service.h
@@ -6,9 +6,7 @@
 UCLASS()
 class UBehaviorTreeGraphNode_Service : public UBehaviorTreeGraphNode
 {
-	GENERATED_BODY()
-public:
-	UBehaviorTreeGraphNode_Service(const FObjectInitializer& ObjectInitializer = FObjectInitializer::Get());
+	GENERATED_UCLASS_BODY()
 
 	virtual FText GetNodeTitle(ENodeTitleType::Type TitleType) const override;
 	virtual void AllocateDefaultPins() override;
diff --git a/Engine/Source/Editor/BehaviorTreeEditor/Classes/BehaviorTreeGraphNode_SimpleParallel.h b/Engine/Source/Editor/BehaviorTreeEditor/Classes/BehaviorTreeGraphNode_SimpleParallel.h
--- a/Engine/Source/Editor/BehaviorTreeEditor/Classes/BehaviorTreeGraphNode_SimpleParallel.h
+++ b/Engine/Source/Editor/BehaviorTreeEditor/Classes/BehaviorTreeGraphNode_SimpleParallel.h
@@ -6,9 +6,7 @@
 UCLASS()
 class UBehaviorTreeGraphNode_SimpleParallel : public UBehaviorTreeGraphNode_Composite
 {
-	GENERATED_BODY()
-public:
-	UBehaviorTreeGraphNode_SimpleParallel(const FObjectInitializer& ObjectInitializer = FObjectInitializer::Get());
+	GENERATED_UCLASS_BODY()
 
 	virtual void AllocateDefaultPins() override;
 	virtual void GetPinHoverText(const UEdGraphPin& Pin, FString& HoverTextOut) const override;
diff --git a/Engine/Source/Editor/BehaviorTreeEditor/Classes/BehaviorTreeGraphNode_SubtreeTask.h b/Engine/Source/Editor/BehaviorTreeEditor/Classes/BehaviorTreeGraphNode_SubtreeTask.h
--- a/Engine/Source/Editor/BehaviorTreeEditor/Classes/BehaviorTreeGraphNode_SubtreeTask.h
+++ b/Engine/Source/Editor/BehaviorTreeEditor/Classes/BehaviorTreeGraphNode_SubtreeTask.h
@@ -6,9 +6,7 @@
 UCLASS()
 class UBehaviorTreeGraphNode_SubtreeTask : public UBehaviorTreeGraphNode_Task
 {
-	GENERATED_BODY()
-public:
-	UBehaviorTreeGraphNode_SubtreeTask(const FObjectInitializer& ObjectInitializer = FObjectInitializer::Get());
+	GENERATED_UCLASS_BODY()
 
 	/** UBehaviorTreeGraph.UpdateCounter value of subtree graph */
 	int32 SubtreeVersion;
diff --git a/Engine/Source/Editor/BehaviorTreeEditor/Classes/BehaviorTreeGraphNode_Task.h b/Engine/Source/Editor/BehaviorTreeEditor/Classes/BehaviorTreeGraphNode_Task.h
--- a/Engine/Source/Editor/BehaviorTreeEditor/Classes/BehaviorTreeGraphNode_Task.h
+++ b/Engine/Source/Editor/BehaviorTreeEditor/Classes/BehaviorTreeGraphNode_Task.h
@@ -7,9 +7,7 @@
 UCLASS()
 class UBehaviorTreeGraphNode_Task : public UBehaviorTreeGraphNode
 {
-	GENERATED_BODY()
-public:
-	UBehaviorTreeGraphNode_Task(const FObjectInitializer& ObjectInitializer = FObjectInitializer::Get());
+	GENERATED_UCLASS_BODY()
 
 	virtual void AllocateDefaultPins() override;
 	virtual FText GetNodeTitle(ENodeTitleType::Type TitleType) const override;
diff --git a/Engine/Source/Editor/BehaviorTreeEditor/Classes/BlackboardDataFactory.h b/Engine/Source/Editor/BehaviorTreeEditor/Classes/BlackboardDataFactory.h
--- a/Engine/Source/Editor/BehaviorTreeEditor/Classes/BlackboardDataFactory.h
+++ b/Engine/Source/Editor/BehaviorTreeEditor/Classes/BlackboardDataFactory.h
@@ -6,9 +6,7 @@
 UCLASS()
 class UBlackboardDataFactory : public UFactory
 {
-	GENERATED_BODY()
-public:
-	UBlackboardDataFactory(const FObjectInitializer& ObjectInitializer = FObjectInitializer::Get());
+	GENERATED_UCLASS_BODY()
 
 	// UFactory interface
 	virtual UObject* FactoryCreateNew(UClass* Class,UObject* InParent,FName Name,EObjectFlags Flags,UObject* Context,FFeedbackContext* Warn) override;
diff --git a/Engine/Source/Editor/BehaviorTreeEditor/Classes/EdGraphSchema_BehaviorTree.h b/Engine/Source/Editor/BehaviorTreeEditor/Classes/EdGraphSchema_BehaviorTree.h
--- a/Engine/Source/Editor/BehaviorTreeEditor/Classes/EdGraphSchema_BehaviorTree.h
+++ b/Engine/Source/Editor/BehaviorTreeEditor/Classes/EdGraphSchema_BehaviorTree.h
@@ -9,7 +9,7 @@
 USTRUCT()
 struct FBehaviorTreeSchemaAction_AutoArrange : public FEdGraphSchemaAction
 {
-	GENERATED_BODY()
+	GENERATED_USTRUCT_BODY();
 
 	FBehaviorTreeSchemaAction_AutoArrange() 
 		: FEdGraphSchemaAction() {}
@@ -26,9 +26,7 @@ struct FBehaviorTreeSchemaAction_AutoArrange : public FEdGraphSchemaAction
 UCLASS(MinimalAPI)
 class UEdGraphSchema_BehaviorTree : public UAIGraphSchema
 {
-	GENERATED_BODY()
-public:
-	BEHAVIORTREEEDITOR_API UEdGraphSchema_BehaviorTree(const FObjectInitializer& ObjectInitializer = FObjectInitializer::Get());
+	GENERATED_UCLASS_BODY()
 
 	// Begin EdGraphSchema interface
 	virtual void CreateDefaultNodesForGraph(UEdGraph& Graph) const override;
diff --git a/Engine/Source/Editor/BehaviorTreeEditor/Classes/EdGraphSchema_BehaviorTreeDecorator.h b/Engine/Source/Editor/BehaviorTreeEditor/Classes/EdGraphSchema_BehaviorTreeDecorator.h
--- a/Engine/Source/Editor/BehaviorTreeEditor/Classes/EdGraphSchema_BehaviorTreeDecorator.h
+++ b/Engine/Source/Editor/BehaviorTreeEditor/Classes/EdGraphSchema_BehaviorTreeDecorator.h
@@ -7,7 +7,7 @@
 USTRUCT()
 struct FDecoratorSchemaAction_NewNode : public FEdGraphSchemaAction
 {
-	GENERATED_BODY()
+	GENERATED_USTRUCT_BODY();
 
 	/** Template of node we want to create */
 	UPROPERTY()
@@ -42,9 +42,7 @@ struct FDecoratorSchemaAction_NewNode : public FEdGraphSchemaAction
 UCLASS(MinimalAPI)
 class UEdGraphSchema_BehaviorTreeDecorator : public UEdGraphSchema
 {
-	GENERATED_BODY()
-public:
-	BEHAVIORTREEEDITOR_API UEdGraphSchema_BehaviorTreeDecorator(const FObjectInitializer& ObjectInitializer = FObjectInitializer::Get());
+	GENERATED_UCLASS_BODY()
 
 	UPROPERTY()
 	FString PC_Boolean;
diff --git a/Engine/Source/Editor/BlueprintGraph/Classes/BlueprintBoundEventNodeSpawner.h b/Engine/Source/Editor/BlueprintGraph/Classes/BlueprintBoundEventNodeSpawner.h
--- a/Engine/Source/Editor/BlueprintGraph/Classes/BlueprintBoundEventNodeSpawner.h
+++ b/Engine/Source/Editor/BlueprintGraph/Classes/BlueprintBoundEventNodeSpawner.h
@@ -17,9 +17,7 @@ class UK2Node_Event;
 UCLASS(Transient)
 class BLUEPRINTGRAPH_API UBlueprintBoundEventNodeSpawner : public UBlueprintEventNodeSpawner
 {
-	GENERATED_BODY()
-public:
-	UBlueprintBoundEventNodeSpawner(const FObjectInitializer& ObjectInitializer = FObjectInitializer::Get());
+	GENERATED_UCLASS_BODY()
 
 public:
 	/**
diff --git a/Engine/Source/Editor/BlueprintGraph/Classes/BlueprintBoundNodeSpawner.h b/Engine/Source/Editor/BlueprintGraph/Classes/BlueprintBoundNodeSpawner.h
--- a/Engine/Source/Editor/BlueprintGraph/Classes/BlueprintBoundNodeSpawner.h
+++ b/Engine/Source/Editor/BlueprintGraph/Classes/BlueprintBoundNodeSpawner.h
@@ -12,9 +12,7 @@
 UCLASS(Transient)
 class BLUEPRINTGRAPH_API UBlueprintBoundNodeSpawner : public UBlueprintNodeSpawner
 {
-	GENERATED_BODY()
-public:
-	UBlueprintBoundNodeSpawner(const FObjectInitializer& ObjectInitializer = FObjectInitializer::Get());
+	GENERATED_UCLASS_BODY()
 	DECLARE_DELEGATE_RetVal_OneParam(bool, FCanBindObjectDelegate, UObject const*);
 	DECLARE_DELEGATE_RetVal_TwoParams(bool, FOnBindObjectDelegate, UEdGraphNode*, UObject*);
 	DECLARE_DELEGATE_RetVal_TwoParams(UEdGraphNode*, FFindPreExistingNodeDelegate, const UBlueprint*, IBlueprintNodeBinder::FBindingSet const& );
diff --git a/Engine/Source/Editor/BlueprintGraph/Classes/BlueprintComponentNodeSpawner.h b/Engine/Source/Editor/BlueprintGraph/Classes/BlueprintComponentNodeSpawner.h
--- a/Engine/Source/Editor/BlueprintGraph/Classes/BlueprintComponentNodeSpawner.h
+++ b/Engine/Source/Editor/BlueprintGraph/Classes/BlueprintComponentNodeSpawner.h
@@ -17,9 +17,7 @@ class UActorComponent;
 UCLASS(Transient)
 class BLUEPRINTGRAPH_API UBlueprintComponentNodeSpawner : public UBlueprintNodeSpawner
 {
-	GENERATED_BODY()
-public:
-	UBlueprintComponentNodeSpawner(const FObjectInitializer& ObjectInitializer = FObjectInitializer::Get());
+	GENERATED_UCLASS_BODY()
 
 public:
 	/**
diff --git a/Engine/Source/Editor/BlueprintGraph/Classes/BlueprintDelegateNodeSpawner.h b/Engine/Source/Editor/BlueprintGraph/Classes/BlueprintDelegateNodeSpawner.h
--- a/Engine/Source/Editor/BlueprintGraph/Classes/BlueprintDelegateNodeSpawner.h
+++ b/Engine/Source/Editor/BlueprintGraph/Classes/BlueprintDelegateNodeSpawner.h
@@ -16,9 +16,7 @@ class UK2Node_BaseMCDelegate;
 UCLASS(Transient)
 class BLUEPRINTGRAPH_API UBlueprintDelegateNodeSpawner : public UBlueprintFieldNodeSpawner
 {
-	GENERATED_BODY()
-public:
-	UBlueprintDelegateNodeSpawner(const FObjectInitializer& ObjectInitializer = FObjectInitializer::Get());
+	GENERATED_UCLASS_BODY()
 
 public:
 	/**
diff --git a/Engine/Source/Editor/BlueprintGraph/Classes/BlueprintEventNodeSpawner.h b/Engine/Source/Editor/BlueprintGraph/Classes/BlueprintEventNodeSpawner.h
--- a/Engine/Source/Editor/BlueprintGraph/Classes/BlueprintEventNodeSpawner.h
+++ b/Engine/Source/Editor/BlueprintGraph/Classes/BlueprintEventNodeSpawner.h
@@ -18,9 +18,7 @@ class UK2Node_Event;
 UCLASS(Transient)
 class BLUEPRINTGRAPH_API UBlueprintEventNodeSpawner : public UBlueprintNodeSpawner
 {
-	GENERATED_BODY()
-public:
-	UBlueprintEventNodeSpawner(const FObjectInitializer& ObjectInitializer = FObjectInitializer::Get());
+	GENERATED_UCLASS_BODY()
 
 public:
 	/**
diff --git a/Engine/Source/Editor/BlueprintGraph/Classes/BlueprintFieldNodeSpawner.h b/Engine/Source/Editor/BlueprintGraph/Classes/BlueprintFieldNodeSpawner.h
--- a/Engine/Source/Editor/BlueprintGraph/Classes/BlueprintFieldNodeSpawner.h
+++ b/Engine/Source/Editor/BlueprintGraph/Classes/BlueprintFieldNodeSpawner.h
@@ -16,9 +16,7 @@ class UK2Node;
 UCLASS(Transient)
 class BLUEPRINTGRAPH_API UBlueprintFieldNodeSpawner : public UBlueprintNodeSpawner
 {
-	GENERATED_BODY()
-public:
-	UBlueprintFieldNodeSpawner(const FObjectInitializer& ObjectInitializer = FObjectInitializer::Get());
+	GENERATED_UCLASS_BODY()
 	DECLARE_DELEGATE_TwoParams(FSetNodeFieldDelegate, UEdGraphNode*, UField const*);
 
 public:
diff --git a/Engine/Source/Editor/BlueprintGraph/Classes/BlueprintFunctionNodeSpawner.h b/Engine/Source/Editor/BlueprintGraph/Classes/BlueprintFunctionNodeSpawner.h
--- a/Engine/Source/Editor/BlueprintGraph/Classes/BlueprintFunctionNodeSpawner.h
+++ b/Engine/Source/Editor/BlueprintGraph/Classes/BlueprintFunctionNodeSpawner.h
@@ -15,9 +15,7 @@ class UK2Node_CallFunction;
 UCLASS(Transient)
 class BLUEPRINTGRAPH_API UBlueprintFunctionNodeSpawner : public UBlueprintFieldNodeSpawner
 {
-	GENERATED_BODY()
-public:
-	UBlueprintFunctionNodeSpawner(const FObjectInitializer& ObjectInitializer = FObjectInitializer::Get());
+	GENERATED_UCLASS_BODY()
 
 public:
 	/**
diff --git a/Engine/Source/Editor/BlueprintGraph/Classes/EdGraphSchema_K2.h b/Engine/Source/Editor/BlueprintGraph/Classes/EdGraphSchema_K2.h
--- a/Engine/Source/Editor/BlueprintGraph/Classes/EdGraphSchema_K2.h
+++ b/Engine/Source/Editor/BlueprintGraph/Classes/EdGraphSchema_K2.h
@@ -8,7 +8,7 @@
 USTRUCT()
 struct BLUEPRINTGRAPH_API FEdGraphSchemaAction_K2Struct : public FEdGraphSchemaAction
 {
-	GENERATED_BODY()
+	GENERATED_USTRUCT_BODY()
 
 	// Simple type info
 	static FName StaticGetTypeId() {static FName Type("FEdGraphSchemaAction_K2Struct"); return Type;}
@@ -214,7 +214,7 @@ USTRUCT()
 // into their never versions (with class param instead of blueprint).
 struct FBlueprintCallableFunctionRedirect
 {
-	GENERATED_BODY()
+	GENERATED_USTRUCT_BODY()
 
 	UPROPERTY()
 	FString ClassName;
@@ -235,9 +235,7 @@ struct FBlueprintCallableFunctionRedirect
 UCLASS(config=Editor)
 class BLUEPRINTGRAPH_API UEdGraphSchema_K2 : public UEdGraphSchema
 {
-	GENERATED_BODY()
-public:
-	UEdGraphSchema_K2(const FObjectInitializer& ObjectInitializer = FObjectInitializer::Get());
+	GENERATED_UCLASS_BODY()
 
 	// Allowable PinType.PinCategory values
 	static const FString PC_Exec;
diff --git a/Engine/Source/Editor/BlueprintGraph/Classes/K2Node.h b/Engine/Source/Editor/BlueprintGraph/Classes/K2Node.h
--- a/Engine/Source/Editor/BlueprintGraph/Classes/K2Node.h
+++ b/Engine/Source/Editor/BlueprintGraph/Classes/K2Node.h
@@ -13,7 +13,7 @@ class UDynamicBlueprintBinding;
 USTRUCT()
 struct FOptionalPinFromProperty
 {
-	GENERATED_BODY()
+	GENERATED_USTRUCT_BODY()
 
 	UPROPERTY(EditAnywhere, Category=Hi, BlueprintReadOnly)
 	FName PropertyName;
@@ -87,9 +87,7 @@ enum ERenamePinResult
 UCLASS(abstract, MinimalAPI)
 class UK2Node : public UEdGraphNode
 {
-	GENERATED_BODY()
-public:
-	BLUEPRINTGRAPH_API UK2Node(const FObjectInitializer& ObjectInitializer = FObjectInitializer::Get());
+	GENERATED_UCLASS_BODY()
 
 	// UObject interface
 	BLUEPRINTGRAPH_API virtual void PostLoad() override;
diff --git a/Engine/Source/Editor/BlueprintGraph/Classes/K2Node_ActorBoundEvent.h b/Engine/Source/Editor/BlueprintGraph/Classes/K2Node_ActorBoundEvent.h
--- a/Engine/Source/Editor/BlueprintGraph/Classes/K2Node_ActorBoundEvent.h
+++ b/Engine/Source/Editor/BlueprintGraph/Classes/K2Node_ActorBoundEvent.h
@@ -9,9 +9,7 @@
 UCLASS(MinimalAPI)
 class UK2Node_ActorBoundEvent : public UK2Node_Event
 {
-	GENERATED_BODY()
-public:
-	BLUEPRINTGRAPH_API UK2Node_ActorBoundEvent(const FObjectInitializer& ObjectInitializer = FObjectInitializer::Get());
+	GENERATED_UCLASS_BODY()
 
 	/** Delegate property name that this event is associated with */
 	UPROPERTY()
diff --git a/Engine/Source/Editor/BlueprintGraph/Classes/K2Node_AddComponent.h b/Engine/Source/Editor/BlueprintGraph/Classes/K2Node_AddComponent.h
--- a/Engine/Source/Editor/BlueprintGraph/Classes/K2Node_AddComponent.h
+++ b/Engine/Source/Editor/BlueprintGraph/Classes/K2Node_AddComponent.h
@@ -9,9 +9,7 @@
 UCLASS(MinimalAPI)
 class UK2Node_AddComponent : public UK2Node_CallFunction
 {
-	GENERATED_BODY()
-public:
-	BLUEPRINTGRAPH_API UK2Node_AddComponent(const FObjectInitializer& ObjectInitializer = FObjectInitializer::Get());
+	GENERATED_UCLASS_BODY()
 
 	UPROPERTY()
 	uint32 bHasExposedVariable:1;
diff --git a/Engine/Source/Editor/BlueprintGraph/Classes/K2Node_AddDelegate.h b/Engine/Source/Editor/BlueprintGraph/Classes/K2Node_AddDelegate.h
--- a/Engine/Source/Editor/BlueprintGraph/Classes/K2Node_AddDelegate.h
+++ b/Engine/Source/Editor/BlueprintGraph/Classes/K2Node_AddDelegate.h
@@ -6,9 +6,7 @@
 UCLASS(MinimalAPI)
 class UK2Node_AddDelegate : public UK2Node_BaseMCDelegate
 {
-	GENERATED_BODY()
-public:
-	BLUEPRINTGRAPH_API UK2Node_AddDelegate(const FObjectInitializer& ObjectInitializer = FObjectInitializer::Get());
+	GENERATED_UCLASS_BODY()
 
 public:
 	// UEdGraphNode interface
diff --git a/Engine/Source/Editor/BlueprintGraph/Classes/K2Node_AssignDelegate.h b/Engine/Source/Editor/BlueprintGraph/Classes/K2Node_AssignDelegate.h
--- a/Engine/Source/Editor/BlueprintGraph/Classes/K2Node_AssignDelegate.h
+++ b/Engine/Source/Editor/BlueprintGraph/Classes/K2Node_AssignDelegate.h
@@ -13,9 +13,7 @@
 UCLASS(MinimalAPI)
 class UK2Node_AssignDelegate : public UK2Node_AddDelegate
 {
-	GENERATED_BODY()
-public:
-	BLUEPRINTGRAPH_API UK2Node_AssignDelegate(const FObjectInitializer& ObjectInitializer = FObjectInitializer::Get());
+	GENERATED_UCLASS_BODY()
 
 public:
 	// UEdGraphNode interface
diff --git a/Engine/Source/Editor/BlueprintGraph/Classes/K2Node_AssignmentStatement.h b/Engine/Source/Editor/BlueprintGraph/Classes/K2Node_AssignmentStatement.h
--- a/Engine/Source/Editor/BlueprintGraph/Classes/K2Node_AssignmentStatement.h
+++ b/Engine/Source/Editor/BlueprintGraph/Classes/K2Node_AssignmentStatement.h
@@ -7,9 +7,7 @@
 UCLASS(MinimalAPI)
 class UK2Node_AssignmentStatement : public UK2Node
 {
-	GENERATED_BODY()
-public:
-	BLUEPRINTGRAPH_API UK2Node_AssignmentStatement(const FObjectInitializer& ObjectInitializer = FObjectInitializer::Get());
+	GENERATED_UCLASS_BODY()
 
 
 	// Name of the Variable pin for this node
diff --git a/Engine/Source/Editor/BlueprintGraph/Classes/K2Node_BaseAsyncTask.h b/Engine/Source/Editor/BlueprintGraph/Classes/K2Node_BaseAsyncTask.h
--- a/Engine/Source/Editor/BlueprintGraph/Classes/K2Node_BaseAsyncTask.h
+++ b/Engine/Source/Editor/BlueprintGraph/Classes/K2Node_BaseAsyncTask.h
@@ -13,9 +13,7 @@ class UK2Node_TemporaryVariable;
 UCLASS(Abstract)
 class BLUEPRINTGRAPH_API UK2Node_BaseAsyncTask : public UK2Node
 {
-	GENERATED_BODY()
-public:
-	UK2Node_BaseAsyncTask(const FObjectInitializer& ObjectInitializer = FObjectInitializer::Get());
+	GENERATED_UCLASS_BODY()
 
 	// UEdGraphNode interface
 	virtual void AllocateDefaultPins() override;
diff --git a/Engine/Source/Editor/BlueprintGraph/Classes/K2Node_BaseMCDelegate.h b/Engine/Source/Editor/BlueprintGraph/Classes/K2Node_BaseMCDelegate.h
--- a/Engine/Source/Editor/BlueprintGraph/Classes/K2Node_BaseMCDelegate.h
+++ b/Engine/Source/Editor/BlueprintGraph/Classes/K2Node_BaseMCDelegate.h
@@ -8,9 +8,7 @@
 UCLASS(MinimalAPI, abstract)
 class UK2Node_BaseMCDelegate : public UK2Node
 {
-	GENERATED_BODY()
-public:
-	BLUEPRINTGRAPH_API UK2Node_BaseMCDelegate(const FObjectInitializer& ObjectInitializer = FObjectInitializer::Get());
+	GENERATED_UCLASS_BODY()
 
 	/** Reference to delegate */
 	UPROPERTY()
diff --git a/Engine/Source/Editor/BlueprintGraph/Classes/K2Node_CallArrayFunction.h b/Engine/Source/Editor/BlueprintGraph/Classes/K2Node_CallArrayFunction.h
--- a/Engine/Source/Editor/BlueprintGraph/Classes/K2Node_CallArrayFunction.h
+++ b/Engine/Source/Editor/BlueprintGraph/Classes/K2Node_CallArrayFunction.h
@@ -19,9 +19,7 @@ struct FArrayPropertyPinCombo
 UCLASS(MinimalAPI)
 class UK2Node_CallArrayFunction : public UK2Node_CallFunction
 {
-	GENERATED_BODY()
-public:
-	BLUEPRINTGRAPH_API UK2Node_CallArrayFunction(const FObjectInitializer& ObjectInitializer = FObjectInitializer::Get());
+	GENERATED_UCLASS_BODY()
 
 	// Begin UEdGraphNode interface
 	virtual void AllocateDefaultPins() override;
diff --git a/Engine/Source/Editor/BlueprintGraph/Classes/K2Node_CallDelegate.h b/Engine/Source/Editor/BlueprintGraph/Classes/K2Node_CallDelegate.h
--- a/Engine/Source/Editor/BlueprintGraph/Classes/K2Node_CallDelegate.h
+++ b/Engine/Source/Editor/BlueprintGraph/Classes/K2Node_CallDelegate.h
@@ -6,9 +6,7 @@
 UCLASS(MinimalAPI)
 class UK2Node_CallDelegate : public UK2Node_BaseMCDelegate
 {
-	GENERATED_BODY()
-public:
-	BLUEPRINTGRAPH_API UK2Node_CallDelegate(const FObjectInitializer& ObjectInitializer = FObjectInitializer::Get());
+	GENERATED_UCLASS_BODY()
 
 public:
 
@@ -24,4 +22,4 @@ public:
 
 	virtual void ValidateNodeDuringCompilation(class FCompilerResultsLog& MessageLog) const override;
 	bool CreatePinsForFunctionInputs(const UFunction* Function);
-};
+};
\ No newline at end of file
diff --git a/Engine/Source/Editor/BlueprintGraph/Classes/K2Node_CallFunction.h b/Engine/Source/Editor/BlueprintGraph/Classes/K2Node_CallFunction.h
--- a/Engine/Source/Editor/BlueprintGraph/Classes/K2Node_CallFunction.h
+++ b/Engine/Source/Editor/BlueprintGraph/Classes/K2Node_CallFunction.h
@@ -10,9 +10,7 @@
 UCLASS()
 class BLUEPRINTGRAPH_API UK2Node_CallFunction : public UK2Node
 {
-	GENERATED_BODY()
-public:
-	UK2Node_CallFunction(const FObjectInitializer& ObjectInitializer = FObjectInitializer::Get());
+	GENERATED_UCLASS_BODY()
 
 	/** Indicates that this is a call to a pure function */
 	UPROPERTY()
diff --git a/Engine/Source/Editor/BlueprintGraph/Classes/K2Node_CallFunctionOnMember.h b/Engine/Source/Editor/BlueprintGraph/Classes/K2Node_CallFunctionOnMember.h
--- a/Engine/Source/Editor/BlueprintGraph/Classes/K2Node_CallFunctionOnMember.h
+++ b/Engine/Source/Editor/BlueprintGraph/Classes/K2Node_CallFunctionOnMember.h
@@ -9,9 +9,7 @@
 UCLASS(MinimalAPI)
 class UK2Node_CallFunctionOnMember : public UK2Node_CallFunction
 {
-	GENERATED_BODY()
-public:
-	BLUEPRINTGRAPH_API UK2Node_CallFunctionOnMember(const FObjectInitializer& ObjectInitializer = FObjectInitializer::Get());
+	GENERATED_UCLASS_BODY()
 
 	/** Reference to member variable to call function on */
 	UPROPERTY()
diff --git a/Engine/Source/Editor/BlueprintGraph/Classes/K2Node_CallMaterialParameterCollectionFunction.h b/Engine/Source/Editor/BlueprintGraph/Classes/K2Node_CallMaterialParameterCollectionFunction.h
--- a/Engine/Source/Editor/BlueprintGraph/Classes/K2Node_CallMaterialParameterCollectionFunction.h
+++ b/Engine/Source/Editor/BlueprintGraph/Classes/K2Node_CallMaterialParameterCollectionFunction.h
@@ -7,11 +7,9 @@
 UCLASS(MinimalAPI)
 class UK2Node_CallMaterialParameterCollectionFunction : public UK2Node_CallFunction
 {
-	GENERATED_BODY()
-public:
-	BLUEPRINTGRAPH_API UK2Node_CallMaterialParameterCollectionFunction(const FObjectInitializer& ObjectInitializer = FObjectInitializer::Get());
+	GENERATED_UCLASS_BODY()
 
 	// Begin EdGraphNode interface
 	virtual void PinDefaultValueChanged(UEdGraphPin* Pin) override;
 	// End EdGraphNode interface
-};
+};
\ No newline at end of file
diff --git a/Engine/Source/Editor/BlueprintGraph/Classes/K2Node_CallParentFunction.h b/Engine/Source/Editor/BlueprintGraph/Classes/K2Node_CallParentFunction.h
--- a/Engine/Source/Editor/BlueprintGraph/Classes/K2Node_CallParentFunction.h
+++ b/Engine/Source/Editor/BlueprintGraph/Classes/K2Node_CallParentFunction.h
@@ -8,9 +8,7 @@
 UCLASS(MinimalAPI)
 class UK2Node_CallParentFunction : public UK2Node_CallFunction
 {
-	GENERATED_BODY()
-public:
-	BLUEPRINTGRAPH_API UK2Node_CallParentFunction(const FObjectInitializer& ObjectInitializer = FObjectInitializer::Get());
+	GENERATED_UCLASS_BODY()
 
 	// Begin EdGraphNode interface
 	virtual FText GetNodeTitle(ENodeTitleType::Type TitleType) const override;
diff --git a/Engine/Source/Editor/BlueprintGraph/Classes/K2Node_CastByteToEnum.h b/Engine/Source/Editor/BlueprintGraph/Classes/K2Node_CastByteToEnum.h
--- a/Engine/Source/Editor/BlueprintGraph/Classes/K2Node_CastByteToEnum.h
+++ b/Engine/Source/Editor/BlueprintGraph/Classes/K2Node_CastByteToEnum.h
@@ -8,9 +8,7 @@
 UCLASS(MinimalAPI)
 class UK2Node_CastByteToEnum : public UK2Node
 {
-	GENERATED_BODY()
-public:
-	BLUEPRINTGRAPH_API UK2Node_CastByteToEnum(const FObjectInitializer& ObjectInitializer = FObjectInitializer::Get());
+	GENERATED_UCLASS_BODY()
 
 	UPROPERTY()
 	UEnum* Enum;
diff --git a/Engine/Source/Editor/BlueprintGraph/Classes/K2Node_ClassDynamicCast.h b/Engine/Source/Editor/BlueprintGraph/Classes/K2Node_ClassDynamicCast.h
--- a/Engine/Source/Editor/BlueprintGraph/Classes/K2Node_ClassDynamicCast.h
+++ b/Engine/Source/Editor/BlueprintGraph/Classes/K2Node_ClassDynamicCast.h
@@ -8,9 +8,7 @@
 UCLASS(MinimalAPI)
 class UK2Node_ClassDynamicCast : public UK2Node_DynamicCast
 {
-	GENERATED_BODY()
-public:
-	BLUEPRINTGRAPH_API UK2Node_ClassDynamicCast(const FObjectInitializer& ObjectInitializer = FObjectInitializer::Get());
+	GENERATED_UCLASS_BODY()
 
 	// Begin UEdGraphNode interface
 	virtual void AllocateDefaultPins() override;
diff --git a/Engine/Source/Editor/BlueprintGraph/Classes/K2Node_ClearDelegate.h b/Engine/Source/Editor/BlueprintGraph/Classes/K2Node_ClearDelegate.h
--- a/Engine/Source/Editor/BlueprintGraph/Classes/K2Node_ClearDelegate.h
+++ b/Engine/Source/Editor/BlueprintGraph/Classes/K2Node_ClearDelegate.h
@@ -7,9 +7,7 @@
 UCLASS(MinimalAPI)
 class UK2Node_ClearDelegate : public UK2Node_BaseMCDelegate
 {
-	GENERATED_BODY()
-public:
-	BLUEPRINTGRAPH_API UK2Node_ClearDelegate(const FObjectInitializer& ObjectInitializer = FObjectInitializer::Get());
+	GENERATED_UCLASS_BODY()
 
 public:
 	// UEdGraphNode interface
diff --git a/Engine/Source/Editor/BlueprintGraph/Classes/K2Node_CommutativeAssociativeBinaryOperator.h b/Engine/Source/Editor/BlueprintGraph/Classes/K2Node_CommutativeAssociativeBinaryOperator.h
--- a/Engine/Source/Editor/BlueprintGraph/Classes/K2Node_CommutativeAssociativeBinaryOperator.h
+++ b/Engine/Source/Editor/BlueprintGraph/Classes/K2Node_CommutativeAssociativeBinaryOperator.h
@@ -7,9 +7,7 @@
 UCLASS(MinimalAPI)
 class UK2Node_CommutativeAssociativeBinaryOperator : public UK2Node_CallFunction
 {
-	GENERATED_BODY()
-public:
-	BLUEPRINTGRAPH_API UK2Node_CommutativeAssociativeBinaryOperator(const FObjectInitializer& ObjectInitializer = FObjectInitializer::Get());
+	GENERATED_UCLASS_BODY()
 
 	/** The number of additional input pins to generate for this node (2 base pins are not included) */
 	UPROPERTY()
@@ -44,4 +42,4 @@ public:
 	virtual void GetContextMenuActions(const FGraphNodeContextMenuBuilder& Context) const override;
 	virtual void ExpandNode(class FKismetCompilerContext& CompilerContext, UEdGraph* SourceGraph) override;
 	// End of UK2Node interface
-};
+};
\ No newline at end of file
diff --git a/Engine/Source/Editor/BlueprintGraph/Classes/K2Node_ComponentBoundEvent.h b/Engine/Source/Editor/BlueprintGraph/Classes/K2Node_ComponentBoundEvent.h
--- a/Engine/Source/Editor/BlueprintGraph/Classes/K2Node_ComponentBoundEvent.h
+++ b/Engine/Source/Editor/BlueprintGraph/Classes/K2Node_ComponentBoundEvent.h
@@ -11,9 +11,7 @@ class UDynamicBlueprintBinding;
 UCLASS(MinimalAPI)
 class UK2Node_ComponentBoundEvent : public UK2Node_Event
 {
-	GENERATED_BODY()
-public:
-	BLUEPRINTGRAPH_API UK2Node_ComponentBoundEvent(const FObjectInitializer& ObjectInitializer = FObjectInitializer::Get());
+	GENERATED_UCLASS_BODY()
 
 	/** Delegate property name that this event is associated with */
 	UPROPERTY()
diff --git a/Engine/Source/Editor/BlueprintGraph/Classes/K2Node_Composite.h b/Engine/Source/Editor/BlueprintGraph/Classes/K2Node_Composite.h
--- a/Engine/Source/Editor/BlueprintGraph/Classes/K2Node_Composite.h
+++ b/Engine/Source/Editor/BlueprintGraph/Classes/K2Node_Composite.h
@@ -7,9 +7,7 @@
 UCLASS(MinimalAPI)
 class UK2Node_Composite : public UK2Node_Tunnel
 {
-	GENERATED_BODY()
-public:
-	BLUEPRINTGRAPH_API UK2Node_Composite(const FObjectInitializer& ObjectInitializer = FObjectInitializer::Get());
+	GENERATED_UCLASS_BODY()
 
 	// The graph that this composite node is representing
 	UPROPERTY()
diff --git a/Engine/Source/Editor/BlueprintGraph/Classes/K2Node_ConstructObjectFromClass.h b/Engine/Source/Editor/BlueprintGraph/Classes/K2Node_ConstructObjectFromClass.h
--- a/Engine/Source/Editor/BlueprintGraph/Classes/K2Node_ConstructObjectFromClass.h
+++ b/Engine/Source/Editor/BlueprintGraph/Classes/K2Node_ConstructObjectFromClass.h
@@ -9,9 +9,7 @@
 UCLASS(abstract)
 class BLUEPRINTGRAPH_API UK2Node_ConstructObjectFromClass : public UK2Node
 {
-	GENERATED_BODY()
-public:
-	UK2Node_ConstructObjectFromClass(const FObjectInitializer& ObjectInitializer = FObjectInitializer::Get());
+	GENERATED_UCLASS_BODY()
 
 	// Begin UEdGraphNode interface.
 	virtual void AllocateDefaultPins() override;
diff --git a/Engine/Source/Editor/BlueprintGraph/Classes/K2Node_CreateDelegate.h b/Engine/Source/Editor/BlueprintGraph/Classes/K2Node_CreateDelegate.h
--- a/Engine/Source/Editor/BlueprintGraph/Classes/K2Node_CreateDelegate.h
+++ b/Engine/Source/Editor/BlueprintGraph/Classes/K2Node_CreateDelegate.h
@@ -6,9 +6,7 @@
 UCLASS(MinimalAPI)
 class UK2Node_CreateDelegate : public UK2Node
 {
-	GENERATED_BODY()
-public:
-	BLUEPRINTGRAPH_API UK2Node_CreateDelegate(const FObjectInitializer& ObjectInitializer = FObjectInitializer::Get());
+	GENERATED_UCLASS_BODY()
 
 	UPROPERTY()
 	FName SelectedFunctionName;
diff --git a/Engine/Source/Editor/BlueprintGraph/Classes/K2Node_CustomEvent.h b/Engine/Source/Editor/BlueprintGraph/Classes/K2Node_CustomEvent.h
--- a/Engine/Source/Editor/BlueprintGraph/Classes/K2Node_CustomEvent.h
+++ b/Engine/Source/Editor/BlueprintGraph/Classes/K2Node_CustomEvent.h
@@ -9,9 +9,7 @@
 UCLASS(MinimalAPI)
 class UK2Node_CustomEvent : public UK2Node_Event
 {
-	GENERATED_BODY()
-public:
-	BLUEPRINTGRAPH_API UK2Node_CustomEvent(const FObjectInitializer& ObjectInitializer = FObjectInitializer::Get());
+	GENERATED_UCLASS_BODY()
 
 	/** Specifies that the event can be triggered in Editor */
 	UPROPERTY()
diff --git a/Engine/Source/Editor/BlueprintGraph/Classes/K2Node_DeadClass.h b/Engine/Source/Editor/BlueprintGraph/Classes/K2Node_DeadClass.h
--- a/Engine/Source/Editor/BlueprintGraph/Classes/K2Node_DeadClass.h
+++ b/Engine/Source/Editor/BlueprintGraph/Classes/K2Node_DeadClass.h
@@ -7,7 +7,5 @@
 UCLASS(MinimalAPI)
 class UK2Node_DeadClass : public UK2Node
 {
-	GENERATED_BODY()
-public:
-	BLUEPRINTGRAPH_API UK2Node_DeadClass(const FObjectInitializer& ObjectInitializer = FObjectInitializer::Get());
+	GENERATED_UCLASS_BODY()
 };
diff --git a/Engine/Source/Editor/BlueprintGraph/Classes/K2Node_DelegateSet.h b/Engine/Source/Editor/BlueprintGraph/Classes/K2Node_DelegateSet.h
--- a/Engine/Source/Editor/BlueprintGraph/Classes/K2Node_DelegateSet.h
+++ b/Engine/Source/Editor/BlueprintGraph/Classes/K2Node_DelegateSet.h
@@ -9,9 +9,7 @@
 UCLASS(MinimalAPI)
 class UK2Node_DelegateSet : public UK2Node
 {
-	GENERATED_BODY()
-public:
-	BLUEPRINTGRAPH_API UK2Node_DelegateSet(const FObjectInitializer& ObjectInitializer = FObjectInitializer::Get());
+	GENERATED_UCLASS_BODY()
 
 	/** Delegate property name that this event is associated with on the target */
 	UPROPERTY()
diff --git a/Engine/Source/Editor/BlueprintGraph/Classes/K2Node_DoOnceMultiInput.h b/Engine/Source/Editor/BlueprintGraph/Classes/K2Node_DoOnceMultiInput.h
--- a/Engine/Source/Editor/BlueprintGraph/Classes/K2Node_DoOnceMultiInput.h
+++ b/Engine/Source/Editor/BlueprintGraph/Classes/K2Node_DoOnceMultiInput.h
@@ -7,9 +7,7 @@
 UCLASS(MinimalAPI)
 class UK2Node_DoOnceMultiInput : public UK2Node
 {
-	GENERATED_BODY()
-public:
-	BLUEPRINTGRAPH_API UK2Node_DoOnceMultiInput(const FObjectInitializer& ObjectInitializer = FObjectInitializer::Get());
+	GENERATED_UCLASS_BODY()
 
 	/** The number of additional input pins to generate for this node (2 base pins are not included) */
 	UPROPERTY()
diff --git a/Engine/Source/Editor/BlueprintGraph/Classes/K2Node_DynamicCast.h b/Engine/Source/Editor/BlueprintGraph/Classes/K2Node_DynamicCast.h
--- a/Engine/Source/Editor/BlueprintGraph/Classes/K2Node_DynamicCast.h
+++ b/Engine/Source/Editor/BlueprintGraph/Classes/K2Node_DynamicCast.h
@@ -9,9 +9,7 @@
 UCLASS(MinimalAPI)
 class UK2Node_DynamicCast : public UK2Node
 {
-	GENERATED_BODY()
-public:
-	BLUEPRINTGRAPH_API UK2Node_DynamicCast(const FObjectInitializer& ObjectInitializer = FObjectInitializer::Get());
+	GENERATED_UCLASS_BODY()
 
 	/** The type that the input should try to be cast to */
 	UPROPERTY()
diff --git a/Engine/Source/Editor/BlueprintGraph/Classes/K2Node_EaseFunction.h b/Engine/Source/Editor/BlueprintGraph/Classes/K2Node_EaseFunction.h
--- a/Engine/Source/Editor/BlueprintGraph/Classes/K2Node_EaseFunction.h
+++ b/Engine/Source/Editor/BlueprintGraph/Classes/K2Node_EaseFunction.h
@@ -10,7 +10,7 @@
 USTRUCT()
 struct FCustomPin
 {
-	GENERATED_BODY()
+	GENERATED_USTRUCT_BODY()
 
 	/** Name of the pin within the current node */
 	UPROPERTY()
@@ -28,9 +28,7 @@ struct FCustomPin
 UCLASS(MinimalAPI)
 class UK2Node_EaseFunction : public UK2Node
 {
-	GENERATED_BODY()
-public:
-	BLUEPRINTGRAPH_API UK2Node_EaseFunction(const FObjectInitializer& ObjectInitializer = FObjectInitializer::Get());
+	GENERATED_UCLASS_BODY()
 
 	// Begin UEdGraphNode interface.
 	virtual void AllocateDefaultPins() override;
diff --git a/Engine/Source/Editor/BlueprintGraph/Classes/K2Node_EditablePinBase.h b/Engine/Source/Editor/BlueprintGraph/Classes/K2Node_EditablePinBase.h
--- a/Engine/Source/Editor/BlueprintGraph/Classes/K2Node_EditablePinBase.h
+++ b/Engine/Source/Editor/BlueprintGraph/Classes/K2Node_EditablePinBase.h
@@ -7,7 +7,7 @@
 USTRUCT()
 struct FUserPinInfo
 {
-	GENERATED_BODY()
+	GENERATED_USTRUCT_BODY()
 
 	/** The name of the pin, as defined by the user */
 	UPROPERTY()
@@ -48,7 +48,7 @@ struct FUserPinInfo
 USTRUCT()
 struct FKismetUserDeclaredFunctionMetadata
 {
-	GENERATED_BODY()
+	GENERATED_USTRUCT_BODY()
 
 	UPROPERTY()
 	FString ToolTip;
@@ -77,9 +77,7 @@ public:
 UCLASS(abstract, MinimalAPI)
 class UK2Node_EditablePinBase : public UK2Node
 {
-	GENERATED_BODY()
-public:
-	BLUEPRINTGRAPH_API UK2Node_EditablePinBase(const FObjectInitializer& ObjectInitializer = FObjectInitializer::Get());
+	GENERATED_UCLASS_BODY()
 
 	/** Whether or not this entry node should be user-editable with the function editor */
 	UPROPERTY()
diff --git a/Engine/Source/Editor/BlueprintGraph/Classes/K2Node_EnumEquality.h b/Engine/Source/Editor/BlueprintGraph/Classes/K2Node_EnumEquality.h
--- a/Engine/Source/Editor/BlueprintGraph/Classes/K2Node_EnumEquality.h
+++ b/Engine/Source/Editor/BlueprintGraph/Classes/K2Node_EnumEquality.h
@@ -8,9 +8,7 @@
 UCLASS(MinimalAPI)
 class UK2Node_EnumEquality : public UK2Node
 {
-	GENERATED_BODY()
-public:
-	BLUEPRINTGRAPH_API UK2Node_EnumEquality(const FObjectInitializer& ObjectInitializer = FObjectInitializer::Get());
+	GENERATED_UCLASS_BODY()
 
 	// Begin UEdGraphNode interface
 	virtual void AllocateDefaultPins() override;
diff --git a/Engine/Source/Editor/BlueprintGraph/Classes/K2Node_EnumLiteral.h b/Engine/Source/Editor/BlueprintGraph/Classes/K2Node_EnumLiteral.h
--- a/Engine/Source/Editor/BlueprintGraph/Classes/K2Node_EnumLiteral.h
+++ b/Engine/Source/Editor/BlueprintGraph/Classes/K2Node_EnumLiteral.h
@@ -9,9 +9,7 @@
 UCLASS(MinimalAPI)
 class UK2Node_EnumLiteral : public UK2Node, public INodeDependingOnEnumInterface
 {
-	GENERATED_BODY()
-public:
-	BLUEPRINTGRAPH_API UK2Node_EnumLiteral(const FObjectInitializer& ObjectInitializer = FObjectInitializer::Get());
+	GENERATED_UCLASS_BODY()
 
 	UPROPERTY()
 	UEnum* Enum;
diff --git a/Engine/Source/Editor/BlueprintGraph/Classes/K2Node_Event.h b/Engine/Source/Editor/BlueprintGraph/Classes/K2Node_Event.h
--- a/Engine/Source/Editor/BlueprintGraph/Classes/K2Node_Event.h
+++ b/Engine/Source/Editor/BlueprintGraph/Classes/K2Node_Event.h
@@ -11,9 +11,7 @@
 UCLASS(MinimalAPI)
 class UK2Node_Event : public UK2Node_EditablePinBase
 {
-	GENERATED_BODY()
-public:
-	BLUEPRINTGRAPH_API UK2Node_Event(const FObjectInitializer& ObjectInitializer = FObjectInitializer::Get());
+	GENERATED_UCLASS_BODY()
 	BLUEPRINTGRAPH_API static const FString DelegateOutputName;
 
 	/** Name of function signature that this event implements */
diff --git a/Engine/Source/Editor/BlueprintGraph/Classes/K2Node_ExecutionSequence.h b/Engine/Source/Editor/BlueprintGraph/Classes/K2Node_ExecutionSequence.h
--- a/Engine/Source/Editor/BlueprintGraph/Classes/K2Node_ExecutionSequence.h
+++ b/Engine/Source/Editor/BlueprintGraph/Classes/K2Node_ExecutionSequence.h
@@ -8,9 +8,7 @@
 UCLASS(MinimalAPI)
 class UK2Node_ExecutionSequence : public UK2Node
 {
-	GENERATED_BODY()
-public:
-	BLUEPRINTGRAPH_API UK2Node_ExecutionSequence(const FObjectInitializer& ObjectInitializer = FObjectInitializer::Get());
+	GENERATED_UCLASS_BODY()
 
 	// Begin UEdGraphNode interface
 	virtual void AllocateDefaultPins() override;
diff --git a/Engine/Source/Editor/BlueprintGraph/Classes/K2Node_ForEachElementInEnum.h b/Engine/Source/Editor/BlueprintGraph/Classes/K2Node_ForEachElementInEnum.h
--- a/Engine/Source/Editor/BlueprintGraph/Classes/K2Node_ForEachElementInEnum.h
+++ b/Engine/Source/Editor/BlueprintGraph/Classes/K2Node_ForEachElementInEnum.h
@@ -8,9 +8,7 @@
 UCLASS(MinimalAPI)
 class UK2Node_ForEachElementInEnum : public UK2Node
 {
-	GENERATED_BODY()
-public:
-	BLUEPRINTGRAPH_API UK2Node_ForEachElementInEnum(const FObjectInitializer& ObjectInitializer = FObjectInitializer::Get());
+	GENERATED_UCLASS_BODY()
 
 	UPROPERTY()
 	UEnum* Enum;
diff --git a/Engine/Source/Editor/BlueprintGraph/Classes/K2Node_FormatText.h b/Engine/Source/Editor/BlueprintGraph/Classes/K2Node_FormatText.h
--- a/Engine/Source/Editor/BlueprintGraph/Classes/K2Node_FormatText.h
+++ b/Engine/Source/Editor/BlueprintGraph/Classes/K2Node_FormatText.h
@@ -1,4 +1,4 @@
-// Copyright 1998-2015 Epic Games, Inc. All Rights Reserved.
+﻿// Copyright 1998-2015 Epic Games, Inc. All Rights Reserved.
 
 #pragma once
 #include "K2Node_FormatText.generated.h"
@@ -6,9 +6,7 @@
 UCLASS(MinimalAPI)
 class UK2Node_FormatText : public UK2Node
 {
-	GENERATED_BODY()
-public:
-	BLUEPRINTGRAPH_API UK2Node_FormatText(const FObjectInitializer& ObjectInitializer = FObjectInitializer::Get());
+	GENERATED_UCLASS_BODY()
 
 	// UObject interface
 	virtual void PostLoad() override;
diff --git a/Engine/Source/Editor/BlueprintGraph/Classes/K2Node_FunctionEntry.h b/Engine/Source/Editor/BlueprintGraph/Classes/K2Node_FunctionEntry.h
--- a/Engine/Source/Editor/BlueprintGraph/Classes/K2Node_FunctionEntry.h
+++ b/Engine/Source/Editor/BlueprintGraph/Classes/K2Node_FunctionEntry.h
@@ -7,9 +7,7 @@
 UCLASS(MinimalAPI)
 class UK2Node_FunctionEntry : public UK2Node_FunctionTerminator
 {
-	GENERATED_BODY()
-public:
-	BLUEPRINTGRAPH_API UK2Node_FunctionEntry(const FObjectInitializer& ObjectInitializer = FObjectInitializer::Get());
+	GENERATED_UCLASS_BODY()
 
 	/** If specified, the function that is created for this entry point will have this name.  Otherwise, it will have the function signature's name */
 	UPROPERTY()
diff --git a/Engine/Source/Editor/BlueprintGraph/Classes/K2Node_FunctionResult.h b/Engine/Source/Editor/BlueprintGraph/Classes/K2Node_FunctionResult.h
--- a/Engine/Source/Editor/BlueprintGraph/Classes/K2Node_FunctionResult.h
+++ b/Engine/Source/Editor/BlueprintGraph/Classes/K2Node_FunctionResult.h
@@ -8,9 +8,7 @@
 UCLASS(MinimalAPI)
 class UK2Node_FunctionResult : public UK2Node_FunctionTerminator
 {
-	GENERATED_BODY()
-public:
-	BLUEPRINTGRAPH_API UK2Node_FunctionResult(const FObjectInitializer& ObjectInitializer = FObjectInitializer::Get());
+	GENERATED_UCLASS_BODY()
 
 	// Begin UEdGraphNode interface
 	virtual void AllocateDefaultPins() override;
diff --git a/Engine/Source/Editor/BlueprintGraph/Classes/K2Node_FunctionTerminator.h b/Engine/Source/Editor/BlueprintGraph/Classes/K2Node_FunctionTerminator.h
--- a/Engine/Source/Editor/BlueprintGraph/Classes/K2Node_FunctionTerminator.h
+++ b/Engine/Source/Editor/BlueprintGraph/Classes/K2Node_FunctionTerminator.h
@@ -8,9 +8,7 @@
 UCLASS(abstract, MinimalAPI)
 class UK2Node_FunctionTerminator : public UK2Node_EditablePinBase
 {
-	GENERATED_BODY()
-public:
-	BLUEPRINTGRAPH_API UK2Node_FunctionTerminator(const FObjectInitializer& ObjectInitializer = FObjectInitializer::Get());
+	GENERATED_UCLASS_BODY()
 
 	/** 
 	 * The source class that defines the signature, if it is getting that from elsewhere (e.g. interface, base class etc). 
diff --git a/Engine/Source/Editor/BlueprintGraph/Classes/K2Node_GetDataTableRow.h b/Engine/Source/Editor/BlueprintGraph/Classes/K2Node_GetDataTableRow.h
--- a/Engine/Source/Editor/BlueprintGraph/Classes/K2Node_GetDataTableRow.h
+++ b/Engine/Source/Editor/BlueprintGraph/Classes/K2Node_GetDataTableRow.h
@@ -10,9 +10,7 @@
 UCLASS()
 class BLUEPRINTGRAPH_API UK2Node_GetDataTableRow : public UK2Node
 {
-	GENERATED_BODY()
-public:
-	UK2Node_GetDataTableRow(const FObjectInitializer& ObjectInitializer = FObjectInitializer::Get());
+	GENERATED_UCLASS_BODY()
 
 	// Begin UEdGraphNode interface.
 	virtual void AllocateDefaultPins() override;
diff --git a/Engine/Source/Editor/BlueprintGraph/Classes/K2Node_GetEnumeratorName.h b/Engine/Source/Editor/BlueprintGraph/Classes/K2Node_GetEnumeratorName.h
--- a/Engine/Source/Editor/BlueprintGraph/Classes/K2Node_GetEnumeratorName.h
+++ b/Engine/Source/Editor/BlueprintGraph/Classes/K2Node_GetEnumeratorName.h
@@ -7,9 +7,7 @@
 UCLASS(MinimalAPI)
 class UK2Node_GetEnumeratorName : public UK2Node
 {
-	GENERATED_BODY()
-public:
-	BLUEPRINTGRAPH_API UK2Node_GetEnumeratorName(const FObjectInitializer& ObjectInitializer = FObjectInitializer::Get());
+	GENERATED_UCLASS_BODY()
 
 	static FString EnumeratorPinName;
 
diff --git a/Engine/Source/Editor/BlueprintGraph/Classes/K2Node_GetEnumeratorNameAsString.h b/Engine/Source/Editor/BlueprintGraph/Classes/K2Node_GetEnumeratorNameAsString.h
--- a/Engine/Source/Editor/BlueprintGraph/Classes/K2Node_GetEnumeratorNameAsString.h
+++ b/Engine/Source/Editor/BlueprintGraph/Classes/K2Node_GetEnumeratorNameAsString.h
@@ -7,9 +7,7 @@
 UCLASS(MinimalAPI)
 class UK2Node_GetEnumeratorNameAsString : public UK2Node_GetEnumeratorName
 {
-	GENERATED_BODY()
-public:
-	BLUEPRINTGRAPH_API UK2Node_GetEnumeratorNameAsString(const FObjectInitializer& ObjectInitializer = FObjectInitializer::Get());
+	GENERATED_UCLASS_BODY()
 
 	// Begin UEdGraphNode interface
 	virtual void AllocateDefaultPins() override;
diff --git a/Engine/Source/Editor/BlueprintGraph/Classes/K2Node_GetInputAxisKeyValue.h b/Engine/Source/Editor/BlueprintGraph/Classes/K2Node_GetInputAxisKeyValue.h
--- a/Engine/Source/Editor/BlueprintGraph/Classes/K2Node_GetInputAxisKeyValue.h
+++ b/Engine/Source/Editor/BlueprintGraph/Classes/K2Node_GetInputAxisKeyValue.h
@@ -8,9 +8,7 @@
 UCLASS(MinimalAPI)
 class UK2Node_GetInputAxisKeyValue : public UK2Node_CallFunction
 {
-	GENERATED_BODY()
-public:
-	BLUEPRINTGRAPH_API UK2Node_GetInputAxisKeyValue(const FObjectInitializer& ObjectInitializer = FObjectInitializer::Get());
+	GENERATED_UCLASS_BODY()
 
 	UPROPERTY()
 	FKey InputAxisKey;
@@ -48,4 +46,4 @@ private:
 	/** Constructing FText strings can be costly, so we cache the node's title/tooltip */
 	FNodeTextCache CachedTooltip;
 	FNodeTextCache CachedNodeTitle;
-};
+};
\ No newline at end of file
diff --git a/Engine/Source/Editor/BlueprintGraph/Classes/K2Node_GetInputAxisValue.h b/Engine/Source/Editor/BlueprintGraph/Classes/K2Node_GetInputAxisValue.h
--- a/Engine/Source/Editor/BlueprintGraph/Classes/K2Node_GetInputAxisValue.h
+++ b/Engine/Source/Editor/BlueprintGraph/Classes/K2Node_GetInputAxisValue.h
@@ -8,9 +8,7 @@
 UCLASS(MinimalAPI)
 class UK2Node_GetInputAxisValue : public UK2Node_CallFunction
 {
-	GENERATED_BODY()
-public:
-	BLUEPRINTGRAPH_API UK2Node_GetInputAxisValue(const FObjectInitializer& ObjectInitializer = FObjectInitializer::Get());
+	GENERATED_UCLASS_BODY()
 
 	UPROPERTY()
 	FName InputAxisName;
diff --git a/Engine/Source/Editor/BlueprintGraph/Classes/K2Node_GetInputVectorAxisValue.h b/Engine/Source/Editor/BlueprintGraph/Classes/K2Node_GetInputVectorAxisValue.h
--- a/Engine/Source/Editor/BlueprintGraph/Classes/K2Node_GetInputVectorAxisValue.h
+++ b/Engine/Source/Editor/BlueprintGraph/Classes/K2Node_GetInputVectorAxisValue.h
@@ -7,9 +7,7 @@
 UCLASS(MinimalAPI)
 class UK2Node_GetInputVectorAxisValue : public UK2Node_GetInputAxisKeyValue
 {
-	GENERATED_BODY()
-public:
-	BLUEPRINTGRAPH_API UK2Node_GetInputVectorAxisValue(const FObjectInitializer& ObjectInitializer = FObjectInitializer::Get());
+	GENERATED_UCLASS_BODY()
 
 	// Begin EdGraphNode interface
 	virtual FText GetTooltipText() const override;
diff --git a/Engine/Source/Editor/BlueprintGraph/Classes/K2Node_GetNumEnumEntries.h b/Engine/Source/Editor/BlueprintGraph/Classes/K2Node_GetNumEnumEntries.h
--- a/Engine/Source/Editor/BlueprintGraph/Classes/K2Node_GetNumEnumEntries.h
+++ b/Engine/Source/Editor/BlueprintGraph/Classes/K2Node_GetNumEnumEntries.h
@@ -8,9 +8,7 @@
 UCLASS(MinimalAPI)
 class UK2Node_GetNumEnumEntries : public UK2Node, public INodeDependingOnEnumInterface
 {
-	GENERATED_BODY()
-public:
-	BLUEPRINTGRAPH_API UK2Node_GetNumEnumEntries(const FObjectInitializer& ObjectInitializer = FObjectInitializer::Get());
+	GENERATED_UCLASS_BODY()
 
 	UPROPERTY()
 	UEnum* Enum;
diff --git a/Engine/Source/Editor/BlueprintGraph/Classes/K2Node_IfThenElse.h b/Engine/Source/Editor/BlueprintGraph/Classes/K2Node_IfThenElse.h
--- a/Engine/Source/Editor/BlueprintGraph/Classes/K2Node_IfThenElse.h
+++ b/Engine/Source/Editor/BlueprintGraph/Classes/K2Node_IfThenElse.h
@@ -8,9 +8,7 @@
 UCLASS(MinimalAPI)
 class UK2Node_IfThenElse : public UK2Node
 {
-	GENERATED_BODY()
-public:
-	BLUEPRINTGRAPH_API UK2Node_IfThenElse(const FObjectInitializer& ObjectInitializer = FObjectInitializer::Get());
+	GENERATED_UCLASS_BODY()
 
 	// Begin UEdGraphNode interface
 	virtual void AllocateDefaultPins() override;
diff --git a/Engine/Source/Editor/BlueprintGraph/Classes/K2Node_InputActionEvent.h b/Engine/Source/Editor/BlueprintGraph/Classes/K2Node_InputActionEvent.h
--- a/Engine/Source/Editor/BlueprintGraph/Classes/K2Node_InputActionEvent.h
+++ b/Engine/Source/Editor/BlueprintGraph/Classes/K2Node_InputActionEvent.h
@@ -8,9 +8,7 @@
 UCLASS(MinimalAPI)
 class UK2Node_InputActionEvent : public UK2Node_Event
 {
-	GENERATED_BODY()
-public:
-	BLUEPRINTGRAPH_API UK2Node_InputActionEvent(const FObjectInitializer& ObjectInitializer = FObjectInitializer::Get());
+	GENERATED_UCLASS_BODY()
 
 	UPROPERTY()
 	FName InputActionName;
@@ -31,4 +29,4 @@ public:
 	virtual UClass* GetDynamicBindingClass() const override;
 	virtual void RegisterDynamicBinding(UDynamicBlueprintBinding* BindingObject) const override;
 	// End UK2Node interface
-};
+};
\ No newline at end of file
diff --git a/Engine/Source/Editor/BlueprintGraph/Classes/K2Node_InputAxisEvent.h b/Engine/Source/Editor/BlueprintGraph/Classes/K2Node_InputAxisEvent.h
--- a/Engine/Source/Editor/BlueprintGraph/Classes/K2Node_InputAxisEvent.h
+++ b/Engine/Source/Editor/BlueprintGraph/Classes/K2Node_InputAxisEvent.h
@@ -8,9 +8,7 @@
 UCLASS(MinimalAPI)
 class UK2Node_InputAxisEvent : public UK2Node_Event
 {
-	GENERATED_BODY()
-public:
-	BLUEPRINTGRAPH_API UK2Node_InputAxisEvent(const FObjectInitializer& ObjectInitializer = FObjectInitializer::Get());
+	GENERATED_UCLASS_BODY()
 
 	UPROPERTY()
 	FName InputAxisName;
diff --git a/Engine/Source/Editor/BlueprintGraph/Classes/K2Node_InputKey.h b/Engine/Source/Editor/BlueprintGraph/Classes/K2Node_InputKey.h
--- a/Engine/Source/Editor/BlueprintGraph/Classes/K2Node_InputKey.h
+++ b/Engine/Source/Editor/BlueprintGraph/Classes/K2Node_InputKey.h
@@ -9,9 +9,7 @@
 UCLASS(MinimalAPI)
 class UK2Node_InputKey : public UK2Node
 {
-	GENERATED_BODY()
-public:
-	BLUEPRINTGRAPH_API UK2Node_InputKey(const FObjectInitializer& ObjectInitializer = FObjectInitializer::Get());
+	GENERATED_UCLASS_BODY()
 
 	// The key that is bound
 	UPROPERTY(EditAnywhere, Category="Input")
diff --git a/Engine/Source/Editor/BlueprintGraph/Classes/K2Node_InputKeyEvent.h b/Engine/Source/Editor/BlueprintGraph/Classes/K2Node_InputKeyEvent.h
--- a/Engine/Source/Editor/BlueprintGraph/Classes/K2Node_InputKeyEvent.h
+++ b/Engine/Source/Editor/BlueprintGraph/Classes/K2Node_InputKeyEvent.h
@@ -9,9 +9,7 @@
 UCLASS(MinimalAPI)
 class UK2Node_InputKeyEvent : public UK2Node_Event
 {
-	GENERATED_BODY()
-public:
-	BLUEPRINTGRAPH_API UK2Node_InputKeyEvent(const FObjectInitializer& ObjectInitializer = FObjectInitializer::Get());
+	GENERATED_UCLASS_BODY()
 
 	UPROPERTY()
 	FInputChord InputChord;
@@ -32,4 +30,4 @@ public:
 	virtual UClass* GetDynamicBindingClass() const override;
 	virtual void RegisterDynamicBinding(UDynamicBlueprintBinding* BindingObject) const override;
 	// End UK2Node interface
-};
+};
\ No newline at end of file
diff --git a/Engine/Source/Editor/BlueprintGraph/Classes/K2Node_InputTouchEvent.h b/Engine/Source/Editor/BlueprintGraph/Classes/K2Node_InputTouchEvent.h
--- a/Engine/Source/Editor/BlueprintGraph/Classes/K2Node_InputTouchEvent.h
+++ b/Engine/Source/Editor/BlueprintGraph/Classes/K2Node_InputTouchEvent.h
@@ -8,9 +8,7 @@
 UCLASS(MinimalAPI)
 class UK2Node_InputTouchEvent : public UK2Node_Event
 {
-	GENERATED_BODY()
-public:
-	BLUEPRINTGRAPH_API UK2Node_InputTouchEvent(const FObjectInitializer& ObjectInitializer = FObjectInitializer::Get());
+	GENERATED_UCLASS_BODY()
 
 	UPROPERTY()
 	TEnumAsByte<EInputEvent> InputKeyEvent;
@@ -28,4 +26,4 @@ public:
 	virtual UClass* GetDynamicBindingClass() const override;
 	virtual void RegisterDynamicBinding(UDynamicBlueprintBinding* BindingObject) const override;
 	// End UK2Node interface
-};
+};
\ No newline at end of file
diff --git a/Engine/Source/Editor/BlueprintGraph/Classes/K2Node_InputVectorAxisEvent.h b/Engine/Source/Editor/BlueprintGraph/Classes/K2Node_InputVectorAxisEvent.h
--- a/Engine/Source/Editor/BlueprintGraph/Classes/K2Node_InputVectorAxisEvent.h
+++ b/Engine/Source/Editor/BlueprintGraph/Classes/K2Node_InputVectorAxisEvent.h
@@ -7,9 +7,7 @@
 UCLASS(MinimalAPI)
 class UK2Node_InputVectorAxisEvent : public UK2Node_InputAxisKeyEvent
 {
-	GENERATED_BODY()
-public:
-	BLUEPRINTGRAPH_API UK2Node_InputVectorAxisEvent(const FObjectInitializer& ObjectInitializer = FObjectInitializer::Get());
+	GENERATED_UCLASS_BODY()
 
 	// UObject interface
 	virtual void Serialize(FArchive& Ar) override;
@@ -21,4 +19,4 @@ public:
 	virtual void RegisterDynamicBinding(UDynamicBlueprintBinding* BindingObject) const override;
 	virtual void GetMenuActions(FBlueprintActionDatabaseRegistrar& ActionRegistrar) const override;
 	// End UK2Node interface
-};
+};
\ No newline at end of file
diff --git a/Engine/Source/Editor/BlueprintGraph/Classes/K2Node_Knot.h b/Engine/Source/Editor/BlueprintGraph/Classes/K2Node_Knot.h
--- a/Engine/Source/Editor/BlueprintGraph/Classes/K2Node_Knot.h
+++ b/Engine/Source/Editor/BlueprintGraph/Classes/K2Node_Knot.h
@@ -7,9 +7,7 @@
 UCLASS(MinimalAPI)
 class UK2Node_Knot : public UK2Node
 {
-	GENERATED_BODY()
-public:
-	BLUEPRINTGRAPH_API UK2Node_Knot(const FObjectInitializer& ObjectInitializer = FObjectInitializer::Get());
+	GENERATED_UCLASS_BODY()
 
 public:
 	// UEdGraphNode interface
diff --git a/Engine/Source/Editor/BlueprintGraph/Classes/K2Node_Literal.h b/Engine/Source/Editor/BlueprintGraph/Classes/K2Node_Literal.h
--- a/Engine/Source/Editor/BlueprintGraph/Classes/K2Node_Literal.h
+++ b/Engine/Source/Editor/BlueprintGraph/Classes/K2Node_Literal.h
@@ -8,9 +8,7 @@
 UCLASS(MinimalAPI)
 class UK2Node_Literal : public UK2Node
 {
-	GENERATED_BODY()
-public:
-	BLUEPRINTGRAPH_API UK2Node_Literal(const FObjectInitializer& ObjectInitializer = FObjectInitializer::Get());
+	GENERATED_UCLASS_BODY()
 
 private:
 	/** If this is an object reference literal, keep a reference here so that it can be updated as objects move around */
diff --git a/Engine/Source/Editor/BlueprintGraph/Classes/K2Node_LocalVariable.h b/Engine/Source/Editor/BlueprintGraph/Classes/K2Node_LocalVariable.h
--- a/Engine/Source/Editor/BlueprintGraph/Classes/K2Node_LocalVariable.h
+++ b/Engine/Source/Editor/BlueprintGraph/Classes/K2Node_LocalVariable.h
@@ -7,9 +7,7 @@
 UCLASS(MinimalAPI, deprecated)
 class UDEPRECATED_K2Node_LocalVariable : public UK2Node_TemporaryVariable
 {
-	GENERATED_BODY()
-public:
-	BLUEPRINTGRAPH_API UDEPRECATED_K2Node_LocalVariable(const FObjectInitializer& ObjectInitializer = FObjectInitializer::Get());
+	GENERATED_UCLASS_BODY()
 
 	/** If this is not an override, allow user to specify a name for the function created by this entry point */
 	UPROPERTY()
@@ -42,4 +40,4 @@ public:
 	 * @param InVariableType		The type to assign this local variable to
 	 */
 	BLUEPRINTGRAPH_API void ChangeVariableType(const FEdGraphPinType& InVariableType);
-};
+};
\ No newline at end of file
diff --git a/Engine/Source/Editor/BlueprintGraph/Classes/K2Node_MacroInstance.h b/Engine/Source/Editor/BlueprintGraph/Classes/K2Node_MacroInstance.h
--- a/Engine/Source/Editor/BlueprintGraph/Classes/K2Node_MacroInstance.h
+++ b/Engine/Source/Editor/BlueprintGraph/Classes/K2Node_MacroInstance.h
@@ -10,9 +10,7 @@
 UCLASS(MinimalAPI)
 class UK2Node_MacroInstance : public UK2Node_Tunnel
 {
-	GENERATED_BODY()
-public:
-	BLUEPRINTGRAPH_API UK2Node_MacroInstance(const FObjectInitializer& ObjectInitializer = FObjectInitializer::Get());
+	GENERATED_UCLASS_BODY()
 
 private:
 	/** A macro is like a composite node, except that the associated graph lives
diff --git a/Engine/Source/Editor/BlueprintGraph/Classes/K2Node_MakeArray.h b/Engine/Source/Editor/BlueprintGraph/Classes/K2Node_MakeArray.h
--- a/Engine/Source/Editor/BlueprintGraph/Classes/K2Node_MakeArray.h
+++ b/Engine/Source/Editor/BlueprintGraph/Classes/K2Node_MakeArray.h
@@ -8,9 +8,7 @@
 UCLASS(MinimalAPI)
 class UK2Node_MakeArray : public UK2Node
 {
-	GENERATED_BODY()
-public:
-	BLUEPRINTGRAPH_API UK2Node_MakeArray(const FObjectInitializer& ObjectInitializer = FObjectInitializer::Get());
+	GENERATED_UCLASS_BODY()
 
 	/** The number of input pins to generate for this node */
 	UPROPERTY()
diff --git a/Engine/Source/Editor/BlueprintGraph/Classes/K2Node_MathExpression.h b/Engine/Source/Editor/BlueprintGraph/Classes/K2Node_MathExpression.h
--- a/Engine/Source/Editor/BlueprintGraph/Classes/K2Node_MathExpression.h
+++ b/Engine/Source/Editor/BlueprintGraph/Classes/K2Node_MathExpression.h
@@ -16,9 +16,7 @@
 UCLASS()
 class BLUEPRINTGRAPH_API UK2Node_MathExpression : public UK2Node_Composite
 {
-	GENERATED_BODY()
-public:
-	UK2Node_MathExpression(const FObjectInitializer& ObjectInitializer = FObjectInitializer::Get());
+	GENERATED_UCLASS_BODY()
 
 public:
 
diff --git a/Engine/Source/Editor/BlueprintGraph/Classes/K2Node_MatineeController.h b/Engine/Source/Editor/BlueprintGraph/Classes/K2Node_MatineeController.h
--- a/Engine/Source/Editor/BlueprintGraph/Classes/K2Node_MatineeController.h
+++ b/Engine/Source/Editor/BlueprintGraph/Classes/K2Node_MatineeController.h
@@ -8,9 +8,7 @@
 UCLASS(MinimalAPI)
 class UK2Node_MatineeController : public UK2Node
 {
-	GENERATED_BODY()
-public:
-	BLUEPRINTGRAPH_API UK2Node_MatineeController(const FObjectInitializer& ObjectInitializer = FObjectInitializer::Get());
+	GENERATED_UCLASS_BODY()
 
 	/** The matinee actor in the level that this node controls */
 	UPROPERTY(EditAnywhere, Category=K2Node_MatineeController)
diff --git a/Engine/Source/Editor/BlueprintGraph/Classes/K2Node_Message.h b/Engine/Source/Editor/BlueprintGraph/Classes/K2Node_Message.h
--- a/Engine/Source/Editor/BlueprintGraph/Classes/K2Node_Message.h
+++ b/Engine/Source/Editor/BlueprintGraph/Classes/K2Node_Message.h
@@ -9,9 +9,7 @@
 UCLASS(MinimalAPI)
 class UK2Node_Message : public UK2Node_CallFunction
 {
-	GENERATED_BODY()
-public:
-	BLUEPRINTGRAPH_API UK2Node_Message(const FObjectInitializer& ObjectInitializer = FObjectInitializer::Get());
+	GENERATED_UCLASS_BODY()
 
 
 	// Begin UEdGraphNode interface.
diff --git a/Engine/Source/Editor/BlueprintGraph/Classes/K2Node_MultiGate.h b/Engine/Source/Editor/BlueprintGraph/Classes/K2Node_MultiGate.h
--- a/Engine/Source/Editor/BlueprintGraph/Classes/K2Node_MultiGate.h
+++ b/Engine/Source/Editor/BlueprintGraph/Classes/K2Node_MultiGate.h
@@ -7,9 +7,7 @@
 UCLASS(MinimalAPI)
 class UK2Node_MultiGate : public UK2Node_ExecutionSequence
 {
-	GENERATED_BODY()
-public:
-	BLUEPRINTGRAPH_API UK2Node_MultiGate(const FObjectInitializer& ObjectInitializer = FObjectInitializer::Get());
+	GENERATED_UCLASS_BODY()
 
 	/** Reference to the integer that contains */
 	UPROPERTY(transient)
diff --git a/Engine/Source/Editor/BlueprintGraph/Classes/K2Node_PureAssignmentStatement.h b/Engine/Source/Editor/BlueprintGraph/Classes/K2Node_PureAssignmentStatement.h
--- a/Engine/Source/Editor/BlueprintGraph/Classes/K2Node_PureAssignmentStatement.h
+++ b/Engine/Source/Editor/BlueprintGraph/Classes/K2Node_PureAssignmentStatement.h
@@ -7,9 +7,7 @@
 UCLASS(MinimalAPI)
 class UK2Node_PureAssignmentStatement : public UK2Node
 {
-	GENERATED_BODY()
-public:
-	BLUEPRINTGRAPH_API UK2Node_PureAssignmentStatement(const FObjectInitializer& ObjectInitializer = FObjectInitializer::Get());
+	GENERATED_UCLASS_BODY()
 
 	// Name of the Variable pin for this node
 	static FString VariablePinName;
diff --git a/Engine/Source/Editor/BlueprintGraph/Classes/K2Node_Select.h b/Engine/Source/Editor/BlueprintGraph/Classes/K2Node_Select.h
--- a/Engine/Source/Editor/BlueprintGraph/Classes/K2Node_Select.h
+++ b/Engine/Source/Editor/BlueprintGraph/Classes/K2Node_Select.h
@@ -9,9 +9,7 @@
 UCLASS(MinimalAPI)
 class UK2Node_Select : public UK2Node, public INodeDependingOnEnumInterface
 {
-	GENERATED_BODY()
-public:
-	BLUEPRINTGRAPH_API UK2Node_Select(const FObjectInitializer& ObjectInitializer = FObjectInitializer::Get());
+	GENERATED_UCLASS_BODY()
 
 	/** The number of selectable options this node currently has */
 	UPROPERTY()
diff --git a/Engine/Source/Editor/BlueprintGraph/Classes/K2Node_Self.h b/Engine/Source/Editor/BlueprintGraph/Classes/K2Node_Self.h
--- a/Engine/Source/Editor/BlueprintGraph/Classes/K2Node_Self.h
+++ b/Engine/Source/Editor/BlueprintGraph/Classes/K2Node_Self.h
@@ -8,9 +8,7 @@
 UCLASS(MinimalAPI)
 class UK2Node_Self : public UK2Node
 {
-	GENERATED_BODY()
-public:
-	BLUEPRINTGRAPH_API UK2Node_Self(const FObjectInitializer& ObjectInitializer = FObjectInitializer::Get());
+	GENERATED_UCLASS_BODY()
 
 	// Begin UEdGraphNode interface
 	virtual FText GetNodeTitle(ENodeTitleType::Type TitleType) const override;
diff --git a/Engine/Source/Editor/BlueprintGraph/Classes/K2Node_SetFieldsInStruct.h b/Engine/Source/Editor/BlueprintGraph/Classes/K2Node_SetFieldsInStruct.h
--- a/Engine/Source/Editor/BlueprintGraph/Classes/K2Node_SetFieldsInStruct.h
+++ b/Engine/Source/Editor/BlueprintGraph/Classes/K2Node_SetFieldsInStruct.h
@@ -9,9 +9,7 @@
 UCLASS(MinimalAPI)
 class UK2Node_SetFieldsInStruct : public UK2Node_MakeStruct
 {
-	GENERATED_BODY()
-public:
-	BLUEPRINTGRAPH_API UK2Node_SetFieldsInStruct(const FObjectInitializer& ObjectInitializer = FObjectInitializer::Get());
+	GENERATED_UCLASS_BODY()
 	// Begin UEdGraphNode interface
 	virtual void AllocateDefaultPins() override;
 	virtual FText GetNodeTitle(ENodeTitleType::Type TitleType) const override;
diff --git a/Engine/Source/Editor/BlueprintGraph/Classes/K2Node_SpawnActor.h b/Engine/Source/Editor/BlueprintGraph/Classes/K2Node_SpawnActor.h
--- a/Engine/Source/Editor/BlueprintGraph/Classes/K2Node_SpawnActor.h
+++ b/Engine/Source/Editor/BlueprintGraph/Classes/K2Node_SpawnActor.h
@@ -9,9 +9,7 @@
 UCLASS(MinimalAPI)
 class UK2Node_SpawnActor : public UK2Node
 {
-	GENERATED_BODY()
-public:
-	BLUEPRINTGRAPH_API UK2Node_SpawnActor(const FObjectInitializer& ObjectInitializer = FObjectInitializer::Get());
+	GENERATED_UCLASS_BODY()
 
 	// Begin UEdGraphNode interface.
 	virtual void AllocateDefaultPins() override;
diff --git a/Engine/Source/Editor/BlueprintGraph/Classes/K2Node_SpawnActorFromClass.h b/Engine/Source/Editor/BlueprintGraph/Classes/K2Node_SpawnActorFromClass.h
--- a/Engine/Source/Editor/BlueprintGraph/Classes/K2Node_SpawnActorFromClass.h
+++ b/Engine/Source/Editor/BlueprintGraph/Classes/K2Node_SpawnActorFromClass.h
@@ -9,9 +9,7 @@
 UCLASS()
 class BLUEPRINTGRAPH_API UK2Node_SpawnActorFromClass : public UK2Node
 {
-	GENERATED_BODY()
-public:
-	UK2Node_SpawnActorFromClass(const FObjectInitializer& ObjectInitializer = FObjectInitializer::Get());
+	GENERATED_UCLASS_BODY()
 
 	// Begin UEdGraphNode interface.
 	virtual void AllocateDefaultPins() override;
diff --git a/Engine/Source/Editor/BlueprintGraph/Classes/K2Node_StructMemberGet.h b/Engine/Source/Editor/BlueprintGraph/Classes/K2Node_StructMemberGet.h
--- a/Engine/Source/Editor/BlueprintGraph/Classes/K2Node_StructMemberGet.h
+++ b/Engine/Source/Editor/BlueprintGraph/Classes/K2Node_StructMemberGet.h
@@ -9,9 +9,7 @@
 UCLASS(MinimalAPI)
 class UK2Node_StructMemberGet : public UK2Node_StructOperation
 {
-	GENERATED_BODY()
-public:
-	BLUEPRINTGRAPH_API UK2Node_StructMemberGet(const FObjectInitializer& ObjectInitializer = FObjectInitializer::Get());
+	GENERATED_UCLASS_BODY()
 
 	UPROPERTY(EditAnywhere, Category=PinOptions, EditFixedSize)
 	TArray<FOptionalPinFromProperty> ShowPinForProperties;
