Missing chages for 2481343.

Some P4V problem with backing out large changelist missed some files.

diff --git a/Engine/Plugins/2D/Paper2D/Source/Paper2D/Classes/PaperFlipbook.h b/Engine/Plugins/2D/Paper2D/Source/Paper2D/Classes/PaperFlipbook.h
--- a/Engine/Plugins/2D/Paper2D/Source/Paper2D/Classes/PaperFlipbook.h
+++ b/Engine/Plugins/2D/Paper2D/Source/Paper2D/Classes/PaperFlipbook.h
@@ -8,8 +8,8 @@
 USTRUCT()
 struct FPaperFlipbookKeyFrame
 {
-	GENERATED_BODY()
 public:
+	GENERATED_USTRUCT_BODY()
 
 	UPROPERTY(Category=Sprite, EditAnywhere)
 	UPaperSprite* Sprite;
@@ -46,9 +46,7 @@ namespace EFlipbookCollisionMode
 UCLASS(BlueprintType, meta = (DisplayThumbnail = "true"))
 class PAPER2D_API UPaperFlipbook : public UObject
 {
-	GENERATED_BODY()
-public:
-	UPaperFlipbook(const FObjectInitializer& ObjectInitializer = FObjectInitializer::Get());
+	GENERATED_UCLASS_BODY()
 
 protected:
 	// The nominal frame rate to play this flipbook animation back at
diff --git a/Engine/Plugins/2D/Paper2D/Source/Paper2D/Classes/PaperRuntimeSettings.h b/Engine/Plugins/2D/Paper2D/Source/Paper2D/Classes/PaperRuntimeSettings.h
--- a/Engine/Plugins/2D/Paper2D/Source/Paper2D/Classes/PaperRuntimeSettings.h
+++ b/Engine/Plugins/2D/Paper2D/Source/Paper2D/Classes/PaperRuntimeSettings.h
@@ -10,9 +10,7 @@
 UCLASS(config = Engine, defaultconfig)
 class PAPER2D_API UPaperRuntimeSettings : public UObject
 {
-	GENERATED_BODY()
-public:
-	UPaperRuntimeSettings(const FObjectInitializer& ObjectInitializer = FObjectInitializer::Get());
+	GENERATED_UCLASS_BODY()
 
 	// The default scaling factor between pixels and Unreal units (cm) to use for newly created sprite assets (e.g., 0.64 would make a 64 pixel wide sprite take up 100 cm)
 	UPROPERTY(GlobalConfig, EditAnywhere, Category=Settings)
diff --git a/Engine/Plugins/2D/Paper2D/Source/Paper2DEditor/Private/Atlasing/PaperSpriteAtlasFactory.h b/Engine/Plugins/2D/Paper2D/Source/Paper2DEditor/Private/Atlasing/PaperSpriteAtlasFactory.h
--- a/Engine/Plugins/2D/Paper2D/Source/Paper2DEditor/Private/Atlasing/PaperSpriteAtlasFactory.h
+++ b/Engine/Plugins/2D/Paper2D/Source/Paper2DEditor/Private/Atlasing/PaperSpriteAtlasFactory.h
@@ -7,9 +7,7 @@
 UCLASS()
 class UPaperSpriteAtlasFactory : public UFactory
 {
-	GENERATED_BODY()
-public:
-	UPaperSpriteAtlasFactory(const FObjectInitializer& ObjectInitializer = FObjectInitializer::Get());
+	GENERATED_UCLASS_BODY()
 
 	// UFactory interface
 	virtual UObject* FactoryCreateNew(UClass* Class, UObject* InParent, FName Name, EObjectFlags Flags, UObject* Context, FFeedbackContext* Warn) override;
diff --git a/Engine/Plugins/2D/Paper2D/Source/PaperJsonImporter/Private/ReimportPaperJsonImporterFactory.h b/Engine/Plugins/2D/Paper2D/Source/PaperJsonImporter/Private/ReimportPaperJsonImporterFactory.h
--- a/Engine/Plugins/2D/Paper2D/Source/PaperJsonImporter/Private/ReimportPaperJsonImporterFactory.h
+++ b/Engine/Plugins/2D/Paper2D/Source/PaperJsonImporter/Private/ReimportPaperJsonImporterFactory.h
@@ -7,9 +7,7 @@
 UCLASS()
 class UReimportPaperJsonImporterFactory : public UPaperJsonImporterFactory, public FReimportHandler
 {
-	GENERATED_BODY()
-public:
-	UReimportPaperJsonImporterFactory(const FObjectInitializer& ObjectInitializer = FObjectInitializer::Get());
+	GENERATED_UCLASS_BODY()
 
 	// Begin FReimportHandler interface
 	virtual bool CanReimport(UObject* Obj, TArray<FString>& OutFilenames) override;
diff --git a/Engine/Plugins/2D/Paper2D/Source/PaperTiledImporter/Classes/PaperTiledImporterFactory.h b/Engine/Plugins/2D/Paper2D/Source/PaperTiledImporter/Classes/PaperTiledImporterFactory.h
--- a/Engine/Plugins/2D/Paper2D/Source/PaperTiledImporter/Classes/PaperTiledImporterFactory.h
+++ b/Engine/Plugins/2D/Paper2D/Source/PaperTiledImporter/Classes/PaperTiledImporterFactory.h
@@ -22,9 +22,7 @@ namespace ETiledOrientation
 UCLASS()
 class UPaperTiledImporterFactory : public UFactory, public FReimportHandler
 {
-	GENERATED_BODY()
-public:
-	UPaperTiledImporterFactory(const FObjectInitializer& ObjectInitializer = FObjectInitializer::Get());
+	GENERATED_UCLASS_BODY()
 
 	// UFactory interface
 	virtual FText GetToolTip() const override;
diff --git a/Engine/Plugins/Experimental/Phya/Source/Phya/Classes/PhyaCollisionHandler.h b/Engine/Plugins/Experimental/Phya/Source/Phya/Classes/PhyaCollisionHandler.h
--- a/Engine/Plugins/Experimental/Phya/Source/Phya/Classes/PhyaCollisionHandler.h
+++ b/Engine/Plugins/Experimental/Phya/Source/Phya/Classes/PhyaCollisionHandler.h
@@ -28,9 +28,7 @@ class FPhyaPairInfo : public TSharedFromThis<FPhyaPairInfo>
 UCLASS()
 class UPhyaCollisionHandler : public UPhysicsCollisionHandler
 {
-	GENERATED_BODY()
-public:
-	UPhyaCollisionHandler(const FObjectInitializer& ObjectInitializer = FObjectInitializer::Get());
+	GENERATED_UCLASS_BODY()
 
 	// Begin UPhysicsCollisionHandler interface
 	virtual void InitCollisionHandler() override;
diff --git a/Engine/Plugins/Media/MediaPlayerEditor/Source/MediaPlayerEditor/Private/Factories/MediaPlayerFactoryNew.h b/Engine/Plugins/Media/MediaPlayerEditor/Source/MediaPlayerEditor/Private/Factories/MediaPlayerFactoryNew.h
--- a/Engine/Plugins/Media/MediaPlayerEditor/Source/MediaPlayerEditor/Private/Factories/MediaPlayerFactoryNew.h
+++ b/Engine/Plugins/Media/MediaPlayerEditor/Source/MediaPlayerEditor/Private/Factories/MediaPlayerFactoryNew.h
@@ -12,9 +12,7 @@ UCLASS(hidecategories=Object)
 class UMediaPlayerFactoryNew
 	: public UFactory
 {
-	GENERATED_BODY()
-public:
-	UMediaPlayerFactoryNew(const FObjectInitializer& ObjectInitializer = FObjectInitializer::Get());
+	GENERATED_UCLASS_BODY()
 
 public:
 
diff --git a/Engine/Plugins/Runtime/CustomMeshComponent/Source/CustomMeshComponent/Classes/CustomMeshComponent.h b/Engine/Plugins/Runtime/CustomMeshComponent/Source/CustomMeshComponent/Classes/CustomMeshComponent.h
--- a/Engine/Plugins/Runtime/CustomMeshComponent/Source/CustomMeshComponent/Classes/CustomMeshComponent.h
+++ b/Engine/Plugins/Runtime/CustomMeshComponent/Source/CustomMeshComponent/Classes/CustomMeshComponent.h
@@ -7,7 +7,7 @@
 USTRUCT(BlueprintType)
 struct FCustomMeshTriangle
 {
-	GENERATED_BODY()
+	GENERATED_USTRUCT_BODY()
 
 	UPROPERTY(EditAnywhere, BlueprintReadWrite, Category=Triangle)
 	FVector Vertex0;
@@ -23,9 +23,7 @@ struct FCustomMeshTriangle
 UCLASS(hidecategories=(Object,LOD, Physics, Collision), editinlinenew, meta=(BlueprintSpawnableComponent), ClassGroup=Rendering)
 class CUSTOMMESHCOMPONENT_API UCustomMeshComponent : public UMeshComponent
 {
-	GENERATED_BODY()
-public:
-	UCustomMeshComponent(const FObjectInitializer& ObjectInitializer = FObjectInitializer::Get());
+	GENERATED_UCLASS_BODY()
 
 	/** Set the geometry to use on this triangle mesh */
 	UFUNCTION(BlueprintCallable, Category="Components|CustomMesh")
diff --git a/Engine/Plugins/Runtime/LeapMotionController/Source/LeapMotionController/Public/LeapMotionControllerActor.h b/Engine/Plugins/Runtime/LeapMotionController/Source/LeapMotionController/Public/LeapMotionControllerActor.h
--- a/Engine/Plugins/Runtime/LeapMotionController/Source/LeapMotionController/Public/LeapMotionControllerActor.h
+++ b/Engine/Plugins/Runtime/LeapMotionController/Source/LeapMotionController/Public/LeapMotionControllerActor.h
@@ -15,9 +15,7 @@
 UCLASS(ClassGroup = LeapMotion, BlueprintType, Blueprintable)
 class LEAPMOTIONCONTROLLER_API ALeapMotionControllerActor : public AActor
 {
-	GENERATED_BODY()
-public:
-	ALeapMotionControllerActor(const FObjectInitializer& ObjectInitializer = FObjectInitializer::Get());
+	GENERATED_UCLASS_BODY()
 
 	/** Scale of hands, relatively to real-world size. This ignores actor's transform's scale */
 	UPROPERTY(VisibleAnywhere, BlueprintReadOnly, Category = LeapMotion)
diff --git a/Engine/Plugins/Runtime/SoundMod/Source/SoundModImporter/Classes/SoundModImporterFactory.h b/Engine/Plugins/Runtime/SoundMod/Source/SoundModImporter/Classes/SoundModImporterFactory.h
--- a/Engine/Plugins/Runtime/SoundMod/Source/SoundModImporter/Classes/SoundModImporterFactory.h
+++ b/Engine/Plugins/Runtime/SoundMod/Source/SoundModImporter/Classes/SoundModImporterFactory.h
@@ -8,9 +8,7 @@
 UCLASS()
 class USoundModImporterFactory : public UFactory
 {
-	GENERATED_BODY()
-public:
-	USoundModImporterFactory(const FObjectInitializer& ObjectInitializer = FObjectInitializer::Get());
+	GENERATED_UCLASS_BODY()
 
 	// UFactory interface
 	virtual UObject* FactoryCreateBinary(UClass* Class, UObject* InParent, FName Name, EObjectFlags Flags, UObject* Context, const TCHAR* Type, const uint8*& Buffer, const uint8* BufferEnd, FFeedbackContext* Warn) override;
diff --git a/Engine/Plugins/ScriptPlugin/Source/ScriptEditorPlugin/Classes/ReimportScriptFactory.h b/Engine/Plugins/ScriptPlugin/Source/ScriptEditorPlugin/Classes/ReimportScriptFactory.h
--- a/Engine/Plugins/ScriptPlugin/Source/ScriptEditorPlugin/Classes/ReimportScriptFactory.h
+++ b/Engine/Plugins/ScriptPlugin/Source/ScriptEditorPlugin/Classes/ReimportScriptFactory.h
@@ -14,9 +14,7 @@
 UCLASS(MinimalApi, collapsecategories, EarlyAccessPreview)
 class UReimportScriptFactory : public UScriptFactory, public FReimportHandler
 {
-	GENERATED_BODY()
-public:
-	SCRIPTEDITORPLUGIN_API UReimportScriptFactory(const FObjectInitializer& ObjectInitializer = FObjectInitializer::Get());
+	GENERATED_UCLASS_BODY()
 
 	UPROPERTY()
 	class UScriptBlueprint* OriginalScript;
diff --git a/Engine/Source/Developer/AITestSuite/Classes/MockAI.h b/Engine/Source/Developer/AITestSuite/Classes/MockAI.h
--- a/Engine/Source/Developer/AITestSuite/Classes/MockAI.h
+++ b/Engine/Source/Developer/AITestSuite/Classes/MockAI.h
@@ -20,9 +20,7 @@ struct FTestTickHelper : FTickableGameObject
 UCLASS()
 class UMockAI : public UObject
 {
-	GENERATED_BODY()
-public:
-	UMockAI(const FObjectInitializer& ObjectInitializer = FObjectInitializer::Get());
+	GENERATED_UCLASS_BODY()
 
 	virtual ~UMockAI();
 
diff --git a/Engine/Source/Developer/Android/AndroidPlatformEditor/Private/AndroidSDKSettings.h b/Engine/Source/Developer/Android/AndroidPlatformEditor/Private/AndroidSDKSettings.h
--- a/Engine/Source/Developer/Android/AndroidPlatformEditor/Private/AndroidSDKSettings.h
+++ b/Engine/Source/Developer/Android/AndroidPlatformEditor/Private/AndroidSDKSettings.h
@@ -13,10 +13,8 @@
 UCLASS(config=Engine, globaluserconfig)
 class ANDROIDPLATFORMEDITOR_API UAndroidSDKSettings : public UObject
 {
-	GENERATED_BODY()
-public:
-	UAndroidSDKSettings(const FObjectInitializer& ObjectInitializer = FObjectInitializer::Get());
 public:
+	GENERATED_UCLASS_BODY()
 
 	// Location on disk of the Android SDK (falls back to ANDROID_HOME environment variable if this is left blank)
 	UPROPERTY(GlobalConfig, EditAnywhere, Category = SDKConfig, Meta = (DisplayName = "Location of Android SDK (the directory usually contains 'android-sdk-')"))
diff --git a/Engine/Source/Editor/AIGraph/Classes/AIGraphSchema.h b/Engine/Source/Editor/AIGraph/Classes/AIGraphSchema.h
--- a/Engine/Source/Editor/AIGraph/Classes/AIGraphSchema.h
+++ b/Engine/Source/Editor/AIGraph/Classes/AIGraphSchema.h
@@ -9,7 +9,7 @@
 USTRUCT()
 struct AIGRAPH_API FAISchemaAction_NewNode : public FEdGraphSchemaAction
 {
-	GENERATED_BODY()
+	GENERATED_USTRUCT_BODY();
 
 	/** Template of node we want to create */
 	UPROPERTY()
@@ -45,7 +45,7 @@ struct AIGRAPH_API FAISchemaAction_NewNode : public FEdGraphSchemaAction
 USTRUCT()
 struct AIGRAPH_API FAISchemaAction_NewSubNode : public FEdGraphSchemaAction
 {
-	GENERATED_BODY()
+	GENERATED_USTRUCT_BODY();
 
 	/** Template of node we want to create */
 	UPROPERTY()
@@ -77,9 +77,7 @@ struct AIGRAPH_API FAISchemaAction_NewSubNode : public FEdGraphSchemaAction
 UCLASS()
 class AIGRAPH_API UAIGraphSchema : public UEdGraphSchema
 {
-	GENERATED_BODY()
-public:
-	UAIGraphSchema(const FObjectInitializer& ObjectInitializer = FObjectInitializer::Get());
+	GENERATED_UCLASS_BODY()
 
 	void GetBreakLinkToSubMenuActions(class FMenuBuilder& MenuBuilder, class UEdGraphPin* InGraphPin);
 
diff --git a/Engine/Source/Editor/AnimGraph/Classes/AnimGraphNode_BlendListBase.h b/Engine/Source/Editor/AnimGraph/Classes/AnimGraphNode_BlendListBase.h
--- a/Engine/Source/Editor/AnimGraph/Classes/AnimGraphNode_BlendListBase.h
+++ b/Engine/Source/Editor/AnimGraph/Classes/AnimGraphNode_BlendListBase.h
@@ -7,9 +7,7 @@
 UCLASS(MinimalAPI, Abstract)
 class UAnimGraphNode_BlendListBase : public UAnimGraphNode_Base
 {
-	GENERATED_BODY()
-public:
-	ANIMGRAPH_API UAnimGraphNode_BlendListBase(const FObjectInitializer& ObjectInitializer = FObjectInitializer::Get());
+	GENERATED_UCLASS_BODY()
 
 	// UEdGraphNode interface
 	ANIMGRAPH_API virtual FLinearColor GetNodeTitleColor() const override;
diff --git a/Engine/Source/Editor/AnimGraph/Classes/AnimGraphNode_CustomTransitionResult.h b/Engine/Source/Editor/AnimGraph/Classes/AnimGraphNode_CustomTransitionResult.h
--- a/Engine/Source/Editor/AnimGraph/Classes/AnimGraphNode_CustomTransitionResult.h
+++ b/Engine/Source/Editor/AnimGraph/Classes/AnimGraphNode_CustomTransitionResult.h
@@ -7,9 +7,7 @@
 UCLASS(MinimalAPI)
 class UAnimGraphNode_CustomTransitionResult : public UAnimGraphNode_Root
 {
-	GENERATED_BODY()
-public:
-	ANIMGRAPH_API UAnimGraphNode_CustomTransitionResult(const FObjectInitializer& ObjectInitializer = FObjectInitializer::Get());
+	GENERATED_UCLASS_BODY()
 
 	// UEdGraphNode interface
 	virtual FLinearColor GetNodeTitleColor() const override;
diff --git a/Engine/Source/Editor/AnimGraph/Classes/AnimGraphNode_SkeletalControlBase.h b/Engine/Source/Editor/AnimGraph/Classes/AnimGraphNode_SkeletalControlBase.h
--- a/Engine/Source/Editor/AnimGraph/Classes/AnimGraphNode_SkeletalControlBase.h
+++ b/Engine/Source/Editor/AnimGraph/Classes/AnimGraphNode_SkeletalControlBase.h
@@ -8,9 +8,7 @@
 UCLASS(Abstract, MinimalAPI)
 class UAnimGraphNode_SkeletalControlBase : public UAnimGraphNode_Base
 {
-	GENERATED_BODY()
-public:
-	ANIMGRAPH_API UAnimGraphNode_SkeletalControlBase(const FObjectInitializer& ObjectInitializer = FObjectInitializer::Get());
+	GENERATED_UCLASS_BODY()
 
 public:
 	// UEdGraphNode interface
diff --git a/Engine/Source/Editor/AnimGraph/Classes/AnimGraphNode_UseCachedPose.h b/Engine/Source/Editor/AnimGraph/Classes/AnimGraphNode_UseCachedPose.h
--- a/Engine/Source/Editor/AnimGraph/Classes/AnimGraphNode_UseCachedPose.h
+++ b/Engine/Source/Editor/AnimGraph/Classes/AnimGraphNode_UseCachedPose.h
@@ -9,9 +9,7 @@
 UCLASS(MinimalAPI)
 class UAnimGraphNode_UseCachedPose : public UAnimGraphNode_Base
 {
-	GENERATED_BODY()
-public:
-	ANIMGRAPH_API UAnimGraphNode_UseCachedPose(const FObjectInitializer& ObjectInitializer = FObjectInitializer::Get());
+	GENERATED_UCLASS_BODY()
 
 public:
 	UPROPERTY()
diff --git a/Engine/Source/Editor/AnimGraph/Classes/AnimationStateMachineSchema.h b/Engine/Source/Editor/AnimGraph/Classes/AnimationStateMachineSchema.h
--- a/Engine/Source/Editor/AnimGraph/Classes/AnimationStateMachineSchema.h
+++ b/Engine/Source/Editor/AnimGraph/Classes/AnimationStateMachineSchema.h
@@ -8,7 +8,7 @@
 USTRUCT()
 struct ANIMGRAPH_API FEdGraphSchemaAction_NewStateNode : public FEdGraphSchemaAction
 {
-	GENERATED_BODY()
+	GENERATED_USTRUCT_BODY()
 
 	UEdGraphNode* NodeTemplate;
 
@@ -39,7 +39,7 @@ struct ANIMGRAPH_API FEdGraphSchemaAction_NewStateNode : public FEdGraphSchemaAc
 USTRUCT()
 struct ANIMGRAPH_API FEdGraphSchemaAction_NewStateComment : public FEdGraphSchemaAction
 {
-	GENERATED_BODY()
+	GENERATED_USTRUCT_BODY();
 
 	FEdGraphSchemaAction_NewStateComment() 
 		: FEdGraphSchemaAction()
@@ -58,9 +58,7 @@ struct ANIMGRAPH_API FEdGraphSchemaAction_NewStateComment : public FEdGraphSchem
 UCLASS(MinimalAPI)
 class UAnimationStateMachineSchema : public UEdGraphSchema
 {
-	GENERATED_BODY()
-public:
-	ANIMGRAPH_API UAnimationStateMachineSchema(const FObjectInitializer& ObjectInitializer = FObjectInitializer::Get());
+	GENERATED_UCLASS_BODY()
 
 	UPROPERTY()
 	FString PC_Exec;
diff --git a/Engine/Source/Editor/BehaviorTreeEditor/Classes/BehaviorTreeDecoratorGraphNode_Decorator.h b/Engine/Source/Editor/BehaviorTreeEditor/Classes/BehaviorTreeDecoratorGraphNode_Decorator.h
--- a/Engine/Source/Editor/BehaviorTreeEditor/Classes/BehaviorTreeDecoratorGraphNode_Decorator.h
+++ b/Engine/Source/Editor/BehaviorTreeEditor/Classes/BehaviorTreeDecoratorGraphNode_Decorator.h
@@ -8,9 +8,7 @@
 UCLASS()
 class UBehaviorTreeDecoratorGraphNode_Decorator : public UBehaviorTreeDecoratorGraphNode 
 {
-	GENERATED_BODY()
-public:
-	UBehaviorTreeDecoratorGraphNode_Decorator(const FObjectInitializer& ObjectInitializer = FObjectInitializer::Get());
+	GENERATED_UCLASS_BODY()
 
 	UPROPERTY()
 	UObject* NodeInstance;
diff --git a/Engine/Source/Editor/BehaviorTreeEditor/Classes/BehaviorTreeFactory.h b/Engine/Source/Editor/BehaviorTreeEditor/Classes/BehaviorTreeFactory.h
--- a/Engine/Source/Editor/BehaviorTreeEditor/Classes/BehaviorTreeFactory.h
+++ b/Engine/Source/Editor/BehaviorTreeEditor/Classes/BehaviorTreeFactory.h
@@ -6,9 +6,7 @@
 UCLASS()
 class UBehaviorTreeFactory : public UFactory
 {
-	GENERATED_BODY()
-public:
-	UBehaviorTreeFactory(const FObjectInitializer& ObjectInitializer = FObjectInitializer::Get());
+	GENERATED_UCLASS_BODY()
 
 	// UFactory interface
 	virtual UObject* FactoryCreateNew(UClass* Class,UObject* InParent,FName Name,EObjectFlags Flags,UObject* Context,FFeedbackContext* Warn) override;
diff --git a/Engine/Source/Editor/BlueprintGraph/Classes/BlueprintVariableNodeSpawner.h b/Engine/Source/Editor/BlueprintGraph/Classes/BlueprintVariableNodeSpawner.h
--- a/Engine/Source/Editor/BlueprintGraph/Classes/BlueprintVariableNodeSpawner.h
+++ b/Engine/Source/Editor/BlueprintGraph/Classes/BlueprintVariableNodeSpawner.h
@@ -19,9 +19,7 @@ class UK2Node_Variable;
 UCLASS(Transient)
 class BLUEPRINTGRAPH_API UBlueprintVariableNodeSpawner : public UBlueprintFieldNodeSpawner
 {
-	GENERATED_BODY()
-public:
-	UBlueprintVariableNodeSpawner(const FObjectInitializer& ObjectInitializer = FObjectInitializer::Get());
+	GENERATED_UCLASS_BODY()
 
 public:
 	/**
@@ -105,4 +103,4 @@ private:
 	/** */
 	UPROPERTY()
 	FBPVariableDescription LocalVarDesc;
-};
+};
\ No newline at end of file
diff --git a/Engine/Source/Editor/BlueprintGraph/Classes/EdGraphSchema_K2_Actions.h b/Engine/Source/Editor/BlueprintGraph/Classes/EdGraphSchema_K2_Actions.h
--- a/Engine/Source/Editor/BlueprintGraph/Classes/EdGraphSchema_K2_Actions.h
+++ b/Engine/Source/Editor/BlueprintGraph/Classes/EdGraphSchema_K2_Actions.h
@@ -12,7 +12,7 @@
 USTRUCT()
 struct BLUEPRINTGRAPH_API FEdGraphSchemaAction_K2NewNode : public FEdGraphSchemaAction
 {
-	GENERATED_BODY()
+	GENERATED_USTRUCT_BODY()
 
 	// Simple type info
 	static FName StaticGetTypeId() {static FName Type("FEdGraphSchemaAction_K2NewNode"); return Type;}
@@ -63,7 +63,7 @@ struct BLUEPRINTGRAPH_API FEdGraphSchemaAction_K2NewNode : public FEdGraphSchema
 USTRUCT()
 struct BLUEPRINTGRAPH_API FEdGraphSchemaAction_K2ViewNode : public FEdGraphSchemaAction
 {
-	GENERATED_BODY()
+	GENERATED_USTRUCT_BODY()
 
 	// Simple type info
 	static FName StaticGetTypeId() {static FName Type("FEdGraphSchemaAction_K2ViewNode"); return Type;}
@@ -98,7 +98,7 @@ struct BLUEPRINTGRAPH_API FEdGraphSchemaAction_K2ViewNode : public FEdGraphSchem
 USTRUCT()
 struct BLUEPRINTGRAPH_API FEdGraphSchemaAction_K2AssignDelegate : public FEdGraphSchemaAction_K2NewNode
 {
-	GENERATED_BODY()
+	GENERATED_USTRUCT_BODY()
 
 	static FName StaticGetTypeId() {static FName Type("FEdGraphSchemaAction_K2AssignDelegate"); return Type;}
 	virtual FName GetTypeId() const override { return StaticGetTypeId(); } 
@@ -126,7 +126,7 @@ struct BLUEPRINTGRAPH_API FEdGraphSchemaAction_K2AssignDelegate : public FEdGrap
 USTRUCT()
 struct BLUEPRINTGRAPH_API FEdGraphSchemaAction_EventFromFunction : public FEdGraphSchemaAction
 {
-	GENERATED_BODY()
+	GENERATED_USTRUCT_BODY()
 
 	static FName StaticGetTypeId() {static FName Type("FEdGraphSchemaAction_EventFromFunction"); return Type;}
 	virtual FName GetTypeId() const override { return StaticGetTypeId(); } 
@@ -159,7 +159,7 @@ struct BLUEPRINTGRAPH_API FEdGraphSchemaAction_EventFromFunction : public FEdGra
 USTRUCT()
 struct FEdGraphSchemaAction_K2AddComponent : public FEdGraphSchemaAction_K2NewNode
 {
-	GENERATED_BODY()
+	GENERATED_USTRUCT_BODY()
 
 	// Simple type info
 	static FName StaticGetTypeId() {static FName Type("FEdGraphSchemaAction_K2AddComponent"); return Type;}
@@ -199,7 +199,7 @@ struct FEdGraphSchemaAction_K2AddComponent : public FEdGraphSchemaAction_K2NewNo
 USTRUCT()
 struct FEdGraphSchemaAction_K2AddTimeline : public FEdGraphSchemaAction_K2NewNode
 {
-	GENERATED_BODY()
+	GENERATED_USTRUCT_BODY()
 
 	// Simple type info
 	static FName StaticGetTypeId() {static FName Type("FEdGraphSchemaAction_K2AddTimeline"); return Type;}
@@ -226,7 +226,7 @@ struct FEdGraphSchemaAction_K2AddTimeline : public FEdGraphSchemaAction_K2NewNod
 USTRUCT()
 struct FEdGraphSchemaAction_K2AddEvent : public FEdGraphSchemaAction_K2NewNode
 {
-	GENERATED_BODY()
+	GENERATED_USTRUCT_BODY()
 
 	// Simple type info
 	static FName StaticGetTypeId() {static FName Type("FEdGraphSchemaAction_K2AddEvent"); return Type;}
@@ -256,7 +256,7 @@ struct FEdGraphSchemaAction_K2AddEvent : public FEdGraphSchemaAction_K2NewNode
 USTRUCT()
 struct FEdGraphSchemaAction_K2AddCustomEvent : public FEdGraphSchemaAction_K2NewNode
 {
-	GENERATED_BODY()
+	GENERATED_USTRUCT_BODY()
 
 	// Simple type info
 	static FName StaticGetTypeId() {static FName Type("FEdGraphSchemaAction_K2AddCustomEvent"); return Type;}
@@ -283,7 +283,7 @@ struct FEdGraphSchemaAction_K2AddCustomEvent : public FEdGraphSchemaAction_K2New
 USTRUCT()
 struct FEdGraphSchemaAction_K2AddCallOnActor : public FEdGraphSchemaAction_K2NewNode
 {
-	GENERATED_BODY()
+	GENERATED_USTRUCT_BODY()
 
 	// Simple type info
 	static FName StaticGetTypeId() {static FName Type("FEdGraphSchemaAction_K2AddCallOnActor"); return Type;}
@@ -315,7 +315,7 @@ struct FEdGraphSchemaAction_K2AddCallOnActor : public FEdGraphSchemaAction_K2New
 USTRUCT()
 struct FEdGraphSchemaAction_K2AddCallOnVariable : public FEdGraphSchemaAction_K2NewNode
 {
-	GENERATED_BODY()
+	GENERATED_USTRUCT_BODY()
 
 	/** Property name we want to call the function on */
 	UPROPERTY()
@@ -342,7 +342,7 @@ struct FEdGraphSchemaAction_K2AddCallOnVariable : public FEdGraphSchemaAction_K2
 USTRUCT()
 struct FEdGraphSchemaAction_K2AddComment : public FEdGraphSchemaAction
 {
-	GENERATED_BODY()
+	GENERATED_USTRUCT_BODY()
 
 	// Simple type info
 	static FName StaticGetTypeId() {static FName Type("FEdGraphSchemaAction_K2AddComment"); return Type;}
@@ -377,7 +377,7 @@ struct FEdGraphSchemaAction_K2AddComment : public FEdGraphSchemaAction
 USTRUCT()
 struct FEdGraphSchemaAction_K2AddDocumentation : public FEdGraphSchemaAction
 {
-	GENERATED_BODY()
+	GENERATED_USTRUCT_BODY()
 
 	// Simple type info
 	static FName StaticGetTypeId() {static FName Type("FEdGraphSchemaAction_K2AddDocumentation"); return Type;}
@@ -412,7 +412,7 @@ struct FEdGraphSchemaAction_K2AddDocumentation : public FEdGraphSchemaAction
 USTRUCT()
 struct BLUEPRINTGRAPH_API FEdGraphSchemaAction_K2TargetNode : public FEdGraphSchemaAction_K2NewNode
 {
-	GENERATED_BODY()
+	GENERATED_USTRUCT_BODY()
 
 
 	// Simple type info
@@ -440,7 +440,7 @@ struct BLUEPRINTGRAPH_API FEdGraphSchemaAction_K2TargetNode : public FEdGraphSch
 USTRUCT()
 struct BLUEPRINTGRAPH_API FEdGraphSchemaAction_K2PasteHere : public FEdGraphSchemaAction_K2NewNode
 {
-	GENERATED_BODY()
+	GENERATED_USTRUCT_BODY()
 
 	// Simple type info
 	static FName StaticGetTypeId() {static FName Type("FEdGraphSchemaAction_K2PasteHere"); return Type;}
@@ -467,7 +467,7 @@ struct BLUEPRINTGRAPH_API FEdGraphSchemaAction_K2PasteHere : public FEdGraphSche
 USTRUCT()
 struct BLUEPRINTGRAPH_API FEdGraphSchemaAction_K2Enum : public FEdGraphSchemaAction
 {
-	GENERATED_BODY()
+	GENERATED_USTRUCT_BODY()
 
 	// Simple type info
 	static FName StaticGetTypeId() {static FName Type("FEdGraphSchemaAction_K2Enum"); return Type;}
@@ -505,7 +505,7 @@ struct BLUEPRINTGRAPH_API FEdGraphSchemaAction_K2Enum : public FEdGraphSchemaAct
 USTRUCT()
 struct BLUEPRINTGRAPH_API FEdGraphSchemaAction_K2Var : public FEdGraphSchemaAction
 {
-	GENERATED_BODY()
+	GENERATED_USTRUCT_BODY()
 
 private:
 	/** Name of function or class */
@@ -572,7 +572,7 @@ public:
 USTRUCT()
 struct BLUEPRINTGRAPH_API FEdGraphSchemaAction_K2LocalVar : public FEdGraphSchemaAction
 {
-	GENERATED_BODY()
+	GENERATED_USTRUCT_BODY()
 
 private:
 	/** Name of function or class */
@@ -654,7 +654,7 @@ namespace EEdGraphSchemaAction_K2Graph
 USTRUCT()
 struct BLUEPRINTGRAPH_API FEdGraphSchemaAction_K2Graph : public FEdGraphSchemaAction
 {
-	GENERATED_BODY()
+	GENERATED_USTRUCT_BODY()
 
 	// Simple type info
 	static FName StaticGetTypeId() {static FName Type("FEdGraphSchemaAction_K2Graph"); return Type;}
@@ -693,7 +693,7 @@ struct BLUEPRINTGRAPH_API FEdGraphSchemaAction_K2Graph : public FEdGraphSchemaAc
 USTRUCT()
 struct BLUEPRINTGRAPH_API FEdGraphSchemaAction_K2Event : public FEdGraphSchemaAction_K2TargetNode
 {
-	GENERATED_BODY()
+	GENERATED_USTRUCT_BODY()
 
 	/**
 	 * An empty default constructor (stubbed out for arrays and other containers 
@@ -743,7 +743,7 @@ struct BLUEPRINTGRAPH_API FEdGraphSchemaAction_K2Event : public FEdGraphSchemaAc
 USTRUCT()
 struct BLUEPRINTGRAPH_API FEdGraphSchemaAction_K2InputAction : public FEdGraphSchemaAction_K2TargetNode
 {
-	GENERATED_BODY()
+	GENERATED_USTRUCT_BODY()
 
 	/**
 	 * An empty default constructor (stubbed out for arrays and other containers
@@ -790,7 +790,7 @@ struct BLUEPRINTGRAPH_API FEdGraphSchemaAction_K2InputAction : public FEdGraphSc
 USTRUCT()
 struct BLUEPRINTGRAPH_API FEdGraphSchemaAction_K2Delegate : public FEdGraphSchemaAction
 {
-	GENERATED_BODY()
+	GENERATED_USTRUCT_BODY()
 
 private:
 	/** Name of the property */
diff --git a/Engine/Source/Editor/BlueprintGraph/Classes/K2Node_AIMoveTo.h b/Engine/Source/Editor/BlueprintGraph/Classes/K2Node_AIMoveTo.h
--- a/Engine/Source/Editor/BlueprintGraph/Classes/K2Node_AIMoveTo.h
+++ b/Engine/Source/Editor/BlueprintGraph/Classes/K2Node_AIMoveTo.h
@@ -9,9 +9,7 @@
 UCLASS()
 class BLUEPRINTGRAPH_API UK2Node_AIMoveTo : public UK2Node_BaseAsyncTask
 {
-	GENERATED_BODY()
-public:
-	UK2Node_AIMoveTo(const FObjectInitializer& ObjectInitializer = FObjectInitializer::Get());
+	GENERATED_UCLASS_BODY()
 
 	// Begin UEdGraphNode interface
 	virtual FText GetTooltipText() const override;
diff --git a/Engine/Source/Editor/BlueprintGraph/Classes/K2Node_BreakStruct.h b/Engine/Source/Editor/BlueprintGraph/Classes/K2Node_BreakStruct.h
--- a/Engine/Source/Editor/BlueprintGraph/Classes/K2Node_BreakStruct.h
+++ b/Engine/Source/Editor/BlueprintGraph/Classes/K2Node_BreakStruct.h
@@ -8,9 +8,7 @@
 UCLASS(MinimalAPI)
 class UK2Node_BreakStruct : public UK2Node_StructMemberGet
 {
-	GENERATED_BODY()
-public:
-	BLUEPRINTGRAPH_API UK2Node_BreakStruct(const FObjectInitializer& ObjectInitializer = FObjectInitializer::Get());
+	GENERATED_UCLASS_BODY()
 
 	static bool CanBeBroken(const UScriptStruct* Struct);
 	static bool CanCreatePinForProperty(const UProperty* Property);
diff --git a/Engine/Source/Editor/BlueprintGraph/Classes/K2Node_CallDataTableFunction.h b/Engine/Source/Editor/BlueprintGraph/Classes/K2Node_CallDataTableFunction.h
--- a/Engine/Source/Editor/BlueprintGraph/Classes/K2Node_CallDataTableFunction.h
+++ b/Engine/Source/Editor/BlueprintGraph/Classes/K2Node_CallDataTableFunction.h
@@ -7,12 +7,10 @@
 UCLASS(MinimalAPI)
 class UK2Node_CallDataTableFunction : public UK2Node_CallFunction
 {
-	GENERATED_BODY()
-public:
-	BLUEPRINTGRAPH_API UK2Node_CallDataTableFunction(const FObjectInitializer& ObjectInitializer = FObjectInitializer::Get());
+	GENERATED_UCLASS_BODY()
 
 	// Begin EdGraphNode interface
 	virtual void PinDefaultValueChanged(UEdGraphPin* Pin) override;
 	virtual void NotifyPinConnectionListChanged(UEdGraphPin* Pin) override;
 	// End EdGraphNode interface
-};
+};
\ No newline at end of file
diff --git a/Engine/Source/Editor/BlueprintGraph/Classes/K2Node_InputAction.h b/Engine/Source/Editor/BlueprintGraph/Classes/K2Node_InputAction.h
--- a/Engine/Source/Editor/BlueprintGraph/Classes/K2Node_InputAction.h
+++ b/Engine/Source/Editor/BlueprintGraph/Classes/K2Node_InputAction.h
@@ -9,9 +9,7 @@
 UCLASS(MinimalAPI)
 class UK2Node_InputAction : public UK2Node
 {
-	GENERATED_BODY()
-public:
-	BLUEPRINTGRAPH_API UK2Node_InputAction(const FObjectInitializer& ObjectInitializer = FObjectInitializer::Get());
+	GENERATED_UCLASS_BODY()
 
 	UPROPERTY()
 	FName InputActionName;
diff --git a/Engine/Source/Editor/BlueprintGraph/Classes/K2Node_InputAxisKeyEvent.h b/Engine/Source/Editor/BlueprintGraph/Classes/K2Node_InputAxisKeyEvent.h
--- a/Engine/Source/Editor/BlueprintGraph/Classes/K2Node_InputAxisKeyEvent.h
+++ b/Engine/Source/Editor/BlueprintGraph/Classes/K2Node_InputAxisKeyEvent.h
@@ -7,9 +7,7 @@
 UCLASS(MinimalAPI)
 class UK2Node_InputAxisKeyEvent : public UK2Node_Event
 {
-	GENERATED_BODY()
-public:
-	BLUEPRINTGRAPH_API UK2Node_InputAxisKeyEvent(const FObjectInitializer& ObjectInitializer = FObjectInitializer::Get());
+	GENERATED_UCLASS_BODY()
 
 	UPROPERTY()
 	FKey AxisKey;
diff --git a/Engine/Source/Editor/BlueprintGraph/Classes/K2Node_InputTouch.h b/Engine/Source/Editor/BlueprintGraph/Classes/K2Node_InputTouch.h
--- a/Engine/Source/Editor/BlueprintGraph/Classes/K2Node_InputTouch.h
+++ b/Engine/Source/Editor/BlueprintGraph/Classes/K2Node_InputTouch.h
@@ -8,9 +8,7 @@
 UCLASS(MinimalAPI)
 class UK2Node_InputTouch : public UK2Node
 {
-	GENERATED_BODY()
-public:
-	BLUEPRINTGRAPH_API UK2Node_InputTouch(const FObjectInitializer& ObjectInitializer = FObjectInitializer::Get());
+	GENERATED_UCLASS_BODY()
 
 	// Prevents actors with lower priority from handling this input
 	UPROPERTY(EditAnywhere, Category="Input")
diff --git a/Engine/Source/Editor/BlueprintGraph/Classes/K2Node_MakeStruct.h b/Engine/Source/Editor/BlueprintGraph/Classes/K2Node_MakeStruct.h
--- a/Engine/Source/Editor/BlueprintGraph/Classes/K2Node_MakeStruct.h
+++ b/Engine/Source/Editor/BlueprintGraph/Classes/K2Node_MakeStruct.h
@@ -9,9 +9,7 @@
 UCLASS(MinimalAPI)
 class UK2Node_MakeStruct : public UK2Node_StructMemberSet
 {
-	GENERATED_BODY()
-public:
-	BLUEPRINTGRAPH_API UK2Node_MakeStruct(const FObjectInitializer& ObjectInitializer = FObjectInitializer::Get());
+	GENERATED_UCLASS_BODY()
 
 	BLUEPRINTGRAPH_API static bool CanBeMade(const UScriptStruct* Struct);
 	static bool CanBeExposed(const UProperty* Property);
diff --git a/Engine/Source/Editor/BlueprintGraph/Classes/K2Node_RemoveDelegate.h b/Engine/Source/Editor/BlueprintGraph/Classes/K2Node_RemoveDelegate.h
--- a/Engine/Source/Editor/BlueprintGraph/Classes/K2Node_RemoveDelegate.h
+++ b/Engine/Source/Editor/BlueprintGraph/Classes/K2Node_RemoveDelegate.h
@@ -6,9 +6,7 @@
 UCLASS(MinimalAPI)
 class UK2Node_RemoveDelegate : public UK2Node_BaseMCDelegate
 {
-	GENERATED_BODY()
-public:
-	BLUEPRINTGRAPH_API UK2Node_RemoveDelegate(const FObjectInitializer& ObjectInitializer = FObjectInitializer::Get());
+	GENERATED_UCLASS_BODY()
 
 public:
 	// UEdGraphNode interface
diff --git a/Engine/Source/Editor/BlueprintGraph/Classes/K2Node_SwitchEnum.h b/Engine/Source/Editor/BlueprintGraph/Classes/K2Node_SwitchEnum.h
--- a/Engine/Source/Editor/BlueprintGraph/Classes/K2Node_SwitchEnum.h
+++ b/Engine/Source/Editor/BlueprintGraph/Classes/K2Node_SwitchEnum.h
@@ -10,9 +10,7 @@
 UCLASS(MinimalAPI)
 class UK2Node_SwitchEnum : public UK2Node_Switch, public INodeDependingOnEnumInterface
 {
-	GENERATED_BODY()
-public:
-	BLUEPRINTGRAPH_API UK2Node_SwitchEnum(const FObjectInitializer& ObjectInitializer = FObjectInitializer::Get());
+	GENERATED_UCLASS_BODY()
 
 	/** Name of the enum being switched on */
 	UPROPERTY()
diff --git a/Engine/Source/Editor/BlueprintGraph/Classes/K2Node_TemporaryVariable.h b/Engine/Source/Editor/BlueprintGraph/Classes/K2Node_TemporaryVariable.h
--- a/Engine/Source/Editor/BlueprintGraph/Classes/K2Node_TemporaryVariable.h
+++ b/Engine/Source/Editor/BlueprintGraph/Classes/K2Node_TemporaryVariable.h
@@ -8,9 +8,7 @@
 UCLASS(MinimalAPI)
 class UK2Node_TemporaryVariable : public UK2Node
 {
-	GENERATED_BODY()
-public:
-	BLUEPRINTGRAPH_API UK2Node_TemporaryVariable(const FObjectInitializer& ObjectInitializer = FObjectInitializer::Get());
+	GENERATED_UCLASS_BODY()
 
 	UPROPERTY()
 	struct FEdGraphPinType VariableType;
diff --git a/Engine/Source/Editor/BlueprintGraph/Classes/K2Node_VariableSet.h b/Engine/Source/Editor/BlueprintGraph/Classes/K2Node_VariableSet.h
--- a/Engine/Source/Editor/BlueprintGraph/Classes/K2Node_VariableSet.h
+++ b/Engine/Source/Editor/BlueprintGraph/Classes/K2Node_VariableSet.h
@@ -9,9 +9,7 @@
 UCLASS(MinimalAPI)
 class UK2Node_VariableSet : public UK2Node_Variable
 {
-	GENERATED_BODY()
-public:
-	BLUEPRINTGRAPH_API UK2Node_VariableSet(const FObjectInitializer& ObjectInitializer = FObjectInitializer::Get());
+	GENERATED_UCLASS_BODY()
 
 	// Begin UEdGraphNode interface
 	virtual void AllocateDefaultPins() override;
diff --git a/Engine/Source/Editor/BlueprintGraph/Classes/K2Node_VariableSetRef.h b/Engine/Source/Editor/BlueprintGraph/Classes/K2Node_VariableSetRef.h
--- a/Engine/Source/Editor/BlueprintGraph/Classes/K2Node_VariableSetRef.h
+++ b/Engine/Source/Editor/BlueprintGraph/Classes/K2Node_VariableSetRef.h
@@ -9,9 +9,7 @@
 UCLASS(MinimalAPI)
 class UK2Node_VariableSetRef : public UK2Node
 {
-	GENERATED_BODY()
-public:
-	BLUEPRINTGRAPH_API UK2Node_VariableSetRef(const FObjectInitializer& ObjectInitializer = FObjectInitializer::Get());
+	GENERATED_UCLASS_BODY()
 
 	// Begin UEdGraphNode interface
 	virtual void AllocateDefaultPins() override;
diff --git a/Engine/Source/Editor/EnvironmentQueryEditor/Classes/EnvironmentQueryGraphNode_Root.h b/Engine/Source/Editor/EnvironmentQueryEditor/Classes/EnvironmentQueryGraphNode_Root.h
--- a/Engine/Source/Editor/EnvironmentQueryEditor/Classes/EnvironmentQueryGraphNode_Root.h
+++ b/Engine/Source/Editor/EnvironmentQueryEditor/Classes/EnvironmentQueryGraphNode_Root.h
@@ -5,9 +5,7 @@
 UCLASS()
 class UEnvironmentQueryGraphNode_Root : public UEnvironmentQueryGraphNode
 {
-	GENERATED_BODY()
-public:
-	UEnvironmentQueryGraphNode_Root(const FObjectInitializer& ObjectInitializer = FObjectInitializer::Get());
+	GENERATED_UCLASS_BODY()
 
 	virtual void AllocateDefaultPins() override;
 	virtual FText GetNodeTitle(ENodeTitleType::Type TitleType) const override;
diff --git a/Engine/Source/Editor/GameplayAbilitiesEditor/Classes/GameplayAbilityGraph.h b/Engine/Source/Editor/GameplayAbilitiesEditor/Classes/GameplayAbilityGraph.h
--- a/Engine/Source/Editor/GameplayAbilitiesEditor/Classes/GameplayAbilityGraph.h
+++ b/Engine/Source/Editor/GameplayAbilitiesEditor/Classes/GameplayAbilityGraph.h
@@ -7,8 +7,6 @@
 UCLASS(MinimalAPI)
 class UGameplayAbilityGraph : public UEdGraph
 {
-	GENERATED_BODY()
-public:
-	GAMEPLAYABILITIESEDITOR_API UGameplayAbilityGraph(const FObjectInitializer& ObjectInitializer = FObjectInitializer::Get());
+	GENERATED_UCLASS_BODY()
 };
 
diff --git a/Engine/Source/Editor/GameplayTagsEditor/Classes/GameplayTagsK2Node_MultiCompareBase.h b/Engine/Source/Editor/GameplayTagsEditor/Classes/GameplayTagsK2Node_MultiCompareBase.h
--- a/Engine/Source/Editor/GameplayTagsEditor/Classes/GameplayTagsK2Node_MultiCompareBase.h
+++ b/Engine/Source/Editor/GameplayTagsEditor/Classes/GameplayTagsK2Node_MultiCompareBase.h
@@ -11,9 +11,7 @@ struct FGameplayTagContainer;
 UCLASS()
 class UGameplayTagsK2Node_MultiCompareBase : public UK2Node
 {
-	GENERATED_BODY()
-public:
-	UGameplayTagsK2Node_MultiCompareBase(const FObjectInitializer& ObjectInitializer = FObjectInitializer::Get());
+	GENERATED_UCLASS_BODY()
 
 	UPROPERTY(EditAnywhere, Category = PinOptions)
 	int32 NumberOfPins;
diff --git a/Engine/Source/Editor/GameplayTagsEditor/Classes/GameplayTagsK2Node_MultiCompareGameplayTagAssetInterface.h b/Engine/Source/Editor/GameplayTagsEditor/Classes/GameplayTagsK2Node_MultiCompareGameplayTagAssetInterface.h
--- a/Engine/Source/Editor/GameplayTagsEditor/Classes/GameplayTagsK2Node_MultiCompareGameplayTagAssetInterface.h
+++ b/Engine/Source/Editor/GameplayTagsEditor/Classes/GameplayTagsK2Node_MultiCompareGameplayTagAssetInterface.h
@@ -8,9 +8,7 @@
 UCLASS()
 class UGameplayTagsK2Node_MultiCompareGameplayTagAssetInterface : public UGameplayTagsK2Node_MultiCompareBase
 {
-	GENERATED_BODY()
-public:
-	UGameplayTagsK2Node_MultiCompareGameplayTagAssetInterface(const FObjectInitializer& ObjectInitializer = FObjectInitializer::Get());
+	GENERATED_UCLASS_BODY()
 
 	// UEdGraphNode interface
 	virtual void AllocateDefaultPins() override;
@@ -25,4 +23,4 @@ public:
 private:
 
 	virtual void AddPinToSwitchNode() override;
-};
+};
\ No newline at end of file
diff --git a/Engine/Source/Editor/GeometryMode/Classes/GeomModifier_Clip.h b/Engine/Source/Editor/GeometryMode/Classes/GeomModifier_Clip.h
--- a/Engine/Source/Editor/GeometryMode/Classes/GeomModifier_Clip.h
+++ b/Engine/Source/Editor/GeometryMode/Classes/GeomModifier_Clip.h
@@ -11,9 +11,7 @@
 UCLASS()
 class UGeomModifier_Clip : public UGeomModifier_Edit
 {
-	GENERATED_BODY()
-public:
-	UGeomModifier_Clip(const FObjectInitializer& ObjectInitializer = FObjectInitializer::Get());
+	GENERATED_UCLASS_BODY()
 
 	UPROPERTY(EditAnywhere, Category=Settings)
 	uint32 bFlipNormal:1;
diff --git a/Engine/Source/Editor/Kismet/Classes/BlueprintPaletteFavorites.h b/Engine/Source/Editor/Kismet/Classes/BlueprintPaletteFavorites.h
--- a/Engine/Source/Editor/Kismet/Classes/BlueprintPaletteFavorites.h
+++ b/Engine/Source/Editor/Kismet/Classes/BlueprintPaletteFavorites.h
@@ -16,7 +16,7 @@ struct FBlueprintActionInfo;
 USTRUCT()
 struct FFavoritedBlueprintPaletteItem
 {
-	GENERATED_BODY()
+	GENERATED_USTRUCT_BODY()
 
 public:
 	/**
@@ -88,9 +88,7 @@ private:
 UCLASS(config=EditorUserSettings)
 class KISMET_API UBlueprintPaletteFavorites : public UObject
 {
-	GENERATED_BODY()
-public:
-	UBlueprintPaletteFavorites(const FObjectInitializer& ObjectInitializer = FObjectInitializer::Get());
+	GENERATED_UCLASS_BODY()
 
 public:
 	// Begin UObject interface
diff --git a/Engine/Source/Editor/KismetCompiler/Private/KismetCompilerCppBackend.cpp b/Engine/Source/Editor/KismetCompiler/Private/KismetCompilerCppBackend.cpp
--- a/Engine/Source/Editor/KismetCompiler/Private/KismetCompilerCppBackend.cpp
+++ b/Engine/Source/Editor/KismetCompiler/Private/KismetCompilerCppBackend.cpp
@@ -774,8 +774,7 @@ void FKismetCppBackend::GenerateCodeFromClass(UClass* SourceClass, TIndirectArra
 	Emit(Header,
 		TEXT("{\n")
 		TEXT("public:\n"));
-	Emit(Header, *FString::Printf(TEXT("\tGENERATED_BODY()\n")));
-	Emit(Header, *FString::Printf(TEXT("public:\n")));
+	Emit(Header, *FString::Printf(TEXT("\tGENERATED_UCLASS_BODY()\n")));
 
 	EmitClassProperties(Header, SourceClass);
 
@@ -795,6 +794,11 @@ void FKismetCppBackend::GenerateCodeFromClass(UClass* SourceClass, TIndirectArra
 	Emit(Body, *FString::Printf(TEXT("#include \"%s.h\"\n"), FApp::GetGameName()));
 	Emit(Body, TEXT("#include \"GeneratedCodeHelpers.h\"\n\n"));
 
+	//constructor
+	Emit(Body, *FString::Printf(
+		TEXT("%s::%s(const FObjectInitializer& ObjectInitializer) : Super(ObjectInitializer) {}\n\n"),
+		*CppClassName, *CppClassName));
+
 	for (int32 i = 0; i < Functions.Num(); ++i)
 	{
 		if (Functions[i].IsValid())
diff --git a/Engine/Source/Editor/Matinee/Classes/MatineeTrackDirectorHelper.h b/Engine/Source/Editor/Matinee/Classes/MatineeTrackDirectorHelper.h
--- a/Engine/Source/Editor/Matinee/Classes/MatineeTrackDirectorHelper.h
+++ b/Engine/Source/Editor/Matinee/Classes/MatineeTrackDirectorHelper.h
@@ -8,9 +8,7 @@
 UCLASS()
 class UMatineeTrackDirectorHelper : public UInterpTrackHelper
 {
-	GENERATED_BODY()
-public:
-	UMatineeTrackDirectorHelper(const FObjectInitializer& ObjectInitializer = FObjectInitializer::Get());
+	GENERATED_UCLASS_BODY()
 
 	void OnAddKeyTextEntry(const FString& ChosenText, IMatineeBase* Matinee, UInterpTrack* Track);
 
diff --git a/Engine/Source/Editor/Matinee/Classes/MatineeTrackLinearColorPropHelper.h b/Engine/Source/Editor/Matinee/Classes/MatineeTrackLinearColorPropHelper.h
--- a/Engine/Source/Editor/Matinee/Classes/MatineeTrackLinearColorPropHelper.h
+++ b/Engine/Source/Editor/Matinee/Classes/MatineeTrackLinearColorPropHelper.h
@@ -8,9 +8,7 @@
 UCLASS()
 class UMatineeTrackLinearColorPropHelper : public UMatineeTrackVectorPropHelper
 {
-	GENERATED_BODY()
-public:
-	UMatineeTrackLinearColorPropHelper(const FObjectInitializer& ObjectInitializer = FObjectInitializer::Get());
+	GENERATED_UCLASS_BODY()
 
 public:
 
diff --git a/Engine/Source/Editor/Matinee/Classes/MatineeTrackToggleHelper.h b/Engine/Source/Editor/Matinee/Classes/MatineeTrackToggleHelper.h
--- a/Engine/Source/Editor/Matinee/Classes/MatineeTrackToggleHelper.h
+++ b/Engine/Source/Editor/Matinee/Classes/MatineeTrackToggleHelper.h
@@ -8,9 +8,7 @@
 UCLASS()
 class UMatineeTrackToggleHelper : public UInterpTrackHelper
 {
-	GENERATED_BODY()
-public:
-	UMatineeTrackToggleHelper(const FObjectInitializer& ObjectInitializer = FObjectInitializer::Get());
+	GENERATED_UCLASS_BODY()
 
 	void OnAddKeyTextEntry(const FString& ChosenText, IMatineeBase* Matinee, UInterpTrack* Track);
 
diff --git a/Engine/Source/Editor/OnlineBlueprintSupport/Classes/K2Node_InAppPurchaseQuery.h b/Engine/Source/Editor/OnlineBlueprintSupport/Classes/K2Node_InAppPurchaseQuery.h
--- a/Engine/Source/Editor/OnlineBlueprintSupport/Classes/K2Node_InAppPurchaseQuery.h
+++ b/Engine/Source/Editor/OnlineBlueprintSupport/Classes/K2Node_InAppPurchaseQuery.h
@@ -7,7 +7,5 @@
 UCLASS()
 class ONLINEBLUEPRINTSUPPORT_API UK2Node_InAppPurchaseQuery : public UK2Node_BaseAsyncTask
 {
-	GENERATED_BODY()
-public:
-	UK2Node_InAppPurchaseQuery(const FObjectInitializer& ObjectInitializer = FObjectInitializer::Get());
+	GENERATED_UCLASS_BODY()
 };
diff --git a/Engine/Source/Editor/OnlineBlueprintSupport/Classes/K2Node_LatentOnlineCall.h b/Engine/Source/Editor/OnlineBlueprintSupport/Classes/K2Node_LatentOnlineCall.h
--- a/Engine/Source/Editor/OnlineBlueprintSupport/Classes/K2Node_LatentOnlineCall.h
+++ b/Engine/Source/Editor/OnlineBlueprintSupport/Classes/K2Node_LatentOnlineCall.h
@@ -8,9 +8,7 @@
 UCLASS()
 class ONLINEBLUEPRINTSUPPORT_API UK2Node_LatentOnlineCall : public UK2Node_BaseAsyncTask
 {
-	GENERATED_BODY()
-public:
-	UK2Node_LatentOnlineCall(const FObjectInitializer& ObjectInitializer = FObjectInitializer::Get());
+	GENERATED_UCLASS_BODY()
 	
 	// UEdGraphNode interface
 	virtual void GetMenuEntries(FGraphContextMenuBuilder& ContextMenuBuilder) const override;
diff --git a/Engine/Source/Editor/TranslationEditor/Private/InternationalizationExportSettings.h b/Engine/Source/Editor/TranslationEditor/Private/InternationalizationExportSettings.h
--- a/Engine/Source/Editor/TranslationEditor/Private/InternationalizationExportSettings.h
+++ b/Engine/Source/Editor/TranslationEditor/Private/InternationalizationExportSettings.h
@@ -6,9 +6,7 @@
 UCLASS(hidecategories = Object, MinimalAPI, config=InternationalizationExport)
 class UInternationalizationExportSettings : public UObject
 {
-	GENERATED_BODY()
-public:
-	TRANSLATIONEDITOR_API UInternationalizationExportSettings(const FObjectInitializer& ObjectInitializer = FObjectInitializer::Get());
+	GENERATED_UCLASS_BODY()
 
 public:
 
@@ -51,4 +49,4 @@ public:
 	/** Whether or not to use culture path */
 	UPROPERTY(Category = GatherTextStep, EditAnywhere, config)
 	bool bUseCultureDirectory;
-};
+};
\ No newline at end of file
diff --git a/Engine/Source/Editor/UMGEditor/Classes/WidgetGraphSchema.h b/Engine/Source/Editor/UMGEditor/Classes/WidgetGraphSchema.h
--- a/Engine/Source/Editor/UMGEditor/Classes/WidgetGraphSchema.h
+++ b/Engine/Source/Editor/UMGEditor/Classes/WidgetGraphSchema.h
@@ -8,9 +8,7 @@
 UCLASS(MinimalAPI)
 class UWidgetGraphSchema : public UEdGraphSchema_K2
 {
-	GENERATED_BODY()
-public:
-	UMGEDITOR_API UWidgetGraphSchema(const FObjectInitializer& ObjectInitializer = FObjectInitializer::Get());
+	GENERATED_UCLASS_BODY()
 
 	virtual bool ShouldAlwaysPurgeOnModification() const override { return true; }
 
@@ -33,4 +31,4 @@ public:
 	FName DefaultEvaluationHandlerName;
 	
 	// TODO UMG - Add more things specific to the needs of the widget blueprint graph.
-};
+};
\ No newline at end of file
diff --git a/Engine/Source/Editor/UnrealEd/Classes/ActorFactories/ActorFactoryAnimationAsset.h b/Engine/Source/Editor/UnrealEd/Classes/ActorFactories/ActorFactoryAnimationAsset.h
--- a/Engine/Source/Editor/UnrealEd/Classes/ActorFactories/ActorFactoryAnimationAsset.h
+++ b/Engine/Source/Editor/UnrealEd/Classes/ActorFactories/ActorFactoryAnimationAsset.h
@@ -10,9 +10,7 @@ class USkeletalMesh;
 UCLASS(MinimalAPI, config=Editor, hidecategories=Object)
 class UActorFactoryAnimationAsset : public UActorFactorySkeletalMesh
 {
-	GENERATED_BODY()
-public:
-	UNREALED_API UActorFactoryAnimationAsset(const FObjectInitializer& ObjectInitializer = FObjectInitializer::Get());
+	GENERATED_UCLASS_BODY()
 
 protected:
 	// Begin UActorFactory Interface
diff --git a/Engine/Source/Editor/UnrealEd/Classes/ActorFactories/ActorFactoryBasicShape.h b/Engine/Source/Editor/UnrealEd/Classes/ActorFactories/ActorFactoryBasicShape.h
--- a/Engine/Source/Editor/UnrealEd/Classes/ActorFactories/ActorFactoryBasicShape.h
+++ b/Engine/Source/Editor/UnrealEd/Classes/ActorFactories/ActorFactoryBasicShape.h
@@ -7,9 +7,7 @@
 UCLASS(MinimalAPI,config=Editor)
 class UActorFactoryBasicShape : public UActorFactoryStaticMesh
 {
-	GENERATED_BODY()
-public:
-	UNREALED_API UActorFactoryBasicShape(const FObjectInitializer& ObjectInitializer = FObjectInitializer::Get());
+	GENERATED_UCLASS_BODY()
 
 	virtual bool CanCreateActorFrom( const FAssetData& AssetData, FText& OutErrorMsg ) override;
 	virtual void PostSpawnActor(UObject* Asset, AActor* NewActor) override;
diff --git a/Engine/Source/Editor/UnrealEd/Classes/ActorFactories/ActorFactoryBlueprint.h b/Engine/Source/Editor/UnrealEd/Classes/ActorFactories/ActorFactoryBlueprint.h
--- a/Engine/Source/Editor/UnrealEd/Classes/ActorFactories/ActorFactoryBlueprint.h
+++ b/Engine/Source/Editor/UnrealEd/Classes/ActorFactories/ActorFactoryBlueprint.h
@@ -8,9 +8,7 @@
 UCLASS(MinimalAPI, config=Editor, collapsecategories, hidecategories=Object)
 class UActorFactoryBlueprint : public UActorFactory
 {
-	GENERATED_BODY()
-public:
-	UNREALED_API UActorFactoryBlueprint(const FObjectInitializer& ObjectInitializer = FObjectInitializer::Get());
+	GENERATED_UCLASS_BODY()
 
 	// Begin UActorFactory Interface
 	virtual bool PreSpawnActor( UObject* Asset, FTransform& InOutLocation ) override;
diff --git a/Engine/Source/Editor/UnrealEd/Classes/ActorFactories/ActorFactoryNiagara.h b/Engine/Source/Editor/UnrealEd/Classes/ActorFactories/ActorFactoryNiagara.h
--- a/Engine/Source/Editor/UnrealEd/Classes/ActorFactories/ActorFactoryNiagara.h
+++ b/Engine/Source/Editor/UnrealEd/Classes/ActorFactories/ActorFactoryNiagara.h
@@ -8,9 +8,7 @@
 UCLASS(MinimalAPI, config = Editor, collapsecategories, hidecategories = Object)
 class UActorFactoryNiagara : public UActorFactory
 {
-	GENERATED_BODY()
-public:
-	UNREALED_API UActorFactoryNiagara(const FObjectInitializer& ObjectInitializer = FObjectInitializer::Get());
+	GENERATED_UCLASS_BODY()
 
 	// Begin UActorFactory Interface
 	virtual void PostSpawnActor(UObject* Asset, AActor* NewActor) override;
diff --git a/Engine/Source/Editor/UnrealEd/Classes/ActorFactories/ActorFactorySphereReflectionCapture.h b/Engine/Source/Editor/UnrealEd/Classes/ActorFactories/ActorFactorySphereReflectionCapture.h
--- a/Engine/Source/Editor/UnrealEd/Classes/ActorFactories/ActorFactorySphereReflectionCapture.h
+++ b/Engine/Source/Editor/UnrealEd/Classes/ActorFactories/ActorFactorySphereReflectionCapture.h
@@ -7,9 +7,7 @@
 UCLASS(MinimalAPI, config=Editor, collapsecategories, hidecategories=Object)
 class UActorFactorySphereReflectionCapture : public UActorFactory
 {
-	GENERATED_BODY()
-public:
-	UNREALED_API UActorFactorySphereReflectionCapture(const FObjectInitializer& ObjectInitializer = FObjectInitializer::Get());
+	GENERATED_UCLASS_BODY()
 };
 
 
diff --git a/Engine/Source/Editor/UnrealEd/Classes/Builders/SheetBuilder.h b/Engine/Source/Editor/UnrealEd/Classes/Builders/SheetBuilder.h
--- a/Engine/Source/Editor/UnrealEd/Classes/Builders/SheetBuilder.h
+++ b/Engine/Source/Editor/UnrealEd/Classes/Builders/SheetBuilder.h
@@ -19,9 +19,11 @@ enum ESheetAxis
 UCLASS(MinimalAPI, autoexpandcategories=BrushSettings, EditInlineNew, NotPlaceable, meta=(DisplayName="Plane"))
 class USheetBuilder : public UEditorBrushBuilder
 {
+public:
 	GENERATED_BODY()
+
 public:
-	UNREALED_API USheetBuilder(const FObjectInitializer& ObjectInitializer = FObjectInitializer::Get());
+	USheetBuilder(const FObjectInitializer& ObjectInitializer = FObjectInitializer::Get());
 
 	UPROPERTY(EditAnywhere, Category=BrushSettings, meta=(ClampMin = "1"))
 	int32 X;
diff --git a/Engine/Source/Editor/UnrealEd/Classes/Commandlets/CookCommandlet.h b/Engine/Source/Editor/UnrealEd/Classes/Commandlets/CookCommandlet.h
--- a/Engine/Source/Editor/UnrealEd/Classes/Commandlets/CookCommandlet.h
+++ b/Engine/Source/Editor/UnrealEd/Classes/Commandlets/CookCommandlet.h
@@ -12,9 +12,7 @@ UCLASS(config=Editor)
 class UCookCommandlet
 	: public UCommandlet
 {
-	GENERATED_BODY()
-public:
-	UCookCommandlet(const FObjectInitializer& ObjectInitializer = FObjectInitializer::Get());
+	GENERATED_UCLASS_BODY()
 
 	/** List of asset types that will force GC after loading them during cook */
 	UPROPERTY(config)
@@ -159,4 +157,4 @@ private:
 			InOutFilesToCook.AddUnique(InFilename);
 		}
 	}
-};
+};
\ No newline at end of file
diff --git a/Engine/Source/Editor/UnrealEd/Classes/Commandlets/FixupRedirectsCommandlet.h b/Engine/Source/Editor/UnrealEd/Classes/Commandlets/FixupRedirectsCommandlet.h
--- a/Engine/Source/Editor/UnrealEd/Classes/Commandlets/FixupRedirectsCommandlet.h
+++ b/Engine/Source/Editor/UnrealEd/Classes/Commandlets/FixupRedirectsCommandlet.h
@@ -7,9 +7,7 @@
 UCLASS()
 class UFixupRedirectsCommandlet : public UCommandlet
 {
-	GENERATED_BODY()
-public:
-	UFixupRedirectsCommandlet(const FObjectInitializer& ObjectInitializer = FObjectInitializer::Get());
+	GENERATED_UCLASS_BODY()
 	// Begin UCommandlet Interface
 	virtual void CreateCustomEngine(const FString& Params) override;
 	virtual int32 Main(const FString& Params) override;
diff --git a/Engine/Source/Editor/UnrealEd/Classes/Exporters/LevelExporterLOD.h b/Engine/Source/Editor/UnrealEd/Classes/Exporters/LevelExporterLOD.h
--- a/Engine/Source/Editor/UnrealEd/Classes/Exporters/LevelExporterLOD.h
+++ b/Engine/Source/Editor/UnrealEd/Classes/Exporters/LevelExporterLOD.h
@@ -10,10 +10,11 @@
 UCLASS()
 class ULevelExporterLOD : public UExporter
 {
-	GENERATED_BODY()
 public:
-	ULevelExporterLOD(const FObjectInitializer& ObjectInitializer = FObjectInitializer::Get());
+	GENERATED_BODY()
 
+public:
+	ULevelExporterLOD(const FObjectInitializer& ObjectInitializer = FObjectInitializer());
 
 	// Begin UExporter Interface
 	virtual bool ExportText( const FExportObjectInnerContext* Context, UObject* Object, const TCHAR* Type, FOutputDevice& Ar, FFeedbackContext* Warn, uint32 PortFlags=0 ) override;
diff --git a/Engine/Source/Editor/UnrealEd/Classes/Exporters/PolysExporterOBJ.h b/Engine/Source/Editor/UnrealEd/Classes/Exporters/PolysExporterOBJ.h
--- a/Engine/Source/Editor/UnrealEd/Classes/Exporters/PolysExporterOBJ.h
+++ b/Engine/Source/Editor/UnrealEd/Classes/Exporters/PolysExporterOBJ.h
@@ -10,10 +10,11 @@
 UCLASS()
 class UPolysExporterOBJ : public UExporter
 {
-	GENERATED_BODY()
 public:
-	UPolysExporterOBJ(const FObjectInitializer& ObjectInitializer = FObjectInitializer::Get());
+	GENERATED_BODY()
 
+public:
+	UPolysExporterOBJ(const FObjectInitializer& ObjectInitializer = FObjectInitializer());
 
 	// Begin UExporter Interface
 	virtual bool ExportText( const FExportObjectInnerContext* Context, UObject* Object, const TCHAR* Type, FOutputDevice& Ar, FFeedbackContext* Warn, uint32 PortFlags=0 ) override;
diff --git a/Engine/Source/Editor/UnrealEd/Classes/Exporters/PolysExporterT3D.h b/Engine/Source/Editor/UnrealEd/Classes/Exporters/PolysExporterT3D.h
--- a/Engine/Source/Editor/UnrealEd/Classes/Exporters/PolysExporterT3D.h
+++ b/Engine/Source/Editor/UnrealEd/Classes/Exporters/PolysExporterT3D.h
@@ -10,10 +10,11 @@
 UCLASS()
 class UPolysExporterT3D : public UExporter
 {
-	GENERATED_BODY()
 public:
-	UPolysExporterT3D(const FObjectInitializer& ObjectInitializer = FObjectInitializer::Get());
+	GENERATED_BODY()
 
+public:
+	UPolysExporterT3D(const FObjectInitializer& ObjectInitializer = FObjectInitializer());
 
 	// Begin UExporter Interface
 	virtual bool ExportText( const FExportObjectInnerContext* Context, UObject* Object, const TCHAR* Type, FOutputDevice& Ar, FFeedbackContext* Warn, uint32 PortFlags=0 ) override;
diff --git a/Engine/Source/Editor/UnrealEd/Classes/Exporters/StaticMeshExporterOBJ.h b/Engine/Source/Editor/UnrealEd/Classes/Exporters/StaticMeshExporterOBJ.h
--- a/Engine/Source/Editor/UnrealEd/Classes/Exporters/StaticMeshExporterOBJ.h
+++ b/Engine/Source/Editor/UnrealEd/Classes/Exporters/StaticMeshExporterOBJ.h
@@ -10,9 +10,7 @@
 UCLASS()
 class UStaticMeshExporterOBJ : public UExporter
 {
-	GENERATED_BODY()
-public:
-	UStaticMeshExporterOBJ(const FObjectInitializer& ObjectInitializer = FObjectInitializer::Get());
+	GENERATED_UCLASS_BODY()
 
 
 	// Begin UExporter Interface
diff --git a/Engine/Source/Editor/UnrealEd/Classes/Exporters/TextureCubeExporterHDR.h b/Engine/Source/Editor/UnrealEd/Classes/Exporters/TextureCubeExporterHDR.h
--- a/Engine/Source/Editor/UnrealEd/Classes/Exporters/TextureCubeExporterHDR.h
+++ b/Engine/Source/Editor/UnrealEd/Classes/Exporters/TextureCubeExporterHDR.h
@@ -10,9 +10,7 @@
 UCLASS()
 class UNREALED_API UTextureCubeExporterHDR : public UExporter
 {
-	GENERATED_BODY()
-public:
-	UTextureCubeExporterHDR(const FObjectInitializer& ObjectInitializer = FObjectInitializer::Get());
+	GENERATED_UCLASS_BODY()
 
 	// Begin UExporter Interface
 	virtual bool ExportBinary(UObject* Object, const TCHAR* Type, FArchive& Ar, FFeedbackContext* Warn, int32 FileIndex = 0, uint32 PortFlags = 0) override;
diff --git a/Engine/Source/Editor/UnrealEd/Classes/Exporters/TextureExporterHDR.h b/Engine/Source/Editor/UnrealEd/Classes/Exporters/TextureExporterHDR.h
--- a/Engine/Source/Editor/UnrealEd/Classes/Exporters/TextureExporterHDR.h
+++ b/Engine/Source/Editor/UnrealEd/Classes/Exporters/TextureExporterHDR.h
@@ -10,9 +10,7 @@
 UCLASS()
 class UNREALED_API UTextureExporterHDR : public UExporter
 {
-	GENERATED_BODY()
-public:
-	UTextureExporterHDR(const FObjectInitializer& ObjectInitializer = FObjectInitializer::Get());
+	GENERATED_UCLASS_BODY()
 
 	// Begin UExporter Interface
 	virtual bool ExportBinary(UObject* Object, const TCHAR* Type, FArchive& Ar, FFeedbackContext* Warn, int32 FileIndex = 0, uint32 PortFlags = 0) override;
diff --git a/Engine/Source/Editor/UnrealEd/Classes/Factories/FbxImportUI.h b/Engine/Source/Editor/UnrealEd/Classes/Factories/FbxImportUI.h
--- a/Engine/Source/Editor/UnrealEd/Classes/Factories/FbxImportUI.h
+++ b/Engine/Source/Editor/UnrealEd/Classes/Factories/FbxImportUI.h
@@ -25,9 +25,7 @@ enum EFBXImportType
 UCLASS(config=EditorUserSettings, AutoExpandCategories=(FTransform), HideCategories=Object, MinimalAPI)
 class UFbxImportUI : public UObject
 {
-	GENERATED_BODY()
-public:
-	UNREALED_API UFbxImportUI(const FObjectInitializer& ObjectInitializer = FObjectInitializer::Get());
+	GENERATED_UCLASS_BODY()
 
 	/** Whether or not the imported file is in OBJ format */
 	UPROPERTY()
diff --git a/Engine/Source/Editor/UnrealEd/Classes/Factories/MaterialFactoryNew.h b/Engine/Source/Editor/UnrealEd/Classes/Factories/MaterialFactoryNew.h
--- a/Engine/Source/Editor/UnrealEd/Classes/Factories/MaterialFactoryNew.h
+++ b/Engine/Source/Editor/UnrealEd/Classes/Factories/MaterialFactoryNew.h
@@ -10,9 +10,7 @@
 UCLASS(hidecategories=Object, collapsecategories, MinimalAPI)
 class UMaterialFactoryNew : public UFactory
 {
-	GENERATED_BODY()
-public:
-	UNREALED_API UMaterialFactoryNew(const FObjectInitializer& ObjectInitializer = FObjectInitializer::Get());
+	GENERATED_UCLASS_BODY()
 
 	/** An initial texture to place in the newly created material */
 	UPROPERTY()
@@ -21,4 +19,4 @@ public:
 	// Begin UFactory Interface
 	virtual UObject* FactoryCreateNew(UClass* Class,UObject* InParent,FName Name,EObjectFlags Flags,UObject* Context,FFeedbackContext* Warn) override;
 	// Begin UFactory Interface	
-};
+};
\ No newline at end of file
diff --git a/Engine/Source/Editor/UnrealEd/Classes/Factories/NiagaraEffectFactoryNew.h b/Engine/Source/Editor/UnrealEd/Classes/Factories/NiagaraEffectFactoryNew.h
--- a/Engine/Source/Editor/UnrealEd/Classes/Factories/NiagaraEffectFactoryNew.h
+++ b/Engine/Source/Editor/UnrealEd/Classes/Factories/NiagaraEffectFactoryNew.h
@@ -6,9 +6,7 @@
 UCLASS(hidecategories = Object)
 class UNiagaraEffectFactoryNew : public UFactory
 {
-	GENERATED_BODY()
-public:
-	UNiagaraEffectFactoryNew(const FObjectInitializer& ObjectInitializer = FObjectInitializer::Get());
+	GENERATED_UCLASS_BODY()
 
 	// Begin UFactory Interface
 	virtual UObject* FactoryCreateNew(UClass* Class, UObject* InParent, FName Name, EObjectFlags Flags, UObject* Context, FFeedbackContext* Warn) override;
diff --git a/Engine/Source/Editor/UnrealEd/Classes/Factories/ReimportCurveTableFactory.h b/Engine/Source/Editor/UnrealEd/Classes/Factories/ReimportCurveTableFactory.h
--- a/Engine/Source/Editor/UnrealEd/Classes/Factories/ReimportCurveTableFactory.h
+++ b/Engine/Source/Editor/UnrealEd/Classes/Factories/ReimportCurveTableFactory.h
@@ -6,9 +6,7 @@
 UCLASS()
 class UReimportCurveTableFactory : public UCSVImportFactory, public FReimportHandler
 {
-	GENERATED_BODY()
-public:
-	UReimportCurveTableFactory(const FObjectInitializer& ObjectInitializer = FObjectInitializer::Get());
+	GENERATED_UCLASS_BODY()
 
 	// Begin FReimportHandler interface
 	virtual bool CanReimport( UObject* Obj, TArray<FString>& OutFilenames ) override;
diff --git a/Engine/Source/Editor/UnrealEd/Classes/Factories/ReimportTextureFactory.h b/Engine/Source/Editor/UnrealEd/Classes/Factories/ReimportTextureFactory.h
--- a/Engine/Source/Editor/UnrealEd/Classes/Factories/ReimportTextureFactory.h
+++ b/Engine/Source/Editor/UnrealEd/Classes/Factories/ReimportTextureFactory.h
@@ -10,9 +10,7 @@
 UCLASS(hidecategories=(LightMap, DitherMipMaps, LODGroup), collapsecategories)
 class UReimportTextureFactory : public UTextureFactory, public FReimportHandler
 {
-	GENERATED_BODY()
-public:
-	UReimportTextureFactory(const FObjectInitializer& ObjectInitializer = FObjectInitializer::Get());
+	GENERATED_UCLASS_BODY()
 
 	UPROPERTY()
 	class UTexture* pOriginalTex;
diff --git a/Engine/Source/Editor/UnrealEd/Classes/Factories/StructureFactory.h b/Engine/Source/Editor/UnrealEd/Classes/Factories/StructureFactory.h
--- a/Engine/Source/Editor/UnrealEd/Classes/Factories/StructureFactory.h
+++ b/Engine/Source/Editor/UnrealEd/Classes/Factories/StructureFactory.h
@@ -7,9 +7,7 @@
 UCLASS(hidecategories=Object, MinimalAPI)
 class UStructureFactory : public UFactory
 {
-	GENERATED_BODY()
-public:
-	UNREALED_API UStructureFactory(const FObjectInitializer& ObjectInitializer = FObjectInitializer::Get());
+	GENERATED_UCLASS_BODY()
 
 	// Begin UFactory Interface
 	virtual UObject* FactoryCreateNew(UClass* Class,UObject* InParent,FName Name,EObjectFlags Flags,UObject* Context,FFeedbackContext* Warn) override;
diff --git a/Engine/Source/Editor/UnrealEd/Classes/Factories/TextureCubeThumbnailRenderer.h b/Engine/Source/Editor/UnrealEd/Classes/Factories/TextureCubeThumbnailRenderer.h
--- a/Engine/Source/Editor/UnrealEd/Classes/Factories/TextureCubeThumbnailRenderer.h
+++ b/Engine/Source/Editor/UnrealEd/Classes/Factories/TextureCubeThumbnailRenderer.h
@@ -11,9 +11,7 @@
 UCLASS()
 class UTextureCubeThumbnailRenderer : public UTextureThumbnailRenderer
 {
-	GENERATED_BODY()
-public:
-	UTextureCubeThumbnailRenderer(const FObjectInitializer& ObjectInitializer = FObjectInitializer::Get());
+	GENERATED_UCLASS_BODY()
 
 
 	// Begin UThumbnailRenderer Interface
diff --git a/Engine/Source/Editor/UnrealEd/Classes/MaterialEditor/DEditorStaticSwitchParameterValue.h b/Engine/Source/Editor/UnrealEd/Classes/MaterialEditor/DEditorStaticSwitchParameterValue.h
--- a/Engine/Source/Editor/UnrealEd/Classes/MaterialEditor/DEditorStaticSwitchParameterValue.h
+++ b/Engine/Source/Editor/UnrealEd/Classes/MaterialEditor/DEditorStaticSwitchParameterValue.h
@@ -9,9 +9,7 @@
 UCLASS(hidecategories=Object, collapsecategories)
 class UNREALED_API UDEditorStaticSwitchParameterValue : public UDEditorParameterValue
 {
-	GENERATED_BODY()
-public:
-	UDEditorStaticSwitchParameterValue(const FObjectInitializer& ObjectInitializer = FObjectInitializer::Get());
+	GENERATED_UCLASS_BODY()
 
 	UPROPERTY(EditAnywhere, Category=DEditorStaticSwitchParameterValue)
 	uint32 ParameterValue:1;
diff --git a/Engine/Source/Editor/UnrealEd/Classes/MaterialGraph/MaterialGraphNode_Comment.h b/Engine/Source/Editor/UnrealEd/Classes/MaterialGraph/MaterialGraphNode_Comment.h
--- a/Engine/Source/Editor/UnrealEd/Classes/MaterialGraph/MaterialGraphNode_Comment.h
+++ b/Engine/Source/Editor/UnrealEd/Classes/MaterialGraph/MaterialGraphNode_Comment.h
@@ -9,9 +9,7 @@
 UCLASS(MinimalAPI)
 class UMaterialGraphNode_Comment : public UEdGraphNode_Comment
 {
-	GENERATED_BODY()
-public:
-	UNREALED_API UMaterialGraphNode_Comment(const FObjectInitializer& ObjectInitializer = FObjectInitializer::Get());
+	GENERATED_UCLASS_BODY()
 
 	/** Material Comment that this node represents */
 	UPROPERTY()
diff --git a/Engine/Source/Editor/UnrealEd/Classes/MaterialGraph/MaterialGraphSchema.h b/Engine/Source/Editor/UnrealEd/Classes/MaterialGraph/MaterialGraphSchema.h
--- a/Engine/Source/Editor/UnrealEd/Classes/MaterialGraph/MaterialGraphSchema.h
+++ b/Engine/Source/Editor/UnrealEd/Classes/MaterialGraph/MaterialGraphSchema.h
@@ -8,7 +8,7 @@
 USTRUCT()
 struct UNREALED_API FMaterialGraphSchemaAction_NewNode : public FEdGraphSchemaAction
 {
-	GENERATED_BODY()
+	GENERATED_USTRUCT_BODY();
 
 	/** Class of expression we want to create */
 	UPROPERTY()
@@ -45,7 +45,7 @@ struct UNREALED_API FMaterialGraphSchemaAction_NewNode : public FEdGraphSchemaAc
 USTRUCT()
 struct UNREALED_API FMaterialGraphSchemaAction_NewFunctionCall : public FEdGraphSchemaAction
 {
-	GENERATED_BODY()
+	GENERATED_USTRUCT_BODY();
 
 	/** Path to function that we want to call */
 	UPROPERTY()
@@ -72,7 +72,7 @@ struct UNREALED_API FMaterialGraphSchemaAction_NewFunctionCall : public FEdGraph
 USTRUCT()
 struct UNREALED_API FMaterialGraphSchemaAction_NewComment : public FEdGraphSchemaAction
 {
-	GENERATED_BODY()
+	GENERATED_USTRUCT_BODY();
 
 	// Simple type info
 	static FName StaticGetTypeId() {static FName Type("FMaterialGraphSchemaAction_NewComment"); return Type;}
@@ -95,7 +95,7 @@ struct UNREALED_API FMaterialGraphSchemaAction_NewComment : public FEdGraphSchem
 USTRUCT()
 struct UNREALED_API FMaterialGraphSchemaAction_Paste : public FEdGraphSchemaAction
 {
-	GENERATED_BODY()
+	GENERATED_USTRUCT_BODY();
 
 	// Simple type info
 	static FName StaticGetTypeId() {static FName Type("FMaterialGraphSchemaAction_Paste"); return Type;}
@@ -117,9 +117,7 @@ struct UNREALED_API FMaterialGraphSchemaAction_Paste : public FEdGraphSchemaActi
 UCLASS(MinimalAPI)
 class UMaterialGraphSchema : public UEdGraphSchema
 {
-	GENERATED_BODY()
-public:
-	UNREALED_API UMaterialGraphSchema(const FObjectInitializer& ObjectInitializer = FObjectInitializer::Get());
+	GENERATED_UCLASS_BODY()
 
 	// Allowable PinType.PinCategory values
 	UPROPERTY()
diff --git a/Engine/Source/Editor/UnrealEd/Classes/Niagara/NiagaraScriptSource.h b/Engine/Source/Editor/UnrealEd/Classes/Niagara/NiagaraScriptSource.h
--- a/Engine/Source/Editor/UnrealEd/Classes/Niagara/NiagaraScriptSource.h
+++ b/Engine/Source/Editor/UnrealEd/Classes/Niagara/NiagaraScriptSource.h
@@ -9,9 +9,7 @@
 UCLASS(MinimalAPI)
 class UNiagaraScriptSource : public UNiagaraScriptSourceBase
 {
-	GENERATED_BODY()
-public:
-	UNREALED_API UNiagaraScriptSource(const FObjectInitializer& ObjectInitializer = FObjectInitializer::Get());
+	GENERATED_UCLASS_BODY()
 
 	/** Graph for particle update expression */
 	UPROPERTY()
diff --git a/Engine/Source/Editor/UnrealEd/Classes/Preferences/LightmassOptionsObject.h b/Engine/Source/Editor/UnrealEd/Classes/Preferences/LightmassOptionsObject.h
--- a/Engine/Source/Editor/UnrealEd/Classes/Preferences/LightmassOptionsObject.h
+++ b/Engine/Source/Editor/UnrealEd/Classes/Preferences/LightmassOptionsObject.h
@@ -15,9 +15,7 @@
 UCLASS(hidecategories=Object, editinlinenew)
 class ULightmassOptionsObject : public UObject
 {
-	GENERATED_BODY()
-public:
-	ULightmassOptionsObject(const FObjectInitializer& ObjectInitializer = FObjectInitializer::Get());
+	GENERATED_UCLASS_BODY()
 
 	UPROPERTY(EditAnywhere, Category=Debug)
 	struct FLightmassDebugOptions DebugSettings;
diff --git a/Engine/Source/Editor/UnrealEd/Classes/Preferences/UnrealEdKeyBindings.h b/Engine/Source/Editor/UnrealEd/Classes/Preferences/UnrealEdKeyBindings.h
--- a/Engine/Source/Editor/UnrealEd/Classes/Preferences/UnrealEdKeyBindings.h
+++ b/Engine/Source/Editor/UnrealEd/Classes/Preferences/UnrealEdKeyBindings.h
@@ -11,7 +11,7 @@
 USTRUCT()
 struct FEditorKeyBinding
 {
-	GENERATED_BODY()
+	GENERATED_USTRUCT_BODY()
 
 	UPROPERTY()
 	uint32 bCtrlDown:1;
@@ -38,9 +38,7 @@ struct FEditorKeyBinding
 UCLASS(Config=EditorKeyBindings)
 class UUnrealEdKeyBindings : public UObject
 {
-	GENERATED_BODY()
-public:
-	UUnrealEdKeyBindings(const FObjectInitializer& ObjectInitializer = FObjectInitializer::Get());
+	GENERATED_UCLASS_BODY()
 
 	/** Array of keybindings */
 	UPROPERTY(config)
diff --git a/Engine/Source/Editor/UnrealEd/Classes/Settings/LevelEditorPlaySettings.h b/Engine/Source/Editor/UnrealEd/Classes/Settings/LevelEditorPlaySettings.h
--- a/Engine/Source/Editor/UnrealEd/Classes/Settings/LevelEditorPlaySettings.h
+++ b/Engine/Source/Editor/UnrealEd/Classes/Settings/LevelEditorPlaySettings.h
@@ -84,7 +84,7 @@ enum EPlayNetMode
 USTRUCT()
 struct FPlayScreenResolution
 {
-	GENERATED_BODY()
+	GENERATED_USTRUCT_BODY()
 
 public:
 
@@ -113,9 +113,7 @@ UCLASS(config=EditorUserSettings)
 class UNREALED_API ULevelEditorPlaySettings
 	: public UObject
 {
-	GENERATED_BODY()
-public:
-	ULevelEditorPlaySettings(const FObjectInitializer& ObjectInitializer = FObjectInitializer::Get());
+	GENERATED_UCLASS_BODY()
 
 public:
 
diff --git a/Engine/Source/Editor/UnrealEd/Classes/SoundClassGraph/SoundClassGraphSchema.h b/Engine/Source/Editor/UnrealEd/Classes/SoundClassGraph/SoundClassGraphSchema.h
--- a/Engine/Source/Editor/UnrealEd/Classes/SoundClassGraph/SoundClassGraphSchema.h
+++ b/Engine/Source/Editor/UnrealEd/Classes/SoundClassGraph/SoundClassGraphSchema.h
@@ -7,7 +7,7 @@
 USTRUCT()
 struct UNREALED_API FSoundClassGraphSchemaAction_NewNode : public FEdGraphSchemaAction
 {
-	GENERATED_BODY()
+	GENERATED_USTRUCT_BODY();
 
 	// Simple type info
 	static FName StaticGetTypeId() {static FName Type("FSoundClassGraphSchemaAction_NewNode"); return Type;}
@@ -34,9 +34,7 @@ struct UNREALED_API FSoundClassGraphSchemaAction_NewNode : public FEdGraphSchema
 UCLASS(MinimalAPI)
 class USoundClassGraphSchema : public UEdGraphSchema
 {
-	GENERATED_BODY()
-public:
-	UNREALED_API USoundClassGraphSchema(const FObjectInitializer& ObjectInitializer = FObjectInitializer::Get());
+	GENERATED_UCLASS_BODY()
 
 	/** Check whether connecting these pins would cause a loop */
 	bool ConnectionCausesLoop(const UEdGraphPin* InputPin, const UEdGraphPin* OutputPin) const;
diff --git a/Engine/Source/Editor/UnrealEd/Classes/ThumbnailRendering/AnimSequenceThumbnailRenderer.h b/Engine/Source/Editor/UnrealEd/Classes/ThumbnailRendering/AnimSequenceThumbnailRenderer.h
--- a/Engine/Source/Editor/UnrealEd/Classes/ThumbnailRendering/AnimSequenceThumbnailRenderer.h
+++ b/Engine/Source/Editor/UnrealEd/Classes/ThumbnailRendering/AnimSequenceThumbnailRenderer.h
@@ -11,9 +11,7 @@
 UCLASS(config=Editor, MinimalAPI)
 class UAnimSequenceThumbnailRenderer : public UDefaultSizedThumbnailRenderer
 {
-	GENERATED_BODY()
-public:
-	UNREALED_API UAnimSequenceThumbnailRenderer(const FObjectInitializer& ObjectInitializer = FObjectInitializer::Get());
+	GENERATED_UCLASS_BODY()
 
 
 	// Begin UThumbnailRenderer Object
diff --git a/Engine/Source/Editor/UnrealEd/Classes/ThumbnailRendering/BlendSpaceThumbnailRenderer.h b/Engine/Source/Editor/UnrealEd/Classes/ThumbnailRendering/BlendSpaceThumbnailRenderer.h
--- a/Engine/Source/Editor/UnrealEd/Classes/ThumbnailRendering/BlendSpaceThumbnailRenderer.h
+++ b/Engine/Source/Editor/UnrealEd/Classes/ThumbnailRendering/BlendSpaceThumbnailRenderer.h
@@ -11,9 +11,7 @@
 UCLASS(config=Editor, MinimalAPI)
 class UBlendSpaceThumbnailRenderer : public UDefaultSizedThumbnailRenderer
 {
-	GENERATED_BODY()
-public:
-	UNREALED_API UBlendSpaceThumbnailRenderer(const FObjectInitializer& ObjectInitializer = FObjectInitializer::Get());
+	GENERATED_UCLASS_BODY()
 
 
 	// Begin UThumbnailRenderer Object
diff --git a/Engine/Source/Editor/UnrealEd/Classes/ThumbnailRendering/ClassThumbnailRenderer.h b/Engine/Source/Editor/UnrealEd/Classes/ThumbnailRendering/ClassThumbnailRenderer.h
--- a/Engine/Source/Editor/UnrealEd/Classes/ThumbnailRendering/ClassThumbnailRenderer.h
+++ b/Engine/Source/Editor/UnrealEd/Classes/ThumbnailRendering/ClassThumbnailRenderer.h
@@ -6,9 +6,7 @@
 UCLASS(config=Editor,MinimalAPI)
 class UClassThumbnailRenderer : public UDefaultSizedThumbnailRenderer
 {
-	GENERATED_BODY()
-public:
-	UNREALED_API UClassThumbnailRenderer(const FObjectInitializer& ObjectInitializer = FObjectInitializer::Get());
+	GENERATED_UCLASS_BODY()
 
 
 	// Begin UThumbnailRenderer Object
diff --git a/Engine/Source/Editor/UnrealEd/Classes/ThumbnailRendering/SceneThumbnailInfoWithPrimitive.h b/Engine/Source/Editor/UnrealEd/Classes/ThumbnailRendering/SceneThumbnailInfoWithPrimitive.h
--- a/Engine/Source/Editor/UnrealEd/Classes/ThumbnailRendering/SceneThumbnailInfoWithPrimitive.h
+++ b/Engine/Source/Editor/UnrealEd/Classes/ThumbnailRendering/SceneThumbnailInfoWithPrimitive.h
@@ -14,9 +14,7 @@
 UCLASS(MinimalAPI)
 class USceneThumbnailInfoWithPrimitive : public USceneThumbnailInfo
 {
-	GENERATED_BODY()
-public:
-	UNREALED_API USceneThumbnailInfoWithPrimitive(const FObjectInitializer& ObjectInitializer = FObjectInitializer::Get());
+	GENERATED_UCLASS_BODY()
 
 	/** The type of primitive used in this thumbnail */
 	UPROPERTY(EditAnywhere, Category=Thumbnail)
diff --git a/Engine/Source/Editor/WorldBrowser/Private/Tiles/WorldTileDetails.h b/Engine/Source/Editor/WorldBrowser/Private/Tiles/WorldTileDetails.h
--- a/Engine/Source/Editor/WorldBrowser/Private/Tiles/WorldTileDetails.h
+++ b/Engine/Source/Editor/WorldBrowser/Private/Tiles/WorldTileDetails.h
@@ -12,7 +12,7 @@ class FWorldTileModel;
 USTRUCT()
 struct FTileLODEntryDetails
 {
-	GENERATED_BODY()
+	GENERATED_USTRUCT_BODY()
 	
 	// Maximum deviation of details percentage
 	UPROPERTY(Category=LODDetails, VisibleAnywhere)
@@ -37,9 +37,7 @@ struct FTileLODEntryDetails
 UCLASS()
 class UWorldTileDetails : public UObject
 {
-	GENERATED_BODY()
-public:
-	UWorldTileDetails(const FObjectInitializer& ObjectInitializer = FObjectInitializer::Get());
+	GENERATED_UCLASS_BODY()
 	
 	// Whether this tile properties can be edited via details panel
 	UPROPERTY(Category=Tile, VisibleAnywhere)
diff --git a/Engine/Source/Runtime/AIModule/Classes/AIResourceInterface.h b/Engine/Source/Runtime/AIModule/Classes/AIResourceInterface.h
--- a/Engine/Source/Runtime/AIModule/Classes/AIResourceInterface.h
+++ b/Engine/Source/Runtime/AIModule/Classes/AIResourceInterface.h
@@ -7,15 +7,12 @@
 UINTERFACE(MinimalAPI, meta=(CannotImplementInterfaceInBlueprint))
 class UAIResourceInterface : public UInterface
 {
-	GENERATED_BODY()
-public:
-	AIMODULE_API UAIResourceInterface(const FObjectInitializer& ObjectInitializer = FObjectInitializer::Get());
+	GENERATED_UINTERFACE_BODY()
 };
 
 class IAIResourceInterface
 {
-	GENERATED_BODY()
-public:
+	GENERATED_IINTERFACE_BODY()
 
 	/** If resource is lockable lock it with indicated priority */
 	virtual void LockResource(EAIRequestPriority::Type LockSource) {}
diff --git a/Engine/Source/Runtime/AIModule/Classes/BehaviorTree/BTNode.h b/Engine/Source/Runtime/AIModule/Classes/BehaviorTree/BTNode.h
--- a/Engine/Source/Runtime/AIModule/Classes/BehaviorTree/BTNode.h
+++ b/Engine/Source/Runtime/AIModule/Classes/BehaviorTree/BTNode.h
@@ -22,9 +22,7 @@ struct FBTInstancedNodeMemory
 UCLASS(Abstract,config=Game)
 class AIMODULE_API UBTNode : public UObject
 {
-	GENERATED_BODY()
-public:
-	UBTNode(const FObjectInitializer& ObjectInitializer = FObjectInitializer::Get());
+	GENERATED_UCLASS_BODY()
 
 	virtual UWorld* GetWorld() const override;
 
diff --git a/Engine/Source/Runtime/AIModule/Classes/BehaviorTree/BTService.h b/Engine/Source/Runtime/AIModule/Classes/BehaviorTree/BTService.h
--- a/Engine/Source/Runtime/AIModule/Classes/BehaviorTree/BTService.h
+++ b/Engine/Source/Runtime/AIModule/Classes/BehaviorTree/BTService.h
@@ -33,9 +33,7 @@
 UCLASS(Abstract)
 class AIMODULE_API UBTService : public UBTAuxiliaryNode
 {
-	GENERATED_BODY()
-public:
-	UBTService(const FObjectInitializer& ObjectInitializer = FObjectInitializer::Get());
+	GENERATED_UCLASS_BODY()
 
 	virtual FString GetStaticDescription() const override;
 
diff --git a/Engine/Source/Runtime/AIModule/Classes/BehaviorTree/Decorators/BTDecorator_ConditionalLoop.h b/Engine/Source/Runtime/AIModule/Classes/BehaviorTree/Decorators/BTDecorator_ConditionalLoop.h
--- a/Engine/Source/Runtime/AIModule/Classes/BehaviorTree/Decorators/BTDecorator_ConditionalLoop.h
+++ b/Engine/Source/Runtime/AIModule/Classes/BehaviorTree/Decorators/BTDecorator_ConditionalLoop.h
@@ -11,9 +11,7 @@
 UCLASS(HideCategories=(FlowControl))
 class AIMODULE_API UBTDecorator_ConditionalLoop : public UBTDecorator_Blackboard
 {
-	GENERATED_BODY()
-public:
-	UBTDecorator_ConditionalLoop(const FObjectInitializer& ObjectInitializer = FObjectInitializer::Get());
+	GENERATED_UCLASS_BODY()
 
 	virtual bool CalculateRawConditionValue(UBehaviorTreeComponent& OwnerComp, uint8* NodeMemory) const override;
 	virtual void OnBlackboardChange(const UBlackboardComponent& Blackboard, FBlackboard::FKey ChangedKeyID) override;
diff --git a/Engine/Source/Runtime/AIModule/Classes/BehaviorTree/Decorators/BTDecorator_ReachedMoveGoal.h b/Engine/Source/Runtime/AIModule/Classes/BehaviorTree/Decorators/BTDecorator_ReachedMoveGoal.h
--- a/Engine/Source/Runtime/AIModule/Classes/BehaviorTree/Decorators/BTDecorator_ReachedMoveGoal.h
+++ b/Engine/Source/Runtime/AIModule/Classes/BehaviorTree/Decorators/BTDecorator_ReachedMoveGoal.h
@@ -11,9 +11,7 @@
 UCLASS()
 class AIMODULE_API UBTDecorator_ReachedMoveGoal : public UBTDecorator
 {
-	GENERATED_BODY()
-public:
-	UBTDecorator_ReachedMoveGoal(const FObjectInitializer& ObjectInitializer = FObjectInitializer::Get());
+	GENERATED_UCLASS_BODY()
 
 	virtual bool CalculateRawConditionValue(UBehaviorTreeComponent& OwnerComp, uint8* NodeMemory) const override;
 
diff --git a/Engine/Source/Runtime/AIModule/Classes/BehaviorTree/Tasks/BTTask_BlackboardBase.h b/Engine/Source/Runtime/AIModule/Classes/BehaviorTree/Tasks/BTTask_BlackboardBase.h
--- a/Engine/Source/Runtime/AIModule/Classes/BehaviorTree/Tasks/BTTask_BlackboardBase.h
+++ b/Engine/Source/Runtime/AIModule/Classes/BehaviorTree/Tasks/BTTask_BlackboardBase.h
@@ -10,9 +10,7 @@ class UBehaviorTree;
 UCLASS(Abstract)
 class AIMODULE_API UBTTask_BlackboardBase : public UBTTaskNode
 {
-	GENERATED_BODY()
-public:
-	UBTTask_BlackboardBase(const FObjectInitializer& ObjectInitializer = FObjectInitializer::Get());
+	GENERATED_UCLASS_BODY()
 
 	/** initialize any asset related data */
 	virtual void InitializeFromAsset(UBehaviorTree& Asset) override;
diff --git a/Engine/Source/Runtime/AIModule/Classes/BehaviorTree/Tasks/BTTask_BlueprintBase.h b/Engine/Source/Runtime/AIModule/Classes/BehaviorTree/Tasks/BTTask_BlueprintBase.h
--- a/Engine/Source/Runtime/AIModule/Classes/BehaviorTree/Tasks/BTTask_BlueprintBase.h
+++ b/Engine/Source/Runtime/AIModule/Classes/BehaviorTree/Tasks/BTTask_BlueprintBase.h
@@ -18,9 +18,7 @@ class FBehaviorBlueprintDetails;
 UCLASS(Abstract, Blueprintable)
 class AIMODULE_API UBTTask_BlueprintBase : public UBTTaskNode
 {
-	GENERATED_BODY()
-public:
-	UBTTask_BlueprintBase(const FObjectInitializer& ObjectInitializer = FObjectInitializer::Get());
+	GENERATED_UCLASS_BODY()
 
 	/** setup node name */
 	virtual void PostInitProperties() override;
diff --git a/Engine/Source/Runtime/AIModule/Classes/BehaviorTree/Tasks/BTTask_PawnActionBase.h b/Engine/Source/Runtime/AIModule/Classes/BehaviorTree/Tasks/BTTask_PawnActionBase.h
--- a/Engine/Source/Runtime/AIModule/Classes/BehaviorTree/Tasks/BTTask_PawnActionBase.h
+++ b/Engine/Source/Runtime/AIModule/Classes/BehaviorTree/Tasks/BTTask_PawnActionBase.h
@@ -20,9 +20,7 @@ class UPawnAction;
 UCLASS(Abstract)
 class AIMODULE_API UBTTask_PawnActionBase : public UBTTaskNode
 {
-	GENERATED_BODY()
-public:
-	UBTTask_PawnActionBase(const FObjectInitializer& ObjectInitializer = FObjectInitializer::Get());
+	GENERATED_UCLASS_BODY()
 
 	virtual EBTNodeResult::Type AbortTask(UBehaviorTreeComponent& OwnerComp, uint8* NodeMemory) override;
 
diff --git a/Engine/Source/Runtime/AIModule/Classes/Blueprint/AIAsyncTaskBlueprintProxy.h b/Engine/Source/Runtime/AIModule/Classes/Blueprint/AIAsyncTaskBlueprintProxy.h
--- a/Engine/Source/Runtime/AIModule/Classes/Blueprint/AIAsyncTaskBlueprintProxy.h
+++ b/Engine/Source/Runtime/AIModule/Classes/Blueprint/AIAsyncTaskBlueprintProxy.h
@@ -14,9 +14,7 @@ DECLARE_DYNAMIC_MULTICAST_DELEGATE_OneParam(FOAISimpleDelegate, EPathFollowingRe
 UCLASS(MinimalAPI)
 class UAIAsyncTaskBlueprintProxy : public UObject
 {
-	GENERATED_BODY()
-public:
-	AIMODULE_API UAIAsyncTaskBlueprintProxy(const FObjectInitializer& ObjectInitializer = FObjectInitializer::Get());
+	GENERATED_UCLASS_BODY()
 
 	UPROPERTY(BlueprintAssignable)
 	FOAISimpleDelegate	OnSuccess;
diff --git a/Engine/Source/Runtime/AIModule/Classes/EnvironmentQuery/EQSQueryResultSourceInterface.h b/Engine/Source/Runtime/AIModule/Classes/EnvironmentQuery/EQSQueryResultSourceInterface.h
--- a/Engine/Source/Runtime/AIModule/Classes/EnvironmentQuery/EQSQueryResultSourceInterface.h
+++ b/Engine/Source/Runtime/AIModule/Classes/EnvironmentQuery/EQSQueryResultSourceInterface.h
@@ -7,15 +7,12 @@
 UINTERFACE(MinimalAPI, meta=(CannotImplementInterfaceInBlueprint))
 class UEQSQueryResultSourceInterface : public UInterface
 {
-	GENERATED_BODY()
-public:
-	AIMODULE_API UEQSQueryResultSourceInterface(const FObjectInitializer& ObjectInitializer = FObjectInitializer::Get());
+	GENERATED_UINTERFACE_BODY()
 };
 
 class IEQSQueryResultSourceInterface
 {
-	GENERATED_BODY()
-public:
+	GENERATED_IINTERFACE_BODY()
 
 	virtual const struct FEnvQueryResult* GetQueryResult() const { return NULL; }
 	virtual const struct FEnvQueryInstance* GetQueryInstance() const { return NULL; }
diff --git a/Engine/Source/Runtime/AIModule/Classes/EnvironmentQuery/EnvQuery.h b/Engine/Source/Runtime/AIModule/Classes/EnvironmentQuery/EnvQuery.h
--- a/Engine/Source/Runtime/AIModule/Classes/EnvironmentQuery/EnvQuery.h
+++ b/Engine/Source/Runtime/AIModule/Classes/EnvironmentQuery/EnvQuery.h
@@ -13,9 +13,7 @@ class UEdGraph;
 UCLASS()
 class AIMODULE_API UEnvQuery : public UObject
 {
-	GENERATED_BODY()
-public:
-	UEnvQuery(const FObjectInitializer& ObjectInitializer = FObjectInitializer::Get());
+	GENERATED_UCLASS_BODY()
 
 #if WITH_EDITORONLY_DATA
 	/** Graph for query */
diff --git a/Engine/Source/Runtime/AIModule/Classes/EnvironmentQuery/EnvQueryGenerator.h b/Engine/Source/Runtime/AIModule/Classes/EnvironmentQuery/EnvQueryGenerator.h
--- a/Engine/Source/Runtime/AIModule/Classes/EnvironmentQuery/EnvQueryGenerator.h
+++ b/Engine/Source/Runtime/AIModule/Classes/EnvironmentQuery/EnvQueryGenerator.h
@@ -19,9 +19,7 @@ namespace EnvQueryGeneratorVersion
 UCLASS(EditInlineNew, Abstract, meta = (Category = "Generators"))
 class AIMODULE_API UEnvQueryGenerator : public UEnvQueryNode
 {
-	GENERATED_BODY()
-public:
-	UEnvQueryGenerator(const FObjectInitializer& ObjectInitializer = FObjectInitializer::Get());
+	GENERATED_UCLASS_BODY()
 
 	UPROPERTY(EditDefaultsOnly, Category=Option)
 	FString OptionName;
diff --git a/Engine/Source/Runtime/AIModule/Classes/EnvironmentQuery/Generators/EnvQueryGenerator_Composite.h b/Engine/Source/Runtime/AIModule/Classes/EnvironmentQuery/Generators/EnvQueryGenerator_Composite.h
--- a/Engine/Source/Runtime/AIModule/Classes/EnvironmentQuery/Generators/EnvQueryGenerator_Composite.h
+++ b/Engine/Source/Runtime/AIModule/Classes/EnvironmentQuery/Generators/EnvQueryGenerator_Composite.h
@@ -15,9 +15,7 @@ struct FEnvQueryInstance;
 UCLASS(meta = (DisplayName = "Composite"))
 class AIMODULE_API UEnvQueryGenerator_Composite : public UEnvQueryGenerator
 {
-	GENERATED_BODY()
-public:
-	UEnvQueryGenerator_Composite(const FObjectInitializer& ObjectInitializer = FObjectInitializer::Get());
+	GENERATED_UCLASS_BODY()
 
 	UPROPERTY(EditDefaultsOnly, Instanced, Category = Generator)
 	TArray<UEnvQueryGenerator*> Generators;
diff --git a/Engine/Source/Runtime/AIModule/Classes/EnvironmentQuery/Items/EnvQueryItemType_Direction.h b/Engine/Source/Runtime/AIModule/Classes/EnvironmentQuery/Items/EnvQueryItemType_Direction.h
--- a/Engine/Source/Runtime/AIModule/Classes/EnvironmentQuery/Items/EnvQueryItemType_Direction.h
+++ b/Engine/Source/Runtime/AIModule/Classes/EnvironmentQuery/Items/EnvQueryItemType_Direction.h
@@ -9,9 +9,7 @@ struct FEnvQueryContextData;
 UCLASS()
 class AIMODULE_API UEnvQueryItemType_Direction : public UEnvQueryItemType_VectorBase
 {
-	GENERATED_BODY()
-public:
-	UEnvQueryItemType_Direction(const FObjectInitializer& ObjectInitializer = FObjectInitializer::Get());
+	GENERATED_UCLASS_BODY()
 
 	static FVector GetValue(const uint8* RawData);
 	static void SetValue(uint8* RawData, const FVector& Value);
diff --git a/Engine/Source/Runtime/AIModule/Classes/Navigation/CrowdFollowingComponent.h b/Engine/Source/Runtime/AIModule/Classes/Navigation/CrowdFollowingComponent.h
--- a/Engine/Source/Runtime/AIModule/Classes/Navigation/CrowdFollowingComponent.h
+++ b/Engine/Source/Runtime/AIModule/Classes/Navigation/CrowdFollowingComponent.h
@@ -24,9 +24,7 @@ namespace ECrowdAvoidanceQuality
 UCLASS(BlueprintType)
 class AIMODULE_API UCrowdFollowingComponent : public UPathFollowingComponent, public ICrowdAgentInterface
 {
-	GENERATED_BODY()
-public:
-	UCrowdFollowingComponent(const FObjectInitializer& ObjectInitializer = FObjectInitializer::Get());
+	GENERATED_UCLASS_BODY()
 
 	UPROPERTY()
 	FVector CrowdAgentMoveDirection;
diff --git a/Engine/Source/Runtime/AIModule/Classes/Navigation/CrowdManager.h b/Engine/Source/Runtime/AIModule/Classes/Navigation/CrowdManager.h
--- a/Engine/Source/Runtime/AIModule/Classes/Navigation/CrowdManager.h
+++ b/Engine/Source/Runtime/AIModule/Classes/Navigation/CrowdManager.h
@@ -52,7 +52,7 @@ class dtObstacleAvoidanceDebugData;
 USTRUCT()
 struct AIMODULE_API FCrowdAvoidanceConfig
 {
-	GENERATED_BODY()
+	GENERATED_USTRUCT_BODY()
 
 	UPROPERTY(EditAnywhere, Category=Crowd)
 	float VelocityBias;
@@ -98,7 +98,7 @@ struct AIMODULE_API FCrowdAvoidanceConfig
 USTRUCT()
 struct AIMODULE_API FCrowdAvoidanceSamplingPattern
 {
-	GENERATED_BODY()
+	GENERATED_USTRUCT_BODY()
 
 	UPROPERTY(EditAnywhere, Category=Crowd)
 	TArray<float> Angles;
@@ -153,9 +153,7 @@ struct FCrowdTickHelper : FTickableGameObject
 UCLASS(config = Engine, defaultconfig)
 class AIMODULE_API UCrowdManager : public UObject
 {
-	GENERATED_BODY()
-public:
-	UCrowdManager(const FObjectInitializer& ObjectInitializer = FObjectInitializer::Get());
+	GENERATED_UCLASS_BODY()
 
 	virtual void Tick(float DeltaTime);
 	virtual void BeginDestroy() override;
diff --git a/Engine/Source/Runtime/AIModule/Classes/Navigation/PathFollowingComponent.h b/Engine/Source/Runtime/AIModule/Classes/Navigation/PathFollowingComponent.h
--- a/Engine/Source/Runtime/AIModule/Classes/Navigation/PathFollowingComponent.h
+++ b/Engine/Source/Runtime/AIModule/Classes/Navigation/PathFollowingComponent.h
@@ -90,9 +90,7 @@ namespace EPathFollowingMessage
 UCLASS(config=Engine)
 class AIMODULE_API UPathFollowingComponent : public UActorComponent, public IAIResourceInterface
 {
-	GENERATED_BODY()
-public:
-	UPathFollowingComponent(const FObjectInitializer& ObjectInitializer = FObjectInitializer::Get());
+	GENERATED_UCLASS_BODY()
 
 	DECLARE_DELEGATE_TwoParams(FPostProcessMoveSignature, UPathFollowingComponent* /*comp*/, FVector& /*velocity*/);
 	DECLARE_DELEGATE_OneParam(FRequestCompletedSignature, EPathFollowingResult::Type /*Result*/);
diff --git a/Engine/Source/Runtime/AIModule/Classes/Perception/AISense_Hearing.h b/Engine/Source/Runtime/AIModule/Classes/Perception/AISense_Hearing.h
--- a/Engine/Source/Runtime/AIModule/Classes/Perception/AISense_Hearing.h
+++ b/Engine/Source/Runtime/AIModule/Classes/Perception/AISense_Hearing.h
@@ -11,7 +11,7 @@ class UAISenseConfig_Hearing;
 USTRUCT()
 struct AIMODULE_API FAINoiseEvent
 {	
-	GENERATED_BODY()
+	GENERATED_USTRUCT_BODY()
 
 	typedef class UAISense_Hearing FSenseClass;
 
@@ -39,9 +39,7 @@ struct AIMODULE_API FAINoiseEvent
 UCLASS(ClassGroup=AI, Config=Game)
 class AIMODULE_API UAISense_Hearing : public UAISense
 {
-	GENERATED_BODY()
-public:
-	UAISense_Hearing(const FObjectInitializer& ObjectInitializer = FObjectInitializer::Get());
+	GENERATED_UCLASS_BODY()
 		
 protected:
 	UPROPERTY()
diff --git a/Engine/Source/Runtime/Engine/Classes/AI/Navigation/NavAreas/NavAreaMeta.h b/Engine/Source/Runtime/Engine/Classes/AI/Navigation/NavAreas/NavAreaMeta.h
--- a/Engine/Source/Runtime/Engine/Classes/AI/Navigation/NavAreas/NavAreaMeta.h
+++ b/Engine/Source/Runtime/Engine/Classes/AI/Navigation/NavAreas/NavAreaMeta.h
@@ -11,9 +11,7 @@ class AActor;
 UCLASS(Abstract)
 class ENGINE_API UNavAreaMeta : public UNavArea
 {
-	GENERATED_BODY()
-public:
-	UNavAreaMeta(const FObjectInitializer& ObjectInitializer = FObjectInitializer::Get());
+	GENERATED_UCLASS_BODY()
 
 	/** 
 	 *	Picks an navigation area class that should be used for Actor when
diff --git a/Engine/Source/Runtime/Engine/Classes/AI/Navigation/NavigationData.h b/Engine/Source/Runtime/Engine/Classes/AI/Navigation/NavigationData.h
--- a/Engine/Source/Runtime/Engine/Classes/AI/Navigation/NavigationData.h
+++ b/Engine/Source/Runtime/Engine/Classes/AI/Navigation/NavigationData.h
@@ -17,7 +17,7 @@ class UWorld;
 USTRUCT()
 struct FSupportedAreaData
 {
-	GENERATED_BODY()
+	GENERATED_USTRUCT_BODY()
 
 	UPROPERTY()
 	FString AreaClassName;
@@ -410,9 +410,7 @@ enum class ERuntimeGenerationType : uint8
 UCLASS(config=Engine, defaultconfig, NotBlueprintable, abstract)
 class ENGINE_API ANavigationData : public AActor
 {
-	GENERATED_BODY()
-public:
-	ANavigationData(const FObjectInitializer& ObjectInitializer = FObjectInitializer::Get());
+	GENERATED_UCLASS_BODY()
 	
 	UPROPERTY(transient, duplicatetransient)
 	UPrimitiveComponent* RenderingComp;
diff --git a/Engine/Source/Runtime/Engine/Classes/Animation/AnimBlueprintGeneratedClass.h b/Engine/Source/Runtime/Engine/Classes/Animation/AnimBlueprintGeneratedClass.h
--- a/Engine/Source/Runtime/Engine/Classes/Animation/AnimBlueprintGeneratedClass.h
+++ b/Engine/Source/Runtime/Engine/Classes/Animation/AnimBlueprintGeneratedClass.h
@@ -19,7 +19,7 @@ class  UAnimGraphNode_StateResult;
 USTRUCT()
 struct FStateMachineDebugData
 {
-	GENERATED_BODY()
+	GENERATED_USTRUCT_BODY()
 
 public:
 	// Map from state nodes to their state entry in a state machine
@@ -41,7 +41,7 @@ public:
 USTRUCT()
 struct FAnimationFrameSnapshot
 {
-	GENERATED_BODY()
+	GENERATED_USTRUCT_BODY()
 
 	FAnimationFrameSnapshot()
 #if WITH_EDITORONLY_DATA
@@ -68,7 +68,7 @@ public:
 USTRUCT()
 struct ENGINE_API FAnimBlueprintDebugData
 {
-	GENERATED_BODY()
+	GENERATED_USTRUCT_BODY()
 
 	FAnimBlueprintDebugData()
 #if WITH_EDITORONLY_DATA
@@ -161,9 +161,7 @@ namespace EPropertySearchMode
 UCLASS()
 class ENGINE_API UAnimBlueprintGeneratedClass : public UBlueprintGeneratedClass
 {
-	GENERATED_BODY()
-public:
-	UAnimBlueprintGeneratedClass(const FObjectInitializer& ObjectInitializer = FObjectInitializer::Get());
+	GENERATED_UCLASS_BODY()
 
 	// List of state machines present in this blueprint class
 	UPROPERTY()
@@ -301,4 +299,4 @@ NodeType* GetNodeFromPropertyIndex(UAnimInstance* AnimInstance, const UAnimBluep
 	}
 
 	return NULL;
-}
+}
\ No newline at end of file
diff --git a/Engine/Source/Runtime/Engine/Classes/Animation/AnimCompress_RemoveEverySecondKey.h b/Engine/Source/Runtime/Engine/Classes/Animation/AnimCompress_RemoveEverySecondKey.h
--- a/Engine/Source/Runtime/Engine/Classes/Animation/AnimCompress_RemoveEverySecondKey.h
+++ b/Engine/Source/Runtime/Engine/Classes/Animation/AnimCompress_RemoveEverySecondKey.h
@@ -12,9 +12,7 @@
 UCLASS(MinimalAPI)
 class UAnimCompress_RemoveEverySecondKey : public UAnimCompress
 {
-	GENERATED_BODY()
-public:
-	ENGINE_API UAnimCompress_RemoveEverySecondKey(const FObjectInitializer& ObjectInitializer = FObjectInitializer::Get());
+	GENERATED_UCLASS_BODY()
 
 	/** Animations with fewer than MinKeys will not lose any keys. */
 	UPROPERTY(EditAnywhere, Category=AnimationCompressionAlgorithm_RemoveEverySecondKey)
diff --git a/Engine/Source/Runtime/Engine/Classes/Animation/AnimData/BoneMaskFilter.h b/Engine/Source/Runtime/Engine/Classes/Animation/AnimData/BoneMaskFilter.h
--- a/Engine/Source/Runtime/Engine/Classes/Animation/AnimData/BoneMaskFilter.h
+++ b/Engine/Source/Runtime/Engine/Classes/Animation/AnimData/BoneMaskFilter.h
@@ -12,7 +12,7 @@
 USTRUCT()
 struct FBranchFilter
 {
-	GENERATED_BODY()
+	GENERATED_USTRUCT_BODY()
 
 	/** Bone Name to filter **/
 	UPROPERTY(EditAnywhere, Category=Branch)
@@ -26,7 +26,7 @@ struct FBranchFilter
 USTRUCT()
 struct FInputBlendPose
 {
-	GENERATED_BODY()
+	GENERATED_USTRUCT_BODY()
 
 	/** Bone Name to filter **/
 	UPROPERTY(EditAnywhere, Category=Filter)
@@ -36,9 +36,7 @@ struct FInputBlendPose
 UCLASS(hidecategories=Object, BlueprintType,MinimalAPI)
 class UBoneMaskFilter : public UObject
 {
-	GENERATED_BODY()
-public:
-	ENGINE_API UBoneMaskFilter(const FObjectInitializer& ObjectInitializer = FObjectInitializer::Get());
+	GENERATED_UCLASS_BODY()
 
 	UPROPERTY(EditAnywhere, Category=BoneMask)
 	TArray<FInputBlendPose>	BlendPoses;
diff --git a/Engine/Source/Runtime/Engine/Classes/Animation/AnimationAsset.h b/Engine/Source/Runtime/Engine/Classes/Animation/AnimationAsset.h
--- a/Engine/Source/Runtime/Engine/Classes/Animation/AnimationAsset.h
+++ b/Engine/Source/Runtime/Engine/Classes/Animation/AnimationAsset.h
@@ -17,7 +17,7 @@
 USTRUCT(BlueprintType)
 struct FBlendSampleData
 {
-	GENERATED_BODY()
+	GENERATED_USTRUCT_BODY()
 
 	UPROPERTY()
 	int32 SampleDataIndex;
@@ -59,7 +59,7 @@ struct FBlendSampleData
 USTRUCT()
 struct FBlendFilter
 {
-	GENERATED_BODY()
+	GENERATED_USTRUCT_BODY()
 
 	FFIRFilterTimeBased FilterPerAxis[3];
 
@@ -113,7 +113,7 @@ namespace ERootMotionMode
 USTRUCT()
 struct FAnimExtractContext
 {
-	GENERATED_BODY()
+	GENERATED_USTRUCT_BODY()
 
 	/** Is Root Motion being extracted? */
 	UPROPERTY()
@@ -149,7 +149,7 @@ struct FAnimExtractContext
 USTRUCT()
 struct FAnimTickRecord
 {
-	GENERATED_BODY()
+	GENERATED_USTRUCT_BODY()
 
 	UPROPERTY()
 	class UAnimationAsset* SourceAsset;
@@ -187,7 +187,7 @@ namespace EAnimGroupRole
 USTRUCT()
 struct FAnimGroupInstance
 {
-	GENERATED_BODY()
+	GENERATED_USTRUCT_BODY()
 
 public:
 	// The list of animation players in this group which are going to be evaluated this frame
@@ -240,7 +240,7 @@ public:
 USTRUCT()
 struct FRootMotionMovementParams
 {
-	GENERATED_BODY()
+	GENERATED_USTRUCT_BODY()
 
 	UPROPERTY()
 	bool bHasRootMotion;
@@ -411,7 +411,7 @@ private:
 USTRUCT()
 struct FAnimationGroupReference
 {
-	GENERATED_BODY()
+	GENERATED_USTRUCT_BODY()
 	
 	// The name of the group
 	UPROPERTY(EditAnywhere, Category=Settings)
@@ -430,9 +430,7 @@ struct FAnimationGroupReference
 UCLASS(abstract, MinimalAPI)
 class UAnimationAsset : public UObject
 {
-	GENERATED_BODY()
-public:
-	ENGINE_API UAnimationAsset(const FObjectInitializer& ObjectInitializer = FObjectInitializer::Get());
+	GENERATED_UCLASS_BODY()
 
 private:
 	/** Pointer to the Skeleton this asset can be played on .	*/
diff --git a/Engine/Source/Runtime/Engine/Classes/Animation/AnimationSettings.h b/Engine/Source/Runtime/Engine/Classes/Animation/AnimationSettings.h
--- a/Engine/Source/Runtime/Engine/Classes/Animation/AnimationSettings.h
+++ b/Engine/Source/Runtime/Engine/Classes/Animation/AnimationSettings.h
@@ -16,9 +16,7 @@ UCLASS(config=Engine, defaultconfig)
 class ENGINE_API UAnimationSettings
 	: public UObject
 {
-	GENERATED_BODY()
-public:
-	UAnimationSettings(const FObjectInitializer& ObjectInitializer = FObjectInitializer::Get());
+	GENERATED_UCLASS_BODY()
 
 	// compression upgrade version
 	UPROPERTY(config, VisibleAnywhere, Category = Compression)
diff --git a/Engine/Source/Runtime/Engine/Classes/Animation/BlendSpace1D.h b/Engine/Source/Runtime/Engine/Classes/Animation/BlendSpace1D.h
--- a/Engine/Source/Runtime/Engine/Classes/Animation/BlendSpace1D.h
+++ b/Engine/Source/Runtime/Engine/Classes/Animation/BlendSpace1D.h
@@ -13,9 +13,7 @@
 UCLASS(config=Engine, hidecategories=Object, MinimalAPI, BlueprintType)
 class UBlendSpace1D : public UBlendSpaceBase
 {
-	GENERATED_BODY()
-public:
-	ENGINE_API UBlendSpace1D(const FObjectInitializer& ObjectInitializer = FObjectInitializer::Get());
+	GENERATED_UCLASS_BODY()
 
 public:
 
diff --git a/Engine/Source/Runtime/Engine/Classes/Atmosphere/AtmosphericFogComponent.h b/Engine/Source/Runtime/Engine/Classes/Atmosphere/AtmosphericFogComponent.h
--- a/Engine/Source/Runtime/Engine/Classes/Atmosphere/AtmosphericFogComponent.h
+++ b/Engine/Source/Runtime/Engine/Classes/Atmosphere/AtmosphericFogComponent.h
@@ -7,7 +7,7 @@
 USTRUCT()
 struct FAtmospherePrecomputeParameters
 {
-	GENERATED_BODY()
+	GENERATED_USTRUCT_BODY()
 	
 	/** Rayleigh scattering density height scale, ranges from [0...1] */
 	UPROPERTY(EditAnywhere, BlueprintReadOnly, Category=AtmosphereParam)
@@ -89,9 +89,7 @@ struct FAtmospherePrecomputeParameters
 UCLASS(ClassGroup=Rendering, collapsecategories, hidecategories=(Object, Mobility, Activation, "Components|Activation"), editinlinenew, meta=(BlueprintSpawnableComponent), MinimalAPI)
 class UAtmosphericFogComponent : public USceneComponent
 {
-	GENERATED_BODY()
-public:
-	ENGINE_API UAtmosphericFogComponent(const FObjectInitializer& ObjectInitializer = FObjectInitializer::Get());
+	GENERATED_UCLASS_BODY()
 
 	/** Global scattering factor. */
 	UPROPERTY(EditAnywhere, BlueprintReadWrite, interp, Category=Atmosphere)
diff --git a/Engine/Source/Runtime/Engine/Classes/Components/BillboardComponent.h b/Engine/Source/Runtime/Engine/Classes/Components/BillboardComponent.h
--- a/Engine/Source/Runtime/Engine/Classes/Components/BillboardComponent.h
+++ b/Engine/Source/Runtime/Engine/Classes/Components/BillboardComponent.h
@@ -13,9 +13,7 @@ struct FConvexVolume;
 UCLASS(ClassGroup=Rendering, collapsecategories, hidecategories=(Object,Activation,"Components|Activation",Physics,Collision,Lighting,LOD,Mesh,PhysicsVolume), editinlinenew, meta=(BlueprintSpawnableComponent),MinimalAPI)
 class UBillboardComponent : public UPrimitiveComponent
 {
-	GENERATED_BODY()
-public:
-	ENGINE_API UBillboardComponent(const FObjectInitializer& ObjectInitializer = FObjectInitializer::Get());
+	GENERATED_UCLASS_BODY()
 
 	UPROPERTY(EditAnywhere, BlueprintReadWrite, Category=Sprite)
 	class UTexture2D* Sprite;
diff --git a/Engine/Source/Runtime/Engine/Classes/Components/BrushComponent.h b/Engine/Source/Runtime/Engine/Classes/Components/BrushComponent.h
--- a/Engine/Source/Runtime/Engine/Classes/Components/BrushComponent.h
+++ b/Engine/Source/Runtime/Engine/Classes/Components/BrushComponent.h
@@ -14,9 +14,7 @@ struct FConvexVolume;
 UCLASS(editinlinenew, MinimalAPI, hidecategories=(Physics, Lighting, LOD, Rendering, TextureStreaming, Transform, Activation, "Components|Activation"), showcategories=(Mobility, "Rendering|Material"))
 class UBrushComponent : public UPrimitiveComponent
 {
-	GENERATED_BODY()
-public:
-	ENGINE_API UBrushComponent(const FObjectInitializer& ObjectInitializer = FObjectInitializer::Get());
+	GENERATED_UCLASS_BODY()
 
 	UPROPERTY()
 	class UModel* Brush;
diff --git a/Engine/Source/Runtime/Engine/Classes/Components/DrawSphereComponent.h b/Engine/Source/Runtime/Engine/Classes/Components/DrawSphereComponent.h
--- a/Engine/Source/Runtime/Engine/Classes/Components/DrawSphereComponent.h
+++ b/Engine/Source/Runtime/Engine/Classes/Components/DrawSphereComponent.h
@@ -11,9 +11,7 @@
 UCLASS(collapsecategories, hidecategories=Object, editinlinenew, MinimalAPI)
 class UDrawSphereComponent : public USphereComponent
 {
-	GENERATED_BODY()
-public:
-	ENGINE_API UDrawSphereComponent(const FObjectInitializer& ObjectInitializer = FObjectInitializer::Get());
+	GENERATED_UCLASS_BODY()
 
 public:
 #if WITH_EDITOR
diff --git a/Engine/Source/Runtime/Engine/Classes/Components/MaterialBillboardComponent.h b/Engine/Source/Runtime/Engine/Classes/Components/MaterialBillboardComponent.h
--- a/Engine/Source/Runtime/Engine/Classes/Components/MaterialBillboardComponent.h
+++ b/Engine/Source/Runtime/Engine/Classes/Components/MaterialBillboardComponent.h
@@ -8,7 +8,7 @@
 USTRUCT()
 struct FMaterialSpriteElement
 {
-	GENERATED_BODY()
+	GENERATED_USTRUCT_BODY()
 
 	/** The material that the sprite is rendered with. */
 	UPROPERTY(EditAnywhere, BlueprintReadWrite, Category=MaterialSpriteElement)
@@ -53,9 +53,7 @@ struct FMaterialSpriteElement
 UCLASS(ClassGroup=Rendering, collapsecategories, hidecategories=(Object,Activation,"Components|Activation",Physics,Collision,Lighting,LOD,Mesh,PhysicsVolume), editinlinenew, meta=(BlueprintSpawnableComponent))
 class ENGINE_API UMaterialBillboardComponent : public UPrimitiveComponent
 {
-	GENERATED_BODY()
-public:
-	UMaterialBillboardComponent(const FObjectInitializer& ObjectInitializer = FObjectInitializer::Get());
+	GENERATED_UCLASS_BODY()
 
 	UPROPERTY(EditAnywhere, BlueprintReadWrite, Category=Sprite)
 	TArray<struct FMaterialSpriteElement> Elements;
diff --git a/Engine/Source/Runtime/Engine/Classes/Components/SkeletalMeshComponent.h b/Engine/Source/Runtime/Engine/Classes/Components/SkeletalMeshComponent.h
--- a/Engine/Source/Runtime/Engine/Classes/Components/SkeletalMeshComponent.h
+++ b/Engine/Source/Runtime/Engine/Classes/Components/SkeletalMeshComponent.h
@@ -193,7 +193,7 @@ namespace EAnimationMode
 USTRUCT()
 struct FSingleAnimationPlayData
 {
-	GENERATED_BODY()
+	GENERATED_USTRUCT_BODY()
 
 	// @todo in the future, we should make this one UObject
 	// and have detail customization to display different things
@@ -245,7 +245,7 @@ struct FSingleAnimationPlayData
 USTRUCT()
 struct FSkeletalMeshComponentPreClothTickFunction : public FTickFunction
 {
-	GENERATED_BODY()
+	GENERATED_USTRUCT_BODY()
 
 	/** World this tick function belongs to **/
 	class USkeletalMeshComponent*	Target;
@@ -272,9 +272,7 @@ struct FSkeletalMeshComponentPreClothTickFunction : public FTickFunction
 UCLASS(ClassGroup=(Rendering, Common), hidecategories=Object, config=Engine, editinlinenew, meta=(BlueprintSpawnableComponent))
 class ENGINE_API USkeletalMeshComponent : public USkinnedMeshComponent, public IInterface_CollisionDataProvider
 {
-	GENERATED_BODY()
-public:
-	USkeletalMeshComponent(const FObjectInitializer& ObjectInitializer = FObjectInitializer::Get());
+	GENERATED_UCLASS_BODY()
 	
 	/**
 	 * Animation 
