Added a new graph action filter that restricts level script blueprints properties and functions from being referenced in external blueprints.
#UE-10663 - Prevent creating references to external level functions in sub level and vice versa
#Branch UE4
#Proj BlueprintGraph

diff --git a/Engine/Source/Editor/BlueprintGraph/Private/BlueprintActionFilter.cpp b/Engine/Source/Editor/BlueprintGraph/Private/BlueprintActionFilter.cpp
--- a/Engine/Source/Editor/BlueprintGraph/Private/BlueprintActionFilter.cpp
+++ b/Engine/Source/Editor/BlueprintGraph/Private/BlueprintActionFilter.cpp
@@ -236,6 +236,15 @@ namespace BlueprintActionFilterImpl
 	 */
 	static bool IsOutOfScopeLocalVariable(FBlueprintActionFilter const& Filter, FBlueprintActionInfo& BlueprintAction);
 
+	/**
+	 * 
+	 * 
+	 * @param  Filter	
+	 * @param  BlueprintAction	
+	 * @return 
+	 */
+	static bool IsLevelScriptActionValid(FBlueprintActionFilter const& Filter, FBlueprintActionInfo& BlueprintAction);
+
 	/**
 	 * 
 	 * 
@@ -922,6 +931,42 @@ static bool BlueprintActionFilterImpl::IsOutOfScopeLocalVariable(FBlueprintActio
 	return bIsFilteredOut;
 }
 
+//------------------------------------------------------------------------------
+static bool BlueprintActionFilterImpl::IsLevelScriptActionValid(FBlueprintActionFilter const& Filter, FBlueprintActionInfo& BlueprintAction)
+{
+	bool bIsFilteredOut = false;
+	FBlueprintActionContext const& FilterContext = Filter.Context;
+	UClass* OuterClass = nullptr;
+
+	if( UFunction const* NodeFunction = BlueprintAction.GetAssociatedFunction() )
+	{
+		OuterClass = Cast<UClass>( NodeFunction->GetOuter() );
+	}
+	else if( UProperty const* NodeProperty = BlueprintAction.GetAssociatedProperty() )
+	{
+		OuterClass = Cast<UClass>( NodeProperty->GetOuter() );
+	}
+
+	if( OuterClass )
+	{
+		UBlueprint const* BPOwner = Cast<UBlueprint>( OuterClass->ClassGeneratedBy );
+		if( BPOwner && BPOwner->BlueprintType == BPTYPE_LevelScript )
+		{
+			bIsFilteredOut = true;
+			for( auto Blueprint : FilterContext.Blueprints )
+			{
+				if( BPOwner->GetBlueprintGuid() == Blueprint->GetBlueprintGuid() )
+				{
+					bIsFilteredOut = false;
+					break;
+				}
+			}
+		}
+	}
+
+	return bIsFilteredOut;
+}
+
 //------------------------------------------------------------------------------
 static bool BlueprintActionFilterImpl::IsSchemaIncompatible(FBlueprintActionFilter const& Filter, FBlueprintActionInfo& BlueprintAction)
 {
@@ -1605,6 +1650,7 @@ FBlueprintActionFilter::FBlueprintActionFilter(uint32 Flags/*= 0x00*/)
 	AddRejectionTest(FRejectionTestDelegate::CreateStatic(IsNonTargetMember, !(Flags & BPFILTER_RejectGlobalFields)));
 	AddRejectionTest(FRejectionTestDelegate::CreateStatic(IsUnBoundBindingSpawner));
 	AddRejectionTest(FRejectionTestDelegate::CreateStatic(IsOutOfScopeLocalVariable));
+	AddRejectionTest(FRejectionTestDelegate::CreateStatic(IsLevelScriptActionValid));
 }
 
 //------------------------------------------------------------------------------
