Graph Editor: Add a shortcut to create reroute nodes by double-clicking on BP spline connections

diff --git a/Engine/Source/Editor/BlueprintGraph/Classes/EdGraphSchema_K2.h b/Engine/Source/Editor/BlueprintGraph/Classes/EdGraphSchema_K2.h
--- a/Engine/Source/Editor/BlueprintGraph/Classes/EdGraphSchema_K2.h
+++ b/Engine/Source/Editor/BlueprintGraph/Classes/EdGraphSchema_K2.h
@@ -424,6 +424,7 @@ public:
 	virtual bool ShouldAlwaysPurgeOnModification() const override { return false; }
 	virtual void SplitPin(UEdGraphPin* Pin) const override;
 	virtual void RecombinePin(UEdGraphPin* Pin) const override;
+	virtual void OnPinConnectionDoubleCicked(UEdGraphPin* PinA, UEdGraphPin* PinB, const FVector2D& GraphPosition) const override;
 	// End EdGraphSchema Interface
 
 	/**
diff --git a/Engine/Source/Editor/BlueprintGraph/Private/EdGraphSchema_K2.cpp b/Engine/Source/Editor/BlueprintGraph/Private/EdGraphSchema_K2.cpp
--- a/Engine/Source/Editor/BlueprintGraph/Private/EdGraphSchema_K2.cpp
+++ b/Engine/Source/Editor/BlueprintGraph/Private/EdGraphSchema_K2.cpp
@@ -5635,6 +5635,29 @@ void UEdGraphSchema_K2::RecombinePin(UEdGraphPin* Pin) const
 	FBlueprintEditorUtils::MarkBlueprintAsModified(Blueprint);
 }
 
+void UEdGraphSchema_K2::OnPinConnectionDoubleCicked(UEdGraphPin* PinA, UEdGraphPin* PinB, const FVector2D& GraphPosition) const
+{
+	const FScopedTransaction Transaction(LOCTEXT("CreateRerouteNodeOnWire", "Create Reroute Node"));
+
+	//@TODO: This constant is duplicated from inside of SGraphNodeKnot
+	const FVector2D NodeSpacerSize(42.0f, 24.0f);
+	const FVector2D KnotTopLeft = GraphPosition - (NodeSpacerSize * 0.5f);
+
+	// Create a new knot
+	UEdGraph* ParentGraph = PinA->GetOwningNode()->GetGraph();
+	UK2Node_Knot* NewKnot = FEdGraphSchemaAction_K2NewNode::SpawnNodeFromTemplate<UK2Node_Knot>(ParentGraph, NewObject<UK2Node_Knot>(), KnotTopLeft);
+
+	// Move the connections across (only notifying the knot, as the other two didn't really change)
+	PinA->BreakLinkTo(PinB);
+	PinA->MakeLinkTo((PinA->Direction == EGPD_Output) ? NewKnot->GetInputPin() : NewKnot->GetOutputPin());
+	PinB->MakeLinkTo((PinB->Direction == EGPD_Output) ? NewKnot->GetInputPin() : NewKnot->GetOutputPin());
+	NewKnot->PostReconstructNode();
+
+	// Dirty the blueprint
+	UBlueprint* Blueprint = FBlueprintEditorUtils::FindBlueprintForGraphChecked(CastChecked<UEdGraph>(ParentGraph));
+	FBlueprintEditorUtils::MarkBlueprintAsModified(Blueprint);
+}
+
 void UEdGraphSchema_K2::ConfigureVarNode(UK2Node_Variable* InVarNode, FName InVariableName, UStruct* InVariableSource, UBlueprint* InTargetBlueprint)
 {
 	// See if this is a 'self context' (ie. blueprint class is owner (or child of owner) of dropped var class)
diff --git a/Engine/Source/Editor/GraphEditor/Private/SGraphPanel.cpp b/Engine/Source/Editor/GraphEditor/Private/SGraphPanel.cpp
--- a/Engine/Source/Editor/GraphEditor/Private/SGraphPanel.cpp
+++ b/Engine/Source/Editor/GraphEditor/Private/SGraphPanel.cpp
@@ -673,9 +673,13 @@ FReply SGraphPanel::OnMouseButtonUp(const FGeometry& MyGeometry, const FPointerE
 
 FReply SGraphPanel::OnMouseButtonDoubleClick(const FGeometry& MyGeometry, const FPointerEvent& MouseEvent)
 {
-	if (SGraphPin* BestPinFromHoveredSpline = GetBestPinFromHoveredSpline())
+	if (PreviousFrameSplineOverlap.IsValid() && (PreviousFrameSplineOverlap.Pin1 != nullptr) && (PreviousFrameSplineOverlap.Pin2 != nullptr))
 	{
-		//@TODO: Should insert a reroute node here
+		// Give the schema a chance to do something interesting with a double click on a proper spline (both ends are attached to a pin, i.e., not a preview/drag one)
+		const FVector2D DoubleClickPositionInGraphSpace = PanelCoordToGraphCoord(MyGeometry.AbsoluteToLocal(MouseEvent.GetScreenSpacePosition()));
+
+		const UEdGraphSchema* Schema = GraphObj->GetSchema();
+		Schema->OnPinConnectionDoubleCicked(PreviousFrameSplineOverlap.Pin1, PreviousFrameSplineOverlap.Pin2, DoubleClickPositionInGraphSpace);
 	}
 
 	return SNodePanel::OnMouseButtonDoubleClick(MyGeometry, MouseEvent);
diff --git a/Engine/Source/Runtime/Engine/Classes/EdGraph/EdGraphSchema.h b/Engine/Source/Runtime/Engine/Classes/EdGraph/EdGraphSchema.h
--- a/Engine/Source/Runtime/Engine/Classes/EdGraph/EdGraphSchema.h
+++ b/Engine/Source/Runtime/Engine/Classes/EdGraph/EdGraphSchema.h
@@ -620,6 +620,9 @@ class ENGINE_API UEdGraphSchema : public UObject
 	/** Collapses a pin and its siblings back in to the original pin */
 	virtual void RecombinePin(UEdGraphPin* Pin) const { };
 
+	/** Handles double-clicking on a pin<->pin connection */
+	virtual void OnPinConnectionDoubleCicked(UEdGraphPin* PinA, UEdGraphPin* PinB, const FVector2D& GraphPosition) const { }
+
 	/** Break links on this pin and create links instead on MoveToPin */
 	virtual FPinConnectionResponse MovePinLinks(UEdGraphPin& MoveFromPin, UEdGraphPin& MoveToPin, bool bIsIntermediateMove = false) const;
 
