#ue4 Fix for autocreaterefterm not being properly handled for lists

diff --git a/Engine/Source/Editor/BlueprintGraph/Private/EdGraphSchema_K2.cpp b/Engine/Source/Editor/BlueprintGraph/Private/EdGraphSchema_K2.cpp
--- a/Engine/Source/Editor/BlueprintGraph/Private/EdGraphSchema_K2.cpp
+++ b/Engine/Source/Editor/BlueprintGraph/Private/EdGraphSchema_K2.cpp
@@ -3786,7 +3786,10 @@ bool UEdGraphSchema_K2::IsAutoCreateRefTerm(const UEdGraphPin* Pin) const
 		UFunction* TargetFunction = FuncNode->GetTargetFunction();
 		if (TargetFunction && !Pin->PinName.IsEmpty())
 		{
-			bIsAutoCreateRefTerm = (Pin->PinName == TargetFunction->GetMetaData(FBlueprintMetadata::MD_AutoCreateRefTerm));
+			TArray<FString> AutoCreateParameterNames;
+			FString MetaData = TargetFunction->GetMetaData(FBlueprintMetadata::MD_AutoCreateRefTerm);
+			MetaData.ParseIntoArray(AutoCreateParameterNames, TEXT(","), true);
+			bIsAutoCreateRefTerm = AutoCreateParameterNames.Contains(Pin->PinName);
 		}
 	}
 
