#ue4 BP local variables are created for unconnected "AutoCreateRefTerm" pins, even when default value is empty.

UE-12595 A struct param marked with AutoCreateRefTerm crashes when running Blueprint

diff --git a/Engine/Source/Editor/BlueprintGraph/Private/K2Node_CallFunction.cpp b/Engine/Source/Editor/BlueprintGraph/Private/K2Node_CallFunction.cpp
--- a/Engine/Source/Editor/BlueprintGraph/Private/K2Node_CallFunction.cpp
+++ b/Engine/Source/Editor/BlueprintGraph/Private/K2Node_CallFunction.cpp
@@ -1870,14 +1870,14 @@ void UK2Node_CallFunction::ExpandNode(class FKismetCompilerContext& CompilerCont
 			{
 				if ( Pin && bHasAutoCreateRefTerms && AutoCreateRefTermPinNames.Contains(Pin->PinName) )
 				{
-					const bool bHasDefaultValue = !Pin->DefaultValue.IsEmpty() || Pin->DefaultObject || !Pin->DefaultTextValue.IsEmpty();
 					const bool bValidAutoRefPin = Pin->PinType.bIsReference
 						&& !CompilerContext.GetSchema()->IsMetaPin(*Pin)
 						&& ( Pin->Direction == EGPD_Input )
-						&& !Pin->LinkedTo.Num()
-						&& ( Pin->PinType.bIsArray || bHasDefaultValue );
+						&& !Pin->LinkedTo.Num();
 					if ( bValidAutoRefPin )
 					{
+						const bool bHasDefaultValue = !Pin->DefaultValue.IsEmpty() || Pin->DefaultObject || !Pin->DefaultTextValue.IsEmpty();
+
 						//default values can be reset when the pin is connected
 						const auto DefaultValue = Pin->DefaultValue;
 						const auto DefaultObject = Pin->DefaultObject;
