Fixing various cook errors. No longer using a ranged for loop when the array node's pins can change during the loop.

diff --git a/Engine/Source/Editor/BlueprintGraph/Private/K2Node_CallArrayFunction.cpp b/Engine/Source/Editor/BlueprintGraph/Private/K2Node_CallArrayFunction.cpp
--- a/Engine/Source/Editor/BlueprintGraph/Private/K2Node_CallArrayFunction.cpp
+++ b/Engine/Source/Editor/BlueprintGraph/Private/K2Node_CallArrayFunction.cpp
@@ -39,8 +39,14 @@ void UK2Node_CallArrayFunction::AllocateDefaultPins()
 
 void UK2Node_CallArrayFunction::PostReconstructNode()
 {
-	for (UEdGraphPin* Pin : Pins)
+	// cannot use a ranged for here, as PinConnectionListChanged() might end up 
+	// collapsing split pins (subtracting elements from Pins)
+	// 
+	// @TODO: this breaks connections to split outputs (by collapsing them), 
+	//        that is obviously a bug
+	for (int32 PinIndex = 0; PinIndex < Pins.Num(); ++PinIndex)
 	{
+		UEdGraphPin* Pin = Pins[PinIndex];
 		if (Pin->LinkedTo.Num() > 0)
 		{
 			PinConnectionListChanged(Pin);
