[UE-12083] Fix for function used in async task that was not merged over. We can never compile this, but we shouldn't crash!

diff --git a/Engine/Source/Editor/BlueprintGraph/Private/K2Node_BaseAsyncTask.cpp b/Engine/Source/Editor/BlueprintGraph/Private/K2Node_BaseAsyncTask.cpp
--- a/Engine/Source/Editor/BlueprintGraph/Private/K2Node_BaseAsyncTask.cpp
+++ b/Engine/Source/Editor/BlueprintGraph/Private/K2Node_BaseAsyncTask.cpp
@@ -28,6 +28,10 @@ FText UK2Node_BaseAsyncTask::GetTooltipText() const
 
 FText UK2Node_BaseAsyncTask::GetNodeTitle(ENodeTitleType::Type TitleType) const
 {
+	if (GetFactoryFunction() == nullptr)
+	{
+		return FText(LOCTEXT("UK2Node_BaseAsyncTaskGetNodeTitle", "Async Task: Missing Function"));
+	}
 	const FText FunctionToolTipText = UK2Node_CallFunction::GetUserFacingFunctionName(GetFactoryFunction());
 	return FunctionToolTipText;
 }
@@ -272,15 +276,25 @@ void UK2Node_BaseAsyncTask::ExpandNode(class FKismetCompilerContext& CompilerCon
 	UK2Node_CallFunction* const CallCreateProxyObjectNode = CompilerContext.SpawnIntermediateNode<UK2Node_CallFunction>(this, SourceGraph);
 	CallCreateProxyObjectNode->FunctionReference.SetExternalMember(ProxyFactoryFunctionName, ProxyFactoryClass);
 	CallCreateProxyObjectNode->AllocateDefaultPins();
+	if (CallCreateProxyObjectNode->GetTargetFunction() == nullptr)
+	{
+		const FString ClassName = ProxyFactoryClass ? ProxyFactoryClass->GetName() : LOCTEXT("MissingClassString", "Unknown Class").ToString();
+		const FString RawMessage = LOCTEXT("AsyncTaskError", "BaseAsyncTask: Missing function %s from class %s for async task @@").ToString();
+		const FString FormattedMessage = FString::Printf(*RawMessage, *ProxyFactoryFunctionName.GetPlainNameString(), *ClassName);
+		CompilerContext.MessageLog.Error(*FormattedMessage, this);
+		return;
+	}
+
 	bIsErrorFree &= CompilerContext.MovePinLinksToIntermediate(*FindPinChecked(Schema->PN_Execute), *CallCreateProxyObjectNode->FindPinChecked(Schema->PN_Execute)).CanSafeConnect();
+
 	for (auto CurrentPin : Pins)
-		{
+	{
 		if (FBaseAsyncTaskHelper::ValidDataPin(CurrentPin, EGPD_Input, Schema))
-			{
+		{
 			UEdGraphPin* DestPin = CallCreateProxyObjectNode->FindPin(CurrentPin->PinName); // match function inputs, to pass data to function from CallFunction node
 			bIsErrorFree &= DestPin && CompilerContext.MovePinLinksToIntermediate(*CurrentPin, *DestPin).CanSafeConnect();
-			}
 		}
+	}
 
 	// GATHER OUTPUT PARAMETERS AND PAIR THEM WITH LOCAL VARIABLES
 	TArray<FBaseAsyncTaskHelper::FOutputPinAndLocalVariable> VariableOutputs;
@@ -300,15 +314,15 @@ void UK2Node_BaseAsyncTask::ExpandNode(class FKismetCompilerContext& CompilerCon
 	UEdGraphPin* LastThenPin = CallCreateProxyObjectNode->FindPinChecked(Schema->PN_Then);
 	UEdGraphPin* const ProxyObjectPin = CallCreateProxyObjectNode->GetReturnValuePin();
 	for (TFieldIterator<UMulticastDelegateProperty> PropertyIt(ProxyClass, EFieldIteratorFlags::ExcludeSuper); PropertyIt && bIsErrorFree; ++PropertyIt)
-		{
+	{
 		bIsErrorFree &= FBaseAsyncTaskHelper::HandleDelegateImplementation(*PropertyIt, VariableOutputs, ProxyObjectPin, LastThenPin, this, SourceGraph, CompilerContext);
 	}
 
 	if (CallCreateProxyObjectNode->FindPinChecked(Schema->PN_Then) == LastThenPin)
-			{
+	{
 		CompilerContext.MessageLog.Error(*LOCTEXT("MissingDelegateProperties", "BaseAsyncTask: Proxy has no delegates defined. @@").ToString(), this);
 		return;
-			}
+	}
 
 	// Create a call to activate the proxy object if necessary
 	if (ProxyActivateFunctionName != NAME_None)
@@ -337,11 +351,11 @@ void UK2Node_BaseAsyncTask::ExpandNode(class FKismetCompilerContext& CompilerCon
 	if (!bIsErrorFree)
 	{
 		CompilerContext.MessageLog.Error(*LOCTEXT("InternalConnectionError", "BaseAsyncTask: Internal connection error. @@").ToString(), this);
-		}
+	}
 
 	// Make sure we caught everything
 	BreakAllNodeLinks();
-	}
+}
 
 bool UK2Node_BaseAsyncTask::HasExternalBlueprintDependencies(TArray<class UStruct*>* OptionalOutput) const
 {
@@ -397,7 +411,7 @@ UFunction* UK2Node_BaseAsyncTask::GetFactoryFunction() const
 {
 	if (ProxyFactoryClass == nullptr)
 	{
-		UE_LOG(LogBlueprint, Fatal, TEXT("ProxyFactoryClass null in %s. Was a class deleted or saved on a non promoted build?"), *GetFullName());
+		UE_LOG(LogBlueprint, Error, TEXT("ProxyFactoryClass null in %s. Was a class deleted or saved on a non promoted build?"), *GetFullName());
 		return nullptr;
 	}
 	
