Making it so we don't drop split pins for array nodes when they're refreshed (which was breaking existing split pin connections).

diff --git a/Engine/Source/Editor/BlueprintGraph/Private/K2Node_CallArrayFunction.cpp b/Engine/Source/Editor/BlueprintGraph/Private/K2Node_CallArrayFunction.cpp
--- a/Engine/Source/Editor/BlueprintGraph/Private/K2Node_CallArrayFunction.cpp
+++ b/Engine/Source/Editor/BlueprintGraph/Private/K2Node_CallArrayFunction.cpp
@@ -41,9 +41,6 @@ void UK2Node_CallArrayFunction::PostReconstructNode()
 {
 	// cannot use a ranged for here, as PinConnectionListChanged() might end up 
 	// collapsing split pins (subtracting elements from Pins)
-	// 
-	// @TODO: this breaks connections to split outputs (by collapsing them), 
-	//        that is obviously a bug
 	for (int32 PinIndex = 0; PinIndex < Pins.Num(); ++PinIndex)
 	{
 		UEdGraphPin* Pin = Pins[PinIndex];
@@ -211,6 +208,7 @@ void UK2Node_CallArrayFunction::PropagateArrayTypeInfo(const UEdGraphPin* Source
 	if( SourcePin )
 	{
 		const UEdGraphSchema_K2* Schema = GetDefault<UEdGraphSchema_K2>();
+		const FEdGraphPinType& SourcePinType = SourcePin->PinType;
 
 		TArray<UEdGraphPin*> DependentPins;
 		GetArrayTypeDependentPins(DependentPins);
@@ -221,14 +219,25 @@ void UK2Node_CallArrayFunction::PropagateArrayTypeInfo(const UEdGraphPin* Source
 		{
 			if (CurrentPin != SourcePin)
 			{
+				FEdGraphPinType& CurrentPinType = CurrentPin->PinType;
+
+				bool const bHasTypeMismatch = (CurrentPinType.PinCategory != SourcePinType.PinCategory) ||
+					(CurrentPinType.PinSubCategory != SourcePinType.PinSubCategory) ||
+					(CurrentPinType.PinSubCategoryObject != SourcePinType.PinSubCategoryObject);
+
+				if (!bHasTypeMismatch)
+				{
+					continue;
+				}
+
 				if (CurrentPin->SubPins.Num() > 0)
 				{
 					Schema->RecombinePin(CurrentPin->SubPins[0]);
 				}
 
-				CurrentPin->PinType.PinCategory = SourcePin->PinType.PinCategory;
-				CurrentPin->PinType.PinSubCategory = SourcePin->PinType.PinSubCategory;
-				CurrentPin->PinType.PinSubCategoryObject = SourcePin->PinType.PinSubCategoryObject;
+				CurrentPinType.PinCategory          = SourcePinType.PinCategory;
+				CurrentPinType.PinSubCategory       = SourcePinType.PinSubCategory;
+				CurrentPinType.PinSubCategoryObject = SourcePinType.PinSubCategoryObject;
 
 				// Reset default values
 				CurrentPin->DefaultValue = FString();
