UE-12990 Contains node connected to Enum Array value selection cleared on editor restart

diff --git a/Engine/Source/Editor/BlueprintGraph/Private/K2Node.cpp b/Engine/Source/Editor/BlueprintGraph/Private/K2Node.cpp
--- a/Engine/Source/Editor/BlueprintGraph/Private/K2Node.cpp
+++ b/Engine/Source/Editor/BlueprintGraph/Private/K2Node.cpp
@@ -558,6 +558,17 @@ void UK2Node::ReconstructSinglePin(UEdGraphPin* NewPin, UEdGraphPin* OldPin, ERe
 {
 	UBlueprint* Blueprint = GetBlueprint();
 
+	check(NewPin && OldPin);
+
+	// Copy the type info for any unconnected wildcard pin
+	const UEdGraphSchema_K2* K2Schema = CastChecked<const UEdGraphSchema_K2>(GetSchema());
+	const bool bStrongTypeInNewPin = NewPin->PinType.PinCategory != K2Schema->PC_Wildcard;
+	const bool bStrongTypeInOldPin = OldPin->PinType.PinCategory != K2Schema->PC_Wildcard;
+	if (bStrongTypeInOldPin && !bStrongTypeInNewPin && !OldPin->LinkedTo.Num())
+	{
+		NewPin->PinType = OldPin->PinType;
+	}
+
 	// Copy over modified persistent data
 	NewPin->CopyPersistentDataFromOldPin(*OldPin);
 
