Better fix for UE-12990 Contains node connected to Enum Array value selection cleared on editor restart

diff --git a/Engine/Source/Editor/BlueprintGraph/Private/K2Node.cpp b/Engine/Source/Editor/BlueprintGraph/Private/K2Node.cpp
--- a/Engine/Source/Editor/BlueprintGraph/Private/K2Node.cpp
+++ b/Engine/Source/Editor/BlueprintGraph/Private/K2Node.cpp
@@ -560,15 +560,6 @@ void UK2Node::ReconstructSinglePin(UEdGraphPin* NewPin, UEdGraphPin* OldPin, ERe
 
 	check(NewPin && OldPin);
 
-	// Copy the type info for any unconnected wildcard pin
-	const UEdGraphSchema_K2* K2Schema = CastChecked<const UEdGraphSchema_K2>(GetSchema());
-	const bool bStrongTypeInNewPin = NewPin->PinType.PinCategory != K2Schema->PC_Wildcard;
-	const bool bStrongTypeInOldPin = OldPin->PinType.PinCategory != K2Schema->PC_Wildcard;
-	if (bStrongTypeInOldPin && !bStrongTypeInNewPin && !OldPin->LinkedTo.Num())
-	{
-		NewPin->PinType = OldPin->PinType;
-	}
-
 	// Copy over modified persistent data
 	NewPin->CopyPersistentDataFromOldPin(*OldPin);
 
diff --git a/Engine/Source/Editor/BlueprintGraph/Private/K2Node_CallArrayFunction.cpp b/Engine/Source/Editor/BlueprintGraph/Private/K2Node_CallArrayFunction.cpp
--- a/Engine/Source/Editor/BlueprintGraph/Private/K2Node_CallArrayFunction.cpp
+++ b/Engine/Source/Editor/BlueprintGraph/Private/K2Node_CallArrayFunction.cpp
@@ -207,7 +207,7 @@ void UK2Node_CallArrayFunction::PropagateArrayTypeInfo(const UEdGraphPin* Source
 {
 	if( SourcePin )
 	{
-		const UEdGraphSchema_K2* Schema = GetDefault<UEdGraphSchema_K2>();
+		const UEdGraphSchema_K2* Schema = CastChecked<UEdGraphSchema_K2>(GetSchema());
 		const FEdGraphPinType& SourcePinType = SourcePin->PinType;
 
 		TArray<UEdGraphPin*> DependentPins;
@@ -240,9 +240,10 @@ void UK2Node_CallArrayFunction::PropagateArrayTypeInfo(const UEdGraphPin* Source
 				CurrentPinType.PinSubCategoryObject = SourcePinType.PinSubCategoryObject;
 
 				// Reset default values
-				CurrentPin->DefaultValue = FString();
-				CurrentPin->DefaultTextValue = FText();
-				CurrentPin->DefaultObject = NULL;
+				if (!Schema->IsPinDefaultValid(CurrentPin, CurrentPin->DefaultValue, CurrentPin->DefaultObject, CurrentPin->DefaultTextValue).IsEmpty())
+				{
+					CurrentPin->ResetDefaultValue();
+				}
 
 				// Verify that all previous connections to this pin are still valid with the new type
 				for (TArray<UEdGraphPin*>::TIterator ConnectionIt(CurrentPin->LinkedTo); ConnectionIt; ++ConnectionIt)
