BP: skeletal class is recompiled after Event Node was added

diff --git a/Engine/Source/Editor/BlueprintGraph/Private/K2Node_Event.cpp b/Engine/Source/Editor/BlueprintGraph/Private/K2Node_Event.cpp
--- a/Engine/Source/Editor/BlueprintGraph/Private/K2Node_Event.cpp
+++ b/Engine/Source/Editor/BlueprintGraph/Private/K2Node_Event.cpp
@@ -257,8 +257,8 @@ void UK2Node_Event::ValidateNodeDuringCompilation(class FCompilerResultsLog& Mes
 
 bool UK2Node_Event::NodeCausesStructuralBlueprintChange() const
 {
-	// will only change class structure when UGPF is disabled
-	return !UBlueprintGeneratedClass::UsePersistentUberGraphFrame();
+	// FBlueprintEditor::CanAddParentNode requires actual data in skel class
+	return true;
 }
 
 void UK2Node_Event::GetRedirectPinNames(const UEdGraphPin& Pin, TArray<FString>& RedirectPinNames) const
