Added enum-inequality kismet node.

diff --git a/Engine/Source/Editor/BlueprintGraph/Classes/K2Node_EnumEquality.h b/Engine/Source/Editor/BlueprintGraph/Classes/K2Node_EnumEquality.h
--- a/Engine/Source/Editor/BlueprintGraph/Classes/K2Node_EnumEquality.h
+++ b/Engine/Source/Editor/BlueprintGraph/Classes/K2Node_EnumEquality.h
@@ -41,6 +41,6 @@ class UK2Node_EnumEquality : public UK2Node
 	BLUEPRINTGRAPH_API UEdGraphPin* GetInput2Pin() const;
 
 	/** Gets the name and class of the EqualEqual_ByteByte function */
-	BLUEPRINTGRAPH_API void GetConditionalFunction(FName& FunctionName, UClass** FunctionClass);
+	BLUEPRINTGRAPH_API virtual void GetConditionalFunction(FName& FunctionName, UClass** FunctionClass) const;
 };
 
diff --git a/Engine/Source/Editor/BlueprintGraph/Classes/K2Node_EnumInequality.h b/Engine/Source/Editor/BlueprintGraph/Classes/K2Node_EnumInequality.h
new file mode 100644
--- /dev/null
+++ b/Engine/Source/Editor/BlueprintGraph/Classes/K2Node_EnumInequality.h
@@ -0,0 +1,25 @@
+// Copyright 1998-2015 Epic Games, Inc. All Rights Reserved.
+
+#pragma once
+#include "K2Node.h"
+#include "K2Node_EnumInequality.generated.h"
+
+UCLASS(MinimalAPI)
+class UK2Node_EnumInequality : public UK2Node_EnumEquality
+{
+	GENERATED_UCLASS_BODY()
+
+	// Begin UEdGraphNode interface
+	virtual FText GetTooltipText() const override;
+	virtual FText GetNodeTitle(ENodeTitleType::Type TitleType) const override;
+	virtual FString GetKeywords() const override;
+	// End UEdGraphNode interface
+
+	// Begin UK2Node interface
+	virtual FText GetCompactNodeTitle() const override { return NSLOCTEXT("K2Node", "EqualEqual", "!="); }
+	// End UK2Node interface
+
+	/** Gets the name and class of the EqualEqual_ByteByte function */
+	BLUEPRINTGRAPH_API virtual void GetConditionalFunction(FName& FunctionName, UClass** FunctionClass) const override;
+};
+
diff --git a/Engine/Source/Editor/BlueprintGraph/Private/K2Node_EnumEquality.cpp b/Engine/Source/Editor/BlueprintGraph/Private/K2Node_EnumEquality.cpp
--- a/Engine/Source/Editor/BlueprintGraph/Private/K2Node_EnumEquality.cpp
+++ b/Engine/Source/Editor/BlueprintGraph/Private/K2Node_EnumEquality.cpp
@@ -172,7 +172,7 @@ UEdGraphPin* UK2Node_EnumEquality::GetInput2Pin() const
 	return Pin;
 }
 
-void UK2Node_EnumEquality::GetConditionalFunction(FName& FunctionName, UClass** FunctionClass)
+void UK2Node_EnumEquality::GetConditionalFunction(FName& FunctionName, UClass** FunctionClass) const
 {
 	FunctionName = GET_FUNCTION_NAME_CHECKED(UKismetMathLibrary, EqualEqual_ByteByte);
 	*FunctionClass = UKismetMathLibrary::StaticClass();
diff --git a/Engine/Source/Editor/BlueprintGraph/Private/K2Node_EnumInequality.cpp b/Engine/Source/Editor/BlueprintGraph/Private/K2Node_EnumInequality.cpp
new file mode 100644
--- /dev/null
+++ b/Engine/Source/Editor/BlueprintGraph/Private/K2Node_EnumInequality.cpp
@@ -0,0 +1,40 @@
+// Copyright 1998-2015 Epic Games, Inc. All Rights Reserved.
+
+#include "BlueprintGraphPrivatePCH.h"
+//#include "KismetCompiler.h"
+//#include "Kismet/KismetMathLibrary.h"
+//#include "K2ActionMenuBuilder.h" // for FK2ActionMenuBuilder::AddNewNodeAction()
+#include "K2Node_EnumInequality.h"
+//#include "BlueprintNodeSpawner.h"
+//#include "EditorCategoryUtils.h"
+//#include "BlueprintActionDatabaseRegistrar.h"
+
+#define LOCTEXT_NAMESPACE "K2Node_EnumInequality"
+
+UK2Node_EnumInequality::UK2Node_EnumInequality(const FObjectInitializer& ObjectInitializer)
+	: Super(ObjectInitializer)
+{
+}
+
+FText UK2Node_EnumInequality::GetTooltipText() const
+{
+	return LOCTEXT("EnumInequalityTooltip", "Returns true if A is NOT equal to B (A != B)");
+}
+
+FText UK2Node_EnumInequality::GetNodeTitle(ENodeTitleType::Type TitleType) const
+{
+	return NSLOCTEXT("K2Node_EnumInequality", "NotEqualEnum", "NotEqual (Enum)");
+}
+
+FString UK2Node_EnumInequality::GetKeywords() const
+{
+	return TEXT("!=");
+}
+void UK2Node_EnumInequality::GetConditionalFunction(FName& FunctionName, UClass** FunctionClass) const
+{
+	FunctionName = GET_FUNCTION_NAME_CHECKED(UKismetMathLibrary, NotEqual_ByteByte);
+	*FunctionClass = UKismetMathLibrary::StaticClass();
+}
+
+
+#undef LOCTEXT_NAMESPACE
