[UE-12898] Guarding against invalid BP action database objects (ones that were pending deletion) from being registered later with the database.

diff --git a/Engine/Source/Editor/BlueprintGraph/Private/BlueprintActionDatabase.cpp b/Engine/Source/Editor/BlueprintGraph/Private/BlueprintActionDatabase.cpp
--- a/Engine/Source/Editor/BlueprintGraph/Private/BlueprintActionDatabase.cpp
+++ b/Engine/Source/Editor/BlueprintGraph/Private/BlueprintActionDatabase.cpp
@@ -825,13 +825,20 @@ static void BlueprintActionDatabaseImpl::OnAssetsPendingDelete(TArray<UObject*>
 	FBlueprintActionDatabase& ActionDatabase = FBlueprintActionDatabase::Get();
 	for (UObject* DeletingObject : ObjectsForDelete)
 	{
+// 		if (!IsObjectValidForDatabase(DeletingObject))
+// 		{
+// 			continue;
+// 		}
+
 		// have to temporarily remove references (so that this delete isn't 
 		// blocked by dangling references)
-		ActionDatabase.ClearAssetActions(DeletingObject);
-
-		// in case they choose not to delete the object, we need to add 
-		// these back in to the database, so we track them here
-		PendingDelete.Add(DeletingObject);
+		if (ActionDatabase.ClearAssetActions(DeletingObject))
+		{
+			ensure(IsObjectValidForDatabase(DeletingObject));
+			// in case they choose not to delete the object, we need to add 
+			// these back in to the database, so we track them here
+			PendingDelete.Add(DeletingObject);
+		}		
 	}
 }
 
@@ -1148,7 +1155,6 @@ void FBlueprintActionDatabase::RefreshClassActions(UClass* const Class)
 void FBlueprintActionDatabase::RefreshAssetActions(UObject* const AssetObject)
 {
 	using namespace BlueprintActionDatabaseImpl;
-	check(BlueprintActionDatabaseImpl::IsObjectValidForDatabase(AssetObject));
 
 	bool const bHadExistingEntry = ActionRegistry.Contains(AssetObject);
 	FActionList& AssetActionList = ActionRegistry.FindOrAdd(AssetObject);
@@ -1161,6 +1167,11 @@ void FBlueprintActionDatabase::RefreshAssetActions(UObject* const AssetObject)
 	}
 	AssetActionList.Empty();
 
+	if (!ensure(IsObjectValidForDatabase(AssetObject)))
+	{
+		return;
+	}
+
 	if(const USkeleton* Skeleton = Cast<USkeleton>(AssetObject))
 	{
 		AddSkeletonActions(*Skeleton, AssetActionList);
@@ -1237,7 +1248,7 @@ void FBlueprintActionDatabase::RefreshComponentActions()
 }
 
 //------------------------------------------------------------------------------
-void FBlueprintActionDatabase::ClearAssetActions(UObject* const AssetObject)
+bool FBlueprintActionDatabase::ClearAssetActions(UObject* const AssetObject)
 {
 	FActionList* ActionList = ActionRegistry.Find(AssetObject);
 
@@ -1264,6 +1275,7 @@ void FBlueprintActionDatabase::ClearAssetActions(UObject* const AssetObject)
 	{
 		EntryRemovedDelegate.Broadcast(AssetObject);
 	}
+	return bHasEntry;
 }
 
 //------------------------------------------------------------------------------
diff --git a/Engine/Source/Editor/BlueprintGraph/Public/BlueprintActionDatabase.h b/Engine/Source/Editor/BlueprintGraph/Public/BlueprintActionDatabase.h
--- a/Engine/Source/Editor/BlueprintGraph/Public/BlueprintActionDatabase.h
+++ b/Engine/Source/Editor/BlueprintGraph/Public/BlueprintActionDatabase.h
@@ -96,8 +96,9 @@ public:
 	 * won't be rebuilt, unless RefreshAssetActions() is explicitly called after.
 	 * 
 	 * @param  AssetObject	
+	 * @return True if an entry was found and removed.
 	 */
-	void ClearAssetActions(UObject* const AssetObject);
+	bool ClearAssetActions(UObject* const AssetObject);
 	
 	/**
 	 * Finds the database entry for the specified unloaded asset and wipes it.
