Fix for crash when creating a generic select node with no pin for context (just right click and create a select node)

diff --git a/Engine/Source/Editor/BlueprintGraph/Private/K2Node_Select.cpp b/Engine/Source/Editor/BlueprintGraph/Private/K2Node_Select.cpp
--- a/Engine/Source/Editor/BlueprintGraph/Private/K2Node_Select.cpp
+++ b/Engine/Source/Editor/BlueprintGraph/Private/K2Node_Select.cpp
@@ -264,18 +264,21 @@ void UK2Node_Select::AllocateDefaultPins()
 
 void UK2Node_Select::AutowireNewNode(UEdGraphPin* FromPin)
 {
-	// Attempt to autowire to the index pin as users generally drag off of something intending to use
-	// it as an index in a select statement rather than an arbitrary entry:
-	const UEdGraphSchema_K2* K2Schema = CastChecked<UEdGraphSchema_K2>(GetSchema());
-	UEdGraphPin* IndexPin = GetIndexPin();
-	ECanCreateConnectionResponse ConnectResponse = K2Schema->CanCreateConnection(FromPin, IndexPin).Response;
-	if (ConnectResponse == ECanCreateConnectionResponse::CONNECT_RESPONSE_MAKE)
-	{
-		if (K2Schema->TryCreateConnection(FromPin, IndexPin))
+	if (FromPin)
+	{
+		// Attempt to autowire to the index pin as users generally drag off of something intending to use
+		// it as an index in a select statement rather than an arbitrary entry:
+		const UEdGraphSchema_K2* K2Schema = CastChecked<UEdGraphSchema_K2>(GetSchema());
+		UEdGraphPin* IndexPin = GetIndexPin();
+		ECanCreateConnectionResponse ConnectResponse = K2Schema->CanCreateConnection(FromPin, IndexPin).Response;
+		if (ConnectResponse == ECanCreateConnectionResponse::CONNECT_RESPONSE_MAKE)
 		{
-			FromPin->GetOwningNode()->NodeConnectionListChanged();
-			this->NodeConnectionListChanged();
-			return;
+			if (K2Schema->TryCreateConnection(FromPin, IndexPin))
+			{
+				FromPin->GetOwningNode()->NodeConnectionListChanged();
+				this->NodeConnectionListChanged();
+				return;
+			}
 		}
 	}
 
