BP: Fixed non-unity build

diff --git a/Engine/Source/Editor/BlueprintGraph/Classes/K2Node_EnumInequality.h b/Engine/Source/Editor/BlueprintGraph/Classes/K2Node_EnumInequality.h
--- a/Engine/Source/Editor/BlueprintGraph/Classes/K2Node_EnumInequality.h
+++ b/Engine/Source/Editor/BlueprintGraph/Classes/K2Node_EnumInequality.h
@@ -1,9 +1,12 @@
 // Copyright 1998-2015 Epic Games, Inc. All Rights Reserved.
 
 #pragma once
+
 #include "K2Node.h"
+#include "K2Node_EnumEquality.h"
 #include "K2Node_EnumInequality.generated.h"
 
+
 UCLASS(MinimalAPI)
 class UK2Node_EnumInequality : public UK2Node_EnumEquality
 {
@@ -22,4 +25,3 @@ class UK2Node_EnumInequality : public UK2Node_EnumEquality
 	/** Gets the name and class of the EqualEqual_ByteByte function */
 	BLUEPRINTGRAPH_API virtual void GetConditionalFunction(FName& FunctionName, UClass** FunctionClass) const override;
 };
-
