Cached node titles can now be forced to refresh without iteration over every node.

This will occur with every structural modification to any BP, node titles refresh only when visible.

diff --git a/Engine/Plugins/Experimental/LiveEditor/Source/LiveEditor/Private/K2Node_LiveEditObject.cpp b/Engine/Plugins/Experimental/LiveEditor/Source/LiveEditor/Private/K2Node_LiveEditObject.cpp
--- a/Engine/Plugins/Experimental/LiveEditor/Source/LiveEditor/Private/K2Node_LiveEditObject.cpp
+++ b/Engine/Plugins/Experimental/LiveEditor/Source/LiveEditor/Private/K2Node_LiveEditObject.cpp
@@ -190,7 +190,7 @@ FText UK2Node_LiveEditObject::GetNodeTitle(ENodeTitleType::Type TitleType) const
 	{
 		return NSLOCTEXT("K2Node", "LiveEditObject_NullTitle", "LiveEditObject NONE");
 	}
-	else if (CachedNodeTitle.IsOutOfDate())
+	else if (CachedNodeTitle.IsOutOfDate(this))
 	{
 		FNumberFormattingOptions NumberOptions;
 		NumberOptions.UseGrouping = false;
@@ -198,7 +198,7 @@ FText UK2Node_LiveEditObject::GetNodeTitle(ENodeTitleType::Type TitleType) const
 		Args.Add(TEXT("SpawnString"), FText::FromName(BaseClassPin->DefaultObject->GetFName()));
 		Args.Add(TEXT("ID"), FText::AsNumber(GetUniqueID(), &NumberOptions));
 
-		CachedNodeTitle = FText::Format(NSLOCTEXT("K2Node", "LiveEditObject", "LiveEditObject {SpawnString}_{ID}"), Args);
+		CachedNodeTitle.SetCachedText(FText::Format(NSLOCTEXT("K2Node", "LiveEditObject", "LiveEditObject {SpawnString}_{ID}"), Args), this);
 	}
 	return CachedNodeTitle;
 }
diff --git a/Engine/Source/Developer/MovieSceneTools/Private/K2Node_PlayMovieScene.cpp b/Engine/Source/Developer/MovieSceneTools/Private/K2Node_PlayMovieScene.cpp
--- a/Engine/Source/Developer/MovieSceneTools/Private/K2Node_PlayMovieScene.cpp
+++ b/Engine/Source/Developer/MovieSceneTools/Private/K2Node_PlayMovieScene.cpp
@@ -85,12 +85,12 @@ FText UK2Node_PlayMovieScene::GetNodeTitle(ENodeTitleType::Type TitleType) const
 	// @TODO: don't know enough about this node type to comfortably assert that
 	//        the MovieScene won't change after the node has spawned... until
 	//        then, we'll leave this optimization off
-	else //if (CachedNodeTitle.IsOutOfDate())
+	else //if (CachedNodeTitle.IsOutOfDate(this))
 	{
 		FFormatNamedArguments Args;
 		Args.Add(TEXT("SceneName"), FText::FromString(MovieScene->GetName()));
 		// FText::Format() is slow, so we cache this to save on performance
-		CachedNodeTitle = FText::Format(NSLOCTEXT("PlayMovieSceneNode", "NodeTitle", "Play Movie Scene: {SceneName}"), Args);
+		CachedNodeTitle.SetCachedText(FText::Format(NSLOCTEXT("PlayMovieSceneNode", "NodeTitle", "Play Movie Scene: {SceneName}"), Args), this);
 	}
 	return CachedNodeTitle;
 }
diff --git a/Engine/Source/Editor/AnimGraph/Private/AnimGraphNode_BlendListByEnum.cpp b/Engine/Source/Editor/AnimGraph/Private/AnimGraphNode_BlendListByEnum.cpp
--- a/Engine/Source/Editor/AnimGraph/Private/AnimGraphNode_BlendListByEnum.cpp
+++ b/Engine/Source/Editor/AnimGraph/Private/AnimGraphNode_BlendListByEnum.cpp
@@ -41,12 +41,12 @@ FText UAnimGraphNode_BlendListByEnum::GetNodeTitle(ENodeTitleType::Type TitleTyp
 	// @TODO: don't know enough about this node type to comfortably assert that
 	//        the BoundEnum won't change after the node has spawned... until
 	//        then, we'll leave this optimization off
-	else //if (CachedNodeTitle.IsOutOfDate())
+	else //if (CachedNodeTitle.IsOutOfDate(this))
 	{
 		FFormatNamedArguments Args;
 		Args.Add(TEXT("EnumName"), FText::FromString(BoundEnum->GetName()));
 		// FText::Format() is slow, so we cache this to save on performance
-		CachedNodeTitle = FText::Format(LOCTEXT("AnimGraphNode_BlendListByEnum_Title", "Blend Poses ({EnumName})"), Args);
+		CachedNodeTitle.SetCachedText(FText::Format(LOCTEXT("AnimGraphNode_BlendListByEnum_Title", "Blend Poses ({EnumName})"), Args), this);
 	}
 	return CachedNodeTitle;
 }
diff --git a/Engine/Source/Editor/AnimGraph/Private/AnimGraphNode_BlendSpaceEvaluator.cpp b/Engine/Source/Editor/AnimGraph/Private/AnimGraphNode_BlendSpaceEvaluator.cpp
--- a/Engine/Source/Editor/AnimGraph/Private/AnimGraphNode_BlendSpaceEvaluator.cpp
+++ b/Engine/Source/Editor/AnimGraph/Private/AnimGraphNode_BlendSpaceEvaluator.cpp
@@ -36,7 +36,7 @@ FText UAnimGraphNode_BlendSpaceEvaluator::GetNodeTitle(ENodeTitleType::Type Titl
 	}
 	// @TODO: the bone can be altered in the property editor, so we have to 
 	//        choose to mark this dirty when that happens for this to properly work
-	else //if (!CachedNodeTitles.IsTitleCached(TitleType))
+	else //if (!CachedNodeTitles.IsTitleCached(TitleType, this))
 	{
 		const FText BlendSpaceName = FText::FromString(Node.BlendSpace->GetName());
 
@@ -46,7 +46,7 @@ FText UAnimGraphNode_BlendSpaceEvaluator::GetNodeTitle(ENodeTitleType::Type Titl
 			Args.Add(TEXT("BlendSpaceName"), BlendSpaceName);
 
 			// FText::Format() is slow, so we cache this to save on performance
-			CachedNodeTitles.SetCachedTitle(TitleType, FText::Format(LOCTEXT("BlendSpaceEvaluatorListTitle", "Blendspace Evaluator '{BlendSpaceName}'"), Args));
+			CachedNodeTitles.SetCachedTitle(TitleType, FText::Format(LOCTEXT("BlendSpaceEvaluatorListTitle", "Blendspace Evaluator '{BlendSpaceName}'"), Args), this);
 		}
 		else
 		{
@@ -62,7 +62,7 @@ FText UAnimGraphNode_BlendSpaceEvaluator::GetNodeTitle(ENodeTitleType::Type Titl
 				Title = FText::Format(LOCTEXT("BlendSpaceNodeGroupSubtitle", "{Title}\nSync group {SyncGroupName}"), Args);
 			}
 			// FText::Format() is slow, so we cache this to save on performance
-			CachedNodeTitles.SetCachedTitle(TitleType, Title);
+			CachedNodeTitles.SetCachedTitle(TitleType, Title, this);
 		}
 	}
 	return CachedNodeTitles[TitleType];
diff --git a/Engine/Source/Editor/AnimGraph/Private/AnimGraphNode_BlendSpacePlayer.cpp b/Engine/Source/Editor/AnimGraph/Private/AnimGraphNode_BlendSpacePlayer.cpp
--- a/Engine/Source/Editor/AnimGraph/Private/AnimGraphNode_BlendSpacePlayer.cpp
+++ b/Engine/Source/Editor/AnimGraph/Private/AnimGraphNode_BlendSpacePlayer.cpp
@@ -38,7 +38,7 @@ FText UAnimGraphNode_BlendSpacePlayer::GetNodeTitle(ENodeTitleType::Type TitleTy
 	}
 	// @TODO: the bone can be altered in the property editor, so we have to 
 	//        choose to mark this dirty when that happens for this to properly work
-	else //if (!CachedNodeTitles.IsTitleCached(TitleType))
+	else //if (!CachedNodeTitles.IsTitleCached(TitleType, this))
 	{
 		const FText BlendSpaceName = FText::FromString(Node.BlendSpace->GetName());
 
@@ -47,7 +47,7 @@ FText UAnimGraphNode_BlendSpacePlayer::GetNodeTitle(ENodeTitleType::Type TitleTy
 			FFormatNamedArguments Args;
 			Args.Add(TEXT("BlendSpaceName"), BlendSpaceName);
 			// FText::Format() is slow, so we cache this to save on performance
-			CachedNodeTitles.SetCachedTitle(TitleType, FText::Format(LOCTEXT("BlendspacePlayer", "Blendspace Player '{BlendSpaceName}'"), Args));
+			CachedNodeTitles.SetCachedTitle(TitleType, FText::Format(LOCTEXT("BlendspacePlayer", "Blendspace Player '{BlendSpaceName}'"), Args), this);
 		}
 		else
 		{
@@ -63,7 +63,7 @@ FText UAnimGraphNode_BlendSpacePlayer::GetNodeTitle(ENodeTitleType::Type TitleTy
 				Title = FText::Format(LOCTEXT("BlendSpaceNodeGroupSubtitle", "{Title}\nSync group {SyncGroupName}"), Args);
 			}
 			// FText::Format() is slow, so we cache this to save on performance
-			CachedNodeTitles.SetCachedTitle(TitleType, Title);
+			CachedNodeTitles.SetCachedTitle(TitleType, Title, this);
 		}
 	}
 	return CachedNodeTitles[TitleType];
diff --git a/Engine/Source/Editor/AnimGraph/Private/AnimGraphNode_BoneDrivenController.cpp b/Engine/Source/Editor/AnimGraph/Private/AnimGraphNode_BoneDrivenController.cpp
--- a/Engine/Source/Editor/AnimGraph/Private/AnimGraphNode_BoneDrivenController.cpp
+++ b/Engine/Source/Editor/AnimGraph/Private/AnimGraphNode_BoneDrivenController.cpp
@@ -24,7 +24,7 @@ FText UAnimGraphNode_BoneDrivenController::GetNodeTitle(ENodeTitleType::Type Tit
 	}
 	// @TODO: the bone can be altered in the property editor, so we have to 
 	//        choose to mark this dirty when that happens for this to properly work
-	else //if (!CachedNodeTitles.IsTitleCached(TitleType))
+	else //if (!CachedNodeTitles.IsTitleCached(TitleType, this))
 	{
 		FFormatNamedArguments Args;
 		Args.Add(TEXT("ControllerDesc"), GetControllerDescription());
@@ -41,7 +41,7 @@ FText UAnimGraphNode_BoneDrivenController::GetNodeTitle(ENodeTitleType::Type Tit
 		}
 
 		// FText::Format() is slow, so we cache this to save on performance
-		CachedNodeTitles.SetCachedTitle(TitleType, FText::Format(LOCTEXT("AnimGraphNode_BoneDrivenController_Title", "{ControllerDesc}{Delim}Driving Bone: {SourceBone}{Delim}Driven Bone: {TargetBone}"), Args));
+		CachedNodeTitles.SetCachedTitle(TitleType, FText::Format(LOCTEXT("AnimGraphNode_BoneDrivenController_Title", "{ControllerDesc}{Delim}Driving Bone: {SourceBone}{Delim}Driven Bone: {TargetBone}"), Args), this);
 	}	
 	return CachedNodeTitles[TitleType];
 }
diff --git a/Engine/Source/Editor/AnimGraph/Private/AnimGraphNode_CopyBone.cpp b/Engine/Source/Editor/AnimGraph/Private/AnimGraphNode_CopyBone.cpp
--- a/Engine/Source/Editor/AnimGraph/Private/AnimGraphNode_CopyBone.cpp
+++ b/Engine/Source/Editor/AnimGraph/Private/AnimGraphNode_CopyBone.cpp
@@ -31,7 +31,7 @@ FText UAnimGraphNode_CopyBone::GetNodeTitle(ENodeTitleType::Type TitleType) cons
 	}
 	// @TODO: the bone can be altered in the property editor, so we have to 
 	//        choose to mark this dirty when that happens for this to properly work
-	else //if (!CachedNodeTitles.IsTitleCached(TitleType))
+	else //if (!CachedNodeTitles.IsTitleCached(TitleType, this))
 	{
 		FFormatNamedArguments Args;
 		Args.Add(TEXT("ControllerDescription"), GetControllerDescription());
@@ -40,11 +40,11 @@ FText UAnimGraphNode_CopyBone::GetNodeTitle(ENodeTitleType::Type TitleType) cons
 
 		if (TitleType == ENodeTitleType::ListView || TitleType == ENodeTitleType::MenuTitle)
 		{
-			CachedNodeTitles.SetCachedTitle(TitleType, FText::Format(LOCTEXT("AnimGraphNode_CopyBone_ListTitle", "{ControllerDescription} - Source Bone: {SourceBoneName} - Target Bone: {TargetBoneName}"), Args));
+			CachedNodeTitles.SetCachedTitle(TitleType, FText::Format(LOCTEXT("AnimGraphNode_CopyBone_ListTitle", "{ControllerDescription} - Source Bone: {SourceBoneName} - Target Bone: {TargetBoneName}"), Args), this);
 		}
 		else
 		{
-			CachedNodeTitles.SetCachedTitle(TitleType, FText::Format(LOCTEXT("AnimGraphNode_CopyBone_Title", "{ControllerDescription}\nSource Bone: {SourceBoneName}\nTarget Bone: {TargetBoneName}"), Args));
+			CachedNodeTitles.SetCachedTitle(TitleType, FText::Format(LOCTEXT("AnimGraphNode_CopyBone_Title", "{ControllerDescription}\nSource Bone: {SourceBoneName}\nTarget Bone: {TargetBoneName}"), Args), this);
 		}
 	}
 	return CachedNodeTitles[TitleType];
diff --git a/Engine/Source/Editor/AnimGraph/Private/AnimGraphNode_LookAt.cpp b/Engine/Source/Editor/AnimGraph/Private/AnimGraphNode_LookAt.cpp
--- a/Engine/Source/Editor/AnimGraph/Private/AnimGraphNode_LookAt.cpp
+++ b/Engine/Source/Editor/AnimGraph/Private/AnimGraphNode_LookAt.cpp
@@ -36,7 +36,7 @@ FText UAnimGraphNode_LookAt::GetNodeTitle(ENodeTitleType::Type TitleType) const
 	}
 	// @TODO: the bone can be altered in the property editor, so we have to 
 	//        choose to mark this dirty when that happens for this to properly work
-	else //if (!CachedNodeTitles.IsTitleCached(TitleType))
+	else //if (!CachedNodeTitles.IsTitleCached(TitleType, this))
 	{
 		FFormatNamedArguments Args;
 		Args.Add(TEXT("ControllerDescription"), GetControllerDescription());
@@ -45,11 +45,11 @@ FText UAnimGraphNode_LookAt::GetNodeTitle(ENodeTitleType::Type TitleType) const
 		// FText::Format() is slow, so we cache this to save on performance
 		if (TitleType == ENodeTitleType::ListView || TitleType == ENodeTitleType::MenuTitle)
 		{
-			CachedNodeTitles.SetCachedTitle(TitleType, FText::Format(LOCTEXT("AnimGraphNode_LookAt_ListTitle", "{ControllerDescription} - Bone: {BoneName}"), Args));
+			CachedNodeTitles.SetCachedTitle(TitleType, FText::Format(LOCTEXT("AnimGraphNode_LookAt_ListTitle", "{ControllerDescription} - Bone: {BoneName}"), Args), this);
 		}
 		else
 		{
-			CachedNodeTitles.SetCachedTitle(TitleType, FText::Format(LOCTEXT("AnimGraphNode_LookAt_Title", "{ControllerDescription}\nBone: {BoneName}"), Args));
+			CachedNodeTitles.SetCachedTitle(TitleType, FText::Format(LOCTEXT("AnimGraphNode_LookAt_Title", "{ControllerDescription}\nBone: {BoneName}"), Args), this);
 		}
 	}
 	return CachedNodeTitles[TitleType];		
diff --git a/Engine/Source/Editor/AnimGraph/Private/AnimGraphNode_ModifyBone.cpp b/Engine/Source/Editor/AnimGraph/Private/AnimGraphNode_ModifyBone.cpp
--- a/Engine/Source/Editor/AnimGraph/Private/AnimGraphNode_ModifyBone.cpp
+++ b/Engine/Source/Editor/AnimGraph/Private/AnimGraphNode_ModifyBone.cpp
@@ -37,7 +37,7 @@ FText UAnimGraphNode_ModifyBone::GetNodeTitle(ENodeTitleType::Type TitleType) co
 	}
 	// @TODO: the bone can be altered in the property editor, so we have to 
 	//        choose to mark this dirty when that happens for this to properly work
-	else //if (!CachedNodeTitles.IsTitleCached(TitleType))
+	else //if (!CachedNodeTitles.IsTitleCached(TitleType, this))
 	{
 		FFormatNamedArguments Args;
 		Args.Add(TEXT("ControllerDescription"), GetControllerDescription());
@@ -46,11 +46,11 @@ FText UAnimGraphNode_ModifyBone::GetNodeTitle(ENodeTitleType::Type TitleType) co
 		// FText::Format() is slow, so we cache this to save on performance
 		if (TitleType == ENodeTitleType::ListView || TitleType == ENodeTitleType::MenuTitle)
 		{
-			CachedNodeTitles.SetCachedTitle(TitleType, FText::Format(LOCTEXT("AnimGraphNode_ModifyBone_ListTitle", "{ControllerDescription} - Bone: {BoneName}"), Args));
+			CachedNodeTitles.SetCachedTitle(TitleType, FText::Format(LOCTEXT("AnimGraphNode_ModifyBone_ListTitle", "{ControllerDescription} - Bone: {BoneName}"), Args), this);
 		}
 		else
 		{
-			CachedNodeTitles.SetCachedTitle(TitleType, FText::Format(LOCTEXT("AnimGraphNode_ModifyBone_Title", "{ControllerDescription}\nBone: {BoneName}"), Args));
+			CachedNodeTitles.SetCachedTitle(TitleType, FText::Format(LOCTEXT("AnimGraphNode_ModifyBone_Title", "{ControllerDescription}\nBone: {BoneName}"), Args), this);
 		}
 	}
 	return CachedNodeTitles[TitleType];
diff --git a/Engine/Source/Editor/AnimGraph/Private/AnimGraphNode_RotationMultiplier.cpp b/Engine/Source/Editor/AnimGraph/Private/AnimGraphNode_RotationMultiplier.cpp
--- a/Engine/Source/Editor/AnimGraph/Private/AnimGraphNode_RotationMultiplier.cpp
+++ b/Engine/Source/Editor/AnimGraph/Private/AnimGraphNode_RotationMultiplier.cpp
@@ -30,7 +30,7 @@ FText UAnimGraphNode_RotationMultiplier::GetNodeTitle(ENodeTitleType::Type Title
 	}
 	// @TODO: the bone can be altered in the property editor, so we have to 
 	//        choose to mark this dirty when that happens for this to properly work
-	else //if (!CachedNodeTitles.IsTitleCached(TitleType))
+	else //if (!CachedNodeTitles.IsTitleCached(TitleType, this))
 	{
 		FFormatNamedArguments Args;
 		Args.Add(TEXT("ControllerDescription"), GetControllerDescription());
@@ -39,11 +39,11 @@ FText UAnimGraphNode_RotationMultiplier::GetNodeTitle(ENodeTitleType::Type Title
 		// FText::Format() is slow, so we cache this to save on performance
 		if (TitleType == ENodeTitleType::ListView || TitleType == ENodeTitleType::MenuTitle)
 		{
-			CachedNodeTitles.SetCachedTitle(TitleType, FText::Format(LOCTEXT("AnimGraphNode_ModifyBone_ListTitle", "{ControllerDescription} - Bone: {BoneName}"), Args));
+			CachedNodeTitles.SetCachedTitle(TitleType, FText::Format(LOCTEXT("AnimGraphNode_ModifyBone_ListTitle", "{ControllerDescription} - Bone: {BoneName}"), Args), this);
 		}
 		else
 		{
-			CachedNodeTitles.SetCachedTitle(TitleType, FText::Format(LOCTEXT("AnimGraphNode_ModifyBone_Title", "{ControllerDescription}\nBone: {BoneName}"), Args));
+			CachedNodeTitles.SetCachedTitle(TitleType, FText::Format(LOCTEXT("AnimGraphNode_ModifyBone_Title", "{ControllerDescription}\nBone: {BoneName}"), Args), this);
 		}
 	}
 	return CachedNodeTitles[TitleType];
diff --git a/Engine/Source/Editor/AnimGraph/Private/AnimGraphNode_RotationOffsetBlendSpace.cpp b/Engine/Source/Editor/AnimGraph/Private/AnimGraphNode_RotationOffsetBlendSpace.cpp
--- a/Engine/Source/Editor/AnimGraph/Private/AnimGraphNode_RotationOffsetBlendSpace.cpp
+++ b/Engine/Source/Editor/AnimGraph/Private/AnimGraphNode_RotationOffsetBlendSpace.cpp
@@ -38,7 +38,7 @@ FText UAnimGraphNode_RotationOffsetBlendSpace::GetNodeTitle(ENodeTitleType::Type
 	}
 	// @TODO: the bone can be altered in the property editor, so we have to 
 	//        choose to mark this dirty when that happens for this to properly work
-	else //if (!CachedNodeTitles.IsTitleCached(TitleType))
+	else //if (!CachedNodeTitles.IsTitleCached(TitleType, this))
 	{
 		const FText BlendSpaceName = FText::FromString(Node.BlendSpace->GetName());
 
@@ -48,11 +48,11 @@ FText UAnimGraphNode_RotationOffsetBlendSpace::GetNodeTitle(ENodeTitleType::Type
 		// FText::Format() is slow, so we cache this to save on performance
 		if (TitleType == ENodeTitleType::ListView || TitleType == ENodeTitleType::MenuTitle)
 		{
-			CachedNodeTitles.SetCachedTitle(TitleType, FText::Format(LOCTEXT("AimOffsetListTitle", "AimOffset '{BlendSpaceName}'"), Args));
+			CachedNodeTitles.SetCachedTitle(TitleType, FText::Format(LOCTEXT("AimOffsetListTitle", "AimOffset '{BlendSpaceName}'"), Args), this);
 		}
 		else
 		{
-			CachedNodeTitles.SetCachedTitle(TitleType, FText::Format(LOCTEXT("AimOffsetFullTitle", "{BlendSpaceName}\nAimOffset"), Args));
+			CachedNodeTitles.SetCachedTitle(TitleType, FText::Format(LOCTEXT("AimOffsetFullTitle", "{BlendSpaceName}\nAimOffset"), Args), this);
 		}
 	}
 	return CachedNodeTitles[TitleType];
diff --git a/Engine/Source/Editor/AnimGraph/Private/AnimGraphNode_SaveCachedPose.cpp b/Engine/Source/Editor/AnimGraph/Private/AnimGraphNode_SaveCachedPose.cpp
--- a/Engine/Source/Editor/AnimGraph/Private/AnimGraphNode_SaveCachedPose.cpp
+++ b/Engine/Source/Editor/AnimGraph/Private/AnimGraphNode_SaveCachedPose.cpp
@@ -55,14 +55,11 @@ FText UAnimGraphNode_SaveCachedPose::GetNodeTitle(ENodeTitleType::Type TitleType
 	{
 		return LOCTEXT("NewSaveCachedPose", "New Save cached pose...");
 	}
-	// @TODO: don't know enough about this node type to comfortably assert that
-	//        the CacheName won't change after the node has spawned... until
-	//        then, we'll leave this optimization off
-	else //if (CachedNodeTitle.IsOutOfDate())
+	else if (CachedNodeTitle.IsOutOfDate(this))
 	{
 		FFormatNamedArguments Args;
 		Args.Add(TEXT("NodeTitle"), FText::FromString(CacheName));
-		CachedNodeTitle = FText::Format(LOCTEXT("AnimGraphNode_SaveCachedPose_Title", "Save cached pose '{NodeTitle}'"), Args);
+		CachedNodeTitle.SetCachedText(FText::Format(LOCTEXT("AnimGraphNode_SaveCachedPose_Title", "Save cached pose '{NodeTitle}'"), Args), this);
 	}
 	return CachedNodeTitle;
 }
diff --git a/Engine/Source/Editor/AnimGraph/Private/AnimGraphNode_SequenceEvaluator.cpp b/Engine/Source/Editor/AnimGraph/Private/AnimGraphNode_SequenceEvaluator.cpp
--- a/Engine/Source/Editor/AnimGraph/Private/AnimGraphNode_SequenceEvaluator.cpp
+++ b/Engine/Source/Editor/AnimGraph/Private/AnimGraphNode_SequenceEvaluator.cpp
@@ -51,7 +51,7 @@ FText UAnimGraphNode_SequenceEvaluator::GetNodeTitle(ENodeTitleType::Type TitleT
 	// @TODO: don't know enough about this node type to comfortably assert that
 	//        the CacheName won't change after the node has spawned... until
 	//        then, we'll leave this optimization off
-	else //if (CachedNodeTitle.IsOutOfDate())
+	else //if (CachedNodeTitle.IsOutOfDate(this))
 	{
 		const FText SequenceName = FText::FromString(Node.Sequence->GetName());
 
@@ -61,11 +61,11 @@ FText UAnimGraphNode_SequenceEvaluator::GetNodeTitle(ENodeTitleType::Type TitleT
 		// FText::Format() is slow, so we cache this to save on performance
 		if (Node.Sequence->IsValidAdditive())
 		{
-			CachedNodeTitle = FText::Format(LOCTEXT("EvaluateSequence_Additive", "Evaluate {SequenceName} (additive)"), Args);
+			CachedNodeTitle.SetCachedText(FText::Format(LOCTEXT("EvaluateSequence_Additive", "Evaluate {SequenceName} (additive)"), Args), this);
 		}
 		else
 		{
-			CachedNodeTitle = FText::Format(LOCTEXT("EvaluateSequence", "Evaluate {SequenceName}"), Args);
+			CachedNodeTitle.SetCachedText(FText::Format(LOCTEXT("EvaluateSequence", "Evaluate {SequenceName}"), Args), this);
 		}
 	}
 
diff --git a/Engine/Source/Editor/AnimGraph/Private/AnimGraphNode_Slot.cpp b/Engine/Source/Editor/AnimGraph/Private/AnimGraphNode_Slot.cpp
--- a/Engine/Source/Editor/AnimGraph/Private/AnimGraphNode_Slot.cpp
+++ b/Engine/Source/Editor/AnimGraph/Private/AnimGraphNode_Slot.cpp
@@ -41,7 +41,7 @@ FText UAnimGraphNode_Slot::GetNodeTitle(ENodeTitleType::Type TitleType) const
 	}
 	// @TODO: the bone can be altered in the property editor, so we have to 
 	//        choose to mark this dirty when that happens for this to properly work
-	else //if (!CachedNodeTitles.IsTitleCached(TitleType))
+	else //if (!CachedNodeTitles.IsTitleCached(TitleType, this))
 	{
 		UAnimBlueprint* AnimBlueprint = GetAnimBlueprint();
 		FName GroupName = (AnimBlueprint->TargetSkeleton) ? AnimBlueprint->TargetSkeleton->GetSlotGroupName(Node.SlotName) : FAnimSlotGroup::DefaultGroupName;
@@ -53,11 +53,11 @@ FText UAnimGraphNode_Slot::GetNodeTitle(ENodeTitleType::Type TitleType) const
 		// FText::Format() is slow, so we cache this to save on performance
 		if (TitleType == ENodeTitleType::ListView || TitleType == ENodeTitleType::MenuTitle)
 		{
-			CachedNodeTitles.SetCachedTitle(TitleType, FText::Format(LOCTEXT("SlotNodeListTitle", "Slot '{SlotName}'"), Args));
+			CachedNodeTitles.SetCachedTitle(TitleType, FText::Format(LOCTEXT("SlotNodeListTitle", "Slot '{SlotName}'"), Args), this);
 		}
 		else
 		{
-			CachedNodeTitles.SetCachedTitle(TitleType, FText::Format(LOCTEXT("SlotNodeTitle", "Slot '{SlotName}'\nGroup '{GroupName}'"), Args));
+			CachedNodeTitles.SetCachedTitle(TitleType, FText::Format(LOCTEXT("SlotNodeTitle", "Slot '{SlotName}'\nGroup '{GroupName}'"), Args), this);
 		}
 	}
 	return CachedNodeTitles[TitleType];
diff --git a/Engine/Source/Editor/AnimGraph/Private/AnimGraphNode_SpringBone.cpp b/Engine/Source/Editor/AnimGraph/Private/AnimGraphNode_SpringBone.cpp
--- a/Engine/Source/Editor/AnimGraph/Private/AnimGraphNode_SpringBone.cpp
+++ b/Engine/Source/Editor/AnimGraph/Private/AnimGraphNode_SpringBone.cpp
@@ -31,7 +31,7 @@ FText UAnimGraphNode_SpringBone::GetNodeTitle(ENodeTitleType::Type TitleType) co
 	}
 	// @TODO: the bone can be altered in the property editor, so we have to 
 	//        choose to mark this dirty when that happens for this to properly work
-	else //if(!CachedNodeTitles.IsTitleCached(TitleType))
+	else //if(!CachedNodeTitles.IsTitleCached(TitleType, this))
 	{
 		FFormatNamedArguments Args;
 		Args.Add(TEXT("ControllerDescription"), GetControllerDescription());
@@ -40,11 +40,11 @@ FText UAnimGraphNode_SpringBone::GetNodeTitle(ENodeTitleType::Type TitleType) co
 		// FText::Format() is slow, so we cache this to save on performance
 		if (TitleType == ENodeTitleType::ListView || TitleType == ENodeTitleType::MenuTitle)
 		{
-			CachedNodeTitles.SetCachedTitle(TitleType, FText::Format(LOCTEXT("AnimGraphNode_SpringBone_ListTitle", "{ControllerDescription} - Bone: {BoneName}"), Args));
+			CachedNodeTitles.SetCachedTitle(TitleType, FText::Format(LOCTEXT("AnimGraphNode_SpringBone_ListTitle", "{ControllerDescription} - Bone: {BoneName}"), Args), this);
 		}
 		else
 		{
-			CachedNodeTitles.SetCachedTitle(TitleType, FText::Format(LOCTEXT("AnimGraphNode_SpringBone_Title", "{ControllerDescription}\nBone: {BoneName}"), Args));
+			CachedNodeTitles.SetCachedTitle(TitleType, FText::Format(LOCTEXT("AnimGraphNode_SpringBone_Title", "{ControllerDescription}\nBone: {BoneName}"), Args), this);
 		}
 	}
 
diff --git a/Engine/Source/Editor/AnimGraph/Private/AnimGraphNode_StateMachineBase.cpp b/Engine/Source/Editor/AnimGraph/Private/AnimGraphNode_StateMachineBase.cpp
--- a/Engine/Source/Editor/AnimGraph/Private/AnimGraphNode_StateMachineBase.cpp
+++ b/Engine/Source/Editor/AnimGraph/Private/AnimGraphNode_StateMachineBase.cpp
@@ -76,15 +76,12 @@ FText UAnimGraphNode_StateMachineBase::GetNodeTitle(ENodeTitleType::Type TitleTy
 	}
 	else if (TitleType == ENodeTitleType::FullTitle)
 	{
-		// @TODO: don't know enough about this node type to comfortably assert that
-		//        the EditorStateMachineGraph won't change after the node has 
-		//        spawned... until then, we'll leave this optimization off
-		//if (CachedFullTitle.IsOutOfDate())
+		if (CachedFullTitle.IsOutOfDate(this))
 		{
 			FFormatNamedArguments Args;
 			Args.Add(TEXT("Title"), FText::FromName(EditorStateMachineGraph->GetFName()));
 			// FText::Format() is slow, so we cache this to save on performance
-			CachedFullTitle = FText::Format(LOCTEXT("StateMachineFullTitle", "{Title}\nState Machine"), Args);
+			CachedFullTitle.SetCachedText(FText::Format(LOCTEXT("StateMachineFullTitle", "{Title}\nState Machine"), Args), this);
 		}
 		return CachedFullTitle;
 	}
diff --git a/Engine/Source/Editor/AnimGraph/Private/AnimGraphNode_Trail.cpp b/Engine/Source/Editor/AnimGraph/Private/AnimGraphNode_Trail.cpp
--- a/Engine/Source/Editor/AnimGraph/Private/AnimGraphNode_Trail.cpp
+++ b/Engine/Source/Editor/AnimGraph/Private/AnimGraphNode_Trail.cpp
@@ -31,7 +31,7 @@ FText UAnimGraphNode_Trail::GetNodeTitle(ENodeTitleType::Type TitleType) const
 	}
 	// @TODO: the bone can be altered in the property editor, so we have to 
 	//        choose to mark this dirty when that happens for this to properly work
-	else //if (!CachedNodeTitles.IsTitleCached(TitleType))
+	else //if (!CachedNodeTitles.IsTitleCached(TitleType, this))
 	{
 		FFormatNamedArguments Args;
 		Args.Add(TEXT("ControllerDescription"), GetControllerDescription());
@@ -39,11 +39,11 @@ FText UAnimGraphNode_Trail::GetNodeTitle(ENodeTitleType::Type TitleType) const
 
 		if (TitleType == ENodeTitleType::ListView || TitleType == ENodeTitleType::MenuTitle)
 		{
-			CachedNodeTitles.SetCachedTitle(TitleType, FText::Format(LOCTEXT("AnimGraphNode_Trail_ListTitle", "{ControllerDescription} - Bone: {BoneName}"), Args));
+			CachedNodeTitles.SetCachedTitle(TitleType, FText::Format(LOCTEXT("AnimGraphNode_Trail_ListTitle", "{ControllerDescription} - Bone: {BoneName}"), Args), this);
 		}
 		else
 		{
-			CachedNodeTitles.SetCachedTitle(TitleType, FText::Format(LOCTEXT("AnimGraphNode_Trail_Title", "{ControllerDescription}\nBone: {BoneName}"), Args));
+			CachedNodeTitles.SetCachedTitle(TitleType, FText::Format(LOCTEXT("AnimGraphNode_Trail_Title", "{ControllerDescription}\nBone: {BoneName}"), Args), this);
 		}
 	}
 	return CachedNodeTitles[TitleType];
diff --git a/Engine/Source/Editor/AnimGraph/Private/AnimGraphNode_TwoBoneIK.cpp b/Engine/Source/Editor/AnimGraph/Private/AnimGraphNode_TwoBoneIK.cpp
--- a/Engine/Source/Editor/AnimGraph/Private/AnimGraphNode_TwoBoneIK.cpp
+++ b/Engine/Source/Editor/AnimGraph/Private/AnimGraphNode_TwoBoneIK.cpp
@@ -58,7 +58,7 @@ FText UAnimGraphNode_TwoBoneIK::GetNodeTitle(ENodeTitleType::Type TitleType) con
 	}
 	// @TODO: the bone can be altered in the property editor, so we have to 
 	//        choose to mark this dirty when that happens for this to properly work
-	else //if (!CachedNodeTitles.IsTitleCached(TitleType))
+	else //if (!CachedNodeTitles.IsTitleCached(TitleType, this))
 	{
 		FFormatNamedArguments Args;
 		Args.Add(TEXT("ControllerDescription"), GetControllerDescription());
@@ -67,11 +67,11 @@ FText UAnimGraphNode_TwoBoneIK::GetNodeTitle(ENodeTitleType::Type TitleType) con
 		// FText::Format() is slow, so we cache this to save on performance
 		if (TitleType == ENodeTitleType::ListView || TitleType == ENodeTitleType::MenuTitle)
 		{
-			CachedNodeTitles.SetCachedTitle(TitleType, FText::Format(LOCTEXT("AnimGraphNode_IKBone_ListTitle", "{ControllerDescription} - Bone: {BoneName}"), Args));
+			CachedNodeTitles.SetCachedTitle(TitleType, FText::Format(LOCTEXT("AnimGraphNode_IKBone_ListTitle", "{ControllerDescription} - Bone: {BoneName}"), Args), this);
 		}
 		else
 		{
-			CachedNodeTitles.SetCachedTitle(TitleType, FText::Format(LOCTEXT("AnimGraphNode_IKBone_Title", "{ControllerDescription}\nBone: {BoneName}"), Args));
+			CachedNodeTitles.SetCachedTitle(TitleType, FText::Format(LOCTEXT("AnimGraphNode_IKBone_Title", "{ControllerDescription}\nBone: {BoneName}"), Args), this);
 		}
 	}
 	return CachedNodeTitles[TitleType];
diff --git a/Engine/Source/Editor/BlueprintGraph/Classes/EdGraphSchema_K2.h b/Engine/Source/Editor/BlueprintGraph/Classes/EdGraphSchema_K2.h
--- a/Engine/Source/Editor/BlueprintGraph/Classes/EdGraphSchema_K2.h
+++ b/Engine/Source/Editor/BlueprintGraph/Classes/EdGraphSchema_K2.h
@@ -304,6 +304,9 @@ class BLUEPRINTGRAPH_API UEdGraphSchema_K2 : public UEdGraphSchema
 	// Somewhat hacky mechanism to prevent tooltips created for pins from including the display name and type when generating BP API documentation
 	static bool bGeneratingDocumentation;
 
+	// ID for checking dirty status of node titles against, increases every compile
+	static int32 CurrentCacheRefreshID;
+
 	UPROPERTY(globalconfig)
 	TArray<FBlueprintCallableFunctionRedirect> EditoronlyBPFunctionRedirects;
 
@@ -427,6 +430,9 @@ public:
 	virtual void OnPinConnectionDoubleCicked(UEdGraphPin* PinA, UEdGraphPin* PinB, const FVector2D& GraphPosition) const override;
 	virtual UEdGraphPin* DropPinOnNode(UEdGraphNode* InTargetNode, const FString& InSourcePinName, const FEdGraphPinType& InSourcePinType, EEdGraphPinDirection InSourcePinDirection) const override;
 	virtual bool SupportsDropPinOnNode(UEdGraphNode* InTargetNode, const FEdGraphPinType& InSourcePinType, EEdGraphPinDirection InSourcePinDirection, FText& OutErrorMessage) const override;
+	virtual bool IsCacheVisualizationOutOfDate(int32 InVisualizationCacheID) const override;
+	virtual int32 GetCurrentVisualizationCacheID() const override;
+	virtual void ForceVisualizationCacheClear() const override;
 	// End EdGraphSchema Interface
 
 	/**
diff --git a/Engine/Source/Editor/BlueprintGraph/Private/EdGraphSchema_K2.cpp b/Engine/Source/Editor/BlueprintGraph/Private/EdGraphSchema_K2.cpp
--- a/Engine/Source/Editor/BlueprintGraph/Private/EdGraphSchema_K2.cpp
+++ b/Engine/Source/Editor/BlueprintGraph/Private/EdGraphSchema_K2.cpp
@@ -483,6 +483,7 @@ const UScriptStruct* UEdGraphSchema_K2::LinearColorStruct = nullptr;
 const UScriptStruct* UEdGraphSchema_K2::ColorStruct = nullptr;
 
 bool UEdGraphSchema_K2::bGeneratingDocumentation = false;
+int32 UEdGraphSchema_K2::CurrentCacheRefreshID = 0;
 
 UEdGraphSchema_K2::UEdGraphSchema_K2(const FObjectInitializer& ObjectInitializer)
 	: Super(ObjectInitializer)
@@ -5803,6 +5804,21 @@ bool UEdGraphSchema_K2::SupportsDropPinOnNode(UEdGraphNode* InTargetNode, const
 	return bIsSupported;
 }
 
+bool UEdGraphSchema_K2::IsCacheVisualizationOutOfDate(int32 InVisualizationCacheID) const
+{
+	return CurrentCacheRefreshID != InVisualizationCacheID;
+}
+
+int32 UEdGraphSchema_K2::GetCurrentVisualizationCacheID() const
+{
+	return CurrentCacheRefreshID;
+}
+
+void UEdGraphSchema_K2::ForceVisualizationCacheClear() const
+{
+	++CurrentCacheRefreshID;
+}
+
 /////////////////////////////////////////////////////
 
 #undef LOCTEXT_NAMESPACE
diff --git a/Engine/Source/Editor/BlueprintGraph/Private/K2Node_ActorBoundEvent.cpp b/Engine/Source/Editor/BlueprintGraph/Private/K2Node_ActorBoundEvent.cpp
--- a/Engine/Source/Editor/BlueprintGraph/Private/K2Node_ActorBoundEvent.cpp
+++ b/Engine/Source/Editor/BlueprintGraph/Private/K2Node_ActorBoundEvent.cpp
@@ -89,14 +89,14 @@ FText UK2Node_ActorBoundEvent::GetNodeTitle(ENodeTitleType::Type TitleType) cons
 		Args.Add(TEXT("DelegatePropertyName"), FText::FromName(DelegatePropertyName));
 		return FText::Format(LOCTEXT("ActorBoundEventTitle", "{DelegatePropertyName} (None)"), Args);
 	}
-	else if (CachedNodeTitle.IsOutOfDate())
+	else if (CachedNodeTitle.IsOutOfDate(this))
 	{
 		FFormatNamedArguments Args;
 		Args.Add(TEXT("DelegatePropertyName"), FText::FromName(DelegatePropertyName));
 		Args.Add(TEXT("TargetName"), FText::FromString(EventOwner->GetActorLabel()));
 
 		// FText::Format() is slow, so we cache this to save on performance
-		CachedNodeTitle = FText::Format(LOCTEXT("ActorBoundEventTitle", "{DelegatePropertyName} ({TargetName})"), Args);
+		CachedNodeTitle.SetCachedText(FText::Format(LOCTEXT("ActorBoundEventTitle", "{DelegatePropertyName} ({TargetName})"), Args), this);
 	}
 	return CachedNodeTitle;
 }
diff --git a/Engine/Source/Editor/BlueprintGraph/Private/K2Node_AssignDelegate.cpp b/Engine/Source/Editor/BlueprintGraph/Private/K2Node_AssignDelegate.cpp
--- a/Engine/Source/Editor/BlueprintGraph/Private/K2Node_AssignDelegate.cpp
+++ b/Engine/Source/Editor/BlueprintGraph/Private/K2Node_AssignDelegate.cpp
@@ -22,7 +22,7 @@ FText UK2Node_AssignDelegate::GetNodeTitle(ENodeTitleType::Type TitleType) const
 {
 	if (TitleType == ENodeTitleType::ListView || TitleType == ENodeTitleType::MenuTitle)
 	{
-		if (CachedListTitle.IsOutOfDate())
+		if (CachedListTitle.IsOutOfDate(this))
 		{
 			FText PropertyName;
 			if (UProperty* Property = GetProperty())
@@ -31,7 +31,7 @@ FText UK2Node_AssignDelegate::GetNodeTitle(ENodeTitleType::Type TitleType) const
 				PropertyName = FText::FromString(bShowFriendlyNames ? UEditorEngine::GetFriendlyName(Property) : Property->GetName());
 
 				// FText::Format() is slow, so we cache this to save on performance
-				CachedListTitle = FText::Format(LOCTEXT("AssignDelegateTitle", "Assign {0}"), PropertyName);
+				CachedListTitle.SetCachedText(FText::Format(LOCTEXT("AssignDelegateTitle", "Assign {0}"), PropertyName), this);
 			}
 			else
 			{
diff --git a/Engine/Source/Editor/BlueprintGraph/Private/K2Node_BreakStruct.cpp b/Engine/Source/Editor/BlueprintGraph/Private/K2Node_BreakStruct.cpp
--- a/Engine/Source/Editor/BlueprintGraph/Private/K2Node_BreakStruct.cpp
+++ b/Engine/Source/Editor/BlueprintGraph/Private/K2Node_BreakStruct.cpp
@@ -197,13 +197,13 @@ FText UK2Node_BreakStruct::GetNodeTitle(ENodeTitleType::Type TitleType) const
 	{
 		return LOCTEXT("BreakNullStruct_Title", "Break <unknown struct>");
 	}
-	else if (CachedNodeTitle.IsOutOfDate())
+	else if (CachedNodeTitle.IsOutOfDate(this))
 	{
 		FFormatNamedArguments Args;
 		Args.Add(TEXT("StructName"), FText::FromString(StructType->GetName()));
 
 		// FText::Format() is slow, so we cache this to save on performance
-		CachedNodeTitle = FText::Format(LOCTEXT("BreakNodeTitle", "Break {StructName}"), Args);
+		CachedNodeTitle.SetCachedText(FText::Format(LOCTEXT("BreakNodeTitle", "Break {StructName}"), Args), this);
 	}
 	return CachedNodeTitle;
 }
@@ -214,13 +214,13 @@ FText UK2Node_BreakStruct::GetTooltipText() const
 	{
 		return LOCTEXT("BreakNullStruct_Tooltip", "Adds a node that breaks an '<unknown struct>' into its member fields");
 	}
-	else if (CachedTooltip.IsOutOfDate())
+	else if (CachedTooltip.IsOutOfDate(this))
 	{
 		// FText::Format() is slow, so we cache this to save on performance
-		CachedTooltip = FText::Format(
+		CachedTooltip.SetCachedText(FText::Format(
 			LOCTEXT("BreakStruct_Tooltip", "Adds a node that breaks a '{0}' into its member fields"),
 			FText::FromName(StructType->GetFName())
-		);
+		), this);
 	}
 	return CachedTooltip;
 }
diff --git a/Engine/Source/Editor/BlueprintGraph/Private/K2Node_CallFunction.cpp b/Engine/Source/Editor/BlueprintGraph/Private/K2Node_CallFunction.cpp
--- a/Engine/Source/Editor/BlueprintGraph/Private/K2Node_CallFunction.cpp
+++ b/Engine/Source/Editor/BlueprintGraph/Private/K2Node_CallFunction.cpp
@@ -1121,7 +1121,7 @@ FText UK2Node_CallFunction::GetTooltipText() const
 	{
 		return FText::Format(LOCTEXT("CallUnknownFunction", "Call unknown function {0}"), FText::FromName(FunctionReference.GetMemberName()));
 	}
-	else if (CachedTooltip.IsOutOfDate())
+	else if (CachedTooltip.IsOutOfDate(this))
 	{
 		FText BaseTooltip = FText::FromString(GetDefaultTooltipForFunction(Function));
 
@@ -1135,7 +1135,7 @@ FText UK2Node_CallFunction::GetTooltipText() const
 				NSLOCTEXT("K2Node", "ServerFunction", "Authority Only. This function will only execute on the server.")
 			);
 			// FText::Format() is slow, so we cache this to save on performance
-			CachedTooltip = FText::Format(LOCTEXT("CallFunction_SubtitledTooltip", "{DefaultTooltip}\n\n{ClientString}"), Args);
+			CachedTooltip.SetCachedText(FText::Format(LOCTEXT("CallFunction_SubtitledTooltip", "{DefaultTooltip}\n\n{ClientString}"), Args), this);
 		}
 		else if (Function->HasAllFunctionFlags(FUNC_BlueprintCosmetic))
 		{
@@ -1144,11 +1144,11 @@ FText UK2Node_CallFunction::GetTooltipText() const
 				NSLOCTEXT("K2Node", "ClientEvent", "Cosmetic. This event is only for cosmetic, non-gameplay actions.")
 			);
 			// FText::Format() is slow, so we cache this to save on performance
-			CachedTooltip = FText::Format(LOCTEXT("CallFunction_SubtitledTooltip", "{DefaultTooltip}\n\n{ClientString}"), Args);
+			CachedTooltip.SetCachedText(FText::Format(LOCTEXT("CallFunction_SubtitledTooltip", "{DefaultTooltip}\n\n{ClientString}"), Args), this);
 		} 
 		else
 		{
-			CachedTooltip = BaseTooltip;
+			CachedTooltip.SetCachedText(BaseTooltip, this);
 		}
 	}
 	return CachedTooltip;
diff --git a/Engine/Source/Editor/BlueprintGraph/Private/K2Node_CastByteToEnum.cpp b/Engine/Source/Editor/BlueprintGraph/Private/K2Node_CastByteToEnum.cpp
--- a/Engine/Source/Editor/BlueprintGraph/Private/K2Node_CastByteToEnum.cpp
+++ b/Engine/Source/Editor/BlueprintGraph/Private/K2Node_CastByteToEnum.cpp
@@ -39,13 +39,13 @@ FText UK2Node_CastByteToEnum::GetTooltipText() const
 	{
 		return NSLOCTEXT("K2Node", "CastByteToEnum_NullTooltip", "Byte to Enum (bad enum)");
 	}
-	else if(CachedTooltip.IsOutOfDate())
+	else if(CachedTooltip.IsOutOfDate(this))
 	{
 		// FText::Format() is slow, so we cache this to save on performance
-		CachedTooltip = FText::Format(
+		CachedTooltip.SetCachedText(FText::Format(
 			NSLOCTEXT("K2Node", "CastByteToEnum_Tooltip", "Byte to Enum {0}"),
 			FText::FromName(Enum->GetFName())
-		);
+		), this);
 	}
 	return CachedTooltip;
 }
diff --git a/Engine/Source/Editor/BlueprintGraph/Private/K2Node_ClassDynamicCast.cpp b/Engine/Source/Editor/BlueprintGraph/Private/K2Node_ClassDynamicCast.cpp
--- a/Engine/Source/Editor/BlueprintGraph/Private/K2Node_ClassDynamicCast.cpp
+++ b/Engine/Source/Editor/BlueprintGraph/Private/K2Node_ClassDynamicCast.cpp
@@ -73,9 +73,9 @@ FLinearColor UK2Node_ClassDynamicCast::GetNodeTitleColor() const
 
 FText UK2Node_ClassDynamicCast::GetNodeTitle(ENodeTitleType::Type TitleType) const
 {
-	if (CachedNodeTitle.IsOutOfDate())
+	if (CachedNodeTitle.IsOutOfDate(this))
 	{
-		CachedNodeTitle = FText::Format(LOCTEXT("NodeTitle", "{0} Class"), Super::GetNodeTitle(TitleType));
+		CachedNodeTitle.SetCachedText(FText::Format(LOCTEXT("NodeTitle", "{0} Class"), Super::GetNodeTitle(TitleType)), this);
 	}
 	return CachedNodeTitle;
 }
diff --git a/Engine/Source/Editor/BlueprintGraph/Private/K2Node_ComponentBoundEvent.cpp b/Engine/Source/Editor/BlueprintGraph/Private/K2Node_ComponentBoundEvent.cpp
--- a/Engine/Source/Editor/BlueprintGraph/Private/K2Node_ComponentBoundEvent.cpp
+++ b/Engine/Source/Editor/BlueprintGraph/Private/K2Node_ComponentBoundEvent.cpp
@@ -20,14 +20,14 @@ bool UK2Node_ComponentBoundEvent::Modify(bool bAlwaysMarkDirty)
 
 FText UK2Node_ComponentBoundEvent::GetNodeTitle(ENodeTitleType::Type TitleType) const
 {
-	if (CachedNodeTitle.IsOutOfDate())
+	if (CachedNodeTitle.IsOutOfDate(this))
 	{
 		FFormatNamedArguments Args;
 		Args.Add(TEXT("DelegatePropertyName"), FText::FromName(DelegatePropertyName));
 		Args.Add(TEXT("ComponentPropertyName"), FText::FromName(ComponentPropertyName));
 
 		// FText::Format() is slow, so we cache this to save on performance
-		CachedNodeTitle = FText::Format(LOCTEXT("ComponentBoundEvent_Title", "{DelegatePropertyName} ({ComponentPropertyName})"), Args);
+		CachedNodeTitle.SetCachedText(FText::Format(LOCTEXT("ComponentBoundEvent_Title", "{DelegatePropertyName} ({ComponentPropertyName})"), Args), this);
 	}
 	return CachedNodeTitle;
 }
diff --git a/Engine/Source/Editor/BlueprintGraph/Private/K2Node_Composite.cpp b/Engine/Source/Editor/BlueprintGraph/Private/K2Node_Composite.cpp
--- a/Engine/Source/Editor/BlueprintGraph/Private/K2Node_Composite.cpp
+++ b/Engine/Source/Editor/BlueprintGraph/Private/K2Node_Composite.cpp
@@ -173,12 +173,12 @@ FText UK2Node_Composite::GetNodeTitle(ENodeTitleType::Type TitleType) const
 	{
 		return FText::FromString(BoundGraph->GetName());
 	}
-	else if (CachedNodeTitle.IsOutOfDate()) // TitleType == ENodeTitleType::FullTitle
+	else if (CachedNodeTitle.IsOutOfDate(this)) // TitleType == ENodeTitleType::FullTitle
 	{
 		FFormatNamedArguments Args;
 		Args.Add(TEXT("BoundGraphName"), (BoundGraph) ? FText::FromString(BoundGraph->GetName()) : LOCTEXT("InvalidGraph", "Invalid Graph"));
 		// FText::Format() is slow, so we cache this to save on performance
-		CachedNodeTitle = FText::Format(LOCTEXT("Collapsed_Name", "{BoundGraphName}\nCollapsed Graph"), Args);
+		CachedNodeTitle.SetCachedText(FText::Format(LOCTEXT("Collapsed_Name", "{BoundGraphName}\nCollapsed Graph"), Args), this);
 	}
 	return CachedNodeTitle;
 }
diff --git a/Engine/Source/Editor/BlueprintGraph/Private/K2Node_ConstructObjectFromClass.cpp b/Engine/Source/Editor/BlueprintGraph/Private/K2Node_ConstructObjectFromClass.cpp
--- a/Engine/Source/Editor/BlueprintGraph/Private/K2Node_ConstructObjectFromClass.cpp
+++ b/Engine/Source/Editor/BlueprintGraph/Private/K2Node_ConstructObjectFromClass.cpp
@@ -269,12 +269,12 @@ FText UK2Node_ConstructObjectFromClass::GetNodeTitle(ENodeTitleType::Type TitleT
 		{
 			return NSLOCTEXT("K2Node", "ConstructObject_Title_NONE", "Construct NONE");
 		}
-		else if (CachedNodeTitle.IsOutOfDate())
+		else if (CachedNodeTitle.IsOutOfDate(this))
 		{
 			FFormatNamedArguments Args;
 			Args.Add(TEXT("ClassName"), FText::FromString(ClassPin->DefaultObject->GetName()));
 			// FText::Format() is slow, so we cache this to save on performance
-			CachedNodeTitle = FText::Format(GetNodeTitleFormat(), Args);
+			CachedNodeTitle.SetCachedText(FText::Format(GetNodeTitleFormat(), Args), this);
 		}
 	}
 	else
diff --git a/Engine/Source/Editor/BlueprintGraph/Private/K2Node_CustomEvent.cpp b/Engine/Source/Editor/BlueprintGraph/Private/K2Node_CustomEvent.cpp
--- a/Engine/Source/Editor/BlueprintGraph/Private/K2Node_CustomEvent.cpp
+++ b/Engine/Source/Editor/BlueprintGraph/Private/K2Node_CustomEvent.cpp
@@ -122,7 +122,7 @@ FText UK2Node_CustomEvent::GetNodeTitle(ENodeTitleType::Type TitleType) const
 	{
 		return FText::FromName(CustomFunctionName);
 	}
-	else if (CachedNodeTitle.IsOutOfDate())
+	else if (CachedNodeTitle.IsOutOfDate(this))
 	{
 		FString RPCString = UK2Node_Event::GetLocalizedNetString(FunctionFlags, false);
 		
@@ -131,7 +131,7 @@ FText UK2Node_CustomEvent::GetNodeTitle(ENodeTitleType::Type TitleType) const
 		Args.Add(TEXT("RPCString"), FText::FromString(RPCString));
 
 		// FText::Format() is slow, so we cache this to save on performance
-		CachedNodeTitle = FText::Format(NSLOCTEXT("K2Node", "CustomEvent_Name", "{FunctionName}{RPCString}\nCustom Event"), Args);
+		CachedNodeTitle.SetCachedText(FText::Format(NSLOCTEXT("K2Node", "CustomEvent_Name", "{FunctionName}{RPCString}\nCustom Event"), Args), this);
 	}
 
 	return CachedNodeTitle;
diff --git a/Engine/Source/Editor/BlueprintGraph/Private/K2Node_DelegateSet.cpp b/Engine/Source/Editor/BlueprintGraph/Private/K2Node_DelegateSet.cpp
--- a/Engine/Source/Editor/BlueprintGraph/Private/K2Node_DelegateSet.cpp
+++ b/Engine/Source/Editor/BlueprintGraph/Private/K2Node_DelegateSet.cpp
@@ -148,17 +148,17 @@ void UK2Node_DelegateSet::AllocateDefaultPins()
 
 FText UK2Node_DelegateSet::GetTooltipText() const
 {
-	if (CachedTooltip.IsOutOfDate())
+	if (CachedTooltip.IsOutOfDate(this))
 	{
 		// FText::Format() is slow, so we cache this to save on performance
-		CachedTooltip = FText::Format(NSLOCTEXT("K2Node", "CreateEventForDelegate", "Create an event tied to the delegate {0}"), FText::FromName(DelegatePropertyName));
+		CachedTooltip.SetCachedText(FText::Format(NSLOCTEXT("K2Node", "CreateEventForDelegate", "Create an event tied to the delegate {0}"), FText::FromName(DelegatePropertyName)), this);
 		if (UFunction* Function = GetDelegateSignature())
 		{
 			const FText SignatureTooltip = Function->GetToolTipText();
 
 			if (!SignatureTooltip.IsEmpty())
 			{
-				CachedTooltip = FText::Format(LOCTEXT("DelegateSet_SubtitledTooltip", "{0}\n{1}"), (FText&)CachedTooltip, SignatureTooltip);
+				CachedTooltip.SetCachedText(FText::Format(LOCTEXT("DelegateSet_SubtitledTooltip", "{0}\n{1}"), (FText&)CachedTooltip, SignatureTooltip), this);
 			}
 		}
 	}
@@ -167,12 +167,12 @@ FText UK2Node_DelegateSet::GetTooltipText() const
 
 FText UK2Node_DelegateSet::GetNodeTitle(ENodeTitleType::Type TitleType) const
 {
-	if (CachedNodeTitle.IsOutOfDate())
+	if (CachedNodeTitle.IsOutOfDate(this))
 	{
 		FFormatNamedArguments Args;
 		Args.Add(TEXT("DelegatePropertyName"), FText::FromName(DelegatePropertyName));
 		// FText::Format() is slow, so we cache this to save on performance
-		CachedNodeTitle = FText::Format(NSLOCTEXT("K2Node", "Assign_Name", "Assign {DelegatePropertyName}"), Args);
+		CachedNodeTitle.SetCachedText(FText::Format(NSLOCTEXT("K2Node", "Assign_Name", "Assign {DelegatePropertyName}"), Args), this);
 	}
 	return CachedNodeTitle;
 }
diff --git a/Engine/Source/Editor/BlueprintGraph/Private/K2Node_DynamicCast.cpp b/Engine/Source/Editor/BlueprintGraph/Private/K2Node_DynamicCast.cpp
--- a/Engine/Source/Editor/BlueprintGraph/Private/K2Node_DynamicCast.cpp
+++ b/Engine/Source/Editor/BlueprintGraph/Private/K2Node_DynamicCast.cpp
@@ -77,7 +77,7 @@ FText UK2Node_DynamicCast::GetNodeTitle(ENodeTitleType::Type TitleType) const
 	{
 		return NSLOCTEXT("K2Node_DynamicCast", "BadCastNode", "Bad cast node");
 	}
-	else if (CachedNodeTitle.IsOutOfDate())
+	else if (CachedNodeTitle.IsOutOfDate(this))
 	{
 		// If casting to BP class, use BP name not class name (ie. remove the _C)
 		FString TargetName;
@@ -95,7 +95,7 @@ FText UK2Node_DynamicCast::GetNodeTitle(ENodeTitleType::Type TitleType) const
 		Args.Add(TEXT("TargetName"), FText::FromString(TargetName));
 
 		// FText::Format() is slow, so we cache this to save on performance
-		CachedNodeTitle = FText::Format(NSLOCTEXT("K2Node_DynamicCast", "CastTo", "Cast To {TargetName}"), Args);
+		CachedNodeTitle.SetCachedText(FText::Format(NSLOCTEXT("K2Node_DynamicCast", "CastTo", "Cast To {TargetName}"), Args), this);
 	}
 	return CachedNodeTitle;
 }
@@ -284,10 +284,10 @@ bool UK2Node_DynamicCast::HasExternalBlueprintDependencies(TArray<class UStruct*
 FText UK2Node_DynamicCast::GetMenuCategory() const
 {
 	static FNodeTextCache CachedCategory;
-	if (CachedCategory.IsOutOfDate())
+	if (CachedCategory.IsOutOfDate(this))
 	{
 		// FText::Format() is slow, so we cache this to save on performance
-		CachedCategory = FEditorCategoryUtils::BuildCategoryString(FCommonEditorCategory::Utilities, LOCTEXT("ActionMenuCategory", "Casting"));
+		CachedCategory.SetCachedText(FEditorCategoryUtils::BuildCategoryString(FCommonEditorCategory::Utilities, LOCTEXT("ActionMenuCategory", "Casting")), this);
 	}
 	return CachedCategory;
 }
diff --git a/Engine/Source/Editor/BlueprintGraph/Private/K2Node_EaseFunction.cpp b/Engine/Source/Editor/BlueprintGraph/Private/K2Node_EaseFunction.cpp
--- a/Engine/Source/Editor/BlueprintGraph/Private/K2Node_EaseFunction.cpp
+++ b/Engine/Source/Editor/BlueprintGraph/Private/K2Node_EaseFunction.cpp
@@ -260,10 +260,10 @@ void UK2Node_EaseFunction::GetMenuActions(FBlueprintActionDatabaseRegistrar& Act
 FText UK2Node_EaseFunction::GetMenuCategory() const
 {
 	static FNodeTextCache CachedCategory;
-	if (CachedCategory.IsOutOfDate())
+	if (CachedCategory.IsOutOfDate(this))
 	{
 		// FText::Format() is slow, so we cache this to save on performance
-		CachedCategory = FEditorCategoryUtils::BuildCategoryString(FCommonEditorCategory::Math, LOCTEXT("InterpCategory", "Interpolation"));
+		CachedCategory.SetCachedText(FEditorCategoryUtils::BuildCategoryString(FCommonEditorCategory::Math, LOCTEXT("InterpCategory", "Interpolation")), this);
 	}
 	return CachedCategory;
 }
diff --git a/Engine/Source/Editor/BlueprintGraph/Private/K2Node_EnumLiteral.cpp b/Engine/Source/Editor/BlueprintGraph/Private/K2Node_EnumLiteral.cpp
--- a/Engine/Source/Editor/BlueprintGraph/Private/K2Node_EnumLiteral.cpp
+++ b/Engine/Source/Editor/BlueprintGraph/Private/K2Node_EnumLiteral.cpp
@@ -52,18 +52,18 @@ FText UK2Node_EnumLiteral::GetTooltipText() const
 	{
 		return NSLOCTEXT("K2Node", "BadEnumLiteral_Tooltip", "Literal enum (bad enum)");
 	}
-	else if (CachedTooltip.IsOutOfDate())
+	else if (CachedTooltip.IsOutOfDate(this))
 	{
 		FFormatNamedArguments Args;
 		Args.Add(TEXT("EnumName"), FText::FromName(Enum->GetFName()));
-		CachedTooltip = FText::Format(NSLOCTEXT("K2Node", "EnumLiteral_Tooltip", "Literal enum {EnumName}"), Args);
+		CachedTooltip.SetCachedText(FText::Format(NSLOCTEXT("K2Node", "EnumLiteral_Tooltip", "Literal enum {EnumName}"), Args), this);
 	}
 	return CachedTooltip;	
 }
 
 FText UK2Node_EnumLiteral::GetNodeTitle(ENodeTitleType::Type TitleType) const
 {
-	if (CachedTooltip.IsOutOfDate())
+	if (CachedTooltip.IsOutOfDate(this))
 	{
 		return GetTooltipText();
 	}
diff --git a/Engine/Source/Editor/BlueprintGraph/Private/K2Node_Event.cpp b/Engine/Source/Editor/BlueprintGraph/Private/K2Node_Event.cpp
--- a/Engine/Source/Editor/BlueprintGraph/Private/K2Node_Event.cpp
+++ b/Engine/Source/Editor/BlueprintGraph/Private/K2Node_Event.cpp
@@ -107,9 +107,9 @@ FText UK2Node_Event::GetNodeTitle(ENodeTitleType::Type TitleType) const
 FText UK2Node_Event::GetTooltipText() const
 {
 	UFunction* Function = EventReference.ResolveMember<UFunction>(GetBlueprintClassFromNode());
-	if (CachedTooltip.IsOutOfDate() && (Function != nullptr))
+	if (CachedTooltip.IsOutOfDate(this) && (Function != nullptr))
 	{
-		CachedTooltip = FText::FromString(UK2Node_CallFunction::GetDefaultTooltipForFunction(Function));
+		CachedTooltip.SetCachedText(FText::FromString(UK2Node_CallFunction::GetDefaultTooltipForFunction(Function)), this);
 
 		if (bOverrideFunction || (CustomFunctionName == NAME_None))
 		{
@@ -124,7 +124,7 @@ FText UK2Node_Event::GetTooltipText() const
 					NSLOCTEXT("K2Node", "ClientEvent", "\n\nCosmetic. This event is only for cosmetic, non-gameplay actions.")
 				);
 				// FText::Format() is slow, so we cache this to save on performance
-				CachedTooltip = FText::Format(LOCTEXT("Event_SubtitledTooltip", "{FunctionTooltip}\n\n{ClientString}"), Args);
+				CachedTooltip.SetCachedText(FText::Format(LOCTEXT("Event_SubtitledTooltip", "{FunctionTooltip}\n\n{ClientString}"), Args), this);
 			} 
 			else if(Function->HasAllFunctionFlags(FUNC_BlueprintAuthorityOnly))
 			{
@@ -133,7 +133,7 @@ FText UK2Node_Event::GetTooltipText() const
 					NSLOCTEXT("K2Node", "ServerEvent", "\n\nAuthority Only. This event only fires on the server.")
 				);
 				// FText::Format() is slow, so we cache this to save on performance
-				CachedTooltip = FText::Format(LOCTEXT("Event_SubtitledTooltip", "{FunctionTooltip}\n\n{ClientString}"), Args);
+				CachedTooltip.SetCachedText(FText::Format(LOCTEXT("Event_SubtitledTooltip", "{FunctionTooltip}\n\n{ClientString}"), Args), this);
 			}			
 		}		
 	}
diff --git a/Engine/Source/Editor/BlueprintGraph/Private/K2Node_ForEachElementInEnum.cpp b/Engine/Source/Editor/BlueprintGraph/Private/K2Node_ForEachElementInEnum.cpp
--- a/Engine/Source/Editor/BlueprintGraph/Private/K2Node_ForEachElementInEnum.cpp
+++ b/Engine/Source/Editor/BlueprintGraph/Private/K2Node_ForEachElementInEnum.cpp
@@ -142,12 +142,12 @@ FText UK2Node_ForEachElementInEnum::GetNodeTitle(ENodeTitleType::Type TitleType)
 	{
 		return LOCTEXT("ForEachElementInUnknownEnum_Title", "ForEach UNKNOWN");
 	}
-	else if (CachedNodeTitle.IsOutOfDate())
+	else if (CachedNodeTitle.IsOutOfDate(this))
 	{
 		FFormatNamedArguments Args;
 		Args.Add(TEXT("EnumName"), FText::FromName(Enum->GetFName()));
 		// FText::Format() is slow, so we cache this to save on performance
-		CachedNodeTitle = FText::Format(LOCTEXT("ForEachElementInEnum_Title", "ForEach {EnumName}"), Args);
+		CachedNodeTitle.SetCachedText(FText::Format(LOCTEXT("ForEachElementInEnum_Title", "ForEach {EnumName}"), Args), this);
 	}
 	return CachedNodeTitle;
 }
diff --git a/Engine/Source/Editor/BlueprintGraph/Private/K2Node_GetDataTableRow.cpp b/Engine/Source/Editor/BlueprintGraph/Private/K2Node_GetDataTableRow.cpp
--- a/Engine/Source/Editor/BlueprintGraph/Private/K2Node_GetDataTableRow.cpp
+++ b/Engine/Source/Editor/BlueprintGraph/Private/K2Node_GetDataTableRow.cpp
@@ -294,14 +294,14 @@ FText UK2Node_GetDataTableRow::GetNodeTitle(ENodeTitleType::Type TitleType) cons
 		{
 			return NSLOCTEXT("K2Node", "DataTable_Title_None", "Get Data Table Row NONE");
 		}
-		else if (CachedNodeTitle.IsOutOfDate())
+		else if (CachedNodeTitle.IsOutOfDate(this))
 		{
 			FFormatNamedArguments Args;
 			Args.Add(TEXT("DataTableName"), FText::FromString(DataTablePin->DefaultObject->GetName()));
 
 			FText LocFormat = NSLOCTEXT("K2Node", "DataTable", "Get Data Table Row {DataTableName}");
 			// FText::Format() is slow, so we cache this to save on performance
-			CachedNodeTitle = FText::Format(LocFormat, Args);
+			CachedNodeTitle.SetCachedText(FText::Format(LocFormat, Args), this);
 		}
 	}
 	else
diff --git a/Engine/Source/Editor/BlueprintGraph/Private/K2Node_GetInputAxisKeyValue.cpp b/Engine/Source/Editor/BlueprintGraph/Private/K2Node_GetInputAxisKeyValue.cpp
--- a/Engine/Source/Editor/BlueprintGraph/Private/K2Node_GetInputAxisKeyValue.cpp
+++ b/Engine/Source/Editor/BlueprintGraph/Private/K2Node_GetInputAxisKeyValue.cpp
@@ -41,12 +41,12 @@ FText UK2Node_GetInputAxisKeyValue::GetNodeTitle(ENodeTitleType::Type TitleType)
 	{
 		return InputAxisKey.GetDisplayName();
 	}
-	else if (CachedNodeTitle.IsOutOfDate())
+	else if (CachedNodeTitle.IsOutOfDate(this))
 	{
 		FFormatNamedArguments Args;
 		Args.Add(TEXT("AxisKey"), InputAxisKey.GetDisplayName());
 		// FText::Format() is slow, so we cache this to save on performance
-		CachedNodeTitle = FText::Format(NSLOCTEXT("K2Node", "GetInputAxisKey_Name", "Get {AxisKey}"), Args);
+		CachedNodeTitle.SetCachedText(FText::Format(NSLOCTEXT("K2Node", "GetInputAxisKey_Name", "Get {AxisKey}"), Args), this);
 	}
 
 	return CachedNodeTitle;
@@ -59,12 +59,12 @@ FString UK2Node_GetInputAxisKeyValue::GetKeywords() const
 
 FText UK2Node_GetInputAxisKeyValue::GetTooltipText() const
 {
-	if (CachedTooltip.IsOutOfDate())
+	if (CachedTooltip.IsOutOfDate(this))
 	{
 		FFormatNamedArguments Args;
 		Args.Add(TEXT("AxisKey"), InputAxisKey.GetDisplayName());
 		// FText::Format() is slow, so we cache this to save on performance
-		CachedTooltip = FText::Format(NSLOCTEXT("K2Node", "GetInputAxisKey_Tooltip", "Returns the current value of input axis key {AxisKey}.  If input is disabled for the actor the value will be 0."), Args);
+		CachedTooltip.SetCachedText(FText::Format(NSLOCTEXT("K2Node", "GetInputAxisKey_Tooltip", "Returns the current value of input axis key {AxisKey}.  If input is disabled for the actor the value will be 0."), Args), this);
 	}
 	return CachedTooltip;
 }
@@ -202,10 +202,10 @@ FText UK2Node_GetInputAxisKeyValue::GetMenuCategory() const
 		CategoryIndex = KeyValueCategory;
 	}
 
-	if (CachedCategories[CategoryIndex].IsOutOfDate())
+	if (CachedCategories[CategoryIndex].IsOutOfDate(this))
 	{
 		// FText::Format() is slow, so we cache this to save on performance
-		CachedCategories[CategoryIndex] = FEditorCategoryUtils::BuildCategoryString(FCommonEditorCategory::Input, SubCategory);
+		CachedCategories[CategoryIndex].SetCachedText(FEditorCategoryUtils::BuildCategoryString(FCommonEditorCategory::Input, SubCategory), this);
 	}
 	return CachedCategories[CategoryIndex];
 }
diff --git a/Engine/Source/Editor/BlueprintGraph/Private/K2Node_GetInputAxisValue.cpp b/Engine/Source/Editor/BlueprintGraph/Private/K2Node_GetInputAxisValue.cpp
--- a/Engine/Source/Editor/BlueprintGraph/Private/K2Node_GetInputAxisValue.cpp
+++ b/Engine/Source/Editor/BlueprintGraph/Private/K2Node_GetInputAxisValue.cpp
@@ -38,14 +38,14 @@ FText UK2Node_GetInputAxisValue::GetNodeTitle(ENodeTitleType::Type TitleType) co
 	{
 		return FText::FromName(InputAxisName);
 	}
-	else if (CachedNodeTitle.IsOutOfDate())
+	else if (CachedNodeTitle.IsOutOfDate(this))
 	{
 		FFormatNamedArguments Args;
 		Args.Add(TEXT("InputAxisName"), FText::FromName(InputAxisName));
 
 		FText LocFormat = NSLOCTEXT("K2Node", "GetInputAxis_Name", "Get {InputAxisName}");
 		// FText::Format() is slow, so we cache this to save on performance
-		CachedNodeTitle = FText::Format(LocFormat, Args);
+		CachedNodeTitle.SetCachedText(FText::Format(LocFormat, Args), this);
 	}
 
 	return CachedNodeTitle;
@@ -58,10 +58,10 @@ FString UK2Node_GetInputAxisValue::GetKeywords() const
 
 FText UK2Node_GetInputAxisValue::GetTooltipText() const
 {
-	if (CachedTooltip.IsOutOfDate())
+	if (CachedTooltip.IsOutOfDate(this))
 	{
 		// FText::Format() is slow, so we cache this to save on performance
-		CachedTooltip = FText::Format(NSLOCTEXT("K2Node", "GetInputAxis_Tooltip", "Returns the current value of input axis {0}.  If input is disabled for the actor the value will be 0."), FText::FromName(InputAxisName));
+		CachedTooltip.SetCachedText(FText::Format(NSLOCTEXT("K2Node", "GetInputAxis_Tooltip", "Returns the current value of input axis {0}.  If input is disabled for the actor the value will be 0."), FText::FromName(InputAxisName)), this);
 	}
 	return CachedTooltip;
 }
@@ -154,10 +154,10 @@ void UK2Node_GetInputAxisValue::GetMenuActions(FBlueprintActionDatabaseRegistrar
 FText UK2Node_GetInputAxisValue::GetMenuCategory() const
 {
 	static FNodeTextCache CachedCategory;
-	if (CachedCategory.IsOutOfDate())
+	if (CachedCategory.IsOutOfDate(this))
 	{
 		// FText::Format() is slow, so we cache this to save on performance
-		CachedCategory = FEditorCategoryUtils::BuildCategoryString(FCommonEditorCategory::Input, LOCTEXT("ActionMenuCategory", "Axis Values"));
+		CachedCategory.SetCachedText(FEditorCategoryUtils::BuildCategoryString(FCommonEditorCategory::Input, LOCTEXT("ActionMenuCategory", "Axis Values")), this);
 	}
 	return CachedCategory;
 }
diff --git a/Engine/Source/Editor/BlueprintGraph/Private/K2Node_GetNumEnumEntries.cpp b/Engine/Source/Editor/BlueprintGraph/Private/K2Node_GetNumEnumEntries.cpp
--- a/Engine/Source/Editor/BlueprintGraph/Private/K2Node_GetNumEnumEntries.cpp
+++ b/Engine/Source/Editor/BlueprintGraph/Private/K2Node_GetNumEnumEntries.cpp
@@ -30,10 +30,10 @@ FText UK2Node_GetNumEnumEntries::GetTooltipText() const
 	{
 		return NSLOCTEXT("K2Node", "GetNumEnumEntries_BadTooltip", "Returns (bad enum)_MAX value");
 	}
-	else if (CachedTooltip.IsOutOfDate())
+	else if (CachedTooltip.IsOutOfDate(this))
 	{
 		// FText::Format() is slow, so we cache this to save on performance
-		CachedTooltip = FText::Format(NSLOCTEXT("K2Node", "GetNumEnumEntries_Tooltip", "Returns {0}_MAX value"), FText::FromName(Enum->GetFName()));
+		CachedTooltip.SetCachedText(FText::Format(NSLOCTEXT("K2Node", "GetNumEnumEntries_Tooltip", "Returns {0}_MAX value"), FText::FromName(Enum->GetFName())), this);
 	}
 	return CachedTooltip;
 }
@@ -44,11 +44,11 @@ FText UK2Node_GetNumEnumEntries::GetNodeTitle(ENodeTitleType::Type TitleType) co
 	{
 		return NSLOCTEXT("K2Node", "GetNumEnumEntries_BadEnumTitle", "Get number of entries in (bad enum)");
 	}
-	else if (CachedNodeTitle.IsOutOfDate())
+	else if (CachedNodeTitle.IsOutOfDate(this))
 	{
 		FFormatNamedArguments Args;
 		Args.Add(TEXT("EnumName"), FText::FromString(Enum->GetName()));
-		CachedNodeTitle = FText::Format(NSLOCTEXT("K2Node", "GetNumEnumEntries_Title", "Get number of entries in {EnumName}"), Args);
+		CachedNodeTitle.SetCachedText(FText::Format(NSLOCTEXT("K2Node", "GetNumEnumEntries_Title", "Get number of entries in {EnumName}"), Args), this);
 	}
 	return CachedNodeTitle;
 }
diff --git a/Engine/Source/Editor/BlueprintGraph/Private/K2Node_InputAction.cpp b/Engine/Source/Editor/BlueprintGraph/Private/K2Node_InputAction.cpp
--- a/Engine/Source/Editor/BlueprintGraph/Private/K2Node_InputAction.cpp
+++ b/Engine/Source/Editor/BlueprintGraph/Private/K2Node_InputAction.cpp
@@ -52,14 +52,14 @@ FText UK2Node_InputAction::GetNodeTitle(ENodeTitleType::Type TitleType) const
 	{
 		return FText::FromName(InputActionName);
 	}
-	else if (CachedNodeTitle.IsOutOfDate())
+	else if (CachedNodeTitle.IsOutOfDate(this))
 	{
 		FFormatNamedArguments Args;
 		Args.Add(TEXT("InputActionName"), FText::FromName(InputActionName));
 
 		FText LocFormat = NSLOCTEXT("K2Node", "InputAction_Name", "InputAction {InputActionName}");
 		// FText::Format() is slow, so we cache this to save on performance
-		CachedNodeTitle = FText::Format(LocFormat, Args);
+		CachedNodeTitle.SetCachedText(FText::Format(LocFormat, Args), this);
 	}
 
 	return CachedNodeTitle;
@@ -67,10 +67,10 @@ FText UK2Node_InputAction::GetNodeTitle(ENodeTitleType::Type TitleType) const
 
 FText UK2Node_InputAction::GetTooltipText() const
 {
-	if (CachedTooltip.IsOutOfDate())
+	if (CachedTooltip.IsOutOfDate(this))
 	{
 		// FText::Format() is slow, so we cache this to save on performance
-		CachedTooltip = FText::Format(NSLOCTEXT("K2Node", "InputAction_Tooltip", "Event for when the keys bound to input action {0} are pressed or released."), FText::FromName(InputActionName));
+		CachedTooltip.SetCachedText(FText::Format(NSLOCTEXT("K2Node", "InputAction_Tooltip", "Event for when the keys bound to input action {0} are pressed or released."), FText::FromName(InputActionName)), this);
 	}
 	return CachedTooltip;
 }
@@ -190,10 +190,10 @@ void UK2Node_InputAction::GetMenuActions(FBlueprintActionDatabaseRegistrar& Acti
 FText UK2Node_InputAction::GetMenuCategory() const
 {
 	static FNodeTextCache CachedCategory;
-	if (CachedCategory.IsOutOfDate())
+	if (CachedCategory.IsOutOfDate(this))
 	{
 		// FText::Format() is slow, so we cache this to save on performance
-		CachedCategory = FEditorCategoryUtils::BuildCategoryString(FCommonEditorCategory::Input, LOCTEXT("ActionMenuCategory", "Action Events"));
+		CachedCategory.SetCachedText(FEditorCategoryUtils::BuildCategoryString(FCommonEditorCategory::Input, LOCTEXT("ActionMenuCategory", "Action Events")), this);
 	}
 	return CachedCategory;
 }
diff --git a/Engine/Source/Editor/BlueprintGraph/Private/K2Node_InputAxisEvent.cpp b/Engine/Source/Editor/BlueprintGraph/Private/K2Node_InputAxisEvent.cpp
--- a/Engine/Source/Editor/BlueprintGraph/Private/K2Node_InputAxisEvent.cpp
+++ b/Engine/Source/Editor/BlueprintGraph/Private/K2Node_InputAxisEvent.cpp
@@ -56,14 +56,14 @@ FText UK2Node_InputAxisEvent::GetNodeTitle(ENodeTitleType::Type TitleType) const
 	{
 		return FText::FromName(InputAxisName);
 	}
-	else if (CachedNodeTitle.IsOutOfDate())
+	else if (CachedNodeTitle.IsOutOfDate(this))
 	{
 		FFormatNamedArguments Args;
 		Args.Add(TEXT("InputAxisName"), FText::FromName(InputAxisName));
 
 		FText LocFormat = NSLOCTEXT("K2Node", "InputAxis_Name", "InputAxis {InputAxisName}");
 		// FText::Format() is slow, so we cache this to save on performance
-		CachedNodeTitle = FText::Format(LocFormat, Args);
+		CachedNodeTitle.SetCachedText(FText::Format(LocFormat, Args), this);
 	}
 
 	return CachedNodeTitle;
@@ -71,10 +71,10 @@ FText UK2Node_InputAxisEvent::GetNodeTitle(ENodeTitleType::Type TitleType) const
 
 FText UK2Node_InputAxisEvent::GetTooltipText() const
 {
-	if (CachedTooltip.IsOutOfDate())
+	if (CachedTooltip.IsOutOfDate(this))
 	{
 		// FText::Format() is slow, so we cache this to save on performance
-		CachedTooltip = FText::Format(NSLOCTEXT("K2Node", "InputAxis_Tooltip", "Event that provides the current value of the {0} axis once per frame when input is enabled for the containing actor."), FText::FromName(InputAxisName));
+		CachedTooltip.SetCachedText(FText::Format(NSLOCTEXT("K2Node", "InputAxis_Tooltip", "Event that provides the current value of the {0} axis once per frame when input is enabled for the containing actor."), FText::FromName(InputAxisName)), this);
 	}
 	return CachedTooltip;
 }
@@ -174,10 +174,10 @@ void UK2Node_InputAxisEvent::GetMenuActions(FBlueprintActionDatabaseRegistrar& A
 FText UK2Node_InputAxisEvent::GetMenuCategory() const
 {
 	static FNodeTextCache CachedCategory;
-	if (CachedCategory.IsOutOfDate())
+	if (CachedCategory.IsOutOfDate(this))
 	{
 		// FText::Format() is slow, so we cache this to save on performance
-		CachedCategory = FEditorCategoryUtils::BuildCategoryString(FCommonEditorCategory::Input, LOCTEXT("ActionMenuCategory", "Axis Events"));
+		CachedCategory.SetCachedText(FEditorCategoryUtils::BuildCategoryString(FCommonEditorCategory::Input, LOCTEXT("ActionMenuCategory", "Axis Events")), this);
 	}
 	return CachedCategory;
 }
diff --git a/Engine/Source/Editor/BlueprintGraph/Private/K2Node_InputAxisKeyEvent.cpp b/Engine/Source/Editor/BlueprintGraph/Private/K2Node_InputAxisKeyEvent.cpp
--- a/Engine/Source/Editor/BlueprintGraph/Private/K2Node_InputAxisKeyEvent.cpp
+++ b/Engine/Source/Editor/BlueprintGraph/Private/K2Node_InputAxisKeyEvent.cpp
@@ -48,10 +48,10 @@ FText UK2Node_InputAxisKeyEvent::GetNodeTitle(ENodeTitleType::Type TitleType) co
 
 FText UK2Node_InputAxisKeyEvent::GetTooltipText() const
 {
-	if (CachedTooltip.IsOutOfDate())
+	if (CachedTooltip.IsOutOfDate(this))
 	{
 		// FText::Format() is slow, so we cache this to save on performance
-		CachedTooltip = FText::Format(NSLOCTEXT("K2Node", "InputAxisKey_Tooltip", "Event that provides the current value of the {0} axis once per frame when input is enabled for the containing actor."), AxisKey.GetDisplayName());
+		CachedTooltip.SetCachedText(FText::Format(NSLOCTEXT("K2Node", "InputAxisKey_Tooltip", "Event that provides the current value of the {0} axis once per frame when input is enabled for the containing actor."), AxisKey.GetDisplayName()), this);
 	}
 	return CachedTooltip;
 }
@@ -199,10 +199,10 @@ FText UK2Node_InputAxisKeyEvent::GetMenuCategory() const
 		CategoryIndex = KeyEventCategory;
 	}
 
-	if (CachedCategories[CategoryIndex].IsOutOfDate())
+	if (CachedCategories[CategoryIndex].IsOutOfDate(this))
 	{
 		// FText::Format() is slow, so we cache this to save on performance
-		CachedCategories[CategoryIndex] = FEditorCategoryUtils::BuildCategoryString(FCommonEditorCategory::Input, SubCategory);
+		CachedCategories[CategoryIndex].SetCachedText(FEditorCategoryUtils::BuildCategoryString(FCommonEditorCategory::Input, SubCategory), this);
 	}
 	return CachedCategories[CategoryIndex];
 }
diff --git a/Engine/Source/Editor/BlueprintGraph/Private/K2Node_InputKey.cpp b/Engine/Source/Editor/BlueprintGraph/Private/K2Node_InputKey.cpp
--- a/Engine/Source/Editor/BlueprintGraph/Private/K2Node_InputKey.cpp
+++ b/Engine/Source/Editor/BlueprintGraph/Private/K2Node_InputKey.cpp
@@ -154,14 +154,14 @@ FText UK2Node_InputKey::GetNodeTitle(ENodeTitleType::Type TitleType) const
 {
 	if (bControl || bAlt || bShift || bCommand)
 	{
-		if (CachedNodeTitle.IsOutOfDate())
+		if (CachedNodeTitle.IsOutOfDate(this))
 		{
 			FFormatNamedArguments Args;
 			Args.Add(TEXT("ModifierKey"), GetModifierText());
 			Args.Add(TEXT("Key"), GetKeyText());
 			
 			// FText::Format() is slow, so we cache this to save on performance
-			CachedNodeTitle = FText::Format(NSLOCTEXT("K2Node", "InputKey_Name_WithModifiers", "{ModifierKey} {Key}"), Args);
+			CachedNodeTitle.SetCachedText(FText::Format(NSLOCTEXT("K2Node", "InputKey_Name_WithModifiers", "{ModifierKey} {Key}"), Args), this);
 		}
 		return CachedNodeTitle;
 	}
@@ -173,7 +173,7 @@ FText UK2Node_InputKey::GetNodeTitle(ENodeTitleType::Type TitleType) const
 
 FText UK2Node_InputKey::GetTooltipText() const
 {
-	if (CachedTooltip.IsOutOfDate())
+	if (CachedTooltip.IsOutOfDate(this))
 	{
 		FText ModifierText = GetModifierText();
 		FText KeyText = GetKeyText();
@@ -181,11 +181,11 @@ FText UK2Node_InputKey::GetTooltipText() const
 		// FText::Format() is slow, so we cache this to save on performance
 		if (!ModifierText.IsEmpty())
 		{
-			CachedTooltip = FText::Format(NSLOCTEXT("K2Node", "InputKey_Tooltip_Modifiers", "Events for when the {0} key is pressed or released while {1} is also held."), KeyText, ModifierText);
+			CachedTooltip.SetCachedText(FText::Format(NSLOCTEXT("K2Node", "InputKey_Tooltip_Modifiers", "Events for when the {0} key is pressed or released while {1} is also held."), KeyText, ModifierText), this);
 		}
 		else
 		{
-			CachedTooltip = FText::Format(NSLOCTEXT("K2Node", "InputKey_Tooltip", "Events for when the {0} key is pressed or released."), KeyText);
+			CachedTooltip.SetCachedText(FText::Format(NSLOCTEXT("K2Node", "InputKey_Tooltip", "Events for when the {0} key is pressed or released."), KeyText), this);
 		}
 	}
 	return CachedTooltip;
@@ -362,10 +362,10 @@ FText UK2Node_InputKey::GetMenuCategory() const
 		CategoryIndex = KeyEventCategory;
 	}
 
-	if (CachedCategories[CategoryIndex].IsOutOfDate())
+	if (CachedCategories[CategoryIndex].IsOutOfDate(this))
 	{
 		// FText::Format() is slow, so we cache this to save on performance
-		CachedCategories[CategoryIndex] = FEditorCategoryUtils::BuildCategoryString(FCommonEditorCategory::Input, SubCategory);
+		CachedCategories[CategoryIndex].SetCachedText(FEditorCategoryUtils::BuildCategoryString(FCommonEditorCategory::Input, SubCategory), this);
 	}
 	return CachedCategories[CategoryIndex];
 }
diff --git a/Engine/Source/Editor/BlueprintGraph/Private/K2Node_MCDelegate.cpp b/Engine/Source/Editor/BlueprintGraph/Private/K2Node_MCDelegate.cpp
--- a/Engine/Source/Editor/BlueprintGraph/Private/K2Node_MCDelegate.cpp
+++ b/Engine/Source/Editor/BlueprintGraph/Private/K2Node_MCDelegate.cpp
@@ -282,12 +282,12 @@ void UK2Node_AddDelegate::AllocateDefaultPins()
 
 FText UK2Node_AddDelegate::GetNodeTitle(ENodeTitleType::Type TitleType) const
 {
-	if (CachedNodeTitle.IsOutOfDate())
+	if (CachedNodeTitle.IsOutOfDate(this))
 	{
 		FFormatNamedArguments Args;
 		Args.Add(TEXT("PropertyName"), FText::FromName(GetPropertyName()));
 		// FText::Format() is slow, so we cache this to save on performance
-		CachedNodeTitle = FText::Format(NSLOCTEXT("K2Node", "AddDelegate", "Bind Event to {PropertyName}"), Args);
+		CachedNodeTitle.SetCachedText(FText::Format(NSLOCTEXT("K2Node", "AddDelegate", "Bind Event to {PropertyName}"), Args), this);
 	}
 	return CachedNodeTitle;
 }
@@ -313,12 +313,12 @@ UK2Node_ClearDelegate::UK2Node_ClearDelegate(const FObjectInitializer& ObjectIni
 
 FText UK2Node_ClearDelegate::GetNodeTitle(ENodeTitleType::Type TitleType) const
 {
-	if (CachedNodeTitle.IsOutOfDate())
+	if (CachedNodeTitle.IsOutOfDate(this))
 	{
 		FFormatNamedArguments Args;
 		Args.Add(TEXT("PropertyName"), FText::FromName(GetPropertyName()));
 		// FText::Format() is slow, so we cache this to save on performance
-		CachedNodeTitle = FText::Format(NSLOCTEXT("K2Node", "ClearDelegate", "Unbind all Events from {PropertyName}"), Args);
+		CachedNodeTitle.SetCachedText(FText::Format(NSLOCTEXT("K2Node", "ClearDelegate", "Unbind all Events from {PropertyName}"), Args), this);
 	}
 	return CachedNodeTitle;
 }
@@ -350,12 +350,12 @@ void UK2Node_RemoveDelegate::AllocateDefaultPins()
 
 FText UK2Node_RemoveDelegate::GetNodeTitle(ENodeTitleType::Type TitleType) const
 {
-	if (CachedNodeTitle.IsOutOfDate())
+	if (CachedNodeTitle.IsOutOfDate(this))
 	{
 		FFormatNamedArguments Args;
 		Args.Add(TEXT("PropertyName"), FText::FromName(GetPropertyName()));
 		// FText::Format() is slow, so we cache this to save on performance
-		CachedNodeTitle = FText::Format(NSLOCTEXT("K2Node", "RemoveDelegate", "Unbind Event from {PropertyName}"), Args);
+		CachedNodeTitle.SetCachedText(FText::Format(NSLOCTEXT("K2Node", "RemoveDelegate", "Unbind Event from {PropertyName}"), Args), this);
 	}
 	return CachedNodeTitle;
 }
@@ -403,12 +403,12 @@ void UK2Node_CallDelegate::AllocateDefaultPins()
 
 FText UK2Node_CallDelegate::GetNodeTitle(ENodeTitleType::Type TitleType) const
 {
-	if (CachedNodeTitle.IsOutOfDate())
+	if (CachedNodeTitle.IsOutOfDate(this))
 	{
 		FFormatNamedArguments Args;
 		Args.Add(TEXT("PropertyName"), FText::FromName(GetPropertyName()));
 		// FText::Format() is slow, so we cache this to save on performance
-		CachedNodeTitle = FText::Format(NSLOCTEXT("K2Node", "CallDelegate", "Call {PropertyName}"), Args);
+		CachedNodeTitle.SetCachedText(FText::Format(NSLOCTEXT("K2Node", "CallDelegate", "Call {PropertyName}"), Args), this);
 	}
 	return CachedNodeTitle;
 }
diff --git a/Engine/Source/Editor/BlueprintGraph/Private/K2Node_MacroInstance.cpp b/Engine/Source/Editor/BlueprintGraph/Private/K2Node_MacroInstance.cpp
--- a/Engine/Source/Editor/BlueprintGraph/Private/K2Node_MacroInstance.cpp
+++ b/Engine/Source/Editor/BlueprintGraph/Private/K2Node_MacroInstance.cpp
@@ -135,10 +135,10 @@ FText UK2Node_MacroInstance::GetTooltipText() const
 	{
 		return NSLOCTEXT("K2Node", "Macro_Tooltip", "Macro");
 	}
-	else if (CachedTooltip.IsOutOfDate())
+	else if (CachedTooltip.IsOutOfDate(this))
 	{
 		// FText::Format() is slow, so we cache this to save on performance
-		CachedTooltip = FText::Format(NSLOCTEXT("K2Node", "MacroGraphInstance_Tooltip", "{0} instance"), FText::FromName(MacroGraph->GetFName()));
+		CachedTooltip.SetCachedText(FText::Format(NSLOCTEXT("K2Node", "MacroGraphInstance_Tooltip", "{0} instance"), FText::FromName(MacroGraph->GetFName())), this);
 	}
 	return CachedTooltip;
 }
diff --git a/Engine/Source/Editor/BlueprintGraph/Private/K2Node_MakeArray.cpp b/Engine/Source/Editor/BlueprintGraph/Private/K2Node_MakeArray.cpp
--- a/Engine/Source/Editor/BlueprintGraph/Private/K2Node_MakeArray.cpp
+++ b/Engine/Source/Editor/BlueprintGraph/Private/K2Node_MakeArray.cpp
@@ -452,10 +452,10 @@ void UK2Node_MakeArray::GetMenuActions(FBlueprintActionDatabaseRegistrar& Action
 FText UK2Node_MakeArray::GetMenuCategory() const
 {
 	static FNodeTextCache CachedCategory;
-	if (CachedCategory.IsOutOfDate())
+	if (CachedCategory.IsOutOfDate(this))
 	{
 		// FText::Format() is slow, so we cache this to save on performance
-		CachedCategory = FEditorCategoryUtils::BuildCategoryString(FCommonEditorCategory::Utilities, LOCTEXT("ActionMenuCategory", "Array"));
+		CachedCategory.SetCachedText(FEditorCategoryUtils::BuildCategoryString(FCommonEditorCategory::Utilities, LOCTEXT("ActionMenuCategory", "Array")), this);
 	}
 	return CachedCategory;
 }
diff --git a/Engine/Source/Editor/BlueprintGraph/Private/K2Node_MakeStruct.cpp b/Engine/Source/Editor/BlueprintGraph/Private/K2Node_MakeStruct.cpp
--- a/Engine/Source/Editor/BlueprintGraph/Private/K2Node_MakeStruct.cpp
+++ b/Engine/Source/Editor/BlueprintGraph/Private/K2Node_MakeStruct.cpp
@@ -159,12 +159,12 @@ FText UK2Node_MakeStruct::GetNodeTitle(ENodeTitleType::Type TitleType) const
 	{
 		return LOCTEXT("MakeNullStructTitle", "Make <unknown struct>");
 	}
-	else if (CachedNodeTitle.IsOutOfDate())
+	else if (CachedNodeTitle.IsOutOfDate(this))
 	{
 		FFormatNamedArguments Args;
 		Args.Add(TEXT("StructName"), FText::FromName(StructType->GetFName()));
 		// FText::Format() is slow, so we cache this to save on performance
-		CachedNodeTitle = FText::Format(LOCTEXT("MakeNodeTitle", "Make {StructName}"), Args);
+		CachedNodeTitle.SetCachedText(FText::Format(LOCTEXT("MakeNodeTitle", "Make {StructName}"), Args), this);
 	}
 	return CachedNodeTitle;
 }
@@ -175,13 +175,13 @@ FText UK2Node_MakeStruct::GetTooltipText() const
 	{
 		return LOCTEXT("MakeNullStruct_Tooltip", "Adds a node that create an '<unknown struct>' from its members");
 	}
-	else if (CachedTooltip.IsOutOfDate())
+	else if (CachedTooltip.IsOutOfDate(this))
 	{
 		// FText::Format() is slow, so we cache this to save on performance
-		CachedTooltip = FText::Format(
+		CachedTooltip.SetCachedText(FText::Format(
 			LOCTEXT("MakeStruct_Tooltip", "Adds a node that create a '{0}' from its members"),
 			FText::FromName(StructType->GetFName())
-		);
+		), this);
 	}
 	return CachedTooltip;
 }
diff --git a/Engine/Source/Editor/BlueprintGraph/Private/K2Node_MathExpression.cpp b/Engine/Source/Editor/BlueprintGraph/Private/K2Node_MathExpression.cpp
--- a/Engine/Source/Editor/BlueprintGraph/Private/K2Node_MathExpression.cpp
+++ b/Engine/Source/Editor/BlueprintGraph/Private/K2Node_MathExpression.cpp
@@ -2481,10 +2481,10 @@ FText UK2Node_MathExpression::GetNodeTitle(ENodeTitleType::Type TitleType) const
 	{
 		return FText::FromString(Expression);
 	}
-	else if (CachedNodeTitle.IsOutOfDate())
+	else if (CachedNodeTitle.IsOutOfDate(this))
 	{
 		// FText::Format() is slow, so we cache this to save on performance
-		CachedNodeTitle = FText::Format(LOCTEXT("MathExpressionSecondTitleLine", "{0}\nMath Expression"), FText::FromString(Expression));
+		CachedNodeTitle.SetCachedText(FText::Format(LOCTEXT("MathExpressionSecondTitleLine", "{0}\nMath Expression"), FText::FromString(Expression)), this);
 	}
 	return CachedNodeTitle;
 }
diff --git a/Engine/Source/Editor/BlueprintGraph/Private/K2Node_Message.cpp b/Engine/Source/Editor/BlueprintGraph/Private/K2Node_Message.cpp
--- a/Engine/Source/Editor/BlueprintGraph/Private/K2Node_Message.cpp
+++ b/Engine/Source/Editor/BlueprintGraph/Private/K2Node_Message.cpp
@@ -19,13 +19,13 @@ FText UK2Node_Message::GetNodeTitle(ENodeTitleType::Type TitleType) const
 	FText NodeName;
 	if (UFunction* Function = GetTargetFunction())
 	{
-		if (!CachedNodeTitles.IsTitleCached(TitleType))
+		if (!CachedNodeTitles.IsTitleCached(TitleType, this))
 		{
 			FText NodeNameText = UK2Node_CallFunction::GetUserFacingFunctionName(Function);
 			if (TitleType == ENodeTitleType::MenuTitle)
 			{
 				// FText::Format() is slow, so we cache this to save on performance
-				CachedNodeTitles.SetCachedTitle(TitleType, FText::Format(LOCTEXT("ListTitle", "{0} (Message)"), NodeNameText));
+				CachedNodeTitles.SetCachedTitle(TitleType, FText::Format(LOCTEXT("ListTitle", "{0} (Message)"), NodeNameText), this);
 			}
 			else
 			{
@@ -35,7 +35,7 @@ FText UK2Node_Message::GetNodeTitle(ENodeTitleType::Type TitleType) const
 
 				FText NodeTitle = FText::Format(NSLOCTEXT("K2Node", "CallInterfaceContext", "{NodeName}\nUsing Interface {OuterClassName}"), Args);
 				// FText::Format() is slow, so we cache this to save on performance
-				CachedNodeTitles.SetCachedTitle(TitleType, NodeTitle);
+				CachedNodeTitles.SetCachedTitle(TitleType, NodeTitle, this);
 			}
 		}
 	}
diff --git a/Engine/Source/Editor/BlueprintGraph/Private/K2Node_SetFieldsInStruct.cpp b/Engine/Source/Editor/BlueprintGraph/Private/K2Node_SetFieldsInStruct.cpp
--- a/Engine/Source/Editor/BlueprintGraph/Private/K2Node_SetFieldsInStruct.cpp
+++ b/Engine/Source/Editor/BlueprintGraph/Private/K2Node_SetFieldsInStruct.cpp
@@ -59,12 +59,12 @@ FText UK2Node_SetFieldsInStruct::GetNodeTitle(ENodeTitleType::Type TitleType) co
 	{
 		return LOCTEXT("SetFieldsInNullStructNodeTitle", "Set members in <unknown struct>");
 	}
-	else if (CachedNodeTitle.IsOutOfDate())
+	else if (CachedNodeTitle.IsOutOfDate(this))
 	{
 		FFormatNamedArguments Args;
 		Args.Add(TEXT("StructName"), FText::FromName(StructType->GetFName()));
 		// FText::Format() is slow, so we cache this to save on performance
-		CachedNodeTitle = FText::Format(LOCTEXT("SetFieldsInStructNodeTitle", "Set members in {StructName}"), Args);
+		CachedNodeTitle.SetCachedText(FText::Format(LOCTEXT("SetFieldsInStructNodeTitle", "Set members in {StructName}"), Args), this);
 	}
 	return CachedNodeTitle;
 }
@@ -75,13 +75,13 @@ FText UK2Node_SetFieldsInStruct::GetTooltipText() const
 	{
 		return LOCTEXT("SetFieldsInStruct_NullTooltip", "Adds a node that modifies an '<unknown struct>'");
 	}
-	else if (CachedTooltip.IsOutOfDate())
+	else if (CachedTooltip.IsOutOfDate(this))
 	{
 		// FText::Format() is slow, so we cache this to save on performance
-		CachedTooltip = FText::Format(
+		CachedTooltip.SetCachedText(FText::Format(
 			LOCTEXT("SetFieldsInStruct_Tooltip", "Adds a node that modifies a '{0}'"),
 			FText::FromName(StructType->GetFName())
-		);
+		), this);
 	}
 	return CachedTooltip;
 }
diff --git a/Engine/Source/Editor/BlueprintGraph/Private/K2Node_SpawnActor.cpp b/Engine/Source/Editor/BlueprintGraph/Private/K2Node_SpawnActor.cpp
--- a/Engine/Source/Editor/BlueprintGraph/Private/K2Node_SpawnActor.cpp
+++ b/Engine/Source/Editor/BlueprintGraph/Private/K2Node_SpawnActor.cpp
@@ -245,12 +245,12 @@ FText UK2Node_SpawnActor::GetNodeTitle(ENodeTitleType::Type TitleType) const
 		// Blueprint will be determined dynamically, so we don't have the name in this case
 		return NSLOCTEXT("K2Node", "SpawnActorUnknown_Title", "SpawnActor");
 	}
-	else if (CachedNodeTitle.IsOutOfDate())
+	else if (CachedNodeTitle.IsOutOfDate(this))
 	{
 		FFormatNamedArguments Args;
 		Args.Add(TEXT("ActorName"), FText::FromString(BlueprintPin->DefaultObject->GetName()));
 		// FText::Format() is slow, so we cache this to save on performance
-		CachedNodeTitle = FText::Format(NSLOCTEXT("K2Node", "SpawnActor", "SpawnActor {ActorName}"), Args);
+		CachedNodeTitle.SetCachedText(FText::Format(NSLOCTEXT("K2Node", "SpawnActor", "SpawnActor {ActorName}"), Args), this);
 	}
 	return CachedNodeTitle;
 }
diff --git a/Engine/Source/Editor/BlueprintGraph/Private/K2Node_SpawnActorFromClass.cpp b/Engine/Source/Editor/BlueprintGraph/Private/K2Node_SpawnActorFromClass.cpp
--- a/Engine/Source/Editor/BlueprintGraph/Private/K2Node_SpawnActorFromClass.cpp
+++ b/Engine/Source/Editor/BlueprintGraph/Private/K2Node_SpawnActorFromClass.cpp
@@ -326,7 +326,7 @@ FText UK2Node_SpawnActorFromClass::GetNodeTitle(ENodeTitleType::Type TitleType)
 			}
 			else
 			{
-				if (CachedNodeTitle.IsOutOfDate())
+				if (CachedNodeTitle.IsOutOfDate(this))
 				{
 					FText ClassName;
 					if (UClass* PickedClass = Cast<UClass>(ClassPin->DefaultObject))
@@ -338,7 +338,7 @@ FText UK2Node_SpawnActorFromClass::GetNodeTitle(ENodeTitleType::Type TitleType)
 					Args.Add(TEXT("ClassName"), ClassName);
 
 					// FText::Format() is slow, so we cache this to save on performance
-					CachedNodeTitle = FText::Format(NSLOCTEXT("K2Node", "SpawnActor", "SpawnActor {ClassName}"), Args);
+					CachedNodeTitle.SetCachedText(FText::Format(NSLOCTEXT("K2Node", "SpawnActor", "SpawnActor {ClassName}"), Args), this);
 				}
 				NodeTitle = CachedNodeTitle;
 			} 
diff --git a/Engine/Source/Editor/BlueprintGraph/Private/K2Node_StructMemberGet.cpp b/Engine/Source/Editor/BlueprintGraph/Private/K2Node_StructMemberGet.cpp
--- a/Engine/Source/Editor/BlueprintGraph/Private/K2Node_StructMemberGet.cpp
+++ b/Engine/Source/Editor/BlueprintGraph/Private/K2Node_StructMemberGet.cpp
@@ -71,24 +71,24 @@ void UK2Node_StructMemberGet::AllocatePinsForSingleMemberGet(FName MemberName)
 
 FText UK2Node_StructMemberGet::GetTooltipText() const
 {
-	if (CachedTooltip.IsOutOfDate())
+	if (CachedTooltip.IsOutOfDate(this))
 	{
 		FFormatNamedArguments Args;
 		Args.Add(TEXT("VariableName"), FText::FromString(GetVarNameString()));
 		// FText::Format() is slow, so we cache this to save on performance
-		CachedTooltip = FText::Format(LOCTEXT("K2Node_StructMemberGet_Tooltip", "Get member variables of {VariableName}"), Args);
+		CachedTooltip.SetCachedText(FText::Format(LOCTEXT("K2Node_StructMemberGet_Tooltip", "Get member variables of {VariableName}"), Args), this);
 	}
 	return CachedTooltip;
 }
 
 FText UK2Node_StructMemberGet::GetNodeTitle(ENodeTitleType::Type TitleType) const
 {
-	if (CachedNodeTitle.IsOutOfDate())
+	if (CachedNodeTitle.IsOutOfDate(this))
 	{
 		FFormatNamedArguments Args;
 		Args.Add(TEXT("VariableName"), FText::FromString(GetVarNameString()));
 		// FText::Format() is slow, so we cache this to save on performance
-		CachedNodeTitle = FText::Format(LOCTEXT("GetMembersInVariable", "Get members in {VariableName}"), Args);
+		CachedNodeTitle.SetCachedText(FText::Format(LOCTEXT("GetMembersInVariable", "Get members in {VariableName}"), Args), this);
 	}
 	return CachedNodeTitle;
 }
diff --git a/Engine/Source/Editor/BlueprintGraph/Private/K2Node_StructMemberSet.cpp b/Engine/Source/Editor/BlueprintGraph/Private/K2Node_StructMemberSet.cpp
--- a/Engine/Source/Editor/BlueprintGraph/Private/K2Node_StructMemberSet.cpp
+++ b/Engine/Source/Editor/BlueprintGraph/Private/K2Node_StructMemberSet.cpp
@@ -44,24 +44,24 @@ void UK2Node_StructMemberSet::AllocateDefaultPins()
 
 FText UK2Node_StructMemberSet::GetTooltipText() const
 {
-	if (CachedTooltip.IsOutOfDate())
+	if (CachedTooltip.IsOutOfDate(this))
 	{
 		FFormatNamedArguments Args;
 		Args.Add(TEXT("VariableName"), FText::FromName(VariableReference.GetMemberName()));
 		// FText::Format() is slow, so we cache this to save on performance
-		CachedTooltip = FText::Format(LOCTEXT("K2Node_StructMemberSet_Tooltip", "Set member variables of {VariableName}"), Args);
+		CachedTooltip.SetCachedText(FText::Format(LOCTEXT("K2Node_StructMemberSet_Tooltip", "Set member variables of {VariableName}"), Args), this);
 	}
 	return CachedTooltip;
 }
 
 FText UK2Node_StructMemberSet::GetNodeTitle(ENodeTitleType::Type TitleType) const
 {
-	if (CachedNodeTitle.IsOutOfDate())
+	if (CachedNodeTitle.IsOutOfDate(this))
 	{
 		FFormatNamedArguments Args;
 		Args.Add(TEXT("VariableName"), FText::FromName(VariableReference.GetMemberName()));
 		// FText::Format() is slow, so we cache this to save on performance
-		CachedNodeTitle = FText::Format(LOCTEXT("SetMembersInVariable", "Set members in {VariableName}"), Args);
+		CachedNodeTitle.SetCachedText(FText::Format(LOCTEXT("SetMembersInVariable", "Set members in {VariableName}"), Args), this);
 	}
 	return CachedNodeTitle;
 }
diff --git a/Engine/Source/Editor/BlueprintGraph/Private/K2Node_Switch.cpp b/Engine/Source/Editor/BlueprintGraph/Private/K2Node_Switch.cpp
--- a/Engine/Source/Editor/BlueprintGraph/Private/K2Node_Switch.cpp
+++ b/Engine/Source/Editor/BlueprintGraph/Private/K2Node_Switch.cpp
@@ -316,10 +316,10 @@ FNodeHandlingFunctor* UK2Node_Switch::CreateNodeHandler(FKismetCompilerContext&
 FText UK2Node_Switch::GetMenuCategory() const
 {
 	static FNodeTextCache CachedCategory;
-	if (CachedCategory.IsOutOfDate())
+	if (CachedCategory.IsOutOfDate(this))
 	{
 		// FText::Format() is slow, so we cache this to save on performance
-		CachedCategory = FEditorCategoryUtils::BuildCategoryString(FCommonEditorCategory::FlowControl, LOCTEXT("ActionMenuCategory", "Switch"));
+		CachedCategory.SetCachedText(FEditorCategoryUtils::BuildCategoryString(FCommonEditorCategory::FlowControl, LOCTEXT("ActionMenuCategory", "Switch")), this);
 	}
 	return CachedCategory;
 }
diff --git a/Engine/Source/Editor/BlueprintGraph/Private/K2Node_SwitchEnum.cpp b/Engine/Source/Editor/BlueprintGraph/Private/K2Node_SwitchEnum.cpp
--- a/Engine/Source/Editor/BlueprintGraph/Private/K2Node_SwitchEnum.cpp
+++ b/Engine/Source/Editor/BlueprintGraph/Private/K2Node_SwitchEnum.cpp
@@ -55,12 +55,12 @@ FText UK2Node_SwitchEnum::GetNodeTitle(ENodeTitleType::Type TitleType) const
 	{
 		return LOCTEXT("SwitchEnum_BadEnumTitle", "Switch on (bad enum)");
 	}
-	else if (CachedNodeTitle.IsOutOfDate())
+	else if (CachedNodeTitle.IsOutOfDate(this))
 	{
 		FFormatNamedArguments Args;
 		Args.Add(TEXT("EnumName"), FText::FromString(Enum->GetName()));
 		// FText::Format() is slow, so we cache this to save on performance
-		CachedNodeTitle = FText::Format(NSLOCTEXT("K2Node", "Switch_Enum", "Switch on {EnumName}"), Args);
+		CachedNodeTitle.SetCachedText(FText::Format(NSLOCTEXT("K2Node", "Switch_Enum", "Switch on {EnumName}"), Args), this);
 	}
 	return CachedNodeTitle;
 }
diff --git a/Engine/Source/Editor/BlueprintGraph/Private/K2Node_TemporaryVariable.cpp b/Engine/Source/Editor/BlueprintGraph/Private/K2Node_TemporaryVariable.cpp
--- a/Engine/Source/Editor/BlueprintGraph/Private/K2Node_TemporaryVariable.cpp
+++ b/Engine/Source/Editor/BlueprintGraph/Private/K2Node_TemporaryVariable.cpp
@@ -49,26 +49,26 @@ void UK2Node_TemporaryVariable::AllocateDefaultPins()
 
 FText UK2Node_TemporaryVariable::GetTooltipText() const
 {
-	if (CachedTooltip.IsOutOfDate())
+	if (CachedTooltip.IsOutOfDate(this))
 	{
 		FFormatNamedArguments Args;
 		Args.Add(TEXT("VariableType"), UEdGraphSchema_K2::TypeToText(VariableType));
 		// FText::Format() is slow, so we cache this to save on performance
-		CachedTooltip = FText::Format(NSLOCTEXT("K2Node", "LocalTemporaryVariable", "Local temporary {VariableType} variable"), Args);
+		CachedTooltip.SetCachedText(FText::Format(NSLOCTEXT("K2Node", "LocalTemporaryVariable", "Local temporary {VariableType} variable"), Args), this);
 	}
 	return CachedTooltip;
 }
 
 FText UK2Node_TemporaryVariable::GetNodeTitle(ENodeTitleType::Type TitleType) const
 {
-	if (CachedNodeTitle.IsOutOfDate())
+	if (CachedNodeTitle.IsOutOfDate(this))
 	{
 		FFormatNamedArguments Args;
 		Args.Add(TEXT("VariableType"), UEdGraphSchema_K2::TypeToText(VariableType));
 
 		FText TitleFormat = !bIsPersistent ? NSLOCTEXT("K2Node", "LocalVariable", "Local {VariableType}") : NSLOCTEXT("K2Node", "PersistentLocalVariable", "Persistent Local {VariableType}");
 		// FText::Format() is slow, so we cache this to save on performance
-		CachedNodeTitle = FText::Format(TitleFormat, Args);
+		CachedNodeTitle.SetCachedText(FText::Format(TitleFormat, Args), this);
 	}
 	
 	return CachedNodeTitle;
diff --git a/Engine/Source/Editor/BlueprintGraph/Private/K2Node_VariableGet.cpp b/Engine/Source/Editor/BlueprintGraph/Private/K2Node_VariableGet.cpp
--- a/Engine/Source/Editor/BlueprintGraph/Private/K2Node_VariableGet.cpp
+++ b/Engine/Source/Editor/BlueprintGraph/Private/K2Node_VariableGet.cpp
@@ -171,22 +171,19 @@ FText UK2Node_VariableGet::GetBlueprintVarTooltip(FBPVariableDescription const&
 
 FText UK2Node_VariableGet::GetTooltipText() const
 {
-	// @TODO: The variable name mutates as the user makes changes to the 
-	//        underlying property, so until we can catch all those cases, we're
-	//        going to leave this optimization off
-	//if (CachedTooltip.IsOutOfDate())
+	if (CachedTooltip.IsOutOfDate(this))
 	{
 		if (UProperty* Property = GetPropertyForVariable())
 		{
-			CachedTooltip = GetPropertyTooltip(Property);
+			CachedTooltip.SetCachedText(GetPropertyTooltip(Property), this);
 		}
 		else if (FBPVariableDescription const* VarDesc = GetBlueprintVarDescription())
 		{
-			CachedTooltip = GetBlueprintVarTooltip(*VarDesc);
+			CachedTooltip.SetCachedText(GetBlueprintVarTooltip(*VarDesc), this);
 		}
 		else
 		{
-			CachedTooltip = K2Node_VariableGetImpl::GetBaseTooltip(GetVarName());
+			CachedTooltip.SetCachedText(K2Node_VariableGetImpl::GetBaseTooltip(GetVarName()), this);
 		}
 	}
 	return CachedTooltip;
@@ -212,15 +209,12 @@ FText UK2Node_VariableGet::GetNodeTitle(ENodeTitleType::Type TitleType) const
 	{
 		return LOCTEXT("Get", "Get");
 	}
-	// @TODO: The variable name mutates as the user makes changes to the 
-	//        underlying property, so until we can catch all those cases, we're
-	//        going to leave this optimization off
-	else //if (CachedNodeTitle.IsOutOfDate())
+	else if (CachedNodeTitle.IsOutOfDate(this))
 	{
 		FFormatNamedArguments Args;
 		Args.Add(TEXT("PinName"), FText::FromString(OutputPinName));
 		// FText::Format() is slow, so we cache this to save on performance
-		CachedNodeTitle = FText::Format(LOCTEXT("GetPinName", "Get {PinName}"), Args);
+		CachedNodeTitle.SetCachedText(FText::Format(LOCTEXT("GetPinName", "Get {PinName}"), Args), this);
 	}
 	return CachedNodeTitle;
 }
diff --git a/Engine/Source/Editor/BlueprintGraph/Private/K2Node_VariableSet.cpp b/Engine/Source/Editor/BlueprintGraph/Private/K2Node_VariableSet.cpp
--- a/Engine/Source/Editor/BlueprintGraph/Private/K2Node_VariableSet.cpp
+++ b/Engine/Source/Editor/BlueprintGraph/Private/K2Node_VariableSet.cpp
@@ -217,61 +217,21 @@ FText UK2Node_VariableSet::GetBlueprintVarTooltip(FBPVariableDescription const&
 
 FText UK2Node_VariableSet::GetTooltipText() const
 {
-	// @TODO: The variable name mutates as the user makes changes to the 
-	//        underlying property, so until we can catch all those cases, we're
-	//        going to leave this optimization off
-	//if (CachedTooltip.IsOutOfDate())
+	if (CachedTooltip.IsOutOfDate(this))
 	{
 		if (UProperty* Property = GetPropertyForVariable())
 		{
-			CachedTooltip = GetPropertyTooltip(Property);
+			CachedTooltip.SetCachedText(GetPropertyTooltip(Property), this);
 		}
 		else if (FBPVariableDescription const* VarDesc = GetBlueprintVarDescription())
 		{
-			CachedTooltip = GetBlueprintVarTooltip(*VarDesc);
+			CachedTooltip.SetCachedText(GetBlueprintVarTooltip(*VarDesc), this);
 		}
 		else
 		{
-			CachedTooltip = K2Node_VariableSetImpl::GetBaseTooltip(GetVarName());
-		}
-	}
-	return CachedTooltip;
-
-	
-
-	FFormatNamedArguments Args;
-	Args.Add( TEXT( "VarName" ), FText::FromString( GetVarNameString() ));
-	Args.Add( TEXT( "ReplicationCall" ), FText::GetEmpty());
-	Args.Add( TEXT( "ReplicationNotifyName" ), FText::GetEmpty());
-	Args.Add( TEXT( "TextPartition" ), FText::GetEmpty());
-	Args.Add( TEXT( "MetaData" ), FText::GetEmpty());
-		
-	if(	HasLocalRepNotify() )
-	{
-		Args.Add( TEXT( "ReplicationCall" ), NSLOCTEXT( "K2Node", "VariableSet_ReplicationCall", " and call " ));
-		Args.Add( TEXT( "ReplicationNotifyName" ), FText::FromString( GetRepNotifyName().ToString() ));
-	}
-	if(  UProperty* Property = GetPropertyForVariable() )
-	{
-		// discover if the variable property is a non blueprint user variable
-		UClass* SourceClass = Property->GetOwnerClass();
-		if( SourceClass && SourceClass->ClassGeneratedBy == NULL )
-		{
-			const FString MetaData = Property->GetToolTipText().ToString();
-
-			if( !MetaData.IsEmpty() )
-			{
-				// See if the property associated with this editor has a tooltip
-				FText PropertyMetaData = FText::FromString( *MetaData );
-				FString TooltipName = FString::Printf( TEXT("%s.tooltip"), *(Property->GetName()));
-				FText::FindText( *(Property->GetFullGroupName(true)), *TooltipName, PropertyMetaData );
-				Args.Add( TEXT( "TextPartition" ), FText::FromString( "\n" ));
-				Args.Add( TEXT( "MetaData" ), PropertyMetaData );
-			}
+			CachedTooltip.SetCachedText(K2Node_VariableSetImpl::GetBaseTooltip(GetVarName()), this);
 		}
 	}
-	// FText::Format() is slow, so we cache this to save on performance
-	CachedTooltip = FText::Format(NSLOCTEXT("K2Node", "SetValueOfVariable", "Set the value of variable {VarName}{ReplicationCall}{ReplicationNotifyName}{TextPartition}{MetaData}"), Args);
 	return CachedTooltip;
 }
 
@@ -300,7 +260,7 @@ FText UK2Node_VariableSet::GetNodeTitle(ENodeTitleType::Type TitleType) const
 	// @TODO: The variable name mutates as the user makes changes to the 
 	//        underlying property, so until we can catch all those cases, we're
 	//        going to leave this optimization off
-	else //if (CachedNodeTitle.IsOutOfDate())
+	else if (CachedNodeTitle.IsOutOfDate(this))
 	{
 		FFormatNamedArguments Args;
 		Args.Add(TEXT("PinName"), FText::FromString(InputPinName));
@@ -308,11 +268,11 @@ FText UK2Node_VariableSet::GetNodeTitle(ENodeTitleType::Type TitleType) const
 		// FText::Format() is slow, so we cache this to save on performance
 		if (HasLocalRepNotify())
 		{
-			CachedNodeTitle = FText::Format(NSLOCTEXT("K2Node", "SetWithNotifyPinName", "Set with Notify {PinName}"), Args);
+			CachedNodeTitle.SetCachedText(FText::Format(NSLOCTEXT("K2Node", "SetWithNotifyPinName", "Set with Notify {PinName}"), Args), this);
 		}
 		else
 		{
-			CachedNodeTitle = FText::Format(NSLOCTEXT("K2Node", "SetPinName", "Set {PinName}"), Args);
+			CachedNodeTitle.SetCachedText(FText::Format(NSLOCTEXT("K2Node", "SetPinName", "Set {PinName}"), Args), this);
 		}
 	}
 	return CachedNodeTitle;
diff --git a/Engine/Source/Editor/BlueprintGraph/Private/K2Node_VariableSetRef.cpp b/Engine/Source/Editor/BlueprintGraph/Private/K2Node_VariableSetRef.cpp
--- a/Engine/Source/Editor/BlueprintGraph/Private/K2Node_VariableSetRef.cpp
+++ b/Engine/Source/Editor/BlueprintGraph/Private/K2Node_VariableSetRef.cpp
@@ -137,12 +137,12 @@ FText UK2Node_VariableSetRef::GetNodeTitle(ENodeTitleType::Type TitleType) const
 	{
 		return NSLOCTEXT("K2Node", "SetRefVarNodeTitle", "Set By-Ref Var");
 	}
-	else if (CachedNodeTitle.IsOutOfDate())
+	else if (CachedNodeTitle.IsOutOfDate(this))
 	{
 		FFormatNamedArguments Args;
 		Args.Add(TEXT("PinType"), Schema->TypeToText(TargetPin->PinType));
 		// FText::Format() is slow, so we cache this to save on performance
-		CachedNodeTitle = FText::Format(NSLOCTEXT("K2Node", "SetRefVarNodeTitle_Typed", "Set {PinType}"), Args);
+		CachedNodeTitle.SetCachedText(FText::Format(NSLOCTEXT("K2Node", "SetRefVarNodeTitle_Typed", "Set {PinType}"), Args), this);
 	}
 	return CachedNodeTitle;
 }
diff --git a/Engine/Source/Editor/UnrealEd/Private/EdGraphNode_Comment.cpp b/Engine/Source/Editor/UnrealEd/Private/EdGraphNode_Comment.cpp
--- a/Engine/Source/Editor/UnrealEd/Private/EdGraphNode_Comment.cpp
+++ b/Engine/Source/Editor/UnrealEd/Private/EdGraphNode_Comment.cpp
@@ -54,9 +54,9 @@ void UEdGraphNode_Comment::PostPlacedNewNode()
 
 FText UEdGraphNode_Comment::GetTooltipText() const
 {
-	if (CachedTooltip.IsOutOfDate())
+	if (CachedTooltip.IsOutOfDate(this))
 	{
-		CachedTooltip = FText::Format(NSLOCTEXT("K2Node", "CommentBlock_Tooltip", "Comment:\n{0}"), FText::FromString(NodeComment));
+		CachedTooltip.SetCachedText(FText::Format(NSLOCTEXT("K2Node", "CommentBlock_Tooltip", "Comment:\n{0}"), FText::FromString(NodeComment)), this);
 	}
 	return CachedTooltip;
 }
diff --git a/Engine/Source/Editor/UnrealEd/Private/Kismet2/BlueprintEditorUtils.cpp b/Engine/Source/Editor/UnrealEd/Private/Kismet2/BlueprintEditorUtils.cpp
--- a/Engine/Source/Editor/UnrealEd/Private/Kismet2/BlueprintEditorUtils.cpp
+++ b/Engine/Source/Editor/UnrealEd/Private/Kismet2/BlueprintEditorUtils.cpp
@@ -1892,6 +1892,9 @@ void FBlueprintEditorUtils::UpdateDelegatesInBlueprint(UBlueprint* Blueprint)
 // Blueprint has materially changed.  Recompile the skeleton, notify observers, and mark the package as dirty.
 void FBlueprintEditorUtils::MarkBlueprintAsStructurallyModified(UBlueprint* Blueprint)
 {
+	// The Blueprint has been structurally modified and this means that some node titles will need to be refreshed
+	GetDefault<UEdGraphSchema_K2>()->ForceVisualizationCacheClear();
+
 	Blueprint->bCachedDependenciesUpToDate = false;
 	if (Blueprint->Status != BS_BeingCreated && !Blueprint->bBeingCompiled)
 	{
diff --git a/Engine/Source/Runtime/Engine/Classes/EdGraph/EdGraphSchema.h b/Engine/Source/Runtime/Engine/Classes/EdGraph/EdGraphSchema.h
--- a/Engine/Source/Runtime/Engine/Classes/EdGraph/EdGraphSchema.h
+++ b/Engine/Source/Runtime/Engine/Classes/EdGraph/EdGraphSchema.h
@@ -761,4 +761,18 @@ class ENGINE_API UEdGraphSchema : public UObject
 	 * @return								Returns TRUE if there is support for dropping the pin on the node
 	 */
 	virtual bool SupportsDropPinOnNode(UEdGraphNode* InTargetNode, const FEdGraphPinType& InSourcePinType, EEdGraphPinDirection InSourcePinDirection, FText& OutErrorMessage) const { return false; }
+
+	/**
+	 * Checks if a CacheRefreshID is out of date
+	 *
+	 * @param InVisualizationCacheID	The current refresh ID to check if out of date
+	 * @return							TRUE if dirty
+	 */
+	virtual bool IsCacheVisualizationOutOfDate(int32 InVisualizationCacheID) const { return false; }
+
+	/** Returns the current cache title refresh ID that is appropriate for the passed node */
+	virtual int32 GetCurrentVisualizationCacheID() const { return 0; }
+
+	/** Forces cached visualization data to refresh */
+	virtual void ForceVisualizationCacheClear() const {};
 };
diff --git a/Engine/Source/Runtime/Engine/Public/EdGraph/EdGraphNodeUtils.h b/Engine/Source/Runtime/Engine/Public/EdGraph/EdGraphNodeUtils.h
--- a/Engine/Source/Runtime/Engine/Public/EdGraph/EdGraphNodeUtils.h
+++ b/Engine/Source/Runtime/Engine/Public/EdGraph/EdGraphNodeUtils.h
@@ -3,6 +3,7 @@
 #pragma once
 
 #include "EdGraph/EdGraphNode.h" // for ENodeTitleType
+#include "EdGraph/EdGraphSchema.h"
 
 /*******************************************************************************
  * FNodeTextCache
@@ -15,16 +16,32 @@
 struct FNodeTextCache
 {
 public:
-	/** */
-	FORCEINLINE bool IsOutOfDate() const
+	FNodeTextCache() :
+		CacheRefreshID(0)
+	{}
+
+	/**
+	 * Checks if the title is out of date
+	 *
+	 * @param InOwningNode		Node that owns this title
+	 */
+	FORCEINLINE bool IsOutOfDate(const UEdGraphNode* InOwningNode) const
 	{
-		return CachedText.IsEmpty();
+		check(InOwningNode);
+		return CachedText.IsEmpty() || (InOwningNode->GetSchema()->IsCacheVisualizationOutOfDate(CacheRefreshID));
 	}
 
-	/** */
-	FORCEINLINE void SetCachedText(FText const& Text) const
+	/**
+	 * Checks if the title is out of date
+	 *
+	 * @param InText			Text to cache
+	 * @param InOwningNode		Node that owns this title
+	 */
+	FORCEINLINE void SetCachedText(FText const& InText, const UEdGraphNode* InOwningNode) const
 	{
-		CachedText = Text;
+		check(InOwningNode);
+		CachedText = InText;
+		CacheRefreshID = InOwningNode->GetSchema()->GetCurrentVisualizationCacheID();
 	}
 
 	/** */
@@ -37,6 +54,7 @@ public:
 	FORCEINLINE void MarkDirty() const
 	{
 		Clear();
+		CacheRefreshID = 0;
 	}
 
 	/** */
@@ -45,13 +63,6 @@ public:
 		CachedText = FText::GetEmpty();
 	}
 
-	/** */
-	FORCEINLINE FText& operator=(FText const& Text) const
-	{
-		SetCachedText(Text);
-		return CachedText;
-	}
-
 	/** */
 	FORCEINLINE operator FText&() const
 	{
@@ -61,6 +72,9 @@ public:
 private:
 	/** Mutable so that SetCachedText() can remain const (callable by the node's GetNodeTitle() method)*/
 	mutable FText CachedText;
+
+	/** ID to check if the title should be considered dirty due to outside conditions that may require the title to refresh */
+	mutable int32 CacheRefreshID;
 };
 
 /*******************************************************************************
@@ -70,16 +84,27 @@ private:
 struct FNodeTitleTextTable
 {
 public:
-	/** */
-	FORCEINLINE bool IsTitleCached(ENodeTitleType::Type TitleType) const
+	/**
+	 * Checks if the title of the passed type is cached
+	 *
+	 * @param InTitleType		Type of title to be examined
+	 * @param InOwningNode		Node that owns this title
+	 */
+	FORCEINLINE bool IsTitleCached(ENodeTitleType::Type InTitleType, const UEdGraphNode* InOwningNode) const
 	{
-		return !CachedNodeTitles[TitleType].IsOutOfDate();
+		return !CachedNodeTitles[InTitleType].IsOutOfDate(InOwningNode);
 	}
 
-	/** */
-	FORCEINLINE void SetCachedTitle(ENodeTitleType::Type TitleType, FText const& Text) const
+	/**
+	 * Caches the specific title type
+	 *
+	 * @param InTitleType		Type of title to be examined
+	 * @param InText			Text to cache
+	 * @param InOwningNode		Node that owns this title
+	 */
+	FORCEINLINE void SetCachedTitle(ENodeTitleType::Type InTitleType, FText const& InText, const UEdGraphNode* InOwningNode) const
 	{
-		CachedNodeTitles[TitleType].SetCachedText(Text);
+		CachedNodeTitles[InTitleType].SetCachedText(InText, InOwningNode);
 	}
 
 	/** */
@@ -115,16 +140,25 @@ private:
 struct FNodeTextTable : FNodeTitleTextTable
 {
 public:
-	/** */
-	FORCEINLINE bool IsTooltipCached() const
+	/**
+	 * Checks if there is text cached
+	 *
+	 * @param InOwningNode		Node that owns this title
+	 */
+	FORCEINLINE bool IsTooltipCached(const UEdGraphNode* InOwningNode) const
 	{
-		return !CachedTooltip.IsOutOfDate();
+		return !CachedTooltip.IsOutOfDate(InOwningNode);
 	}
 
-	/** */
-	FORCEINLINE void SetCachedTooltip(FText const& Text) const
+	/**
+	 * Caches the tooltip
+	 *
+	 * @param InText			Text to cache
+	 * @param InOwningNode		Node that owns this title
+	 */
+	FORCEINLINE void SetCachedTooltip(FText const& InText, const UEdGraphNode* InOwningNode) const
 	{
-		CachedTooltip.SetCachedText(Text);
+		CachedTooltip.SetCachedText(InText, InOwningNode);
 	}
 
 	/** */
