Math expression nodes now support variable renames without breaking and local variables can be used with them.

Behind the scenes, variables are converted to their Guid and are converted back to their name for display. The token parser was updated to handle these Guid's which are inside brackets { }.

Added some helper functions for looking up a Member Variable's Guid by name and a Member Variable's name by Guid.

diff --git a/Engine/Source/Editor/BlueprintGraph/Classes/K2Node_MathExpression.h b/Engine/Source/Editor/BlueprintGraph/Classes/K2Node_MathExpression.h
--- a/Engine/Source/Editor/BlueprintGraph/Classes/K2Node_MathExpression.h
+++ b/Engine/Source/Editor/BlueprintGraph/Classes/K2Node_MathExpression.h
@@ -63,12 +63,20 @@ private:
 	*/
 	void ClearExpression();
 
+	/** Sanitizes an expression for display, removing outermost parentheses */
+	FString SanitizeDisplayExpression(FString InExpression) const;
+
+	/** Helper function to build the node's full title */
+	FText GetFullTitle(FText InExpression) const;
 private:
 	/** Cached so we don't have to regenerate it when the graph is recompiled */
 	TSharedPtr<class FCompilerResultsLog> CachedMessageLog;
 
 	/** Constructing FText strings can be costly, so we cache the node's title */
 	FNodeTextCache CachedNodeTitle;
+
+	/** Constructing the display string for a Math Expression is costly, so we cache it */
+	FNodeTextCache CachedDisplayExpression;
 };
 
 
diff --git a/Engine/Source/Editor/BlueprintGraph/Private/BasicTokenParser.cpp b/Engine/Source/Editor/BlueprintGraph/Private/BasicTokenParser.cpp
--- a/Engine/Source/Editor/BlueprintGraph/Private/BasicTokenParser.cpp
+++ b/Engine/Source/Editor/BlueprintGraph/Private/BasicTokenParser.cpp
@@ -109,6 +109,12 @@ void FBasicToken::SetConstString(TCHAR* InString, int32 MaxLength)
 	TokenType = TOKEN_Const;
 }
 
+//------------------------------------------------------------------------------
+void FBasicToken::SetGuid(TCHAR* InString, int32 MaxLength)
+{
+	SetConstString(InString, MaxLength);
+	TokenType = TOKEN_Guid;
+}
 //------------------------------------------------------------------------------
 FString FBasicToken::GetConstantValue() const
 {
@@ -243,7 +249,32 @@ bool FBasicTokenParser::GetToken(FBasicToken& Token, bool bNoConsts/* = false*/)
 	}
 	Token.StartPos		= PrevPos;
 	Token.StartLine		= PrevLine;
-	if( (c>='A' && c<='Z') || (c>='a' && c<='z') || (c=='_') )
+	if( (c=='{') )
+	{
+		// Alphanumeric token.
+		int32 Length=0;
+		Token.Identifier[Length++] = c;
+		do
+		{
+			if( Length >= NAME_SIZE )
+			{
+				Length = ((int32)NAME_SIZE) - 1;
+				Token.Identifier[Length]=0; // need this for the error description
+
+				FText ErrorDesc = FText::Format(LOCTEXT("IdTooLong", "Identifer ({0}...) exceeds maximum length of {1}"), FText::FromString(Token.Identifier), FText::AsNumber((int32)NAME_SIZE));
+				SetError(FErrorState::ParseError, ErrorDesc);
+
+				break;
+			}
+			c = GetChar();
+			Token.Identifier[Length++] = c;
+		} while( c!='}' );
+
+		Token.Identifier[Length]=0;
+		Token.SetGuid(Token.Identifier);
+		return IsValid();
+	}
+	else if( (c>='A' && c<='Z') || (c>='a' && c<='z') || (c=='_') )
 	{
 		// Alphanumeric token.
 		int32 Length=0;
diff --git a/Engine/Source/Editor/BlueprintGraph/Private/BasicTokenParser.h b/Engine/Source/Editor/BlueprintGraph/Private/BasicTokenParser.h
--- a/Engine/Source/Editor/BlueprintGraph/Private/BasicTokenParser.h
+++ b/Engine/Source/Editor/BlueprintGraph/Private/BasicTokenParser.h
@@ -55,6 +55,7 @@ public:
 	void SetConstFloat(float InFloat);
 	void SetConstName(FName InName);
 	void SetConstString(TCHAR* InString, int32 MaxLength=MAX_STRING_CONST_SIZE);
+	void SetGuid(TCHAR* InString, int32 MaxLength=MAX_STRING_CONST_SIZE);
 
 	//--------------------------------------
 	// Constant value getters
@@ -85,6 +86,7 @@ public:
 		TOKEN_Identifier		= 0x01,	// Alphanumeric identifier.
 		TOKEN_Symbol			= 0x02,	// Symbol.
 		TOKEN_Const				= 0x03,	// A constant.
+		TOKEN_Guid				= 0x04, // A variable guid
 		TOKEN_Max				= 0x0D
 	};
 	/** Type of this token. */
diff --git a/Engine/Source/Editor/BlueprintGraph/Private/K2Node_MathExpression.cpp b/Engine/Source/Editor/BlueprintGraph/Private/K2Node_MathExpression.cpp
--- a/Engine/Source/Editor/BlueprintGraph/Private/K2Node_MathExpression.cpp
+++ b/Engine/Source/Editor/BlueprintGraph/Private/K2Node_MathExpression.cpp
@@ -194,6 +194,12 @@ public:
 	 */
 	virtual FString ToString() const = 0;
 
+	/**
+	 * Variable Guid's are stored in the internal expression and must be 
+	 * converted back to their name when showing the expression in the node's title
+	 */
+	virtual FString ToDisplayString(UBlueprint* InBlueprint) const { return ToString(); };
+
 protected:
 	IFExpressionNode() {}
 };
@@ -226,7 +232,7 @@ public:
 	/** For debug purposes, constructs a textual representation of this expression */
 	virtual FString ToString() const override
 	{
-		if (Token.TokenType == FBasicToken::TOKEN_Identifier)
+		if (Token.TokenType == FBasicToken::TOKEN_Identifier || FBasicToken::TOKEN_Guid)
 		{
 			return FString::Printf(TEXT("%s"), Token.Identifier);
 		}
@@ -240,6 +246,24 @@ public:
 		}
 	}
 
+	virtual FString ToDisplayString(UBlueprint* InBlueprint) const
+	{
+		if (Token.TokenType == FBasicToken::TOKEN_Guid)
+		{
+			FGuid VariableGuid;
+			if(FGuid::Parse(FString(Token.Identifier), VariableGuid))
+			{
+				FName VariableName = FBlueprintEditorUtils::FindMemberVariableNameByGuid(InBlueprint, VariableGuid);
+
+				if (VariableName.IsNone())
+				{
+					VariableName = FBlueprintEditorUtils::FindLocalVariableNameByGuid(InBlueprint, VariableGuid);
+				}
+				return VariableName.ToString();
+			}
+		}
+		return ToString();
+	}
 public:
 	/** The base token which this leaf node represents */
 	FBasicToken Token;
@@ -290,6 +314,13 @@ public:
 		return FString::Printf(TEXT("(%s %s %s)"), *LeftStr, *Operator, *RightStr);
 	}
 
+	virtual FString ToDisplayString(UBlueprint* InBlueprint) const
+	{
+		const FString LeftStr = LHS->ToDisplayString(InBlueprint);
+		const FString RightStr = RHS->ToDisplayString(InBlueprint);
+		return FString::Printf(TEXT("(%s %s %s)"), *LeftStr, *Operator, *RightStr);
+	}
+
 public:
 	FString Operator;
 	TSharedRef<IFExpressionNode> LHS;
@@ -339,6 +370,11 @@ public:
 		return FString::Printf(TEXT("(%s%s)"), *Operator, *RightStr);
 	}
 
+	virtual FString ToDisplayString(UBlueprint* InBlueprint) const
+	{
+		const FString RightStr = RHS->ToDisplayString(InBlueprint);
+		return FString::Printf(TEXT("(%s%s)"), *Operator, *RightStr);
+	}
 public:
 	FString Operator;
 	TSharedRef<IFExpressionNode> RHS;
@@ -392,6 +428,13 @@ public:
 		return FString::Printf(TEXT("(%s ? %s : %s)"), *ConditionStr, *TrueStr, *FalseStr);
 	}
 
+	virtual FString ToDisplayString(UBlueprint* InBlueprint) const
+	{
+		const FString ConditionStr = Condition->ToDisplayString(InBlueprint);
+		const FString TrueStr = TruePart->ToDisplayString(InBlueprint);
+		const FString FalseStr = FalsePart->ToDisplayString(InBlueprint);
+		return FString::Printf(TEXT("(%s ? %s : %s)"), *ConditionStr, *TrueStr, *FalseStr);
+	}
 public:
 	TSharedRef<IFExpressionNode> Condition;
 	TSharedRef<IFExpressionNode> TruePart;
@@ -449,6 +492,24 @@ public:
 		return AsString;
 	}
 
+	virtual FString ToDisplayString(UBlueprint* InBlueprint) const
+	{
+		FString AsString("(");
+		for (TSharedRef<IFExpressionNode> Child : Children)
+		{
+			AsString += Child->ToDisplayString(InBlueprint);
+			if (Child == Children.Last())
+			{
+				AsString += ")";
+			}
+			else 
+			{
+				AsString += ", ";
+			}
+		}
+		return AsString;
+	}
+
 	virtual ~FExpressionList() {}
 
 public:
@@ -494,6 +555,11 @@ public:
 		return FString::Printf(TEXT("(%s%s)"), *FuncName, *ParamsString);
 	}
 
+	virtual FString ToDisplayString(UBlueprint* InBlueprint) const
+	{
+		FString const ParamsString = ParamList->ToDisplayString(InBlueprint);
+		return FString::Printf(TEXT("(%s%s)"), *FuncName, *ParamsString);
+	}
 public:
 	FString FuncName;
 	TSharedRef<FExpressionList> ParamList;
@@ -1255,20 +1321,70 @@ public:
 		check(ActiveMessageLog != nullptr);
 		const UEdGraphSchema_K2* K2Schema = GetDefault<UEdGraphSchema_K2>();
 
-		if (ExpressionNode.Token.TokenType == FBasicToken::TOKEN_Identifier)
+		if (ExpressionNode.Token.TokenType == FBasicToken::TOKEN_Identifier || ExpressionNode.Token.TokenType == FBasicToken::TOKEN_Guid)
 		{
 			FString const VariableIdentifier = ExpressionNode.Token.Identifier;
 			// first we try to match up variables with existing variable properties on the blueprint
-			if (UProperty* VariableProperty = FindField<UProperty>(TargetBlueprint->SkeletonGeneratedClass, *VariableIdentifier))
+
+			FMemberReference VariableReference;
+			FString VariableName;
+			FGuid VariableGuid;
+
+			if (ExpressionNode.Token.TokenType == FBasicToken::TOKEN_Guid && FGuid::Parse(VariableIdentifier, VariableGuid))
 			{
-				TSharedPtr<FCodeGenFragment_VariableGet> VariableGetFragment = GeneratePropertyFragment(ExpressionNode, VariableProperty, *ActiveMessageLog);
+				// First look the variable up as a Member variable
+				FName VariableFName;
+				VariableFName = FBlueprintEditorUtils::FindMemberVariableNameByGuid(TargetBlueprint, VariableGuid);
+
+				// If the variable was not found, look it up as a local variable
+				if (VariableFName.IsNone())
+				{
+					VariableFName = FBlueprintEditorUtils::FindLocalVariableNameByGuid(TargetBlueprint, VariableGuid);
+					VariableReference.SetLocalMember(VariableFName, CompilingNode->GetGraph()->GetName(), VariableGuid);
+				}
+				else
+				{
+					VariableReference.SetSelfMember(VariableFName);
+				}
+
+				VariableName = VariableFName.ToString();
+			}
+			else 
+			{
+				VariableName = VariableIdentifier;
+
+				// First look the variable up as a Member variable
+				VariableGuid = FBlueprintEditorUtils::FindMemberVariableGuidByName(TargetBlueprint, FName(*VariableName));
+
+				// If the variable was not found, look it up as a local variable
+				if (!VariableGuid.IsValid())
+				{
+					VariableGuid = FBlueprintEditorUtils::FindLocalVariableGuidByName(TargetBlueprint, CompilingNode->GetGraph(), FName(*VariableName));
+					VariableReference.SetLocalMember(FName(*VariableName), CompilingNode->GetGraph()->GetName(), VariableGuid);
+				}
+				else
+				{
+					VariableReference.SetSelfMember(FName(*VariableName));
+				}
+
+				// If we found a valid Guid, change the expression's identifier to be the Guid
+				if (VariableGuid.IsValid())
+				{
+					FCString::Strncpy(ExpressionNode.Token.Identifier, *VariableGuid.ToString(EGuidFormats::DigitsWithHyphensInBraces), NAME_SIZE);
+					ExpressionNode.Token.TokenType = FBasicToken::TOKEN_Guid;
+				}
+			}
+
+			if (UProperty* VariableProperty = VariableReference.ResolveMember<UProperty>(TargetBlueprint))
+			{
+				TSharedPtr<FCodeGenFragment_VariableGet> VariableGetFragment = GeneratePropertyFragment(ExpressionNode, VariableProperty, VariableReference, *ActiveMessageLog);
 				if (VariableGetFragment.IsValid())
 				{
 					CompiledFragments.Add(&ExpressionNode, VariableGetFragment);
 				}
 			}
 			// if a variable-get couldn't be created for it, it needs to be an input to the math node
-			else 
+			else if(ExpressionNode.Token.TokenType != FBasicToken::TOKEN_Guid)
 			{
 				CompiledFragments.Add(&ExpressionNode, GenerateInputPinFragment(VariableIdentifier));
 			}
@@ -1497,25 +1613,26 @@ private:
      * wrapper is created and returned (to aid in linking the node later).
      *
      * @param  ExpressionContext    The expression node that we're generating this fragment for.
-     * @param  VariableProperty     The variable that the generated UK2Node should access.
+	 * @param  VariableProperty     The variable that the generated UK2Node should access.
+     * @param  VariableReference    Variable reference to assign to the node
      * @return An empty pointer if we failed to generate something, otherwise new variable-get fragment.
      */
-	TSharedPtr<FCodeGenFragment_VariableGet> GeneratePropertyFragment(FTokenWrapperNode& ExpressionContext, UProperty* VariableProperty, FCompilerResultsLog& MessageLog)
+	TSharedPtr<FCodeGenFragment_VariableGet> GeneratePropertyFragment(FTokenWrapperNode& ExpressionContext, UProperty* VariableProperty, FMemberReference& MemberReference, FCompilerResultsLog& MessageLog)
 	{
-		check(ExpressionContext.Token.TokenType == FBasicToken::TOKEN_Identifier);
+		check(ExpressionContext.Token.TokenType == FBasicToken::TOKEN_Identifier || ExpressionContext.Token.TokenType == FBasicToken::TOKEN_Guid);
 		check(VariableProperty != nullptr);
 		UEdGraphSchema_K2 const* K2Schema = GetDefault<UEdGraphSchema_K2>();
 		
 		TSharedPtr<FCodeGenFragment_VariableGet> VariableGetFragment;
 		
 		UClass* VariableAccessClass = TargetBlueprint->SkeletonGeneratedClass;
-		if (UEdGraphSchema_K2::CanUserKismetAccessVariable(VariableProperty, VariableAccessClass, UEdGraphSchema_K2::CannotBeDelegate))
+		if (MemberReference.IsLocalScope() || UEdGraphSchema_K2::CanUserKismetAccessVariable(VariableProperty, VariableAccessClass, UEdGraphSchema_K2::CannotBeDelegate))
 		{
 			FEdGraphPinType VarType;
 			if (K2Schema->ConvertPropertyToPinType(VariableProperty, /*out*/VarType))
 			{
 				UK2Node_VariableGet* NodeTemplate = NewObject<UK2Node_VariableGet>();
-				NodeTemplate->VariableReference.SetSelfMember(VariableProperty->GetFName());
+				NodeTemplate->VariableReference = MemberReference;
 				UK2Node_VariableGet* VariableGetNode = SpawnNodeFromTemplate<UK2Node_VariableGet>(&ExpressionContext, NodeTemplate);
 				
 				VariableGetFragment = MakeShareable(new FCodeGenFragment_VariableGet(VariableGetNode, VarType));
@@ -2350,6 +2467,10 @@ void UK2Node_MathExpression::RebuildExpression(FString InExpression)
 		ClearExpression();
 		Expression = InExpression;
 
+		// This should not be sanitized, if anything fails to occur, what the user inputed should be what is displayed
+		CachedDisplayExpression.SetCachedText(FText::FromString(Expression), this);
+		CachedNodeTitle.SetCachedText(GetFullTitle(CachedDisplayExpression), this);
+
 		if (!InExpression.IsEmpty()) // @TODO: is this needed?
 		{
 			// build a expression tree from the string
@@ -2364,7 +2485,13 @@ void UK2Node_MathExpression::RebuildExpression(FString InExpression)
 				// a series of errors being attached to the node).
 				if (!GraphGenerator.GenerateCode(ExpressionRoot.ToSharedRef(), *CachedMessageLog))
 				{
-					CachedMessageLog->Error(*LOCTEXT("MathExprFailedGen", "Failed to generate full expression graph for: '@@'").ToString(), this);
+					CachedMessageLog->Error(*LOCTEXT("MathExprGFailedGen", "Failed to generate full expression graph for: '@@'").ToString(), this);
+				}
+				else
+				{
+					Expression = ExpressionRoot->ToString();
+					CachedDisplayExpression.SetCachedText(FText::FromString(SanitizeDisplayExpression(ExpressionRoot->ToDisplayString(GetBlueprint()))), this);
+					CachedNodeTitle.SetCachedText(GetFullTitle(CachedDisplayExpression), this);
 				}
 
 				if (UK2Node_Tunnel* EntryNode = GetEntryNode())
@@ -2479,12 +2606,27 @@ FText UK2Node_MathExpression::GetNodeTitle(ENodeTitleType::Type TitleType) const
 	}
 	else if (TitleType != ENodeTitleType::FullTitle)
 	{
-		return FText::FromString(Expression);
+		if (CachedDisplayExpression.IsOutOfDate(this))
+		{
+			FExpressionParser Parser;
+			TSharedPtr<IFExpressionNode> ExpressionRoot = Parser.ParseExpression(Expression);
+			if(Parser.IsValid())
+			{
+				CachedDisplayExpression.SetCachedText(FText::FromString(SanitizeDisplayExpression(ExpressionRoot->ToDisplayString(GetBlueprint()))), this);
+			}
+		}
+		return CachedDisplayExpression;
 	}
 	else if (CachedNodeTitle.IsOutOfDate(this))
 	{
-		// FText::Format() is slow, so we cache this to save on performance
-		CachedNodeTitle.SetCachedText(FText::Format(LOCTEXT("MathExpressionSecondTitleLine", "{0}\nMath Expression"), FText::FromString(Expression)), this);
+		FExpressionParser Parser;
+		TSharedPtr<IFExpressionNode> ExpressionRoot = Parser.ParseExpression(Expression);
+
+		if(Parser.IsValid())
+		{
+			CachedDisplayExpression.SetCachedText(FText::FromString(SanitizeDisplayExpression(ExpressionRoot->ToDisplayString(GetBlueprint()))), this);
+		}
+		CachedNodeTitle.SetCachedText(GetFullTitle(CachedDisplayExpression), this);
 	}
 	return CachedNodeTitle;
 }
@@ -2506,4 +2648,20 @@ void UK2Node_MathExpression::ReconstructNode()
 	Super::ReconstructNode();
 }
 
+//------------------------------------------------------------------------------
+FString UK2Node_MathExpression::SanitizeDisplayExpression(FString InExpression) const
+{
+	// We do not want the outermost parentheses in the display expression, they add nothing to the logical comprehension
+	InExpression.RemoveFromStart(TEXT("("));
+	InExpression.RemoveFromEnd(TEXT(")"));
+
+	return InExpression;
+}
+
+FText UK2Node_MathExpression::GetFullTitle(FText InExpression) const
+{
+	// FText::Format() is slow, so we cache this to save on performance
+	return FText::Format(LOCTEXT("MathExpressionSecondTitleLine", "{0}\nMath Expression"), InExpression);
+}
+
 #undef LOCTEXT_NAMESPACE
diff --git a/Engine/Source/Editor/UnrealEd/Private/Kismet2/BlueprintEditorUtils.cpp b/Engine/Source/Editor/UnrealEd/Private/Kismet2/BlueprintEditorUtils.cpp
--- a/Engine/Source/Editor/UnrealEd/Private/Kismet2/BlueprintEditorUtils.cpp
+++ b/Engine/Source/Editor/UnrealEd/Private/Kismet2/BlueprintEditorUtils.cpp
@@ -3777,6 +3777,37 @@ void FBlueprintEditorUtils::BulkRemoveMemberVariables(UBlueprint* Blueprint, con
 	}
 }
 
+FGuid FBlueprintEditorUtils::FindMemberVariableGuidByName(UBlueprint* InBlueprint, const FName InVariableName)
+{
+	while(InBlueprint)
+	{
+		const int32 VarIndex = FBlueprintEditorUtils::FindNewVariableIndex(InBlueprint, InVariableName);
+		if (VarIndex != INDEX_NONE)
+		{
+			return InBlueprint->NewVariables[VarIndex].VarGuid;
+		}
+		InBlueprint = Cast<UBlueprint>(InBlueprint->ParentClass->ClassGeneratedBy);
+	}
+	return FGuid(); 
+}
+
+FName FBlueprintEditorUtils::FindMemberVariableNameByGuid(UBlueprint* InBlueprint, const FGuid& InVariableGuid)
+{
+	while(InBlueprint)
+	{
+		for(FBPVariableDescription& Variable : InBlueprint->NewVariables)
+		{
+			if(Variable.VarGuid == InVariableGuid)
+			{
+				return Variable.VarName;
+			}
+		}
+
+		InBlueprint = Cast<UBlueprint>(InBlueprint->ParentClass->ClassGeneratedBy);
+	}
+	return NAME_None;
+}
+
 void FBlueprintEditorUtils::RemoveVariableNodes(UBlueprint* Blueprint, const FName& VarName, bool const bForSelfOnly, UEdGraph* LocalGraphScope)
 {
 	TArray<UEdGraph*> AllGraphs;
diff --git a/Engine/Source/Editor/UnrealEd/Public/Kismet2/BlueprintEditorUtils.h b/Engine/Source/Editor/UnrealEd/Public/Kismet2/BlueprintEditorUtils.h
--- a/Engine/Source/Editor/UnrealEd/Public/Kismet2/BlueprintEditorUtils.h
+++ b/Engine/Source/Editor/UnrealEd/Public/Kismet2/BlueprintEditorUtils.h
@@ -630,6 +630,24 @@ public:
 	 */
 	static void BulkRemoveMemberVariables(UBlueprint* Blueprint, const TArray<FName>& VarNames);
 
+	/**
+	 * Finds a member variable Guid using the variable's name
+	 *
+	 * @param InBlueprint		Blueprint to search for the local variable
+	 * @param InVariableGuid	Local variable's name to search for
+	 * @return					The Guid associated with the local variable
+	 */
+	static FGuid FindMemberVariableGuidByName(UBlueprint* InBlueprint, const FName InVariableName);
+
+	/**
+	 * Finds a member variable name using the variable's Guid
+	 *
+	 * @param InBlueprint		Blueprint to search for the local variable
+	 * @param InVariableGuid	Guid to identify the local variable with
+	 * @return					Local variable's name
+	 */
+	static FName FindMemberVariableNameByGuid(UBlueprint* InBlueprint, const FGuid& InVariableGuid);
+
 	/**
 	 * Removes the variable nodes associated with the specified var name
 	 *
@@ -713,7 +731,6 @@ public:
 	 * Finds a local variable name using the variable's Guid
 	 *
 	 * @param InBlueprint		Blueprint to search for the local variable
-	 * @param InScope			Local variable's scope
 	 * @param InVariableGuid	Guid to identify the local variable with
 	 * @return					Local variable's name
 	 */
