Fixing Mac compiler error.

diff --git a/Engine/Source/Editor/BlueprintGraph/Private/BasicTokenParser.cpp b/Engine/Source/Editor/BlueprintGraph/Private/BasicTokenParser.cpp
--- a/Engine/Source/Editor/BlueprintGraph/Private/BasicTokenParser.cpp
+++ b/Engine/Source/Editor/BlueprintGraph/Private/BasicTokenParser.cpp
@@ -249,7 +249,7 @@ bool FBasicTokenParser::GetToken(FBasicToken& Token, bool bNoConsts/* = false*/)
 	}
 	Token.StartPos		= PrevPos;
 	Token.StartLine		= PrevLine;
-	if( (c=='{') )
+	if( c=='{' )
 	{
 		// Alphanumeric token.
 		int32 Length=0;
