Fix display of {Proto Pin Display Name} in split pins

diff --git a/Engine/Source/Editor/BlueprintGraph/Private/EdGraphSchema_K2.cpp b/Engine/Source/Editor/BlueprintGraph/Private/EdGraphSchema_K2.cpp
--- a/Engine/Source/Editor/BlueprintGraph/Private/EdGraphSchema_K2.cpp
+++ b/Engine/Source/Editor/BlueprintGraph/Private/EdGraphSchema_K2.cpp
@@ -5503,7 +5503,7 @@ void UEdGraphSchema_K2::SplitPin(UEdGraphPin* Pin) const
 			{
 				FFormatNamedArguments Arguments;
 				Arguments.Add(TEXT("PinDisplayName"), Pin->GetDisplayName());
-				Arguments.Add(TEXT("PinDisplayName"), ProtoPin->GetDisplayName());
+				Arguments.Add(TEXT("ProtoPinDisplayName"), ProtoPin->GetDisplayName());
 				SubPin->PinFriendlyName = FText::Format(LOCTEXT("SplitPinFriendlyNameFormat", "{PinDisplayName} {ProtoPinDisplayName}"), Arguments);
 			}
 
