Fixes issues with undo/redo and the Format Text node.

diff --git a/Engine/Source/Editor/BlueprintGraph/Private/K2Node_FormatText.cpp b/Engine/Source/Editor/BlueprintGraph/Private/K2Node_FormatText.cpp
--- a/Engine/Source/Editor/BlueprintGraph/Private/K2Node_FormatText.cpp
+++ b/Engine/Source/Editor/BlueprintGraph/Private/K2Node_FormatText.cpp
@@ -97,6 +97,7 @@ void UK2Node_FormatText::PinConnectionListChanged(UEdGraphPin* Pin)
 {
 	const auto FormatPin = GetFormatPin();
 
+	Modify();
 	// Clear all pins.
 	if(Pin == FormatPin && !FormatPin->DefaultTextValue.IsEmpty())
 	{
@@ -108,6 +109,7 @@ void UK2Node_FormatText::PinConnectionListChanged(UEdGraphPin* Pin)
 			UEdGraphPin* CheckPin = *It;
 			if(CheckPin != FormatPin && CheckPin->Direction == EGPD_Input)
 			{
+				CheckPin->Modify();
 				CheckPin->BreakAllPinLinks();
 				Pins.Remove(CheckPin);
 				--It;
