UE-13070 Hot reload causes AnimBlueprint to lose functionality and crash

diff --git a/Engine/Source/Editor/BlueprintGraph/Private/K2Node_DynamicCast.cpp b/Engine/Source/Editor/BlueprintGraph/Private/K2Node_DynamicCast.cpp
--- a/Engine/Source/Editor/BlueprintGraph/Private/K2Node_DynamicCast.cpp
+++ b/Engine/Source/Editor/BlueprintGraph/Private/K2Node_DynamicCast.cpp
@@ -22,9 +22,12 @@ UK2Node_DynamicCast::UK2Node_DynamicCast(const FObjectInitializer& ObjectInitial
 
 void UK2Node_DynamicCast::AllocateDefaultPins()
 {
-	// Check to track down possible BP comms corruption
-	//@TODO: Move this somewhere more sensible
-	check((TargetType == NULL) || (!TargetType->HasAnyClassFlags(CLASS_NewerVersionExists)));
+	const bool bReferenceObsoleteClass = TargetType && TargetType->HasAnyClassFlags(CLASS_NewerVersionExists);
+	if (bReferenceObsoleteClass)
+	{
+		Message_Error(FString::Printf(TEXT("Node '%s' references obsolete class '%s'"), *GetPathName(), *TargetType->GetPathName()));
+	}
+	ensure(!bReferenceObsoleteClass);
 
 	const UEdGraphSchema_K2* K2Schema = Cast<UEdGraphSchema_K2>(GetSchema());
 	check(K2Schema != nullptr);
