SplitPin crashes when the structure doesn't exist anymore.

A part of: UE-11813 User crashes in UObject::SetLinker() during Ctrl+Z and other operations

diff --git a/Engine/Source/Editor/BlueprintGraph/Private/EdGraphSchema_K2.cpp b/Engine/Source/Editor/BlueprintGraph/Private/EdGraphSchema_K2.cpp
--- a/Engine/Source/Editor/BlueprintGraph/Private/EdGraphSchema_K2.cpp
+++ b/Engine/Source/Editor/BlueprintGraph/Private/EdGraphSchema_K2.cpp
@@ -5422,7 +5422,15 @@ UK2Node* UEdGraphSchema_K2::CreateSplitPinNode(UEdGraphPin* Pin, FKismetCompiler
 {
 	UEdGraphNode* GraphNode = Pin->GetOwningNode();
 	UEdGraph* Graph = GraphNode->GetGraph();
-	UScriptStruct* StructType = CastChecked<UScriptStruct>(Pin->PinType.PinSubCategoryObject.Get());
+	UScriptStruct* StructType = Cast<UScriptStruct>(Pin->PinType.PinSubCategoryObject.Get());
+	if (!StructType)
+	{
+		if (CompilerContext)
+		{
+			CompilerContext->MessageLog.Error(TEXT("No structure in SubCategoryObject in pin @@"), Pin);
+		}
+		StructType = GetFallbackStruct();
+	}
 	UK2Node* SplitPinNode = NULL;
 
 	if (Pin->Direction == EGPD_Input)
diff --git a/Engine/Source/Editor/BlueprintGraph/Private/K2Node.cpp b/Engine/Source/Editor/BlueprintGraph/Private/K2Node.cpp
--- a/Engine/Source/Editor/BlueprintGraph/Private/K2Node.cpp
+++ b/Engine/Source/Editor/BlueprintGraph/Private/K2Node.cpp
@@ -520,9 +520,10 @@ UK2Node::ERedirectType UK2Node::DoPinsMatchForReconstruction(const UEdGraphPin*
 				for (int32 ParentIndex = ParentHierarchy.Num() - 1; ParentIndex >= 0; --ParentIndex)
 				{
 					const UEdGraphPin* CurPin = ParentHierarchy[ParentIndex].Pin;
-					const UEdGraphPin* ParentPin = CurPin->ParentPin;
+					const UEdGraphPin* ParentPin = CurPin ? CurPin->ParentPin : nullptr;
+					const UObject* SubCategoryObject = ParentPin ? ParentPin->PinType.PinSubCategoryObject.Get() : nullptr;
 
-					TMap<FName, FName>* StructRedirects = UStruct::TaggedPropertyRedirects.Find(ParentPin->PinType.PinSubCategoryObject->GetFName());
+					TMap<FName, FName>* StructRedirects = SubCategoryObject ? UStruct::TaggedPropertyRedirects.Find(SubCategoryObject->GetFName()) : nullptr;
 					if (StructRedirects)
 					{
 						FName* PropertyRedirect = StructRedirects->Find(FName(*ParentHierarchy[ParentIndex].PropertyName));
