Strict check in UEdGraphSchema_K2::CreateSplitPinNode

diff --git a/Engine/Source/Editor/BlueprintGraph/Private/EdGraphSchema_K2.cpp b/Engine/Source/Editor/BlueprintGraph/Private/EdGraphSchema_K2.cpp
--- a/Engine/Source/Editor/BlueprintGraph/Private/EdGraphSchema_K2.cpp
+++ b/Engine/Source/Editor/BlueprintGraph/Private/EdGraphSchema_K2.cpp
@@ -5422,7 +5422,7 @@ UK2Node* UEdGraphSchema_K2::CreateSplitPinNode(UEdGraphPin* Pin, FKismetCompiler
 {
 	UEdGraphNode* GraphNode = Pin->GetOwningNode();
 	UEdGraph* Graph = GraphNode->GetGraph();
-	UScriptStruct* StructType = Cast<UScriptStruct>(Pin->PinType.PinSubCategoryObject.Get());
+	UScriptStruct* StructType = CastChecked<UScriptStruct>(Pin->PinType.PinSubCategoryObject.Get(), ECastCheckedType::NullAllowed);
 	if (!StructType)
 	{
 		if (CompilerContext)
