BP - The CustomEvent K2 Node now properly handles generating the proper signature when dragged from a standard, non-multicast delegate pin.

diff --git a/Engine/Source/Editor/BlueprintGraph/Private/K2Node_CustomEvent.cpp b/Engine/Source/Editor/BlueprintGraph/Private/K2Node_CustomEvent.cpp
--- a/Engine/Source/Editor/BlueprintGraph/Private/K2Node_CustomEvent.cpp
+++ b/Engine/Source/Editor/BlueprintGraph/Private/K2Node_CustomEvent.cpp
@@ -309,17 +309,30 @@ void UK2Node_CustomEvent::ReconstructNode()
 {
 	CachedNodeTitle.MarkDirty();
 
+	const UEdGraphSchema_K2* K2Schema = GetDefault<UEdGraphSchema_K2>();
+
 	const UEdGraphPin* DelegateOutPin = FindPin(DelegateOutputName);
+	const UEdGraphPin* LinkedPin = ( DelegateOutPin && DelegateOutPin->LinkedTo.Num() && DelegateOutPin->LinkedTo[0] ) ? DelegateOutPin->LinkedTo[0] : nullptr;
+
+	const UFunction* DelegateSignature = nullptr;
 
-	const UK2Node_BaseMCDelegate* OtherNode = (DelegateOutPin && DelegateOutPin->LinkedTo.Num() && DelegateOutPin->LinkedTo[0]) ?
-		Cast<const UK2Node_BaseMCDelegate>(DelegateOutPin->LinkedTo[0]->GetOwningNode()) : NULL;
-	const UFunction* DelegateSignature = OtherNode ? OtherNode->GetDelegateSignature() : NULL;
+	if ( LinkedPin )
+	{
+		if ( const UK2Node_BaseMCDelegate* OtherNode = Cast<const UK2Node_BaseMCDelegate>(LinkedPin->GetOwningNode()) )
+		{
+			DelegateSignature = OtherNode->GetDelegateSignature();
+		}
+		else if ( LinkedPin->PinType.PinCategory == K2Schema->PC_Delegate )
+		{
+			DelegateSignature = FMemberReference::ResolveSimpleMemberReference<UFunction>(LinkedPin->PinType.PinSubCategoryMemberReference);
+		}
+	}
+	
 	const bool bUseDelegateSignature = (NULL == FindEventSignatureFunction()) && DelegateSignature;
 
 	if (bUseDelegateSignature)
 	{
 		UserDefinedPins.Empty();
-		const UEdGraphSchema_K2* K2Schema = GetDefault<UEdGraphSchema_K2>();
 		for (TFieldIterator<UProperty> PropIt(DelegateSignature); PropIt && (PropIt->PropertyFlags & CPF_Parm); ++PropIt)
 		{
 			const UProperty* Param = *PropIt;
