[INTEGRATE] Change 2471544 by [REDACTED_PERSON][REDACTED_HANDLE] on 2015/04/06 12:10

	Unshelved from pending changelist '2501794':
	[UE-10865] Adding backwards compatibility fo dynamic cast nodes that do not have the boolean result pin yet (they have yet to be regenerated).

diff --git a/Engine/Source/Editor/BlueprintGraph/Private/DynamicCastHandler.cpp b/Engine/Source/Editor/BlueprintGraph/Private/DynamicCastHandler.cpp
--- a/Engine/Source/Editor/BlueprintGraph/Private/DynamicCastHandler.cpp
+++ b/Engine/Source/Editor/BlueprintGraph/Private/DynamicCastHandler.cpp
@@ -11,6 +11,23 @@
 void FKCHandler_DynamicCast::RegisterNets(FKismetFunctionContext& Context, UEdGraphNode* Node)
 {
 	FNodeHandlingFunctor::RegisterNets(Context, Node);
+
+	if (const UK2Node_DynamicCast* DynamicCastNode = Cast<UK2Node_DynamicCast>(Node))
+	{
+		UEdGraphPin* BoolSuccessPin = DynamicCastNode->GetBoolSuccessPin();
+		// this is to support backwards compatibility (when a cast node is generating code, but has yet to be reconstructed)
+		// @TODO: remove this at some point, when backwards compatibility isn't a concern
+		if (BoolSuccessPin == nullptr)
+		{
+			// Create a term to determine if the cast was successful or not
+			FBPTerminal* BoolTerm = Context.CreateLocalTerminal();
+			BoolTerm->Type.PinCategory = CompilerContext.GetSchema()->PC_Boolean;
+			BoolTerm->Source = Node;
+			BoolTerm->Name = Context.NetNameMap->MakeValidName(Node) + TEXT("_CastSuccess");
+			BoolTermMap.Add(Node, BoolTerm);
+		}
+	}
+	
 }
 
 void FKCHandler_DynamicCast::RegisterNet(FKismetFunctionContext& Context, UEdGraphPin* Net)
@@ -106,9 +123,17 @@ void FKCHandler_DynamicCast::Compile(FKismetFunctionContext& Context, UEdGraphNo
 	CastStatement.RHS.Add(ClassTerm);
 	CastStatement.RHS.Add(*ObjectToCast);
 
-	UEdGraphPin*  BoolSuccessPin  = DynamicCastNode->GetBoolSuccessPin();
-	FBPTerminal** BoolSuccessTerm = Context.NetMap.Find(BoolSuccessPin);
+	FBPTerminal** BoolSuccessTerm = nullptr;
+	if (UEdGraphPin* BoolSuccessPin = DynamicCastNode->GetBoolSuccessPin())
+	{
+		BoolSuccessTerm = Context.NetMap.Find(BoolSuccessPin);
+	}
+	else
+	{
+		BoolSuccessTerm = BoolTermMap.Find(DynamicCastNode);
+	}	
 	check(BoolSuccessTerm != nullptr);
+
 	// Check result of cast statement
 	FBlueprintCompiledStatement& CheckResultStatement = Context.AppendStatementForNode(Node);
 	CheckResultStatement.Type = KCST_ObjectToBool;
diff --git a/Engine/Source/Editor/BlueprintGraph/Private/DynamicCastHandler.h b/Engine/Source/Editor/BlueprintGraph/Private/DynamicCastHandler.h
--- a/Engine/Source/Editor/BlueprintGraph/Private/DynamicCastHandler.h
+++ b/Engine/Source/Editor/BlueprintGraph/Private/DynamicCastHandler.h
@@ -10,6 +10,7 @@
 class FKCHandler_DynamicCast : public FNodeHandlingFunctor
 {
 protected:
+	TMap<UEdGraphNode*, FBPTerminal*> BoolTermMap; // should be deprecated at a certain point (left in for backwards compatibility)
 	EKismetCompiledStatementType CastType;
 
 public:
diff --git a/Engine/Source/Editor/BlueprintGraph/Private/K2Node_ClassDynamicCast.cpp b/Engine/Source/Editor/BlueprintGraph/Private/K2Node_ClassDynamicCast.cpp
--- a/Engine/Source/Editor/BlueprintGraph/Private/K2Node_ClassDynamicCast.cpp
+++ b/Engine/Source/Editor/BlueprintGraph/Private/K2Node_ClassDynamicCast.cpp
@@ -90,8 +90,7 @@ UEdGraphPin* UK2Node_ClassDynamicCast::GetCastSourcePin() const
 UEdGraphPin* UK2Node_ClassDynamicCast::GetBoolSuccessPin() const
 {
 	UEdGraphPin* Pin = FindPin(FClassDynamicCastHelper::CastSuccessPinName);
-	check(Pin != nullptr);
-	check(Pin->Direction == EGPD_Output);
+	check((Pin == nullptr) || (Pin->Direction == EGPD_Output));
 	return Pin;
 }
 
diff --git a/Engine/Source/Editor/BlueprintGraph/Private/K2Node_DynamicCast.cpp b/Engine/Source/Editor/BlueprintGraph/Private/K2Node_DynamicCast.cpp
--- a/Engine/Source/Editor/BlueprintGraph/Private/K2Node_DynamicCast.cpp
+++ b/Engine/Source/Editor/BlueprintGraph/Private/K2Node_DynamicCast.cpp
@@ -209,8 +209,7 @@ UEdGraphPin* UK2Node_DynamicCast::GetCastSourcePin() const
 UEdGraphPin* UK2Node_DynamicCast::GetBoolSuccessPin() const
 {
 	UEdGraphPin* Pin = FindPin(UK2Node_DynamicCastImpl::CastSuccessPinName);
-	check(Pin != nullptr);
-	check(Pin->Direction == EGPD_Output);
+	check((Pin == nullptr) || (Pin->Direction == EGPD_Output));
 	return Pin;
 }
 
diff --git a/Engine/Source/Editor/KismetCompiler/Private/KismetCompiler.cpp b/Engine/Source/Editor/KismetCompiler/Private/KismetCompiler.cpp
--- a/Engine/Source/Editor/KismetCompiler/Private/KismetCompiler.cpp
+++ b/Engine/Source/Editor/KismetCompiler/Private/KismetCompiler.cpp
@@ -2890,10 +2890,6 @@ void FKismetCompilerContext::ExpandTunnelsAndMacros(UEdGraph* SourceGraph)
 					/*out*/ NodeOffsetY);
 			}
 
-			// If regenerating on load and the MacroGraph's Blueprint did not reach PostLoad yet, the nodes expanded from the graph will need to be reconstructed
-			UBlueprint* MacroBP = FBlueprintEditorUtils::FindBlueprintForGraph(MacroGraph);
-			bool bMacroNodesNeedRefresh = Blueprint->bIsRegeneratingOnLoad && MacroBP && !MacroBP->bHasBeenRegenerated;
-
 			// Record intermediate object creation nodes, offset the nodes, and handle tunnels
 			for (TArray<UEdGraphNode*>::TIterator MacroNodeIt(MacroNodes); MacroNodeIt; ++MacroNodeIt)
 			{
@@ -2901,10 +2897,6 @@ void FKismetCompilerContext::ExpandTunnelsAndMacros(UEdGraph* SourceGraph)
 				
 				if( DuplicatedNode != NULL )
 				{
-					if (bMacroNodesNeedRefresh)
-					{
-						DuplicatedNode->ReconstructNode();
-					}
 					// Record the source node mapping for the intermediate node first, as it's going to be overwritten through the MessageLog below
 					UEdGraphNode* MacroSourceNode = Cast<UEdGraphNode>(MessageLog.FindSourceObject(DuplicatedNode));
 					if (MacroSourceNode)
