When pasting Event nodes in Blueprints, if there is a ghost node version of the Event, that node will be removed in favor of the newly pasted Event.

diff --git a/Engine/Source/Editor/BlueprintGraph/Classes/K2Node_Event.h b/Engine/Source/Editor/BlueprintGraph/Classes/K2Node_Event.h
--- a/Engine/Source/Editor/BlueprintGraph/Classes/K2Node_Event.h
+++ b/Engine/Source/Editor/BlueprintGraph/Classes/K2Node_Event.h
@@ -91,6 +91,8 @@ class UK2Node_Event : public UK2Node_EditablePinBase
 	 */
 	static FString GetLocalizedNetString(uint32 NetFlags, bool Calling);
 
+	/** Helper function to identify if two Event nodes are the same */
+	static BLUEPRINTGRAPH_API bool AreEventNodesIdentical(const UK2Node_Event* InNodeA, const UK2Node_Event* InNodeB);
 private:
 	/** Constructing FText strings can be costly, so we cache the node's tooltip */
 	FNodeTextCache CachedTooltip;
diff --git a/Engine/Source/Editor/BlueprintGraph/Private/K2Node_Event.cpp b/Engine/Source/Editor/BlueprintGraph/Private/K2Node_Event.cpp
--- a/Engine/Source/Editor/BlueprintGraph/Private/K2Node_Event.cpp
+++ b/Engine/Source/Editor/BlueprintGraph/Private/K2Node_Event.cpp
@@ -426,9 +426,7 @@ bool UK2Node_Event::CanPasteHere(const UEdGraph* TargetGraph) const
 						// If the event function is already handled in this Blueprint, don't paste this event
 						for(int32 i = 0; i < ExistingEventNodes.Num() && !bDisallowPaste; ++i)
 						{
-							bDisallowPaste = ExistingEventNodes[i]->bOverrideFunction
-								&& ExistingEventNodes[i]->EventReference.GetMemberName() == EventReference.GetMemberName()
-								&& ExistingEventNodes[i]->EventReference.GetMemberParentClass(GetBlueprintClassFromNode()) == EventReference.GetMemberParentClass(GetBlueprintClassFromNode());
+							bDisallowPaste = ExistingEventNodes[i]->bOverrideFunction && ExistingEventNodes[i]->bIsNodeEnabled && AreEventNodesIdentical(this, ExistingEventNodes[i]);
 						}
 
 						// We need to also check for 'const' BPIE methods that might already be implemented as functions with a read-only 'self' context (these were previously implemented as events)
@@ -718,4 +716,10 @@ UObject* UK2Node_Event::GetJumpTargetForDoubleClick() const
 	return NULL;
 }
 
+bool UK2Node_Event::AreEventNodesIdentical(const UK2Node_Event* InNodeA, const UK2Node_Event* InNodeB)
+{
+	return InNodeA->EventReference.GetMemberName() == InNodeB->EventReference.GetMemberName()
+		&& InNodeA->EventReference.GetMemberParentClass(InNodeA->GetBlueprintClassFromNode()) == InNodeB->EventReference.GetMemberParentClass(InNodeB->GetBlueprintClassFromNode());
+}
+
 #undef LOCTEXT_NAMESPACE
\ No newline at end of file
diff --git a/Engine/Source/Editor/Kismet/Private/BlueprintEditor.cpp b/Engine/Source/Editor/Kismet/Private/BlueprintEditor.cpp
--- a/Engine/Source/Editor/Kismet/Private/BlueprintEditor.cpp
+++ b/Engine/Source/Editor/Kismet/Private/BlueprintEditor.cpp
@@ -5216,6 +5216,29 @@ void FBlueprintEditor::PasteNodesHere(class UEdGraph* DestinationGraph, const FV
 		{
 			bNeedToModifyStructurally = true;
 		}
+
+		// For pasted Event nodes, we need to see if there is an already existing node in a disabled state that needs to be cleaned up
+		if (UK2Node_Event* EventNode = Cast<UK2Node_Event>(Node))
+		{
+			// Gather all existing event nodes
+			TArray<UK2Node_Event*> ExistingEventNodes;
+			FBlueprintEditorUtils::GetAllNodesOfClass<UK2Node_Event>(GetBlueprintObj(), ExistingEventNodes);
+
+			for (UK2Node_Event* ExistingEventNode : ExistingEventNodes)
+			{
+				bool bIdenticalNode = EventNode != ExistingEventNode && ExistingEventNode->bOverrideFunction && UK2Node_Event::AreEventNodesIdentical(EventNode, ExistingEventNode);
+
+				// Check if the nodes are identical, if they are we need to delete the original because it is disabled. Identical nodes that are in an enabled state will never make it this far and still be enabled.
+				if(bIdenticalNode)
+				{
+					// Should not have made it to being a pasted node if the pre-existing node wasn't disabled.
+					ensure(!ExistingEventNode->bIsNodeEnabled);
+
+					// Destroy the pre-existing node, we do not need it.
+					ExistingEventNode->DestroyNode();
+				}
+			}
+		}
 		// Log new node created to analytics
 		AnalyticsTrackNodeEvent( GetBlueprintObj(), Node, false );
 	}
