Improved searching for pin type classes to find classes, structs, enums, and objects when searching with a sanitized version of the name or unsanitized.

Structs now display, in the selector, with their display name.

Unloaded assets will now display with a sanitized name.

diff --git a/Engine/Source/Editor/BlueprintGraph/Private/EdGraphSchema_K2.cpp b/Engine/Source/Editor/BlueprintGraph/Private/EdGraphSchema_K2.cpp
--- a/Engine/Source/Editor/BlueprintGraph/Private/EdGraphSchema_K2.cpp
+++ b/Engine/Source/Editor/BlueprintGraph/Private/EdGraphSchema_K2.cpp
@@ -124,7 +124,7 @@ private:
 		const FString Tooltip = (TooltipPtr && !TooltipPtr->IsEmpty()) ? *TooltipPtr : InAsset.ObjectPath.ToString();
 
 		FPinTypeTreeInfoPtr TypeTreeInfo = MakeShareable(new UEdGraphSchema_K2::FPinTypeTreeInfo(InCategoryName, InAsset.ToStringReference(), FText::FromString(Tooltip)));
-		TypeTreeInfo->FriendlyName = FText::FromName(InAsset.AssetName);
+		TypeTreeInfo->FriendlyName = FText::FromString(FName::NameToDisplayString(InAsset.AssetName.ToString(), false));
 		OutChildren.Add(TypeTreeInfo);
 	}
 
@@ -412,13 +412,13 @@ FText UEdGraphSchema_K2::FPinTypeTreeInfo::GetDescription() const
 	else if (PinType.PinSubCategoryObject.IsValid())
 	{
 		FText DisplayName;
-		if (UClass* SubCategoryClass = Cast<UClass>(PinType.PinSubCategoryObject.Get()))
+		if (UField* SubCategoryField = Cast<UField>(PinType.PinSubCategoryObject.Get()))
 		{
-			DisplayName = FBlueprintEditorUtils::GetFriendlyClassDisplayName(SubCategoryClass);
+			DisplayName = SubCategoryField->GetDisplayNameText();
 		}
 		else
 		{
-			DisplayName = FText::FromString(PinType.PinSubCategoryObject->GetName());
+			DisplayName = FText::FromString(FName::NameToDisplayString(PinType.PinSubCategoryObject->GetName(), PinType.PinCategory == PC_Boolean));
 		}
 
 		return DisplayName;
diff --git a/Engine/Source/Editor/KismetWidgets/Private/SPinTypeSelector.cpp b/Engine/Source/Editor/KismetWidgets/Private/SPinTypeSelector.cpp
--- a/Engine/Source/Editor/KismetWidgets/Private/SPinTypeSelector.cpp
+++ b/Engine/Source/Editor/KismetWidgets/Private/SPinTypeSelector.cpp
@@ -311,7 +311,7 @@ void SPinTypeSelector::OnFilterTextChanged(const FText& NewText)
 	SearchText = NewText;
 	FilteredTypeTreeRoot.Empty();
 
-	GetChildrenMatchingSearch(NewText.ToString(), TypeTreeRoot, FilteredTypeTreeRoot);
+	GetChildrenMatchingSearch(FName::NameToDisplayString(NewText.ToString(), false), TypeTreeRoot, FilteredTypeTreeRoot);
 	TypeTreeView->RequestTreeRefresh();
 
 	// Select the first non-category item
