Strip s when [REDACTED_HANDLE] is used instead of [REDACTED_HANDLE] for a return pin tooltip

diff --git a/Engine/Source/Editor/BlueprintGraph/Private/K2Node_CallFunction.cpp b/Engine/Source/Editor/BlueprintGraph/Private/K2Node_CallFunction.cpp
--- a/Engine/Source/Editor/BlueprintGraph/Private/K2Node_CallFunction.cpp
+++ b/Engine/Source/Editor/BlueprintGraph/Private/K2Node_CallFunction.cpp
@@ -1159,7 +1159,8 @@ void UK2Node_CallFunction::GeneratePinTooltipFromFunction(UEdGraphPin& Pin, cons
 	// figure what tag we should be parsing for (is this a return-val pin, or a parameter?)
 	FString ParamName;
 	FString TagStr = TEXT("@param");
-	if (Pin.PinName == UEdGraphSchema_K2::PN_ReturnValue)
+	const bool bReturnPin = Pin.PinName == UEdGraphSchema_K2::PN_ReturnValue;
+	if (bReturnPin)
 	{
 		TagStr = TEXT("@return");
 	}
@@ -1185,6 +1186,12 @@ void UK2Node_CallFunction::GeneratePinTooltipFromFunction(UEdGraphPin& Pin, cons
 		// advance past the tag
 		CurStrPos += TagStr.Len();
 
+		// handle people having done @returns instead of @return
+		if (bReturnPin && CurStrPos < FullToolTipLen && FunctionToolTipText[CurStrPos] == TEXT('s'))
+		{
+			++CurStrPos;
+		}
+
 		// advance past whitespace
 		while(CurStrPos < FullToolTipLen && FChar::IsWhitespace(FunctionToolTipText[CurStrPos]))
 		{
