Renaming pins on macros will correctly refresh the details view.

Fixes finding a valid name for macro and collapsed graph pins.

diff --git a/Engine/Source/Editor/BlueprintGraph/Classes/K2Node_Tunnel.h b/Engine/Source/Editor/BlueprintGraph/Classes/K2Node_Tunnel.h
--- a/Engine/Source/Editor/BlueprintGraph/Classes/K2Node_Tunnel.h
+++ b/Engine/Source/Editor/BlueprintGraph/Classes/K2Node_Tunnel.h
@@ -41,6 +41,7 @@ class UK2Node_Tunnel : public UK2Node_EditablePinBase
 	virtual bool CanUserDeleteNode() const override;
 	BLUEPRINTGRAPH_API virtual bool CanDuplicateNode() const override;
 	virtual UObject* GetJumpTargetForDoubleClick() const override;
+	virtual FString CreateUniquePinName(FString SourcePinName) const override;
 	// End UEdGraphNode interface
 
 	// Begin UK2Node interface.
diff --git a/Engine/Source/Editor/BlueprintGraph/Private/K2Node_Tunnel.cpp b/Engine/Source/Editor/BlueprintGraph/Private/K2Node_Tunnel.cpp
--- a/Engine/Source/Editor/BlueprintGraph/Private/K2Node_Tunnel.cpp
+++ b/Engine/Source/Editor/BlueprintGraph/Private/K2Node_Tunnel.cpp
@@ -69,6 +69,33 @@ FText UK2Node_Tunnel::GetNodeTitle(ENodeTitleType::Type TitleType) const
 	}
 }
 
+FString UK2Node_Tunnel::CreateUniquePinName(FString InSourcePinName) const
+{
+	if (GetClass() == UK2Node_Tunnel::StaticClass())
+	{
+		// When dealing with a tunnel node that is not a sub class (macro/collapsed graph entry and result), attempt to find the paired node and find a valid name between the two
+		TWeakObjectPtr<UK2Node_EditablePinBase> TunnelEntry;
+		TWeakObjectPtr<UK2Node_EditablePinBase> TunnelResult;
+		FBlueprintEditorUtils::GetEntryAndResultNodes(GetGraph(), TunnelEntry, TunnelResult);
+
+		if (TunnelEntry.IsValid() && TunnelResult.IsValid())
+		{
+			FString PinName(InSourcePinName);
+
+			int32 Index = 1;
+			while (TunnelEntry.Get()->FindPin(PinName) != nullptr || TunnelResult.Get()->FindPin(PinName) != nullptr)
+			{
+				++Index;
+				PinName = InSourcePinName + FString::FromInt(Index);
+			}
+
+			return PinName;
+		}
+	}
+
+	return Super::CreateUniquePinName(InSourcePinName);
+}
+
 bool UK2Node_Tunnel::CanUserDeleteNode() const
 {
 	// Disallow deletion of tunnels that are inside a tunnel graph, but allow it on top level tunnels that have gotten there on accident
diff --git a/Engine/Source/Editor/Kismet/Private/BlueprintDetailsCustomization.cpp b/Engine/Source/Editor/Kismet/Private/BlueprintDetailsCustomization.cpp
--- a/Engine/Source/Editor/Kismet/Private/BlueprintDetailsCustomization.cpp
+++ b/Engine/Source/Editor/Kismet/Private/BlueprintDetailsCustomization.cpp
@@ -3209,7 +3209,7 @@ bool FBaseBlueprintGraphActionDetails::OnPinRenamed(UK2Node_EditablePinBase* Tar
 
 		if (FunctionResultNodePtr.IsValid())
 		{
-			PinRenamedHelper.NodesToRename.Add(FunctionEntryNodePtr.Get());
+			PinRenamedHelper.NodesToRename.Add(FunctionResultNodePtr.Get());
 		}
 
 		PinRenamedHelper.ModifiedBlueprints.Add(GetBlueprintObj());
diff --git a/Engine/Source/Editor/UnrealEd/Private/Kismet2/BlueprintEditorUtils.cpp b/Engine/Source/Editor/UnrealEd/Private/Kismet2/BlueprintEditorUtils.cpp
--- a/Engine/Source/Editor/UnrealEd/Private/Kismet2/BlueprintEditorUtils.cpp
+++ b/Engine/Source/Editor/UnrealEd/Private/Kismet2/BlueprintEditorUtils.cpp
@@ -4431,7 +4431,8 @@ UFunction* FFunctionFromNodeHelper::FunctionFromNode(const UK2Node* Node)
 			check(EntryNodes.Num() == 1);
 			Node = EntryNodes[0];
 		}
-		else if (const UK2Node_FunctionEntry* FunctionNode = Cast<UK2Node_FunctionEntry>(Node))
+		
+		if (const UK2Node_FunctionEntry* FunctionNode = Cast<UK2Node_FunctionEntry>(Node))
 		{
 			const FName FunctionName = (FunctionNode->CustomGeneratedFunctionName != NAME_None) ? FunctionNode->CustomGeneratedFunctionName : FunctionNode->GetGraph()->GetFName();
 			Function = SearchScope->FindFunctionByName(FunctionName);
