[UE-13209] Fixing it so multigate nodes don't stomp memory (when called both inside and out of a function).

diff --git a/Engine/Source/Editor/BlueprintGraph/Private/K2Node_MultiGate.cpp b/Engine/Source/Editor/BlueprintGraph/Private/K2Node_MultiGate.cpp
--- a/Engine/Source/Editor/BlueprintGraph/Private/K2Node_MultiGate.cpp
+++ b/Engine/Source/Editor/BlueprintGraph/Private/K2Node_MultiGate.cpp
@@ -22,16 +22,24 @@ protected:
 	// Map to an int used to keep track of which outputs have been used
 	TMap<UEdGraphNode*, FBPTerminal*> DataTermMap;
 
-	// Generic bool term used for run-time conditions
-	FBPTerminal* GenericBoolTerm;
-	// Index term used for run-time index determination
-	FBPTerminal* IndexTerm;
+	/** 
+	 * These are locals that don't need to be independent for each node, instead 
+	 * they can be shared between each multigate node in the graph.
+	 */
+	struct FFunctionScopedTerms
+	{
+		FFunctionScopedTerms() : GenericBoolTerm(nullptr), IndexTerm(nullptr) {}
+
+		// Generic bool term used for run-time conditions
+		FBPTerminal* GenericBoolTerm;
+		// Index term used for run-time index determination
+		FBPTerminal* IndexTerm;
+	};
+	TMap<UFunction*, FFunctionScopedTerms> FunctionTermMap;
 
 public:
 	FKCHandler_MultiGate(FKismetCompilerContext& InCompilerContext)
 		: FNodeHandlingFunctor(InCompilerContext)
-		, GenericBoolTerm(NULL)
-		, IndexTerm(NULL)
 	{
 	}
 
@@ -59,22 +67,24 @@ public:
 			DataTermMap.Add(Node, DataTerm);
 		}
 
+		FFunctionScopedTerms& FuncLocals = FunctionTermMap.FindOrAdd(Context.Function);
+
 		// Create a local scratch bool for run-time if there isn't already one
-		if (!GenericBoolTerm)
+		if (!FuncLocals.GenericBoolTerm)
 		{
-			GenericBoolTerm = Context.CreateLocalTerminal();
-			GenericBoolTerm->Type.PinCategory = CompilerContext.GetSchema()->PC_Boolean;
-			GenericBoolTerm->Source = Node;
-			GenericBoolTerm->Name = BaseNetName + TEXT("_ScratchBool");
+			FuncLocals.GenericBoolTerm = Context.CreateLocalTerminal();
+			FuncLocals.GenericBoolTerm->Type.PinCategory = CompilerContext.GetSchema()->PC_Boolean;
+			FuncLocals.GenericBoolTerm->Source = Node;
+			FuncLocals.GenericBoolTerm->Name = BaseNetName + TEXT("_ScratchBool");
 		}
 
 		// Create a local scratch int for run-time index tracking if there isn't already one
-		if (!IndexTerm)
+		if (!FuncLocals.IndexTerm)
 		{
-			IndexTerm = Context.CreateLocalTerminal();
-			IndexTerm->Type.PinCategory = CompilerContext.GetSchema()->PC_Int;
-			IndexTerm->Source = Node;
-			IndexTerm->Name = BaseNetName + TEXT("_ScratchIndex");
+			FuncLocals.IndexTerm = Context.CreateLocalTerminal();
+			FuncLocals.IndexTerm->Type.PinCategory = CompilerContext.GetSchema()->PC_Int;
+			FuncLocals.IndexTerm->Source = Node;
+			FuncLocals.IndexTerm->Name = BaseNetName + TEXT("_ScratchIndex");
 		}
 	}
 
@@ -185,6 +195,9 @@ public:
 		// See if this is the first time in
 		///////////////////////////////////////////////////
 
+		FFunctionScopedTerms& FuncLocals = FunctionTermMap.FindChecked(Context.Function);
+		check(FuncLocals.GenericBoolTerm != nullptr);
+
 		// (bIsNotFirstTime != true)
 		FBlueprintCompiledStatement& BoolNotEqualStatement = Context.AppendStatementForNode(Node);
 		BoolNotEqualStatement.Type = KCST_CallFunction;
@@ -192,7 +205,7 @@ public:
 		BoolNotEqualStatement.FunctionContext = NULL;
 		BoolNotEqualStatement.bIsParentContext = false;
 		// Set the params
-		BoolNotEqualStatement.LHS = GenericBoolTerm;
+		BoolNotEqualStatement.LHS = FuncLocals.GenericBoolTerm;
 		BoolNotEqualStatement.RHS.Add(FirstRunBoolTerm);
 		BoolNotEqualStatement.RHS.Add(TrueBoolTerm);
 
@@ -200,7 +213,7 @@ public:
 		// {
 		FBlueprintCompiledStatement& IfFirstTimeStatement = Context.AppendStatementForNode(Node);
 		IfFirstTimeStatement.Type = KCST_GotoIfNot;
-		IfFirstTimeStatement.LHS = GenericBoolTerm;
+		IfFirstTimeStatement.LHS = FuncLocals.GenericBoolTerm;
 
 		///////////////////////////////////////////////////////////////////
 		// This is the first time in... set the bool and the start index
@@ -222,7 +235,7 @@ public:
 		Statement.FunctionToCall = ConditionFunction;
 		Statement.FunctionContext = NULL;
 		Statement.bIsParentContext = false;
-		Statement.LHS = GenericBoolTerm;
+		Statement.LHS = FuncLocals.GenericBoolTerm;
 		Statement.RHS.Add(*StartIndexPinTerm);
 		Statement.RHS.Add(InvalidIndexTerm);
 
@@ -230,7 +243,7 @@ public:
 		// {
 		FBlueprintCompiledStatement& IfHasIndexStatement = Context.AppendStatementForNode(Node);
 		IfHasIndexStatement.Type = KCST_GotoIfNot;
-		IfHasIndexStatement.LHS = GenericBoolTerm;
+		IfHasIndexStatement.LHS = FuncLocals.GenericBoolTerm;
 
 		///////////////////////////////////////////////////////////////////
 		// They supplied a start index so set the index to it
@@ -239,7 +252,7 @@ public:
 		// Index = StartIndex; // (StartIndex is from multi gate pin for it)
 		FBlueprintCompiledStatement& AssignSuppliedIndexStatement = Context.AppendStatementForNode(Node);
 		AssignSuppliedIndexStatement.Type = KCST_Assignment;
-		AssignSuppliedIndexStatement.LHS = IndexTerm;
+		AssignSuppliedIndexStatement.LHS = FuncLocals.IndexTerm;
 		AssignSuppliedIndexStatement.RHS.Add(*StartIndexPinTerm);
 
 		// Jump to index usage
@@ -253,12 +266,14 @@ public:
 		// They did NOT supply a start index so figure one out
 		///////////////////////////////////////////////////////////////////
 
+		check(FuncLocals.IndexTerm != nullptr);
+
 		// Index = GetUnmarkedBit(Data, -1, bRandom);
 		FBlueprintCompiledStatement& GetStartIndexStatement = Context.AppendStatementForNode(Node);
 		GetStartIndexStatement.Type = KCST_CallFunction;
 		GetStartIndexStatement.FunctionToCall = GetUnmarkedBitFunction;
 		GetStartIndexStatement.bIsParentContext = false;
-		GetStartIndexStatement.LHS = IndexTerm;
+		GetStartIndexStatement.LHS = FuncLocals.IndexTerm;
 		GetStartIndexStatement.RHS.Add(DataTerm);
 		GetStartIndexStatement.RHS.Add(*StartIndexPinTerm);
 		GetStartIndexStatement.RHS.Add(NumOutsTerm);
@@ -284,7 +299,7 @@ public:
 		IsAvailableStatement.FunctionToCall = HasUnmarkedBitFunction;
 		IsAvailableStatement.FunctionContext = NULL;
 		IsAvailableStatement.bIsParentContext = false;
-		IsAvailableStatement.LHS = GenericBoolTerm;
+		IsAvailableStatement.LHS = FuncLocals.GenericBoolTerm;
 		IsAvailableStatement.RHS.Add(DataTerm);
 		IsAvailableStatement.RHS.Add(NumOutsTerm);
 		// Hook the IfFirstTimeStatement jump to this node
@@ -295,7 +310,7 @@ public:
 		// {
 		FBlueprintCompiledStatement& IfIsAvailableStatement = Context.AppendStatementForNode(Node);
 		IfIsAvailableStatement.Type = KCST_GotoIfNot;
-		IfIsAvailableStatement.LHS = GenericBoolTerm;
+		IfIsAvailableStatement.LHS = FuncLocals.GenericBoolTerm;
 
 		////////////////////////////////////////////////////////////////////////////
 		// Has available index so figure it out and jump to its' usage
@@ -306,7 +321,7 @@ public:
 		GetNextIndexStatement.Type = KCST_CallFunction;
 		GetNextIndexStatement.FunctionToCall = GetUnmarkedBitFunction;
 		GetNextIndexStatement.bIsParentContext = false;
-		GetNextIndexStatement.LHS = IndexTerm;
+		GetNextIndexStatement.LHS = FuncLocals.IndexTerm;
 		GetNextIndexStatement.RHS.Add(DataTerm);
 		GetNextIndexStatement.RHS.Add(*StartIndexPinTerm);
 		GetNextIndexStatement.RHS.Add(NumOutsTerm);
@@ -369,9 +384,9 @@ public:
 		MarkIndexStatement.Type = KCST_CallFunction;
 		MarkIndexStatement.FunctionToCall = MarkBitFunction;
 		MarkIndexStatement.bIsParentContext = false;
-		MarkIndexStatement.LHS = IndexTerm;
+		MarkIndexStatement.LHS = FuncLocals.IndexTerm;
 		MarkIndexStatement.RHS.Add(DataTerm);
-		MarkIndexStatement.RHS.Add(IndexTerm);
+		MarkIndexStatement.RHS.Add(FuncLocals.IndexTerm);
 
 		// Setup jump label
 		MarkIndexStatement.bIsJumpTarget = true;
@@ -400,14 +415,14 @@ public:
 			LiteralIndexTerm->Type.PinCategory = CompilerContext.GetSchema()->PC_Int;
 			LiteralIndexTerm->Name = FString::Printf(TEXT("%d"), OutIdx);
 			// Set the params
-			IndexEqualityStatement.LHS = GenericBoolTerm;
-			IndexEqualityStatement.RHS.Add(IndexTerm);
+			IndexEqualityStatement.LHS = FuncLocals.GenericBoolTerm;
+			IndexEqualityStatement.RHS.Add(FuncLocals.IndexTerm);
 			IndexEqualityStatement.RHS.Add(LiteralIndexTerm);
 
 			// if (Index == OutIdx)
 			FBlueprintCompiledStatement& IfIndexMatchesStatement = Context.AppendStatementForNode(Node);
 			IfIndexMatchesStatement.Type = KCST_GotoIfNot;
-			IfIndexMatchesStatement.LHS = GenericBoolTerm;
+			IfIndexMatchesStatement.LHS = FuncLocals.GenericBoolTerm;
 			// {
 			//////////////////////////////////////
 			// Found a match - Jump there
