BP: flag bExecutionAfterReturn added.

diff --git a/Engine/Config/BaseEngine.ini b/Engine/Config/BaseEngine.ini
--- a/Engine/Config/BaseEngine.ini
+++ b/Engine/Config/BaseEngine.ini
@@ -1544,6 +1544,7 @@ bEnableInheritableComponents=true
 bDeferDependencyLoads=true
 bForceDisableCookedDependencyDeferring=false
 bFirstCompileChildrenThenReinstance=true
+bExecutionAfterReturn=false
 
 [/Script/Engine.Blueprint]
 bRecompileOnLoad=true
diff --git a/Engine/Source/Editor/BlueprintGraph/Private/K2Node_FunctionResult.cpp b/Engine/Source/Editor/BlueprintGraph/Private/K2Node_FunctionResult.cpp
--- a/Engine/Source/Editor/BlueprintGraph/Private/K2Node_FunctionResult.cpp
+++ b/Engine/Source/Editor/BlueprintGraph/Private/K2Node_FunctionResult.cpp
@@ -68,18 +68,28 @@ public:
 
 	virtual void Compile(FKismetFunctionContext& Context, UEdGraphNode* Node) override
 	{
-		GenerateAssigments(Context, Node);
+		static const FBoolConfigValueHelper ExecutionAfterReturn(TEXT("Kismet"), TEXT("bExecutionAfterReturn"), GEngineIni);
 
-		if (Context.bCreateDebugData && Node)
+		if (ExecutionAfterReturn)
 		{
-			FBlueprintCompiledStatement& TraceStatement = Context.AppendStatementForNode(Node);
-			TraceStatement.Type = KCST_WireTraceSite;
-			TraceStatement.Comment = Node->NodeComment.IsEmpty() ? Node->GetName() : Node->NodeComment;
+			// for backward compatibility only
+			FKCHandler_VariableSet::Compile(Context, Node);
 		}
+		else
+		{
+			GenerateAssigments(Context, Node);
 
-		// always go to return
-		FBlueprintCompiledStatement& GotoStatement = Context.AppendStatementForNode(Node);
-		GotoStatement.Type = KCST_GotoReturn; 
+			if (Context.bCreateDebugData && Node)
+			{
+				FBlueprintCompiledStatement& TraceStatement = Context.AppendStatementForNode(Node);
+				TraceStatement.Type = KCST_WireTraceSite;
+				TraceStatement.Comment = Node->NodeComment.IsEmpty() ? Node->GetName() : Node->NodeComment;
+			}
+
+			// always go to return
+			FBlueprintCompiledStatement& GotoStatement = Context.AppendStatementForNode(Node);
+			GotoStatement.Type = KCST_GotoReturn;
+		}
 	}
 };
 
