This change adds better tracking of changes in levels and worlds in order to keep the blueprint action database current.
#UE-10665 - when a sub level is removed it fails to remove any graph actions registered for itself
#UE-10962 - Blueprint level script functions are missing from action menu
#UE-7004 - Custom Events can't be called after changing level blueprint or switching maps
#Branch UE4
#Proj BlueprintGraph

diff --git a/Engine/Source/Editor/BlueprintGraph/Private/BlueprintActionDatabase.cpp b/Engine/Source/Editor/BlueprintGraph/Private/BlueprintActionDatabase.cpp
--- a/Engine/Source/Editor/BlueprintGraph/Private/BlueprintActionDatabase.cpp
+++ b/Engine/Source/Editor/BlueprintGraph/Private/BlueprintActionDatabase.cpp
@@ -448,6 +448,13 @@ namespace BlueprintActionDatabaseImpl
 	 */
 	static void OnAssetRenamed(FAssetData const& AssetInfo, const FString& InOldName);
 
+	/**
+	 * Callback to refresh/add all level blueprints owned by this world to the database
+	 * 
+	 * @param  NewWorld		The world that was added.
+	 */
+	static void OnWorldAdded(UWorld* NewWorld);
+
 	/**
 	 * Callback to clear all levels from the database when a world is destroyed 
 	 * 
@@ -455,6 +462,14 @@ namespace BlueprintActionDatabaseImpl
 	 */
 	static void OnWorldDestroyed(UWorld* DestroyedWorld);
 
+	/**
+	 * Callback to re-evaluate all level blueprints owned by the world when the layout has changed
+	 * 
+	 * @param  Level			The level that has added/removed.
+	 * @param  World			The owner of the level and the world to be re-evaluated.
+	 */
+	static void OnLevelsChanged(ULevel* Level, UWorld* World);
+
 	/**
 	 * Returns TRUE if the Object is valid for the database
 	 *
@@ -885,16 +900,22 @@ static void BlueprintActionDatabaseImpl::OnAssetRenamed(FAssetData const& AssetI
 	}
 }
 
+//------------------------------------------------------------------------------
+static void BlueprintActionDatabaseImpl::OnWorldAdded(UWorld* NewWorld)
+{
+	FBlueprintActionDatabase::Get().RefreshAssetActions( (UObject*)NewWorld );
+}
+
 //------------------------------------------------------------------------------
 static void BlueprintActionDatabaseImpl::OnWorldDestroyed(UWorld* DestroyedWorld)
 {
-	for(auto It = DestroyedWorld->GetLevelIterator() ; It ; ++It)
-	{
-		if(const ULevelScriptBlueprint* LevelScript = (*It)->GetLevelScriptBlueprint(/*bDontCreate =*/true))
-		{
-			FBlueprintActionDatabase::Get().ClearAssetActions((UObject*)LevelScript);
-		}
-	}
+	FBlueprintActionDatabase::Get().ClearAssetActions( (UObject*)DestroyedWorld );
+}
+
+//------------------------------------------------------------------------------
+static void BlueprintActionDatabaseImpl::OnLevelsChanged(ULevel* Level, UWorld* World)
+{
+	FBlueprintActionDatabase::Get().RefreshAssetActions( (UObject*)World );
 }
 
 //------------------------------------------------------------------------------
@@ -919,6 +940,11 @@ static bool BlueprintActionDatabaseImpl::IsObjectValidForDatabase(UObject const*
 		// Level scripts are sometimes not assets because they have not been saved yet, but they are still valid for the database.
 		bReturn = FBlueprintEditorUtils::IsLevelScriptBlueprint(Blueprint);
 	}
+	else if( UWorld const* World = Cast<UWorld>(Object))
+	{
+		// We now use worlds as databse keys to manage the level scripts they own.
+		bReturn = true;
+	}
 	return bReturn;
 }
 
@@ -952,7 +978,11 @@ FBlueprintActionDatabase::FBlueprintActionDatabase()
 	FEditorDelegates::OnAssetsPreDelete.AddStatic(&BlueprintActionDatabaseImpl::OnAssetsPendingDelete);
 	FKismetEditorUtilities::OnBlueprintUnloaded.AddStatic(&BlueprintActionDatabaseImpl::OnBlueprintUnloaded);
 
+	GEngine->OnWorldAdded().AddStatic(&BlueprintActionDatabaseImpl::OnWorldAdded);
 	GEngine->OnWorldDestroyed().AddStatic(&BlueprintActionDatabaseImpl::OnWorldDestroyed);
+	FWorldDelegates::LevelAddedToWorld.AddStatic(&BlueprintActionDatabaseImpl::OnLevelsChanged);
+	FWorldDelegates::LevelRemovedFromWorld.AddStatic(&BlueprintActionDatabaseImpl::OnLevelsChanged);
+	FWorldDelegates::LevelScriptBlueprintCreated.AddStatic(&BlueprintActionDatabaseImpl::OnLevelsChanged);
 
 	IHotReloadInterface& HotReloadSupport = FModuleManager::LoadModuleChecked<IHotReloadInterface>("HotReload");
 	HotReloadSupport.OnHotReload().AddStatic(&BlueprintActionDatabaseImpl::OnProjectHotReloaded);
@@ -1059,6 +1089,27 @@ void FBlueprintActionDatabase::RefreshAll()
 	// this handles creating entries for components that were loaded before the database was alive:
 	FComponentTypeRegistry::Get().SubscribeToComponentList(ComponentTypes).AddRaw(this, &FBlueprintActionDatabase::RefreshComponentActions);
 	RefreshComponentActions();
+
+	// Refresh existing worlds
+	RefreshWorlds();
+}
+
+//------------------------------------------------------------------------------
+void FBlueprintActionDatabase::RefreshWorlds()
+{
+	// Add all level scripts from current world
+	const TIndirectArray<FWorldContext>& WorldContexts = GEngine->GetWorldContexts();
+
+	for( auto Context : WorldContexts )
+	{
+		if( Context.WorldType == EWorldType::Editor )
+		{
+			if( UWorld* CurrentWorld = Context.World())
+			{
+				RefreshAssetActions((UObject*)CurrentWorld);
+			}
+		}
+	}
 }
 
 //------------------------------------------------------------------------------
@@ -1068,9 +1119,10 @@ void FBlueprintActionDatabase::RefreshClassActions(UClass* const Class)
 	check(Class != nullptr);
 
 	bool const bOutOfDateClass   = Class->HasAnyClassFlags(CLASS_NewerVersionExists);
-	bool const bIsBlueprintClass = (Cast<UBlueprintGeneratedClass>(Class) != nullptr);	
+	bool const bIsBlueprintClass = (Cast<UBlueprintGeneratedClass>(Class) != nullptr);
+	bool const bIsLevelScript	 = Class->ClassGeneratedBy && Cast<UBlueprint>(Class->ClassGeneratedBy)->BlueprintType == EBlueprintType::BPTYPE_LevelScript;
 
-	if (bOutOfDateClass)
+	if (bOutOfDateClass || bIsLevelScript)
 	{
 		ActionRegistry.Remove(Class);
 		return;
@@ -1200,6 +1252,31 @@ void FBlueprintActionDatabase::RefreshAssetActions(UObject* const AssetObject)
 		}
 	}
 
+	UWorld* WorldAsset = Cast<UWorld>( AssetObject );
+	if( WorldAsset && WorldAsset->WorldType == EWorldType::Editor )
+	{
+		for( auto Level : WorldAsset->GetLevels() )
+		{
+			if(UBlueprint* LevelScript = Cast<UBlueprint>(Level->GetLevelScriptBlueprint(true)))
+			{
+				AddBlueprintGraphActions(LevelScript, AssetActionList);
+				if (UClass* SkeletonClass = LevelScript->SkeletonGeneratedClass)
+				{
+					GetClassMemberActions(SkeletonClass, AssetActionList);
+				}
+				// Register for change and compilation notifications
+				if( !LevelScript->OnChanged().IsBoundToObject( this ))
+				{
+					LevelScript->OnChanged().AddRaw(this, &FBlueprintActionDatabase::OnBlueprintChanged);
+				}
+				if( !LevelScript->OnCompiled().IsBoundToObject( this ))
+				{
+					LevelScript->OnCompiled().AddRaw(this, &FBlueprintActionDatabase::OnBlueprintChanged);
+				}
+			}
+		}
+	}
+
 	FBlueprintActionDatabaseRegistrar Registrar(ActionRegistry, UnloadedActionRegistry, ActionPrimingQueue);
 	Registrar.ActionKeyFilter = AssetObject; // make sure actions only associated with this asset get added
 	// nodes may have actions they wish to add actions for this asset
@@ -1220,8 +1297,8 @@ void FBlueprintActionDatabase::RefreshAssetActions(UObject* const AssetObject)
 	// we don't want to clear entries for blueprints, mainly because we 
 	// use the presence of an entry to know if we've set the blueprint's 
 	// OnChanged(), but also because most blueprints will have actions at some 
-	// later point
-	else if (BlueprintAsset == nullptr)
+	// later point. Same goes for world assets because they are used to managed level scripts blueprints.
+	else if (!BlueprintAsset && !WorldAsset)
 	{
 		ClearAssetActions(AssetObject);
 	}
@@ -1344,7 +1421,15 @@ void FBlueprintActionDatabase::RegisterAllNodeActions(FBlueprintActionDatabaseRe
 
 void FBlueprintActionDatabase::OnBlueprintChanged(UBlueprint* InBlueprint)
 {
-	BlueprintActionDatabaseImpl::OnBlueprintChanged(InBlueprint);
+	if( InBlueprint->BlueprintType == BPTYPE_LevelScript )
+	{
+		// Levelscript blueprints are managed through their owning worlds.
+		RefreshWorlds();
+	}
+	else
+	{
+		BlueprintActionDatabaseImpl::OnBlueprintChanged(InBlueprint);
+	}
 }
 
 #undef LOCTEXT_NAMESPACE
diff --git a/Engine/Source/Editor/BlueprintGraph/Public/BlueprintActionDatabase.h b/Engine/Source/Editor/BlueprintGraph/Public/BlueprintActionDatabase.h
--- a/Engine/Source/Editor/BlueprintGraph/Public/BlueprintActionDatabase.h
+++ b/Engine/Source/Editor/BlueprintGraph/Public/BlueprintActionDatabase.h
@@ -70,6 +70,11 @@ public:
 	 */
 	void RefreshAll();
 
+	/**
+	 * Populates the action database with all level script actions from all active editor worlds.
+	 */
+	void RefreshWorlds();
+
 	/**
 	 * Finds the database entry for the specified class and wipes it, 
 	 * repopulating it with a fresh set of associated node-spawners. 
diff --git a/Engine/Source/Editor/UnrealEd/Private/EditorLevelUtils.cpp b/Engine/Source/Editor/UnrealEd/Private/EditorLevelUtils.cpp
--- a/Engine/Source/Editor/UnrealEd/Private/EditorLevelUtils.cpp
+++ b/Engine/Source/Editor/UnrealEd/Private/EditorLevelUtils.cpp
@@ -420,6 +420,9 @@ namespace EditorLevelUtils
 
 			EditorDestroyLevel(InLevel);
 
+			// Scrub any levelscript actions so GC can take effect
+			FWorldDelegates::LevelRemovedFromWorld.Broadcast( InLevel, OwningWorld );
+
 			// Redraw the main editor viewports.
 			FEditorSupportDelegates::RedrawAllViewports.Broadcast();
 
diff --git a/Engine/Source/Runtime/Engine/Classes/Engine/World.h b/Engine/Source/Runtime/Engine/Classes/Engine/World.h
--- a/Engine/Source/Runtime/Engine/Classes/Engine/World.h
+++ b/Engine/Source/Runtime/Engine/Classes/Engine/World.h
@@ -3229,6 +3229,14 @@ public:
 	// called by UWorld::GetAssetRegistryTags()
 	static FWorldGetAssetTags GetAssetTags;
 
+#if WITH_EDITOR
+	// Delegate type for Levelscript blueprint creation events
+	DECLARE_MULTICAST_DELEGATE_TwoParams(FLevelScriptCreationEvent, ULevel*, UWorld* );
+
+	// Called when a ULevel has a level script blueprint created for it.
+	static FLevelScriptCreationEvent LevelScriptBlueprintCreated;
+#endif
+
 private:
 	FWorldDelegates() {}
 };
diff --git a/Engine/Source/Runtime/Engine/Private/Level.cpp b/Engine/Source/Runtime/Engine/Private/Level.cpp
--- a/Engine/Source/Runtime/Engine/Private/Level.cpp
+++ b/Engine/Source/Runtime/Engine/Private/Level.cpp
@@ -1659,6 +1659,8 @@ ULevelScriptBlueprint* ULevel::GetLevelScriptBlueprint(bool bDontCreate)
 		// LevelScript blueprints should not be standalone
 		LevelScriptBlueprint->ClearFlags(RF_Standalone);
 		ULevel::LevelDirtiedEvent.Broadcast();
+		// Broadcast Levelscript creation.
+		FWorldDelegates::LevelScriptBlueprintCreated.Broadcast(this, OwningWorld);
 	}
 
 	// Ensure that friendly name is always up-to-date
diff --git a/Engine/Source/Runtime/Engine/Private/World.cpp b/Engine/Source/Runtime/Engine/Private/World.cpp
--- a/Engine/Source/Runtime/Engine/Private/World.cpp
+++ b/Engine/Source/Runtime/Engine/Private/World.cpp
@@ -104,6 +104,9 @@ FWorldDelegates::FWorldEvent FWorldDelegates::OnPreWorldFinishDestroy;
 FWorldDelegates::FOnLevelChanged FWorldDelegates::LevelAddedToWorld;
 FWorldDelegates::FOnLevelChanged FWorldDelegates::LevelRemovedFromWorld;
 FWorldDelegates::FWorldGetAssetTags FWorldDelegates::GetAssetTags;
+#if WITH_EDITOR
+FWorldDelegates::FLevelScriptCreationEvent FWorldDelegates::LevelScriptBlueprintCreated;
+#endif // WITH_EDITOR
 
 UWorld::UWorld( const FObjectInitializer& ObjectInitializer )
 :	UObject(ObjectInitializer)
