[UE-13595] Variable node autowire codepath did not handle local variables, which have no self pin

diff --git a/Engine/Source/Editor/BlueprintGraph/Private/K2Node_Variable.cpp b/Engine/Source/Editor/BlueprintGraph/Private/K2Node_Variable.cpp
--- a/Engine/Source/Editor/BlueprintGraph/Private/K2Node_Variable.cpp
+++ b/Engine/Source/Editor/BlueprintGraph/Private/K2Node_Variable.cpp
@@ -647,17 +647,20 @@ void UK2Node_Variable::AutowireNewNode(UEdGraphPin* FromPin)
 					if(FromPin->PinType.PinSubCategoryObject == PropertyOwner || dynamic_cast<UClass*>(FromPin->PinType.PinSubCategoryObject.Get())->IsChildOf(PropertyOwner))
 					{
 						UEdGraphPin* TargetPin = FindPin(K2Schema->PN_Self);
-						TargetPin->PinType.PinSubCategoryObject = PropertyOwner;
-
-						if(K2Schema->TryCreateConnection(FromPin, TargetPin))
+						if (TargetPin)
 						{
-							bConnected = true;
+							TargetPin->PinType.PinSubCategoryObject = PropertyOwner;
+
+							if(K2Schema->TryCreateConnection(FromPin, TargetPin))
+							{
+								bConnected = true;
 
-							// Setup the VariableReference correctly since it may no longer be a self member
-							VariableReference.SetFromField<UProperty>(GetPropertyForVariable(), false);
-							TargetPin->bHidden = false;
-							FromPin->GetOwningNode()->NodeConnectionListChanged();
-							this->NodeConnectionListChanged();
+								// Setup the VariableReference correctly since it may no longer be a self member
+								VariableReference.SetFromField<UProperty>(GetPropertyForVariable(), false);
+								TargetPin->bHidden = false;
+								FromPin->GetOwningNode()->NodeConnectionListChanged();
+								this->NodeConnectionListChanged();
+							}
 						}
 					}
 				}
