This fixes a problem cleaning up sublevels when they are removed.
#UE-10665 - when a sub level is removed it fails to remove any graph actions registered for itself
#Branch UE4
#Proj BlueprintGraph

diff --git a/Engine/Source/Editor/BlueprintGraph/Private/BlueprintActionDatabase.cpp b/Engine/Source/Editor/BlueprintGraph/Private/BlueprintActionDatabase.cpp
--- a/Engine/Source/Editor/BlueprintGraph/Private/BlueprintActionDatabase.cpp
+++ b/Engine/Source/Editor/BlueprintGraph/Private/BlueprintActionDatabase.cpp
@@ -466,10 +466,9 @@ namespace BlueprintActionDatabaseImpl
 	/**
 	 * Callback to re-evaluate all level blueprints owned by the world when the layout has changed
 	 * 
-	 * @param  Level			The level that has added/removed.
 	 * @param  World			The owner of the level and the world to be re-evaluated.
 	 */
-	static void OnLevelsChanged(ULevel* Level, UWorld* World);
+	static void OnRefreshLevelScripts(UWorld* World);
 
 	/**
 	 * Returns TRUE if the Object is valid for the database
@@ -914,7 +913,7 @@ static void BlueprintActionDatabaseImpl::OnWorldDestroyed(UWorld* DestroyedWorld
 }
 
 //------------------------------------------------------------------------------
-static void BlueprintActionDatabaseImpl::OnLevelsChanged(ULevel* Level, UWorld* World)
+static void BlueprintActionDatabaseImpl::OnRefreshLevelScripts(UWorld* World)
 {
 	FBlueprintActionDatabase::Get().RefreshAssetActions( (UObject*)World );
 }
@@ -981,9 +980,7 @@ FBlueprintActionDatabase::FBlueprintActionDatabase()
 
 	GEngine->OnWorldAdded().AddStatic(&BlueprintActionDatabaseImpl::OnWorldAdded);
 	GEngine->OnWorldDestroyed().AddStatic(&BlueprintActionDatabaseImpl::OnWorldDestroyed);
-	FWorldDelegates::LevelAddedToWorld.AddStatic(&BlueprintActionDatabaseImpl::OnLevelsChanged);
-	FWorldDelegates::LevelRemovedFromWorld.AddStatic(&BlueprintActionDatabaseImpl::OnLevelsChanged);
-	FWorldDelegates::LevelScriptBlueprintCreated.AddStatic(&BlueprintActionDatabaseImpl::OnLevelsChanged);
+	FWorldDelegates::RefreshLevelScriptActions.AddStatic(&BlueprintActionDatabaseImpl::OnRefreshLevelScripts);
 
 	IHotReloadInterface& HotReloadSupport = FModuleManager::LoadModuleChecked<IHotReloadInterface>("HotReload");
 	HotReloadSupport.OnHotReload().AddStatic(&BlueprintActionDatabaseImpl::OnProjectHotReloaded);
@@ -1253,12 +1250,13 @@ void FBlueprintActionDatabase::RefreshAssetActions(UObject* const AssetObject)
 		}
 	}
 
-	UWorld* WorldAsset = Cast<UWorld>( AssetObject );
-	if( WorldAsset && WorldAsset->WorldType == EWorldType::Editor )
+	UWorld* WorldAsset = Cast<UWorld>(AssetObject);
+	if (WorldAsset && WorldAsset->WorldType == EWorldType::Editor)
 	{
 		for( auto Level : WorldAsset->GetLevels() )
 		{
-			if(UBlueprint* LevelScript = Cast<UBlueprint>(Level->GetLevelScriptBlueprint(true)))
+			UBlueprint* LevelScript = Cast<UBlueprint>(Level->GetLevelScriptBlueprint(true));
+			if (Level->bIsVisible && LevelScript)
 			{
 				AddBlueprintGraphActions(LevelScript, AssetActionList);
 				if (UClass* SkeletonClass = LevelScript->SkeletonGeneratedClass)
@@ -1266,11 +1264,11 @@ void FBlueprintActionDatabase::RefreshAssetActions(UObject* const AssetObject)
 					GetClassMemberActions(SkeletonClass, AssetActionList);
 				}
 				// Register for change and compilation notifications
-				if( !LevelScript->OnChanged().IsBoundToObject( this ))
+				if (!LevelScript->OnChanged().IsBoundToObject(this))
 				{
 					LevelScript->OnChanged().AddRaw(this, &FBlueprintActionDatabase::OnBlueprintChanged);
 				}
-				if( !LevelScript->OnCompiled().IsBoundToObject( this ))
+				if (!LevelScript->OnCompiled().IsBoundToObject(this))
 				{
 					LevelScript->OnCompiled().AddRaw(this, &FBlueprintActionDatabase::OnBlueprintChanged);
 				}
diff --git a/Engine/Source/Runtime/Engine/Classes/Engine/World.h b/Engine/Source/Runtime/Engine/Classes/Engine/World.h
--- a/Engine/Source/Runtime/Engine/Classes/Engine/World.h
+++ b/Engine/Source/Runtime/Engine/Classes/Engine/World.h
@@ -2948,7 +2948,7 @@ public:
 	void NotifyOfBlueprintDebuggingAssociation(class UBlueprint* Blueprint, UObject* DebugObject);
 
 	/** Broadcasts that the number of levels has changed. */
-	void BroadcastLevelsChanged() { LevelsChangedEvent.Broadcast(); }
+	void BroadcastLevelsChanged();
 
 	/** Returns the LevelsChangedEvent member. */
 	FOnLevelsChangedEvent& OnLevelsChanged() { return LevelsChangedEvent; }
@@ -3245,11 +3245,11 @@ public:
 	static FWorldGetAssetTags GetAssetTags;
 
 #if WITH_EDITOR
-	// Delegate type for Levelscript blueprint creation events
-	DECLARE_MULTICAST_DELEGATE_TwoParams(FLevelScriptCreationEvent, ULevel*, UWorld* );
+	// Delegate called when levelscript actions need refreshing
+	DECLARE_MULTICAST_DELEGATE_OneParam(FRefreshLevelScriptActionsEvent, UWorld*);
 
-	// Called when a ULevel has a level script blueprint created for it.
-	static FLevelScriptCreationEvent LevelScriptBlueprintCreated;
+	// Called when changes in the levels require blueprint actions to be refreshed.
+	static FRefreshLevelScriptActionsEvent RefreshLevelScriptActions;
 #endif
 
 private:
diff --git a/Engine/Source/Runtime/Engine/Private/Level.cpp b/Engine/Source/Runtime/Engine/Private/Level.cpp
--- a/Engine/Source/Runtime/Engine/Private/Level.cpp
+++ b/Engine/Source/Runtime/Engine/Private/Level.cpp
@@ -1659,8 +1659,8 @@ ULevelScriptBlueprint* ULevel::GetLevelScriptBlueprint(bool bDontCreate)
 		// LevelScript blueprints should not be standalone
 		LevelScriptBlueprint->ClearFlags(RF_Standalone);
 		ULevel::LevelDirtiedEvent.Broadcast();
-		// Broadcast Levelscript creation.
-		FWorldDelegates::LevelScriptBlueprintCreated.Broadcast(this, OwningWorld);
+		// Refresh level script actions
+		FWorldDelegates::RefreshLevelScriptActions.Broadcast(OwningWorld);
 	}
 
 	// Ensure that friendly name is always up-to-date
diff --git a/Engine/Source/Runtime/Engine/Private/World.cpp b/Engine/Source/Runtime/Engine/Private/World.cpp
--- a/Engine/Source/Runtime/Engine/Private/World.cpp
+++ b/Engine/Source/Runtime/Engine/Private/World.cpp
@@ -107,7 +107,7 @@ FWorldDelegates::FOnLevelChanged FWorldDelegates::LevelRemovedFromWorld;
 FWorldDelegates::FLevelOffsetEvent FWorldDelegates::PostApplyLevelOffset;
 FWorldDelegates::FWorldGetAssetTags FWorldDelegates::GetAssetTags;
 #if WITH_EDITOR
-FWorldDelegates::FLevelScriptCreationEvent FWorldDelegates::LevelScriptBlueprintCreated;
+FWorldDelegates::FRefreshLevelScriptActionsEvent FWorldDelegates::RefreshLevelScriptActions;
 #endif // WITH_EDITOR
 
 UWorld::UWorld( const FObjectInitializer& ObjectInitializer )
@@ -1548,6 +1548,14 @@ void UWorld::NotifyOfBlueprintDebuggingAssociation(class UBlueprint* Blueprint,
 #endif	//#if WITH_EDITOR
 }
 
+void UWorld::BroadcastLevelsChanged()
+{
+	LevelsChangedEvent.Broadcast();
+#if WITH_EDITOR
+	FWorldDelegates::RefreshLevelScriptActions.Broadcast(this);
+#endif	//#if WITH_EDITOR
+}
+
 DEFINE_STAT(STAT_AddToWorldTime);
 DEFINE_STAT(STAT_RemoveFromWorldTime);
 DEFINE_STAT(STAT_UpdateLevelStreamingTime);
