[UE-13740] Adding a "Owner" param to the SpawnActorFromClass node.

diff --git a/Engine/Source/Editor/BlueprintGraph/Classes/K2Node_SpawnActorFromClass.h b/Engine/Source/Editor/BlueprintGraph/Classes/K2Node_SpawnActorFromClass.h
--- a/Engine/Source/Editor/BlueprintGraph/Classes/K2Node_SpawnActorFromClass.h
+++ b/Engine/Source/Editor/BlueprintGraph/Classes/K2Node_SpawnActorFromClass.h
@@ -55,6 +55,8 @@ class BLUEPRINTGRAPH_API UK2Node_SpawnActorFromClass : public UK2Node
 	UEdGraphPin* GetSpawnTransformPin() const;
 	/** Get the spawn NoCollisionFail input pin */
 	UEdGraphPin* GetNoCollisionFailPin() const;
+	/** Get the actor owner pin */
+	UEdGraphPin* GetOwnerPin() const;
 	/** Get the result output pin */
 	UEdGraphPin* GetResultPin() const;
 
diff --git a/Engine/Source/Editor/BlueprintGraph/Private/K2Node_SpawnActorFromClass.cpp b/Engine/Source/Editor/BlueprintGraph/Private/K2Node_SpawnActorFromClass.cpp
--- a/Engine/Source/Editor/BlueprintGraph/Private/K2Node_SpawnActorFromClass.cpp
+++ b/Engine/Source/Editor/BlueprintGraph/Private/K2Node_SpawnActorFromClass.cpp
@@ -13,12 +13,14 @@ struct FK2Node_SpawnActorFromClassHelper
 	static FString ClassPinName;
 	static FString SpawnTransformPinName;
 	static FString NoCollisionFailPinName;
+	static FString OwnerPinName;
 };
 
 FString FK2Node_SpawnActorFromClassHelper::WorldContextPinName(TEXT("WorldContextObject"));
 FString FK2Node_SpawnActorFromClassHelper::ClassPinName(TEXT("Class"));
 FString FK2Node_SpawnActorFromClassHelper::SpawnTransformPinName(TEXT("SpawnTransform"));
 FString FK2Node_SpawnActorFromClassHelper::NoCollisionFailPinName(TEXT("SpawnEvenIfColliding"));
+FString FK2Node_SpawnActorFromClassHelper::OwnerPinName(TEXT("Owner"));
 
 #define LOCTEXT_NAMESPACE "K2Node_SpawnActorFromClass"
 
@@ -55,6 +57,14 @@ void UK2Node_SpawnActorFromClass::AllocateDefaultPins()
 	UEdGraphPin* NoCollisionFailPin = CreatePin(EGPD_Input, K2Schema->PC_Boolean, TEXT(""), NULL, false, false, FK2Node_SpawnActorFromClassHelper::NoCollisionFailPinName);
 	K2Schema->ConstructBasicPinTooltip(*NoCollisionFailPin, LOCTEXT("NoCollisionFailPinDescription", "Determines if the Actor should be spawned when the location is blocked by a collision"), NoCollisionFailPin->PinToolTip);
 
+	UEdGraphPin* OwnerPin = CreatePin(EGPD_Input, K2Schema->PC_Object, TEXT(""), AActor::StaticClass(),/*bIsArray =*/false, /*bIsReference =*/false, FK2Node_SpawnActorFromClassHelper::OwnerPinName);
+	OwnerPin->bAdvancedView = true;
+	if (ENodeAdvancedPins::NoPins == AdvancedPinDisplay)
+	{
+		AdvancedPinDisplay = ENodeAdvancedPins::Hidden;
+	}
+	K2Schema->ConstructBasicPinTooltip(*OwnerPin, LOCTEXT("OwnerPinDescription", "Can be left empty; primarily used for replication (bNetUseOwnerRelevancy and bOnlyRelevantToOwner), or visibility (PrimitiveComponent's bOwnerNoSee/bOnlyOwnerSee)"), OwnerPin->PinToolTip);
+
 	// Result pin
 	UEdGraphPin* ResultPin = CreatePin(EGPD_Output, K2Schema->PC_Object, TEXT(""), AActor::StaticClass(), false, false, K2Schema->PN_ReturnValue);
 	K2Schema->ConstructBasicPinTooltip(*ResultPin, LOCTEXT("ResultPinDescription", "The spawned Actor"), ResultPin->PinToolTip);
@@ -173,7 +183,8 @@ bool UK2Node_SpawnActorFromClass::IsSpawnVarPin(UEdGraphPin* Pin)
 			Pin->PinName != FK2Node_SpawnActorFromClassHelper::ClassPinName &&
 			Pin->PinName != FK2Node_SpawnActorFromClassHelper::WorldContextPinName &&
 			Pin->PinName != FK2Node_SpawnActorFromClassHelper::NoCollisionFailPinName &&
-			Pin->PinName != FK2Node_SpawnActorFromClassHelper::SpawnTransformPinName );
+			Pin->PinName != FK2Node_SpawnActorFromClassHelper::SpawnTransformPinName && 
+			Pin->PinName != FK2Node_SpawnActorFromClassHelper::OwnerPinName );
 }
 
 void UK2Node_SpawnActorFromClass::OnClassPinChanged()
@@ -285,6 +296,13 @@ UEdGraphPin* UK2Node_SpawnActorFromClass::GetNoCollisionFailPin() const
 	return Pin;
 }
 
+UEdGraphPin* UK2Node_SpawnActorFromClass::GetOwnerPin() const
+{
+	UEdGraphPin* Pin = FindPin(FK2Node_SpawnActorFromClassHelper::OwnerPinName);
+	check(Pin == NULL || Pin->Direction == EGPD_Input);
+	return Pin;
+}
+
 UEdGraphPin* UK2Node_SpawnActorFromClass::GetWorldContextPin() const
 {
 	UEdGraphPin* Pin = FindPin(FK2Node_SpawnActorFromClassHelper::WorldContextPinName);
@@ -409,6 +427,7 @@ void UK2Node_SpawnActorFromClass::ExpandNode(class FKismetCompilerContext& Compi
 	static FString ActorParamName = FString(TEXT("Actor"));
 	static FString TransformParamName = FString(TEXT("SpawnTransform"));
 	static FString NoCollisionFailParamName = FString(TEXT("bNoCollisionFail"));
+	static FString OwnerParamName = FString(TEXT("Owner"));
 
 	static FString ObjectParamName = FString(TEXT("Object"));
 	static FString ValueParamName = FString(TEXT("Value"));
@@ -420,6 +439,7 @@ void UK2Node_SpawnActorFromClass::ExpandNode(class FKismetCompilerContext& Compi
 	UEdGraphPin* SpawnNodeNoCollisionFail = GetNoCollisionFailPin();
 	UEdGraphPin* SpawnWorldContextPin = SpawnNode->GetWorldContextPin();
 	UEdGraphPin* SpawnClassPin = SpawnNode->GetClassPin();
+	UEdGraphPin* SpawnNodeOwnerPin = SpawnNode->GetOwnerPin();
 	UEdGraphPin* SpawnNodeThen = SpawnNode->GetThenPin();
 	UEdGraphPin* SpawnNodeResult = SpawnNode->GetResultPin();
 
@@ -443,7 +463,8 @@ void UK2Node_SpawnActorFromClass::ExpandNode(class FKismetCompilerContext& Compi
 	UEdGraphPin* CallBeginActorClassPin = CallBeginSpawnNode->FindPinChecked(ActorClassParamName);
 	UEdGraphPin* CallBeginTransform = CallBeginSpawnNode->FindPinChecked(TransformParamName);
 	UEdGraphPin* CallBeginNoCollisionFail = CallBeginSpawnNode->FindPinChecked(NoCollisionFailParamName);
-	UEdGraphPin* CallBeginResult = CallBeginSpawnNode->GetReturnValuePin();
+	UEdGraphPin* CallBeginOwnerPin = CallBeginSpawnNode->FindPinChecked(FK2Node_SpawnActorFromClassHelper::OwnerPinName);
+	UEdGraphPin* CallBeginResult = CallBeginSpawnNode->GetReturnValuePin();	
 
 	// Move 'exec' connection from spawn node to 'begin spawn'
 	CompilerContext.MovePinLinksToIntermediate(*SpawnNodeExec, *CallBeginExec);
@@ -465,6 +486,11 @@ void UK2Node_SpawnActorFromClass::ExpandNode(class FKismetCompilerContext& Compi
 		CompilerContext.MovePinLinksToIntermediate(*SpawnWorldContextPin, *CallBeginWorldContextPin);
 	}
 
+	if (SpawnNodeOwnerPin != nullptr)
+	{
+		CompilerContext.MovePinLinksToIntermediate(*SpawnNodeOwnerPin, *CallBeginOwnerPin);
+	}
+
 	// Copy the 'transform' connection from the spawn node to 'begin spawn'
 	CompilerContext.MovePinLinksToIntermediate(*SpawnNodeTransform, *CallBeginTransform);
 
diff --git a/Engine/Source/Runtime/Engine/Classes/Kismet/GameplayStatics.h b/Engine/Source/Runtime/Engine/Classes/Kismet/GameplayStatics.h
--- a/Engine/Source/Runtime/Engine/Classes/Kismet/GameplayStatics.h
+++ b/Engine/Source/Runtime/Engine/Classes/Kismet/GameplayStatics.h
@@ -38,7 +38,7 @@ class ENGINE_API UGameplayStatics : public UBlueprintFunctionLibrary
 
 	/** Spawns an instance of an actor class, but does not automatically run its construction script.  */
 	UFUNCTION(BlueprintCallable, Category="Spawning", meta=(WorldContext="WorldContextObject", UnsafeDuringActorConstruction = "true", BlueprintInternalUseOnly = "true"))
-	static class AActor* BeginSpawningActorFromClass(UObject* WorldContextObject, TSubclassOf<AActor> ActorClass, const FTransform& SpawnTransform, bool bNoCollisionFail = false);
+	static class AActor* BeginSpawningActorFromClass(UObject* WorldContextObject, TSubclassOf<AActor> ActorClass, const FTransform& SpawnTransform, bool bNoCollisionFail = false, AActor* Owner = nullptr);
 
 	/** 'Finish' spawning an actor.  This will run the construction script. */
 	UFUNCTION(BlueprintCallable, Category="Spawning", meta=(UnsafeDuringActorConstruction = "true", BlueprintInternalUseOnly = "true"))
diff --git a/Engine/Source/Runtime/Engine/Private/GameplayStatics.cpp b/Engine/Source/Runtime/Engine/Private/GameplayStatics.cpp
--- a/Engine/Source/Runtime/Engine/Private/GameplayStatics.cpp
+++ b/Engine/Source/Runtime/Engine/Private/GameplayStatics.cpp
@@ -297,7 +297,7 @@ AActor* UGameplayStatics::BeginSpawningActorFromBlueprint(UObject* WorldContextO
 	return NewActor;
 }
 
-AActor* UGameplayStatics::BeginSpawningActorFromClass(UObject* WorldContextObject, TSubclassOf<AActor> ActorClass, const FTransform& SpawnTransform, bool bNoCollisionFail)
+AActor* UGameplayStatics::BeginSpawningActorFromClass(UObject* WorldContextObject, TSubclassOf<AActor> ActorClass, const FTransform& SpawnTransform, bool bNoCollisionFail, AActor* Owner)
 {
 	const FVector SpawnLoc(SpawnTransform.GetTranslation());
 	const FRotator SpawnRot(SpawnTransform.GetRotation());
@@ -323,7 +323,7 @@ AActor* UGameplayStatics::BeginSpawningActorFromClass(UObject* WorldContextObjec
 		UWorld* World = GEngine->GetWorldFromContextObject(WorldContextObject);
 		if (World)
 		{
-			NewActor = World->SpawnActorDeferred<AActor>(Class, SpawnLoc, SpawnRot, NULL, AutoInstigator, bNoCollisionFail);
+			NewActor = World->SpawnActorDeferred<AActor>(Class, SpawnLoc, SpawnRot, Owner, AutoInstigator, bNoCollisionFail);
 		}
 		else
 		{
