DoOnceMultiInput node now works in non-English languages.

Changed PinNames to always be in English and set the PinFriendlyNames to be the localized form.

diff --git a/Engine/Source/Editor/BlueprintGraph/Classes/K2Node_DoOnceMultiInput.h b/Engine/Source/Editor/BlueprintGraph/Classes/K2Node_DoOnceMultiInput.h
--- a/Engine/Source/Editor/BlueprintGraph/Classes/K2Node_DoOnceMultiInput.h
+++ b/Engine/Source/Editor/BlueprintGraph/Classes/K2Node_DoOnceMultiInput.h
@@ -27,7 +27,7 @@ private:
 	const static int32 NumBaseInputs = 1;
 
 	static int32 GetMaxInputPinsNum();
-	static FString GetNameForPin(int32 PinIndex, bool In);
+	static FText GetNameForPin(int32 PinIndex, bool In);
 
 	FEdGraphPinType GetInType() const;
 	FEdGraphPinType GetOutType() const;
diff --git a/Engine/Source/Editor/BlueprintGraph/Private/K2Node_DoOnceMultiInput.cpp b/Engine/Source/Editor/BlueprintGraph/Private/K2Node_DoOnceMultiInput.cpp
--- a/Engine/Source/Editor/BlueprintGraph/Private/K2Node_DoOnceMultiInput.cpp
+++ b/Engine/Source/Editor/BlueprintGraph/Private/K2Node_DoOnceMultiInput.cpp
@@ -53,13 +53,17 @@ int32 UK2Node_DoOnceMultiInput::GetMaxInputPinsNum()
 	return (TCHAR('Z') - TCHAR('A'));
 }
 
-FString UK2Node_DoOnceMultiInput::GetNameForPin(int32 PinIndex, bool In)
+FText UK2Node_DoOnceMultiInput::GetNameForPin(int32 PinIndex, bool In)
 {
 	check(PinIndex < GetMaxInputPinsNum());
 	FString Name;
 	Name.AppendChar(TCHAR('A') + PinIndex);
-	Name = In ? Name + LOCTEXT("DoOnceMultiIn", " In").ToString() : Name + LOCTEXT("DoOnceMultiOut", " Out").ToString();
-	return Name;
+	
+	FFormatNamedArguments Args;
+	Args.Add(TEXT("Identifier"), FText::FromString(Name));
+	Args.Add(TEXT("Direction"), In ? LOCTEXT("DoOnceMultiIn", "In") : LOCTEXT("DoOnceMultiOut", "Out"));
+
+	return FText::Format(LOCTEXT("DoOnceMultiInputPinName", "{Identifier} {Direction}"), Args);
 }
 
 UK2Node_DoOnceMultiInput::UK2Node_DoOnceMultiInput(const FObjectInitializer& ObjectInitializer)
@@ -178,12 +182,19 @@ void UK2Node_DoOnceMultiInput::AllocateDefaultPins()
 
 	const UEdGraphSchema_K2* K2Schema = GetDefault<UEdGraphSchema_K2>();
 
-	
-	CreatePin(EGPD_Input, K2Schema->PC_Exec, TEXT(""), NULL, false, false, GetNameForPin(0, true));
-	CreatePin(EGPD_Output, K2Schema->PC_Exec, TEXT(""), NULL, false, false, GetNameForPin(0, false));
+	FText InputPinAName = GetNameForPin(0, true);
+	UEdGraphPin* InputPinA = CreatePin(EGPD_Input, K2Schema->PC_Exec, TEXT(""), NULL, false, false, InputPinAName.BuildSourceString());
+	InputPinA->PinFriendlyName = InputPinAName;
+
+	FText OutputPinAName = GetNameForPin(0, false);
+	UEdGraphPin* OutputPinA = CreatePin(EGPD_Output, K2Schema->PC_Exec, TEXT(""), NULL, false, false, OutputPinAName.BuildSourceString());
+	OutputPinA->PinFriendlyName = OutputPinAName;
+
+	UEdGraphPin* DoOnceResetIn = CreatePin(EGPD_Input, K2Schema->PC_Exec, TEXT(""), NULL, false, false, TEXT("Reset In"));
+	DoOnceResetIn->PinFriendlyName = LOCTEXT("DoOnceResetIn", "Reset In");
 
-	CreatePin(EGPD_Input, K2Schema->PC_Exec, TEXT(""), NULL, false, false, LOCTEXT("DoOnceResetIn", "Reset In").ToString());
-	CreatePin(EGPD_Output, K2Schema->PC_Exec, TEXT(""), NULL, false, false, LOCTEXT("DoOnceResetOut", "Reset Out").ToString());
+	UEdGraphPin* DoOnceResetOut = CreatePin(EGPD_Output, K2Schema->PC_Exec, TEXT(""), NULL, false, false, TEXT("Reset Out"));
+	DoOnceResetOut->PinFriendlyName = LOCTEXT("DoOnceResetOut", "Reset Out");
 
 	for (int32 i = 0; i < NumAdditionalInputs; ++i)
 	{
@@ -195,26 +206,30 @@ void UK2Node_DoOnceMultiInput::AddPinsInner(int32 AdditionalPinIndex)
 {
 	{
 		const FEdGraphPinType InputType = GetInType();
-		CreatePin(EGPD_Input, 
+		FText InputPinName = GetNameForPin(AdditionalPinIndex, true);
+		UEdGraphPin* InputPin = CreatePin(EGPD_Input, 
 			InputType.PinCategory, 
 			InputType.PinSubCategory, 
 			InputType.PinSubCategoryObject.Get(), 
 			InputType.bIsArray, 
 			InputType.bIsReference, 
-			*GetNameForPin(AdditionalPinIndex, true)
+			InputPinName.BuildSourceString()
 		);
+		InputPin->PinFriendlyName = InputPinName;
 	}
 
 	{
 		const FEdGraphPinType OutputType = GetOutType();
-		CreatePin(EGPD_Output,
+		FText OutputPinName = GetNameForPin(AdditionalPinIndex, false);
+		UEdGraphPin* OutputPin = CreatePin(EGPD_Output,
 			OutputType.PinCategory,
 			OutputType.PinSubCategory,
 			OutputType.PinSubCategoryObject.Get(),
 			OutputType.bIsArray,
 			OutputType.bIsReference,
-			*GetNameForPin(AdditionalPinIndex, false)
+			OutputPinName.BuildSourceString()
 		);
+		OutputPin->PinFriendlyName = OutputPinName;
 	}
 }
 
@@ -255,7 +270,7 @@ void UK2Node_DoOnceMultiInput::RemoveInputPin(UEdGraphPin* Pin)
 				UEdGraphPin* LocalPin = Pins[PinIndex];
 				if(LocalPin && (LocalPin != OutPin) && (LocalPin != SelfPin))
 				{
-					const FString PinName = GetNameForPin(NameIndex + NumBaseInputs, true);  // FIXME
+					const FString PinName = GetNameForPin(NameIndex + NumBaseInputs, true).BuildSourceString();  // FIXME
 					if(PinName != LocalPin->PinName)
 					{
 						LocalPin->Modify();
