Fix for missing break/make menu entries for structures with blueprint readable/writable properties but no BlueprintType tag

diff --git a/Engine/Source/Editor/BlueprintGraph/Classes/K2Node_BreakStruct.h b/Engine/Source/Editor/BlueprintGraph/Classes/K2Node_BreakStruct.h
--- a/Engine/Source/Editor/BlueprintGraph/Classes/K2Node_BreakStruct.h
+++ b/Engine/Source/Editor/BlueprintGraph/Classes/K2Node_BreakStruct.h
@@ -10,8 +10,21 @@ class UK2Node_BreakStruct : public UK2Node_StructMemberGet
 {
 	GENERATED_UCLASS_BODY()
 
-	static bool CanBeBroken(const UScriptStruct* Struct);
-	static bool CanCreatePinForProperty(const UProperty* Property);
+	/** 
+	 * Returns false if:
+	 *   1. The Struct has a 'native break' method
+	 * Returns true if:
+	 *   1. The Struct is tagged as BlueprintType
+	 *   or
+	 *   2. The Struct has any property that is tagged as CPF_BlueprintVisible
+	 *   or
+	 *   3. The Struct has any property that is tagged as CPF_Edit and bIncludeEditAnywhere is true
+	 *
+	 * When constructing the context menu we do not allow the presence of the EditAnywhere flag to result in
+	 * a creation of a break node, although we could do so in the future. There are legacy break nodes that
+	 * rely on expansion of structs that neither have BlueprintVisible properties nor are tagged as BlueprintType
+	 */
+	static bool CanBeBroken(const UScriptStruct* Struct, bool bIncludeEditAnywhere = true);
 
 	// Begin UEdGraphNode interface
 	virtual void AllocateDefaultPins() override;
diff --git a/Engine/Source/Editor/BlueprintGraph/Classes/K2Node_MakeStruct.h b/Engine/Source/Editor/BlueprintGraph/Classes/K2Node_MakeStruct.h
--- a/Engine/Source/Editor/BlueprintGraph/Classes/K2Node_MakeStruct.h
+++ b/Engine/Source/Editor/BlueprintGraph/Classes/K2Node_MakeStruct.h
@@ -11,9 +11,22 @@ class UK2Node_MakeStruct : public UK2Node_StructMemberSet
 {
 	GENERATED_UCLASS_BODY()
 
-	BLUEPRINTGRAPH_API static bool CanBeMade(const UScriptStruct* Struct);
-	static bool CanBeExposed(const UProperty* Property);
-
+	/**
+	* Returns false if:
+	*   1. The Struct has a 'native make' method
+	* Returns true if:
+	*   1. The Struct is tagged as BlueprintType
+	*   or
+	*   2. The Struct has any property that is tagged as CPF_BlueprintVisible and is not CPF_BlueprintReadOnly
+	*   or
+	*   3. The Struct has any property that is tagged as CPF_Edit and is not CPF_BlueprintReadOnly and bIncludeEditAnywhere is true 
+	*
+	* When constructing the context menu we do not allow the presence of the EditAnywhere flag to result in
+	* a creation of a break node, although we could do so in the future. There are legacy break nodes that
+	* rely on expansion of structs that neither have BlueprintVisible properties nor are tagged as BlueprintType
+	*/
+	BLUEPRINTGRAPH_API static bool CanBeMade(const UScriptStruct* Struct, bool bIncludeEditAnywhere = true);
+	
 	// Begin UEdGraphNode interface
 	virtual void AllocateDefaultPins() override;
 	virtual FText GetNodeTitle(ENodeTitleType::Type TitleType) const override;
diff --git a/Engine/Source/Editor/BlueprintGraph/Private/K2Node_BreakStruct.cpp b/Engine/Source/Editor/BlueprintGraph/Private/K2Node_BreakStruct.cpp
--- a/Engine/Source/Editor/BlueprintGraph/Private/K2Node_BreakStruct.cpp
+++ b/Engine/Source/Editor/BlueprintGraph/Private/K2Node_BreakStruct.cpp
@@ -124,25 +124,29 @@ UK2Node_BreakStruct::UK2Node_BreakStruct(const FObjectInitializer& ObjectInitial
 {
 }
 
-bool UK2Node_BreakStruct::CanCreatePinForProperty(const UProperty* Property)
+static bool CanCreatePinForProperty(const UProperty* Property, bool bIncludeEditAnywhere = true)
 {
 	const UEdGraphSchema_K2* Schema = GetDefault<UEdGraphSchema_K2>();
 	FEdGraphPinType DumbGraphPinType;
 	const bool bConvertable = Schema->ConvertPropertyToPinType(Property, /*out*/ DumbGraphPinType);
 
-	//TODO: remove CPF_Edit in a next release.
-	const bool bVisible = Property ? Property->HasAnyPropertyFlags(CPF_Edit | CPF_BlueprintVisible) : false;
+	const bool bVisible = Property ? (bIncludeEditAnywhere ? Property->HasAnyPropertyFlags(CPF_Edit | CPF_BlueprintVisible) : Property->HasAnyPropertyFlags(CPF_BlueprintVisible) ): false;
 	return bVisible && bConvertable;
 }
 
-bool UK2Node_BreakStruct::CanBeBroken(const UScriptStruct* Struct)
+bool UK2Node_BreakStruct::CanBeBroken(const UScriptStruct* Struct, bool bIncludeEditAnywhere )
 {
 	const UEdGraphSchema_K2* Schema = GetDefault<UEdGraphSchema_K2>();
 	if(Struct && Schema && !Struct->HasMetaData(TEXT("HasNativeBreak")))
 	{
+		if (UEdGraphSchema_K2::IsAllowableBlueprintVariableType(Struct))
+		{
+			return true;
+		}
+
 		for (TFieldIterator<UProperty> It(Struct); It; ++It)
 		{
-			if (CanCreatePinForProperty(*It))
+			if (CanCreatePinForProperty(*It, bIncludeEditAnywhere))
 			{
 				return true;
 			}
@@ -162,7 +166,7 @@ void UK2Node_BreakStruct::AllocateDefaultPins()
 		{
 			virtual bool CanTreatPropertyAsOptional(UProperty* TestProperty) const override
 			{
-				return UK2Node_BreakStruct::CanCreatePinForProperty(TestProperty);
+				return CanCreatePinForProperty(TestProperty);
 			}
 		};
 
@@ -338,7 +342,7 @@ void UK2Node_BreakStruct::GetMenuActions(FBlueprintActionDatabaseRegistrar& Acti
 	for (TObjectIterator<UScriptStruct> StructIt; StructIt; ++StructIt)
 	{
 		UScriptStruct const* Struct = (*StructIt);
-		if (!UEdGraphSchema_K2::IsAllowableBlueprintVariableType(Struct) || !CanBeBroken(Struct))
+		if (!CanBeBroken(Struct, false))
 		{
 			continue;
 		}
diff --git a/Engine/Source/Editor/BlueprintGraph/Private/K2Node_MakeStruct.cpp b/Engine/Source/Editor/BlueprintGraph/Private/K2Node_MakeStruct.cpp
--- a/Engine/Source/Editor/BlueprintGraph/Private/K2Node_MakeStruct.cpp
+++ b/Engine/Source/Editor/BlueprintGraph/Private/K2Node_MakeStruct.cpp
@@ -78,10 +78,31 @@ void UK2Node_MakeStruct::FMakeStructPinManager::CustomizePinData(UEdGraphPin* Pi
 	}
 }
 
-bool UK2Node_MakeStruct::FMakeStructPinManager::CanTreatPropertyAsOptional(UProperty* TestProperty) const
+static bool CanBeExposed(const UProperty* Property, bool bIncludeEditOnly = true)
+{
+	if (Property)
 	{
-		return UK2Node_MakeStruct::CanBeExposed(TestProperty);
+		const UEdGraphSchema_K2* Schema = GetDefault<UEdGraphSchema_K2>();
+		check(Schema);
+
+		FEdGraphPinType DumbGraphPinType;
+		const bool bConvertable = Schema->ConvertPropertyToPinType(Property, /*out*/ DumbGraphPinType);
+
+		//TODO: remove CPF_Edit in a future release
+		const bool bVisible = (bIncludeEditOnly ? Property->HasAnyPropertyFlags(CPF_BlueprintVisible | CPF_Edit) : Property->HasAnyPropertyFlags(CPF_BlueprintVisible)) && !(Property->ArrayDim > 1);
+		const bool bBlueprintReadOnly = Property->HasAllPropertyFlags(CPF_BlueprintReadOnly);
+		if (bVisible && bConvertable && !bBlueprintReadOnly)
+		{
+			return true;
+		}
 	}
+	return false;
+}
+
+bool UK2Node_MakeStruct::FMakeStructPinManager::CanTreatPropertyAsOptional(UProperty* TestProperty) const
+{
+	return CanBeExposed(TestProperty);
+}
 
 UK2Node_MakeStruct::UK2Node_MakeStruct(const FObjectInitializer& ObjectInitializer)
 	: Super(ObjectInitializer)
@@ -205,34 +226,18 @@ FLinearColor UK2Node_MakeStruct::GetNodeTitleColor() const
 	return UK2Node::GetNodeTitleColor();
 }
 
-bool UK2Node_MakeStruct::CanBeExposed(const UProperty* Property)
+bool UK2Node_MakeStruct::CanBeMade(const UScriptStruct* Struct, bool bIncludeEditOnly )
 {
-	if(Property)
+	if (Struct && !Struct->HasMetaData(TEXT("HasNativeMake")))
 	{
-		const UEdGraphSchema_K2* Schema = GetDefault<UEdGraphSchema_K2>();
-		check(Schema);
-
-		FEdGraphPinType DumbGraphPinType;
-		const bool bConvertable = Schema->ConvertPropertyToPinType(Property, /*out*/ DumbGraphPinType);
-
-		//TODO: remove CPF_Edit in a future release
-		const bool bVisible = Property->HasAnyPropertyFlags(CPF_BlueprintVisible|CPF_Edit) && !(Property->ArrayDim > 1);
-		const bool bBlueprintReadOnly = Property->HasAllPropertyFlags(CPF_BlueprintReadOnly);
-		if(bVisible && bConvertable && !bBlueprintReadOnly)
+		if (UEdGraphSchema_K2::IsAllowableBlueprintVariableType(Struct))
 		{
 			return true;
 		}
-	}
-	return false;
-}
 
-bool UK2Node_MakeStruct::CanBeMade(const UScriptStruct* Struct)
-{
-	if (Struct && !Struct->HasMetaData(TEXT("HasNativeMake")))
-	{
 		for (TFieldIterator<UProperty> It(Struct); It; ++It)
 		{
-			if(CanBeExposed(*It))
+			if (CanBeExposed(*It, bIncludeEditOnly))
 			{
 				return true;
 			}
@@ -303,7 +308,7 @@ void UK2Node_MakeStruct::GetMenuActions(FBlueprintActionDatabaseRegistrar& Actio
 	for (TObjectIterator<UScriptStruct> StructIt; StructIt; ++StructIt)
 	{
 		UScriptStruct const* Struct = (*StructIt);
-		if (!UEdGraphSchema_K2::IsAllowableBlueprintVariableType(Struct) || !CanBeMade(Struct))
+		if (!CanBeMade(Struct, false))
 		{
 			continue;
 		}
