Fixed an incomplete sentence for a tooltip in the DynamicCastNode.

diff --git a/Engine/Source/Editor/BlueprintGraph/Private/K2Node_DynamicCast.cpp b/Engine/Source/Editor/BlueprintGraph/Private/K2Node_DynamicCast.cpp
--- a/Engine/Source/Editor/BlueprintGraph/Private/K2Node_DynamicCast.cpp
+++ b/Engine/Source/Editor/BlueprintGraph/Private/K2Node_DynamicCast.cpp
@@ -110,7 +110,7 @@ void UK2Node_DynamicCast::GetContextMenuActions(const FGraphNodeContextMenuBuild
 	Context.MenuBuilder->BeginSection("K2NodeDynamicCast", LOCTEXT("DynamicCastHeader", "Cast"));
 	{
 		FText MenuEntryTitle   = LOCTEXT("MakePureTitle",   "Convert to pure cast");
-		FText MenuEntryTooltip = LOCTEXT("MakePureTooltip", "Removes the execution pins to make the node more versitile (NOTE: the cast could still, resulting in an invalid output).");
+		FText MenuEntryTooltip = LOCTEXT("MakePureTooltip", "Removes the execution pins to make the node more versitile (NOTE: the cast could still fail, resulting in an invalid output).");
 
 		bool bCanTogglePurity = true;
 		auto CanExecutePurityToggle = [](bool const bInCanTogglePurity)->bool
