Variable Get nodes can be toggled between pure and impure which will branch logic based on value validity, if they are of type Object and a non-array.

diff --git a/Engine/Source/Editor/BlueprintGraph/Classes/K2Node_VariableGet.h b/Engine/Source/Editor/BlueprintGraph/Classes/K2Node_VariableGet.h
--- a/Engine/Source/Editor/BlueprintGraph/Classes/K2Node_VariableGet.h
+++ b/Engine/Source/Editor/BlueprintGraph/Classes/K2Node_VariableGet.h
@@ -15,21 +15,44 @@ class BLUEPRINTGRAPH_API UK2Node_VariableGet : public UK2Node_Variable
 	virtual void AllocateDefaultPins() override;
 	virtual FText GetTooltipText() const override;
 	virtual FText GetNodeTitle(ENodeTitleType::Type TitleType) const override;
+	virtual void GetContextMenuActions(const FGraphNodeContextMenuBuilder& Context) const override;
 	// End UEdGraphNode interface
 
 	// Begin K2Node interface
-	virtual bool IsNodePure() const override { return true; }
+	virtual bool IsNodePure() const override { return bIsPureGet; }
 	virtual void ReallocatePinsDuringReconstruction(TArray<UEdGraphPin*>& OldPins) override;
 	virtual bool ShouldShowNodeProperties() const override { return true; }
 	virtual class FNodeHandlingFunctor* CreateNodeHandler(class FKismetCompilerContext& CompilerContext) const override;
+	virtual void ExpandNode(class FKismetCompilerContext& CompilerContext, UEdGraph* SourceGraph) override;
 	// End K2Node interface
 
 	static FText GetPropertyTooltip(UProperty const* VariableProperty);
 	static FText GetBlueprintVarTooltip(FBPVariableDescription const& VarDesc);
 
+	/**
+	 * Will change the node's purity, and reallocate pins accordingly (adding/
+	 * removing exec pins).
+	 * 
+	 * @param  bNewPurity  The new value for bIsPureCast.
+	 */
+	void SetPurity(bool bNewPurity);
+
 private:
+	/** Checks if the pin type is valid to be a non-pure node */
+	static bool IsValidTypeForNonPure(const FEdGraphPinType& InPinType);
+
+	/** Adds pins required for the node to function in a non-pure manner */
+	void CreateNonPurePins(TArray<UEdGraphPin*>* InOldPinsPtr);
+
+	/** Flips the node's purity (adding/removing exec pins as needed). */
+	void TogglePurity();
+
 	/** Constructing FText strings can be costly, so we cache the node's title/tooltip */
 	FNodeTextCache CachedTooltip;
 	FNodeTextCache CachedNodeTitle;
+
+	/** TRUE if the node should function as a pure node, without exec pins */
+	UPROPERTY()
+	bool bIsPureGet;
 };
 
diff --git a/Engine/Source/Editor/BlueprintGraph/Private/K2Node_VariableGet.cpp b/Engine/Source/Editor/BlueprintGraph/Private/K2Node_VariableGet.cpp
--- a/Engine/Source/Editor/BlueprintGraph/Private/K2Node_VariableGet.cpp
+++ b/Engine/Source/Editor/BlueprintGraph/Private/K2Node_VariableGet.cpp
@@ -3,6 +3,7 @@
 
 #include "BlueprintGraphPrivatePCH.h"
 #include "KismetCompiler.h"
+#include "ScopedTransaction.h"
 
 //////////////////////////////////////////////////////////////////////////
 // FKCHandler_VariableGet
@@ -74,13 +75,74 @@ static FText K2Node_VariableGetImpl::GetBaseTooltip(FName VarName)
 
 UK2Node_VariableGet::UK2Node_VariableGet(const FObjectInitializer& ObjectInitializer)
 	: Super(ObjectInitializer)
+	, bIsPureGet(true)
 {
 }
 
+void UK2Node_VariableGet::CreateNonPurePins(TArray<UEdGraphPin*>* InOldPinsPtr)
+{
+	const UEdGraphSchema_K2* K2Schema = Cast<UEdGraphSchema_K2>(GetSchema());
+	check(K2Schema != nullptr);
+	if (!K2Schema->DoesGraphSupportImpureFunctions(GetGraph()))
+	{
+		bIsPureGet = true;
+	}
+
+	if (!bIsPureGet)
+	{
+		FEdGraphPinType PinType;
+		UProperty* VariableProperty = GetPropertyForVariable();
+
+		// We need the pin's type, to both see if it's an array and if it is of the correct types to remain an impure node
+		if (VariableProperty)
+		{
+			K2Schema->ConvertPropertyToPinType(GetPropertyForVariable(), PinType);
+		}
+		// If there is no property and we are given some old pins to look at, find the old value pin and use the type there
+		// This allows nodes to be pasted into other BPs without access to the property
+		else if(InOldPinsPtr)
+		{
+			// find old variable pin and use the type.
+			const FString PinName = GetVarNameString();
+			for(auto Iter = InOldPinsPtr->CreateConstIterator(); Iter; ++Iter)
+			{
+				if(const UEdGraphPin* Pin = *Iter)
+				{
+					if(PinName == Pin->PinName)
+					{
+						PinType = Pin->PinType;
+						break;
+					}
+				}
+			}
+
+		}
+
+		if (IsValidTypeForNonPure(PinType))
+		{
+			// Input - Execution Pin
+			CreatePin(EGPD_Input, K2Schema->PC_Exec, TEXT(""), NULL, false, false, K2Schema->PN_Execute);
+
+			// Output - Execution Pins
+			UEdGraphPin* ValidPin = CreatePin(EGPD_Output, K2Schema->PC_Exec, TEXT(""), NULL, false, false, K2Schema->PN_Then);
+			ValidPin->PinFriendlyName = LOCTEXT("Valid", "Is Valid");
+
+			UEdGraphPin* InvalidPin = CreatePin(EGPD_Output, K2Schema->PC_Exec, TEXT(""), NULL, false, false, K2Schema->PN_Else);
+			InvalidPin->PinFriendlyName = LOCTEXT("Invalid", "Is Not Valid");
+		}
+		else
+		{
+			bIsPureGet = true;
+		}
+	}
+}
+
 void UK2Node_VariableGet::AllocateDefaultPins()
 {
 	if(GetVarName() != NAME_None)
 	{
+		CreateNonPurePins(nullptr);
+
 		if(CreatePinForVariable(EGPD_Output))
 		{
 			CreatePinForSelf();
@@ -94,6 +156,8 @@ void UK2Node_VariableGet::ReallocatePinsDuringReconstruction(TArray<UEdGraphPin*
 {
 	if(GetVarName() != NAME_None)
 	{
+		CreateNonPurePins(&OldPins);
+
 		if(!CreatePinForVariable(EGPD_Output))
 		{
 			if(!RecreatePinForVariable(EGPD_Output, OldPins))
@@ -224,4 +288,158 @@ FNodeHandlingFunctor* UK2Node_VariableGet::CreateNodeHandler(FKismetCompilerCont
 	return new FKCHandler_VariableGet(CompilerContext);
 }
 
+bool UK2Node_VariableGet::IsValidTypeForNonPure(const FEdGraphPinType& InPinType)
+{
+	return InPinType.bIsArray == false && (InPinType.PinCategory == UObject::StaticClass()->GetName() ||InPinType.PinCategory == UClass::StaticClass()->GetName());
+}
+
+void UK2Node_VariableGet::GetContextMenuActions(const FGraphNodeContextMenuBuilder& Context) const
+{
+	Super::GetContextMenuActions(Context);
+
+	const UEdGraphPin* ValuePin = GetValuePin();
+	if (IsValidTypeForNonPure(ValuePin->PinType))
+	{
+		Context.MenuBuilder->BeginSection("K2NodeVariableGet", LOCTEXT("VariableGetHeader", "Variable Get"));
+		{
+			FText MenuEntryTitle;
+			FText MenuEntryTooltip;
+
+			bool bCanTogglePurity = true;
+			auto CanExecutePurityToggle = [](bool const bInCanTogglePurity)->bool
+			{
+				return bInCanTogglePurity;
+			};
+
+			if (bIsPureGet)
+			{
+				MenuEntryTitle   = LOCTEXT("ConvertToImpureGetTitle",   "Convert to impure Get");
+				MenuEntryTooltip = LOCTEXT("ConvertToImpureGetTooltip", "Removes the execution pins to make the node more versitile.");
+
+				const UEdGraphSchema_K2* K2Schema = Cast<UEdGraphSchema_K2>(GetSchema());
+				check(K2Schema != nullptr);
+
+				bCanTogglePurity = K2Schema->DoesGraphSupportImpureFunctions(GetGraph());
+				if (!bCanTogglePurity)
+				{
+					MenuEntryTooltip = LOCTEXT("CannotMakeImpureGetTooltip", "This graph does not support impure calls!");
+				}
+			}
+			else
+			{
+				MenuEntryTitle   = LOCTEXT("ConvertToPureGetTitle",   "Convert to pure Get");
+				MenuEntryTooltip = LOCTEXT("ConvertToPureGetTooltip", "Adds in branching execution pins so that you can separatly handle when the returned value is valid/invalid.");
+			}
+
+			Context.MenuBuilder->AddMenuEntry(
+				MenuEntryTitle,
+				MenuEntryTooltip,
+				FSlateIcon(),
+				FUIAction(
+				FExecuteAction::CreateUObject(this, &UK2Node_VariableGet::TogglePurity),
+				FCanExecuteAction::CreateStatic(CanExecutePurityToggle, bCanTogglePurity),
+				FIsActionChecked()
+				)
+				);
+		}
+		Context.MenuBuilder->EndSection();
+	}
+}
+
+void UK2Node_VariableGet::TogglePurity()
+{
+	FText TransactionTitle;
+	if(bIsPureGet)
+	{
+		TransactionTitle = LOCTEXT("TogglePureGet", "Convert to Pure Get");
+	}
+	else
+	{
+		TransactionTitle = LOCTEXT("ToggleImpureGet", "Convert to Impure Get");
+	}
+	const FScopedTransaction Transaction( TransactionTitle );
+	Modify();
+
+	SetPurity(!bIsPureGet);
+}
+
+void UK2Node_VariableGet::SetPurity(bool bNewPurity)
+{
+	if (bNewPurity != bIsPureGet)
+	{
+		bIsPureGet = bNewPurity;
+
+		bool const bHasBeenConstructed = (Pins.Num() > 0);
+		if (bHasBeenConstructed)
+		{
+			ReconstructNode();
+		}
+	}
+}
+
+void UK2Node_VariableGet::ExpandNode(class FKismetCompilerContext& CompilerContext, UEdGraph* SourceGraph)
+{
+	Super::ExpandNode(CompilerContext, SourceGraph);
+
+	if (!bIsPureGet)
+	{
+		const UEdGraphSchema_K2* Schema = CompilerContext.GetSchema();
+		UEdGraphPin* ValuePin = GetValuePin();
+
+		// Impure Get nodes convert into three nodes:
+		// 1. A pure Get node
+		// 2. An IsValid node
+		// 3. A Branch node (only impure part
+		
+		// Create the impure Get node
+		UK2Node_VariableGet* VariableGetNode = CompilerContext.SpawnIntermediateNode<UK2Node_VariableGet>(this, SourceGraph);
+		VariableGetNode->VariableReference = VariableReference;
+		VariableGetNode->AllocateDefaultPins();
+		CompilerContext.MessageLog.NotifyIntermediateObjectCreation(VariableGetNode, this);
+
+		// Move pin links from Get node we are expanding, to the new pure one we've created
+		CompilerContext.MovePinLinksToIntermediate(*ValuePin, *VariableGetNode->GetValuePin());
+		CompilerContext.MovePinLinksToIntermediate(*FindPin(Schema->PN_Self), *VariableGetNode->FindPin(Schema->PN_Self));
+
+		// Create the IsValid node
+		UK2Node_CallFunction* IsValidFunction = CompilerContext.SpawnIntermediateNode<UK2Node_CallFunction>(this, SourceGraph);
+
+		// Based on if the type is an "Object" or a "Class" changes which function to use
+		if (ValuePin->PinType.PinCategory == UObject::StaticClass()->GetName())
+		{
+			IsValidFunction->SetFromFunction(UKismetSystemLibrary::StaticClass()->FindFunctionByName("IsValid"));
+		}
+		else if (ValuePin->PinType.PinCategory == UClass::StaticClass()->GetName())
+		{
+			IsValidFunction->SetFromFunction(UKismetSystemLibrary::StaticClass()->FindFunctionByName("IsValidClass"));
+		}
+		IsValidFunction->AllocateDefaultPins();
+		CompilerContext.MessageLog.NotifyIntermediateObjectCreation(IsValidFunction, this);
+
+		// Connect the value pin from the new Get node to the IsValid
+		UEdGraphPin* ObjectPin = IsValidFunction->Pins[1];
+		check(ObjectPin->Direction == EGPD_Input);
+		ObjectPin->MakeLinkTo(VariableGetNode->GetValuePin());
+
+		// Create the Branch node
+		UK2Node_IfThenElse* BranchNode = CompilerContext.SpawnIntermediateNode<UK2Node_IfThenElse>(this, SourceGraph);
+		BranchNode->AllocateDefaultPins();
+		CompilerContext.MessageLog.NotifyIntermediateObjectCreation(BranchNode, this);
+
+		// Connect the bool output pin from IsValid node to the Branch node
+		UEdGraphPin* BoolPin = IsValidFunction->Pins[2];
+		check(BoolPin->Direction == EGPD_Output);
+		BoolPin->MakeLinkTo(BranchNode->GetConditionPin());
+
+		// Connect the Branch node to the input of the impure Get node
+		CompilerContext.MovePinLinksToIntermediate(*GetExecPin(), *BranchNode->GetExecPin());
+
+		// Move the two Branch pins to the Branch node
+		CompilerContext.MovePinLinksToIntermediate(*FindPin(Schema->PN_Then), *BranchNode->FindPin(Schema->PN_Then));
+		CompilerContext.MovePinLinksToIntermediate(*FindPin(Schema->PN_Else), *BranchNode->FindPin(Schema->PN_Else));
+
+		BreakAllNodeLinks();
+	}
+}
+
 #undef LOCTEXT_NAMESPACE
\ No newline at end of file
diff --git a/Engine/Source/Editor/GraphEditor/Private/KismetNodes/SGraphNodeK2Var.cpp b/Engine/Source/Editor/GraphEditor/Private/KismetNodes/SGraphNodeK2Var.cpp
--- a/Engine/Source/Editor/GraphEditor/Private/KismetNodes/SGraphNodeK2Var.cpp
+++ b/Engine/Source/Editor/GraphEditor/Private/KismetNodes/SGraphNodeK2Var.cpp
@@ -38,6 +38,7 @@ void SGraphNodeK2Var::UpdateGraphNode()
 	float HorizontalTitleMargin = 0.0f;
 	float VerticalTitleMargin = 8.0f;
 	FMargin TitleMargin = FMargin(0.0f, 8.0f);
+	FMargin ContentAreaMargin = FMargin(0.0f, 4.0f);
 	EHorizontalAlignment TitleHAlign = HAlign_Center;
 	TSharedPtr<SWidget> TitleWidget;
 
@@ -47,7 +48,6 @@ void SGraphNodeK2Var::UpdateGraphNode()
 		if(SetNode->HasLocalRepNotify())
 		{
 			TitleText = NSLOCTEXT("GraphEditor", "VariableSetWithNotify", "SET w/ Notify");
-			//+ FString::Printf(TEXT("\n(%s)"), *SetNode->GetRepNotifyName().ToString());
 		}
 		else
 		{
@@ -146,7 +146,16 @@ void SGraphNodeK2Var::UpdateGraphNode()
 				]
 			];
 	}
-	else
+	else if (UK2Node_VariableGet* VariableGet = Cast<UK2Node_VariableGet>(GraphNode))
+	{
+		if (!VariableGet->IsNodePure())
+		{
+			TitleText = NSLOCTEXT("GraphEditor", "VariableGet", "GET");
+			ContentAreaMargin.Top += 16;
+		}
+	}
+	
+	if (TitleText.IsEmpty())
 	{
 		TitleWidget = SNullWidget::NullWidget;
 	}
@@ -214,7 +223,7 @@ void SGraphNodeK2Var::UpdateGraphNode()
 				TitleWidget.ToSharedRef()
 			]
 			+ SOverlay::Slot()
-			.Padding( FMargin(0,4) )
+			.Padding( ContentAreaMargin )
 			[
 				// NODE CONTENT AREA
 				SNew(SHorizontalBox)
