Fix reverse logic on impure/pure variable Get node transaction string.

diff --git a/Engine/Source/Editor/BlueprintGraph/Private/K2Node_VariableGet.cpp b/Engine/Source/Editor/BlueprintGraph/Private/K2Node_VariableGet.cpp
--- a/Engine/Source/Editor/BlueprintGraph/Private/K2Node_VariableGet.cpp
+++ b/Engine/Source/Editor/BlueprintGraph/Private/K2Node_VariableGet.cpp
@@ -349,7 +349,7 @@ void UK2Node_VariableGet::GetContextMenuActions(const FGraphNodeContextMenuBuild
 void UK2Node_VariableGet::TogglePurity()
 {
 	FText TransactionTitle;
-	if(bIsPureGet)
+	if(!bIsPureGet)
 	{
 		TransactionTitle = LOCTEXT("TogglePureGet", "Convert to Pure Get");
 	}
