[UE-11537] Node had a bunch of mutable state that it was stomping when it was refreshed... Removed the unneeded state.

diff --git a/Engine/Source/Editor/BlueprintGraph/Classes/K2Node_EaseFunction.h b/Engine/Source/Editor/BlueprintGraph/Classes/K2Node_EaseFunction.h
--- a/Engine/Source/Editor/BlueprintGraph/Classes/K2Node_EaseFunction.h
+++ b/Engine/Source/Editor/BlueprintGraph/Classes/K2Node_EaseFunction.h
@@ -71,13 +71,7 @@ private:
 	void ChangePinType(UEdGraphPin* Pin);
 
 	/** Generates the available custom pins for the current ease function if applicable */
-	bool ConditionalGenerateCustomPins();
-
-	/** Try to remove all custom pins that exist */
-	bool RemoveCustomPins(const bool bRemoveValuePins);
-
-	/** Generate custom pins used by the current ease function */
-	bool GenerateCustomPins(const int32 NewEasingFunc);
+	void RefreshPinVisibility();
 
 	/** Update MyPin with data from another pin. Return FALSE in case nothing was changed, TRUE otherwise. */
 	bool UpdatePin(UEdGraphPin* MyPin, UEdGraphPin* OtherPin);
@@ -94,14 +88,6 @@ private:
 	/** Called right after the wildcard pins has been assigned a value */
 	void GenerateExtraPins();
 
-	/** Custom pins that are curently used */
-	UPROPERTY()
-	TArray<FCustomPin> CustomPinNames;
-
-	/** Previous easing function, used to only recreate custom pins if required */
-	UPROPERTY()
-	int32 OldEasingFunc;
-
 	/** Name of the kismet ease function to be called */
 	UPROPERTY()
 	FString EaseFunctionName;
diff --git a/Engine/Source/Editor/BlueprintGraph/Private/K2Node_EaseFunction.cpp b/Engine/Source/Editor/BlueprintGraph/Private/K2Node_EaseFunction.cpp
--- a/Engine/Source/Editor/BlueprintGraph/Private/K2Node_EaseFunction.cpp
+++ b/Engine/Source/Editor/BlueprintGraph/Private/K2Node_EaseFunction.cpp
@@ -73,7 +73,6 @@ UK2Node_EaseFunction::UK2Node_EaseFunction(const FObjectInitializer& ObjectIniti
 	, CachedEaseFuncPin(NULL)
 {
 	NodeTooltip = LOCTEXT("NodeTooltip", "Interpolates from value A to value B using a user specified easing function");
-	OldEasingFunc = INDEX_NONE;
 	EaseFunctionName = TEXT("");
 }
 
@@ -92,21 +91,32 @@ void UK2Node_EaseFunction::AllocateDefaultPins()
 	// Make sure that the default value is set correctly if none has been set
 	K2Schema->SetPinDefaultValue(CachedEaseFuncPin);
 
-	UEdGraphPin *AlphaPin = CreatePin(EGPD_Input, K2Schema->PC_Float, TEXT(""), NULL, false, false, FEaseFunctionNodeHelper::GetAlphaPinName());
+	UEdGraphPin* AlphaPin = CreatePin(EGPD_Input, K2Schema->PC_Float, TEXT(""), NULL, false, false, FEaseFunctionNodeHelper::GetAlphaPinName());
 	SetPinToolTip(*AlphaPin, LOCTEXT("AlphaPinTooltip", "Alpha value used to specify the easing in time."));
 
 	// Add wildcard pins for A, B and the return Pin
-	UEdGraphPin *APin = CreatePin(EGPD_Input, K2Schema->PC_Wildcard, TEXT(""), NULL, false, false, FEaseFunctionNodeHelper::GetAPinName());
+	UEdGraphPin* APin = CreatePin(EGPD_Input, K2Schema->PC_Wildcard, TEXT(""), NULL, false, false, FEaseFunctionNodeHelper::GetAPinName());
 	SetPinToolTip(*APin, LOCTEXT("APinDescription", "Easing start value"));
 
-	UEdGraphPin *BPin = CreatePin(EGPD_Input, K2Schema->PC_Wildcard, TEXT(""), NULL, false, false, FEaseFunctionNodeHelper::GetBPinName());
+	UEdGraphPin* BPin = CreatePin(EGPD_Input, K2Schema->PC_Wildcard, TEXT(""), NULL, false, false, FEaseFunctionNodeHelper::GetBPinName());
 	SetPinToolTip(*BPin, LOCTEXT("BPinDescription", "Easing end value"));
 
-	UEdGraphPin *ResultPin = CreatePin(EGPD_Output, K2Schema->PC_Wildcard, TEXT(""), NULL, false, false, FEaseFunctionNodeHelper::GetResultPinName());
+	UEdGraphPin* ResultPin = CreatePin(EGPD_Output, K2Schema->PC_Wildcard, TEXT(""), NULL, false, false, FEaseFunctionNodeHelper::GetResultPinName());
 	SetPinToolTip(*ResultPin, LOCTEXT("ResultPinDescription", "Easing result value"));
 
-	// Try to generate any custom pins that applies
-	ConditionalGenerateCustomPins();
+	UEdGraphPin* ShortestPathPin = CreatePin(EGPD_Input, K2Schema->PC_Boolean, TEXT(""), NULL, false, false, FEaseFunctionNodeHelper::GetShortestPathPinName());
+	SetPinToolTip(*ShortestPathPin, LOCTEXT("ShortestPathPinTooltip", "When interpolating the shortest path should be taken."));
+	ShortestPathPin->DefaultValue = ShortestPathPin->AutogeneratedDefaultValue = TEXT("true");
+
+	UEdGraphPin* BlendExpPin = CreatePin(EGPD_Input, K2Schema->PC_Float, TEXT(""), NULL, false, false, FEaseFunctionNodeHelper::GetBlendExpPinName());
+	SetPinToolTip(*BlendExpPin, LOCTEXT("BlendExpPinDescription", "Blend Exponent for basic ease functions"));
+	BlendExpPin->DefaultValue = BlendExpPin->AutogeneratedDefaultValue = TEXT("2.0");
+
+	UEdGraphPin* StepsPin = CreatePin(EGPD_Input, K2Schema->PC_Int, TEXT(""), NULL, false, false, FEaseFunctionNodeHelper::GetStepsPinName());
+	SetPinToolTip(*StepsPin, LOCTEXT("StepsPinDescription", "Number of steps required to go from A to B"));
+	StepsPin->DefaultValue = StepsPin->AutogeneratedDefaultValue = TEXT("2");
+
+	RefreshPinVisibility();
 }
 
 FText UK2Node_EaseFunction::GetNodeTitle(ENodeTitleType::Type TitleType) const
@@ -147,23 +157,8 @@ void UK2Node_EaseFunction::NotifyPinConnectionListChanged(UEdGraphPin* Pin)
 	const auto EaseFuncPin = GetEaseFuncPin();
 	if (Pin == EaseFuncPin)
 	{
-		bool bDirty = false;
-
-		// Check in which state we are at the moment
-		if (CanCustomizeCurve())
-		{
-			OldEasingFunc = INDEX_NONE;
-			bDirty = ConditionalGenerateCustomPins();
-		}
-		else
-		{
-			bDirty = RemoveCustomPins(false);
-		}
-
-		if (bDirty)
-		{
-			GetGraph()->NotifyGraphChanged();
-		}
+		RefreshPinVisibility();
+		GetGraph()->NotifyGraphChanged();
 	}
 	else
 	{
@@ -175,8 +170,9 @@ void UK2Node_EaseFunction::NotifyPinConnectionListChanged(UEdGraphPin* Pin)
 void UK2Node_EaseFunction::PinDefaultValueChanged(UEdGraphPin* Pin)
 {
 	const auto EaseFuncPin = GetEaseFuncPin();
-	if (Pin == EaseFuncPin && ConditionalGenerateCustomPins())
+	if (Pin == EaseFuncPin )
 	{
+		RefreshPinVisibility();
 		GetGraph()->NotifyGraphChanged();
 	}
 }
@@ -185,14 +181,8 @@ void UK2Node_EaseFunction::PostReconstructNode()
 {
 	Super::PostReconstructNode();
 
-	CustomPinNames.Empty();
 	// Check in which state we are at the moment
-	if (CanCustomizeCurve())
-	{
-		OldEasingFunc = INDEX_NONE;
-		ConditionalGenerateCustomPins();
-	}
-	GenerateExtraPins();
+	RefreshPinVisibility();
 
 	// Find a pin that has connections to use to jumpstart the wildcard process
 	for (int32 PinIndex = 0; PinIndex < Pins.Num(); ++PinIndex)
@@ -210,6 +200,8 @@ void UK2Node_EaseFunction::PostReconstructNode()
 			}
 		}
 	}
+
+	GenerateExtraPins();
 }
 
 bool UK2Node_EaseFunction::IsConnectionDisallowed(const UEdGraphPin* MyPin, const UEdGraphPin* OtherPin, FString& OutReason) const
@@ -389,8 +381,8 @@ void UK2Node_EaseFunction::ExpandNode(class FKismetCompilerContext& CompilerCont
 		the other nodes into the Blueprint.
 	*/
 
-	UFunction* function = UKismetMathLibrary::StaticClass()->FindFunctionByName(*EaseFunctionName);
-	if (function == NULL)
+	UFunction* Function = UKismetMathLibrary::StaticClass()->FindFunctionByName(*EaseFunctionName);
+	if (Function == NULL)
 	{
 		CompilerContext.MessageLog.Error(*LOCTEXT("InvalidFunctionName", "BaseAsyncTask: Type not supported or not initialized. @@").ToString(), this);
 		return;
@@ -402,7 +394,7 @@ void UK2Node_EaseFunction::ExpandNode(class FKismetCompilerContext& CompilerCont
 	// the name of the desired function
 	UK2Node_CallFunction* CallFunction = CompilerContext.SpawnIntermediateNode<UK2Node_CallFunction>(this, SourceGraph);
 		
-	CallFunction->SetFromFunction(function);
+	CallFunction->SetFromFunction(Function);
 	CallFunction->AllocateDefaultPins();
 	CompilerContext.MessageLog.NotifyIntermediateObjectCreation(CallFunction, this);
 
@@ -416,9 +408,22 @@ void UK2Node_EaseFunction::ExpandNode(class FKismetCompilerContext& CompilerCont
 	CompilerContext.MovePinLinksToIntermediate(*FindPin(FEaseFunctionNodeHelper::GetResultPinName()), *CallFunction->GetReturnValuePin());
 
 	// Now move the custom pins to their new locations
-	for (int32 ArgIdx = 0; ArgIdx < CustomPinNames.Num(); ++ArgIdx)
+	UEdGraphPin* ShortestPathPin = FindPinChecked(FEaseFunctionNodeHelper::GetShortestPathPinName());
+	if (!ShortestPathPin->bHidden)
 	{
-		CompilerContext.MovePinLinksToIntermediate(*FindPin(CustomPinNames[ArgIdx].PinName), *CallFunction->FindPin(CustomPinNames[ArgIdx].CallFuncPinName));
+		CompilerContext.MovePinLinksToIntermediate(*ShortestPathPin, *CallFunction->FindPinChecked(TEXT("bShortestPath")));
+	}
+
+	UEdGraphPin* BlendExpPin = FindPinChecked(FEaseFunctionNodeHelper::GetBlendExpPinName());
+	if (!BlendExpPin->bHidden)
+	{
+		CompilerContext.MovePinLinksToIntermediate(*BlendExpPin, *CallFunction->FindPinChecked(FEaseFunctionNodeHelper::GetBlendExpPinName()));
+	}
+
+	UEdGraphPin* StepsPin = FindPinChecked(FEaseFunctionNodeHelper::GetStepsPinName());
+	if (!StepsPin->bHidden)
+	{
+		CompilerContext.MovePinLinksToIntermediate(*StepsPin, *CallFunction->FindPinChecked(FEaseFunctionNodeHelper::GetStepsPinName()));
 	}
 
 	// Cleanup links to ourselfs and we are done!
@@ -434,141 +439,63 @@ UEdGraphPin* UK2Node_EaseFunction::GetEaseFuncPin() const
 	return CachedEaseFuncPin;
 }
 
-bool UK2Node_EaseFunction::ConditionalGenerateCustomPins()
+void UK2Node_EaseFunction::RefreshPinVisibility()
 {
-	// Only generate custom values if we are able to use them
-	if (CanCustomizeCurve())
+	const auto EaseFuncPin = GetEaseFuncPin();
+	UEnum * Enum = FindObject<UEnum>(ANY_PACKAGE, TEXT("EEasingFunc"), true);
+	check(Enum != NULL);
+	const int32 NewEasingFunc = CanCustomizeCurve() ? Enum->FindEnumIndex(*EaseFuncPin->DefaultValue) : -1;
+	
+	// Early exit in case no changes are required
+	const UEdGraphSchema_K2* K2Schema = Cast<const UEdGraphSchema_K2>(GetSchema());
+
+	UEdGraphPin* BlendExpPin = FindPinChecked(FEaseFunctionNodeHelper::GetBlendExpPinName());
+
+	if (NewEasingFunc == EEasingFunc::EaseIn ||
+		NewEasingFunc == EEasingFunc::EaseOut ||
+		NewEasingFunc == EEasingFunc::EaseInOut)
 	{
-		const auto EaseFuncPin = GetEaseFuncPin();
-		if (EaseFuncPin)
-		{
-			UEnum * Enum = FindObject<UEnum>(ANY_PACKAGE, TEXT("EEasingFunc"), true);
-			check(Enum != NULL);
-			const int32 EnumIndex = Enum->FindEnumIndex(*EaseFuncPin->DefaultValue);
-			return GenerateCustomPins(EnumIndex);
-		}
+		// Show the BlendExpPin
+		BlendExpPin->bHidden = false;
 	}
-	return false;
-}
-
-bool UK2Node_EaseFunction::RemoveCustomPins(const bool bRemoveValuePins)
-{
-	bool bResult = false;
-	// Remove old pins
-	for (int32 ArgIdx = CustomPinNames.Num()-1; ArgIdx >= 0; --ArgIdx)
+	else
 	{
-		UEdGraphPin* CustomPin = FindPin(CustomPinNames[ArgIdx].PinName);
-
-		if (CustomPinNames[ArgIdx].bValuePin == bRemoveValuePins)
-		{
-			// Make sure the pin exists
-			if (CustomPin)
-			{
-				CustomPin->BreakAllPinLinks();
-				Pins.Remove(CustomPin);
-				CustomPinNames.RemoveAt(ArgIdx, 1);
-				bResult = true;
+		// Hide the BlendExpPin:
+		BlendExpPin->BreakAllPinLinks();
+		BlendExpPin->bHidden = true;
+	}
 
-			}
-			else
-			{
-				// When loaded from the disk this function is called before pins have been created:
-				CustomPinNames.RemoveAt(ArgIdx, 1);
-			}
-		}
+	UEdGraphPin* StepsPin = FindPinChecked(FEaseFunctionNodeHelper::GetStepsPinName());
+	if (NewEasingFunc == EEasingFunc::Step)
+	{
+		// Show the Steps pin:
+		StepsPin->bHidden = false;
+	}
+	else
+	{
+		// Hide the Steps pin:
+		StepsPin->BreakAllPinLinks();
+		StepsPin->bHidden = true;
 	}
-	return bResult;
 }
 
 void UK2Node_EaseFunction::GenerateExtraPins()
 {
-	bool bResult = false;
 	const UEdGraphSchema_K2* K2Schema = Cast<const UEdGraphSchema_K2>(GetSchema());
 
-	// First try to remove old pins
-	RemoveCustomPins(true);
-
 	// Add pins based on the pin type
-	const UEdGraphPin* APin = FindPin(FEaseFunctionNodeHelper::GetAPinName());
+	const UEdGraphPin* APin = FindPinChecked(FEaseFunctionNodeHelper::GetAPinName());
 	const bool bIsRotator = APin->PinType.PinCategory == K2Schema->PC_Struct && APin->PinType.PinSubCategoryObject.Get()->GetName() == TEXT("Rotator");
+	UEdGraphPin* ShortestPathPin = FindPinChecked(FEaseFunctionNodeHelper::GetShortestPathPinName());
 	if (bIsRotator)
 	{
-		// Easing a rotator lets us choose which path in the rotation to choose
-		UEdGraphPin *ShortestPathPin = CreatePin(EGPD_Input, K2Schema->PC_Boolean, TEXT(""), NULL, false, false, FEaseFunctionNodeHelper::GetShortestPathPinName());
-		SetPinToolTip(*ShortestPathPin, LOCTEXT("ShortestPathPinTooltip", "When interpolating the shortest path should be taken."));
-		if (ShortestPathPin->DefaultValue.Len() == 0)
-		{
-			ShortestPathPin->DefaultValue = TEXT("true");
-		}
-
-		// Add the pin
-		CustomPinNames.AddZeroed();
-		CustomPinNames[CustomPinNames.Num() - 1].PinName = FEaseFunctionNodeHelper::GetShortestPathPinName();
-		CustomPinNames[CustomPinNames.Num() - 1].CallFuncPinName = TEXT("bShortestPath");
-		CustomPinNames[CustomPinNames.Num() - 1].bValuePin = true;
-	}
-}
-
-bool UK2Node_EaseFunction::GenerateCustomPins(const int32 NewEasingFunc)
-{
-	// Early exit in case no changes are required
-	if (OldEasingFunc == NewEasingFunc ||
-		((OldEasingFunc == EEasingFunc::EaseIn || OldEasingFunc == EEasingFunc::EaseOut || OldEasingFunc == EEasingFunc::EaseInOut) &&
-		(NewEasingFunc == EEasingFunc::EaseIn || NewEasingFunc == EEasingFunc::EaseOut || NewEasingFunc == EEasingFunc::EaseInOut)))
-	{
-		OldEasingFunc = NewEasingFunc;
-		return false;
+		ShortestPathPin->bHidden = false;
 	}
-
-	bool bResult = false;
-	const UEdGraphSchema_K2* K2Schema = Cast<const UEdGraphSchema_K2>(GetSchema());
-
-	// First try to remove old pins
-	bResult |= RemoveCustomPins(false);
-
-	// Add new pins (if applicable)
-	UEdGraphPin *CustomPin = NULL;
-	switch (NewEasingFunc)
+	else
 	{
-	case EEasingFunc::EaseIn:
-	case EEasingFunc::EaseOut:
-	case EEasingFunc::EaseInOut:
-		// Generate the BlendExp custom pin
-		CustomPin = CreatePin(EGPD_Input, K2Schema->PC_Float, TEXT(""), NULL, false, false, FEaseFunctionNodeHelper::GetBlendExpPinName());
-		SetPinToolTip(*CustomPin, LOCTEXT("BlendExpPinDescription", "Blend Exponent for basic ease functions"));
-		if (CustomPin->DefaultValue.Len() == 0)
-		{
-			CustomPin->DefaultValue = TEXT("2.0");
-		}
-
-		// Add the pin
-		CustomPinNames.AddZeroed();
-		CustomPinNames[CustomPinNames.Num() - 1].PinName = FEaseFunctionNodeHelper::GetBlendExpPinName();
-		CustomPinNames[CustomPinNames.Num() - 1].CallFuncPinName = FEaseFunctionNodeHelper::GetBlendExpPinName();
-		CustomPinNames[CustomPinNames.Num() - 1].bValuePin = false;
-		bResult |= true;
-		break;
-
-	case EEasingFunc::Step:
-		// Generate the BlendExp custom pin
-		CustomPin = CreatePin(EGPD_Input, K2Schema->PC_Int, TEXT(""), NULL, false, false, FEaseFunctionNodeHelper::GetStepsPinName());
-		SetPinToolTip(*CustomPin, LOCTEXT("StepsPinDescription", "Number of steps required to go from A to B"));
-		if (CustomPin->DefaultValue.Len() == 0)
-		{
-			CustomPin->DefaultValue = TEXT("2");
-		}
-
-		// Add the pin
-		CustomPinNames.AddZeroed();
-		CustomPinNames[CustomPinNames.Num() - 1].PinName = FEaseFunctionNodeHelper::GetStepsPinName();
-		CustomPinNames[CustomPinNames.Num() - 1].CallFuncPinName = FEaseFunctionNodeHelper::GetStepsPinName();
-		CustomPinNames[CustomPinNames.Num() - 1].bValuePin = false;
-		bResult |= true;
-		break;
+		ShortestPathPin->BreakAllPinLinks();
+		ShortestPathPin->bHidden = true;
 	}
-
-	OldEasingFunc = NewEasingFunc;
-	return bResult;
 }
 
 void UK2Node_EaseFunction::ResetToWildcards()
