Minor change to limit levelscript blueprints to be consider from editor worlds only.
#Branch UE4
#Proj BlueprintGraph

diff --git a/Engine/Source/Editor/BlueprintGraph/Private/BlueprintActionDatabase.cpp b/Engine/Source/Editor/BlueprintGraph/Private/BlueprintActionDatabase.cpp
--- a/Engine/Source/Editor/BlueprintGraph/Private/BlueprintActionDatabase.cpp
+++ b/Engine/Source/Editor/BlueprintGraph/Private/BlueprintActionDatabase.cpp
@@ -949,10 +949,13 @@ static bool BlueprintActionDatabaseImpl::IsObjectValidForDatabase(UObject const*
 		// Level scripts are sometimes not assets because they have not been saved yet, but they are still valid for the database.
 		bReturn = FBlueprintEditorUtils::IsLevelScriptBlueprint(Blueprint);
 	}
-	else if( UWorld const* World = Cast<UWorld>(Object))
+	else if(UWorld const* World = Cast<UWorld>(Object))
 	{
-		// We now use worlds as databse keys to manage the level scripts they own.
-		bReturn = true;
+		// We now use worlds as databse keys to manage the level scripts they own, but we only want Editor worlds.
+		if(World->WorldType == EWorldType::Editor )
+		{
+			bReturn = true;
+		}
 	}
 	return bReturn;
 }
