UK2Node_ComponentBoundEvent::DelegateOwnerClass is no longer transient and the data will not be lost after saving BPs with the node.

diff --git a/Engine/Source/Editor/BlueprintGraph/Classes/K2Node_ComponentBoundEvent.h b/Engine/Source/Editor/BlueprintGraph/Classes/K2Node_ComponentBoundEvent.h
--- a/Engine/Source/Editor/BlueprintGraph/Classes/K2Node_ComponentBoundEvent.h
+++ b/Engine/Source/Editor/BlueprintGraph/Classes/K2Node_ComponentBoundEvent.h
@@ -18,7 +18,7 @@ class UK2Node_ComponentBoundEvent : public UK2Node_Event
 	FName DelegatePropertyName;
 
 	/** Delegate property's owner class that this event is associated with */
-	UPROPERTY(Transient)
+	UPROPERTY()
 	UClass* DelegateOwnerClass;
 
 	/** Name of property in Blueprint class that pointer to component we want to bind to */
