Fix for crash I found while deleting some functions. BlueprintNodeTemplateCache and SMyBlueprint::OnDeleteGraph

diff --git a/Engine/Source/Editor/BlueprintGraph/Private/K2Node_CreateDelegate.cpp b/Engine/Source/Editor/BlueprintGraph/Private/K2Node_CreateDelegate.cpp
--- a/Engine/Source/Editor/BlueprintGraph/Private/K2Node_CreateDelegate.cpp
+++ b/Engine/Source/Editor/BlueprintGraph/Private/K2Node_CreateDelegate.cpp
@@ -281,7 +281,13 @@ UFunction* UK2Node_CreateDelegate::GetDelegateSignature() const
 UClass* UK2Node_CreateDelegate::GetScopeClass(bool bDontUseSkeletalClassForSelf/* = false*/) const
 {
 	UEdGraphPin* Pin = FindPin(FK2Node_CreateDelegate_Helper::ObjectInputName);
-	check(Pin != NULL);
+	if (Pin == nullptr)
+	{
+		// The BlueprintNodeTemplateCache creates nodes but doesn't call allocate default pins.
+		// SMyBlueprint::OnDeleteGraph calls this function on every UK2Node_CreateDelegate. Each of 
+		// these systems is violating some first principles, so I've settled on this null check.
+		return nullptr;
+	}
 	check(Pin->LinkedTo.Num() <= 1);
 	if(Pin->LinkedTo.Num())
 	{
