Fixing some cases of finding a UFunction by name without the use of GET_MEMBER_NAME_CHECKED

diff --git a/Engine/Source/Editor/BlueprintGraph/Private/K2Node_FormatText.cpp b/Engine/Source/Editor/BlueprintGraph/Private/K2Node_FormatText.cpp
--- a/Engine/Source/Editor/BlueprintGraph/Private/K2Node_FormatText.cpp
+++ b/Engine/Source/Editor/BlueprintGraph/Private/K2Node_FormatText.cpp
@@ -202,7 +202,7 @@ void UK2Node_FormatText::ExpandNode(class FKismetCompilerContext& CompilerContex
 
 	// This is the node that does all the Format work.
 	UK2Node_CallFunction* CallFunction = CompilerContext.SpawnIntermediateNode<UK2Node_CallFunction>(this, SourceGraph);
-	CallFunction->SetFromFunction(UKismetTextLibrary::StaticClass()->FindFunctionByName("Format"));
+	CallFunction->SetFromFunction(UKismetTextLibrary::StaticClass()->FindFunctionByName(GET_MEMBER_NAME_CHECKED(UKismetTextLibrary, Format)));
 	CallFunction->AllocateDefaultPins();
 	CompilerContext.MessageLog.NotifyIntermediateObjectCreation(CallFunction, this);
 
diff --git a/Engine/Source/Editor/BlueprintGraph/Private/K2Node_VariableGet.cpp b/Engine/Source/Editor/BlueprintGraph/Private/K2Node_VariableGet.cpp
--- a/Engine/Source/Editor/BlueprintGraph/Private/K2Node_VariableGet.cpp
+++ b/Engine/Source/Editor/BlueprintGraph/Private/K2Node_VariableGet.cpp
@@ -409,11 +409,11 @@ void UK2Node_VariableGet::ExpandNode(class FKismetCompilerContext& CompilerConte
 		// Based on if the type is an "Object" or a "Class" changes which function to use
 		if (ValuePin->PinType.PinCategory == UObject::StaticClass()->GetName())
 		{
-			IsValidFunction->SetFromFunction(UKismetSystemLibrary::StaticClass()->FindFunctionByName("IsValid"));
+			IsValidFunction->SetFromFunction(UKismetSystemLibrary::StaticClass()->FindFunctionByName(GET_MEMBER_NAME_CHECKED(UKismetSystemLibrary, IsValid)));
 		}
 		else if (ValuePin->PinType.PinCategory == UClass::StaticClass()->GetName())
 		{
-			IsValidFunction->SetFromFunction(UKismetSystemLibrary::StaticClass()->FindFunctionByName("IsValidClass"));
+			IsValidFunction->SetFromFunction(UKismetSystemLibrary::StaticClass()->FindFunctionByName(GET_MEMBER_NAME_CHECKED(UKismetSystemLibrary, IsValidClass)));
 		}
 		IsValidFunction->AllocateDefaultPins();
 		CompilerContext.MessageLog.NotifyIntermediateObjectCreation(IsValidFunction, this);
