#UE4 Remove hacky fixup code from ComponentBoundEvent as the bad code did not ship in 4.7. Revert this changelist to do local fixups for broken component bound events

diff --git a/Engine/Source/Editor/BlueprintGraph/Private/K2Node_ComponentBoundEvent.cpp b/Engine/Source/Editor/BlueprintGraph/Private/K2Node_ComponentBoundEvent.cpp
--- a/Engine/Source/Editor/BlueprintGraph/Private/K2Node_ComponentBoundEvent.cpp
+++ b/Engine/Source/Editor/BlueprintGraph/Private/K2Node_ComponentBoundEvent.cpp
@@ -123,29 +123,7 @@ void UK2Node_ComponentBoundEvent::Serialize(FArchive& Ar)
 			bNeedsFixup = true;
 		}
 
-		// Recover from the period where DelegateOwnerClass was transient
-		if (!DelegateOwnerClass && HasValidBlueprint())
-		{
-			bNeedsFixup = true;
-			// Search for a component property on the owning class, this should work in most cases
-			UBlueprint* ParentBlueprint = GetBlueprint();
-			UClass* ParentClass = ParentBlueprint ? ParentBlueprint->GeneratedClass : NULL;
-			if (!ParentClass && ParentBlueprint)
-			{
-				// Try the skeleton class
-				ParentClass = ParentBlueprint->SkeletonGeneratedClass;
-			}
-
-			UObjectProperty* ComponentProperty = ParentClass ? Cast<UObjectProperty>(ParentClass->FindPropertyByName(ComponentPropertyName)) : NULL;
-
-			if (ParentClass && ComponentProperty)
-			{
-				UE_LOG(LogBlueprint, Warning, TEXT("Repaired invalid component bound event in node %s."), *GetPathName());
-				DelegateOwnerClass = ComponentProperty->PropertyClass;
-			}
-		}
-
-		if (bNeedsFixup)
+		if (bNeedsFixup || !DelegateOwnerClass)
 		{
 			// We need to fixup our event reference as it may have been saved incorrectly
 			UMulticastDelegateProperty* TargetDelegateProp = GetTargetDelegateProperty();
