[UE-12401] Fixing up the BP menu filtering to account for array functions that have params that get converted to wildcard pins.

diff --git a/Engine/Source/Editor/BlueprintGraph/Classes/EdGraphSchema_K2.h b/Engine/Source/Editor/BlueprintGraph/Classes/EdGraphSchema_K2.h
--- a/Engine/Source/Editor/BlueprintGraph/Classes/EdGraphSchema_K2.h
+++ b/Engine/Source/Editor/BlueprintGraph/Classes/EdGraphSchema_K2.h
@@ -159,6 +159,9 @@ public:
 	static const FName MD_DynamicOutputType;
 	/** Metadata that flags the function output param that will be controlled by the "MD_DynamicOutputType" pin */
 	static const FName MD_DynamicOutputParam;
+
+	static const FName MD_ArrayParam;
+	static const FName MD_ArrayDependentParam;
 	
 private:
 	// This class should never be instantiated
diff --git a/Engine/Source/Editor/BlueprintGraph/Private/BlueprintActionFilter.cpp b/Engine/Source/Editor/BlueprintGraph/Private/BlueprintActionFilter.cpp
--- a/Engine/Source/Editor/BlueprintGraph/Private/BlueprintActionFilter.cpp
+++ b/Engine/Source/Editor/BlueprintGraph/Private/BlueprintActionFilter.cpp
@@ -13,6 +13,7 @@
 #include "ObjectEditorUtils.h"		// for IsFunctionHiddenFromClass()/IsVariableCategoryHiddenFromClass()
 #include "K2Node_VariableSet.h"
 #include "K2Node_VariableGet.h"
+#include "K2Node_CallArrayFunction.h" // for IsFunctionMissingPinParam()
 // "impure" node types (utilized in BlueprintActionFilterImpl::IsImpure)
 #include "K2Node_IfThenElse.h"
 #include "K2Node_MultiGate.h"
@@ -281,6 +282,11 @@ namespace BlueprintActionFilterImpl
 	 */
 	static bool IsFunctionMissingPinParam(FBlueprintActionFilter const& Filter, FBlueprintActionInfo& BlueprintAction);
 
+	/**
+	 * 
+	 */
+	static bool ArrayFunctionHasParamOfType(const UFunction* ArrayFunction, UEdGraph const* InGraph, const FEdGraphPinType& DesiredPinType, bool bWantOutput);
+
 	/**
 	 * 
 	 * 
@@ -1142,7 +1148,8 @@ static bool BlueprintActionFilterImpl::IsFunctionMissingPinParam(FBlueprintActio
 	if (UFunction const* AssociatedFunc = BlueprintAction.GetAssociatedFunction())
 	{
 		UEdGraphSchema_K2 const* K2Schema = GetDefault<UEdGraphSchema_K2>();
-		bool const bIsEventSpawner = BlueprintAction.GetNodeClass()->IsChildOf<UK2Node_Event>();
+		bool const bIsEventSpawner  = BlueprintAction.GetNodeClass()->IsChildOf<UK2Node_Event>();
+		bool const bIsArrayFunction = BlueprintAction.GetNodeClass()->IsChildOf<UK2Node_CallArrayFunction>();
 
 		for (int32 PinIndex = 0; !bIsFilteredOut && (PinIndex < Filter.Context.Pins.Num()); ++PinIndex)
 		{
@@ -1162,8 +1169,9 @@ static bool BlueprintActionFilterImpl::IsFunctionMissingPinParam(FBlueprintActio
 				// the function signature would have them as inputs), so we 
 				// want to flip the connotation here
 				bool const bWantsOutputConnection = (PinDir == EGPD_Input) ^ bIsEventSpawner;
-
-				if (K2Schema->FunctionHasParamOfType(AssociatedFunc, K2Node->GetGraph(), PinType, bWantsOutputConnection))
+				
+				if ( K2Schema->FunctionHasParamOfType(AssociatedFunc, K2Node->GetGraph(), PinType, bWantsOutputConnection) || 
+					(bIsArrayFunction && ArrayFunctionHasParamOfType(AssociatedFunc, K2Node->GetGraph(), PinType, bWantsOutputConnection)) )
 				{
 					bIsFilteredOut = false;
 				}
@@ -1173,13 +1181,65 @@ static bool BlueprintActionFilterImpl::IsFunctionMissingPinParam(FBlueprintActio
 					bIsFilteredOut = bIsEventSpawner || !IsPinCompatibleWithTargetSelf(ContextPin, BlueprintAction);
 				}
 			}
-			
 		}
 	}
 
 	return bIsFilteredOut;
 }
 
+//------------------------------------------------------------------------------
+static bool BlueprintActionFilterImpl::ArrayFunctionHasParamOfType(const UFunction* ArrayFunction, UEdGraph const* InGraph, const FEdGraphPinType& DesiredPinType, bool bWantOutput)
+{
+	UEdGraphSchema_K2 const* K2Schema = GetDefault<UEdGraphSchema_K2>();
+
+	TSet<FString> HiddenPins;
+	FBlueprintEditorUtils::GetHiddenPinsForFunction(InGraph, ArrayFunction, HiddenPins);
+
+	FName ParamTag = FBlueprintMetadata::MD_ArrayDependentParam;
+	if (DesiredPinType.bIsArray)
+	{
+		ParamTag = FBlueprintMetadata::MD_ArrayParam;
+	}
+	const FString FlaggedParamMetaData = ArrayFunction->GetMetaData(ParamTag);
+
+	TArray<FString> WildcardPinNames;
+	FlaggedParamMetaData.ParseIntoArray(WildcardPinNames, TEXT(","), /*CullEmpty =*/true);
+
+	for (TFieldIterator<UProperty> PropIt(ArrayFunction); PropIt && (PropIt->PropertyFlags & CPF_Parm); ++PropIt)
+	{
+		UProperty* FuncParam = *PropIt;
+		const FString ParamName = FuncParam->GetName();
+
+		const bool bIsFunctionInput = !FuncParam->HasAnyPropertyFlags(CPF_OutParm) || FuncParam->HasAnyPropertyFlags(CPF_ReferenceParm);
+		if (bWantOutput == bIsFunctionInput)
+		{
+			continue;
+		}
+
+		if (!WildcardPinNames.Contains(ParamName) || HiddenPins.Contains(ParamName))
+		{
+			continue;
+		}
+
+		FEdGraphPinType ParamPinType;
+		if (K2Schema->ConvertPropertyToPinType(FuncParam, ParamPinType))
+		{
+			ParamPinType.PinCategory = UEdGraphSchema_K2::PC_Wildcard;
+
+			const FEdGraphPinType& InputType  = (bIsFunctionInput) ? ParamPinType : DesiredPinType;
+			const FEdGraphPinType& OutputType = (bIsFunctionInput) ? DesiredPinType : ParamPinType;
+
+			if (K2Schema->ArePinTypesCompatible(OutputType, InputType))
+			{
+				return true;
+			}
+		}
+
+	}
+
+	return false;
+}
+
 //------------------------------------------------------------------------------
 static bool BlueprintActionFilterImpl::IsMissmatchedPropertyType(FBlueprintActionFilter const& Filter, FBlueprintActionInfo& BlueprintAction)
 {
@@ -1708,6 +1768,14 @@ bool FBlueprintActionFilter::IsFilteredByThis(FBlueprintActionInfo& BlueprintAct
 {
 	FBlueprintActionFilter const& FilterRef = *this;
 
+	// for debugging purposes:
+// 	FBlueprintActionUiSpec UiSpec = BlueprintAction.NodeSpawner->GetUiSpec(Context, BlueprintAction.GetBindings());
+// 	bool bDebugBreak = false;
+// 	if (UiSpec.MenuName.ToString().Contains("Set Array Elem"))
+// 	{
+// 		bDebugBreak = true;
+// 	}
+
 	bool bIsFiltered = false;
 	// iterate backwards so that custom user test are ran first (and the slow
 	// internal tests are ran last).
diff --git a/Engine/Source/Editor/BlueprintGraph/Private/BlueprintFunctionNodeSpawner.cpp b/Engine/Source/Editor/BlueprintGraph/Private/BlueprintFunctionNodeSpawner.cpp
--- a/Engine/Source/Editor/BlueprintGraph/Private/BlueprintFunctionNodeSpawner.cpp
+++ b/Engine/Source/Editor/BlueprintGraph/Private/BlueprintFunctionNodeSpawner.cpp
@@ -182,7 +182,7 @@ UBlueprintFunctionNodeSpawner* UBlueprintFunctionNodeSpawner::Create(UFunction c
 	check(Function != nullptr);
 
 	bool const bIsPure = Function->HasAllFunctionFlags(FUNC_BlueprintPure);
-	bool const bHasArrayPointerParms = Function->HasMetaData(TEXT("ArrayParm"));
+	bool const bHasArrayPointerParms = Function->HasMetaData(FBlueprintMetadata::MD_ArrayParam);
 	bool const bIsCommutativeAssociativeBinaryOp = Function->HasMetaData(FBlueprintMetadata::MD_CommutativeAssociativeBinaryOperator);
 	bool const bIsMaterialParamCollectionFunc = Function->HasMetaData(FBlueprintMetadata::MD_MaterialParameterCollectionFunction);
 	bool const bIsDataTableFunc = Function->HasMetaData(FBlueprintMetadata::MD_DataTablePin);
diff --git a/Engine/Source/Editor/BlueprintGraph/Private/EdGraphSchema_K2.cpp b/Engine/Source/Editor/BlueprintGraph/Private/EdGraphSchema_K2.cpp
--- a/Engine/Source/Editor/BlueprintGraph/Private/EdGraphSchema_K2.cpp
+++ b/Engine/Source/Editor/BlueprintGraph/Private/EdGraphSchema_K2.cpp
@@ -95,6 +95,9 @@ const FName FBlueprintMetadata::MD_NativeBreakFunction(TEXT("HasNativeBreak"));
 const FName FBlueprintMetadata::MD_DynamicOutputType(TEXT("DeterminesOutputType"));
 const FName FBlueprintMetadata::MD_DynamicOutputParam(TEXT("DynamicOutputParam"));
 
+const FName FBlueprintMetadata::MD_ArrayParam(TEXT("ArrayParm"));
+const FName FBlueprintMetadata::MD_ArrayDependentParam(TEXT("ArrayTypeDependentParams"));
+
 //////////////////////////////////////////////////////////////////////////
 
 #define LOCTEXT_NAMESPACE "KismetSchema"
@@ -651,7 +654,7 @@ bool UEdGraphSchema_K2::CanFunctionBeUsedInGraph(const UClass* InClass, const UF
 
 		const bool bFunctionStatic = InFunction->HasAllFunctionFlags(FUNC_Static);
 		const bool bHasReturnParams = (InFunction->GetReturnProperty() != NULL);
-		const bool bHasArrayPointerParms = InFunction->HasMetaData(TEXT("ArrayParm"));
+		const bool bHasArrayPointerParms = InFunction->HasMetaData(FBlueprintMetadata::MD_ArrayParam);
 
 		const bool bAllowForEachCall = !bFunctionStatic && !bIsLatent && !bIsPureFunc && !bIsConstFunc && !bHasReturnParams && !bHasArrayPointerParms;
 		if (bInCalledForEach && !bAllowForEachCall)
diff --git a/Engine/Source/Editor/BlueprintGraph/Private/K2Node_CallArrayFunction.cpp b/Engine/Source/Editor/BlueprintGraph/Private/K2Node_CallArrayFunction.cpp
--- a/Engine/Source/Editor/BlueprintGraph/Private/K2Node_CallArrayFunction.cpp
+++ b/Engine/Source/Editor/BlueprintGraph/Private/K2Node_CallArrayFunction.cpp
@@ -136,7 +136,7 @@ void UK2Node_CallArrayFunction::GetArrayPins(TArray< FArrayPropertyPinCombo >& O
 
 	UFunction* TargetFunction = GetTargetFunction();
 	check(TargetFunction);
-	FString ArrayPointerMetaData = TargetFunction->GetMetaData(TEXT("ArrayParm"));
+	FString ArrayPointerMetaData = TargetFunction->GetMetaData(FBlueprintMetadata::MD_ArrayParam);
 	TArray<FString> ArrayPinComboNames;
 	ArrayPointerMetaData.ParseIntoArray(ArrayPinComboNames, TEXT(","), true);
 
@@ -163,7 +163,7 @@ bool UK2Node_CallArrayFunction::IsWildcardProperty(UFunction* InArrayFunction, c
 {
 	if(InArrayFunction && InProperty)
 	{
-		FString ArrayPointerMetaData = InArrayFunction->GetMetaData(TEXT("ArrayParm"));
+		FString ArrayPointerMetaData = InArrayFunction->GetMetaData(FBlueprintMetadata::MD_ArrayParam);
 		TArray<FString> ArrayPinComboNames;
 		ArrayPointerMetaData.ParseIntoArray(ArrayPinComboNames, TEXT(","), true);
 
@@ -188,7 +188,7 @@ void UK2Node_CallArrayFunction::GetArrayTypeDependentPins(TArray<UEdGraphPin*>&
 	UFunction* TargetFunction = GetTargetFunction();
 	check(TargetFunction);
 
-	const FString DependentPinMetaData = TargetFunction->GetMetaData(TEXT("ArrayTypeDependentParams"));
+	const FString DependentPinMetaData = TargetFunction->GetMetaData(FBlueprintMetadata::MD_ArrayDependentParam);
 	TArray<FString> TypeDependentPinNames;
 	DependentPinMetaData.ParseIntoArray(TypeDependentPinNames, TEXT(","), true);
 
diff --git a/Engine/Source/Editor/KismetCompiler/Private/KismetCompilerMisc.cpp b/Engine/Source/Editor/KismetCompiler/Private/KismetCompilerMisc.cpp
--- a/Engine/Source/Editor/KismetCompiler/Private/KismetCompilerMisc.cpp
+++ b/Engine/Source/Editor/KismetCompiler/Private/KismetCompilerMisc.cpp
@@ -45,7 +45,7 @@ bool FKismetCompilerUtilities::IsTypeCompatibleWithProperty(UEdGraphPin* SourceP
 			if(OwningFunction)
 			{
 				// Check for the magic ArrayParm property, which always matches array types
-				FString ArrayPointerMetaData = OwningFunction->GetMetaData(TEXT("ArrayParm"));
+				FString ArrayPointerMetaData = OwningFunction->GetMetaData(FBlueprintMetadata::MD_ArrayParam);
 				TArray<FString> ArrayPinComboNames;
 				ArrayPointerMetaData.ParseIntoArray(ArrayPinComboNames, TEXT(","), true);
 
@@ -76,10 +76,10 @@ bool FKismetCompilerUtilities::IsTypeCompatibleWithProperty(UEdGraphPin* SourceP
 	}
 
 	// Check for the early out...if this is a type dependent parameter in an array function
-	if ( (OwningFunction != NULL) && OwningFunction->HasMetaData(TEXT("ArrayParm")) )
+	if ( (OwningFunction != NULL) && OwningFunction->HasMetaData(FBlueprintMetadata::MD_ArrayParam) )
 	{
 		// Check to see if this param is type dependent on an array parameter
-		const FString DependentParams = OwningFunction->GetMetaData(TEXT("ArrayTypeDependentParams"));
+		const FString DependentParams = OwningFunction->GetMetaData(FBlueprintMetadata::MD_ArrayDependentParam);
 		TArray<FString>	DependentParamNames;
 		DependentParams.ParseIntoArray(DependentParamNames, TEXT(","), true);
 		if (DependentParamNames.Find(SourcePin->PinName) != INDEX_NONE)
