UE-12558 'Directions are not Compatible' error on ForEachLoopWithBreak if split struct pin is disconnected from Enum type

The problem is related to pins default value for an enum. But it seems, it should be fixed in the Marco node.

diff --git a/Engine/Source/Editor/BlueprintGraph/Private/EdGraphSchema_K2.cpp b/Engine/Source/Editor/BlueprintGraph/Private/EdGraphSchema_K2.cpp
--- a/Engine/Source/Editor/BlueprintGraph/Private/EdGraphSchema_K2.cpp
+++ b/Engine/Source/Editor/BlueprintGraph/Private/EdGraphSchema_K2.cpp
@@ -1643,7 +1643,7 @@ void UEdGraphSchema_K2::GetBreakLinkToSubMenuActions( class FMenuBuilder& MenuBu
 		++Count;
 
 		MenuBuilder.AddMenuEntry( Description, Description, FSlateIcon(), FUIAction(
-			FExecuteAction::CreateUObject((USoundClassGraphSchema*const)this, &USoundClassGraphSchema::BreakSinglePinLink, const_cast< UEdGraphPin* >(InGraphPin), *Links) ) );
+			FExecuteAction::CreateUObject(this, &UEdGraphSchema_K2::BreakSinglePinLink, const_cast< UEdGraphPin* >(InGraphPin), *Links)));
 	}
 }
 
diff --git a/Engine/Source/Editor/BlueprintGraph/Private/K2Node.cpp b/Engine/Source/Editor/BlueprintGraph/Private/K2Node.cpp
--- a/Engine/Source/Editor/BlueprintGraph/Private/K2Node.cpp
+++ b/Engine/Source/Editor/BlueprintGraph/Private/K2Node.cpp
@@ -259,8 +259,11 @@ void UK2Node::PinConnectionListChanged(UEdGraphPin* Pin)
 		{
 			const UEdGraphSchema_K2* Schema = Cast<const UEdGraphSchema_K2>(GetSchema());
 
-			Pin->DefaultObject = NULL;
-			Schema->SetPinDefaultValueBasedOnType(Pin);
+			Pin->ResetDefaultValue();
+			if (EEdGraphPinDirection::EGPD_Input == Pin->Direction)
+			{
+				Schema->SetPinDefaultValueBasedOnType(Pin);
+			}
 		}
 	}
 
