UE-14063 CRASH: If you create a MakeArray pin when dragging off an exec pin and compile
UE-14164 Crash dragging off Get Character Movement node

Better validation in MakeArray node.

diff --git a/Engine/Source/Editor/BlueprintGraph/Classes/K2Node_MakeArray.h b/Engine/Source/Editor/BlueprintGraph/Classes/K2Node_MakeArray.h
--- a/Engine/Source/Editor/BlueprintGraph/Classes/K2Node_MakeArray.h
+++ b/Engine/Source/Editor/BlueprintGraph/Classes/K2Node_MakeArray.h
@@ -28,6 +28,7 @@ public:
 	virtual void PostReconstructNode() override;
 	virtual FText GetTooltipText() const override;
 	virtual FName GetPaletteIcon(FLinearColor& OutColor) const override{ return TEXT("GraphEditor.MakeArray_16x"); }
+	virtual void ValidateNodeDuringCompilation(class FCompilerResultsLog& MessageLog) const override;
 	// End of UEdGraphNode interface
 
 	// UK2Node interface
diff --git a/Engine/Source/Editor/BlueprintGraph/Private/K2Node_MakeArray.cpp b/Engine/Source/Editor/BlueprintGraph/Private/K2Node_MakeArray.cpp
--- a/Engine/Source/Editor/BlueprintGraph/Private/K2Node_MakeArray.cpp
+++ b/Engine/Source/Editor/BlueprintGraph/Private/K2Node_MakeArray.cpp
@@ -414,17 +414,32 @@ void UK2Node_MakeArray::GetContextMenuActions(const FGraphNodeContextMenuBuilder
 
 bool UK2Node_MakeArray::IsConnectionDisallowed(const UEdGraphPin* MyPin, const UEdGraphPin* OtherPin, FString& OutReason) const
 {
-	const UEdGraphSchema_K2* Schema = GetDefault<UEdGraphSchema_K2>();
-
 	if(OtherPin->PinType.bIsArray == true && MyPin->Direction == EGPD_Input)
 	{
 		OutReason = NSLOCTEXT("K2Node", "MakeArray_InputIsArray", "Cannot make an array with an input of an array!").ToString();
 		return true;
 	}
 
+	auto Schema = Cast<const UEdGraphSchema_K2>(GetSchema());
+	if (!ensure(Schema) || (ensure(OtherPin) && Schema->IsExecPin(*OtherPin)))
+	{
+		OutReason = NSLOCTEXT("K2Node", "MakeArray_InputIsExec", "Cannot make an array with an execution input!").ToString();
+		return true;
+	}
+
 	return false;
 }
 
+void UK2Node_MakeArray::ValidateNodeDuringCompilation(class FCompilerResultsLog& MessageLog) const
+{
+	auto Schema = Cast<const UEdGraphSchema_K2>(GetSchema());
+	auto OutputPin = GetOutputPin();
+	if (!ensure(Schema) || !ensure(OutputPin) || Schema->IsExecPin(*OutputPin))
+	{
+		MessageLog.Error(*NSLOCTEXT("K2Node", "MakeArray_OutputIsExec", "Uaccepted array type in @@").ToString(), this);
+	}
+}
+
 void UK2Node_MakeArray::GetMenuActions(FBlueprintActionDatabaseRegistrar& ActionRegistrar) const
 {
 	// actions get registered under specific object-keys; the idea is that 
