"Convert to impure Get" menu option for Variable Get nodes now listed as "Convert to Validated Get"

diff --git a/Engine/Source/Editor/BlueprintGraph/Private/K2Node_VariableGet.cpp b/Engine/Source/Editor/BlueprintGraph/Private/K2Node_VariableGet.cpp
--- a/Engine/Source/Editor/BlueprintGraph/Private/K2Node_VariableGet.cpp
+++ b/Engine/Source/Editor/BlueprintGraph/Private/K2Node_VariableGet.cpp
@@ -315,7 +315,7 @@ void UK2Node_VariableGet::GetContextMenuActions(const FGraphNodeContextMenuBuild
 
 			if (bIsPureGet)
 			{
-				MenuEntryTitle   = LOCTEXT("ConvertToImpureGetTitle",   "Convert to impure Get");
+				MenuEntryTitle   = LOCTEXT("ConvertToImpureGetTitle",   "Convert to Validated Get");
 				MenuEntryTooltip = LOCTEXT("ConvertToImpureGetTooltip", "Removes the execution pins to make the node more versitile.");
 
 				const UEdGraphSchema_K2* K2Schema = Cast<UEdGraphSchema_K2>(GetSchema());
