Switch On String node will obey case insensitivity by default as the details view says it is set to.

diff --git a/Engine/Source/Editor/BlueprintGraph/Classes/K2Node_SwitchString.h b/Engine/Source/Editor/BlueprintGraph/Classes/K2Node_SwitchString.h
--- a/Engine/Source/Editor/BlueprintGraph/Classes/K2Node_SwitchString.h
+++ b/Engine/Source/Editor/BlueprintGraph/Classes/K2Node_SwitchString.h
@@ -42,4 +42,7 @@ protected:
 	virtual void CreateSelectionPin() override;
 	virtual void CreateCasePins() override;
 	virtual void RemovePin(UEdGraphPin* TargetPin) override;
+
+	/** Sets the nodes function to call based on case sensitivity */
+	void SetupCaseSensitivityFunction();
 };
diff --git a/Engine/Source/Editor/BlueprintGraph/Private/K2Node_SwitchString.cpp b/Engine/Source/Editor/BlueprintGraph/Private/K2Node_SwitchString.cpp
--- a/Engine/Source/Editor/BlueprintGraph/Private/K2Node_SwitchString.cpp
+++ b/Engine/Source/Editor/BlueprintGraph/Private/K2Node_SwitchString.cpp
@@ -12,7 +12,7 @@ UK2Node_SwitchString::UK2Node_SwitchString(const FObjectInitializer& ObjectIniti
 {
 
 	bIsCaseSensitive = false;
-	FunctionName = TEXT("NotEqual_StrStr");
+	SetupCaseSensitivityFunction();
 	FunctionClass = UKismetStringLibrary::StaticClass();
 }
 
@@ -26,9 +26,7 @@ void UK2Node_SwitchString::PostEditChangeProperty(struct FPropertyChangedEvent&
 	}
 	else if (PropertyName == TEXT("bIsCaseSensitive"))
 	{
-		FunctionName = (bIsCaseSensitive == true)
-			?  TEXT("NotEqual_StrStr")
-			: TEXT("NotEqual_StriStri");
+		SetupCaseSensitivityFunction();
 
 		FunctionClass = UKismetStringLibrary::StaticClass();
 		bIsDirty = true;
@@ -42,6 +40,13 @@ void UK2Node_SwitchString::PostEditChangeProperty(struct FPropertyChangedEvent&
 	Super::PostEditChangeProperty(PropertyChangedEvent);
 }
 
+void UK2Node_SwitchString::SetupCaseSensitivityFunction()
+{
+	FunctionName = (bIsCaseSensitive == true)
+		?  TEXT("NotEqual_StrStr")
+		: TEXT("NotEqual_StriStri");
+}
+
 FText UK2Node_SwitchString::GetNodeTitle(ENodeTitleType::Type TitleType) const
 {
 	return NSLOCTEXT("K2Node", "Switch_String", "Switch on String");
