[UE-11455] Fixing up some Blueprint breaking  issues that were caused by upgrading a project to 4.7/4.8

    * Preventing variable object type changes from breaking pin connections (when the type was changed to a super type)
    * Turning "wasted cast" errors into node warnings.
    * ValidateVariableNames() now walks the whole super chain to look for conflicting property names
    * Added a warning to the SpawnActor node (to catch when users are passing a null class)

diff --git a/Engine/Source/Editor/BlueprintGraph/Classes/K2Node_DynamicCast.h b/Engine/Source/Editor/BlueprintGraph/Classes/K2Node_DynamicCast.h
--- a/Engine/Source/Editor/BlueprintGraph/Classes/K2Node_DynamicCast.h
+++ b/Engine/Source/Editor/BlueprintGraph/Classes/K2Node_DynamicCast.h
@@ -35,6 +35,7 @@ class UK2Node_DynamicCast : public UK2Node
 	virtual bool IsConnectionDisallowed(const UEdGraphPin* MyPin, const UEdGraphPin* OtherPin, FString& OutReason) const override;
 	virtual void NotifyPinConnectionListChanged(UEdGraphPin* Pin) override;
 	virtual void ReallocatePinsDuringReconstruction(TArray<UEdGraphPin*>& OldPins) override;
+	virtual void ValidateNodeDuringCompilation(class FCompilerResultsLog& MessageLog) const override;
 	// End of UK2Node interface
 
 	/** Get the 'valid cast' exec pin */
diff --git a/Engine/Source/Editor/BlueprintGraph/Private/K2Node_DynamicCast.cpp b/Engine/Source/Editor/BlueprintGraph/Private/K2Node_DynamicCast.cpp
--- a/Engine/Source/Editor/BlueprintGraph/Private/K2Node_DynamicCast.cpp
+++ b/Engine/Source/Editor/BlueprintGraph/Private/K2Node_DynamicCast.cpp
@@ -328,39 +328,7 @@ bool UK2Node_DynamicCast::IsConnectionDisallowed(const UEdGraphPin* MyPin, const
 		}
 		else if (OtherPinType.PinCategory == UEdGraphSchema_K2::PC_Object)
 		{
-			UClass* ObjectClass = Cast<UClass>(OtherPinType.PinSubCategoryObject.Get());
-			if ((ObjectClass == nullptr) && (OtherPinType.PinSubCategory == UEdGraphSchema_K2::PSC_Self))
-			{
-				if (UK2Node* K2Node = Cast<UK2Node>(OtherPin->GetOwningNode()))
-				{
-					ObjectClass = K2Node->GetBlueprint()->GeneratedClass;
-				}
-			}
-			// if the ObjectClass is still null, assume it is a UObject, which 
-			// will work with everything (so don't disallow it)
-			
-			if (ObjectClass != nullptr)
-			{
-				
-				if (ObjectClass == TargetType)
-				{
-					bIsDisallowed = true;
-					OutReason = FText::Format(LOCTEXT("EqualObjectCast", "'{0}' is already a '{1}', you don't need the cast."),
-						OtherPinName, TargetType->GetDisplayNameText()).ToString();
-				}
-				else if (ObjectClass->IsChildOf(TargetType))
-				{
-					bIsDisallowed = true;
-					OutReason = FText::Format(LOCTEXT("UnneededObjectCast", "'{0}' is already a '{1}' (which inherits from '{2}'), so you don't need the cast."),
-						OtherPinName, ObjectClass->GetDisplayNameText(), TargetType->GetDisplayNameText()).ToString();
-				}
-				else if (!TargetType->IsChildOf(ObjectClass))
-				{
-					bIsDisallowed = true;
-					OutReason = FText::Format(LOCTEXT("DisallowedObjectCast", "'{0}' does not inherit from '{1}' (the cast would always fail)."),
-						TargetType->GetDisplayNameText(), ObjectClass->GetDisplayNameText()).ToString();
-				}
-			}
+			// let's handle wasted cast inputs with warnings in ValidateNodeDuringCompilation() instead
 		}
 		else
 		{
@@ -412,6 +380,64 @@ void UK2Node_DynamicCast::ReallocatePinsDuringReconstruction(TArray<UEdGraphPin*
 	ReconnectPureExecPins(OldPins);
 }
 
+void UK2Node_DynamicCast::ValidateNodeDuringCompilation(FCompilerResultsLog& MessageLog) const
+{
+	Super::ValidateNodeDuringCompilation(MessageLog);
+
+	UEdGraphPin* SourcePin = GetCastSourcePin();
+	if ((SourcePin->LinkedTo.Num() > 0) && (TargetType != nullptr))
+	{
+		for (UEdGraphPin* CastInput : SourcePin->LinkedTo)
+		{
+			const FEdGraphPinType& SourcePinType = CastInput->PinType;
+			if (SourcePinType.PinCategory != UEdGraphSchema_K2::PC_Object)
+			{
+				// all other types should have been rejected by IsConnectionDisallowed()
+				continue;
+			}
+
+			UClass* SourceClass = Cast<UClass>(SourcePinType.PinSubCategoryObject.Get());
+			if ((SourceClass == nullptr) && (SourcePinType.PinSubCategory == UEdGraphSchema_K2::PSC_Self))
+			{
+				if (UK2Node* K2Node = Cast<UK2Node>(CastInput->GetOwningNode()))
+				{
+					SourceClass = K2Node->GetBlueprint()->GeneratedClass;
+				}
+			}
+
+			if (SourceClass == nullptr)
+			{
+				const FString SourcePinName = CastInput->PinFriendlyName.IsEmpty() ? CastInput->PinName : CastInput->PinFriendlyName.ToString();
+
+				FText const ErrorFormat = LOCTEXT("BadCastInput", "'%s' does not have a clear object type (invalid input into @@).");
+				MessageLog.Error( *FString::Printf(*ErrorFormat.ToString(), *SourcePinName), this );
+
+				continue;
+			}
+
+			if (SourceClass == TargetType)
+			{
+				const FString SourcePinName = CastInput->PinFriendlyName.IsEmpty() ? CastInput->PinName : CastInput->PinFriendlyName.ToString();
+
+				FText const WarningFormat = LOCTEXT("EqualObjectCast", "'%s' is already a '%s', you don't need @@.");
+				MessageLog.Warning( *FString::Printf(*WarningFormat.ToString(), *SourcePinName, *TargetType->GetDisplayNameText().ToString()), this );
+			}
+			else if (SourceClass->IsChildOf(TargetType))
+			{
+				const FString SourcePinName = CastInput->PinFriendlyName.IsEmpty() ? CastInput->PinName : CastInput->PinFriendlyName.ToString();
+
+				FText const WarningFormat = LOCTEXT("UnneededObjectCast", "'%s' is already a '%s' (which inherits from '%s'), so you don't need @@.");
+				MessageLog.Warning( *FString::Printf(*WarningFormat.ToString(), *SourcePinName, *SourceClass->GetDisplayNameText().ToString(), *TargetType->GetDisplayNameText().ToString()), this );
+			}
+			else if (!TargetType->IsChildOf(SourceClass))
+			{
+				FText const WarningFormat = LOCTEXT("DisallowedObjectCast", "'%s' does not inherit from '%s' (@@ would always fail).");
+				MessageLog.Warning( *FString::Printf(*WarningFormat.ToString(), *TargetType->GetDisplayNameText().ToString(), *SourceClass->GetDisplayNameText().ToString()), this );
+			}
+		}
+	}
+}
+
 
 bool UK2Node_DynamicCast::ReconnectPureExecPins(TArray<UEdGraphPin*>& OldPins)
 {
diff --git a/Engine/Source/Editor/BlueprintGraph/Private/K2Node_Variable.cpp b/Engine/Source/Editor/BlueprintGraph/Private/K2Node_Variable.cpp
--- a/Engine/Source/Editor/BlueprintGraph/Private/K2Node_Variable.cpp
+++ b/Engine/Source/Editor/BlueprintGraph/Private/K2Node_Variable.cpp
@@ -202,11 +202,27 @@ UK2Node::ERedirectType UK2Node_Variable::DoPinsMatchForReconstruction( const UEd
 		return Super::DoPinsMatchForReconstruction(NewPin, NewPinIndex, OldPin, OldPinIndex);
 	}
 
-	const bool bCanMatchSelfs = ((OldPin->PinName == K2Schema->PN_Self) == (NewPin->PinName == K2Schema->PN_Self));
+	const bool bPinNamesMatch = (OldPin->PinName == NewPin->PinName);
+	const bool bCanMatchSelfs = bPinNamesMatch || ((OldPin->PinName == K2Schema->PN_Self) == (NewPin->PinName == K2Schema->PN_Self));
 	const bool bTheSameDirection = (NewPin->Direction == OldPin->Direction);
+
 	if (bCanMatchSelfs && bTheSameDirection)
 	{
-		if (K2Schema->ArePinTypesCompatible(NewPin->PinType, OldPin->PinType))
+		// the order that the PinTypes are passed to ArePinTypesCompatible() 
+		// matters; object pin types are seen as compatible when the output-
+		// pin's type is a subclass of the input-pin's type, so we want to keep 
+		// that in mind here (should the pins "MatchForReconstruction" if the 
+		// variable has been changed to a super class of the original? what 
+		// about a subclass?
+		// 
+		// if these are output nodes, then it is perfectly acceptable that the 
+		// variable has been altered to be a sub-class ref (meaning we should 
+		// treat the NewPin as an output)... the opposite applies if the pins 
+		// are inputs
+		const FEdGraphPinType& InputType  = (OldPin->Direction == EGPD_Output) ? OldPin->PinType : NewPin->PinType;
+		const FEdGraphPinType& OutputType = (OldPin->Direction == EGPD_Output) ? NewPin->PinType : OldPin->PinType;
+
+		if (K2Schema->ArePinTypesCompatible(OutputType, InputType))
 		{
 			// If these are split pins, we need to do some name checking logic
 			if (NewPin->ParentPin)
@@ -249,30 +265,65 @@ UK2Node::ERedirectType UK2Node_Variable::DoPinsMatchForReconstruction( const UEd
 
 			return ERedirectType_Name;
 		}
-		else if ((OldPin->PinName == NewPin->PinName) && ((NewPin->PinType.PinCategory == K2Schema->PC_Object) ||
-			(NewPin->PinType.PinCategory == K2Schema->PC_Interface)) && (NewPin->PinType.PinSubCategoryObject == NULL))
+		else
 		{
-			// Special Case:  If we had a pin match, and the class isn't loaded yet because of a cyclic dependency, temporarily cast away the const, and fix up.
-			// @TODO:  Fix this up to be less hacky
-			UBlueprintGeneratedClass* TypeClass = Cast<UBlueprintGeneratedClass>(OldPin->PinType.PinSubCategoryObject.Get());
-			if (TypeClass && TypeClass->ClassGeneratedBy && TypeClass->ClassGeneratedBy->HasAnyFlags(RF_BeingRegenerated))
+			const bool bNewPinIsObject = (NewPin->PinType.PinCategory == K2Schema->PC_Object);
+
+			// Special Case: If we had a pin match, and the class isn't loaded 
+			//               yet because of a cyclic dependency, temporarily 
+			//               cast away the const, and fix up.
+			if ( bPinNamesMatch &&
+				(bNewPinIsObject || (NewPin->PinType.PinCategory == K2Schema->PC_Interface)) &&
+				(NewPin->PinType.PinSubCategoryObject == NULL) )
 			{
-				UEdGraphPin* NonConstNewPin = (UEdGraphPin*)NewPin;
-				NonConstNewPin->PinType.PinSubCategoryObject = OldPin->PinType.PinSubCategoryObject.Get();
-				return ERedirectType_Name;
+				// @TODO:  Fix this up to be less hacky
+				UBlueprintGeneratedClass* TypeClass = Cast<UBlueprintGeneratedClass>(OldPin->PinType.PinSubCategoryObject.Get());
+				if (TypeClass && TypeClass->ClassGeneratedBy && TypeClass->ClassGeneratedBy->HasAnyFlags(RF_BeingRegenerated))
+				{
+					UEdGraphPin* NonConstNewPin = (UEdGraphPin*)NewPin;
+					NonConstNewPin->PinType.PinSubCategoryObject = OldPin->PinType.PinSubCategoryObject.Get();
+					return ERedirectType_Name;
+				}
 			}
-		}
-		else
-		{
-			// Special Case:  If we're migrating from old blueprint references to class references, allow pins to be reconnected if coerced
-			const UClass* PSCOClass = Cast<UClass>(OldPin->PinType.PinSubCategoryObject.Get());
-			const bool bOldIsBlueprint = PSCOClass && PSCOClass->IsChildOf(UBlueprint::StaticClass());
-			const bool bNewIsClass = (NewPin->PinType.PinCategory == K2Schema->PC_Class);
-			if (bNewIsClass && bOldIsBlueprint)
+			// Special Case: if we have object pins that are "compatible" in the
+			//               reverse order (meaning one's type is a sub-class of 
+			//               the other's), then they could still be acceptable 
+			//               if all their connections are still valid (for 
+			//               example: if the OldPin was an output only connected 
+			//               to super-class pins)
+			else if (bNewPinIsObject && K2Schema->ArePinTypesCompatible(InputType, OutputType))
 			{
-				UEdGraphPin* OldPinNonConst = (UEdGraphPin*)OldPin;
-				OldPinNonConst->PinName = NewPin->PinName;
-				return ERedirectType_Name;
+				bool bLinksCompatible = (OldPin->LinkedTo.Num() > 0) && (OldPin->DefaultObject == nullptr);
+				for (UEdGraphPin* OldLink : OldPin->LinkedTo)
+				{
+					const FEdGraphPinType& LinkInputType  = (OldPin->Direction == EGPD_Input) ? NewPin->PinType : OldLink->PinType;
+					const FEdGraphPinType& LinkOutputType = (OldPin->Direction == EGPD_Input) ? OldLink->PinType : NewPin->PinType;
+				
+					if (!K2Schema->ArePinTypesCompatible(LinkOutputType, LinkInputType))
+					{
+						bLinksCompatible = false;
+						break;
+					}
+				}
+
+				if (bLinksCompatible)
+				{
+					return ERedirectType_Name;
+				}
+			}
+			else
+			{
+				const UClass* PSCOClass = Cast<UClass>(OldPin->PinType.PinSubCategoryObject.Get());
+				const bool bOldIsBlueprint = PSCOClass && PSCOClass->IsChildOf(UBlueprint::StaticClass());
+				const bool bNewIsClass     = (NewPin->PinType.PinCategory == K2Schema->PC_Class);
+				// Special Case: If we're migrating from old blueprint references 
+				//               to class references, allow pins to be reconnected if coerced
+				if (bNewIsClass && bOldIsBlueprint)
+				{
+					UEdGraphPin* OldPinNonConst = (UEdGraphPin*)OldPin;
+					OldPinNonConst->PinName = NewPin->PinName;
+					return ERedirectType_Name;
+				}
 			}
 		}
 	}
diff --git a/Engine/Source/Editor/KismetCompiler/Private/KismetCompiler.cpp b/Engine/Source/Editor/KismetCompiler/Private/KismetCompiler.cpp
--- a/Engine/Source/Editor/KismetCompiler/Private/KismetCompiler.cpp
+++ b/Engine/Source/Editor/KismetCompiler/Private/KismetCompiler.cpp
@@ -449,9 +449,19 @@ void FKismetCompilerContext::ValidateVariableNames()
 			{
 				NewVarName = FBlueprintEditorUtils::FindUniqueKismetName(Blueprint, VarNameStr);
 			}
-			else if (ParentClass->HasAnyFlags(RF_Native) && FindObject<UObject>(ParentClass, *VarNameStr, /*ExactClass =*/false))
+			else if (ParentClass->HasAnyFlags(RF_Native)) // the above case handles when the parent is a blueprint
 			{
-				NewVarName = FBlueprintEditorUtils::FindUniqueKismetName(Blueprint, VarNameStr);
+				UClass* SuperClass = ParentClass;
+				do
+				{
+					if (FindObject<UObject>(SuperClass, *VarNameStr, /*ExactClass =*/false))
+					{
+						NewVarName = FBlueprintEditorUtils::FindUniqueKismetName(Blueprint, VarNameStr);
+						break;
+					}
+					SuperClass = SuperClass->GetSuperClass();
+
+				} while (SuperClass != nullptr);
 			}
 
 			if (OldVarName != NewVarName)
@@ -1158,6 +1168,7 @@ void FKismetCompilerContext::PrecompileFunction(FKismetFunctionContext& Context)
 		FName NewFunctionName = (Context.EntryPoint->CustomGeneratedFunctionName != NAME_None) ? Context.EntryPoint->CustomGeneratedFunctionName : Context.EntryPoint->SignatureName;
 		if(Context.IsDelegateSignature())
 		{
+			// prefix with the the blueprint name to avoid conflicts with natively defined delegate signatures
 			FString Name = NewFunctionName.ToString();
 			Name += HEADER_GENERATED_DELEGATE_SIGNATURE_SUFFIX;
 			NewFunctionName = FName(*Name);
diff --git a/Engine/Source/Runtime/Engine/Private/GameplayStatics.cpp b/Engine/Source/Runtime/Engine/Private/GameplayStatics.cpp
--- a/Engine/Source/Runtime/Engine/Private/GameplayStatics.cpp
+++ b/Engine/Source/Runtime/Engine/Private/GameplayStatics.cpp
@@ -333,6 +333,10 @@ AActor* UGameplayStatics::BeginSpawningActorFromClass(UObject* WorldContextObjec
 			UE_LOG(LogScript, Warning, TEXT("UGameplayStatics::BeginSpawningActorFromClass: %s can not be spawned in NULL world"), *Class->GetName());		
 		}
 	}
+	else
+	{
+		UE_LOG(LogScript, Warning, TEXT("UGameplayStatics::BeginSpawningActorFromClass: can not spawn an actor from a NULL class"));
+	}
 
 	return NewActor;
 }
