Fixed some dangerous localizable strings and improved Blueprint function/event node title localization support.

diff --git a/Engine/Source/Editor/BlueprintGraph/Classes/K2Node_CallFunction.h b/Engine/Source/Editor/BlueprintGraph/Classes/K2Node_CallFunction.h
--- a/Engine/Source/Editor/BlueprintGraph/Classes/K2Node_CallFunction.h
+++ b/Engine/Source/Editor/BlueprintGraph/Classes/K2Node_CallFunction.h
@@ -152,8 +152,8 @@ public:
 	/** Get the compact name for this function */
 	static FString GetCompactNodeTitle(const UFunction* Function);
 
-	/** Get the string to use to explain the context for this function (used on node title) */
-	virtual FString GetFunctionContextString() const;
+	/** Get the text to use to explain the context for this function (used on node title) */
+	virtual FText GetFunctionContextString() const;
 
 	/** Set properties of this node from a supplied function (does not save ref to function) */
 	virtual void SetFromFunction(const UFunction* Function);
diff --git a/Engine/Source/Editor/BlueprintGraph/Classes/K2Node_CallFunctionOnMember.h b/Engine/Source/Editor/BlueprintGraph/Classes/K2Node_CallFunctionOnMember.h
--- a/Engine/Source/Editor/BlueprintGraph/Classes/K2Node_CallFunctionOnMember.h
+++ b/Engine/Source/Editor/BlueprintGraph/Classes/K2Node_CallFunctionOnMember.h
@@ -19,7 +19,7 @@ class UK2Node_CallFunctionOnMember : public UK2Node_CallFunction
 
 	// Begin UK2Node_CallFunction interface
 	virtual UEdGraphPin* CreateSelfPin(const UFunction* Function) override;
-	virtual FString GetFunctionContextString() const override;
+	virtual FText GetFunctionContextString() const override;
 	virtual class FNodeHandlingFunctor* CreateNodeHandler(class FKismetCompilerContext& CompilerContext) const override;
 	virtual void ExpandNode(class FKismetCompilerContext& CompilerContext, UEdGraph* SourceGraph) override;
 	// End UK2Node_CallFunction interface
diff --git a/Engine/Source/Editor/BlueprintGraph/Classes/K2Node_Event.h b/Engine/Source/Editor/BlueprintGraph/Classes/K2Node_Event.h
--- a/Engine/Source/Editor/BlueprintGraph/Classes/K2Node_Event.h
+++ b/Engine/Source/Editor/BlueprintGraph/Classes/K2Node_Event.h
@@ -90,7 +90,7 @@ class UK2Node_Event : public UK2Node_EditablePinBase
 	/** Returns localized string describing replication settings. 
 	 *		Calling - whether this function is being called ("sending") or showing implementation ("receiving"). Determined whether we output "Replicated To Server" or "Replicated From Client".
 	 */
-	static FString GetLocalizedNetString(uint32 NetFlags, bool Calling);
+	static FText GetLocalizedNetString(uint32 NetFlags, bool Calling);
 
 	/** Helper function to identify if two Event nodes are the same */
 	static BLUEPRINTGRAPH_API bool AreEventNodesIdentical(const UK2Node_Event* InNodeA, const UK2Node_Event* InNodeB);
diff --git a/Engine/Source/Editor/BlueprintGraph/Private/K2Node_CallFunction.cpp b/Engine/Source/Editor/BlueprintGraph/Private/K2Node_CallFunction.cpp
--- a/Engine/Source/Editor/BlueprintGraph/Private/K2Node_CallFunction.cpp
+++ b/Engine/Source/Editor/BlueprintGraph/Private/K2Node_CallFunction.cpp
@@ -428,9 +428,9 @@ FString UK2Node_CallFunction::GetDeprecationMessage() const
 }
 
 
-FString UK2Node_CallFunction::GetFunctionContextString() const
+FText UK2Node_CallFunction::GetFunctionContextString() const
 {
-	FString ContextString;
+	FText ContextString;
 
 	// Don't show 'target is' if no target pin!
 	UEdGraphPin* SelfPin = GetDefault<UEdGraphSchema_K2>()->FindSelfPin(*this, EGPD_Input);
@@ -446,7 +446,9 @@ FString UK2Node_CallFunction::GetFunctionContextString() const
 
 		const FText TargetText = FBlueprintEditorUtils::GetFriendlyClassDisplayName(TrueSelfClass);
 
-		ContextString = FText::Format(LOCTEXT("CallFunctionOnDifferentContext", "\nTarget is {0}"), TargetText).ToString();
+		FFormatNamedArguments Args;
+		Args.Add(TEXT("TargetName"), TargetText);
+		ContextString = FText::Format(LOCTEXT("CallFunctionOnDifferentContext", "Target is {TargetName}"), Args);
 	}
 
 	return ContextString;
@@ -456,8 +458,8 @@ FString UK2Node_CallFunction::GetFunctionContextString() const
 FText UK2Node_CallFunction::GetNodeTitle(ENodeTitleType::Type TitleType) const
 {
 	FText FunctionName;
-	FString ContextString;
-	FString RPCString;
+	FText ContextString;
+	FText RPCString;
 
 	if (UFunction* Function = GetTargetFunction())
 	{
@@ -478,9 +480,25 @@ FText UK2Node_CallFunction::GetNodeTitle(ENodeTitleType::Type TitleType) const
 	{
 		FFormatNamedArguments Args;
 		Args.Add(TEXT("FunctionName"), FunctionName);
-		Args.Add(TEXT("ContextString"), FText::FromString(ContextString));
-		Args.Add(TEXT("RPCString"), FText::FromString(RPCString));
-		return FText::Format(LOCTEXT("CallFunction_FullTitle", "{FunctionName}{ContextString}{RPCString}"), Args);
+		Args.Add(TEXT("ContextString"), ContextString);
+		Args.Add(TEXT("RPCString"), RPCString);
+
+		if (ContextString.IsEmpty() && RPCString.IsEmpty())
+		{
+			return FText::Format(LOCTEXT("CallFunction_FullTitle", "{FunctionName}"), Args);
+		}
+		else if (ContextString.IsEmpty())
+		{
+			return FText::Format(LOCTEXT("CallFunction_FullTitle_WithRPCString", "{FunctionName}\n{RPCString}"), Args);
+		}
+		else if (RPCString.IsEmpty())
+		{
+			return FText::Format(LOCTEXT("CallFunction_FullTitle_WithContextString", "{FunctionName}\n{ContextString}"), Args);
+		}
+		else
+		{
+			return FText::Format(LOCTEXT("CallFunction_FullTitle_WithContextRPCString", "{FunctionName}\n{ContextString}\n{RPCString}"), Args);
+		}
 	}
 	else
 	{
diff --git a/Engine/Source/Editor/BlueprintGraph/Private/K2Node_CallFunctionOnMember.cpp b/Engine/Source/Editor/BlueprintGraph/Private/K2Node_CallFunctionOnMember.cpp
--- a/Engine/Source/Editor/BlueprintGraph/Private/K2Node_CallFunctionOnMember.cpp
+++ b/Engine/Source/Editor/BlueprintGraph/Private/K2Node_CallFunctionOnMember.cpp
@@ -32,11 +32,14 @@ UEdGraphPin* UK2Node_CallFunctionOnMember::CreateSelfPin(const UFunction* Functi
 	return SelfPin;
 }
 
-FString UK2Node_CallFunctionOnMember::GetFunctionContextString() const
+FText UK2Node_CallFunctionOnMember::GetFunctionContextString() const
 {
 	UClass* MemberVarClass = MemberVariableToCallOn.GetMemberParentClass(GetBlueprintClassFromNode());
-	FString CallFunctionClassName = (MemberVarClass != NULL) ? MemberVarClass->GetDisplayNameText().ToString() : TEXT("NULL");
-	return FString(TEXT("\n")) + FString::Printf(TEXT("%s %s (%s)"), *LOCTEXT("CallFunctionOnMemberDifferentContext", "Target is").ToString(), *CallFunctionClassName, *MemberVariableToCallOn.GetMemberName().ToString());
+	FText CallFunctionClassName = (MemberVarClass != NULL) ? MemberVarClass->GetDisplayNameText() : FText::GetEmpty();
+	FFormatNamedArguments Args;
+	Args.Add(TEXT("TargetName"), CallFunctionClassName);
+	Args.Add(TEXT("MemberVariableName"), FText::FromName(MemberVariableToCallOn.GetMemberName()));
+	return FText::Format(LOCTEXT("CallFunctionOnMemberDifferentContext", "Target is {TargetName} ({MemberVariableName})"), Args);
 }
 
 FNodeHandlingFunctor* UK2Node_CallFunctionOnMember::CreateNodeHandler(FKismetCompilerContext& CompilerContext) const
diff --git a/Engine/Source/Editor/BlueprintGraph/Private/K2Node_CustomEvent.cpp b/Engine/Source/Editor/BlueprintGraph/Private/K2Node_CustomEvent.cpp
--- a/Engine/Source/Editor/BlueprintGraph/Private/K2Node_CustomEvent.cpp
+++ b/Engine/Source/Editor/BlueprintGraph/Private/K2Node_CustomEvent.cpp
@@ -124,11 +124,11 @@ FText UK2Node_CustomEvent::GetNodeTitle(ENodeTitleType::Type TitleType) const
 	}
 	else if (CachedNodeTitle.IsOutOfDate(this))
 	{
-		FString RPCString = UK2Node_Event::GetLocalizedNetString(FunctionFlags, false);
+		FText RPCString = UK2Node_Event::GetLocalizedNetString(FunctionFlags, false);
 		
 		FFormatNamedArguments Args;
 		Args.Add(TEXT("FunctionName"), FText::FromName(CustomFunctionName));
-		Args.Add(TEXT("RPCString"), FText::FromString(RPCString));
+		Args.Add(TEXT("RPCString"), RPCString);
 
 		// FText::Format() is slow, so we cache this to save on performance
 		CachedNodeTitle.SetCachedText(FText::Format(NSLOCTEXT("K2Node", "CustomEvent_Name", "{FunctionName}{RPCString}\nCustom Event"), Args), this);
diff --git a/Engine/Source/Editor/BlueprintGraph/Private/K2Node_Event.cpp b/Engine/Source/Editor/BlueprintGraph/Private/K2Node_Event.cpp
--- a/Engine/Source/Editor/BlueprintGraph/Private/K2Node_Event.cpp
+++ b/Engine/Source/Editor/BlueprintGraph/Private/K2Node_Event.cpp
@@ -130,7 +130,7 @@ FText UK2Node_Event::GetTooltipText() const
 			{
 				Args.Add(
 					TEXT("ClientString"),
-					NSLOCTEXT("K2Node", "ServerEvent", "\n\nAuthority Only. This event only fires on the server.")
+					NSLOCTEXT("K2Node", "ServerEvent", "Authority Only. This event only fires on the server.")
 				);
 				// FText::Format() is slow, so we cache this to save on performance
 				CachedTooltip.SetCachedText(FText::Format(LOCTEXT("Event_SubtitledTooltip", "{FunctionTooltip}\n\n{ClientString}"), Args), this);
@@ -506,52 +506,51 @@ bool UK2Node_Event::CanPasteHere(const UEdGraph* TargetGraph) const
 	return !bDisallowPaste;
 }
 
-FString UK2Node_Event::GetLocalizedNetString(uint32 FunctionFlags, bool Calling)
+FText UK2Node_Event::GetLocalizedNetString(uint32 FunctionFlags, bool Calling)
 {
-	FString RPCString;
+	FText RPCString;
 	if (FunctionFlags & FUNC_Net)
 	{
-		RPCString = FString(TEXT("\n"));
-
-		if (FunctionFlags & FUNC_NetReliable)
-		{
-			RPCString += NSLOCTEXT("K2Node", "CustomEvent_ReplicatedReliable", "RELIABLE ").ToString();
-		}
-
 		if (FunctionFlags & FUNC_NetMulticast)
 		{
 			if (Calling)
 			{
-				RPCString += NSLOCTEXT("K2Node", "CustomEvent_ReplicatedMulticast", "Replicated To All (if server)").ToString();
+				RPCString = NSLOCTEXT("K2Node", "CustomEvent_ReplicatedMulticast", "Replicated To All (if server)");
 			}
 			else
 			{
-				RPCString += NSLOCTEXT("K2Node", "CustomEvent_ReplicatedMulticastFrom", "Replicated From Server\nExecutes On All").ToString();
+				RPCString = NSLOCTEXT("K2Node", "CustomEvent_ReplicatedMulticastFrom", "Replicated From Server\nExecutes On All");
 			}
-
 		}
 		else if (FunctionFlags & FUNC_NetServer)
 		{
 			if (Calling)
 			{
-				RPCString += NSLOCTEXT("K2Node", "CustomEvent_ReplicatedServer", "Replicated To Server (if owning client)").ToString();
+				RPCString = NSLOCTEXT("K2Node", "CustomEvent_ReplicatedServer", "Replicated To Server (if owning client)");
 			}
 			else
 			{
-				RPCString += NSLOCTEXT("K2Node", "CustomEvent_ReplicatedServerFrom", "Replicated From Client\nExecutes On Server").ToString();
+				RPCString = NSLOCTEXT("K2Node", "CustomEvent_ReplicatedServerFrom", "Replicated From Client\nExecutes On Server");
 			}
 		}
 		else if (FunctionFlags & FUNC_NetClient)
 		{
 			if (Calling)
 			{
-				RPCString += NSLOCTEXT("K2Node", "CustomEvent_ReplicatedClient", "Replicated To Owning Client (if server)").ToString();
+				RPCString = NSLOCTEXT("K2Node", "CustomEvent_ReplicatedClient", "Replicated To Owning Client (if server)");
 			}
 			else
 			{
-				RPCString += NSLOCTEXT("K2Node", "CustomEvent_ReplicatedClientFrom", "Replicated From Server\nExecutes on Owning Client").ToString();
+				RPCString = NSLOCTEXT("K2Node", "CustomEvent_ReplicatedClientFrom", "Replicated From Server\nExecutes on Owning Client");
 			}
 		}
+
+		if (FunctionFlags & FUNC_NetReliable)
+		{
+			FFormatNamedArguments Args;
+			Args.Add(TEXT("RPCString"), RPCString);
+			RPCString = FText::Format(NSLOCTEXT("K2Node", "CustomEvent_ReplicatedReliable", "RELIABLE {RPCString}"), Args);
+		}
 	}
 	return RPCString;
 }
diff --git a/Engine/Source/Editor/UnrealEd/Private/MaterialGraphNode.cpp b/Engine/Source/Editor/UnrealEd/Private/MaterialGraphNode.cpp
--- a/Engine/Source/Editor/UnrealEd/Private/MaterialGraphNode.cpp
+++ b/Engine/Source/Editor/UnrealEd/Private/MaterialGraphNode.cpp
@@ -181,7 +181,8 @@ FText UMaterialGraphNode::GetNodeTitle(ENodeTitleType::Type TitleType) const
 
 		if (bIsPreviewExpression)
 		{
-			NodeTitle.AppendLine(LOCTEXT("PreviewExpression", "\nPreviewing"));
+			NodeTitle.AppendLine();
+			NodeTitle.AppendLine(LOCTEXT("PreviewExpression", "Previewing"));
 		}
 
 		return NodeTitle.ToText();
