Added alias for ATAN to the MathExpression node so ARCTAN is also supported.

diff --git a/Engine/Source/Editor/BlueprintGraph/Private/K2Node_MathExpression.cpp b/Engine/Source/Editor/BlueprintGraph/Private/K2Node_MathExpression.cpp
--- a/Engine/Source/Editor/BlueprintGraph/Private/K2Node_MathExpression.cpp
+++ b/Engine/Source/Editor/BlueprintGraph/Private/K2Node_MathExpression.cpp
@@ -936,6 +936,11 @@ private:
 			ADD_ALIAS("ACOS")
 			ADD_ALIAS("ARCCOS")
 		FUNC_ALIASES_END
+
+		FUNC_ALIASES_BEGIN("ATan")
+			ADD_ALIAS("ATAN")
+			ADD_ALIAS("ARCTAN")
+		FUNC_ALIASES_END
 		
 		FUNC_ALIASES_BEGIN("MakeVector")
 			ADD_ALIAS("VECTOR")
