UE-14582 Many warnings about bSelfContext == true, but no scope supplied opening QA-Blueprints

diff --git a/Engine/Source/Editor/BlueprintGraph/Private/K2Node.cpp b/Engine/Source/Editor/BlueprintGraph/Private/K2Node.cpp
--- a/Engine/Source/Editor/BlueprintGraph/Private/K2Node.cpp
+++ b/Engine/Source/Editor/BlueprintGraph/Private/K2Node.cpp
@@ -1122,11 +1122,10 @@ void UK2Node::GetPinHoverText(const UEdGraphPin& Pin, FString& HoverTextOut) con
 
 UClass* UK2Node::GetBlueprintClassFromNode() const
 {
-	UClass* BPClass = nullptr;
-	if (HasValidBlueprint())
-	{
-		BPClass = GetBlueprint()->SkeletonGeneratedClass;
-	}
+	auto BP = FBlueprintEditorUtils::FindBlueprintForNode(this);
+	UClass* BPClass = BP
+		? (BP->SkeletonGeneratedClass ? BP->SkeletonGeneratedClass : BP->GeneratedClass)
+		: nullptr;
 	return BPClass;
 }
 
