A warning when a event has output (non-const) parameter.
UE-10836 Blueprint Interfaces cannot pass by reference between two blueprints

diff --git a/Engine/Source/Editor/BlueprintGraph/Private/DelegateNodeHandlers.cpp b/Engine/Source/Editor/BlueprintGraph/Private/DelegateNodeHandlers.cpp
--- a/Engine/Source/Editor/BlueprintGraph/Private/DelegateNodeHandlers.cpp
+++ b/Engine/Source/Editor/BlueprintGraph/Private/DelegateNodeHandlers.cpp
@@ -10,25 +10,7 @@ struct FKCHandlerDelegateHelper
 {
 	static void CheckOutputsParametersInDelegateSignature(const UFunction* SignatureFunc, const UK2Node * DelegateNode, FCompilerResultsLog& MessageLog)
 	{
-		for (TFieldIterator<UProperty> PropIt(SignatureFunc); PropIt && (PropIt->PropertyFlags & CPF_Parm); ++PropIt)
-		{
-			UProperty* FuncParam = *PropIt;
-			if (FuncParam->HasAllPropertyFlags(CPF_OutParm) && !FuncParam->HasAllPropertyFlags(CPF_ConstParm))
-			{
-				const bool bIsArray = FuncParam->IsA<UArrayProperty>(); // array is always passed by reference, see FKismetCompilerContext::CreatePropertiesFromList
-				const FString MessageStr = FString::Printf(
-					*LOCTEXT("DelegatesDontSupportRef", "Event Dispatcher: No value will be return by reference. Parameter '%s'. Node '@@'").ToString(),
-					*FuncParam->GetName());
-				if (bIsArray)
-				{
-					MessageLog.Note(*MessageStr,DelegateNode);
-				}
-				else
-				{
-					MessageLog.Warning(*MessageStr, DelegateNode);
-				}
-			}
-		}
+		FKismetCompilerUtilities::DetectValuesReturnedByRef(SignatureFunc, DelegateNode, MessageLog);
 	}
 
 	static UMulticastDelegateProperty* FindAndCheckDelegateProperty(FKismetFunctionContext& Context, UK2Node_BaseMCDelegate * DelegateNode, FCompilerResultsLog& MessageLog, const UEdGraphSchema_K2* Schema)
diff --git a/Engine/Source/Editor/BlueprintGraph/Private/K2Node_Event.cpp b/Engine/Source/Editor/BlueprintGraph/Private/K2Node_Event.cpp
--- a/Engine/Source/Editor/BlueprintGraph/Private/K2Node_Event.cpp
+++ b/Engine/Source/Editor/BlueprintGraph/Private/K2Node_Event.cpp
@@ -248,11 +248,24 @@ void UK2Node_Event::ValidateNodeDuringCompilation(class FCompilerResultsLog& Mes
 {
 	Super::ValidateNodeDuringCompilation(MessageLog);
 
-	// If we are overriding a function, but we can;t find the function we are overriding, that is a compile error
-	if(bOverrideFunction && EventReference.ResolveMember<UFunction>(GetBlueprintClassFromNode()) == nullptr)
+	UFunction* Function = nullptr;
+	if (bOverrideFunction)
 	{
-		MessageLog.Error(*FString::Printf(*NSLOCTEXT("KismetCompiler", "MissingEventSig_Error", "Missing Event '%s' for @@").ToString(), *EventReference.GetMemberName().ToString()), this);
+		Function = EventReference.ResolveMember<UFunction>(GetBlueprintClassFromNode());
+		if (!Function)
+		{
+			// If we are overriding a function, but we can;t find the function we are overriding, that is a compile error
+			MessageLog.Error(*FString::Printf(*NSLOCTEXT("KismetCompiler", "MissingEventSig_Error", "Missing Event '%s' for @@").ToString(), *EventReference.GetMemberName().ToString()), this);
+		}
+	}
+	else if (UBlueprint* Blueprint = GetBlueprint())
+	{
+		Function = Blueprint->SkeletonGeneratedClass
+			? Blueprint->SkeletonGeneratedClass->FindFunctionByName(CustomFunctionName)
+			: nullptr;
 	}
+
+	FKismetCompilerUtilities::DetectValuesReturnedByRef(Function, this, MessageLog);
 }
 
 bool UK2Node_Event::NodeCausesStructuralBlueprintChange() const
diff --git a/Engine/Source/Editor/KismetCompiler/Private/KismetCompiler.cpp b/Engine/Source/Editor/KismetCompiler/Private/KismetCompiler.cpp
--- a/Engine/Source/Editor/KismetCompiler/Private/KismetCompiler.cpp
+++ b/Engine/Source/Editor/KismetCompiler/Private/KismetCompiler.cpp
@@ -2708,7 +2708,10 @@ void FKismetCompilerContext::CreateAndProcessUbergraph()
 			{
 				for(UEdGraph* OldEventGraph : Blueprint->EventGraphs)
 				{
-					OldEventGraph->Rename(NULL, GetTransientPackage());
+					if (OldEventGraph)
+					{
+						OldEventGraph->Rename(NULL, GetTransientPackage());
+					}
 				}
 			}
 			Blueprint->EventGraphs.Empty();
diff --git a/Engine/Source/Editor/KismetCompiler/Private/KismetCompilerMisc.cpp b/Engine/Source/Editor/KismetCompiler/Private/KismetCompilerMisc.cpp
--- a/Engine/Source/Editor/KismetCompiler/Private/KismetCompilerMisc.cpp
+++ b/Engine/Source/Editor/KismetCompiler/Private/KismetCompilerMisc.cpp
@@ -1126,6 +1126,28 @@ void FKismetCompilerUtilities::ValidateProperEndExecutionPath(FKismetFunctionCon
 	}
 }
 
+void FKismetCompilerUtilities::DetectValuesReturnedByRef(const UFunction* Func, const UK2Node * Node, FCompilerResultsLog& MessageLog)
+{
+	for (TFieldIterator<UProperty> PropIt(Func); PropIt && (PropIt->PropertyFlags & CPF_Parm); ++PropIt)
+	{
+		UProperty* FuncParam = *PropIt;
+		if (FuncParam->HasAllPropertyFlags(CPF_OutParm) && !FuncParam->HasAllPropertyFlags(CPF_ConstParm))
+		{
+			const FString MessageStr = FString::Printf(
+				*LOCTEXT("WrongRefOutput", "No value will be returned by reference. Parameter '%s'. Node: @@").ToString(),
+				*FuncParam->GetName());
+			if (FuncParam->IsA<UArrayProperty>()) // array is always passed by reference, see FKismetCompilerContext::CreatePropertiesFromList
+			{
+				MessageLog.Note(*MessageStr, Node);
+			}
+			else
+			{
+				MessageLog.Warning(*MessageStr, Node);
+			}
+		}
+	}
+}
+
 //////////////////////////////////////////////////////////////////////////
 // FNodeHandlingFunctor
 
diff --git a/Engine/Source/Editor/KismetCompiler/Public/KismetCompilerMisc.h b/Engine/Source/Editor/KismetCompiler/Public/KismetCompilerMisc.h
--- a/Engine/Source/Editor/KismetCompiler/Public/KismetCompilerMisc.h
+++ b/Engine/Source/Editor/KismetCompiler/Public/KismetCompilerMisc.h
@@ -94,6 +94,9 @@ public:
 
 	/** Checks if each execution path ends with a Return node */
 	static void ValidateProperEndExecutionPath(FKismetFunctionContext& Context);
+
+	/** Generate an error for non-const output parameters */
+	static void DetectValuesReturnedByRef(const UFunction* Func, const UK2Node * Node, FCompilerResultsLog& MessageLog);
 };
 
 //////////////////////////////////////////////////////////////////////////
