UE-14706 Need a way to abort a UK2Node_BaseAsyncTask node

diff --git a/Engine/Source/Editor/BlueprintGraph/Classes/K2Node_BaseAsyncTask.h b/Engine/Source/Editor/BlueprintGraph/Classes/K2Node_BaseAsyncTask.h
--- a/Engine/Source/Editor/BlueprintGraph/Classes/K2Node_BaseAsyncTask.h
+++ b/Engine/Source/Editor/BlueprintGraph/Classes/K2Node_BaseAsyncTask.h
@@ -69,5 +69,7 @@ protected:
 			UMulticastDelegateProperty* CurrentProperty, const TArray<FBaseAsyncTaskHelper::FOutputPinAndLocalVariable>& VariableOutputs,
 			UEdGraphPin* ProxyObjectPin, UEdGraphPin*& InOutLastThenPin,
 			UK2Node* CurrentNode, UEdGraph* SourceGraph, FKismetCompilerContext& CompilerContext);
+
+		static const FString& GetAsyncTaskProxyName();
 	};
 };
diff --git a/Engine/Source/Editor/BlueprintGraph/Private/K2Node_BaseAsyncTask.cpp b/Engine/Source/Editor/BlueprintGraph/Private/K2Node_BaseAsyncTask.cpp
--- a/Engine/Source/Editor/BlueprintGraph/Private/K2Node_BaseAsyncTask.cpp
+++ b/Engine/Source/Editor/BlueprintGraph/Private/K2Node_BaseAsyncTask.cpp
@@ -54,18 +54,20 @@ void UK2Node_BaseAsyncTask::AllocateDefaultPins()
 	CreatePin(EGPD_Input, K2Schema->PC_Exec, TEXT(""), NULL, false, false, K2Schema->PN_Execute);
 	CreatePin(EGPD_Output, K2Schema->PC_Exec, TEXT(""), NULL, false, false, K2Schema->PN_Then);
 
+	CreatePin(EGPD_Output, K2Schema->PC_Object, TEXT(""), ProxyClass, false, false, FBaseAsyncTaskHelper::GetAsyncTaskProxyName());
+
 	UFunction* DelegateSignatureFunction = NULL;
 	for (TFieldIterator<UProperty> PropertyIt(ProxyClass, EFieldIteratorFlags::ExcludeSuper); PropertyIt; ++PropertyIt)
 	{
 		if (UMulticastDelegateProperty* Property = Cast<UMulticastDelegateProperty>(*PropertyIt))
 		{
-		CreatePin(EGPD_Output, K2Schema->PC_Exec, TEXT(""), NULL, false, false, *Property->GetName());
-		if (!DelegateSignatureFunction)
-		{
-			DelegateSignatureFunction = Property->SignatureFunction;
+			CreatePin(EGPD_Output, K2Schema->PC_Exec, TEXT(""), NULL, false, false, *Property->GetName());
+			if (!DelegateSignatureFunction)
+			{
+				DelegateSignatureFunction = Property->SignatureFunction;
+			}
 		}
 	}
-	}
 
 	if (DelegateSignatureFunction)
 	{
@@ -130,6 +132,11 @@ void UK2Node_BaseAsyncTask::AllocateDefaultPins()
 	Super::AllocateDefaultPins();
 }
 
+const FString& UK2Node_BaseAsyncTask::FBaseAsyncTaskHelper::GetAsyncTaskProxyName()
+{
+	static const FString Name(TEXT("AsyncTaskProxy"));
+	return Name;
+}
 
 bool UK2Node_BaseAsyncTask::FBaseAsyncTaskHelper::ValidDataPin(const UEdGraphPin* Pin, EEdGraphPinDirection Direction, const UEdGraphSchema_K2* Schema)
 {
@@ -219,11 +226,11 @@ bool UK2Node_BaseAsyncTask::FBaseAsyncTaskHelper::HandleDelegateImplementation(
 	{
 		UEdGraphPin* PinWithData = CurrentCENode->FindPin(OutputPair.OutputPin->PinName);
 		if (PinWithData == NULL)
-	{
+		{
 			FText ErrorMessage = FText::Format(LOCTEXT("MissingDataPin", "ICE: Pin @@ was expecting a data output pin named {0} on @@ (each delegate must have the same signature)"), FText::FromString(OutputPair.OutputPin->PinName));
 			CompilerContext.MessageLog.Error(*ErrorMessage.ToString(), OutputPair.OutputPin, CurrentCENode);
 			return false;
-	}
+		}
 
 		UK2Node_AssignmentStatement* AssignNode = CompilerContext.SpawnIntermediateNode<UK2Node_AssignmentStatement>(CurrentNode, SourceGraph);
 		AssignNode->AllocateDefaultPins();
@@ -273,11 +280,16 @@ void UK2Node_BaseAsyncTask::ExpandNode(class FKismetCompilerContext& CompilerCon
 		}
 	}
 
+	UEdGraphPin* const ProxyObjectPin = CallCreateProxyObjectNode->GetReturnValuePin();
+	check(ProxyObjectPin);
+	auto OutputAsyncTaskProxy = FindPinChecked(FBaseAsyncTaskHelper::GetAsyncTaskProxyName());
+	bIsErrorFree &= CompilerContext.MovePinLinksToIntermediate(*OutputAsyncTaskProxy, *ProxyObjectPin).CanSafeConnect();
+
 	// GATHER OUTPUT PARAMETERS AND PAIR THEM WITH LOCAL VARIABLES
 	TArray<FBaseAsyncTaskHelper::FOutputPinAndLocalVariable> VariableOutputs;
 	for (auto CurrentPin : Pins)
 	{
-		if (FBaseAsyncTaskHelper::ValidDataPin(CurrentPin, EGPD_Output, Schema))
+		if ((OutputAsyncTaskProxy != CurrentPin) && FBaseAsyncTaskHelper::ValidDataPin(CurrentPin, EGPD_Output, Schema))
 		{
 			const FEdGraphPinType& PinType = CurrentPin->PinType;
 			UK2Node_TemporaryVariable* TempVarOutput = CompilerContext.SpawnInternalVariable(
@@ -289,7 +301,6 @@ void UK2Node_BaseAsyncTask::ExpandNode(class FKismetCompilerContext& CompilerCon
 
 	// FOR EACH DELEGATE DEFINE EVENT, CONNECT IT TO DELEGATE AND IMPLEMENT A CHAIN OF ASSIGMENTS
 	UEdGraphPin* LastThenPin = CallCreateProxyObjectNode->FindPinChecked(Schema->PN_Then);
-	UEdGraphPin* const ProxyObjectPin = CallCreateProxyObjectNode->GetReturnValuePin();
 	for (TFieldIterator<UMulticastDelegateProperty> PropertyIt(ProxyClass, EFieldIteratorFlags::ExcludeSuper); PropertyIt && bIsErrorFree; ++PropertyIt)
 	{
 		bIsErrorFree &= FBaseAsyncTaskHelper::HandleDelegateImplementation(*PropertyIt, VariableOutputs, ProxyObjectPin, LastThenPin, this, SourceGraph, CompilerContext);
