[UE-14682] Fixed erronious warning being generated when an object reference was cast to an interface

diff --git a/Engine/Source/Editor/BlueprintGraph/Private/K2Node_DynamicCast.cpp b/Engine/Source/Editor/BlueprintGraph/Private/K2Node_DynamicCast.cpp
--- a/Engine/Source/Editor/BlueprintGraph/Private/K2Node_DynamicCast.cpp
+++ b/Engine/Source/Editor/BlueprintGraph/Private/K2Node_DynamicCast.cpp
@@ -2,9 +2,11 @@
 
 
 #include "BlueprintGraphPrivatePCH.h"
+
+#include "BlueprintEditorSettings.h"
 #include "DynamicCastHandler.h"
 #include "EditorCategoryUtils.h"
-#include "BlueprintEditorSettings.h"
+#include "KismetEditorUtilities.h"
 #include "ScopedTransaction.h"
 
 #define LOCTEXT_NAMESPACE "K2Node_DynamicCast"
@@ -387,6 +389,7 @@ void UK2Node_DynamicCast::ValidateNodeDuringCompilation(FCompilerResultsLog& Mes
 	UEdGraphPin* SourcePin = GetCastSourcePin();
 	if ((SourcePin->LinkedTo.Num() > 0) && (TargetType != nullptr))
 	{
+		const UClass* SourceType = *TargetType;
 		for (UEdGraphPin* CastInput : SourcePin->LinkedTo)
 		{
 			const FEdGraphPinType& SourcePinType = CastInput->PinType;
@@ -415,21 +418,21 @@ void UK2Node_DynamicCast::ValidateNodeDuringCompilation(FCompilerResultsLog& Mes
 				continue;
 			}
 
-			if (SourceClass == TargetType)
+			if (SourceClass == SourceType)
 			{
 				const FString SourcePinName = CastInput->PinFriendlyName.IsEmpty() ? CastInput->PinName : CastInput->PinFriendlyName.ToString();
 
 				FText const WarningFormat = LOCTEXT("EqualObjectCast", "'%s' is already a '%s', you don't need @@.");
 				MessageLog.Warning( *FString::Printf(*WarningFormat.ToString(), *SourcePinName, *TargetType->GetDisplayNameText().ToString()), this );
 			}
-			else if (SourceClass->IsChildOf(TargetType))
+			else if (SourceClass->IsChildOf(SourceType))
 			{
 				const FString SourcePinName = CastInput->PinFriendlyName.IsEmpty() ? CastInput->PinName : CastInput->PinFriendlyName.ToString();
 
 				FText const WarningFormat = LOCTEXT("UnneededObjectCast", "'%s' is already a '%s' (which inherits from '%s'), so you don't need @@.");
 				MessageLog.Warning( *FString::Printf(*WarningFormat.ToString(), *SourcePinName, *SourceClass->GetDisplayNameText().ToString(), *TargetType->GetDisplayNameText().ToString()), this );
 			}
-			else if (!TargetType->IsChildOf(SourceClass))
+			else if (!SourceType->IsChildOf(SourceClass) && !FKismetEditorUtilities::IsClassABlueprintInterface(SourceType))
 			{
 				FText const WarningFormat = LOCTEXT("DisallowedObjectCast", "'%s' does not inherit from '%s' (@@ would always fail).");
 				MessageLog.Warning( *FString::Printf(*WarningFormat.ToString(), *TargetType->GetDisplayNameText().ToString(), *SourceClass->GetDisplayNameText().ToString()), this );
