GetBaseStructure was replaced with TBaseStructure<>::Get()

Related to: https://udn.unrealengine.com/questions/244398/blueprint-struct-with-vector-member-crashes-editor.html

diff --git a/Engine/Source/Developer/BlueprintCompilerCppBackend/Private/BlueprintCompilerCppBackend.h b/Engine/Source/Developer/BlueprintCompilerCppBackend/Private/BlueprintCompilerCppBackend.h
--- a/Engine/Source/Developer/BlueprintCompilerCppBackend/Private/BlueprintCompilerCppBackend.h
+++ b/Engine/Source/Developer/BlueprintCompilerCppBackend/Private/BlueprintCompilerCppBackend.h
@@ -66,9 +66,9 @@ public:
 	{
 		extern UScriptStruct* Z_Construct_UScriptStruct_FVector();
 		VectorStruct = Z_Construct_UScriptStruct_FVector();
-		RotatorStruct = GetBaseStructure(TEXT("Rotator"));
-		TransformStruct = GetBaseStructure(TEXT("Transform"));
-		LinearColorStruct = GetBaseStructure(TEXT("LinearColor"));
+		RotatorStruct = TBaseStructure<FRotator>::Get();
+		TransformStruct = TBaseStructure<FTransform>::Get();
+		LinearColorStruct = TBaseStructure<FLinearColor>::Get();
 		LatentInfoStruct = FLatentActionInfo::StaticStruct();
 	}
 
diff --git a/Engine/Source/Editor/BlueprintGraph/Private/CallFunctionHandler.cpp b/Engine/Source/Editor/BlueprintGraph/Private/CallFunctionHandler.cpp
--- a/Engine/Source/Editor/BlueprintGraph/Private/CallFunctionHandler.cpp
+++ b/Engine/Source/Editor/BlueprintGraph/Private/CallFunctionHandler.cpp
@@ -155,9 +155,9 @@ void FKCHandler_CallFunction::CreateFunctionCallStatement(FKismetFunctionContext
 		int32 LatentTargetParamIndex = INDEX_NONE;
 
 		// Grab the special case structs that use their own literal path
-		UScriptStruct* VectorStruct = GetBaseStructure(TEXT("Vector"));
-		UScriptStruct* RotatorStruct = GetBaseStructure(TEXT("Rotator"));
-		UScriptStruct* TransformStruct = GetBaseStructure(TEXT("Transform"));
+		UScriptStruct* VectorStruct = TBaseStructure<FVector>::Get();
+		UScriptStruct* RotatorStruct = TBaseStructure<FRotator>::Get();
+		UScriptStruct* TransformStruct = TBaseStructure<FTransform>::Get();
 
 		// Check each property
 		bool bMatchedAllParams = true;
diff --git a/Engine/Source/Editor/BlueprintGraph/Private/EdGraphSchema_K2.cpp b/Engine/Source/Editor/BlueprintGraph/Private/EdGraphSchema_K2.cpp
--- a/Engine/Source/Editor/BlueprintGraph/Private/EdGraphSchema_K2.cpp
+++ b/Engine/Source/Editor/BlueprintGraph/Private/EdGraphSchema_K2.cpp
@@ -494,11 +494,11 @@ UEdGraphSchema_K2::UEdGraphSchema_K2(const FObjectInitializer& ObjectInitializer
 	// Initialize cached static references to well-known struct types
 	if (VectorStruct == nullptr)
 	{
-		VectorStruct = GetBaseStructure(TEXT("Vector"));
-		RotatorStruct = GetBaseStructure(TEXT("Rotator"));
-		TransformStruct = GetBaseStructure(TEXT("Transform"));
-		LinearColorStruct = GetBaseStructure(TEXT("LinearColor"));
-		ColorStruct = GetBaseStructure(TEXT("Color"));
+		VectorStruct = TBaseStructure<FVector>::Get();
+		RotatorStruct = TBaseStructure<FRotator>::Get();
+		TransformStruct = TBaseStructure<FTransform>::Get();
+		LinearColorStruct = TBaseStructure<FLinearColor>::Get();
+		ColorStruct = TBaseStructure<FColor>::Get();
 	}
 }
 
@@ -5502,8 +5502,8 @@ void UEdGraphSchema_K2::SplitPin(UEdGraphPin* Pin) const
 	if (Pin->Direction == EGPD_Input)
 	{
 		TArray<FString> OriginalDefaults;
-		if (   StructType == GetBaseStructure(TEXT("Vector"))
-			|| StructType == GetBaseStructure(TEXT("Rotator")))
+		if (   StructType == TBaseStructure<FVector>::Get()
+			|| StructType == TBaseStructure<FRotator>::Get())
 		{
 			Pin->DefaultValue.ParseIntoArray(OriginalDefaults, TEXT(","), false);
 			for (FString& Default : OriginalDefaults)
@@ -5516,7 +5516,7 @@ void UEdGraphSchema_K2::SplitPin(UEdGraphPin* Pin) const
 				OriginalDefaults.Add(TEXT("0.0"));
 			}
 		}
-		else if (StructType == GetBaseStructure(TEXT("Vector2D")))
+		else if (StructType == TBaseStructure<FVector2D>::Get())
 		{
 			FVector2D V2D;
 			V2D.InitFromString(Pin->DefaultValue);
@@ -5524,7 +5524,7 @@ void UEdGraphSchema_K2::SplitPin(UEdGraphPin* Pin) const
 			OriginalDefaults.Add(FString::SanitizeFloat(V2D.X));
 			OriginalDefaults.Add(FString::SanitizeFloat(V2D.Y));
 		}
-		else if (StructType == GetBaseStructure(TEXT("LinearColor")))
+		else if (StructType == TBaseStructure<FLinearColor>::Get())
 		{
 			FLinearColor LC;
 			LC.InitFromString(Pin->DefaultValue);
@@ -5579,14 +5579,14 @@ void UEdGraphSchema_K2::RecombinePin(UEdGraphPin* Pin) const
 		UScriptStruct* StructType = CastChecked<UScriptStruct>(ParentPin->PinType.PinSubCategoryObject.Get());
 
 		TArray<FString> OriginalDefaults;
-		if (   StructType == GetBaseStructure(TEXT("Vector"))
-			|| StructType == GetBaseStructure(TEXT("Rotator")))
+		if (   StructType == TBaseStructure<FVector>::Get()
+			|| StructType == TBaseStructure<FRotator>::Get())
 		{
 			ParentPin->DefaultValue = ParentPin->SubPins[0]->DefaultValue + TEXT(",") 
 									+ ParentPin->SubPins[1]->DefaultValue + TEXT(",")
 									+ ParentPin->SubPins[2]->DefaultValue;
 		}
-		else if (StructType == GetBaseStructure(TEXT("Vector2D")))
+		else if (StructType == TBaseStructure<FVector2D>::Get())
 		{
 			FVector2D V2D;
 			V2D.X = FCString::Atof(*ParentPin->SubPins[0]->DefaultValue);
@@ -5594,7 +5594,7 @@ void UEdGraphSchema_K2::RecombinePin(UEdGraphPin* Pin) const
 			
 			ParentPin->DefaultValue = V2D.ToString();
 		}
-		else if (StructType == GetBaseStructure(TEXT("LinearColor")))
+		else if (StructType == TBaseStructure<FLinearColor>::Get())
 		{
 			FLinearColor LC;
 			LC.R = FCString::Atof(*ParentPin->SubPins[0]->DefaultValue);
diff --git a/Engine/Source/Editor/BlueprintGraph/Private/K2Node_InputTouch.cpp b/Engine/Source/Editor/BlueprintGraph/Private/K2Node_InputTouch.cpp
--- a/Engine/Source/Editor/BlueprintGraph/Private/K2Node_InputTouch.cpp
+++ b/Engine/Source/Editor/BlueprintGraph/Private/K2Node_InputTouch.cpp
@@ -47,7 +47,7 @@ void UK2Node_InputTouch::AllocateDefaultPins()
 	CreatePin(EGPD_Output, K2Schema->PC_Exec, TEXT(""), NULL, false, false, TEXT("Released"));
 	CreatePin(EGPD_Output, K2Schema->PC_Exec, TEXT(""), NULL, false, false, TEXT("Moved"));
 	
-	UScriptStruct* VectorStruct = GetBaseStructure(TEXT("Vector"));
+	UScriptStruct* VectorStruct = TBaseStructure<FVector>::Get();
 	CreatePin(EGPD_Output, K2Schema->PC_Struct, TEXT(""), VectorStruct, false, false, TEXT("Location"));
 
 	CreatePin(EGPD_Output, K2Schema->PC_Byte, TEXT(""), GetTouchIndexEnum(), false, false, TEXT("FingerIndex"));
@@ -170,7 +170,7 @@ void UK2Node_InputTouch::ExpandNode(FKismetCompilerContext& CompilerContext, UEd
 	if( ActivePins.Num() > 1 )
 	{
 		// Create a temporary variable to copy location in to
-		static UScriptStruct* VectorStruct = GetBaseStructure(TEXT("Vector"));
+		static UScriptStruct* VectorStruct = TBaseStructure<FVector>::Get();
 		UK2Node_TemporaryVariable* TouchLocationVar = CompilerContext.SpawnIntermediateNode<UK2Node_TemporaryVariable>(this, SourceGraph);
 		TouchLocationVar->VariableType.PinCategory = Schema->PC_Struct;
 		TouchLocationVar->VariableType.PinSubCategoryObject = VectorStruct;
diff --git a/Engine/Source/Editor/BlueprintGraph/Private/K2Node_SpawnActor.cpp b/Engine/Source/Editor/BlueprintGraph/Private/K2Node_SpawnActor.cpp
--- a/Engine/Source/Editor/BlueprintGraph/Private/K2Node_SpawnActor.cpp
+++ b/Engine/Source/Editor/BlueprintGraph/Private/K2Node_SpawnActor.cpp
@@ -36,7 +36,7 @@ void UK2Node_SpawnActor::AllocateDefaultPins()
 	K2Schema->ConstructBasicPinTooltip(*BlueprintPin, LOCTEXT("BlueprintPinDescription", "The blueprint Actor you want to spawn"), BlueprintPin->PinToolTip);
 
 	// Transform pin
-	UScriptStruct* TransformStruct = GetBaseStructure(TEXT("Transform"));
+	UScriptStruct* TransformStruct = TBaseStructure<FTransform>::Get();
 	UEdGraphPin* TransformPin = CreatePin(EGPD_Input, K2Schema->PC_Struct, TEXT(""), TransformStruct, false, false, SpawnTransformPinName);
 	K2Schema->ConstructBasicPinTooltip(*TransformPin, LOCTEXT("TransformPinDescription", "The transform to spawn the Actor with"), TransformPin->PinToolTip);
 
diff --git a/Engine/Source/Editor/BlueprintGraph/Private/K2Node_SpawnActorFromClass.cpp b/Engine/Source/Editor/BlueprintGraph/Private/K2Node_SpawnActorFromClass.cpp
--- a/Engine/Source/Editor/BlueprintGraph/Private/K2Node_SpawnActorFromClass.cpp
+++ b/Engine/Source/Editor/BlueprintGraph/Private/K2Node_SpawnActorFromClass.cpp
@@ -49,7 +49,7 @@ void UK2Node_SpawnActorFromClass::AllocateDefaultPins()
 	K2Schema->ConstructBasicPinTooltip(*ClassPin, LOCTEXT("ClassPinDescription", "The Actor class you want to spawn"), ClassPin->PinToolTip);
 
 	// Transform pin
-	UScriptStruct* TransformStruct = GetBaseStructure(TEXT("Transform"));
+	UScriptStruct* TransformStruct = TBaseStructure<FTransform>::Get();
 	UEdGraphPin* TransformPin = CreatePin(EGPD_Input, K2Schema->PC_Struct, TEXT(""), TransformStruct, false, false, FK2Node_SpawnActorFromClassHelper::SpawnTransformPinName);
 	K2Schema->ConstructBasicPinTooltip(*TransformPin, LOCTEXT("TransformPinDescription", "The transform to spawn the Actor with"), TransformPin->PinToolTip);
 
diff --git a/Engine/Source/Editor/BlueprintGraph/Private/K2Node_TemporaryVariable.cpp b/Engine/Source/Editor/BlueprintGraph/Private/K2Node_TemporaryVariable.cpp
--- a/Engine/Source/Editor/BlueprintGraph/Private/K2Node_TemporaryVariable.cpp
+++ b/Engine/Source/Editor/BlueprintGraph/Private/K2Node_TemporaryVariable.cpp
@@ -180,15 +180,15 @@ void UK2Node_TemporaryVariable::GetMenuActions(FBlueprintActionDatabaseRegistrar
 	ActionRegistrar.AddBlueprintAction(ActionKey, MakeTempVarNodeSpawner(FEdGraphPinType(K2Schema->PC_Text, TEXT(""), nullptr, /*bIsArray =*/false, /*bIsReference =*/false), /*bIsPersistent =*/false));
 	ActionRegistrar.AddBlueprintAction(ActionKey, MakeTempVarNodeSpawner(FEdGraphPinType(K2Schema->PC_Text, TEXT(""), nullptr, /*bIsArray =*/ true, /*bIsReference =*/false), /*bIsPersistent =*/false));
 
-	UScriptStruct* VectorStruct  = GetBaseStructure(TEXT("Vector"));
+	UScriptStruct* VectorStruct  = TBaseStructure<FVector>::Get();
 	ActionRegistrar.AddBlueprintAction(ActionKey, MakeTempVarNodeSpawner(FEdGraphPinType(K2Schema->PC_Struct, TEXT("Vector"), VectorStruct, /*bIsArray =*/false, /*bIsReference =*/false), /*bIsPersistent =*/false));
 	ActionRegistrar.AddBlueprintAction(ActionKey, MakeTempVarNodeSpawner(FEdGraphPinType(K2Schema->PC_Struct, TEXT("Vector"), VectorStruct, /*bIsArray =*/ true, /*bIsReference =*/false), /*bIsPersistent =*/false));
 	
-	UScriptStruct* RotatorStruct = GetBaseStructure(TEXT("Rotator"));
+	UScriptStruct* RotatorStruct = TBaseStructure<FRotator>::Get();
 	ActionRegistrar.AddBlueprintAction(ActionKey, MakeTempVarNodeSpawner(FEdGraphPinType(K2Schema->PC_Struct, TEXT("Rotator"), RotatorStruct, /*bIsArray =*/false, /*bIsReference =*/false), /*bIsPersistent =*/false));
 	ActionRegistrar.AddBlueprintAction(ActionKey, MakeTempVarNodeSpawner(FEdGraphPinType(K2Schema->PC_Struct, TEXT("Rotator"), RotatorStruct, /*bIsArray =*/ true, /*bIsReference =*/false), /*bIsPersistent =*/false));
 	
-	UScriptStruct* TransformStruct = GetBaseStructure(TEXT("Transform"));
+	UScriptStruct* TransformStruct = TBaseStructure<FTransform>::Get();
 	ActionRegistrar.AddBlueprintAction(ActionKey, MakeTempVarNodeSpawner(FEdGraphPinType(K2Schema->PC_Struct, TEXT("Transform"), TransformStruct, /*bIsArray =*/false, /*bIsReference =*/false), /*bIsPersistent =*/false));
 	ActionRegistrar.AddBlueprintAction(ActionKey, MakeTempVarNodeSpawner(FEdGraphPinType(K2Schema->PC_Struct, TEXT("Transform"), TransformStruct, /*bIsArray =*/ true, /*bIsReference =*/false), /*bIsPersistent =*/false));
 
diff --git a/Engine/Source/Editor/BlueprintGraph/Private/K2Node_Timeline.cpp b/Engine/Source/Editor/BlueprintGraph/Private/K2Node_Timeline.cpp
--- a/Engine/Source/Editor/BlueprintGraph/Private/K2Node_Timeline.cpp
+++ b/Engine/Source/Editor/BlueprintGraph/Private/K2Node_Timeline.cpp
@@ -84,14 +84,14 @@ void UK2Node_Timeline::AllocateDefaultPins()
 		for(int32 i=0; i<Timeline->VectorTracks.Num(); i++)
 		{
 			FTTVectorTrack& VectorTrack = Timeline->VectorTracks[i];
-			UScriptStruct* VectorStruct = GetBaseStructure(TEXT("Vector"));
+			UScriptStruct* VectorStruct = TBaseStructure<FVector>::Get();
 			CreatePin(EGPD_Output, K2Schema->PC_Struct, TEXT(""), VectorStruct, false, false, VectorTrack.TrackName.ToString());			
 		}
 
 		for(int32 i=0; i<Timeline->LinearColorTracks.Num(); i++)
 		{
 			FTTLinearColorTrack& LinearColorTrack = Timeline->LinearColorTracks[i];
-			UScriptStruct* LinearColorStruct = GetBaseStructure(TEXT("LinearColor"));
+			UScriptStruct* LinearColorStruct = TBaseStructure<FLinearColor>::Get();
 			CreatePin(EGPD_Output, K2Schema->PC_Struct, TEXT(""), LinearColorStruct, false, false, LinearColorTrack.TrackName.ToString());			
 		}
 
diff --git a/Engine/Source/Editor/GraphEditor/Private/KismetPins/SGraphPinVector.cpp b/Engine/Source/Editor/GraphEditor/Private/KismetPins/SGraphPinVector.cpp
--- a/Engine/Source/Editor/GraphEditor/Private/KismetPins/SGraphPinVector.cpp
+++ b/Engine/Source/Editor/GraphEditor/Private/KismetPins/SGraphPinVector.cpp
@@ -151,7 +151,7 @@ void SGraphPinVector::Construct(const FArguments& InArgs, UEdGraphPin* InGraphPi
 
 TSharedRef<SWidget>	SGraphPinVector::GetDefaultValueWidget()
 {
-	UScriptStruct* RotatorStruct = GetBaseStructure(TEXT("Rotator"));
+	UScriptStruct* RotatorStruct = TBaseStructure<FRotator>::Get();
 	bIsRotator = (GraphPinObj->PinType.PinSubCategoryObject == RotatorStruct) ? true : false;
 	
 	//Create widget
diff --git a/Engine/Source/Editor/GraphEditor/Private/NiagaraPins/SGraphPinVector4.cpp b/Engine/Source/Editor/GraphEditor/Private/NiagaraPins/SGraphPinVector4.cpp
--- a/Engine/Source/Editor/GraphEditor/Private/NiagaraPins/SGraphPinVector4.cpp
+++ b/Engine/Source/Editor/GraphEditor/Private/NiagaraPins/SGraphPinVector4.cpp
@@ -182,7 +182,7 @@ void SGraphPinVector4::Construct(const FArguments& InArgs, UEdGraphPin* InGraphP
 
 TSharedRef<SWidget>	SGraphPinVector4::GetDefaultValueWidget()
 {
-	UScriptStruct* RotatorStruct = GetBaseStructure(TEXT("Rotator"));
+	UScriptStruct* RotatorStruct = TBaseStructure<FRotator>::Get();
 
 	//Create widget
 	return	SNew(SVector4TextBox)
diff --git a/Engine/Source/Editor/GraphEditor/Private/NodeFactory.cpp b/Engine/Source/Editor/GraphEditor/Private/NodeFactory.cpp
--- a/Engine/Source/Editor/GraphEditor/Private/NodeFactory.cpp
+++ b/Engine/Source/Editor/GraphEditor/Private/NodeFactory.cpp
@@ -309,10 +309,10 @@ TSharedPtr<SGraphPin> FNodeFactory::CreatePinWidget(UEdGraphPin* InPin)
 		else if (InPin->PinType.PinCategory == K2Schema->PC_Struct)
 		{
 			// If you update this logic you'll probably need to update UEdGraphSchema_K2::ShouldHidePinDefaultValue!
-			UScriptStruct* ColorStruct = GetBaseStructure(TEXT("LinearColor"));
-			UScriptStruct* VectorStruct = GetBaseStructure(TEXT("Vector"));
-			UScriptStruct* Vector2DStruct = GetBaseStructure(TEXT("Vector2D"));
-			UScriptStruct* RotatorStruct = GetBaseStructure(TEXT("Rotator"));
+			UScriptStruct* ColorStruct = TBaseStructure<FLinearColor>::Get();
+			UScriptStruct* VectorStruct = TBaseStructure<FVector>::Get();
+			UScriptStruct* Vector2DStruct = TBaseStructure<FVector2D>::Get();
+			UScriptStruct* RotatorStruct = TBaseStructure<FRotator>::Get();
 
 			if (InPin->PinType.PinSubCategoryObject == ColorStruct)
 			{
diff --git a/Engine/Source/Editor/KismetCompiler/Private/KismetCompiler.cpp b/Engine/Source/Editor/KismetCompiler/Private/KismetCompiler.cpp
--- a/Engine/Source/Editor/KismetCompiler/Private/KismetCompiler.cpp
+++ b/Engine/Source/Editor/KismetCompiler/Private/KismetCompiler.cpp
@@ -81,10 +81,10 @@ FKismetCompilerContext::FKismetCompilerContext(UBlueprint* SourceSketch, FCompil
 	MacroSpawnX = MinimumSpawnX;
 	MacroSpawnY = -2000;
 	
-	VectorStruct = GetBaseStructure(TEXT("Vector"));
-	RotatorStruct = GetBaseStructure(TEXT("Rotator"));
-	TransformStruct = GetBaseStructure(TEXT("Transform"));
-	LinearColorStruct = GetBaseStructure(TEXT("LinearColor"));
+	VectorStruct = TBaseStructure<FVector>::Get();
+	RotatorStruct = TBaseStructure<FRotator>::Get();
+	TransformStruct = TBaseStructure<FTransform>::Get();
+	LinearColorStruct = TBaseStructure<FLinearColor>::Get();
 }
 
 FKismetCompilerContext::~FKismetCompilerContext()
diff --git a/Engine/Source/Editor/KismetCompiler/Private/KismetCompilerVMBackend.cpp b/Engine/Source/Editor/KismetCompiler/Private/KismetCompilerVMBackend.cpp
--- a/Engine/Source/Editor/KismetCompiler/Private/KismetCompilerVMBackend.cpp
+++ b/Engine/Source/Editor/KismetCompiler/Private/KismetCompilerVMBackend.cpp
@@ -319,9 +319,9 @@ public:
 		, bIsUbergraph(bInIsUbergraph)
 		, ReturnStatement(InReturnStatement)
 	{
-		VectorStruct = GetBaseStructure(TEXT("Vector"));
-		RotatorStruct = GetBaseStructure(TEXT("Rotator"));
-		TransformStruct = GetBaseStructure(TEXT("Transform"));
+		VectorStruct = TBaseStructure<FVector>::Get();
+		RotatorStruct = TBaseStructure<FRotator>::Get();
+		TransformStruct = TBaseStructure<FTransform>::Get();
 		LatentInfoStruct = FLatentActionInfo::StaticStruct();
 	}
 
diff --git a/Engine/Source/Editor/UnrealEd/Private/Kismet2/BlueprintEditorUtils.cpp b/Engine/Source/Editor/UnrealEd/Private/Kismet2/BlueprintEditorUtils.cpp
--- a/Engine/Source/Editor/UnrealEd/Private/Kismet2/BlueprintEditorUtils.cpp
+++ b/Engine/Source/Editor/UnrealEd/Private/Kismet2/BlueprintEditorUtils.cpp
@@ -4984,9 +4984,9 @@ void FBlueprintEditorUtils::ImportKismetDefaultValueToProperty(UEdGraphPin* Sour
 
 	if (UStructProperty* StructProperty = Cast<UStructProperty>(DestinationProperty))
 	{
-		static UScriptStruct* VectorStruct = GetBaseStructure(TEXT("Vector"));
-		static UScriptStruct* RotatorStruct = GetBaseStructure(TEXT("Rotator"));
-		static UScriptStruct* TransformStruct = GetBaseStructure(TEXT("Transform"));
+		static UScriptStruct* VectorStruct = TBaseStructure<FVector>::Get();
+		static UScriptStruct* RotatorStruct = TBaseStructure<FRotator>::Get();
+		static UScriptStruct* TransformStruct = TBaseStructure<FTransform>::Get();
 
 		if (StructProperty->Struct == VectorStruct)
 		{
@@ -5020,9 +5020,9 @@ void FBlueprintEditorUtils::ExportPropertyToKismetDefaultValue(UEdGraphPin* Targ
 
 	if (UStructProperty* StructProperty = Cast<UStructProperty>(SourceProperty))
 	{
-		static UScriptStruct* VectorStruct = GetBaseStructure(TEXT("Vector"));
-		static UScriptStruct* RotatorStruct = GetBaseStructure(TEXT("Rotator"));
-		static UScriptStruct* TransformStruct = GetBaseStructure(TEXT("Transform"));
+		static UScriptStruct* VectorStruct = TBaseStructure<FVector>::Get();
+		static UScriptStruct* RotatorStruct = TBaseStructure<FRotator>::Get();
+		static UScriptStruct* TransformStruct = TBaseStructure<FTransform>::Get();
 
 		if (StructProperty->Struct == VectorStruct)
 		{
@@ -7183,10 +7183,10 @@ bool FBlueprintEditorUtils::PropertyValueFromString(const UProperty* Property, c
 	}
 	else 
 	{
-		static UScriptStruct* VectorStruct = GetBaseStructure(TEXT("Vector"));
-		static UScriptStruct* RotatorStruct = GetBaseStructure(TEXT("Rotator"));
-		static UScriptStruct* TransformStruct = GetBaseStructure(TEXT("Transform"));
-		static UScriptStruct* LinearColorStruct = GetBaseStructure(TEXT("LinearColor"));
+		static UScriptStruct* VectorStruct = TBaseStructure<FVector>::Get();
+		static UScriptStruct* RotatorStruct = TBaseStructure<FRotator>::Get();
+		static UScriptStruct* TransformStruct = TBaseStructure<FTransform>::Get();
+		static UScriptStruct* LinearColorStruct = TBaseStructure<FLinearColor>::Get();
 
 		const UStructProperty* StructProperty = CastChecked<UStructProperty>(Property);
 
@@ -7240,10 +7240,10 @@ bool FBlueprintEditorUtils::PropertyValueToString(const UProperty* Property, con
 	OutForm.Reset();
 	if (Property->IsA(UStructProperty::StaticClass()))
 	{
-		static UScriptStruct* VectorStruct = GetBaseStructure(TEXT("Vector"));
-		static UScriptStruct* RotatorStruct = GetBaseStructure(TEXT("Rotator"));
-		static UScriptStruct* TransformStruct = GetBaseStructure(TEXT("Transform"));
-		static UScriptStruct* LinearColorStruct = GetBaseStructure(TEXT("LinearColor"));
+		static UScriptStruct* VectorStruct = TBaseStructure<FVector>::Get();
+		static UScriptStruct* RotatorStruct = TBaseStructure<FRotator>::Get();
+		static UScriptStruct* TransformStruct = TBaseStructure<FTransform>::Get();
+		static UScriptStruct* LinearColorStruct = TBaseStructure<FLinearColor>::Get();
 
 		const UStructProperty* StructProperty = CastChecked<UStructProperty>(Property);
 
diff --git a/Engine/Source/Runtime/CoreUObject/Private/UObject/Class.cpp b/Engine/Source/Runtime/CoreUObject/Private/UObject/Class.cpp
--- a/Engine/Source/Runtime/CoreUObject/Private/UObject/Class.cpp
+++ b/Engine/Source/Runtime/CoreUObject/Private/UObject/Class.cpp
@@ -4193,12 +4193,61 @@ bool UFunction::IsSignatureCompatibleWith(const UFunction* OtherFunction, uint64
 	return !(IteratorB && (IteratorB->PropertyFlags & CPF_Parm));
 }
 
-UScriptStruct* GetBaseStructure(const TCHAR* Name)
+static UScriptStruct* StaticGetBaseStructureInternal(const TCHAR* Name)
 {
 	static auto* CoreUObjectPkg = FindObjectChecked<UPackage>(nullptr, TEXT("/Script/CoreUObject"));
 	return FindObjectChecked<UScriptStruct>(CoreUObjectPkg, Name);
 }
 
+UScriptStruct* TBaseStructure<FRotator>::Get()
+{
+	static auto ScriptStruct = StaticGetBaseStructureInternal(TEXT("Rotator"));
+	return ScriptStruct;
+}
+
+UScriptStruct* TBaseStructure<FTransform>::Get()
+{
+	static auto ScriptStruct = StaticGetBaseStructureInternal(TEXT("Transform"));
+	return ScriptStruct;
+}
+
+UScriptStruct* TBaseStructure<FLinearColor>::Get()
+{
+	static auto ScriptStruct = StaticGetBaseStructureInternal(TEXT("LinearColor"));
+	return ScriptStruct;
+}
+
+UScriptStruct* TBaseStructure<FColor>::Get()
+{
+	static auto ScriptStruct = StaticGetBaseStructureInternal(TEXT("Color"));
+	return ScriptStruct;
+}
+
+UScriptStruct* TBaseStructure<FVector>::Get()
+{
+	static auto ScriptStruct = StaticGetBaseStructureInternal(TEXT("Vector"));
+	return ScriptStruct;
+}
+
+UScriptStruct* TBaseStructure<FVector2D>::Get()
+{
+	static auto ScriptStruct = StaticGetBaseStructureInternal(TEXT("Vector2D"));
+	return ScriptStruct;
+}
+
+UScriptStruct* TBaseStructure<FRandomStream>::Get()
+{
+	static auto ScriptStruct = StaticGetBaseStructureInternal(TEXT("RandomStream"));
+	return ScriptStruct;
+}
+
+UScriptStruct* TBaseStructure<FFallbackStruct>::Get()
+{
+	static auto ScriptStruct = StaticGetBaseStructureInternal(TEXT("FallbackStruct"));
+	return ScriptStruct;
+}
+
+
 IMPLEMENT_CORE_INTRINSIC_CLASS(UFunction, UStruct,
 	{
 	}
diff --git a/Engine/Source/Runtime/CoreUObject/Private/UObject/UObjectGlobals.cpp b/Engine/Source/Runtime/CoreUObject/Private/UObject/UObjectGlobals.cpp
--- a/Engine/Source/Runtime/CoreUObject/Private/UObject/UObjectGlobals.cpp
+++ b/Engine/Source/Runtime/CoreUObject/Private/UObject/UObjectGlobals.cpp
@@ -2842,8 +2842,7 @@ FArchive& FScriptInterface::Serialize(FArchive& Ar, UClass* InterfaceType)
 /** A struct used as stub for deleted ones. */
 UScriptStruct* GetFallbackStruct()
 {
-	static UScriptStruct* FallbackStruct = GetBaseStructure(TEXT("FallbackStruct"));
-	return FallbackStruct;
+	return TBaseStructure<FFallbackStruct>::Get();
 }
 
 UObject* FObjectInitializer::CreateDefaultSubobject(UObject* Outer, FName SubobjectFName, UClass* ReturnType, UClass* ClassToCreateByDefault, bool bIsRequired, bool bAbstract, bool bIsTransient) const
diff --git a/Engine/Source/Runtime/CoreUObject/Public/UObject/Class.h b/Engine/Source/Runtime/CoreUObject/Public/UObject/Class.h
--- a/Engine/Source/Runtime/CoreUObject/Public/UObject/Class.h
+++ b/Engine/Source/Runtime/CoreUObject/Public/UObject/Class.h
@@ -2751,4 +2751,52 @@ inline T* GetMutableDefault(UClass *Class)
 	return (T*)Class->GetDefaultObject();
 }
 
-COREUOBJECT_API UScriptStruct* GetBaseStructure(const TCHAR* Name);
\ No newline at end of file
+template< class T > struct TBaseStructure
+{
+	static UScriptStruct* Get()
+	{
+		static_assert(false, "Unknown base structure, try to use T::StaticStruct()");
+		return nullptr;
+	}
+};
+
+template<> struct TBaseStructure<FRotator>
+{
+	COREUOBJECT_API static UScriptStruct* Get();
+};
+
+template<> struct TBaseStructure<FTransform>
+{
+	COREUOBJECT_API static UScriptStruct* Get();
+};
+
+template<> struct TBaseStructure<FLinearColor>
+{
+	COREUOBJECT_API static UScriptStruct* Get();
+};
+
+template<> struct TBaseStructure<FColor>
+{
+	COREUOBJECT_API static UScriptStruct* Get();
+};
+
+template<> struct  TBaseStructure<FVector>
+{
+	COREUOBJECT_API static UScriptStruct* Get();
+};
+
+template<> struct TBaseStructure<FVector2D>
+{
+	COREUOBJECT_API static UScriptStruct* Get();
+};
+
+template<> struct TBaseStructure<FRandomStream>
+{
+	COREUOBJECT_API static UScriptStruct* Get();
+};
+
+struct FFallbackStruct;
+template<> struct TBaseStructure<FFallbackStruct>
+{
+	COREUOBJECT_API static UScriptStruct* Get();
+};
\ No newline at end of file
diff --git a/Engine/Source/Runtime/Engine/Private/ActorConstruction.cpp b/Engine/Source/Runtime/Engine/Private/ActorConstruction.cpp
--- a/Engine/Source/Runtime/Engine/Private/ActorConstruction.cpp
+++ b/Engine/Source/Runtime/Engine/Private/ActorConstruction.cpp
@@ -23,7 +23,7 @@ static TArray<FRandomStream*> FindRandomStreams(AActor* InActor)
 {
 	check(InActor);
 	TArray<FRandomStream*> OutStreams;
-	UScriptStruct* RandomStreamStruct = GetBaseStructure(TEXT("RandomStream"));
+	UScriptStruct* RandomStreamStruct = TBaseStructure<FRandomStream>::Get();
 	for( TFieldIterator<UStructProperty> It(InActor->GetClass()) ; It ; ++It )
 	{
 		UStructProperty* StructProp = *It;
diff --git a/Engine/Source/Runtime/Engine/Private/KismetSystemLibrary.cpp b/Engine/Source/Runtime/Engine/Private/KismetSystemLibrary.cpp
--- a/Engine/Source/Runtime/Engine/Private/KismetSystemLibrary.cpp
+++ b/Engine/Source/Runtime/Engine/Private/KismetSystemLibrary.cpp
@@ -682,7 +682,7 @@ void UKismetSystemLibrary::SetVectorPropertyByName(UObject* Object, FName Proper
 {
 	if(Object != NULL)
 	{
-		UScriptStruct* VectorStruct = GetBaseStructure(TEXT("Vector"));
+		UScriptStruct* VectorStruct = TBaseStructure<FVector>::Get();
 		UStructProperty* VectorProp = FindField<UStructProperty>(Object->GetClass(), PropertyName);
 		if(VectorProp != NULL && VectorProp->Struct == VectorStruct)
 		{
@@ -695,7 +695,7 @@ void UKismetSystemLibrary::SetRotatorPropertyByName(UObject* Object, FName Prope
 {
 	if(Object != NULL)
 	{
-		UScriptStruct* RotatorStruct = GetBaseStructure(TEXT("Rotator"));
+		UScriptStruct* RotatorStruct = TBaseStructure<FRotator>::Get();
 		UStructProperty* RotatorProp = FindField<UStructProperty>(Object->GetClass(), PropertyName);
 		if(RotatorProp != NULL && RotatorProp->Struct == RotatorStruct)
 		{
@@ -708,7 +708,7 @@ void UKismetSystemLibrary::SetLinearColorPropertyByName(UObject* Object, FName P
 {
 	if(Object != NULL)
 	{
-		UScriptStruct* ColorStruct = GetBaseStructure(TEXT("LinearColor"));
+		UScriptStruct* ColorStruct = TBaseStructure<FLinearColor>::Get();
 		UStructProperty* ColorProp = FindField<UStructProperty>(Object->GetClass(), PropertyName);
 		if(ColorProp != NULL && ColorProp->Struct == ColorStruct)
 		{
@@ -721,7 +721,7 @@ void UKismetSystemLibrary::SetTransformPropertyByName(UObject* Object, FName Pro
 {
 	if(Object != NULL)
 	{
-		UScriptStruct* TransformStruct = GetBaseStructure(TEXT("Transform"));
+		UScriptStruct* TransformStruct = TBaseStructure<FTransform>::Get();
 		UStructProperty* TransformProp = FindField<UStructProperty>(Object->GetClass(), PropertyName);
 		if(TransformProp != NULL && TransformProp->Struct == TransformStruct)
 		{
diff --git a/Engine/Source/Runtime/UMG/Public/Binding/DynamicPropertyPath.h b/Engine/Source/Runtime/UMG/Public/Binding/DynamicPropertyPath.h
--- a/Engine/Source/Runtime/UMG/Public/Binding/DynamicPropertyPath.h
+++ b/Engine/Source/Runtime/UMG/Public/Binding/DynamicPropertyPath.h
@@ -97,7 +97,7 @@ inline bool IsConcreteTypeCompatibleWithReflectedType<FString>(UProperty* Proper
 template<>
 inline bool IsConcreteTypeCompatibleWithReflectedType<FLinearColor>(UProperty* Property)
 {
-	static const UScriptStruct* LinearColorStruct = GetBaseStructure(TEXT("LinearColor"));
+	static const UScriptStruct* LinearColorStruct = TBaseStructure<FLinearColor>::Get();
 	
 	if ( UStructProperty* StructProperty = Cast<UStructProperty>(Property) )
 	{
