Connecting function result node array pins to a reroute node will no longer prevent the node from compiling with all it's output parameters.

K2Node_Knot (Reroute node) is now considered pure and will not be pruned during Skeleton compiles.

diff --git a/Engine/Source/Editor/BlueprintGraph/Classes/K2Node_Knot.h b/Engine/Source/Editor/BlueprintGraph/Classes/K2Node_Knot.h
--- a/Engine/Source/Editor/BlueprintGraph/Classes/K2Node_Knot.h
+++ b/Engine/Source/Editor/BlueprintGraph/Classes/K2Node_Knot.h
@@ -30,6 +30,7 @@ public:
 	virtual void PostReconstructNode() override;
 	virtual int32 GetNodeRefreshPriority() const override { return EBaseNodeRefreshPriority::Low_UsesDependentWildcard; }
 	virtual void GetMenuActions(FBlueprintActionDatabaseRegistrar& ActionRegistrar) const override;
+	virtual bool IsNodePure() const { return true; }
 	// End of UK2Node interface
 
 	UEdGraphPin* GetInputPin() const
