Splitting pins on the Function Result node will no longer compile those pins as result pins.

diff --git a/Engine/Source/Editor/BlueprintGraph/Private/K2Node_FunctionResult.cpp b/Engine/Source/Editor/BlueprintGraph/Private/K2Node_FunctionResult.cpp
--- a/Engine/Source/Editor/BlueprintGraph/Private/K2Node_FunctionResult.cpp
+++ b/Engine/Source/Editor/BlueprintGraph/Private/K2Node_FunctionResult.cpp
@@ -61,9 +61,13 @@ public:
 
 	virtual void RegisterNet(FKismetFunctionContext& Context, UEdGraphPin* Net) override
 	{
-		FBPTerminal* Term = new (Context.Results) FBPTerminal();
-		Term->CopyFromPin(Net, Net->PinName);
-		Context.NetMap.Add(Net, Term);
+		// Do not register as a default any Pin that comes from being Split
+		if (Net->ParentPin == nullptr)
+		{
+			FBPTerminal* Term = new (Context.Results) FBPTerminal();
+			Term->CopyFromPin(Net, Net->PinName);
+			Context.NetMap.Add(Net, Term);
+		}
 	}
 
 	virtual void Compile(FKismetFunctionContext& Context, UEdGraphNode* Node) override
