[UE-11595] Making sure we update the BlueprintActionDatabase when a child class's skeleton gets recompiled (to keep UI from getting out of date).

diff --git a/Engine/Source/Editor/BlueprintGraph/Private/BlueprintActionFilter.cpp b/Engine/Source/Editor/BlueprintGraph/Private/BlueprintActionFilter.cpp
--- a/Engine/Source/Editor/BlueprintGraph/Private/BlueprintActionFilter.cpp
+++ b/Engine/Source/Editor/BlueprintGraph/Private/BlueprintActionFilter.cpp
@@ -331,7 +331,7 @@ namespace BlueprintActionFilterImpl
 	 * 
 	 * @param  Filter			Holds the graph context for this test.
 	 * @param  BlueprintAction	The action you wish to query.
-	 * @return true if the action is an animnotification that is incompatible with the current skeleton
+	 * @return True if the action is an animnotification that is incompatible with the current skeleton
 	 */
 	static bool IsIncompatibleAnimNotification( FBlueprintActionFilter const& Filter, FBlueprintActionInfo& BlueprintAction); 
 
@@ -349,9 +349,19 @@ namespace BlueprintActionFilterImpl
 	 * 
 	 * @param  Filter			Holds the graph context for this test.
 	 * @param  BlueprintAction	The action you wish to query.
-	 * @return true if the macro instance is incompatible with the current graph context
+	 * @return True if the macro instance is incompatible with the current graph context
 	 */
 	static bool IsIncompatibleMacroInstance(FBlueprintActionFilter const& Filter, FBlueprintActionInfo& BlueprintAction);
+
+	/**
+	 * Rejection test to help unblock common crashes, where programmers forget 
+	 * to refresh the BlueprintActionDatabase when a blueprint has been re-compiled.
+	 * 
+	 * @param  Filter			Holds the action/field context for this test.
+	 * @param  BlueprintAction	The action you wish to query.
+	 * @return True if the action is stale (associated with a TRASH or REINST class).
+	 */
+	static bool IsStaleFieldAction(FBlueprintActionFilter const& Filter, FBlueprintActionInfo& BlueprintAction);
 };
 
 //------------------------------------------------------------------------------
@@ -1512,6 +1522,14 @@ static bool BlueprintActionFilterImpl::IsIncompatibleMacroInstance(FBlueprintAct
 	return bIsFilteredOut;
 }
 
+//------------------------------------------------------------------------------
+static bool BlueprintActionFilterImpl::IsStaleFieldAction(FBlueprintActionFilter const& Filter, FBlueprintActionInfo& BlueprintAction)
+{
+	bool const bIsFilteredOut = FBlueprintNodeSpawnerUtils::IsStaleFieldAction(BlueprintAction.NodeSpawner);
+	ensureMsgf(!bIsFilteredOut, TEXT("Invalid BlueprintActionDatabase entry (for %s). Was the database properly updated when this class was compiled?"), *BlueprintAction.GetOwnerClass()->GetName());
+	return bIsFilteredOut;
+}
+
 /*******************************************************************************
  * FBlueprintActionInfo
  ******************************************************************************/
@@ -1677,6 +1695,8 @@ FBlueprintActionFilter::FBlueprintActionFilter(uint32 Flags/*= 0x00*/)
 	//
 	// this test in-particular spawns a template-node and then calls 
 	// AllocateDefaultPins() which is costly, so it should be very last!
+	AddRejectionTest(FRejectionTestDelegate::CreateStatic(IsStaleFieldAction));
+
 	AddRejectionTest(FRejectionTestDelegate::CreateStatic(IsIncompatibleAnimNotification));
 	AddRejectionTest(FRejectionTestDelegate::CreateStatic(IsNodeTemplateSelfFiltered));
 	AddRejectionTest(FRejectionTestDelegate::CreateStatic(IsMissingMatchingPinParam));
diff --git a/Engine/Source/Editor/BlueprintGraph/Private/BlueprintBoundEventNodeSpawner.cpp b/Engine/Source/Editor/BlueprintGraph/Private/BlueprintBoundEventNodeSpawner.cpp
--- a/Engine/Source/Editor/BlueprintGraph/Private/BlueprintBoundEventNodeSpawner.cpp
+++ b/Engine/Source/Editor/BlueprintGraph/Private/BlueprintBoundEventNodeSpawner.cpp
@@ -135,6 +135,14 @@ bool UBlueprintBoundEventNodeSpawner::IsBindingCompatible(UObject const* Binding
 	}
 
 	const UMulticastDelegateProperty* Delegate = GetEventDelegate();
+
+	if ( !ensureMsgf(!FBlueprintNodeSpawnerUtils::IsStaleFieldAction(this), 
+			TEXT("Invalid BlueprintBoundEventNodeSpawner (for %s). Was the action database properly updated when this class was compiled?"), 
+			*Delegate->GetOwnerClass()->GetName()))
+	{
+		return false;
+	}
+
 	UClass* DelegateOwner = Delegate->GetOwnerClass()->GetAuthoritativeClass();
 	UClass* BindingClass  = FBlueprintNodeSpawnerUtils::GetBindingClass(BindingCandidate)->GetAuthoritativeClass();
 
diff --git a/Engine/Source/Editor/BlueprintGraph/Private/BlueprintFunctionNodeSpawner.cpp b/Engine/Source/Editor/BlueprintGraph/Private/BlueprintFunctionNodeSpawner.cpp
--- a/Engine/Source/Editor/BlueprintGraph/Private/BlueprintFunctionNodeSpawner.cpp
+++ b/Engine/Source/Editor/BlueprintGraph/Private/BlueprintFunctionNodeSpawner.cpp
@@ -428,6 +428,13 @@ bool UBlueprintFunctionNodeSpawner::IsBindingCompatible(UObject const* BindingCa
 	UFunction const* Function = GetFunction();
 	checkSlow(Function != nullptr);
 
+	if ( !ensureMsgf(!FBlueprintNodeSpawnerUtils::IsStaleFieldAction(this), 
+			TEXT("Invalid BlueprintFunctionNodeSpawner (for %s). Was the action database properly updated when this class was compiled?"), 
+			*Function->GetOwnerClass()->GetName()) )
+	{
+		return false;
+	}
+
 	bool const bNodeTypeMatches = (NodeClass == UK2Node_CallFunction::StaticClass());
 	bool bClassOwnerMatches = false;
 
diff --git a/Engine/Source/Editor/BlueprintGraph/Private/BlueprintNodeSpawnerUtils.cpp b/Engine/Source/Editor/BlueprintGraph/Private/BlueprintNodeSpawnerUtils.cpp
--- a/Engine/Source/Editor/BlueprintGraph/Private/BlueprintNodeSpawnerUtils.cpp
+++ b/Engine/Source/Editor/BlueprintGraph/Private/BlueprintNodeSpawnerUtils.cpp
@@ -78,3 +78,23 @@ UClass* FBlueprintNodeSpawnerUtils::GetBindingClass(const UObject* Binding)
 	}
 	return BindingClass;
 }
+
+//------------------------------------------------------------------------------
+bool FBlueprintNodeSpawnerUtils::IsStaleFieldAction(UBlueprintNodeSpawner const* BlueprintAction)
+{
+	bool bHasStaleAssociatedField= false;
+
+	const UField* AssociatedField = GetAssociatedField(BlueprintAction);
+	if (AssociatedField != nullptr)
+	{
+		UClass* ClassOwner = AssociatedField->GetOwnerClass();
+		if (ClassOwner != nullptr)
+		{
+			// check to see if this field belongs to a TRASH or REINST class,
+			// maybe to a class that was thrown out because of a hot-reload?
+			bHasStaleAssociatedField = ClassOwner->HasAnyClassFlags(CLASS_NewerVersionExists) || (ClassOwner->GetOutermost() == GetTransientPackage());
+		}
+	}
+
+	return bHasStaleAssociatedField;
+}
diff --git a/Engine/Source/Editor/BlueprintGraph/Private/BlueprintNodeSpawnerUtils.h b/Engine/Source/Editor/BlueprintGraph/Private/BlueprintNodeSpawnerUtils.h
--- a/Engine/Source/Editor/BlueprintGraph/Private/BlueprintNodeSpawnerUtils.h
+++ b/Engine/Source/Editor/BlueprintGraph/Private/BlueprintNodeSpawnerUtils.h
@@ -45,5 +45,14 @@ struct FBlueprintNodeSpawnerUtils
 	 * @return A UClass that corresponds to the supplied binding.
 	 */
 	static UClass* GetBindingClass(const UObject* Binding);
+
+	/**
+	 * Checks if the node-spawner's associated action is stale (meaning it 
+	 * belongs to a TRASH or REINST class).
+	 * 
+	 * @param  BlueprintAction    The node-spawner you want to check.
+	 * @return True if the action is stale (associated with a TRASH or REINST class).
+	 */
+	static bool IsStaleFieldAction(UBlueprintNodeSpawner const* BlueprintAction);
 };
 
diff --git a/Engine/Source/Editor/UnrealEd/Private/Kismet2/BlueprintEditorUtils.cpp b/Engine/Source/Editor/UnrealEd/Private/Kismet2/BlueprintEditorUtils.cpp
--- a/Engine/Source/Editor/UnrealEd/Private/Kismet2/BlueprintEditorUtils.cpp
+++ b/Engine/Source/Editor/UnrealEd/Private/Kismet2/BlueprintEditorUtils.cpp
@@ -2019,6 +2019,8 @@ void FBlueprintEditorUtils::MarkBlueprintAsStructurallyModified(UBlueprint* Blue
 					CompileOptions.CompileType = EKismetCompileType::SkeletonOnly;
 					Compiler.CompileBlueprint(SkelBlueprint, CompileOptions, Results);
 
+					SkelBlueprint->BroadcastCompiled();
+
 					SkeletalRecompileChildren(ChildrenOfClass);
 				}
 			}
