[UE-15293] Fix array index out-of-bounds on ForEach* macro expansions when accessing the array element on completion.

change summary:
- modified ForEachLoop and ForEachLoopWithBreak macros to maintain a separate value for the array index that won't increment beyond the input array size
- modified UK2Node_ForEachELementInEnum to expand to a similar construct as the ForEachLoop* macros

diff --git a/Engine/Source/Editor/BlueprintGraph/Private/K2Node_ForEachElementInEnum.cpp b/Engine/Source/Editor/BlueprintGraph/Private/K2Node_ForEachElementInEnum.cpp
--- a/Engine/Source/Editor/BlueprintGraph/Private/K2Node_ForEachElementInEnum.cpp
+++ b/Engine/Source/Editor/BlueprintGraph/Private/K2Node_ForEachElementInEnum.cpp
@@ -20,16 +20,18 @@ struct FForExpandNodeHelper
 	UEdGraphPin* InsideLoopExecOutPin;
 	UEdGraphPin* LoopCompleteOutExecPin;
 
-	UEdGraphPin* IndexOutPin;
-	// for(Index = 0; Index < IndexLimit; ++Index)
-	UEdGraphPin* IndexLimitInPin;
+	UEdGraphPin* ArrayIndexOutPin;
+	UEdGraphPin* LoopCounterOutPin;
+	// for(LoopCounter = 0; LoopCounter < LoopCounterLimit; ++LoopCounter)
+	UEdGraphPin* LoopCounterLimitInPin;
 
 	FForExpandNodeHelper()
-		: StartLoopExecInPin(NULL)
-		, InsideLoopExecOutPin(NULL)
-		, LoopCompleteOutExecPin(NULL)
-		, IndexOutPin(NULL)
-		, IndexLimitInPin(NULL)
+		: StartLoopExecInPin(nullptr)
+		, InsideLoopExecOutPin(nullptr)
+		, LoopCompleteOutExecPin(nullptr)
+		, ArrayIndexOutPin(nullptr)
+		, LoopCounterOutPin(nullptr)
+		, LoopCounterLimitInPin(nullptr)
 	{ }
 
 	bool BuildLoop(UK2Node* Node, FKismetCompilerContext& CompilerContext, UEdGraph* SourceGraph)
@@ -39,25 +41,39 @@ struct FForExpandNodeHelper
 
 		bool bResult = true;
 
-		// Create int Iterator
-		UK2Node_TemporaryVariable* IndexNode = CompilerContext.SpawnIntermediateNode<UK2Node_TemporaryVariable>(Node, SourceGraph);
-		IndexNode->VariableType.PinCategory = Schema->PC_Int;
-		IndexNode->AllocateDefaultPins();
-		IndexOutPin = IndexNode->GetVariablePin();
-		check(IndexOutPin);
-
-		// Initialize iterator
-		UK2Node_AssignmentStatement* IteratorInitialize = CompilerContext.SpawnIntermediateNode<UK2Node_AssignmentStatement>(Node, SourceGraph);
-		IteratorInitialize->AllocateDefaultPins();
-		IteratorInitialize->GetValuePin()->DefaultValue = TEXT("0");
-		bResult &= Schema->TryCreateConnection(IndexOutPin, IteratorInitialize->GetVariablePin());
-		StartLoopExecInPin = IteratorInitialize->GetExecPin();
+		// Create int Loop Counter
+		UK2Node_TemporaryVariable* LoopCounterNode = CompilerContext.SpawnIntermediateNode<UK2Node_TemporaryVariable>(Node, SourceGraph);
+		LoopCounterNode->VariableType.PinCategory = Schema->PC_Int;
+		LoopCounterNode->AllocateDefaultPins();
+		LoopCounterOutPin = LoopCounterNode->GetVariablePin();
+		check(LoopCounterOutPin);
+
+		// Initialize loop counter
+		UK2Node_AssignmentStatement* LoopCounterInitialize = CompilerContext.SpawnIntermediateNode<UK2Node_AssignmentStatement>(Node, SourceGraph);
+		LoopCounterInitialize->AllocateDefaultPins();
+		LoopCounterInitialize->GetValuePin()->DefaultValue = TEXT("0");
+		bResult &= Schema->TryCreateConnection(LoopCounterOutPin, LoopCounterInitialize->GetVariablePin());
+		StartLoopExecInPin = LoopCounterInitialize->GetExecPin();
 		check(StartLoopExecInPin);
 
+		// Create int Array Index
+		UK2Node_TemporaryVariable* ArrayIndexNode = CompilerContext.SpawnIntermediateNode<UK2Node_TemporaryVariable>(Node, SourceGraph);
+		ArrayIndexNode->VariableType.PinCategory = Schema->PC_Int;
+		ArrayIndexNode->AllocateDefaultPins();
+		ArrayIndexOutPin = ArrayIndexNode->GetVariablePin();
+		check(ArrayIndexOutPin);
+
+		// Initialize array index
+		UK2Node_AssignmentStatement* ArrayIndexInitialize = CompilerContext.SpawnIntermediateNode<UK2Node_AssignmentStatement>(Node, SourceGraph);
+		ArrayIndexInitialize->AllocateDefaultPins();
+		ArrayIndexInitialize->GetValuePin()->DefaultValue = TEXT("0");
+		bResult &= Schema->TryCreateConnection(ArrayIndexOutPin, ArrayIndexInitialize->GetVariablePin());
+		bResult &= Schema->TryCreateConnection(LoopCounterInitialize->GetThenPin(), ArrayIndexInitialize->GetExecPin());
+
 		// Do loop branch
 		UK2Node_IfThenElse* Branch = CompilerContext.SpawnIntermediateNode<UK2Node_IfThenElse>(Node, SourceGraph);
 		Branch->AllocateDefaultPins();
-		bResult &= Schema->TryCreateConnection(IteratorInitialize->GetThenPin(), Branch->GetExecPin());
+		bResult &= Schema->TryCreateConnection(ArrayIndexInitialize->GetThenPin(), Branch->GetExecPin());
 		LoopCompleteOutExecPin = Branch->GetElsePin();
 		check(LoopCompleteOutExecPin);
 
@@ -66,31 +82,38 @@ struct FForExpandNodeHelper
 		Condition->SetFromFunction(UKismetMathLibrary::StaticClass()->FindFunctionByName(TEXT("Less_IntInt")));
 		Condition->AllocateDefaultPins();
 		bResult &= Schema->TryCreateConnection(Condition->GetReturnValuePin(), Branch->GetConditionPin());
-		bResult &= Schema->TryCreateConnection(Condition->FindPinChecked(TEXT("A")), IndexOutPin);
-		IndexLimitInPin = Condition->FindPinChecked(TEXT("B"));
-		check(IndexLimitInPin);
+		bResult &= Schema->TryCreateConnection(Condition->FindPinChecked(TEXT("A")), LoopCounterOutPin);
+		LoopCounterLimitInPin = Condition->FindPinChecked(TEXT("B"));
+		check(LoopCounterLimitInPin);
+
+		// Array Index assigned
+		UK2Node_AssignmentStatement* ArrayIndexAssign = CompilerContext.SpawnIntermediateNode<UK2Node_AssignmentStatement>(Node, SourceGraph);
+		ArrayIndexAssign->AllocateDefaultPins();
+		bResult &= Schema->TryCreateConnection(Branch->GetThenPin(), ArrayIndexAssign->GetExecPin());
+		bResult &= Schema->TryCreateConnection(ArrayIndexAssign->GetVariablePin(), ArrayIndexOutPin);
+		bResult &= Schema->TryCreateConnection(ArrayIndexAssign->GetValuePin(), LoopCounterOutPin);
 
 		// body sequence
 		UK2Node_ExecutionSequence* Sequence = CompilerContext.SpawnIntermediateNode<UK2Node_ExecutionSequence>(Node, SourceGraph);
 		Sequence->AllocateDefaultPins();
-		bResult &= Schema->TryCreateConnection(Branch->GetThenPin(), Sequence->GetExecPin());
+		bResult &= Schema->TryCreateConnection(ArrayIndexAssign->GetThenPin(), Sequence->GetExecPin());
 		InsideLoopExecOutPin = Sequence->GetThenPinGivenIndex(0);
 		check(InsideLoopExecOutPin);
 
-		// Iterator increment
+		// Loop Counter increment
 		UK2Node_CallFunction* Increment = CompilerContext.SpawnIntermediateNode<UK2Node_CallFunction>(Node, SourceGraph); 
 		Increment->SetFromFunction(UKismetMathLibrary::StaticClass()->FindFunctionByName(TEXT("Add_IntInt")));
 		Increment->AllocateDefaultPins();
-		bResult &= Schema->TryCreateConnection(Increment->FindPinChecked(TEXT("A")), IndexOutPin);
+		bResult &= Schema->TryCreateConnection(Increment->FindPinChecked(TEXT("A")), LoopCounterOutPin);
 		Increment->FindPinChecked(TEXT("B"))->DefaultValue = TEXT("1");
 
-		// Iterator assigned
-		UK2Node_AssignmentStatement* IteratorAssign = CompilerContext.SpawnIntermediateNode<UK2Node_AssignmentStatement>(Node, SourceGraph);
-		IteratorAssign->AllocateDefaultPins();
-		bResult &= Schema->TryCreateConnection(IteratorAssign->GetExecPin(), Sequence->GetThenPinGivenIndex(1));
-		bResult &= Schema->TryCreateConnection(IteratorAssign->GetVariablePin(), IndexOutPin);
-		bResult &= Schema->TryCreateConnection(IteratorAssign->GetValuePin(), Increment->GetReturnValuePin());
-		bResult &= Schema->TryCreateConnection(IteratorAssign->GetThenPin(), Branch->GetExecPin());
+		// Loop Counter assigned
+		UK2Node_AssignmentStatement* LoopCounterAssign = CompilerContext.SpawnIntermediateNode<UK2Node_AssignmentStatement>(Node, SourceGraph);
+		LoopCounterAssign->AllocateDefaultPins();
+		bResult &= Schema->TryCreateConnection(LoopCounterAssign->GetExecPin(), Sequence->GetThenPinGivenIndex(1));
+		bResult &= Schema->TryCreateConnection(LoopCounterAssign->GetVariablePin(), LoopCounterOutPin);
+		bResult &= Schema->TryCreateConnection(LoopCounterAssign->GetValuePin(), Increment->GetReturnValuePin());
+		bResult &= Schema->TryCreateConnection(LoopCounterAssign->GetThenPin(), Branch->GetExecPin());
 
 		return bResult;
 	}
@@ -177,13 +200,13 @@ void UK2Node_ForEachElementInEnum::ExpandNode(class FKismetCompilerContext& Comp
 	UK2Node_GetNumEnumEntries* GetNumEnumEntries = CompilerContext.SpawnIntermediateNode<UK2Node_GetNumEnumEntries>(this, SourceGraph);
 	GetNumEnumEntries->Enum = Enum;
 	GetNumEnumEntries->AllocateDefaultPins();
-	bool bResult = Schema->TryCreateConnection(GetNumEnumEntries->FindPinChecked(Schema->PN_ReturnValue), ForLoop.IndexLimitInPin);
+	bool bResult = Schema->TryCreateConnection(GetNumEnumEntries->FindPinChecked(Schema->PN_ReturnValue), ForLoop.LoopCounterLimitInPin);
 
 	UK2Node_CallFunction* Conv_Func = CompilerContext.SpawnIntermediateNode<UK2Node_CallFunction>(this, SourceGraph);
 	FName Conv_Func_Name = GET_FUNCTION_NAME_CHECKED(UKismetMathLibrary, Conv_IntToByte);
 	Conv_Func->SetFromFunction(UKismetMathLibrary::StaticClass()->FindFunctionByName(Conv_Func_Name));
 	Conv_Func->AllocateDefaultPins();
-	bResult &= Schema->TryCreateConnection(Conv_Func->FindPinChecked(TEXT("InInt")), ForLoop.IndexOutPin);
+	bResult &= Schema->TryCreateConnection(Conv_Func->FindPinChecked(TEXT("InInt")), ForLoop.ArrayIndexOutPin);
 
 	UK2Node_CastByteToEnum* CastByteToEnum = CompilerContext.SpawnIntermediateNode<UK2Node_CastByteToEnum>(this, SourceGraph);
 	CastByteToEnum->Enum = Enum;
