UE-14387 Support for AssetPtr references in Blueprints
Early version:
- New BP types: Asset and AssetClass
- Now implicit conversion (in Blueprint) from Asset to Object
- Explicit conversion node: K2Node_ConvertAsset
- EX_AssetConst script instruction to store StringAssetReference in bytecode.

diff --git a/Engine/Source/Editor/BlueprintGraph/Classes/EdGraphSchema_K2.h b/Engine/Source/Editor/BlueprintGraph/Classes/EdGraphSchema_K2.h
--- a/Engine/Source/Editor/BlueprintGraph/Classes/EdGraphSchema_K2.h
+++ b/Engine/Source/Editor/BlueprintGraph/Classes/EdGraphSchema_K2.h
@@ -201,6 +201,7 @@ class BLUEPRINTGRAPH_API UEdGraphSchema_K2 : public UEdGraphSchema
 	static const FString PC_Boolean;
 	static const FString PC_Byte;
 	static const FString PC_Class;    // SubCategoryObject is the MetaClass of the Class passed thru this pin, or SubCategory can be 'self'. The DefaultValue string should always be empty, use DefaultObject.
+	static const FString PC_AssetClass;
 	static const FString PC_Int;
 	static const FString PC_Float;
 	static const FString PC_Name;
@@ -208,6 +209,7 @@ class BLUEPRINTGRAPH_API UEdGraphSchema_K2 : public UEdGraphSchema
 	static const FString PC_MCDelegate;  // SubCategoryObject is the UFunction of the delegate signature
 	static const FString PC_Object;    // SubCategoryObject is the Class of the object passed thru this pin, or SubCategory can be 'self'. The DefaultValue string should always be empty, use DefaultObject.
 	static const FString PC_Interface;	// SubCategoryObject is the Class of the object passed thru this pin.
+	static const FString PC_Asset;		// SubCategoryObject is the Class of the AssetPtr passed thru this pin.
 	static const FString PC_String;
 	static const FString PC_Text;
 	static const FString PC_Struct;    // SubCategoryObject is the ScriptStruct of the struct passed thru this pin, 'self' is not a valid SubCategory. DefaultObject should always be empty, the DefaultValue string may be used for supported structs.
diff --git a/Engine/Source/Editor/BlueprintGraph/Classes/K2Node_ConvertAsset.h b/Engine/Source/Editor/BlueprintGraph/Classes/K2Node_ConvertAsset.h
new file mode 100644
--- /dev/null
+++ b/Engine/Source/Editor/BlueprintGraph/Classes/K2Node_ConvertAsset.h
@@ -0,0 +1,30 @@
+// Copyright 1998-2015 Epic Games, Inc. All Rights Reserved.
+
+
+#pragma once
+#include "K2Node.h"
+#include "K2Node_ConvertAsset.generated.h"
+
+UCLASS(MinimalAPI)
+class UK2Node_ConvertAsset : public UK2Node
+{
+	GENERATED_BODY()
+public:
+	UPROPERTY()
+	TSubclassOf<class UObject>  TargetType;
+
+	UPROPERTY()
+	bool bIsAssetClass;
+
+	// UEdGraphNode interface
+	virtual void AllocateDefaultPins() override;
+	// End of UEdGraphNode interface
+
+	// UK2Node interface
+	virtual bool IsNodePure() const override { return true; }
+	virtual bool ShouldDrawCompact() const override { return true; }
+	virtual void ExpandNode(class FKismetCompilerContext& CompilerContext, UEdGraph* SourceGraph) override;
+	virtual bool HasExternalDependencies(TArray<class UStruct*>* OptionalOutput) const override;
+	virtual FText GetCompactNodeTitle() const override;
+	// End of UK2Node interface
+};
\ No newline at end of file
diff --git a/Engine/Source/Editor/BlueprintGraph/Private/EdGraphSchema_K2.cpp b/Engine/Source/Editor/BlueprintGraph/Private/EdGraphSchema_K2.cpp
--- a/Engine/Source/Editor/BlueprintGraph/Private/EdGraphSchema_K2.cpp
+++ b/Engine/Source/Editor/BlueprintGraph/Private/EdGraphSchema_K2.cpp
@@ -30,6 +30,7 @@
 #include "K2Node_GetEnumeratorNameAsString.h"
 #include "K2Node_Tunnel.h"
 #include "K2Node_SetFieldsInStruct.h"
+#include "K2Node_ConvertAsset.h"
 #include "GenericCommands.h"
 
 
@@ -153,7 +154,7 @@ private:
 				}
 			}
 		}
-		else if (Type == UEdGraphSchema_K2::PC_Class)
+		else if (Type == UEdGraphSchema_K2::PC_Class || Type == UEdGraphSchema_K2::PC_AssetClass)
 		{
 			// Generate a list of all potential objects which have "BlueprintType=true" in their metadata
 			for (TObjectIterator<UClass> ClassIt; ClassIt; ++ClassIt)
@@ -165,7 +166,7 @@ private:
 				}
 			}
 		}
-		else if (Type == UEdGraphSchema_K2::PC_Object)
+		else if (Type == UEdGraphSchema_K2::PC_Object || Type == UEdGraphSchema_K2::PC_Asset)
 		{
 			// Generate a list of all potential objects which have "BlueprintType=true" in their metadata (that aren't interfaces)
 			for (TObjectIterator<UClass> ClassIt; ClassIt; ++ClassIt)
@@ -244,7 +245,7 @@ public:
 		{
 			AssetRegistrySearchClass = UUserDefinedEnum::StaticClass();
 		}
-		else if(Schema->PC_Object == CategoryName || Schema->PC_Class == CategoryName || Schema->PC_Interface == CategoryName)
+		else if (Schema->PC_Object == CategoryName || Schema->PC_Class == CategoryName || Schema->PC_Interface == CategoryName || Schema->PC_Asset == CategoryName || Schema->PC_AssetClass == CategoryName)
 		{
 			AssetRegistrySearchClass = UBlueprint::StaticClass();
 		}
@@ -266,7 +267,7 @@ public:
 				}
 			}
 		}
-		else if(Schema->PC_Object == CategoryName || Schema->PC_Class == CategoryName || Schema->PC_Interface == CategoryName)
+		else if (Schema->PC_Object == CategoryName || Schema->PC_Class == CategoryName || Schema->PC_Interface == CategoryName || Schema->PC_Asset == CategoryName || Schema->PC_AssetClass == CategoryName)
 		{
 			const FString BPTypeAllowed = (Schema->PC_Interface == CategoryName)? TEXT("BPTYPE_Interface") : TEXT("BPTYPE_Normal");
 
@@ -448,6 +449,8 @@ const FString UEdGraphSchema_K2::PC_Text(TEXT("text"));
 const FString UEdGraphSchema_K2::PC_Struct(TEXT("struct"));
 const FString UEdGraphSchema_K2::PC_Wildcard(TEXT("wildcard"));
 const FString UEdGraphSchema_K2::PC_Enum(TEXT("enum"));
+const FString UEdGraphSchema_K2::PC_Asset(TEXT("asset"));
+const FString UEdGraphSchema_K2::PC_AssetClass(TEXT("assetclass"));
 const FString UEdGraphSchema_K2::PSC_Self(TEXT("self"));
 const FString UEdGraphSchema_K2::PSC_Index(TEXT("index"));
 const FString UEdGraphSchema_K2::PN_Execute(TEXT("execute"));
@@ -2221,7 +2224,7 @@ bool UEdGraphSchema_K2::FindSpecializedConversionNode(const UEdGraphPin* OutputP
 		else
 		{
 			UClass* InputClass  = Cast<UClass>(InputType.PinSubCategoryObject.Get());
-			UClass* OutputClass = Cast<UClass>(InputType.PinSubCategoryObject.Get());
+			UClass* OutputClass = Cast<UClass>(OutputType.PinSubCategoryObject.Get());
 
 			if ((OutputType.PinCategory == PC_Interface) && (InputType.PinCategory == PC_Object))
 			{
@@ -2243,6 +2246,23 @@ bool UEdGraphSchema_K2::FindSpecializedConversionNode(const UEdGraphPin* OutputP
 				DynCastNode->SetPurity(true);
 				TargetNode = DynCastNode;
 			}
+
+			if (!bCanConvert && InputClass && OutputClass && OutputClass->IsChildOf(InputClass))
+			{
+				const bool ConvertAsset = (OutputType.PinCategory == PC_Asset) && (InputType.PinCategory == PC_Object);
+				const bool ConvertAssetClass = (OutputType.PinCategory == PC_AssetClass) && (InputType.PinCategory == PC_Class);
+				if (ConvertAsset || ConvertAssetClass)
+				{
+					bCanConvert = true;
+					if (bCreateNode)
+					{
+						UK2Node_ConvertAsset* ConvertAsset = NewObject<UK2Node_ConvertAsset>();
+						ConvertAsset->TargetType = InputClass;
+						ConvertAsset->bIsAssetClass = ConvertAssetClass;
+						TargetNode = ConvertAsset;
+					}
+				}
+			}
 		}
 	}
 
@@ -2419,7 +2439,7 @@ bool UEdGraphSchema_K2::DefaultValueSimpleValidation(const FEdGraphPinType& PinT
 			}
 		}
 	}
-	else if (PinCategory == PC_Class)
+	else if ((PinCategory == PC_Class) || (PinCategory == PC_AssetClass))
 	{
 		// Should have an object set but no string
 		if(!NewDefaultValue.IsEmpty())
@@ -2482,9 +2502,9 @@ bool UEdGraphSchema_K2::DefaultValueSimpleValidation(const FEdGraphPinType& PinT
 	{
 		// Anything is allowed
 	}
-	else if ((PinCategory == PC_Object) || (PinCategory == PC_Interface))
+	else if ((PinCategory == PC_Object) || (PinCategory == PC_Interface) || (PinCategory == PC_Asset))
 	{
-		if(PinSubCategoryObject == NULL && PinSubCategory != PSC_Self)
+		if (PinSubCategoryObject == NULL && (PinSubCategory != PSC_Self))
 		{
 			DVSV_RETURN_MSG( FString::Printf(TEXT("PinSubCategoryObject on pin '%s' is NULL and PinSubCategory is '%s' not 'self'"), *(PinName), *PinSubCategory) );
 		}
@@ -2506,6 +2526,11 @@ bool UEdGraphSchema_K2::DefaultValueSimpleValidation(const FEdGraphPinType& PinT
 		{
 			DVSV_RETURN_MSG( FString::Printf(TEXT("%s isn't a %s (specified on pin %s)"), *NewDefaultObject->GetPathName(), *ObjectClass->GetName(), *(PinName)) );		
 		}
+
+		if ((PinCategory == PC_Asset) && NewDefaultObject && !NewDefaultObject->IsAsset())
+		{
+			DVSV_RETURN_MSG(FString::Printf(TEXT("%s is not an asset (specified on pin %s)"), *NewDefaultObject->GetPathName(), *(PinName)));
+		}
 	}
 	else if (PinCategory == PC_String)
 	{
@@ -2667,6 +2692,14 @@ FLinearColor UEdGraphSchema_K2::GetPinTypeColor(const FEdGraphPinType& PinType)
 	{
 		return Settings->NamePinTypeColor;
 	}
+	else if (TypeString == PC_Asset)
+	{
+		return Settings->AssetPinTypeColor;
+	}
+	else if (TypeString == PC_AssetClass)
+	{
+		return Settings->AssetClassPinTypeColor;
+	}
 	else if (TypeString == PC_Delegate)
 	{
 		return Settings->DelegatePinTypeColor;
@@ -3008,9 +3041,14 @@ bool UEdGraphSchema_K2::ConvertPropertyToPinType(const UProperty* Property, /*ou
 	}
 	else if (const UAssetClassProperty* AssetClassProperty = Cast<const UAssetClassProperty>(TestProperty))
 	{
-		TypeOut.PinCategory = PC_Class;
+		TypeOut.PinCategory = PC_AssetClass;
 		TypeOut.PinSubCategoryObject = AssetClassProperty->MetaClass;
 	}
+	else if (const UAssetObjectProperty* AssetObjectProperty = Cast<const UAssetObjectProperty>(TestProperty))
+	{
+		TypeOut.PinCategory = PC_Asset;
+		TypeOut.PinSubCategoryObject = AssetObjectProperty->PropertyClass;
+	}
 	else if (const UObjectPropertyBase* ObjectProperty = Cast<const UObjectPropertyBase>(TestProperty))
 	{
 		TypeOut.PinCategory = PC_Object;
@@ -3164,6 +3202,8 @@ FText UEdGraphSchema_K2::GetCategoryText(const FString& Category, const bool bFo
 		CategoryDescriptions.Add(PC_Struct, LOCTEXT("StructCategory","Structure"));
 		CategoryDescriptions.Add(PC_Wildcard, LOCTEXT("WildcardCategory","Wildcard"));
 		CategoryDescriptions.Add(PC_Enum, LOCTEXT("EnumCategory","Enum"));
+		CategoryDescriptions.Add(PC_Asset, LOCTEXT("AssetCategory", "Asset"));
+		CategoryDescriptions.Add(PC_AssetClass, LOCTEXT("AssetClassCategory", "Class Asset"));
 	}
 
 	if (bForMenu)
@@ -3294,6 +3334,8 @@ void UEdGraphSchema_K2::GetVariableTypeTree( TArray< TSharedPtr<FPinTypeTreeInfo
 	// Add the types that have subtrees
 	TypeTree.Add( MakeShareable( new FPinTypeTreeInfo(GetCategoryText(PC_Struct, true), PC_Struct, this, LOCTEXT("StructType", "Struct (value) types."), true) ) );
 	TypeTree.Add( MakeShareable( new FPinTypeTreeInfo(GetCategoryText(PC_Object, true), PC_Object, this, LOCTEXT("ObjectType", "Object pointer."), true) ) );
+	TypeTree.Add( MakeShareable( new FPinTypeTreeInfo(GetCategoryText(PC_Asset, true), PC_Asset, this, LOCTEXT("AssetType", "Asset pointer."), true) ) );
+	TypeTree.Add( MakeShareable( new FPinTypeTreeInfo(GetCategoryText(PC_AssetClass, true), PC_AssetClass, this, LOCTEXT("ClassAssetType", "Class asset."), true) ) );
 	TypeTree.Add( MakeShareable( new FPinTypeTreeInfo(GetCategoryText(PC_Interface, true), PC_Interface, this, LOCTEXT("InterfaceType", "Interface pointer."), true) ) );
 	TypeTree.Add( MakeShareable( new FPinTypeTreeInfo(GetCategoryText(PC_Class, true), PC_Class, this, LOCTEXT("ClassType", "Class pointers."), true) ) );
 	TypeTree.Add( MakeShareable( new FPinTypeTreeInfo(GetCategoryText(PC_Enum, true), PC_Enum, this, LOCTEXT("EnumType", "Enumeration types."), true) ) );
@@ -3325,7 +3367,7 @@ void UEdGraphSchema_K2::GetVariableIndexTypeTree( TArray< TSharedPtr<FPinTypeTre
 
 bool UEdGraphSchema_K2::DoesTypeHaveSubtypes(const FString& Category) const
 {
-	return (Category == PC_Struct) || (Category == PC_Object) || (Category == PC_Interface) || (Category == PC_Class) || (Category == PC_Enum);
+	return (Category == PC_Struct) || (Category == PC_Object) || (Category == PC_Asset) || (Category == PC_AssetClass) || (Category == PC_Interface) || (Category == PC_Class) || (Category == PC_Enum);
 }
 
 struct FWildcardArrayPinHelper
@@ -3392,6 +3434,16 @@ bool UEdGraphSchema_K2::ArePinTypesCompatible(const FEdGraphPinType& Output, con
 			
 			return OutputClass->IsChildOf(InputClass);
 		}
+		else if (((Output.PinCategory == PC_Asset) && (Input.PinCategory == PC_Asset))
+			|| ((Output.PinCategory == PC_AssetClass) && (Input.PinCategory == PC_AssetClass)))
+		{
+			const UClass* OutputObject = (Output.PinSubCategory == PSC_Self) ? CallingContext : Cast<const UClass>(Output.PinSubCategoryObject.Get());
+			const UClass* InputObject = (Input.PinSubCategory == PSC_Self) ? CallingContext : Cast<const UClass>(Input.PinSubCategoryObject.Get());
+			if ((OutputObject != NULL) && (InputObject != NULL))
+			{
+				return OutputObject->IsChildOf(InputObject);
+			}
+		}
 		else if ((Output.PinCategory == PC_Object) || (Output.PinCategory == PC_Struct) || (Output.PinCategory == PC_Class))
 		{
 			// Subcategory mismatch, but the two could be castable
@@ -3996,6 +4048,16 @@ UFunction* UEdGraphSchema_K2::FindSetVariableByNameFunction(const FEdGraphPinTyp
 		static FName SetTextName(GET_FUNCTION_NAME_CHECKED(UKismetSystemLibrary, SetTextPropertyByName));
 		SetFunctionName = SetTextName;
 	}
+	else if (PinType.PinCategory == K2Schema->PC_Asset)
+	{
+		static FName SetAssetName(GET_FUNCTION_NAME_CHECKED(UKismetSystemLibrary, SetAssetPropertyByName));
+		SetFunctionName = SetAssetName;
+	}
+	else if (PinType.PinCategory == K2Schema->PC_AssetClass)
+	{
+		static FName SetAssetClassName(GET_FUNCTION_NAME_CHECKED(UKismetSystemLibrary, SetAssetClassPropertyByName));
+		SetFunctionName = SetAssetClassName;
+	}
 	else if(PinType.PinCategory == K2Schema->PC_Name)
 	{
 		static FName SetNameName(GET_FUNCTION_NAME_CHECKED(UKismetSystemLibrary, SetNamePropertyByName));
diff --git a/Engine/Source/Editor/BlueprintGraph/Private/K2Node_ConvertAsset.cpp b/Engine/Source/Editor/BlueprintGraph/Private/K2Node_ConvertAsset.cpp
new file mode 100644
--- /dev/null
+++ b/Engine/Source/Editor/BlueprintGraph/Private/K2Node_ConvertAsset.cpp
@@ -0,0 +1,111 @@
+// Copyright 1998-2015 Epic Games, Inc. All Rights Reserved.
+
+
+#include "BlueprintGraphPrivatePCH.h"
+#include "Kismet/KismetSystemLibrary.h"
+#include "K2Node_ClassDynamicCast.h"
+#include "K2Node_ConvertAsset.h"
+
+#define LOCTEXT_NAMESPACE "K2Node_ConvertAsset"
+
+namespace UK2Node_ConvertAssetImpl
+{
+	static const FString InputPinName("Asset");
+	static const FString OutputPinName("Object");
+}
+
+void UK2Node_ConvertAsset::AllocateDefaultPins()
+{
+	const UEdGraphSchema_K2* K2Schema = Cast<UEdGraphSchema_K2>(GetSchema());
+	const bool bReferenceObsoleteClass = !TargetType || TargetType->HasAnyClassFlags(CLASS_NewerVersionExists);
+	if (bReferenceObsoleteClass)
+	{
+		Message_Error(FString::Printf(TEXT("Node '%s' references obsolete class '%s'"), *GetPathName(), *GetPathNameSafe(TargetType)));
+	}
+	ensure(!bReferenceObsoleteClass);
+
+	if (TargetType && K2Schema)
+	{
+		if (bIsAssetClass)
+		{
+			CreatePin(EGPD_Input, K2Schema->PC_AssetClass, TEXT(""), *TargetType, false, false, UK2Node_ConvertAssetImpl::InputPinName);
+			CreatePin(EGPD_Output, K2Schema->PC_Class, TEXT(""), *TargetType, false, false, UK2Node_ConvertAssetImpl::OutputPinName);
+		}
+		else
+		{
+			CreatePin(EGPD_Input, K2Schema->PC_Asset, TEXT(""), *TargetType, false, false, UK2Node_ConvertAssetImpl::InputPinName);
+			CreatePin(EGPD_Output, K2Schema->PC_Object, TEXT(""), *TargetType, false, false, UK2Node_ConvertAssetImpl::OutputPinName);
+		}
+	}
+}
+
+void UK2Node_ConvertAsset::ExpandNode(class FKismetCompilerContext& CompilerContext, UEdGraph* SourceGraph)
+{
+	Super::ExpandNode(CompilerContext, SourceGraph);
+	const UEdGraphSchema_K2* Schema = CompilerContext.GetSchema();
+
+	bool bIsErrorFree = TargetType && Schema && (2 == Pins.Num());
+	if (bIsErrorFree)
+	{
+		//Create Convert Function
+		auto ConvertToObjectFunc = CompilerContext.SpawnIntermediateNode<UK2Node_CallFunction>(this, SourceGraph);
+		const FName ConvertFunctionName = bIsAssetClass 
+			? GET_FUNCTION_NAME_CHECKED(UKismetSystemLibrary, Conv_AssetClassToClass) 
+			: GET_FUNCTION_NAME_CHECKED(UKismetSystemLibrary, Conv_AssetToObject);
+		ConvertToObjectFunc->FunctionReference.SetExternalMember(ConvertFunctionName, UKismetSystemLibrary::StaticClass());
+		ConvertToObjectFunc->AllocateDefaultPins();
+
+		//Connect input to convert
+		auto InputPin = FindPin(UK2Node_ConvertAssetImpl::InputPinName);
+		const FString ConvertInputName = bIsAssetClass
+			? FString(TEXT("AssetClass"))
+			: FString(TEXT("Asset"));
+		auto ConvertInput = ConvertToObjectFunc->FindPin(ConvertInputName);
+		bIsErrorFree &= InputPin && ConvertInput && CompilerContext.MovePinLinksToIntermediate(*InputPin, *ConvertInput).CanSafeConnect();
+
+		//Create Cast Node
+		UK2Node_DynamicCast* CastNode = bIsAssetClass
+			? CompilerContext.SpawnIntermediateNode<UK2Node_ClassDynamicCast>(this, SourceGraph)
+			: CompilerContext.SpawnIntermediateNode<UK2Node_DynamicCast>(this, SourceGraph);
+		CastNode->SetPurity(true);
+		CastNode->TargetType = TargetType;
+		CastNode->AllocateDefaultPins();
+
+		// Connect Object/Class to Cast
+		auto CastInput = CastNode->GetCastSourcePin();
+		auto ConvertOutput = ConvertToObjectFunc->GetReturnValuePin();
+		bIsErrorFree &= ConvertOutput && CastInput && Schema->TryCreateConnection(ConvertOutput, CastInput);
+
+		// Connect output to cast
+		auto OutputPin = FindPin(UK2Node_ConvertAssetImpl::InputPinName);
+		auto CastOutput = CastNode->GetCastResultPin();
+		bIsErrorFree &= OutputPin && CastOutput && CompilerContext.MovePinLinksToIntermediate(*OutputPin, *CastOutput).CanSafeConnect();
+	}
+
+	if (!bIsErrorFree)
+	{
+		CompilerContext.MessageLog.Error(*LOCTEXT("InternalConnectionError", "K2Node_ConvertAsset: Internal connection error. @@").ToString(), this);
+	}
+
+	BreakAllNodeLinks();
+}
+
+bool UK2Node_ConvertAsset::HasExternalDependencies(TArray<class UStruct*>* OptionalOutput) const
+{
+	const UBlueprint* SourceBlueprint = GetBlueprint();
+	UClass* SourceClass = *TargetType;
+	const bool bResult = (SourceClass != NULL) && (SourceClass->ClassGeneratedBy != SourceBlueprint);
+	if (bResult && OptionalOutput)
+	{
+		OptionalOutput->AddUnique(SourceClass);
+	}
+	const bool bSuperResult = Super::HasExternalDependencies(OptionalOutput);
+	return bSuperResult || bResult;
+}
+
+FText UK2Node_ConvertAsset::GetCompactNodeTitle() const
+{
+	return FText::FromString(TEXT("\x2022"));
+}
+
+#undef LOCTEXT_NAMESPACE
\ No newline at end of file
diff --git a/Engine/Source/Editor/GraphEditor/Private/GraphEditorSettings.cpp b/Engine/Source/Editor/GraphEditor/Private/GraphEditorSettings.cpp
--- a/Engine/Source/Editor/GraphEditor/Private/GraphEditorSettings.cpp
+++ b/Engine/Source/Editor/GraphEditor/Private/GraphEditorSettings.cpp
@@ -37,6 +37,8 @@ UGraphEditorSettings::UGraphEditorSettings( const FObjectInitializer& ObjectInit
 	NamePinTypeColor = FLinearColor(0.607717f, 0.224984f, 1.0f, 1.0f);			// lilac
 	DelegatePinTypeColor = FLinearColor(1.0f, 0.04f, 0.04f, 1.0f);				// bright red
 	ObjectPinTypeColor = FLinearColor(0.0f, 0.4f, 0.910000f, 1.0f);				// sharp blue
+	AssetPinTypeColor = FLinearColor(0.3f, 1.0f, 1.0f, 1.0f);					
+	AssetClassPinTypeColor = FLinearColor(1.0f, 0.3f, 1.0f, 1.0f);				
 	InterfacePinTypeColor = FLinearColor(0.8784f, 1.0f, 0.4f, 1.0f);			// pale green
 	StringPinTypeColor = FLinearColor(1.0f, 0.0f, 0.660537f, 1.0f);				// bright pink
 	TextPinTypeColor = FLinearColor(0.8f, 0.2f, 0.4f, 1.0f);					// salmon (light pink)
diff --git a/Engine/Source/Editor/GraphEditor/Private/NodeFactory.cpp b/Engine/Source/Editor/GraphEditor/Private/NodeFactory.cpp
--- a/Engine/Source/Editor/GraphEditor/Private/NodeFactory.cpp
+++ b/Engine/Source/Editor/GraphEditor/Private/NodeFactory.cpp
@@ -290,10 +290,18 @@ TSharedPtr<SGraphPin> FNodeFactory::CreatePinWidget(UEdGraphPin* InPin)
 		{
 			return SNew(SGraphPinObject, InPin);
 		}
+		else if (InPin->PinType.PinCategory == K2Schema->PC_Asset)
+		{
+			return SNew(SGraphPinObject, InPin);
+		}
 		else if (InPin->PinType.PinCategory == K2Schema->PC_Class)
 		{
 			return SNew(SGraphPinClass, InPin);
 		}
+		else if (InPin->PinType.PinCategory == K2Schema->PC_AssetClass)
+		{
+			return SNew(SGraphPinClass, InPin);
+		}
 		else if (InPin->PinType.PinCategory == K2Schema->PC_Int)
 		{
 			return SNew(SGraphPinInteger, InPin);
diff --git a/Engine/Source/Editor/GraphEditor/Public/GraphEditorSettings.h b/Engine/Source/Editor/GraphEditor/Public/GraphEditorSettings.h
--- a/Engine/Source/Editor/GraphEditor/Public/GraphEditorSettings.h
+++ b/Engine/Source/Editor/GraphEditor/Public/GraphEditorSettings.h
@@ -139,6 +139,14 @@ public:
 	UPROPERTY(EditAnywhere, config, Category=PinColors)
 	FLinearColor NamePinTypeColor;
 
+	/** Asset pin type color */
+	UPROPERTY(EditAnywhere, config, Category = PinColors)
+	FLinearColor AssetPinTypeColor;
+
+	/** Asset Class pin type color */
+	UPROPERTY(EditAnywhere, config, Category = PinColors)
+	FLinearColor AssetClassPinTypeColor;
+
 	/** Delegate pin type color */
 	UPROPERTY(EditAnywhere, config, Category=PinColors)
 	FLinearColor DelegatePinTypeColor;
diff --git a/Engine/Source/Editor/KismetCompiler/Private/KismetCompilerMisc.cpp b/Engine/Source/Editor/KismetCompiler/Private/KismetCompilerMisc.cpp
--- a/Engine/Source/Editor/KismetCompiler/Private/KismetCompilerMisc.cpp
+++ b/Engine/Source/Editor/KismetCompiler/Private/KismetCompilerMisc.cpp
@@ -106,7 +106,7 @@ bool FKismetCompilerUtilities::IsTypeCompatibleWithProperty(UEdGraphPin* SourceP
 		UByteProperty* SpecificProperty = Cast<UByteProperty>(TestProperty);
 		bTypeMismatch = (SpecificProperty == NULL);
 	}
-	else if (PinCategory == Schema->PC_Class)
+	else if ((PinCategory == Schema->PC_Class) || (PinCategory == Schema->PC_AssetClass))
 	{
 		const UClass* ClassType = (PinSubCategory == Schema->PSC_Self) ? SelfClass : Cast<const UClass>(PinSubCategoryObject);
 
@@ -135,6 +135,11 @@ bool FKismetCompilerUtilities::IsTypeCompatibleWithProperty(UEdGraphPin* SourceP
 
 				// It matches if it's an exact match or if the output class is more derived than the input class
 				bTypeMismatch = bSubtypeMismatch = !((OutputClass == InputClass) || (OutputClass->IsChildOf(InputClass)));
+
+				if ((PinCategory == Schema->PC_AssetClass) && (!TestProperty->IsA<UAssetClassProperty>()))
+				{
+					bTypeMismatch = true;
+				}
 			}
 			else
 			{
@@ -166,7 +171,7 @@ bool FKismetCompilerUtilities::IsTypeCompatibleWithProperty(UEdGraphPin* SourceP
 			&& PropertyDelegate->SignatureFunction 
 			&& PropertyDelegate->SignatureFunction->IsSignatureCompatibleWith(SignatureFunction));
 	}
-	else if ((PinCategory == Schema->PC_Object) || (PinCategory == Schema->PC_Interface))
+	else if ((PinCategory == Schema->PC_Object) || (PinCategory == Schema->PC_Interface) || (PinCategory == Schema->PC_Asset))
 	{
 		const UClass* ObjectType = (PinSubCategory == Schema->PSC_Self) ? SelfClass : Cast<const UClass>(PinSubCategoryObject);
 
@@ -186,6 +191,11 @@ bool FKismetCompilerUtilities::IsTypeCompatibleWithProperty(UEdGraphPin* SourceP
 
 				// It matches if it's an exact match or if the output class is more derived than the input class
 				bTypeMismatch = bSubtypeMismatch = !((OutputClass == InputClass) || (OutputClass->IsChildOf(InputClass)));
+
+				if ((PinCategory == Schema->PC_Asset) && (!TestProperty->IsA<UAssetObjectProperty>()))
+				{
+					bTypeMismatch = true;
+				}
 			}
 			else if (UInterfaceProperty* IntefaceProperty = Cast<UInterfaceProperty>(TestProperty))
 			{
@@ -742,7 +752,7 @@ UProperty* FKismetCompilerUtilities::CreatePropertyOnScope(UStruct* Scope, const
 	}
 
 	//@TODO: Nasty string if-else tree
-	if ((Type.PinCategory == Schema->PC_Object) || (Type.PinCategory == Schema->PC_Interface))
+	if ((Type.PinCategory == Schema->PC_Object) || (Type.PinCategory == Schema->PC_Interface) || (Type.PinCategory == Schema->PC_Asset))
 	{
 		UClass* SubType = (Type.PinSubCategory == Schema->PSC_Self) ? SelfClass : Cast<UClass>(Type.PinSubCategoryObject.Get());
 
@@ -767,7 +777,11 @@ UProperty* FKismetCompilerUtilities::CreatePropertyOnScope(UStruct* Scope, const
 			{
 				UObjectPropertyBase* NewPropertyObj = NULL;
 
-				if( Type.bIsWeakPointer )
+				if (Type.PinCategory == Schema->PC_Asset)
+				{
+					NewPropertyObj = NewObject<UAssetObjectProperty>(PropertyScope, ValidatedPropertyName, ObjectFlags);
+				}
+				else if( Type.bIsWeakPointer )
 				{
 					NewPropertyObj = NewObject<UWeakObjectProperty>(PropertyScope, ValidatedPropertyName, ObjectFlags);
 				}
@@ -807,7 +821,7 @@ UProperty* FKismetCompilerUtilities::CreatePropertyOnScope(UStruct* Scope, const
 			}
 		}
 	}
-	else if (Type.PinCategory == Schema->PC_Class)
+	else if ((Type.PinCategory == Schema->PC_Class) || (Type.PinCategory == Schema->PC_AssetClass))
 	{
 		UClass* SubType = Cast<UClass>(Type.PinSubCategoryObject.Get());
 		
@@ -825,13 +839,22 @@ UProperty* FKismetCompilerUtilities::CreatePropertyOnScope(UStruct* Scope, const
 
 		if (SubType != NULL)
 		{
-			UClassProperty* NewPropertyClass = NewObject<UClassProperty>(PropertyScope, ValidatedPropertyName, ObjectFlags);
-			// we want to use this setter function instead of setting the 
-			// MetaClass member directly, because it properly handles  
-			// placeholder classes (classes that are stubbed in during load)
-			NewPropertyClass->SetMetaClass(SubType);
-			NewPropertyClass->PropertyClass = UClass::StaticClass();
-			NewProperty = NewPropertyClass;
+			if (Type.PinCategory == Schema->PC_AssetClass)
+			{
+				auto AssetClassProperty = NewObject<UAssetClassProperty>(PropertyScope, ValidatedPropertyName, ObjectFlags);
+				AssetClassProperty->MetaClass = SubType;
+				NewProperty = AssetClassProperty;
+			}
+			else
+			{
+				UClassProperty* NewPropertyClass = NewObject<UClassProperty>(PropertyScope, ValidatedPropertyName, ObjectFlags);
+				// we want to use this setter function instead of setting the 
+				// MetaClass member directly, because it properly handles  
+				// placeholder classes (classes that are stubbed in during load)
+				NewPropertyClass->SetMetaClass(SubType);
+				NewPropertyClass->PropertyClass = UClass::StaticClass();
+				NewProperty = NewPropertyClass;
+			}
 		}
 	}
 	else if (Type.PinCategory == Schema->PC_Delegate)
diff --git a/Engine/Source/Editor/KismetCompiler/Private/KismetCompilerVMBackend.cpp b/Engine/Source/Editor/KismetCompiler/Private/KismetCompilerVMBackend.cpp
--- a/Engine/Source/Editor/KismetCompiler/Private/KismetCompilerVMBackend.cpp
+++ b/Engine/Source/Editor/KismetCompiler/Private/KismetCompilerVMBackend.cpp
@@ -545,9 +545,9 @@ public:
 					{
 						NewTerm.TextLiteral = FText::FromString(NewTerm.Name);
 					}
-					else if (InnerProp->IsA(UObjectProperty::StaticClass()))
+					else if (InnerProp->IsA(UObjectPropertyBase::StaticClass()))
 					{
-						NewTerm.ObjectLiteral = Cast<UObjectProperty>(InnerProp)->GetObjectPropertyValue(RawElemData);
+						NewTerm.ObjectLiteral = Cast<UObjectPropertyBase>(InnerProp)->GetObjectPropertyValue(RawElemData);
 					}
 					EmitTermExpr(&NewTerm, InnerProp);
 				}
@@ -567,6 +567,12 @@ public:
 					Writer << FunctionName;
 				}
 			}
+			else if (CoerceProperty->IsA(UAssetObjectProperty::StaticClass()))
+			{
+				Writer << EX_AssetConst;
+				FAssetPtr AssetPtr(Term->ObjectLiteral);
+				EmitStringLiteral(AssetPtr.GetUniqueID().AssetLongPathname);
+			}
 			else if (CoerceProperty->IsA(UObjectPropertyBase::StaticClass()))
 			{
 				// Note: This case handles both UObjectProperty and UClassProperty
@@ -880,6 +886,7 @@ public:
 		const bool bIsMulticastDelegate = Cast<UMulticastDelegateProperty>(DestinationExpression->AssociatedVarProperty) != NULL;
 		const bool bIsBoolean = Cast<UBoolProperty>(DestinationExpression->AssociatedVarProperty) != NULL;
 		const bool bIsObj = Cast<UObjectPropertyBase>(DestinationExpression->AssociatedVarProperty) != NULL;
+		const bool bIsAsset = Cast<UAssetObjectProperty>(DestinationExpression->AssociatedVarProperty) != NULL;
 		const bool bIsWeakObjPtr = Cast<UWeakObjectProperty>(DestinationExpression->AssociatedVarProperty) != NULL;
 		if (bIsMulticastDelegate)
 		{
@@ -893,7 +900,7 @@ public:
 		{
 			Writer << EX_LetBool;
 		}
-		else if (bIsObj)
+		else if (bIsObj && !bIsAsset)
 		{
 			if( !bIsWeakObjPtr )
 			{
diff --git a/Engine/Source/Editor/UnrealEd/Private/ScriptDisassembler.cpp b/Engine/Source/Editor/UnrealEd/Private/ScriptDisassembler.cpp
--- a/Engine/Source/Editor/UnrealEd/Private/ScriptDisassembler.cpp
+++ b/Engine/Source/Editor/UnrealEd/Private/ScriptDisassembler.cpp
@@ -586,6 +586,12 @@ void FKismetBytecodeDisassembler::ProcessCommon(int32& ScriptIndex, EExprToken O
 			Ar.Logf(TEXT("%s $%X: EX_ObjectConst (%p:%s)"), *Indents, (int32)Opcode, Pointer, *Pointer->GetFullName());
 			break;
 		}
+	case EX_AssetConst:
+		{
+			Ar.Logf(TEXT("%s $%X: EX_AssetConst"), *Indents, (int32)Opcode);
+			SerializeExpr(ScriptIndex);
+			break;
+		}
 	case EX_NameConst:
 		{
 			FString ConstValue = ReadName(ScriptIndex);
diff --git a/Engine/Source/Programs/UnrealHeaderTool/Private/HeaderParser.cpp b/Engine/Source/Programs/UnrealHeaderTool/Private/HeaderParser.cpp
--- a/Engine/Source/Programs/UnrealHeaderTool/Private/HeaderParser.cpp
+++ b/Engine/Source/Programs/UnrealHeaderTool/Private/HeaderParser.cpp
@@ -723,6 +723,7 @@ namespace
 
 		const bool bSupportedType = Property->IsA<UInterfaceProperty>()
 			|| Property->IsA<UClassProperty>()
+			|| Property->IsA<UAssetObjectProperty>()
 			|| Property->IsA<UObjectProperty>()
 			|| Property->IsA<UStructProperty>()
 			|| Property->IsA<UFloatProperty>()
@@ -3768,9 +3769,9 @@ UProperty* FHeaderParser::GetVarNameAndDim
 		// UFunctions with a smart pointer as input parameter wont compile anyway, because of missing P_GET_... macro.
 		// UFunctions with a smart pointer as return type will crash when called via blueprint, because they are not supported in VM.
 		// WeakPointer is supported by VM as return type (see UObject::execLetWeakObjPtr), but there is no P_GET_... macro for WeakPointer.
-		if ((VarProperty.Type == CPT_LazyObjectReference) || (VarProperty.Type == CPT_AssetObjectReference))
+		if ((VarProperty.Type == CPT_LazyObjectReference))
 		{
-			FError::Throwf(TEXT("UFunctions cannot take a smart pointer (LazyPtr, AssetPtr, etc) as a parameter."));
+			FError::Throwf(TEXT("UFunctions cannot take a lazy pointer as a parameter."));
 		}
 	}
 
diff --git a/Engine/Source/Runtime/CoreUObject/Private/UObject/PropertyAssetClass.cpp b/Engine/Source/Runtime/CoreUObject/Private/UObject/PropertyAssetClass.cpp
--- a/Engine/Source/Runtime/CoreUObject/Private/UObject/PropertyAssetClass.cpp
+++ b/Engine/Source/Runtime/CoreUObject/Private/UObject/PropertyAssetClass.cpp
@@ -16,6 +16,11 @@ FString UAssetClassProperty::GetCPPMacroType( FString& ExtendedTypeText ) const
 	return TEXT("ASSETOBJECT");
 }
 
+FString UAssetClassProperty::GetCPPTypeForwardDeclaration() const
+{
+	return FString::Printf(TEXT("class %s%s;"), MetaClass->GetPrefixCPP(), *MetaClass->GetName());
+}
+
 void UAssetClassProperty::Serialize( FArchive& Ar )
 {
 	Super::Serialize( Ar );
diff --git a/Engine/Source/Runtime/CoreUObject/Private/UObject/PropertyAssetObject.cpp b/Engine/Source/Runtime/CoreUObject/Private/UObject/PropertyAssetObject.cpp
--- a/Engine/Source/Runtime/CoreUObject/Private/UObject/PropertyAssetObject.cpp
+++ b/Engine/Source/Runtime/CoreUObject/Private/UObject/PropertyAssetObject.cpp
@@ -16,6 +16,11 @@ FString UAssetObjectProperty::GetCPPMacroType( FString& ExtendedTypeText ) const
 	return TEXT("ASSETOBJECT");
 }
 
+FString UAssetObjectProperty::GetCPPTypeForwardDeclaration() const
+{
+	return FString::Printf(TEXT("class %s%s;"), PropertyClass->GetPrefixCPP(), *PropertyClass->GetName());
+}
+
 FName UAssetObjectProperty::GetID() const
 {
 	return NAME_AssetObjectProperty;
diff --git a/Engine/Source/Runtime/CoreUObject/Private/UObject/ScriptCore.cpp b/Engine/Source/Runtime/CoreUObject/Private/UObject/ScriptCore.cpp
--- a/Engine/Source/Runtime/CoreUObject/Private/UObject/ScriptCore.cpp
+++ b/Engine/Source/Runtime/CoreUObject/Private/UObject/ScriptCore.cpp
@@ -1815,6 +1815,14 @@ void UObject::execObjectConst( FFrame& Stack, RESULT_DECL )
 }
 IMPLEMENT_VM_FUNCTION( EX_ObjectConst, execObjectConst );
 
+void UObject::execAssetConst(FFrame& Stack, RESULT_DECL)
+{
+	FString LongPath;
+	Stack.Step(Stack.Object, &LongPath);
+	*(FAssetPtr*)RESULT_PARAM = FStringAssetReference(LongPath);
+}
+IMPLEMENT_VM_FUNCTION(EX_AssetConst, execAssetConst);
+
 void UObject::execInstanceDelegate( FFrame& Stack, RESULT_DECL )
 {
 	FName FunctionName = Stack.ReadName();
diff --git a/Engine/Source/Runtime/CoreUObject/Public/UObject/Script.h b/Engine/Source/Runtime/CoreUObject/Public/UObject/Script.h
--- a/Engine/Source/Runtime/CoreUObject/Public/UObject/Script.h
+++ b/Engine/Source/Runtime/CoreUObject/Public/UObject/Script.h
@@ -238,6 +238,7 @@ enum EExprToken
 	Ex_LetValueOnPersistentFrame = 0x64,
 	EX_ArrayConst			= 0x65,
 	EX_EndArrayConst		= 0x66,
+	EX_AssetConst			= 0x67,
 	EX_Max					= 0x100,
 };
 
diff --git a/Engine/Source/Runtime/CoreUObject/Public/UObject/ScriptMacros.h b/Engine/Source/Runtime/CoreUObject/Public/UObject/ScriptMacros.h
--- a/Engine/Source/Runtime/CoreUObject/Public/UObject/ScriptMacros.h
+++ b/Engine/Source/Runtime/CoreUObject/Public/UObject/ScriptMacros.h
@@ -60,6 +60,9 @@ enum {MAX_VARIABLE_SIZE = 0x0FFF };
 #define P_GET_TINTERFACE(ObjectType,ParamName)		PARAM_PASSED_BY_VAL(ParamName, UInterfaceProperty, TScriptInterface<ObjectType>)
 #define P_GET_TINTERFACE_REF(ObjectType,ParamName)	PARAM_PASSED_BY_REF(ParamName, UInterfaceProperty, TScriptInterface<ObjectType>)
 
+#define P_GET_ASSETOBJECT(StructType,ParamName)		PARAM_PASSED_BY_VAL(ParamName, UAssetObjectProperty, StructType)
+#define P_GET_ASSETOBJECT_REF(StructType,ParamName)	PARAM_PASSED_BY_REF(ParamName, UAssetObjectProperty, StructType)
+
 #define P_GET_ARRAY(ElementType,ParamName)			ElementType ParamName[(MAX_VARIABLE_SIZE/sizeof(ElementType))+1];		Stack.StepCompiledIn<UProperty>(ParamName);
 #define P_GET_ARRAY_REF(ElementType,ParamName)		ElementType ParamName##Temp[(MAX_VARIABLE_SIZE/sizeof(ElementType))+1]; ElementType* ParamName = Stack.StepCompiledInRef<UProperty, ElementType*>(ParamName##Temp);
 
diff --git a/Engine/Source/Runtime/CoreUObject/Public/UObject/ScriptSerialization.h b/Engine/Source/Runtime/CoreUObject/Public/UObject/ScriptSerialization.h
--- a/Engine/Source/Runtime/CoreUObject/Public/UObject/ScriptSerialization.h
+++ b/Engine/Source/Runtime/CoreUObject/Public/UObject/ScriptSerialization.h
@@ -296,6 +296,11 @@
 
 			break;
 		}
+		case EX_AssetConst:
+		{
+			SerializeExpr(iCode, Ar);
+			break;
+		}
 		case EX_NameConst:
 		{
 			XFERNAME();
diff --git a/Engine/Source/Runtime/CoreUObject/Public/UObject/UObject.h b/Engine/Source/Runtime/CoreUObject/Public/UObject/UObject.h
--- a/Engine/Source/Runtime/CoreUObject/Public/UObject/UObject.h
+++ b/Engine/Source/Runtime/CoreUObject/Public/UObject/UObject.h
@@ -1107,6 +1107,7 @@ public:
 	DECLARE_FUNCTION(execUnicodeStringConst);
 	DECLARE_FUNCTION(execTextConst);
 	DECLARE_FUNCTION(execObjectConst);
+	DECLARE_FUNCTION(execAssetConst);
 
 	// @todo delegate: Multi-cast versions needed for script execution!		(Need Add, Remove, Clear/Empty)
 	DECLARE_FUNCTION(execInstanceDelegate);
diff --git a/Engine/Source/Runtime/CoreUObject/Public/UObject/UnrealType.h b/Engine/Source/Runtime/CoreUObject/Public/UObject/UnrealType.h
--- a/Engine/Source/Runtime/CoreUObject/Public/UObject/UnrealType.h
+++ b/Engine/Source/Runtime/CoreUObject/Public/UObject/UnrealType.h
@@ -1975,6 +1975,7 @@ class COREUOBJECT_API UAssetObjectProperty : public TUObjectPropertyBase<FAssetP
 	// UHT interface
 	virtual FString GetCPPMacroType( FString& ExtendedTypeText ) const  override;
 	virtual FString GetCPPType( FString* ExtendedTypeText, uint32 CPPExportFlags ) const override;
+	virtual FString GetCPPTypeForwardDeclaration() const override;
 	// End of UHT interface
 
 	// UProperty interface
@@ -2090,6 +2091,7 @@ public:
 	// UHT interface
 	virtual FString GetCPPMacroType( FString& ExtendedTypeText ) const  override;
 	virtual FString GetCPPType( FString* ExtendedTypeText, uint32 CPPExportFlags ) const override;
+	virtual FString GetCPPTypeForwardDeclaration() const override;
 	// End of UHT interface
 
 	// UObject interface
diff --git a/Engine/Source/Runtime/Engine/Classes/Kismet/KismetSystemLibrary.h b/Engine/Source/Runtime/Engine/Classes/Kismet/KismetSystemLibrary.h
--- a/Engine/Source/Runtime/Engine/Classes/Kismet/KismetSystemLibrary.h
+++ b/Engine/Source/Runtime/Engine/Classes/Kismet/KismetSystemLibrary.h
@@ -5,6 +5,7 @@
 #include "Engine/LatentActionManager.h"
 #include "BlueprintFunctionLibrary.h"
 #include "Engine/CollisionProfile.h"
+#include "AssetPtr.h"
 #include "KismetSystemLibrary.generated.h"
 
 UENUM(BlueprintType)
@@ -135,6 +136,12 @@ class ENGINE_API UKismetSystemLibrary : public UBlueprintFunctionLibrary
 	UFUNCTION(BlueprintPure, meta=(DisplayName = "ToObject (interface)", CompactNodeTitle = "->"), Category="Utilities")
 	static UObject* Conv_InterfaceToObject(const FScriptInterface& Interface); 
 
+	UFUNCTION(BlueprintPure, meta = (BlueprintInternalUseOnly = "true"), Category = "Utilities")
+	static UObject* Conv_AssetToObject(const TAssetPtr<UObject>& Asset);
+
+	UFUNCTION(BlueprintPure, meta = (BlueprintInternalUseOnly = "true"), Category = "Utilities")
+	static TSubclassOf<UObject> Conv_AssetClassToClass(const TAssetSubclassOf<UObject>& AssetClass);
+
 	/**
 	 * Creates a literal integer
 	 * @param	Value	value to set the integer to
@@ -448,6 +455,14 @@ class ENGINE_API UKismetSystemLibrary : public UBlueprintFunctionLibrary
 	UFUNCTION(BlueprintCallable, meta=(BlueprintInternalUseOnly = "true", AutoCreateRefTerm = "Value" ))
 	static void SetNamePropertyByName(UObject* Object, FName PropertyName, const FName& Value);
 
+	/** Set a ASSET property by name */
+	UFUNCTION(BlueprintCallable, meta = (BlueprintInternalUseOnly = "true", AutoCreateRefTerm = "Value"))
+	static void SetAssetPropertyByName(UObject* Object, FName PropertyName, const TAssetPtr<UObject>& Value);
+
+	/** Set a ASSETCLASS property by name */
+	UFUNCTION(BlueprintCallable, meta = (BlueprintInternalUseOnly = "true", AutoCreateRefTerm = "Value"))
+	static void SetAssetClassPropertyByName(UObject* Object, FName PropertyName, const TAssetSubclassOf<UObject>& Value);
+
 	/** Set a STRING property by name */
 	UFUNCTION(BlueprintCallable, meta=(BlueprintInternalUseOnly = "true", AutoCreateRefTerm = "Value" ))
 	static void SetStringPropertyByName(UObject* Object, FName PropertyName, const FString& Value);
diff --git a/Engine/Source/Runtime/Engine/Private/KismetSystemLibrary.cpp b/Engine/Source/Runtime/Engine/Private/KismetSystemLibrary.cpp
--- a/Engine/Source/Runtime/Engine/Private/KismetSystemLibrary.cpp
+++ b/Engine/Source/Runtime/Engine/Private/KismetSystemLibrary.cpp
@@ -666,6 +666,36 @@ void UKismetSystemLibrary::SetNamePropertyByName(UObject* Object, FName Property
 	}
 }
 
+void UKismetSystemLibrary::SetAssetPropertyByName(UObject* Object, FName PropertyName, const TAssetPtr<UObject>& Value)
+{
+	if (Object != NULL)
+	{
+		auto ObjectProp = FindField<UAssetObjectProperty>(Object->GetClass(), PropertyName);
+		auto AssetPtr = (const FAssetPtr*)(&Value);
+		ObjectProp->SetPropertyValue_InContainer(Object, *AssetPtr);
+	}
+}
+
+void UKismetSystemLibrary::SetAssetClassPropertyByName(UObject* Object, FName PropertyName, const TAssetSubclassOf<UObject>& Value)
+{
+	if (Object != NULL)
+	{
+		auto ObjectProp = FindField<UAssetClassProperty>(Object->GetClass(), PropertyName);
+		auto AssetPtr = (const FAssetPtr*)(&Value);
+		ObjectProp->SetPropertyValue_InContainer(Object, *AssetPtr);
+	}
+}
+
+UObject* UKismetSystemLibrary::Conv_AssetToObject(const TAssetPtr<UObject>& Asset)
+{
+	return ((const FAssetPtr*)&Asset)->Get();
+}
+
+TSubclassOf<UObject> UKismetSystemLibrary::Conv_AssetClassToClass(const TAssetSubclassOf<UObject>& AssetClass)
+{
+	return Cast<UClass>(((const FAssetPtr*)&AssetClass)->Get());
+}
+
 void UKismetSystemLibrary::SetTextPropertyByName(UObject* Object, FName PropertyName, const FText& Value)
 {
 	if(Object != NULL)
@@ -769,7 +799,6 @@ void UKismetSystemLibrary::GetActorListFromComponentList(const TArray<UPrimitive
 	}
 }
 
-
 bool UKismetSystemLibrary::SphereOverlapActors_DEPRECATED(UObject* WorldContextObject, const FVector SpherePos, float SphereRadius, EOverlapFilterOption Filter, UClass* ActorClassFilter, const TArray<AActor*>& ActorsToIgnore, TArray<AActor*>& OutActors)
 {
 	OutActors.Empty();
