Function nodes are categorized correctly again.

diff --git a/Engine/Source/Editor/BlueprintGraph/Private/K2Node_CallFunction.cpp b/Engine/Source/Editor/BlueprintGraph/Private/K2Node_CallFunction.cpp
--- a/Engine/Source/Editor/BlueprintGraph/Private/K2Node_CallFunction.cpp
+++ b/Engine/Source/Editor/BlueprintGraph/Private/K2Node_CallFunction.cpp
@@ -1360,6 +1360,10 @@ FText UK2Node_CallFunction::GetDefaultCategoryForFunction(const UFunction* Funct
 		{
 			NodeCategory = FText::Format(FText::FromString(TEXT("{0}|{1}")), NodeCategory, FuncCategory);
 		}
+		else if (NodeCategory.IsEmpty())
+		{
+			NodeCategory = FuncCategory;
+		}
 	}
 	return NodeCategory;
 }
