Cis Fix - declaration shadows a local variable

diff --git a/Engine/Source/Editor/BlueprintGraph/Private/EdGraphSchema_K2.cpp b/Engine/Source/Editor/BlueprintGraph/Private/EdGraphSchema_K2.cpp
--- a/Engine/Source/Editor/BlueprintGraph/Private/EdGraphSchema_K2.cpp
+++ b/Engine/Source/Editor/BlueprintGraph/Private/EdGraphSchema_K2.cpp
@@ -2249,17 +2249,17 @@ bool UEdGraphSchema_K2::FindSpecializedConversionNode(const UEdGraphPin* OutputP
 
 			if (!bCanConvert && InputClass && OutputClass && OutputClass->IsChildOf(InputClass))
 			{
-				const bool ConvertAsset = (OutputType.PinCategory == PC_Asset) && (InputType.PinCategory == PC_Object);
-				const bool ConvertAssetClass = (OutputType.PinCategory == PC_AssetClass) && (InputType.PinCategory == PC_Class);
-				if (ConvertAsset || ConvertAssetClass)
+				const bool bConvertAsset = (OutputType.PinCategory == PC_Asset) && (InputType.PinCategory == PC_Object);
+				const bool bConvertAssetClass = (OutputType.PinCategory == PC_AssetClass) && (InputType.PinCategory == PC_Class);
+				if (bConvertAsset || bConvertAssetClass)
 				{
 					bCanConvert = true;
 					if (bCreateNode)
 					{
-						UK2Node_ConvertAsset* ConvertAsset = NewObject<UK2Node_ConvertAsset>();
-						ConvertAsset->TargetType = InputClass;
-						ConvertAsset->bIsAssetClass = ConvertAssetClass;
-						TargetNode = ConvertAsset;
+						UK2Node_ConvertAsset* ConvertAssetNode = NewObject<UK2Node_ConvertAsset>();
+						ConvertAssetNode->TargetType = InputClass;
+						ConvertAssetNode->bIsAssetClass = bConvertAssetClass;
+						TargetNode = ConvertAssetNode;
 					}
 				}
 			}
