Cis fix

diff --git a/Engine/Source/Editor/BlueprintGraph/Private/K2Node_ConvertAsset.cpp b/Engine/Source/Editor/BlueprintGraph/Private/K2Node_ConvertAsset.cpp
--- a/Engine/Source/Editor/BlueprintGraph/Private/K2Node_ConvertAsset.cpp
+++ b/Engine/Source/Editor/BlueprintGraph/Private/K2Node_ConvertAsset.cpp
@@ -4,6 +4,7 @@
 #include "BlueprintGraphPrivatePCH.h"
 #include "Kismet/KismetSystemLibrary.h"
 #include "K2Node_ClassDynamicCast.h"
+#include "KismetCompiler.h"
 #include "K2Node_ConvertAsset.h"
 
 #define LOCTEXT_NAMESPACE "K2Node_ConvertAsset"
