Fixed issues with hiding categories in Blueprints.

Switched to using the UiSpec for the filter so the full (and correct) category would be parsed for hiding.

diff --git a/Engine/Source/Editor/BlueprintGraph/Private/BlueprintActionFilter.cpp b/Engine/Source/Editor/BlueprintGraph/Private/BlueprintActionFilter.cpp
--- a/Engine/Source/Editor/BlueprintGraph/Private/BlueprintActionFilter.cpp
+++ b/Engine/Source/Editor/BlueprintGraph/Private/BlueprintActionFilter.cpp
@@ -774,23 +774,13 @@ static bool BlueprintActionFilterImpl::IsFieldCategoryHidden(FBlueprintActionFil
 	DECLARE_DELEGATE_RetVal_OneParam(bool, FIsFieldHiddenDelegate, UClass*);
 	FIsFieldHiddenDelegate IsFieldHiddenDelegate;
 
-	if (UFunction const* Function = BlueprintAction.GetAssociatedFunction())
-	{
-		auto IsFunctionHiddenLambda = [](UClass* Class, UFunction const* InFunction)->bool
-		{
-			// Only hide functions that are not static
-			return !InFunction->HasAnyFunctionFlags(FUNC_Static) && FObjectEditorUtils::IsFunctionHiddenFromClass(InFunction, Class->GetAuthoritativeClass());
-		};
-		IsFieldHiddenDelegate = FIsFieldHiddenDelegate::CreateStatic(IsFunctionHiddenLambda, Function);
-	}
-	else if (UProperty const* Property = BlueprintAction.GetAssociatedProperty())
+	// Use the UiSpec to get the category
+	FBlueprintActionUiSpec UiSpec = BlueprintAction.NodeSpawner->GetUiSpec(Filter.Context, BlueprintAction.GetBindings());
+	auto IsNodeHiddenLambda = [](UClass* Class, FText InCategory)->bool
 	{
-		auto IsPropertyHiddenLambda = [](UClass* Class, UProperty const* InProperty)->bool
-		{
-			return FObjectEditorUtils::IsVariableCategoryHiddenFromClass(InProperty, Class->GetAuthoritativeClass());
-		};
-		IsFieldHiddenDelegate = FIsFieldHiddenDelegate::CreateStatic(IsPropertyHiddenLambda, Property);
-	}
+		return FEditorCategoryUtils::IsCategoryHiddenFromClass(Class->GetAuthoritativeClass(), InCategory.ToString());
+	};
+	IsFieldHiddenDelegate = FIsFieldHiddenDelegate::CreateStatic(IsNodeHiddenLambda, UiSpec.Category);
 
 	if (IsFieldHiddenDelegate.IsBound())
 	{
diff --git a/Engine/Source/Editor/UnrealEd/Private/Kismet2/BlueprintEditorUtils.cpp b/Engine/Source/Editor/UnrealEd/Private/Kismet2/BlueprintEditorUtils.cpp
--- a/Engine/Source/Editor/UnrealEd/Private/Kismet2/BlueprintEditorUtils.cpp
+++ b/Engine/Source/Editor/UnrealEd/Private/Kismet2/BlueprintEditorUtils.cpp
@@ -1715,7 +1715,12 @@ void FBlueprintEditorUtils::RecreateClassMetaData(UBlueprint* Blueprint, UClass*
 		Class->SetMetaData(FBlueprintMetadata::MD_AllowableBlueprintVariableType, TEXT("true"));
 	}
 
-	AllHideCategories.Append(Blueprint->HideCategories);
+	for (FString HideCategory : Blueprint->HideCategories)
+	{
+		HideCategory.ReplaceInline(TEXT(" "), TEXT(""));
+		AllHideCategories.Add(HideCategory);
+	}
+
 	if (AllHideCategories.Num())
 	{
 		Class->SetMetaData(TEXT("HideCategories"), *FString::Join(AllHideCategories, TEXT(" ")));
