Load Asset latent Function
UE-14387 Support for AssetPtr references in Blueprints

diff --git a/Engine/Source/Editor/BlueprintGraph/Classes/K2Node_LoadAsset.h b/Engine/Source/Editor/BlueprintGraph/Classes/K2Node_LoadAsset.h
new file mode 100644
--- /dev/null
+++ b/Engine/Source/Editor/BlueprintGraph/Classes/K2Node_LoadAsset.h
@@ -0,0 +1,56 @@
+// Copyright 1998-2015 Epic Games, Inc. All Rights Reserved.
+
+
+#pragma once
+#include "K2Node.h"
+#include "K2Node_LoadAsset.generated.h"
+
+UCLASS(MinimalAPI)
+class UK2Node_LoadAsset : public UK2Node
+{
+	GENERATED_BODY()
+public:
+	// UEdGraphNode interface
+	virtual void AllocateDefaultPins() override;
+	virtual FText GetTooltipText() const override;
+	virtual FText GetNodeTitle(ENodeTitleType::Type TitleType) const override;
+	virtual bool IsCompatibleWithGraph(const UEdGraph* TargetGraph) const override;
+	// End of UEdGraphNode interface
+
+	// UK2Node interface
+	virtual bool IsNodePure() const override { return false; }
+	virtual void ExpandNode(class FKismetCompilerContext& CompilerContext, UEdGraph* SourceGraph) override;
+	virtual FName GetCornerIcon() const override;
+	virtual void GetMenuActions(FBlueprintActionDatabaseRegistrar& ActionRegistrar) const override;
+	virtual FText GetMenuCategory() const override;
+	// End of UK2Node interface
+
+protected:
+	virtual FName NativeFunctionName() const;
+
+	virtual const FString& GetInputCategory() const;
+	virtual const FString& GetOutputCategory() const;
+
+	virtual const FString& GetInputPinName() const;
+	virtual const FString& GetOutputPinName() const;
+};
+
+UCLASS(MinimalAPI)
+class UK2Node_LoadAssetClass : public UK2Node_LoadAsset
+{
+	GENERATED_BODY()
+public:
+	// UEdGraphNode interface
+	virtual FText GetTooltipText() const override;
+	virtual FText GetNodeTitle(ENodeTitleType::Type TitleType) const override;
+	// End of UEdGraphNode interface
+
+protected:
+	virtual FName NativeFunctionName() const override;
+
+	virtual const FString& GetInputCategory() const override;
+	virtual const FString& GetOutputCategory() const override;
+
+	virtual const FString& GetInputPinName() const override;
+	virtual const FString& GetOutputPinName() const override;
+};
diff --git a/Engine/Source/Editor/BlueprintGraph/Private/K2Node_LoadAsset.cpp b/Engine/Source/Editor/BlueprintGraph/Private/K2Node_LoadAsset.cpp
new file mode 100644
--- /dev/null
+++ b/Engine/Source/Editor/BlueprintGraph/Private/K2Node_LoadAsset.cpp
@@ -0,0 +1,262 @@
+// Copyright 1998-2015 Epic Games, Inc. All Rights Reserved.
+
+
+#include "BlueprintGraphPrivatePCH.h"
+#include "Kismet/KismetSystemLibrary.h"
+#include "K2Node_ClassDynamicCast.h"
+#include "KismetCompiler.h"
+#include "K2Node_LoadAsset.h"
+
+#define LOCTEXT_NAMESPACE "K2Node_LoadAsset"
+
+void UK2Node_LoadAsset::AllocateDefaultPins()
+{
+	const UEdGraphSchema_K2* K2Schema = Cast<UEdGraphSchema_K2>(GetSchema());
+
+	if (K2Schema)
+	{
+		CreatePin(EGPD_Input, K2Schema->PC_Exec, TEXT(""), nullptr, false, false, K2Schema->PN_Execute);
+		UEdGraphPin* OutputExecPin = CreatePin(EGPD_Output, K2Schema->PC_Exec, TEXT(""), nullptr, false, false, K2Schema->PN_Then);
+		if (OutputExecPin)
+		{
+			OutputExecPin->PinFriendlyName = FText::FromString(K2Schema->PN_Completed);
+		}
+
+		CreatePin(EGPD_Input, GetInputCategory(), TEXT(""), UObject::StaticClass(), false, false, GetInputPinName());
+		CreatePin(EGPD_Output, GetOutputCategory(), TEXT(""), UObject::StaticClass(), false, false, GetOutputPinName());
+	}
+}
+
+void UK2Node_LoadAsset::ExpandNode(class FKismetCompilerContext& CompilerContext, UEdGraph* SourceGraph)
+{
+	Super::ExpandNode(CompilerContext, SourceGraph);
+	const UEdGraphSchema_K2* Schema = CompilerContext.GetSchema();
+	check(Schema);
+	bool bIsErrorFree = true;
+
+	// Create LoadAsset function call
+	auto CallLoadAssetNode = CompilerContext.SpawnIntermediateNode<UK2Node_CallFunction>(this, SourceGraph);
+	CallLoadAssetNode->FunctionReference.SetExternalMember(NativeFunctionName(), UKismetSystemLibrary::StaticClass());
+	CallLoadAssetNode->AllocateDefaultPins(); 
+
+	// connect to input exe
+	{
+		auto InputExePin = GetExecPin();
+		auto CallFunctionInputExePin = CallLoadAssetNode->GetExecPin();
+		bIsErrorFree &= InputExePin && CallFunctionInputExePin && CompilerContext.MovePinLinksToIntermediate(*InputExePin, *CallFunctionInputExePin).CanSafeConnect();
+	}
+
+	// Create Local Variable
+	UK2Node_TemporaryVariable* TempVarOutput = CompilerContext.SpawnInternalVariable(this, GetOutputCategory(), FString(), UObject::StaticClass(), false);
+
+	// Create assign node
+	auto AssignNode = CompilerContext.SpawnIntermediateNode<UK2Node_AssignmentStatement>(this, SourceGraph);
+	AssignNode->AllocateDefaultPins();
+
+	auto LoadedObjectVariablePin = TempVarOutput->GetVariablePin();
+
+	// connect local variable to assign node
+	{
+		auto AssignLHSPPin = AssignNode->GetVariablePin();
+		bIsErrorFree &= AssignLHSPPin && LoadedObjectVariablePin && Schema->TryCreateConnection(AssignLHSPPin, LoadedObjectVariablePin);
+	}
+
+	// connect local variable to output
+	{
+		auto OutputObjectPinPin = FindPin(GetOutputPinName());
+		bIsErrorFree &= LoadedObjectVariablePin && OutputObjectPinPin && CompilerContext.MovePinLinksToIntermediate(*OutputObjectPinPin, *LoadedObjectVariablePin).CanSafeConnect();
+	}
+
+	// connect assign exec input to function output
+	{
+		auto CallFunctionOutputExePin = CallLoadAssetNode->FindPin(Schema->PN_Then);
+		auto AssignInputExePin = AssignNode->GetExecPin();
+		bIsErrorFree &= AssignInputExePin && CallFunctionOutputExePin && Schema->TryCreateConnection(AssignInputExePin, CallFunctionOutputExePin);
+	}
+
+	// connect assign exec output to output
+	{
+		auto OutputExePin = FindPin(Schema->PN_Then);
+		auto AssignOutputExePin = AssignNode->GetThenPin();
+		bIsErrorFree &= OutputExePin && AssignOutputExePin && CompilerContext.MovePinLinksToIntermediate(*OutputExePin, *AssignOutputExePin).CanSafeConnect();
+	}
+
+	// connect to asset
+	{
+		auto AssetPin = FindPin(GetInputPinName());
+		auto CallFunctionAssetPin = CallLoadAssetNode->FindPin(GetInputPinName());
+		ensure(CallFunctionAssetPin);
+		bIsErrorFree &= AssetPin && CallFunctionAssetPin && CompilerContext.MovePinLinksToIntermediate(*AssetPin, *CallFunctionAssetPin).CanSafeConnect();
+	}
+
+	// Create OnLoadEvent
+	const FString DelegateOnLoadedParamName(TEXT("OnLoaded"));
+	auto OnLoadEventNode = CompilerContext.SpawnIntermediateNode<UK2Node_CustomEvent>(this, SourceGraph);
+	OnLoadEventNode->CustomFunctionName = *FString::Printf(TEXT("OnLoaded_%s"), *OnLoadEventNode->NodeGuid.ToString());
+	OnLoadEventNode->AllocateDefaultPins();
+	{
+		UFunction* LoadAssetFunction = CallLoadAssetNode->GetTargetFunction();
+		UDelegateProperty* OnLoadDelegateProperty = LoadAssetFunction ? FindField<UDelegateProperty>(LoadAssetFunction, *DelegateOnLoadedParamName) : nullptr;
+		UFunction* OnLoadedSignature = OnLoadDelegateProperty ? OnLoadDelegateProperty->SignatureFunction : nullptr;
+		ensure(OnLoadedSignature);
+		for (TFieldIterator<UProperty> PropIt(OnLoadedSignature); PropIt && (PropIt->PropertyFlags & CPF_Parm); ++PropIt)
+		{
+			const UProperty* Param = *PropIt;
+			if (!Param->HasAnyPropertyFlags(CPF_OutParm) || Param->HasAnyPropertyFlags(CPF_ReferenceParm))
+			{
+				FEdGraphPinType PinType;
+				bIsErrorFree &= Schema->ConvertPropertyToPinType(Param, /*out*/ PinType);
+				bIsErrorFree &= (NULL != OnLoadEventNode->CreateUserDefinedPin(Param->GetName(), PinType, EGPD_Output));
+			}
+		}
+	}
+
+	// connect delegate
+	{
+		auto CallFunctionDelegatePin = CallLoadAssetNode->FindPin(DelegateOnLoadedParamName);
+		ensure(CallFunctionDelegatePin);
+		auto EventDelegatePin = OnLoadEventNode->FindPin(UK2Node_CustomEvent::DelegateOutputName);
+		bIsErrorFree &= CallFunctionDelegatePin && EventDelegatePin && Schema->TryCreateConnection(CallFunctionDelegatePin, EventDelegatePin);
+	}
+
+	// connect loaded object from event to assign
+	{
+		auto LoadedAssetEventPin = OnLoadEventNode->FindPin(TEXT("Loaded"));
+		ensure(LoadedAssetEventPin);
+		auto AssignRHSPPin = AssignNode->GetValuePin();
+		bIsErrorFree &= AssignRHSPPin && LoadedAssetEventPin && Schema->TryCreateConnection(LoadedAssetEventPin, AssignRHSPPin);
+	}
+
+	if (!bIsErrorFree)
+	{
+		CompilerContext.MessageLog.Error(*LOCTEXT("InternalConnectionError", "K2Node_LoadAsset: Internal connection error. @@").ToString(), this);
+	}
+
+	BreakAllNodeLinks();
+}
+
+FText UK2Node_LoadAsset::GetTooltipText() const
+{
+	return FText(LOCTEXT("UK2Node_LoadAssetGetTooltipText", "Async Load Asset"));
+}
+
+FText UK2Node_LoadAsset::GetNodeTitle(ENodeTitleType::Type TitleType) const
+{
+	return FText(LOCTEXT("UK2Node_LoadAssetGetNodeTitle", "Load Asset"));
+}
+
+bool UK2Node_LoadAsset::IsCompatibleWithGraph(const UEdGraph* TargetGraph) const
+{
+	bool bIsCompatible = false;
+	// Can only place events in ubergraphs and macros (other code will help prevent macros with latents from ending up in functions), and basicasync task creates an event node:
+	EGraphType GraphType = TargetGraph->GetSchema()->GetGraphType(TargetGraph);
+	if (GraphType == EGraphType::GT_Ubergraph || GraphType == EGraphType::GT_Macro)
+	{
+		bIsCompatible = true;
+	}
+	return bIsCompatible && Super::IsCompatibleWithGraph(TargetGraph);
+}
+
+FName UK2Node_LoadAsset::GetCornerIcon() const
+{
+	return TEXT("Graph.Latent.LatentIcon");
+}
+
+void UK2Node_LoadAsset::GetMenuActions(FBlueprintActionDatabaseRegistrar& ActionRegistrar) const
+{
+	// actions get registered under specific object-keys; the idea is that 
+	// actions might have to be updated (or deleted) if their object-key is  
+	// mutated (or removed)... here we use the node's class (so if the node 
+	// type disappears, then the action should go with it)
+	UClass* ActionKey = GetClass();
+	// to keep from needlessly instantiating a UBlueprintNodeSpawner, first   
+	// check to make sure that the registrar is looking for actions of this type
+	// (could be regenerating actions for a specific asset, and therefore the 
+	// registrar would only accept actions corresponding to that asset)
+	if (ActionRegistrar.IsOpenForRegistration(ActionKey))
+	{
+		UBlueprintNodeSpawner* NodeSpawner = UBlueprintNodeSpawner::Create(GetClass());
+		check(NodeSpawner != nullptr);
+
+		ActionRegistrar.AddBlueprintAction(ActionKey, NodeSpawner);
+	}
+}
+
+FText UK2Node_LoadAsset::GetMenuCategory() const
+{
+	return FText(LOCTEXT("UK2Node_LoadAssetGetMenuCategory", "Utilities"));
+}
+
+const FString& UK2Node_LoadAsset::GetInputCategory() const
+{
+	const UEdGraphSchema_K2* K2Schema = CastChecked<const UEdGraphSchema_K2>(GetSchema());
+	return K2Schema->PC_Asset;
+}
+
+const FString& UK2Node_LoadAsset::GetOutputCategory() const
+{
+	const UEdGraphSchema_K2* K2Schema = CastChecked<const UEdGraphSchema_K2>(GetSchema());
+	return K2Schema->PC_Object;
+}
+
+const FString& UK2Node_LoadAsset::GetInputPinName() const
+{
+	static const FString InputAssetPinName("Asset");
+	return InputAssetPinName;
+}
+
+const FString& UK2Node_LoadAsset::GetOutputPinName() const
+{
+	static const FString OutputObjectPinName("Object");
+	return OutputObjectPinName;
+}
+
+FName UK2Node_LoadAsset::NativeFunctionName() const
+{
+	return GET_FUNCTION_NAME_CHECKED(UKismetSystemLibrary, LoadAsset);
+}
+
+// UK2Node_LoadAssetClass
+
+FName UK2Node_LoadAssetClass::NativeFunctionName() const
+{
+	return GET_FUNCTION_NAME_CHECKED(UKismetSystemLibrary, LoadAssetClass);
+}
+
+
+FText UK2Node_LoadAssetClass::GetTooltipText() const
+{
+	return FText(LOCTEXT("UK2Node_LoadAssetClassGetTooltipText", "Async Load Asset Class"));
+}
+
+FText UK2Node_LoadAssetClass::GetNodeTitle(ENodeTitleType::Type TitleType) const
+{
+	return FText(LOCTEXT("UK2Node_LoadAssetClassGetNodeTitle", "Load Asset Class"));
+}
+
+const FString& UK2Node_LoadAssetClass::GetInputCategory() const
+{
+	const UEdGraphSchema_K2* K2Schema = CastChecked<const UEdGraphSchema_K2>(GetSchema());
+	return K2Schema->PC_AssetClass;
+}
+
+const FString& UK2Node_LoadAssetClass::GetOutputCategory() const
+{
+	const UEdGraphSchema_K2* K2Schema = CastChecked<const UEdGraphSchema_K2>(GetSchema());
+	return K2Schema->PC_Class;
+}
+
+const FString& UK2Node_LoadAssetClass::GetInputPinName() const
+{
+	static const FString InputAssetPinName("AssetClass");
+	return InputAssetPinName;
+}
+
+const FString& UK2Node_LoadAssetClass::GetOutputPinName() const
+{
+	static const FString OutputObjectPinName("Class");
+	return OutputObjectPinName;
+}
+
+
+#undef LOCTEXT_NAMESPACE
\ No newline at end of file
diff --git a/Engine/Source/Runtime/CoreUObject/Private/UObject/Class.cpp b/Engine/Source/Runtime/CoreUObject/Private/UObject/Class.cpp
--- a/Engine/Source/Runtime/CoreUObject/Private/UObject/Class.cpp
+++ b/Engine/Source/Runtime/CoreUObject/Private/UObject/Class.cpp
@@ -4297,6 +4297,60 @@ bool UFunction::IsSignatureCompatibleWith(const UFunction* OtherFunction) const
 	return IsSignatureCompatibleWith(OtherFunction, IgnoreFlags);
 }
 
+bool FStructUtils::ArePropertiesTheSame(const UProperty* A, const UProperty* B, bool bCheckPropertiesNames)
+{
+	if (A == B)
+	{
+		return true;
+	}
+
+	if (!A != !B) //one of properties is null
+	{
+		return false;
+	}
+
+	if (bCheckPropertiesNames && (A->GetFName() != B->GetFName()))
+	{
+		return false;
+	}
+
+	if (A->GetSize() != B->GetSize())
+	{
+		return false;
+	}
+
+	if (A->GetOffset_ForGC() != B->GetOffset_ForGC())
+	{
+		return false;
+	}
+
+	if (!A->SameType(B))
+	{
+		return false;
+	}
+
+	return true;
+}
+
+bool FStructUtils::TheSameLayout(const UStruct* StructA, const UStruct* StructB, bool bCheckPropertiesNames)
+{
+	bool bResult = false;
+	if (StructA && StructB)
+	{
+		const UProperty* PropertyA = StructA->PropertyLink;
+		const UProperty* PropertyB = StructB->PropertyLink;
+
+		bResult = true;
+		while (bResult && (PropertyA != PropertyB))
+		{
+			bResult = ArePropertiesTheSame(PropertyA, PropertyB, bCheckPropertiesNames);
+			PropertyA = PropertyA ? PropertyA->PropertyLinkNext : NULL;
+			PropertyB = PropertyB ? PropertyB->PropertyLinkNext : NULL;
+		}
+	}
+	return bResult;
+}
+
 bool UFunction::IsSignatureCompatibleWith(const UFunction* OtherFunction, uint64 IgnoreFlags) const
 {
 	// Early out if they're exactly the same function
@@ -4319,12 +4373,9 @@ bool UFunction::IsSignatureCompatibleWith(const UFunction* OtherFunction, uint64
 			UProperty* PropA = *IteratorA;
 			UProperty* PropB = *IteratorB;
 
-			FString TypeA(PropA->GetCPPType());
-			FString TypeB(PropB->GetCPPType());
-
 			// Check the flags as well
 			const uint64 PropertyMash = PropA->PropertyFlags ^ PropB->PropertyFlags;
-			if ((TypeA != TypeB) || ((PropertyMash & ~IgnoreFlags) != 0))
+			if (!FStructUtils::ArePropertiesTheSame(PropA, PropB, false) || ((PropertyMash & ~IgnoreFlags) != 0))
 			{
 				// Type mismatch between an argument of A and B
 				return false;
diff --git a/Engine/Source/Runtime/CoreUObject/Public/UObject/Class.h b/Engine/Source/Runtime/CoreUObject/Public/UObject/Class.h
--- a/Engine/Source/Runtime/CoreUObject/Public/UObject/Class.h
+++ b/Engine/Source/Runtime/CoreUObject/Public/UObject/Class.h
@@ -1301,7 +1301,7 @@ public:
 	FORCEINLINE static uint64 GetDefaultIgnoredSignatureCompatibilityFlags()
 	{
 		//@TODO: UCREMOVAL: CPF_ConstParm added as a hack to get blueprints compiling with a const DamageType parameter.
-		const uint64 IgnoreFlags = CPF_PersistentInstance | CPF_ExportObject | CPF_InstancedReference | CPF_ContainsInstancedReference | CPF_ComputedFlags | CPF_ConstParm | CPF_HasGetValueTypeHash;
+		const uint64 IgnoreFlags = CPF_PersistentInstance | CPF_ExportObject | CPF_InstancedReference | CPF_ContainsInstancedReference | CPF_ComputedFlags | CPF_ConstParm | CPF_HasGetValueTypeHash | CPF_UObjectWrapper;
 		return IgnoreFlags;
 	}
 
@@ -2739,6 +2739,14 @@ inline T* GetMutableDefault(UClass *Class)
 	return (T*)Class->GetDefaultObject();
 }
 
+struct FStructUtils
+{
+	static bool ArePropertiesTheSame(const UProperty* A, const UProperty* B, bool bCheckPropertiesNames);
+
+	// does structures have exactly the same memory layout
+	COREUOBJECT_API static bool TheSameLayout(const UStruct* StructA, const UStruct* StructB, bool bCheckPropertiesNames = false);
+};
+
 template< class T > struct TBaseStructure
 {
 };
diff --git a/Engine/Source/Runtime/Engine/Classes/Engine/BlueprintGeneratedClass.h b/Engine/Source/Runtime/Engine/Classes/Engine/BlueprintGeneratedClass.h
--- a/Engine/Source/Runtime/Engine/Classes/Engine/BlueprintGeneratedClass.h
+++ b/Engine/Source/Runtime/Engine/Classes/Engine/BlueprintGeneratedClass.h
@@ -166,18 +166,6 @@ private:
 	int32 WriteIndex;
 };
 
-
-struct FStructUtils
-{
-private:
-	static bool ArePropertiesTheSame(const UProperty* A, const UProperty* B, bool bCheckPropertiesNames);
-
-public:
-	// does structures have exactly the same memory layout
-	ENGINE_API static bool TheSameLayout(const UStruct* StructA, const UStruct* StructB, bool bCheckPropertiesNames = false);
-};
-
-
 USTRUCT()
 struct ENGINE_API FBlueprintDebugData
 {
diff --git a/Engine/Source/Runtime/Engine/Classes/Kismet/KismetSystemLibrary.h b/Engine/Source/Runtime/Engine/Classes/Kismet/KismetSystemLibrary.h
--- a/Engine/Source/Runtime/Engine/Classes/Kismet/KismetSystemLibrary.h
+++ b/Engine/Source/Runtime/Engine/Classes/Kismet/KismetSystemLibrary.h
@@ -142,6 +142,17 @@ class ENGINE_API UKismetSystemLibrary : public UBlueprintFunctionLibrary
 	UFUNCTION(BlueprintPure, meta = (BlueprintInternalUseOnly = "true"), Category = "Utilities")
 	static TSubclassOf<UObject> Conv_AssetClassToClass(const TAssetSubclassOf<UObject>& AssetClass);
 
+	DECLARE_DYNAMIC_DELEGATE_OneParam(FOnAssetLoaded, class UObject*, Loaded);
+
+	UFUNCTION(BlueprintCallable, meta = (Latent, LatentInfo = "LatentInfo", WorldContext = "WorldContextObject", BlueprintInternalUseOnly = "true"), Category = "Utilities")
+	static void LoadAsset(UObject* WorldContextObject, const TAssetPtr<UObject>& Asset, FOnAssetLoaded OnLoaded, FLatentActionInfo LatentInfo);
+
+	DECLARE_DYNAMIC_DELEGATE_OneParam(FOnAssetClassLoaded, TSubclassOf<UObject>, Loaded);
+
+	UFUNCTION(BlueprintCallable, meta = (Latent, LatentInfo = "LatentInfo", WorldContext = "WorldContextObject", BlueprintInternalUseOnly = "true"), Category = "Utilities")
+	static void LoadAssetClass(UObject* WorldContextObject, const TAssetSubclassOf<UObject>& AssetClass, FOnAssetClassLoaded OnLoaded, FLatentActionInfo LatentInfo);
+
+
 	/**
 	 * Creates a literal integer
 	 * @param	Value	value to set the integer to
diff --git a/Engine/Source/Runtime/Engine/Private/BlueprintGeneratedClass.cpp b/Engine/Source/Runtime/Engine/Private/BlueprintGeneratedClass.cpp
--- a/Engine/Source/Runtime/Engine/Private/BlueprintGeneratedClass.cpp
+++ b/Engine/Source/Runtime/Engine/Private/BlueprintGeneratedClass.cpp
@@ -137,60 +137,6 @@ UClass* UBlueprintGeneratedClass::GetAuthoritativeClass()
 	return (GeneratingBP->GeneratedClass != NULL) ? GeneratingBP->GeneratedClass : this;
 }
 
-bool FStructUtils::ArePropertiesTheSame(const UProperty* A, const UProperty* B, bool bCheckPropertiesNames)
-{
-	if (A == B)
-	{
-		return true;
-	}
-
-	if (!A != !B) //one of properties is null
-	{
-		return false;
-	}
-
-	if (bCheckPropertiesNames && (A->GetFName() != B->GetFName()))
-	{
-		return false;
-	}
-
-	if (A->GetSize() != B->GetSize())
-	{
-		return false;
-	}
-
-	if (A->GetOffset_ForGC() != B->GetOffset_ForGC())
-	{
-		return false;
-	}
-
-	if (!A->SameType(B))
-	{
-		return false;
-	}
-
-	return true;
-}
-
-bool FStructUtils::TheSameLayout(const UStruct* StructA, const UStruct* StructB, bool bCheckPropertiesNames)
-{
-	bool bResult = false;
-	if (StructA && StructB)
-	{
-		const UProperty* PropertyA = StructA->PropertyLink;
-		const UProperty* PropertyB = StructB->PropertyLink;
-
-		bResult = true;
-		while (bResult && (PropertyA != PropertyB))
-		{
-			bResult = ArePropertiesTheSame(PropertyA, PropertyB, bCheckPropertiesNames);
-			PropertyA = PropertyA ? PropertyA->PropertyLinkNext : NULL;
-			PropertyB = PropertyB ? PropertyB->PropertyLinkNext : NULL;
-		}
-	}
-	return bResult;
-}
-
 struct FConditionalRecompileClassHepler
 {
 	enum ENeededAction
diff --git a/Engine/Source/Runtime/Engine/Private/KismetSystemLibrary.cpp b/Engine/Source/Runtime/Engine/Private/KismetSystemLibrary.cpp
--- a/Engine/Source/Runtime/Engine/Private/KismetSystemLibrary.cpp
+++ b/Engine/Source/Runtime/Engine/Private/KismetSystemLibrary.cpp
@@ -13,6 +13,7 @@
 #include "Camera/CameraComponent.h"
 #include "Kismet/GameplayStatics.h"
 #include "SlateCore.h"
+#include "Engine/StreamableManager.h"
 
 //////////////////////////////////////////////////////////////////////////
 // UKismetSystemLibrary
@@ -2986,3 +2987,114 @@ void UKismetSystemLibrary::SetSupressViewportTransitionMessage(UObject* WorldCon
 		World->GetFirstLocalPlayerFromController()->ViewportClient->SetSuppressTransitionMessage(bState);
 	}
 }
+struct FLoadAssetActionBase : public FPendingLatentAction, public FGCObject
+{
+	// @TODO: it would be good to have static/global manager? 
+
+public:
+	FStringAssetReference AssetReference;
+	FStreamableManager StreamableManager;
+	FName ExecutionFunction;
+	int32 OutputLink;
+	FWeakObjectPtr CallbackTarget;
+
+	virtual void OnLoaded() PURE_VIRTUAL(FLoadAssetActionBase::OnLoaded, );
+
+	FLoadAssetActionBase(const FStringAssetReference& InAssetReference, const FLatentActionInfo& InLatentInfo)
+		: AssetReference(InAssetReference)
+		, ExecutionFunction(InLatentInfo.ExecutionFunction)
+		, OutputLink(InLatentInfo.Linkage)
+		, CallbackTarget(InLatentInfo.CallbackTarget)
+	{
+		StreamableManager.SimpleAsyncLoad(AssetReference);
+	}
+
+	virtual ~FLoadAssetActionBase()
+	{
+		StreamableManager.Unload(AssetReference);
+	}
+
+	virtual void UpdateOperation(FLatentResponse& Response) override
+	{
+		const bool bLoaded = StreamableManager.IsAsyncLoadComplete(AssetReference);
+		if (bLoaded)
+		{
+			OnLoaded();
+		}
+		Response.FinishAndTriggerIf(bLoaded, ExecutionFunction, OutputLink, CallbackTarget);
+	}
+
+#if WITH_EDITOR
+	virtual FString GetDescription() const override
+	{
+		return FString::Printf(TEXT("Load Asset Action Base: %s"), *AssetReference.ToString());
+	}
+#endif
+
+	virtual void AddReferencedObjects(FReferenceCollector& Collector) override
+	{
+		StreamableManager.AddStructReferencedObjects(Collector);
+	}
+};
+
+void UKismetSystemLibrary::LoadAsset(UObject* WorldContextObject, const TAssetPtr<UObject>& Asset, UKismetSystemLibrary::FOnAssetLoaded OnLoaded, FLatentActionInfo LatentInfo)
+{
+	struct FLoadAssetAction : public FLoadAssetActionBase
+	{
+	public:
+		UKismetSystemLibrary::FOnAssetLoaded OnLoadedCallback;
+
+		FLoadAssetAction(const FStringAssetReference& InAssetReference, UKismetSystemLibrary::FOnAssetLoaded InOnLoadedCallback, const FLatentActionInfo& InLatentInfo)
+			: FLoadAssetActionBase(InAssetReference, InLatentInfo)
+			, OnLoadedCallback(InOnLoadedCallback)
+		{}
+
+		virtual void OnLoaded() override
+		{
+			UObject* LoadedObject = AssetReference.ResolveObject();
+			OnLoadedCallback.ExecuteIfBound(LoadedObject);
+		}
+	};
+
+	UWorld* World = GEngine->GetWorldFromContextObject(WorldContextObject);
+	if (World != nullptr)
+	{
+		FLatentActionManager& LatentManager = World->GetLatentActionManager();
+		if (LatentManager.FindExistingAction<FLoadAssetAction>(LatentInfo.CallbackTarget, LatentInfo.UUID) == nullptr)
+		{
+			auto NewAction = new FLoadAssetAction(Asset.ToStringReference(), OnLoaded, LatentInfo);
+			LatentManager.AddNewAction(LatentInfo.CallbackTarget, LatentInfo.UUID, NewAction);
+		}
+	}
+}
+
+void UKismetSystemLibrary::LoadAssetClass(UObject* WorldContextObject, const TAssetSubclassOf<UObject>& AssetClass, UKismetSystemLibrary::FOnAssetClassLoaded OnLoaded, FLatentActionInfo LatentInfo)
+{
+	struct FLoadAssetClassAction : public FLoadAssetActionBase
+	{
+	public:
+		UKismetSystemLibrary::FOnAssetClassLoaded OnLoadedCallback;
+
+		FLoadAssetClassAction(const FStringAssetReference& InAssetReference, UKismetSystemLibrary::FOnAssetClassLoaded InOnLoadedCallback, const FLatentActionInfo& InLatentInfo)
+			: FLoadAssetActionBase(InAssetReference, InLatentInfo)
+			, OnLoadedCallback(InOnLoadedCallback)
+		{}
+
+		virtual void OnLoaded() override
+		{
+			auto LoadedObject = Cast<UClass>(AssetReference.ResolveObject());
+			OnLoadedCallback.ExecuteIfBound(LoadedObject);
+		}
+	};
+
+	UWorld* World = GEngine->GetWorldFromContextObject(WorldContextObject);
+	if (World != nullptr)
+	{
+		FLatentActionManager& LatentManager = World->GetLatentActionManager();
+		if (LatentManager.FindExistingAction<FLoadAssetClassAction>(LatentInfo.CallbackTarget, LatentInfo.UUID) == nullptr)
+		{
+			auto NewAction = new FLoadAssetClassAction(AssetClass.ToStringReference(), OnLoaded, LatentInfo);
+			LatentManager.AddNewAction(LatentInfo.CallbackTarget, LatentInfo.UUID, NewAction);
+		}
+	}
+}
\ No newline at end of file
