More usable convert asset node.
UE-14387 Support for AssetPtr references in Blueprints

diff --git a/Engine/Source/Editor/BlueprintGraph/Classes/K2Node_ConvertAsset.h b/Engine/Source/Editor/BlueprintGraph/Classes/K2Node_ConvertAsset.h
--- a/Engine/Source/Editor/BlueprintGraph/Classes/K2Node_ConvertAsset.h
+++ b/Engine/Source/Editor/BlueprintGraph/Classes/K2Node_ConvertAsset.h
@@ -10,21 +10,27 @@ class UK2Node_ConvertAsset : public UK2Node
 {
 	GENERATED_BODY()
 public:
-	UPROPERTY()
-	TSubclassOf<class UObject>  TargetType;
 
-	UPROPERTY()
-	bool bIsAssetClass;
+	UClass* GetTargetClass() const;
+	bool IsAssetClassType() const;
 
 	// UEdGraphNode interface
 	virtual void AllocateDefaultPins() override;
+	virtual FText GetNodeTitle(ENodeTitleType::Type TitleType) const override;
+	virtual FText GetTooltipText() const override;
 	// End of UEdGraphNode interface
 
 	// UK2Node interface
 	virtual bool IsNodePure() const override { return true; }
 	virtual bool ShouldDrawCompact() const override { return true; }
+	virtual int32 GetNodeRefreshPriority() const override { return EBaseNodeRefreshPriority::Low_UsesDependentWildcard; }
+	virtual bool DoesInputWildcardPinAcceptArray(const UEdGraphPin* Pin) const override { return false; }
+
 	virtual void ExpandNode(class FKismetCompilerContext& CompilerContext, UEdGraph* SourceGraph) override;
-	virtual bool HasExternalDependencies(TArray<class UStruct*>* OptionalOutput) const override;
 	virtual FText GetCompactNodeTitle() const override;
+	virtual bool IsConnectionDisallowed(const UEdGraphPin* MyPin, const UEdGraphPin* OtherPin, FString& OutReason) const;
+	virtual void NotifyPinConnectionListChanged(UEdGraphPin* Pin) override;
+	virtual void GetMenuActions(FBlueprintActionDatabaseRegistrar& ActionRegistrar) const override;
+	virtual FText GetMenuCategory() const override;
 	// End of UK2Node interface
 };
\ No newline at end of file
diff --git a/Engine/Source/Editor/BlueprintGraph/Private/EdGraphSchema_K2.cpp b/Engine/Source/Editor/BlueprintGraph/Private/EdGraphSchema_K2.cpp
--- a/Engine/Source/Editor/BlueprintGraph/Private/EdGraphSchema_K2.cpp
+++ b/Engine/Source/Editor/BlueprintGraph/Private/EdGraphSchema_K2.cpp
@@ -2257,8 +2257,6 @@ bool UEdGraphSchema_K2::FindSpecializedConversionNode(const UEdGraphPin* OutputP
 					if (bCreateNode)
 					{
 						UK2Node_ConvertAsset* ConvertAssetNode = NewObject<UK2Node_ConvertAsset>();
-						ConvertAssetNode->TargetType = InputClass;
-						ConvertAssetNode->bIsAssetClass = bConvertAssetClass;
 						TargetNode = ConvertAssetNode;
 					}
 				}
diff --git a/Engine/Source/Editor/BlueprintGraph/Private/K2Node_ConvertAsset.cpp b/Engine/Source/Editor/BlueprintGraph/Private/K2Node_ConvertAsset.cpp
--- a/Engine/Source/Editor/BlueprintGraph/Private/K2Node_ConvertAsset.cpp
+++ b/Engine/Source/Editor/BlueprintGraph/Private/K2Node_ConvertAsset.cpp
@@ -15,36 +15,128 @@ namespace UK2Node_ConvertAssetImpl
 	static const FString OutputPinName("Object");
 }
 
-void UK2Node_ConvertAsset::AllocateDefaultPins()
+UClass* UK2Node_ConvertAsset::GetTargetClass() const
+{
+	auto InutPin = FindPin(UK2Node_ConvertAssetImpl::InputPinName);
+	bool bIsConnected = InutPin && InutPin->LinkedTo.Num() && InutPin->LinkedTo[0];
+	auto SourcePin = bIsConnected ? InutPin->LinkedTo[0] : nullptr;
+	return SourcePin
+		? Cast<UClass>(SourcePin->PinType.PinSubCategoryObject.Get())
+		: nullptr;
+}
+
+bool UK2Node_ConvertAsset::IsAssetClassType() const
 {
 	const UEdGraphSchema_K2* K2Schema = Cast<UEdGraphSchema_K2>(GetSchema());
-	const bool bReferenceObsoleteClass = !TargetType || TargetType->HasAnyClassFlags(CLASS_NewerVersionExists);
-	if (bReferenceObsoleteClass)
-	{
-		Message_Error(FString::Printf(TEXT("Node '%s' references obsolete class '%s'"), *GetPathName(), *GetPathNameSafe(TargetType)));
-	}
-	ensure(!bReferenceObsoleteClass);
 
-	if (TargetType && K2Schema)
+	// get first input, return if class asset
+	auto InutPin = FindPin(UK2Node_ConvertAssetImpl::InputPinName);
+	bool bIsConnected = InutPin && InutPin->LinkedTo.Num() && InutPin->LinkedTo[0];
+	auto SourcePin = bIsConnected ? InutPin->LinkedTo[0] : nullptr;
+	return (SourcePin && K2Schema)
+		? (SourcePin->PinType.PinCategory == K2Schema->PC_AssetClass)
+		: false;
+}
+
+bool UK2Node_ConvertAsset::IsConnectionDisallowed(const UEdGraphPin* MyPin, const UEdGraphPin* OtherPin, FString& OutReason) const
+{
+	const UEdGraphSchema_K2* K2Schema = Cast<UEdGraphSchema_K2>(GetSchema());
+	auto InutPin = FindPin(UK2Node_ConvertAssetImpl::InputPinName);
+	if (K2Schema && InutPin && OtherPin && (InutPin == MyPin) && (MyPin->PinType.PinCategory == K2Schema->PC_Wildcard))
 	{
-		if (bIsAssetClass)
+		if ((OtherPin->PinType.PinCategory != K2Schema->PC_Asset) &&
+			(OtherPin->PinType.PinCategory != K2Schema->PC_AssetClass))
 		{
-			CreatePin(EGPD_Input, K2Schema->PC_AssetClass, TEXT(""), *TargetType, false, false, UK2Node_ConvertAssetImpl::InputPinName);
-			CreatePin(EGPD_Output, K2Schema->PC_Class, TEXT(""), *TargetType, false, false, UK2Node_ConvertAssetImpl::OutputPinName);
+			return true;
 		}
-		else
+	}
+	return false;
+}
+
+void UK2Node_ConvertAsset::NotifyPinConnectionListChanged(UEdGraphPin* Pin)
+{
+	Super::NotifyPinConnectionListChanged(Pin);
+
+	const UEdGraphSchema_K2* K2Schema = CastChecked<UEdGraphSchema_K2>(GetSchema());
+	auto InutPin = FindPin(UK2Node_ConvertAssetImpl::InputPinName);
+	auto OutputPin = FindPin(UK2Node_ConvertAssetImpl::OutputPinName);
+	ensure(InutPin && OutputPin);
+	if (InutPin && OutputPin && (InutPin == Pin))
+	{
+		const bool bIsConnected = InutPin->LinkedTo.Num() > 0;
+		UClass* TargetType = bIsConnected ? GetTargetClass() : nullptr;
+		const bool bIsAssetClass = bIsConnected ? IsAssetClassType() : false;
+
+		const FString InputCategory = bIsConnected
+			? (bIsAssetClass ? K2Schema->PC_AssetClass : K2Schema->PC_Asset)
+			: K2Schema->PC_Wildcard;
+		InutPin->PinType = FEdGraphPinType(InputCategory,FString(), TargetType, false, false);
+
+		const FString OutputCategory = bIsConnected
+			? (bIsAssetClass ? K2Schema->PC_Class : K2Schema->PC_Object)
+			: K2Schema->PC_Wildcard;
+		OutputPin->PinType = FEdGraphPinType(OutputCategory, FString(), TargetType, false, false);
+
+		PinTypeChanged(InutPin);
+		PinTypeChanged(OutputPin);
+
+		if (OutputPin->LinkedTo.Num())
 		{
-			CreatePin(EGPD_Input, K2Schema->PC_Asset, TEXT(""), *TargetType, false, false, UK2Node_ConvertAssetImpl::InputPinName);
-			CreatePin(EGPD_Output, K2Schema->PC_Object, TEXT(""), *TargetType, false, false, UK2Node_ConvertAssetImpl::OutputPinName);
+			UClass const* CallingContext = NULL;
+			if (UBlueprint const* Blueprint = GetBlueprint())
+			{
+				CallingContext = Blueprint->GeneratedClass;
+				if (CallingContext == NULL)
+				{
+					CallingContext = Blueprint->ParentClass;
+				}
+			}
+
+			for (auto TargetPin : OutputPin->LinkedTo)
+			{
+				if (TargetPin && !K2Schema->ArePinsCompatible(OutputPin, TargetPin, CallingContext))
+				{
+					OutputPin->BreakLinkTo(TargetPin);
+				}
+			}
 		}
 	}
 }
 
+void UK2Node_ConvertAsset::AllocateDefaultPins()
+{
+	const UEdGraphSchema_K2* K2Schema = CastChecked<UEdGraphSchema_K2>(GetSchema());
+	CreatePin(EGPD_Input, K2Schema->PC_Wildcard, TEXT(""), nullptr, false, false, UK2Node_ConvertAssetImpl::InputPinName);
+	CreatePin(EGPD_Output, K2Schema->PC_Wildcard, TEXT(""), nullptr, false, false, UK2Node_ConvertAssetImpl::OutputPinName);
+}
+
+void UK2Node_ConvertAsset::GetMenuActions(FBlueprintActionDatabaseRegistrar& ActionRegistrar) const
+{
+	// actions get registered under specific object-keys; the idea is that 
+	// actions might have to be updated (or deleted) if their object-key is  
+	// mutated (or removed)... here we use the node's class (so if the node 
+	// type disappears, then the action should go with it)
+	UClass* ActionKey = GetClass();
+	// to keep from needlessly instantiating a UBlueprintNodeSpawner, first   
+	// check to make sure that the registrar is looking for actions of this type
+	// (could be regenerating actions for a specific asset, and therefore the 
+	// registrar would only accept actions corresponding to that asset)
+	if (ActionRegistrar.IsOpenForRegistration(ActionKey))
+	{
+		UBlueprintNodeSpawner* NodeSpawner = UBlueprintNodeSpawner::Create(GetClass());
+		check(NodeSpawner != nullptr);
+
+		ActionRegistrar.AddBlueprintAction(ActionKey, NodeSpawner);
+	}
+}
+
 void UK2Node_ConvertAsset::ExpandNode(class FKismetCompilerContext& CompilerContext, UEdGraph* SourceGraph)
 {
 	Super::ExpandNode(CompilerContext, SourceGraph);
 	const UEdGraphSchema_K2* Schema = CompilerContext.GetSchema();
-
+	
+	UClass* TargetType = GetTargetClass();
+	const bool bIsAssetClass = IsAssetClassType();
 	bool bIsErrorFree = TargetType && Schema && (2 == Pins.Num());
 	if (bIsErrorFree)
 	{
@@ -91,22 +183,24 @@ void UK2Node_ConvertAsset::ExpandNode(class FKismetCompilerContext& CompilerCont
 	BreakAllNodeLinks();
 }
 
-bool UK2Node_ConvertAsset::HasExternalDependencies(TArray<class UStruct*>* OptionalOutput) const
+FText UK2Node_ConvertAsset::GetCompactNodeTitle() const
 {
-	const UBlueprint* SourceBlueprint = GetBlueprint();
-	UClass* SourceClass = *TargetType;
-	const bool bResult = (SourceClass != NULL) && (SourceClass->ClassGeneratedBy != SourceBlueprint);
-	if (bResult && OptionalOutput)
-	{
-		OptionalOutput->AddUnique(SourceClass);
-	}
-	const bool bSuperResult = Super::HasExternalDependencies(OptionalOutput);
-	return bSuperResult || bResult;
+	return FText::FromString(TEXT("\x2022"));
 }
 
-FText UK2Node_ConvertAsset::GetCompactNodeTitle() const
+FText UK2Node_ConvertAsset::GetMenuCategory() const
 {
-	return FText::FromString(TEXT("\x2022"));
+	return FText(LOCTEXT("UK2Node_LoadAssetGetMenuCategory", "Utilities"));
+}
+
+FText UK2Node_ConvertAsset::GetNodeTitle(ENodeTitleType::Type TitleType) const
+{
+	return FText(LOCTEXT("UK2Node_ConvertAssetGetNodeTitle", "Resolve Asset"));
+}
+
+FText UK2Node_ConvertAsset::GetTooltipText() const
+{
+	return FText(LOCTEXT("UK2Node_ConvertAssetGetTooltipText", "Resolves an asset/assetclass into an object/class. If the asset wasn't loaded it returns none."));
 }
 
 #undef LOCTEXT_NAMESPACE
\ No newline at end of file
