Load Asset node causes structural Blueprint change

diff --git a/Engine/Source/Editor/BlueprintGraph/Classes/K2Node_LoadAsset.h b/Engine/Source/Editor/BlueprintGraph/Classes/K2Node_LoadAsset.h
--- a/Engine/Source/Editor/BlueprintGraph/Classes/K2Node_LoadAsset.h
+++ b/Engine/Source/Editor/BlueprintGraph/Classes/K2Node_LoadAsset.h
@@ -23,6 +23,7 @@ public:
 	virtual FName GetCornerIcon() const override;
 	virtual void GetMenuActions(FBlueprintActionDatabaseRegistrar& ActionRegistrar) const override;
 	virtual FText GetMenuCategory() const override;
+	virtual bool NodeCausesStructuralBlueprintChange() const { return true; }
 	// End of UK2Node interface
 
 protected:
