Cis Fix

diff --git a/Engine/Source/Editor/BlueprintGraph/Private/K2Node_ConvertAsset.cpp b/Engine/Source/Editor/BlueprintGraph/Private/K2Node_ConvertAsset.cpp
--- a/Engine/Source/Editor/BlueprintGraph/Private/K2Node_ConvertAsset.cpp
+++ b/Engine/Source/Editor/BlueprintGraph/Private/K2Node_ConvertAsset.cpp
@@ -5,6 +5,8 @@
 #include "Kismet/KismetSystemLibrary.h"
 #include "K2Node_ClassDynamicCast.h"
 #include "KismetCompiler.h"
+#include "BlueprintNodeSpawner.h"
+#include "BlueprintActionDatabaseRegistrar.h"
 #include "K2Node_ConvertAsset.h"
 
 #define LOCTEXT_NAMESPACE "K2Node_ConvertAsset"
diff --git a/Engine/Source/Editor/BlueprintGraph/Private/K2Node_LoadAsset.cpp b/Engine/Source/Editor/BlueprintGraph/Private/K2Node_LoadAsset.cpp
--- a/Engine/Source/Editor/BlueprintGraph/Private/K2Node_LoadAsset.cpp
+++ b/Engine/Source/Editor/BlueprintGraph/Private/K2Node_LoadAsset.cpp
@@ -5,6 +5,8 @@
 #include "Kismet/KismetSystemLibrary.h"
 #include "K2Node_ClassDynamicCast.h"
 #include "KismetCompiler.h"
+#include "BlueprintNodeSpawner.h"
+#include "BlueprintActionDatabaseRegistrar.h"
 #include "K2Node_LoadAsset.h"
 
 #define LOCTEXT_NAMESPACE "K2Node_LoadAsset"
