Fix non-Unity error in BlueprintActionFilter.cpp

diff --git a/Engine/Source/Editor/BlueprintGraph/Private/BlueprintActionFilter.cpp b/Engine/Source/Editor/BlueprintGraph/Private/BlueprintActionFilter.cpp
--- a/Engine/Source/Editor/BlueprintGraph/Private/BlueprintActionFilter.cpp
+++ b/Engine/Source/Editor/BlueprintGraph/Private/BlueprintActionFilter.cpp
@@ -22,6 +22,7 @@
 #include "K2Node_ExecutionSequence.h"
 #include "K2Node_DynamicCast.h"
 #include "K2Node_BaseAsyncTask.h"
+#include "EditorCategoryUtils.h"
 
 /*******************************************************************************
  * Static BlueprintActionFilter Helpers
