Amending a check() to account for the new Blueprint Asset type.

diff --git a/Engine/Source/Editor/BlueprintGraph/Private/K2Node_MakeStruct.cpp b/Engine/Source/Editor/BlueprintGraph/Private/K2Node_MakeStruct.cpp
--- a/Engine/Source/Editor/BlueprintGraph/Private/K2Node_MakeStruct.cpp
+++ b/Engine/Source/Editor/BlueprintGraph/Private/K2Node_MakeStruct.cpp
@@ -52,7 +52,7 @@ void UK2Node_MakeStruct::FMakeStructPinManager::CustomizePinData(UEdGraphPin* Pi
 		checkSlow(bIsText == ((Schema->PC_Text == Pin->PinType.PinCategory) && !Pin->PinType.bIsArray));
 
 		const bool bIsObject = Property->IsA<UObjectPropertyBase>();
-		checkSlow(bIsObject == ((Schema->PC_Object == Pin->PinType.PinCategory || Schema->PC_Class == Pin->PinType.PinCategory) && !Pin->PinType.bIsArray));
+		checkSlow(bIsObject == ((Schema->PC_Object == Pin->PinType.PinCategory || Schema->PC_Class == Pin->PinType.PinCategory || Schema->PC_Asset == Pin->PinType.PinCategory) && !Pin->PinType.bIsArray));
 
 		const FString& MetadataDefaultValue = Property->GetMetaData(TEXT("MakeStructureDefaultValue"));
 		if (!MetadataDefaultValue.IsEmpty())
