[UE-15683] Functions with a WorldContext meta tag are now accessible if that tag doesn't match a specific param.

diff --git a/Engine/Source/Editor/BlueprintGraph/Private/EdGraphSchema_K2.cpp b/Engine/Source/Editor/BlueprintGraph/Private/EdGraphSchema_K2.cpp
--- a/Engine/Source/Editor/BlueprintGraph/Private/EdGraphSchema_K2.cpp
+++ b/Engine/Source/Editor/BlueprintGraph/Private/EdGraphSchema_K2.cpp
@@ -639,19 +639,23 @@ bool UEdGraphSchema_K2::CanFunctionBeUsedInGraph(const UClass* InClass, const UF
 		{
 			if (InDestGraph && !InFunction->HasMetaData(FBlueprintMetadata::MD_CallableWithoutWorldContext))
 			{
-				auto BP = FBlueprintEditorUtils::FindBlueprintForGraph(InDestGraph);
-				const bool bIsFunctLib = BP && (EBlueprintType::BPTYPE_FunctionLibrary == BP->BlueprintType);
-				UClass* ParentClass = BP ? BP->ParentClass : NULL;
-				const bool bIncompatibleParrent = ParentClass && (!ParentClass->GetDefaultObject()->ImplementsGetWorld() && !ParentClass->HasMetaData(FBlueprintMetadata::MD_ShowWorldContextPin));
-				if (!bIsFunctLib && bIncompatibleParrent)
+				const FString& ContextParam = InFunction->GetMetaData(FBlueprintMetadata::MD_WorldContext);
+				if (InFunction->FindPropertyByName(FName(*ContextParam)) != nullptr)
 				{
-					if(OutReason != nullptr)
+					auto BP = FBlueprintEditorUtils::FindBlueprintForGraph(InDestGraph);
+					const bool bIsFunctLib = BP && (EBlueprintType::BPTYPE_FunctionLibrary == BP->BlueprintType);
+					UClass* ParentClass = BP ? BP->ParentClass : NULL;
+					const bool bIncompatibleParrent = ParentClass && (!ParentClass->GetDefaultObject()->ImplementsGetWorld() && !ParentClass->HasMetaData(FBlueprintMetadata::MD_ShowWorldContextPin));
+					if (!bIsFunctLib && bIncompatibleParrent)
 					{
-						*OutReason = LOCTEXT("FunctionRequiresWorldContext", "Function requires a world context.");
-					}
+						if (OutReason != nullptr)
+						{
+							*OutReason = LOCTEXT("FunctionRequiresWorldContext", "Function requires a world context.");
+						}
 
-					return false;
-				}
+						return false;
+					}
+				}	
 			}
 		}
 
