K2Node_ConvertAsset - fixes

diff --git a/Engine/Source/Editor/BlueprintGraph/Classes/K2Node_ConvertAsset.h b/Engine/Source/Editor/BlueprintGraph/Classes/K2Node_ConvertAsset.h
--- a/Engine/Source/Editor/BlueprintGraph/Classes/K2Node_ConvertAsset.h
+++ b/Engine/Source/Editor/BlueprintGraph/Classes/K2Node_ConvertAsset.h
@@ -21,6 +21,7 @@ public:
 	// End of UEdGraphNode interface
 
 	// UK2Node interface
+	virtual void PostReconstructNode() override;
 	virtual bool IsNodePure() const override { return true; }
 	virtual bool ShouldDrawCompact() const override { return true; }
 	virtual int32 GetNodeRefreshPriority() const override { return EBaseNodeRefreshPriority::Low_UsesDependentWildcard; }
@@ -33,4 +34,7 @@ public:
 	virtual void GetMenuActions(FBlueprintActionDatabaseRegistrar& ActionRegistrar) const override;
 	virtual FText GetMenuCategory() const override;
 	// End of UK2Node interface
+
+protected:
+	void RefreshPinTypes();
 };
\ No newline at end of file
diff --git a/Engine/Source/Editor/BlueprintGraph/Private/K2Node_ConvertAsset.cpp b/Engine/Source/Editor/BlueprintGraph/Private/K2Node_ConvertAsset.cpp
--- a/Engine/Source/Editor/BlueprintGraph/Private/K2Node_ConvertAsset.cpp
+++ b/Engine/Source/Editor/BlueprintGraph/Private/K2Node_ConvertAsset.cpp
@@ -55,15 +55,13 @@ bool UK2Node_ConvertAsset::IsConnectionDisallowed(const UEdGraphPin* MyPin, cons
 	return false;
 }
 
-void UK2Node_ConvertAsset::NotifyPinConnectionListChanged(UEdGraphPin* Pin)
+void UK2Node_ConvertAsset::RefreshPinTypes()
 {
-	Super::NotifyPinConnectionListChanged(Pin);
-
 	const UEdGraphSchema_K2* K2Schema = CastChecked<UEdGraphSchema_K2>(GetSchema());
 	auto InutPin = FindPin(UK2Node_ConvertAssetImpl::InputPinName);
 	auto OutputPin = FindPin(UK2Node_ConvertAssetImpl::OutputPinName);
 	ensure(InutPin && OutputPin);
-	if (InutPin && OutputPin && (InutPin == Pin))
+	if (InutPin && OutputPin)
 	{
 		const bool bIsConnected = InutPin->LinkedTo.Num() > 0;
 		UClass* TargetType = bIsConnected ? GetTargetClass() : nullptr;
@@ -72,7 +70,7 @@ void UK2Node_ConvertAsset::NotifyPinConnectionListChanged(UEdGraphPin* Pin)
 		const FString InputCategory = bIsConnected
 			? (bIsAssetClass ? K2Schema->PC_AssetClass : K2Schema->PC_Asset)
 			: K2Schema->PC_Wildcard;
-		InutPin->PinType = FEdGraphPinType(InputCategory,FString(), TargetType, false, false);
+		InutPin->PinType = FEdGraphPinType(InputCategory, FString(), TargetType, false, false);
 
 		const FString OutputCategory = bIsConnected
 			? (bIsAssetClass ? K2Schema->PC_Class : K2Schema->PC_Object)
@@ -105,6 +103,20 @@ void UK2Node_ConvertAsset::NotifyPinConnectionListChanged(UEdGraphPin* Pin)
 	}
 }
 
+void UK2Node_ConvertAsset::PostReconstructNode()
+{
+	RefreshPinTypes();
+}
+
+void UK2Node_ConvertAsset::NotifyPinConnectionListChanged(UEdGraphPin* Pin)
+{
+	Super::NotifyPinConnectionListChanged(Pin);
+	if (Pin && (UK2Node_ConvertAssetImpl::InputPinName == Pin->PinName))
+	{
+		RefreshPinTypes();
+	}
+}
+
 void UK2Node_ConvertAsset::AllocateDefaultPins()
 {
 	const UEdGraphSchema_K2* K2Schema = CastChecked<UEdGraphSchema_K2>(GetSchema());
@@ -158,23 +170,32 @@ void UK2Node_ConvertAsset::ExpandNode(class FKismetCompilerContext& CompilerCont
 		auto ConvertInput = ConvertToObjectFunc->FindPin(ConvertInputName);
 		bIsErrorFree &= InputPin && ConvertInput && CompilerContext.MovePinLinksToIntermediate(*InputPin, *ConvertInput).CanSafeConnect();
 
-		//Create Cast Node
-		UK2Node_DynamicCast* CastNode = bIsAssetClass
-			? CompilerContext.SpawnIntermediateNode<UK2Node_ClassDynamicCast>(this, SourceGraph)
-			: CompilerContext.SpawnIntermediateNode<UK2Node_DynamicCast>(this, SourceGraph);
-		CastNode->SetPurity(true);
-		CastNode->TargetType = TargetType;
-		CastNode->AllocateDefaultPins();
-
-		// Connect Object/Class to Cast
-		auto CastInput = CastNode->GetCastSourcePin();
 		auto ConvertOutput = ConvertToObjectFunc->GetReturnValuePin();
-		bIsErrorFree &= ConvertOutput && CastInput && Schema->TryCreateConnection(ConvertOutput, CastInput);
+		UEdGraphPin* InnerOutput = nullptr;
+		if (UObject::StaticClass() != TargetType)
+		{
+			//Create Cast Node
+			UK2Node_DynamicCast* CastNode = bIsAssetClass
+				? CompilerContext.SpawnIntermediateNode<UK2Node_ClassDynamicCast>(this, SourceGraph)
+				: CompilerContext.SpawnIntermediateNode<UK2Node_DynamicCast>(this, SourceGraph);
+			CastNode->SetPurity(true);
+			CastNode->TargetType = TargetType;
+			CastNode->AllocateDefaultPins();
+
+			// Connect Object/Class to Cast
+			auto CastInput = CastNode->GetCastSourcePin();
+			bIsErrorFree &= ConvertOutput && CastInput && Schema->TryCreateConnection(ConvertOutput, CastInput);
+
+			// Connect output to cast
+			InnerOutput = CastNode->GetCastResultPin();
+		}
+		else
+		{
+			InnerOutput = ConvertOutput;
+		}
 
-		// Connect output to cast
-		auto OutputPin = FindPin(UK2Node_ConvertAssetImpl::InputPinName);
-		auto CastOutput = CastNode->GetCastResultPin();
-		bIsErrorFree &= OutputPin && CastOutput && CompilerContext.MovePinLinksToIntermediate(*OutputPin, *CastOutput).CanSafeConnect();
+		auto OutputPin = FindPin(UK2Node_ConvertAssetImpl::OutputPinName);
+		bIsErrorFree &= OutputPin && InnerOutput && CompilerContext.MovePinLinksToIntermediate(*OutputPin, *InnerOutput).CanSafeConnect();
 	}
 
 	if (!bIsErrorFree)
