[UE-15998] Guarding against a crash in the BlueprintActionDatabase. Adding an ensure so we can have more info as to what went wrong (for future fixes).

diff --git a/Engine/Source/Editor/BlueprintGraph/Private/BlueprintActionFilter.cpp b/Engine/Source/Editor/BlueprintGraph/Private/BlueprintActionFilter.cpp
--- a/Engine/Source/Editor/BlueprintGraph/Private/BlueprintActionFilter.cpp
+++ b/Engine/Source/Editor/BlueprintGraph/Private/BlueprintActionFilter.cpp
@@ -1686,8 +1686,6 @@ FBlueprintActionFilter::FBlueprintActionFilter(uint32 Flags/*= 0x00*/)
 	//
 	// this test in-particular spawns a template-node and then calls 
 	// AllocateDefaultPins() which is costly, so it should be very last!
-	AddRejectionTest(FRejectionTestDelegate::CreateStatic(IsStaleFieldAction));
-
 	AddRejectionTest(FRejectionTestDelegate::CreateStatic(IsIncompatibleAnimNotification));
 	AddRejectionTest(FRejectionTestDelegate::CreateStatic(IsNodeTemplateSelfFiltered));
 	AddRejectionTest(FRejectionTestDelegate::CreateStatic(IsMissingMatchingPinParam));
@@ -1722,6 +1720,11 @@ FBlueprintActionFilter::FBlueprintActionFilter(uint32 Flags/*= 0x00*/)
 	AddRejectionTest(FRejectionTestDelegate::CreateStatic(IsUnBoundBindingSpawner));
 	AddRejectionTest(FRejectionTestDelegate::CreateStatic(IsOutOfScopeLocalVariable));
 	AddRejectionTest(FRejectionTestDelegate::CreateStatic(IsLevelScriptActionValid));
+
+
+	// added as the first rejection test, so that we don't operate on stale 
+	// (TRASH/REINST) class fields
+	AddRejectionTest(FRejectionTestDelegate::CreateStatic(IsStaleFieldAction));
 }
 
 //------------------------------------------------------------------------------
diff --git a/Engine/Source/Editor/BlueprintGraph/Private/BlueprintNodeSpawner.cpp b/Engine/Source/Editor/BlueprintGraph/Private/BlueprintNodeSpawner.cpp
--- a/Engine/Source/Editor/BlueprintGraph/Private/BlueprintNodeSpawner.cpp
+++ b/Engine/Source/Editor/BlueprintGraph/Private/BlueprintNodeSpawner.cpp
@@ -3,6 +3,7 @@
 #include "BlueprintGraphPrivatePCH.h"
 #include "BlueprintNodeSpawner.h"
 #include "BlueprintNodeTemplateCache.h"
+#include "BlueprintNodeSpawnerUtils.h"
 
 /*******************************************************************************
  * Static UBlueprintNodeSpawner Helpers
@@ -68,6 +69,13 @@ void UBlueprintNodeSpawner::BeginDestroy()
 //------------------------------------------------------------------------------
 void UBlueprintNodeSpawner::Prime()
 {
+	if (FBlueprintNodeSpawnerUtils::IsStaleFieldAction(this))
+	{
+		UField const* AssociatedField = FBlueprintNodeSpawnerUtils::GetAssociatedField(this);
+		ensureMsgf(false, TEXT("Priming invalid BlueprintActionDatabase entry (for %s). Was the database properly updated when this class was compiled?"), *AssociatedField->GetPathName());
+		return;
+	}
+
 	if (UEdGraphNode* CachedTemplateNode = GetTemplateNode())
 	{
 		// since we're priming incrementally, someone could have already
