Can right click in collapsed graphs to add local variables that are part of the outer function graph.

diff --git a/Engine/Source/Editor/BlueprintGraph/Private/BlueprintActionFilter.cpp b/Engine/Source/Editor/BlueprintGraph/Private/BlueprintActionFilter.cpp
--- a/Engine/Source/Editor/BlueprintGraph/Private/BlueprintActionFilter.cpp
+++ b/Engine/Source/Editor/BlueprintGraph/Private/BlueprintActionFilter.cpp
@@ -927,7 +927,7 @@ static bool BlueprintActionFilterImpl::IsOutOfScopeLocalVariable(FBlueprintActio
 			UEdGraph const* VarOuter = Cast<UEdGraph const>(VarSpawner->GetVarOuter());
 			for (UEdGraph* Graph : Filter.Context.Graphs)
 			{
-				if (Graph != VarOuter)
+				if (FBlueprintEditorUtils::GetTopLevelGraph(Graph) != VarOuter)
 				{
 					bIsFilteredOut = true;
 					break;
