The "Get" pin(s) on "Set" variable nodes will now display the name of the pin if the pin has been generated by splitting a struct pin. This allows the user to identify what value the split pins are for.

diff --git a/Engine/Source/Editor/BlueprintGraph/Private/K2Node_VariableSet.cpp b/Engine/Source/Editor/BlueprintGraph/Private/K2Node_VariableSet.cpp
--- a/Engine/Source/Editor/BlueprintGraph/Private/K2Node_VariableSet.cpp
+++ b/Engine/Source/Editor/BlueprintGraph/Private/K2Node_VariableSet.cpp
@@ -333,7 +333,7 @@ FText UK2Node_VariableSet::GetPinNameOverride(const UEdGraphPin& Pin) const
 {
 	// Stop the output pin for the variable, effectively the "get" pin, from displaying a name.
 	const UEdGraphSchema_K2* K2Schema = GetDefault<UEdGraphSchema_K2>();
-	if(Pin.Direction == EGPD_Output || Pin.PinType.PinCategory == K2Schema->PC_Exec)
+	if(Pin.ParentPin == nullptr && (Pin.Direction == EGPD_Output || Pin.PinType.PinCategory == K2Schema->PC_Exec))
 	{
 		return FText::GetEmpty();
 	}
