[UE-16163] Static functions are back to being unaffected by class HideCategories data.

diff --git a/Engine/Source/Editor/BlueprintGraph/Private/BlueprintActionFilter.cpp b/Engine/Source/Editor/BlueprintGraph/Private/BlueprintActionFilter.cpp
--- a/Engine/Source/Editor/BlueprintGraph/Private/BlueprintActionFilter.cpp
+++ b/Engine/Source/Editor/BlueprintGraph/Private/BlueprintActionFilter.cpp
@@ -772,36 +772,25 @@ static bool BlueprintActionFilterImpl::IsFieldCategoryHidden(FBlueprintActionFil
 {
 	bool bIsFilteredOut = false;
 
-	DECLARE_DELEGATE_RetVal_OneParam(bool, FIsFieldHiddenDelegate, UClass*);
-	FIsFieldHiddenDelegate IsFieldHiddenDelegate;
-
-	// Use the UiSpec to get the category
-	FBlueprintActionUiSpec UiSpec = BlueprintAction.NodeSpawner->GetUiSpec(Filter.Context, BlueprintAction.GetBindings());
-	auto IsNodeHiddenLambda = [](UClass* Class, FText InCategory)->bool
+	const UFunction* NodeFunction = BlueprintAction.GetAssociatedFunction();
+	if ((NodeFunction != nullptr) && NodeFunction->HasAnyFunctionFlags(FUNC_Static))
 	{
-		return FEditorCategoryUtils::IsCategoryHiddenFromClass(Class->GetAuthoritativeClass(), InCategory.ToString());
-	};
-	IsFieldHiddenDelegate = FIsFieldHiddenDelegate::CreateStatic(IsNodeHiddenLambda, UiSpec.Category);
-
-	if (IsFieldHiddenDelegate.IsBound())
+		bIsFilteredOut = false;
+	}
+	else
 	{
-		bIsFilteredOut = Filter.TargetClasses.Num() > 0;
+		bIsFilteredOut = (Filter.TargetClasses.Num() > 0);
 
 		for (UClass* TargetClass : Filter.TargetClasses)
 		{
-			if (!IsFieldHiddenDelegate.Execute(TargetClass))
+			// Use the UiSpec to get the category
+			FBlueprintActionUiSpec UiSpec = BlueprintAction.NodeSpawner->GetUiSpec(Filter.Context, BlueprintAction.GetBindings());
+			if (!FEditorCategoryUtils::IsCategoryHiddenFromClass(TargetClass->GetAuthoritativeClass(), UiSpec.Category.ToString()))
 			{
 				bIsFilteredOut = false;
 				break;
 			}
 		}
-
-		// handled now in each IsBindingCompatible() check
-// 		for (auto BindingIt = BlueprintAction.GetBindings().CreateConstIterator(); BindingIt && !bIsFilteredOut; ++BindingIt)
-// 		{
-// 			UClass* BindingClass = FBlueprintNodeSpawnerUtils::GetBindingClass(BindingIt->Get());
-// 			bIsFilteredOut = IsFieldHiddenDelegate.Execute(BindingClass);
-// 		}
 	}
 
 	return bIsFilteredOut;
@@ -1785,7 +1774,7 @@ bool FBlueprintActionFilter::IsFilteredByThis(FBlueprintActionInfo& BlueprintAct
 	// for debugging purposes:
 // 	FBlueprintActionUiSpec UiSpec = BlueprintAction.NodeSpawner->GetUiSpec(Context, BlueprintAction.GetBindings());
 // 	bool bDebugBreak = false;
-// 	if (UiSpec.MenuName.ToString().Contains("Set Array Elem"))
+// 	if (UiSpec.MenuName.ToString().Contains("Print String"))
 // 	{
 // 		bDebugBreak = true;
 // 	}
