[UEBP-56] Optimized K2Node_CallFunction construction, pin tool type generation was being performed for all pins, and is expensive. Made it lazy (pins created when you mouse over a pin)

diff --git a/Engine/Source/Editor/BlueprintGraph/Classes/K2Node_CallFunction.h b/Engine/Source/Editor/BlueprintGraph/Classes/K2Node_CallFunction.h
--- a/Engine/Source/Editor/BlueprintGraph/Classes/K2Node_CallFunction.h
+++ b/Engine/Source/Editor/BlueprintGraph/Classes/K2Node_CallFunction.h
@@ -52,6 +52,9 @@ private:
 	/** Constructing FText strings can be costly, so we cache the node's tooltip */
 	FNodeTextCache CachedTooltip;
 
+	/** Flag used to track validity of pin tooltips, when tooltips are invalid they will be refreshed before being displayed */
+	mutable bool bPinTooltipsValid;
+
 public:
 
 	// UObject interface
@@ -60,6 +63,7 @@ public:
 	// End of UObject interface
 
 	// UEdGraphNode interface
+	virtual void GetPinHoverText(const UEdGraphPin& Pin, FString& HoverTextOut) const override;
 	virtual void AllocateDefaultPins() override;
 	virtual void DestroyNode() override;
 	virtual FLinearColor GetNodeTitleColor() const override;
@@ -198,6 +202,9 @@ private:
 	 */
 	bool ReconnectPureExecPins(TArray<UEdGraphPin*>& OldPins);
 
+	/** Invalidates current pin tool tips, so that they will be refreshed before being displayed: */
+	void InvalidatePinTooltips();
+
 protected:
 	/** Helper function to ensure function is called in our context */
 	virtual void EnsureFunctionIsInBlueprint();
diff --git a/Engine/Source/Editor/BlueprintGraph/Private/K2Node_CallFunction.cpp b/Engine/Source/Editor/BlueprintGraph/Private/K2Node_CallFunction.cpp
--- a/Engine/Source/Editor/BlueprintGraph/Private/K2Node_CallFunction.cpp
+++ b/Engine/Source/Editor/BlueprintGraph/Private/K2Node_CallFunction.cpp
@@ -190,12 +190,6 @@ void FDynamicOutputHelper::ConformOutputType() const
 		{
 			DynamicOutPin->PinType.PinSubCategoryObject = PickedClass;
 
-			if (UFunction* Function = FuncNode->GetTargetFunction())
-			{
-				DynamicOutPin->PinToolTip.Empty();
-				FuncNode->GeneratePinTooltipFromFunction(*DynamicOutPin, Function);
-			}
-
 			// leave the connection, and instead bring the user's attention to 
 			// it via a ValidateNodeDuringCompilation() error
 // 			const UEdGraphSchema* Schema = FuncNode->GetSchema();
@@ -400,6 +394,7 @@ bool FDynamicOutputHelper::CanConformPinType(const UK2Node_CallFunction* FuncNod
 
 UK2Node_CallFunction::UK2Node_CallFunction(const FObjectInitializer& ObjectInitializer)
 	: Super(ObjectInitializer)
+	, bPinTooltipsValid(false)
 {
 }
 
@@ -506,8 +501,25 @@ FText UK2Node_CallFunction::GetNodeTitle(ENodeTitleType::Type TitleType) const
 	}
 }
 
+void UK2Node_CallFunction::GetPinHoverText(const UEdGraphPin& Pin, FString& HoverTextOut) const
+{
+	if (!bPinTooltipsValid)
+	{
+		for (auto& Pin : Pins)
+		{
+			GeneratePinTooltip(*Pin);
+		}
+
+		bPinTooltipsValid = true;
+	}
+
+	return UK2Node::GetPinHoverText(Pin, HoverTextOut);
+}
+
 void UK2Node_CallFunction::AllocateDefaultPins()
 {
+	InvalidatePinTooltips();
+
 	UBlueprint* MyBlueprint = GetBlueprint();
 	
 	UFunction* Function = GetTargetFunction();
@@ -797,9 +809,6 @@ bool UK2Node_CallFunction::CreatePinsForFunctionCall(const UFunction* Function)
 	//Renamed self pin to target
 	SelfPin->PinFriendlyName =  LOCTEXT("Target", "Target");
 
-	// fill out the self-pin's default tool-tip
-	GeneratePinTooltip(*SelfPin);
-
 	const bool bIsProtectedFunc = Function->GetBoolMetaData(FBlueprintMetadata::MD_Protected);
 	const bool bIsStaticFunc = Function->HasAllFunctionFlags(FUNC_Static);
 
@@ -861,9 +870,6 @@ bool UK2Node_CallFunction::CreatePinsForFunctionCall(const UFunction* Function)
 			}
 
 			K2Schema->SetPinDefaultValue(Pin, Function, Param);
-
-			// setup the default tool-tip text for this pin
-			GeneratePinTooltip(*Pin);
 			
 			if (PinsToHide.Contains(Pin->PinName))
 			{
@@ -901,6 +907,7 @@ bool UK2Node_CallFunction::CreatePinsForFunctionCall(const UFunction* Function)
 void UK2Node_CallFunction::PostReconstructNode()
 {
 	Super::PostReconstructNode();
+	InvalidatePinTooltips();
 
 	FCustomStructureParamHelper::UpdateCustomStructurePins(GetTargetFunction(), this);
 
@@ -976,12 +983,14 @@ void UK2Node_CallFunction::NotifyPinConnectionListChanged(UEdGraphPin* Pin)
 		}
 	}
 
+	InvalidatePinTooltips();
 	FDynamicOutputHelper(Pin).ConformOutputType();
 }
 
 void UK2Node_CallFunction::PinDefaultValueChanged(UEdGraphPin* Pin)
 {
 	Super::PinDefaultValueChanged(Pin);
+	InvalidatePinTooltips();
 	FDynamicOutputHelper(Pin).ConformOutputType();
 }
 
@@ -2138,6 +2147,11 @@ bool UK2Node_CallFunction::ReconnectPureExecPins(TArray<UEdGraphPin*>& OldPins)
 	return false;
 }
 
+void UK2Node_CallFunction::InvalidatePinTooltips()
+{
+	bPinTooltipsValid = false;
+}
+
 FText UK2Node_CallFunction::GetToolTipHeading() const
 {
 	FText Heading = Super::GetToolTipHeading();
