Fixed a crash in UK2Node_SetFieldsInStruct::ValidateNodeDuringCompilation

diff --git a/Engine/Source/Editor/BlueprintGraph/Private/K2Node_SetFieldsInStruct.cpp b/Engine/Source/Editor/BlueprintGraph/Private/K2Node_SetFieldsInStruct.cpp
--- a/Engine/Source/Editor/BlueprintGraph/Private/K2Node_SetFieldsInStruct.cpp
+++ b/Engine/Source/Editor/BlueprintGraph/Private/K2Node_SetFieldsInStruct.cpp
@@ -95,13 +95,11 @@ void UK2Node_SetFieldsInStruct::ValidateNodeDuringCompilation(FCompilerResultsLo
 {
 	Super::ValidateNodeDuringCompilation(MessageLog);
 
-	if (UEdGraphPin* FoundPin = FindPinChecked(SetFieldsInStructHelper::StructRefPinName()))
+	UEdGraphPin* FoundPin = FindPin(SetFieldsInStructHelper::StructRefPinName());
+	if (!ensure(FoundPin) || (FoundPin->LinkedTo.Num() <= 0))
 	{
-		if (FoundPin->LinkedTo.Num() <= 0)
-		{
-			FText ErrorMessage = LOCTEXT("SetStructFields_NoStructRefError", "The @@ pin must be connected to the struct that you wish to set.");
-			MessageLog.Error(*ErrorMessage.ToString(), FoundPin);
-		}
+		FText ErrorMessage = LOCTEXT("SetStructFields_NoStructRefError", "The @@ pin must be connected to the struct that you wish to set.");
+		MessageLog.Error(*ErrorMessage.ToString(), FoundPin);
 	}
 }
 
