UE-15907 Opening the Pin Type Selector is way too slow

Optimization.

diff --git a/Engine/Source/Editor/BlueprintGraph/Classes/EdGraphSchema_K2.h b/Engine/Source/Editor/BlueprintGraph/Classes/EdGraphSchema_K2.h
--- a/Engine/Source/Editor/BlueprintGraph/Classes/EdGraphSchema_K2.h
+++ b/Engine/Source/Editor/BlueprintGraph/Classes/EdGraphSchema_K2.h
@@ -191,6 +191,8 @@ struct FBlueprintCallableFunctionRedirect
 	FString ClassParamName;
 };
 
+struct FTypesDatabase;
+
 UCLASS(config=Editor)
 class BLUEPRINTGRAPH_API UEdGraphSchema_K2 : public UEdGraphSchema
 {
@@ -283,7 +285,7 @@ public:
 		{
 		}
 
-		void Init(const FText& FriendlyCategoryName, const FString& CategoryName, const UEdGraphSchema_K2* Schema, const FText& InTooltip, bool bInReadOnly);
+		void Init(const FText& FriendlyCategoryName, const FString& CategoryName, const UEdGraphSchema_K2* Schema, const FText& InTooltip, bool bInReadOnly, FTypesDatabase* TypesDatabase);
 	
 	private:
 		/** The pin type corresponding to the schema type */
@@ -292,6 +294,10 @@ public:
 		/** Asset Reference, used when PinType.PinSubCategoryObject is not loaded yet */
 		FStringAssetReference SubCategoryObjectAssetReference;
 
+		FText CachedDescription;
+
+		FText GenerateDescription();
+
 	public:
 		/** The children of this pin type */
 		TArray< TSharedPtr<FPinTypeTreeInfo> > Children;
@@ -312,7 +318,7 @@ public:
 			PinType.PinSubCategory = SubCategory;
 		}
 
-		FPinTypeTreeInfo(const FText& FriendlyName, const FString& CategoryName, const UEdGraphSchema_K2* Schema, const FText& InTooltip, bool bInReadOnly = false);
+		FPinTypeTreeInfo(const FText& FriendlyName, const FString& CategoryName, const UEdGraphSchema_K2* Schema, const FText& InTooltip, bool bInReadOnly = false, FTypesDatabase* TypesDatabase = nullptr);
 		FPinTypeTreeInfo(const FString& CategoryName, UObject* SubCategoryObject, const FText& InTooltip, bool bInReadOnly = false);
 		FPinTypeTreeInfo(const FString& CategoryName, const FStringAssetReference& SubCategoryObject, const FText& InTooltip, bool bInReadOnly = false);
 
diff --git a/Engine/Source/Editor/BlueprintGraph/Private/EdGraphSchema_K2.cpp b/Engine/Source/Editor/BlueprintGraph/Private/EdGraphSchema_K2.cpp
--- a/Engine/Source/Editor/BlueprintGraph/Private/EdGraphSchema_K2.cpp
+++ b/Engine/Source/Editor/BlueprintGraph/Private/EdGraphSchema_K2.cpp
@@ -103,11 +103,52 @@ const FName FBlueprintMetadata::MD_ArrayDependentParam(TEXT("ArrayTypeDependentP
 
 #define LOCTEXT_NAMESPACE "KismetSchema"
 
-UEdGraphSchema_K2::FPinTypeTreeInfo::FPinTypeTreeInfo(const FText& InFriendlyCategoryName, const FString& CategoryName, const UEdGraphSchema_K2* Schema, const FText& InTooltip, bool bInReadOnly/*=false*/)
+UEdGraphSchema_K2::FPinTypeTreeInfo::FPinTypeTreeInfo(const FText& InFriendlyCategoryName, const FString& CategoryName, const UEdGraphSchema_K2* Schema, const FText& InTooltip, bool bInReadOnly/*=false*/, FTypesDatabase* TypesDatabase /*=nullptr*/)
 {
-	Init(InFriendlyCategoryName, CategoryName, Schema, InTooltip, bInReadOnly);
+	Init(InFriendlyCategoryName, CategoryName, Schema, InTooltip, bInReadOnly, TypesDatabase);
 }
 
+struct FUnloadedAssetData
+{
+	FStringAssetReference StringAssetReference;
+	FText AssetFriendlyName;
+	FText Tooltip;
+
+	FUnloadedAssetData()
+	{}
+
+	FUnloadedAssetData(const FAssetData& InAsset)
+		: StringAssetReference(InAsset.ToStringReference())
+		, AssetFriendlyName(FText::FromString(FName::NameToDisplayString(InAsset.AssetName.ToString(), false)))
+	{
+		const FString* TooltipPtr = InAsset.TagsAndValues.Find(TEXT("Tooltip"));
+		Tooltip = FText::FromString((TooltipPtr && !TooltipPtr->IsEmpty()) ? *TooltipPtr : InAsset.ObjectPath.ToString());
+	}
+};
+
+struct FLoadedAssetData
+{
+	FText Tooltip;
+	UObject* Object;
+
+	FLoadedAssetData() : Object(nullptr) {}
+	FLoadedAssetData(UObject* InObject) 
+		: Object(InObject) 
+	{
+		UStruct* Struct = Cast<UStruct>(Object);
+		Tooltip = Struct ? Struct->GetToolTipText() : FText::GetEmpty();
+	}
+};
+
+struct FTypesDatabase
+{
+	typedef TSharedPtr<TArray<FLoadedAssetData>> FLoadedTypesList;
+	TMap<FString, FLoadedTypesList> LoadedTypesMap;
+
+	typedef TSharedPtr<TArray<FUnloadedAssetData>> FUnLoadedTypesList;
+	TMap<FString, FUnLoadedTypesList> UnLoadedTypesMap;
+};
+
 /** Helper class to gather variable types */
 class FGatherTypesHelper
 {
@@ -205,6 +246,204 @@ private:
 	}
 
 public:
+	static void FillLoadedTypesDatabase(FTypesDatabase& TypesDatabase)
+	{
+		// Loaded types
+		TypesDatabase.LoadedTypesMap.Reset();
+
+		//(Type == UEdGraphSchema_K2::PC_Struct)
+		{
+			FTypesDatabase::FLoadedTypesList LoadedTypesList = MakeShareable(new TArray<FLoadedAssetData>());
+			// Find script structs marked with "BlueprintType=true" in their metadata, and add to the list
+			for (TObjectIterator<UScriptStruct> StructIt; StructIt; ++StructIt)
+			{
+				UScriptStruct* ScriptStruct = *StructIt;
+				if (UEdGraphSchema_K2::IsAllowableBlueprintVariableType(ScriptStruct))
+				{
+					LoadedTypesList->Add(FLoadedAssetData(ScriptStruct));
+				}
+			}
+			TypesDatabase.LoadedTypesMap.Add(UEdGraphSchema_K2::PC_Struct, LoadedTypesList);
+		}
+
+		//(Type == UEdGraphSchema_K2::PC_Class || Type == UEdGraphSchema_K2::PC_AssetClass) UEdGraphSchema_K2::PC_Interface
+		{
+			FTypesDatabase::FLoadedTypesList LoadedTypesList = MakeShareable(new TArray<FLoadedAssetData>());
+			FTypesDatabase::FLoadedTypesList InterfaceLoadedTypesList = MakeShareable(new TArray<FLoadedAssetData>());
+			// Generate a list of all potential objects which have "BlueprintType=true" in their metadata
+			for (TObjectIterator<UClass> ClassIt; ClassIt; ++ClassIt)
+			{
+				UClass* CurrentClass = *ClassIt;
+				if (UEdGraphSchema_K2::IsAllowableBlueprintVariableType(CurrentClass) && !CurrentClass->HasAnyClassFlags(CLASS_Deprecated))
+				{
+					LoadedTypesList->Add(FLoadedAssetData(CurrentClass));
+				}
+				
+				if (CurrentClass->IsChildOf(UInterface::StaticClass()) && UEdGraphSchema_K2::IsAllowableBlueprintVariableType(CurrentClass))
+				{
+					InterfaceLoadedTypesList->Add(FLoadedAssetData(CurrentClass));
+				}
+			}
+			TypesDatabase.LoadedTypesMap.Add(UEdGraphSchema_K2::PC_Class, LoadedTypesList);
+			TypesDatabase.LoadedTypesMap.Add(UEdGraphSchema_K2::PC_AssetClass, LoadedTypesList);
+			TypesDatabase.LoadedTypesMap.Add(UEdGraphSchema_K2::PC_Interface, InterfaceLoadedTypesList);
+		}
+
+		//(Type == UEdGraphSchema_K2::PC_Object || Type == UEdGraphSchema_K2::PC_Asset)
+		{
+			FTypesDatabase::FLoadedTypesList LoadedTypesList = MakeShareable(new TArray<FLoadedAssetData>());
+			// Generate a list of all potential objects which have "BlueprintType=true" in their metadata (that aren't interfaces)
+			for (TObjectIterator<UClass> ClassIt; ClassIt; ++ClassIt)
+			{
+				UClass* CurrentClass = *ClassIt;
+				if (!CurrentClass->IsChildOf(UInterface::StaticClass()) && UEdGraphSchema_K2::IsAllowableBlueprintVariableType(CurrentClass) && !CurrentClass->HasAnyClassFlags(CLASS_Deprecated))
+				{
+					LoadedTypesList->Add(FLoadedAssetData(CurrentClass));
+				}
+			}
+			TypesDatabase.LoadedTypesMap.Add(UEdGraphSchema_K2::PC_Object, LoadedTypesList);
+			TypesDatabase.LoadedTypesMap.Add(UEdGraphSchema_K2::PC_Asset, LoadedTypesList);
+		}
+		
+		//(Type == UEdGraphSchema_K2::PC_Enum)
+		{
+			FTypesDatabase::FLoadedTypesList LoadedTypesList = MakeShareable(new TArray<FLoadedAssetData>());
+			// Generate a list of all potential enums which have "BlueprintType=true" in their metadata
+			for (TObjectIterator<UEnum> EnumIt; EnumIt; ++EnumIt)
+			{
+				UEnum* CurrentEnum = *EnumIt;
+				if (UEdGraphSchema_K2::IsAllowableBlueprintVariableType(CurrentEnum))
+				{
+					LoadedTypesList->Add(FLoadedAssetData(CurrentEnum));
+				}
+			}
+			TypesDatabase.LoadedTypesMap.Add(UEdGraphSchema_K2::PC_Enum, LoadedTypesList);
+		}
+	}
+
+	static void FillUnLoadedTypesDatabase(FTypesDatabase& TypesDatabase)
+	{
+		// Loaded types
+		TypesDatabase.UnLoadedTypesMap.Reset();
+
+		const FAssetRegistryModule& AssetRegistryModule = FModuleManager::LoadModuleChecked<FAssetRegistryModule>(TEXT("AssetRegistry"));
+
+		{
+			TArray<FAssetData> AssetData;
+			AssetRegistryModule.Get().GetAssetsByClass(UUserDefinedStruct::StaticClass()->GetFName(), AssetData);
+
+			FTypesDatabase::FUnLoadedTypesList UnLoadedTypesList = MakeShareable(new TArray<FUnloadedAssetData>());
+			for (int32 AssetIndex = 0; AssetIndex < AssetData.Num(); ++AssetIndex)
+			{
+				const FAssetData& Asset = AssetData[AssetIndex];
+				if (Asset.IsValid() && !Asset.IsAssetLoaded())
+				{
+					UnLoadedTypesList->Add(FUnloadedAssetData(Asset));
+				}
+			}
+
+			TypesDatabase.UnLoadedTypesMap.Add(UEdGraphSchema_K2::PC_Struct, UnLoadedTypesList);
+		}
+
+		{
+			TArray<FAssetData> AssetData;
+			AssetRegistryModule.Get().GetAssetsByClass(UUserDefinedEnum::StaticClass()->GetFName(), AssetData);
+
+			FTypesDatabase::FUnLoadedTypesList UnLoadedTypesList = MakeShareable(new TArray<FUnloadedAssetData>());
+			for (int32 AssetIndex = 0; AssetIndex < AssetData.Num(); ++AssetIndex)
+			{
+				const FAssetData& Asset = AssetData[AssetIndex];
+				if (Asset.IsValid() && !Asset.IsAssetLoaded())
+				{
+					UnLoadedTypesList->Add(FUnloadedAssetData(Asset));
+				}
+			}
+
+			TypesDatabase.UnLoadedTypesMap.Add(UEdGraphSchema_K2::PC_Enum, UnLoadedTypesList);
+		}
+
+		//else if (Schema->PC_Object == CategoryName || Schema->PC_Class == CategoryName || Schema->PC_Interface == CategoryName || Schema->PC_Asset == CategoryName || Schema->PC_AssetClass == CategoryName)
+		{
+			TArray<FAssetData> AssetData;
+			AssetRegistryModule.Get().GetAssetsByClass(UBlueprint::StaticClass()->GetFName(), AssetData);
+
+			const FString BPInterfaceTypeAllowed(TEXT("BPTYPE_Interface")); 
+			const FString BPNormalTypeAllowed(TEXT("BPTYPE_Normal"));
+
+			FTypesDatabase::FUnLoadedTypesList UnLoadedInterfacesList = MakeShareable(new TArray<FUnloadedAssetData>());
+			FTypesDatabase::FUnLoadedTypesList UnLoadedClassesList = MakeShareable(new TArray<FUnloadedAssetData>());
+
+			for (int32 AssetIndex = 0; AssetIndex < AssetData.Num(); ++AssetIndex)
+			{
+				const FAssetData& Asset = AssetData[AssetIndex];
+
+				if (Asset.IsValid() && !Asset.IsAssetLoaded())
+				{
+					const FString* BlueprintTypeStr = Asset.TagsAndValues.Find("BlueprintType");
+					const bool bNormalBP = BlueprintTypeStr && (*BlueprintTypeStr == BPNormalTypeAllowed);
+					const bool bInterfaceBP = BlueprintTypeStr && (*BlueprintTypeStr == BPInterfaceTypeAllowed);
+
+					if (bNormalBP || bInterfaceBP)
+					{
+						uint32 ClassFlags = 0;
+						const FString* ClassFlagsStr = Asset.TagsAndValues.Find("ClassFlags");
+						if (ClassFlagsStr)
+						{
+							ClassFlags = FCString::Atoi(**ClassFlagsStr);
+						}
+
+						if (!(ClassFlags & CLASS_Deprecated))
+						{
+							if (bNormalBP)
+							{
+								UnLoadedClassesList->Add(FUnloadedAssetData(Asset));
+							}
+							else if (bInterfaceBP)
+							{
+								UnLoadedInterfacesList->Add(FUnloadedAssetData(Asset));
+							}
+						}
+					}
+				}
+			}
+			TypesDatabase.UnLoadedTypesMap.Add(UEdGraphSchema_K2::PC_Interface, UnLoadedInterfacesList);
+			TypesDatabase.UnLoadedTypesMap.Add(UEdGraphSchema_K2::PC_Object, UnLoadedClassesList);
+			TypesDatabase.UnLoadedTypesMap.Add(UEdGraphSchema_K2::PC_Class, UnLoadedClassesList);
+			TypesDatabase.UnLoadedTypesMap.Add(UEdGraphSchema_K2::PC_Asset, UnLoadedClassesList);
+			TypesDatabase.UnLoadedTypesMap.Add(UEdGraphSchema_K2::PC_AssetClass, UnLoadedClassesList);
+		}
+	}
+
+	static void Gather(const FText& FriendlyName, const FString& CategoryName, FTypesDatabase& TypesDatabase, TArray<FPinTypeTreeInfoPtr>& OutChildren)
+	{
+		FEdGraphPinType LoadedPinSubtype;
+		LoadedPinSubtype.PinCategory = (CategoryName == UEdGraphSchema_K2::PC_Enum ? UEdGraphSchema_K2::PC_Byte : CategoryName);
+		LoadedPinSubtype.PinSubCategory = TEXT("");
+		LoadedPinSubtype.PinSubCategoryObject = NULL;
+
+		auto LoadedSubTypesPtr = TypesDatabase.LoadedTypesMap.Find(CategoryName);
+		if (LoadedSubTypesPtr && LoadedSubTypesPtr->IsValid())
+		{
+			for (auto LoadedAssetData : *LoadedSubTypesPtr->Get())
+			{
+				OutChildren.Add(MakeShareable(new UEdGraphSchema_K2::FPinTypeTreeInfo(LoadedPinSubtype.PinCategory, LoadedAssetData.Object, LoadedAssetData.Tooltip.IsEmpty() ? FriendlyName : LoadedAssetData.Tooltip)));
+			}
+		}
+
+		auto UnLoadedSubTypesPtr = TypesDatabase.UnLoadedTypesMap.Find(CategoryName);
+		if (UnLoadedSubTypesPtr && UnLoadedSubTypesPtr->IsValid())
+		{
+			for (FUnloadedAssetData& It : *UnLoadedSubTypesPtr->Get())
+			{
+				FPinTypeTreeInfoPtr TypeTreeInfo = MakeShareable(new UEdGraphSchema_K2::FPinTypeTreeInfo(CategoryName, It.StringAssetReference, It.Tooltip));
+				TypeTreeInfo->FriendlyName = It.AssetFriendlyName;
+				OutChildren.Add(TypeTreeInfo);
+			}
+		}
+
+		OutChildren.Sort(FCompareChildren());
+	}
+
 	/**
 	 * Gathers all valid sub-types (loaded and unloaded) of a passed category and sorts them alphabetically
 	 * @param FriendlyName		Friendly name to be used for the tooltip if there is no available data
@@ -362,7 +601,7 @@ const FEdGraphPinType& UEdGraphSchema_K2::FPinTypeTreeInfo::GetPinType(bool bFor
 	return PinType;
 }
 
-void UEdGraphSchema_K2::FPinTypeTreeInfo::Init(const FText& InFriendlyName, const FString& CategoryName, const UEdGraphSchema_K2* Schema, const FText& InTooltip, bool bInReadOnly)
+void UEdGraphSchema_K2::FPinTypeTreeInfo::Init(const FText& InFriendlyName, const FString& CategoryName, const UEdGraphSchema_K2* Schema, const FText& InTooltip, bool bInReadOnly, FTypesDatabase* TypesDatabase)
 {
 	check( !CategoryName.IsEmpty() );
 	check( Schema );
@@ -375,9 +614,18 @@ void UEdGraphSchema_K2::FPinTypeTreeInfo::Init(const FText& InFriendlyName, cons
 
 	bReadOnly = bInReadOnly;
 
+	CachedDescription = GenerateDescription();
+
 	if (Schema->DoesTypeHaveSubtypes(CategoryName))
 	{
-		FGatherTypesHelper::Gather(InFriendlyName, CategoryName, Schema, Children);
+		if (TypesDatabase)
+		{
+			FGatherTypesHelper::Gather(InFriendlyName, CategoryName, *TypesDatabase, Children);
+		}
+		else
+		{
+			FGatherTypesHelper::Gather(InFriendlyName, CategoryName, Schema, Children);
+		}
 	}
 }
 
@@ -391,6 +639,7 @@ UEdGraphSchema_K2::FPinTypeTreeInfo::FPinTypeTreeInfo(const FString& CategoryNam
 	PinType.PinSubCategoryObject = SubCategoryObject;
 
 	bReadOnly = bInReadOnly;
+	CachedDescription = GenerateDescription();
 }
 
 UEdGraphSchema_K2::FPinTypeTreeInfo::FPinTypeTreeInfo(const FString& CategoryName, const FStringAssetReference& SubCategoryObject, const FText& InTooltip, bool bInReadOnly)
@@ -405,9 +654,10 @@ UEdGraphSchema_K2::FPinTypeTreeInfo::FPinTypeTreeInfo(const FString& CategoryNam
 	PinType.PinSubCategoryObject = SubCategoryObjectAssetReference.ResolveObject();
 
 	bReadOnly = bInReadOnly;
+	CachedDescription = GenerateDescription();
 }
 
-FText UEdGraphSchema_K2::FPinTypeTreeInfo::GetDescription() const
+FText UEdGraphSchema_K2::FPinTypeTreeInfo::GenerateDescription()
 {
 	if (!FriendlyName.IsEmpty())
 	{
@@ -433,6 +683,11 @@ FText UEdGraphSchema_K2::FPinTypeTreeInfo::GetDescription() const
 	}
 }
 
+FText UEdGraphSchema_K2::FPinTypeTreeInfo::GetDescription() const
+{
+	return CachedDescription;
+}
+
 const FString UEdGraphSchema_K2::PC_Exec(TEXT("exec"));
 const FString UEdGraphSchema_K2::PC_Boolean(TEXT("bool"));
 const FString UEdGraphSchema_K2::PC_Byte(TEXT("byte"));
@@ -3306,6 +3561,28 @@ FText UEdGraphSchema_K2::TypeToText(const FEdGraphPinType& Type)
 
 void UEdGraphSchema_K2::GetVariableTypeTree( TArray< TSharedPtr<FPinTypeTreeInfo> >& TypeTree, bool bAllowExec, bool bAllowWildCard ) const
 {
+#ifdef SCHEMA_K2_GETVARIABLETYPETREE_LOG_TIME
+	static_assert(false, "Macro redefinition.");
+#endif
+#define SCHEMA_K2_GETVARIABLETYPETREE_LOG_TIME 0
+
+#if SCHEMA_K2_GETVARIABLETYPETREE_LOG_TIME
+	const auto StartTime = FPlatformTime::Seconds();
+#endif //SCHEMA_K2_GETVARIABLETYPETREE_LOG_TIME
+
+	FTypesDatabase TypesDatabase;
+	FGatherTypesHelper::FillLoadedTypesDatabase(TypesDatabase);
+
+#if SCHEMA_K2_GETVARIABLETYPETREE_LOG_TIME
+	const auto DatabaseLoadedTime = FPlatformTime::Seconds();
+#endif //SCHEMA_K2_GETVARIABLETYPETREE_LOG_TIME
+
+	FGatherTypesHelper::FillUnLoadedTypesDatabase(TypesDatabase);
+
+#if SCHEMA_K2_GETVARIABLETYPETREE_LOG_TIME
+	const auto DatabaseUnLoadedTime = FPlatformTime::Seconds();
+#endif //SCHEMA_K2_GETVARIABLETYPETREE_LOG_TIME
+
 	// Clear the list
 	TypeTree.Empty();
 
@@ -3333,14 +3610,22 @@ void UEdGraphSchema_K2::GetVariableTypeTree( TArray< TSharedPtr<FPinTypeTreeInfo
 		TypeTree.Add( MakeShareable( new FPinTypeTreeInfo(GetCategoryText(PC_Wildcard, true), PC_Wildcard, this, LOCTEXT("WildcardType", "Wildcard type (unspecified).")) ) );
 	}
 
+	FTypesDatabase* TypesDatabasePtr = &TypesDatabase;
 	// Add the types that have subtrees
-	TypeTree.Add( MakeShareable( new FPinTypeTreeInfo(GetCategoryText(PC_Struct, true), PC_Struct, this, LOCTEXT("StructType", "Struct (value) types."), true) ) );
-	TypeTree.Add( MakeShareable( new FPinTypeTreeInfo(GetCategoryText(PC_Object, true), PC_Object, this, LOCTEXT("ObjectType", "Object pointer."), true) ) );
-	TypeTree.Add( MakeShareable( new FPinTypeTreeInfo(GetCategoryText(PC_Asset, true), PC_Asset, this, LOCTEXT("AssetType", "Asset pointer."), true) ) );
-	TypeTree.Add( MakeShareable( new FPinTypeTreeInfo(GetCategoryText(PC_AssetClass, true), PC_AssetClass, this, LOCTEXT("ClassAssetType", "Class asset."), true) ) );
-	TypeTree.Add( MakeShareable( new FPinTypeTreeInfo(GetCategoryText(PC_Interface, true), PC_Interface, this, LOCTEXT("InterfaceType", "Interface pointer."), true) ) );
-	TypeTree.Add( MakeShareable( new FPinTypeTreeInfo(GetCategoryText(PC_Class, true), PC_Class, this, LOCTEXT("ClassType", "Class pointers."), true) ) );
-	TypeTree.Add( MakeShareable( new FPinTypeTreeInfo(GetCategoryText(PC_Enum, true), PC_Enum, this, LOCTEXT("EnumType", "Enumeration types."), true) ) );
+	TypeTree.Add( MakeShareable( new FPinTypeTreeInfo(GetCategoryText(PC_Struct, true), PC_Struct, this, LOCTEXT("StructType", "Struct (value) types."), true, TypesDatabasePtr) ) );
+	TypeTree.Add( MakeShareable( new FPinTypeTreeInfo(GetCategoryText(PC_Object, true), PC_Object, this, LOCTEXT("ObjectType", "Object pointer."), true, TypesDatabasePtr) ) );
+	TypeTree.Add( MakeShareable( new FPinTypeTreeInfo(GetCategoryText(PC_Asset, true), PC_Asset, this, LOCTEXT("AssetType", "Asset pointer."), true, TypesDatabasePtr) ) );
+	TypeTree.Add( MakeShareable( new FPinTypeTreeInfo(GetCategoryText(PC_AssetClass, true), PC_AssetClass, this, LOCTEXT("ClassAssetType", "Class asset."), true, TypesDatabasePtr) ) );
+	TypeTree.Add( MakeShareable( new FPinTypeTreeInfo(GetCategoryText(PC_Interface, true), PC_Interface, this, LOCTEXT("InterfaceType", "Interface pointer."), true, TypesDatabasePtr) ) );
+	TypeTree.Add( MakeShareable( new FPinTypeTreeInfo(GetCategoryText(PC_Class, true), PC_Class, this, LOCTEXT("ClassType", "Class pointers."), true, TypesDatabasePtr) ) );
+	TypeTree.Add( MakeShareable( new FPinTypeTreeInfo(GetCategoryText(PC_Enum, true), PC_Enum, this, LOCTEXT("EnumType", "Enumeration types."), true, TypesDatabasePtr) ) );
+
+#if SCHEMA_K2_GETVARIABLETYPETREE_LOG_TIME
+	const auto EndTime = FPlatformTime::Seconds();
+	UE_LOG(LogBlueprint, Log, TEXT("UEdGraphSchema_K2::GetVariableTypeTree times - LoadedTypesDatabase: %f UnLoadedTypesDatabase: %f FPinTypeTreeInfo: %f"), DatabaseLoadedTime - StartTime, DatabaseUnLoadedTime - DatabaseLoadedTime, EndTime - DatabaseUnLoadedTime);
+#endif //SCHEMA_K2_GETVARIABLETYPETREE_LOG_TIME
+
+#undef SCHEMA_K2_GETVARIABLETYPETREE_LOG_TIME
 }
 
 void UEdGraphSchema_K2::GetVariableIndexTypeTree( TArray< TSharedPtr<FPinTypeTreeInfo> >& TypeTree, bool bAllowExec, bool bAllowWildcard ) const
