PC_Asset type in named Asset ID. PC_AssetClass type is named Class Asset ID.

diff --git a/Engine/Source/Editor/BlueprintGraph/Private/EdGraphSchema_K2.cpp b/Engine/Source/Editor/BlueprintGraph/Private/EdGraphSchema_K2.cpp
--- a/Engine/Source/Editor/BlueprintGraph/Private/EdGraphSchema_K2.cpp
+++ b/Engine/Source/Editor/BlueprintGraph/Private/EdGraphSchema_K2.cpp
@@ -3459,8 +3459,8 @@ FText UEdGraphSchema_K2::GetCategoryText(const FString& Category, const bool bFo
 		CategoryDescriptions.Add(PC_Struct, LOCTEXT("StructCategory","Structure"));
 		CategoryDescriptions.Add(PC_Wildcard, LOCTEXT("WildcardCategory","Wildcard"));
 		CategoryDescriptions.Add(PC_Enum, LOCTEXT("EnumCategory","Enum"));
-		CategoryDescriptions.Add(PC_Asset, LOCTEXT("AssetCategory", "Asset"));
-		CategoryDescriptions.Add(PC_AssetClass, LOCTEXT("AssetClassCategory", "Class Asset"));
+		CategoryDescriptions.Add(PC_Asset, LOCTEXT("AssetCategory", "Asset ID"));
+		CategoryDescriptions.Add(PC_AssetClass, LOCTEXT("AssetClassCategory", "Class Asset ID"));
 	}
 
 	if (bForMenu)
diff --git a/Engine/Source/Editor/BlueprintGraph/Private/K2Node_ConvertAsset.cpp b/Engine/Source/Editor/BlueprintGraph/Private/K2Node_ConvertAsset.cpp
--- a/Engine/Source/Editor/BlueprintGraph/Private/K2Node_ConvertAsset.cpp
+++ b/Engine/Source/Editor/BlueprintGraph/Private/K2Node_ConvertAsset.cpp
@@ -218,12 +218,12 @@ FText UK2Node_ConvertAsset::GetMenuCategory() const
 
 FText UK2Node_ConvertAsset::GetNodeTitle(ENodeTitleType::Type TitleType) const
 {
-	return FText(LOCTEXT("UK2Node_ConvertAssetGetNodeTitle", "Resolve Asset"));
+	return FText(LOCTEXT("UK2Node_ConvertAssetGetNodeTitle", "Resolve Asset ID"));
 }
 
 FText UK2Node_ConvertAsset::GetTooltipText() const
 {
-	return FText(LOCTEXT("UK2Node_ConvertAssetGetTooltipText", "Resolves an asset/assetclass into an object/class. If the asset wasn't loaded it returns none."));
+	return FText(LOCTEXT("UK2Node_ConvertAssetGetTooltipText", "Resolves an Asset ID or Class Asset ID into an object/class. If the asset wasn't loaded it returns none."));
 }
 
 #undef LOCTEXT_NAMESPACE
\ No newline at end of file
