When searching the pin type selector, items will no longer appear blank.

Fixed an issue where unloaded items appeared with "Error!", FriendlyName was not properly set before the description cached.

diff --git a/Engine/Source/Editor/BlueprintGraph/Classes/EdGraphSchema_K2.h b/Engine/Source/Editor/BlueprintGraph/Classes/EdGraphSchema_K2.h
--- a/Engine/Source/Editor/BlueprintGraph/Classes/EdGraphSchema_K2.h
+++ b/Engine/Source/Editor/BlueprintGraph/Classes/EdGraphSchema_K2.h
@@ -321,6 +321,7 @@ public:
 		FPinTypeTreeInfo(const FText& FriendlyName, const FString& CategoryName, const UEdGraphSchema_K2* Schema, const FText& InTooltip, bool bInReadOnly = false, FTypesDatabase* TypesDatabase = nullptr);
 		FPinTypeTreeInfo(const FString& CategoryName, UObject* SubCategoryObject, const FText& InTooltip, bool bInReadOnly = false);
 		FPinTypeTreeInfo(const FString& CategoryName, const FStringAssetReference& SubCategoryObject, const FText& InTooltip, bool bInReadOnly = false);
+		FPinTypeTreeInfo(const FText& InFriendlyName, const FString& CategoryName, const FStringAssetReference& SubCategoryObject, const FText& InTooltip, bool bInReadOnly = false);
 
 		FPinTypeTreeInfo(TSharedPtr<FPinTypeTreeInfo> InInfo)
 		{
@@ -329,6 +330,7 @@ public:
 			FriendlyName = InInfo->FriendlyName;
 			Tooltip = InInfo->Tooltip;
 			SubCategoryObjectAssetReference = InInfo->SubCategoryObjectAssetReference;
+			CachedDescription = InInfo->CachedDescription;
 		}
 		
 		/** Returns a succinct menu description of this type */
diff --git a/Engine/Source/Editor/BlueprintGraph/Private/EdGraphSchema_K2.cpp b/Engine/Source/Editor/BlueprintGraph/Private/EdGraphSchema_K2.cpp
--- a/Engine/Source/Editor/BlueprintGraph/Private/EdGraphSchema_K2.cpp
+++ b/Engine/Source/Editor/BlueprintGraph/Private/EdGraphSchema_K2.cpp
@@ -103,9 +103,9 @@ const FName FBlueprintMetadata::MD_ArrayDependentParam(TEXT("ArrayTypeDependentP
 
 #define LOCTEXT_NAMESPACE "KismetSchema"
 
-UEdGraphSchema_K2::FPinTypeTreeInfo::FPinTypeTreeInfo(const FText& InFriendlyCategoryName, const FString& CategoryName, const UEdGraphSchema_K2* Schema, const FText& InTooltip, bool bInReadOnly/*=false*/, FTypesDatabase* TypesDatabase /*=nullptr*/)
+UEdGraphSchema_K2::FPinTypeTreeInfo::FPinTypeTreeInfo(const FText& FriendlyName, const FString& CategoryName, const UEdGraphSchema_K2* Schema, const FText& InTooltip, bool bInReadOnly/*=false*/, FTypesDatabase* TypesDatabase /*=nullptr*/)
 {
-	Init(InFriendlyCategoryName, CategoryName, Schema, InTooltip, bInReadOnly, TypesDatabase);
+	Init(FriendlyName, CategoryName, Schema, InTooltip, bInReadOnly, TypesDatabase);
 }
 
 struct FUnloadedAssetData
@@ -435,8 +435,7 @@ public:
 		{
 			for (FUnloadedAssetData& It : *UnLoadedSubTypesPtr->Get())
 			{
-				FPinTypeTreeInfoPtr TypeTreeInfo = MakeShareable(new UEdGraphSchema_K2::FPinTypeTreeInfo(CategoryName, It.StringAssetReference, It.Tooltip));
-				TypeTreeInfo->FriendlyName = It.AssetFriendlyName;
+				FPinTypeTreeInfoPtr TypeTreeInfo = MakeShareable(new UEdGraphSchema_K2::FPinTypeTreeInfo(It.AssetFriendlyName, CategoryName, It.StringAssetReference, It.Tooltip));
 				OutChildren.Add(TypeTreeInfo);
 			}
 		}
@@ -657,6 +656,23 @@ UEdGraphSchema_K2::FPinTypeTreeInfo::FPinTypeTreeInfo(const FString& CategoryNam
 	CachedDescription = GenerateDescription();
 }
 
+UEdGraphSchema_K2::FPinTypeTreeInfo::FPinTypeTreeInfo(const FText& InFriendlyName, const FString& CategoryName, const FStringAssetReference& SubCategoryObject, const FText& InTooltip, bool bInReadOnly)
+{
+	FriendlyName = InFriendlyName;
+
+	check(!CategoryName.IsEmpty());
+	check(SubCategoryObject.IsValid());
+
+	Tooltip = InTooltip;
+	PinType.PinCategory = CategoryName;
+
+	SubCategoryObjectAssetReference = SubCategoryObject;
+	PinType.PinSubCategoryObject = SubCategoryObjectAssetReference.ResolveObject();
+
+	bReadOnly = bInReadOnly;
+	CachedDescription = GenerateDescription();
+}
+
 FText UEdGraphSchema_K2::FPinTypeTreeInfo::GenerateDescription()
 {
 	if (!FriendlyName.IsEmpty())
