Duplicating the "Add Scene Component" node while the transform pin is split will no longer make visible the transform pin on the duplicated node.

Added EdGraphPin::SafeSetHidden to not do anything when pins are forced hidden (or not) from splitting.

diff --git a/Engine/Source/Editor/BlueprintGraph/Private/K2Node_AddComponent.cpp b/Engine/Source/Editor/BlueprintGraph/Private/K2Node_AddComponent.cpp
--- a/Engine/Source/Editor/BlueprintGraph/Private/K2Node_AddComponent.cpp
+++ b/Engine/Source/Editor/BlueprintGraph/Private/K2Node_AddComponent.cpp
@@ -83,10 +83,10 @@ void UK2Node_AddComponent::AllocatePinsForExposedVariables()
 	const bool bHideTransformPins = (ComponentClass != nullptr) ? !ComponentClass->IsChildOf(USceneComponent::StaticClass()) : false;
 
 	UEdGraphPin* ManualAttachmentPin = GetManualAttachmentPin();
-	ManualAttachmentPin->bHidden = bHideTransformPins;
+	ManualAttachmentPin->SafeSetHidden(bHideTransformPins);
 
 	UEdGraphPin* TransformPin = GetRelativeTransformPin();
-	TransformPin->bHidden = bHideTransformPins;
+	TransformPin->SafeSetHidden(bHideTransformPins);
 }
 
 const UClass* UK2Node_AddComponent::GetSpawnedType() const
diff --git a/Engine/Source/Runtime/Engine/Classes/EdGraph/EdGraphPin.h b/Engine/Source/Runtime/Engine/Classes/EdGraph/EdGraphPin.h
--- a/Engine/Source/Runtime/Engine/Classes/EdGraph/EdGraphPin.h
+++ b/Engine/Source/Runtime/Engine/Classes/EdGraph/EdGraphPin.h
@@ -294,6 +294,15 @@ public:
 		return (InDirection == EGPD_Input) ? EGPD_Output : EGPD_Input;
 	}
 
+	/** Helper to safely set a pin's bHidden property only if it has no sub-pins that are influencing it to be hidden */
+	void SafeSetHidden(bool bIsHidden)
+	{
+		 if (SubPins.Num() == 0)
+		 {
+			 bHidden = bIsHidden;
+		 }
+	}
+
 	/** Get the current DefaultObject path name, or DefaultValue if its null */
 	ENGINE_API FString GetDefaultAsString() const;
 
