On the select node, setting the Index pin's type will always change the pin names of the option pins to "False" and "True".

diff --git a/Engine/Source/Editor/BlueprintGraph/Private/K2Node_Select.cpp b/Engine/Source/Editor/BlueprintGraph/Private/K2Node_Select.cpp
--- a/Engine/Source/Editor/BlueprintGraph/Private/K2Node_Select.cpp
+++ b/Engine/Source/Editor/BlueprintGraph/Private/K2Node_Select.cpp
@@ -763,11 +763,8 @@ void UK2Node_Select::PinTypeChanged(UEdGraphPin* Pin)
 			// Remove all but two options if we switched to a bool index
 			if (IndexPinType.PinCategory == Schema->PC_Boolean)
 			{
-				if (NumOptionPins > 2)
-				{
-					NumOptionPins = 2;
-					bReconstructNode = true;
-				}
+				NumOptionPins = 2;
+				bReconstructNode = true;
 			}
 
 			// Reset the default value
