Preventing MathExpression nodes from counting function inputs as valid variables to access via getter nodes.

diff --git a/Engine/Source/Editor/BlueprintGraph/Private/K2Node_MathExpression.cpp b/Engine/Source/Editor/BlueprintGraph/Private/K2Node_MathExpression.cpp
--- a/Engine/Source/Editor/BlueprintGraph/Private/K2Node_MathExpression.cpp
+++ b/Engine/Source/Editor/BlueprintGraph/Private/K2Node_MathExpression.cpp
@@ -1364,7 +1364,10 @@ public:
 				if (!VariableGuid.IsValid())
 				{
 					VariableGuid = FBlueprintEditorUtils::FindLocalVariableGuidByName(TargetBlueprint, CompilingNode->GetGraph(), FName(*VariableName));
-					VariableReference.SetLocalMember(FName(*VariableName), CompilingNode->GetGraph()->GetName(), VariableGuid);
+					if (VariableGuid.IsValid())
+					{
+						VariableReference.SetLocalMember(FName(*VariableName), CompilingNode->GetGraph()->GetName(), VariableGuid);
+					}
 				}
 				else
 				{
