Fix non-unity include build issue from CL# 2572814

diff --git a/Engine/Source/Editor/BlueprintGraph/Private/K2Node_CustomEvent.cpp b/Engine/Source/Editor/BlueprintGraph/Private/K2Node_CustomEvent.cpp
--- a/Engine/Source/Editor/BlueprintGraph/Private/K2Node_CustomEvent.cpp
+++ b/Engine/Source/Editor/BlueprintGraph/Private/K2Node_CustomEvent.cpp
@@ -9,6 +9,7 @@
 #include "BlueprintActionDatabaseRegistrar.h"
 #include "K2Node_CustomEvent.h"
 #include "K2Node_BaseMCDelegate.h"
+#include "Editor/Kismet/Public/FindInBlueprintManager.h"
 
 #define LOCTEXT_NAMESPACE "K2Node_CustomEvent"
 
