Fixed several shadow variable warnings

diff --git a/Engine/Source/Editor/BlueprintGraph/Classes/EdGraphSchema_K2.h b/Engine/Source/Editor/BlueprintGraph/Classes/EdGraphSchema_K2.h
--- a/Engine/Source/Editor/BlueprintGraph/Classes/EdGraphSchema_K2.h
+++ b/Engine/Source/Editor/BlueprintGraph/Classes/EdGraphSchema_K2.h
@@ -318,7 +318,7 @@ public:
 			PinType.PinSubCategory = SubCategory;
 		}
 
-		FPinTypeTreeInfo(const FText& FriendlyName, const FString& CategoryName, const UEdGraphSchema_K2* Schema, const FText& InTooltip, bool bInReadOnly = false, FTypesDatabase* TypesDatabase = nullptr);
+		FPinTypeTreeInfo(const FText& InFriendlyName, const FString& CategoryName, const UEdGraphSchema_K2* Schema, const FText& InTooltip, bool bInReadOnly = false, FTypesDatabase* TypesDatabase = nullptr);
 		FPinTypeTreeInfo(const FString& CategoryName, UObject* SubCategoryObject, const FText& InTooltip, bool bInReadOnly = false);
 		FPinTypeTreeInfo(const FString& CategoryName, const FStringAssetReference& SubCategoryObject, const FText& InTooltip, bool bInReadOnly = false);
 		FPinTypeTreeInfo(const FText& InFriendlyName, const FString& CategoryName, const FStringAssetReference& SubCategoryObject, const FText& InTooltip, bool bInReadOnly = false);
diff --git a/Engine/Source/Editor/BlueprintGraph/Private/EdGraphSchema_K2.cpp b/Engine/Source/Editor/BlueprintGraph/Private/EdGraphSchema_K2.cpp
--- a/Engine/Source/Editor/BlueprintGraph/Private/EdGraphSchema_K2.cpp
+++ b/Engine/Source/Editor/BlueprintGraph/Private/EdGraphSchema_K2.cpp
@@ -103,9 +103,9 @@ const FName FBlueprintMetadata::MD_ArrayDependentParam(TEXT("ArrayTypeDependentP
 
 #define LOCTEXT_NAMESPACE "KismetSchema"
 
-UEdGraphSchema_K2::FPinTypeTreeInfo::FPinTypeTreeInfo(const FText& FriendlyName, const FString& CategoryName, const UEdGraphSchema_K2* Schema, const FText& InTooltip, bool bInReadOnly/*=false*/, FTypesDatabase* TypesDatabase /*=nullptr*/)
+UEdGraphSchema_K2::FPinTypeTreeInfo::FPinTypeTreeInfo(const FText& InFriendlyName, const FString& CategoryName, const UEdGraphSchema_K2* Schema, const FText& InTooltip, bool bInReadOnly/*=false*/, FTypesDatabase* TypesDatabase /*=nullptr*/)
 {
-	Init(FriendlyName, CategoryName, Schema, InTooltip, bInReadOnly, TypesDatabase);
+	Init(InFriendlyName, CategoryName, Schema, InTooltip, bInReadOnly, TypesDatabase);
 }
 
 struct FUnloadedAssetData
diff --git a/Engine/Source/Editor/BlueprintGraph/Private/K2Node_CallFunction.cpp b/Engine/Source/Editor/BlueprintGraph/Private/K2Node_CallFunction.cpp
--- a/Engine/Source/Editor/BlueprintGraph/Private/K2Node_CallFunction.cpp
+++ b/Engine/Source/Editor/BlueprintGraph/Private/K2Node_CallFunction.cpp
@@ -506,9 +506,9 @@ void UK2Node_CallFunction::GetPinHoverText(const UEdGraphPin& Pin, FString& Hove
 {
 	if (!bPinTooltipsValid)
 	{
-		for (auto& Pin : Pins)
+		for (auto& P : Pins)
 		{
-			GeneratePinTooltip(*Pin);
+			GeneratePinTooltip(*P);
 		}
 
 		bPinTooltipsValid = true;
diff --git a/Engine/Source/Runtime/Engine/Private/Components/SkeletalMeshComponent.cpp b/Engine/Source/Runtime/Engine/Private/Components/SkeletalMeshComponent.cpp
--- a/Engine/Source/Runtime/Engine/Private/Components/SkeletalMeshComponent.cpp
+++ b/Engine/Source/Runtime/Engine/Private/Components/SkeletalMeshComponent.cpp
@@ -2284,15 +2284,15 @@ bool USkeletalMeshComponent::DoCustomNavigableGeometryExport(FNavigableGeometryE
 		const int32 MaxBodies = PhysicsAsset->BodySetup.Num();
 		for (int32 Idx = 0; Idx < MaxBodies; Idx++)
 		{
-			UBodySetup* BodySetup = PhysicsAsset->BodySetup[Idx];
-			int32 BoneIndex = BodySetup ? GetBoneIndex(BodySetup->BoneName) : INDEX_NONE;
+			UBodySetup* const BS = PhysicsAsset->BodySetup[Idx];
+			int32 const BoneIndex = BS ? GetBoneIndex(BS->BoneName) : INDEX_NONE;
 
 			if (BoneIndex != INDEX_NONE)
 			{
 				FTransform WorldBoneTransform = GetBoneTransform(BoneIndex, ComponentToWorld);
 				if (FMath::Abs(WorldBoneTransform.GetDeterminant()) > (float)KINDA_SMALL_NUMBER)
 				{
-					GeomExport.ExportRigidBodySetup(*BodySetup, WorldBoneTransform);
+					GeomExport.ExportRigidBodySetup(*BS, WorldBoneTransform);
 				}
 			}
 		}
diff --git a/Engine/Source/Runtime/Renderer/Private/DeferredShadingRenderer.cpp b/Engine/Source/Runtime/Renderer/Private/DeferredShadingRenderer.cpp
--- a/Engine/Source/Runtime/Renderer/Private/DeferredShadingRenderer.cpp
+++ b/Engine/Source/Runtime/Renderer/Private/DeferredShadingRenderer.cpp
@@ -873,9 +873,9 @@ void FDeferredShadingSceneRenderer::Render(FRHICommandListImmediate& RHICmdList)
 			bOnce = true;
 			GPrevPerBoneMotionBlur.SetVelocityPassCallback(			
 				// this is a strange intermodule bridge so that the skeletal mesh vertex factory knows when to add bone data in a parallel thread
-				[](FRHICommandList& RHICmdList)->bool
+				[](FRHICommandList& InRHICmdList)->bool
 			{
-				return FSceneRenderTargets::Get(RHICmdList).IsVelocityPass();
+				return FSceneRenderTargets::Get(InRHICmdList).IsVelocityPass();
 			});
 		}
 	}
diff --git a/Engine/Source/Runtime/Renderer/Private/PostProcess/PostProcessCombineLUTs.cpp b/Engine/Source/Runtime/Renderer/Private/PostProcess/PostProcessCombineLUTs.cpp
--- a/Engine/Source/Runtime/Renderer/Private/PostProcess/PostProcessCombineLUTs.cpp
+++ b/Engine/Source/Runtime/Renderer/Private/PostProcess/PostProcessCombineLUTs.cpp
@@ -319,9 +319,9 @@ public:
 			float FilmSlopeS = FilmMidYS / (FilmMidXS);
 			float FilmHiYS = 1.0f - FilmHiY;
 			float FilmLoYS = FilmLoY;
-			float FilmToe = FilmLoX;
+			float FilmToeVal = FilmLoX;
 			float FilmHiG = (-FilmHiYS + (FilmSlopeS*FilmHeal)) / (FilmSlopeS*FilmHeal);
-			float FilmLoG = (-FilmLoYS + (FilmSlopeS*FilmToe)) / (FilmSlopeS*FilmToe);
+			float FilmLoG = (-FilmLoYS + (FilmSlopeS*FilmToeVal)) / (FilmSlopeS*FilmToeVal);
 
 			// Constants.
 			float OutColorCurveCh1 = FilmHiYS/FilmHiG;
