Preventing in UK2Node::AutowireNode hidden pins from being auto-connected to passed in Wildcard pins. This prevents issues with wildcards connecting to the "self" pin of functions when dragging off them.

diff --git a/Engine/Source/Editor/BlueprintGraph/Private/K2Node.cpp b/Engine/Source/Editor/BlueprintGraph/Private/K2Node.cpp
--- a/Engine/Source/Editor/BlueprintGraph/Private/K2Node.cpp
+++ b/Engine/Source/Editor/BlueprintGraph/Private/K2Node.cpp
@@ -143,6 +143,12 @@ void UK2Node::AutowireNewNode(UEdGraphPin* FromPin)
 			UEdGraphPin* Pin = Pins[i];
 			check(Pin);
 
+			// Never consider for auto-wiring a hidden pin being connected to a Wildcard. It is never what the user expects
+			if (Pin->bHidden && FromPin->PinType.PinCategory == UEdGraphSchema_K2::PC_Wildcard)
+			{
+				continue;
+			}
+
 			ECanCreateConnectionResponse ConnectResponse = K2Schema->CanCreateConnection(FromPin, Pin).Response;
 			if (ConnectResponse == ECanCreateConnectionResponse::CONNECT_RESPONSE_MAKE)
 			{
