Select nodes using enums with hidden values will no longer incorrectly match up the Option pin connections to the improper enum value.

Fixed a localization issue with display enum values as pin names.

diff --git a/Engine/Source/Editor/BlueprintGraph/Classes/K2Node_Select.h b/Engine/Source/Editor/BlueprintGraph/Classes/K2Node_Select.h
--- a/Engine/Source/Editor/BlueprintGraph/Classes/K2Node_Select.h
+++ b/Engine/Source/Editor/BlueprintGraph/Classes/K2Node_Select.h
@@ -29,7 +29,7 @@ class UK2Node_Select : public UK2Node, public INodeDependingOnEnumInterface
 
 	/** List of the current entries in the enum (Pin Friendly Names) */
 	UPROPERTY()
-	TArray<FName> EnumEntryFriendlyNames;
+	TArray<FText> EnumEntryFriendlyNames;
 
 	/** Whether we need to reconstruct the node after the pins have changed */
 	UPROPERTY(Transient)
diff --git a/Engine/Source/Editor/BlueprintGraph/Private/K2Node_Select.cpp b/Engine/Source/Editor/BlueprintGraph/Private/K2Node_Select.cpp
--- a/Engine/Source/Editor/BlueprintGraph/Private/K2Node_Select.cpp
+++ b/Engine/Source/Editor/BlueprintGraph/Private/K2Node_Select.cpp
@@ -93,7 +93,16 @@ public:
 				FBPTerminal* LiteralTerm = Context.CreateLocalTerminal(ETerminalSpecification::TS_Literal);
 				LiteralTerm->bIsLiteral = true;
 				LiteralTerm->Type.PinCategory = CompilerContext.GetSchema()->PC_Int;
-				LiteralTerm->Name = FString::Printf(TEXT("%d"), OptionIdx);
+
+				if (UEnum* NodeEnum = SelectNode->GetEnum())
+				{
+					int32 EnumIndex = NodeEnum->FindEnumIndex(FName(*OptionPins[OptionIdx]->PinName));
+					LiteralTerm->Name = FString::Printf(TEXT("%d"), EnumIndex);
+				}
+				else
+				{
+					LiteralTerm->Name = FString::Printf(TEXT("%d"), OptionIdx);
+				}
 				Statement.RHS.Add(LiteralTerm);
 				// If there is a previous IfNot statement, hook this one to that one for jumping
 				if (PrevIfNotStatement)
@@ -241,14 +250,7 @@ void UK2Node_Select::AllocateDefaultPins()
 			}
 			else if (Idx < EnumEntryFriendlyNames.Num())
 			{
-				if (EnumEntryFriendlyNames[Idx] != NAME_None)
-				{
-					NewPin->PinFriendlyName = FText::FromName(EnumEntryFriendlyNames[Idx]);
-				}
-				else
-				{
-					NewPin->PinFriendlyName = FText::GetEmpty();
-				}
+				NewPin->PinFriendlyName = EnumEntryFriendlyNames[Idx];
 			}
 		}
 	}
@@ -639,9 +641,9 @@ void UK2Node_Select::SetEnum(UEnum* InEnum, bool bForceRegenerate)
 				if (!bShouldBeHidden)
 				{
 					FString EnumValueName = Enum->GetEnumName(EnumIndex);
-					FString EnumFriendlyName = Enum->GetDisplayNameText(EnumIndex).ToString();
+					FText EnumFriendlyName = Enum->GetDisplayNameText(EnumIndex);
 					EnumEntries.Add(FName(*EnumValueName));
-					EnumEntryFriendlyNames.Add(FName(*EnumFriendlyName));
+					EnumEntryFriendlyNames.Add(EnumFriendlyName);
 				}
 			}
 		}
