Attaching a bool to the "Select" nodes index pin will force the node to update it's pin names even when there are only two options.

diff --git a/Engine/Source/Editor/BlueprintGraph/Private/K2Node_Select.cpp b/Engine/Source/Editor/BlueprintGraph/Private/K2Node_Select.cpp
--- a/Engine/Source/Editor/BlueprintGraph/Private/K2Node_Select.cpp
+++ b/Engine/Source/Editor/BlueprintGraph/Private/K2Node_Select.cpp
@@ -461,8 +461,7 @@ void UK2Node_Select::NotifyPinConnectionListChanged(UEdGraphPin* Pin)
 			}
 
 			// If the index pin is a boolean, we need to remove all but 2 options
-			if (IndexPinType.PinCategory == Schema->PC_Boolean &&
-				NumOptionPins != 2)
+			if (IndexPinType.PinCategory == Schema->PC_Boolean)
 			{
 				NumOptionPins = 2;
 				bReconstructNode = true;
