When dragging off a pin in a function graph, local variable nodes will no longer be filtered out of the context menu for placement if they are valid for connection.

Stores the UProperty of the local variable in the node spawner.

diff --git a/Engine/Source/Editor/BlueprintGraph/Classes/BlueprintVariableNodeSpawner.h b/Engine/Source/Editor/BlueprintGraph/Classes/BlueprintVariableNodeSpawner.h
--- a/Engine/Source/Editor/BlueprintGraph/Classes/BlueprintVariableNodeSpawner.h
+++ b/Engine/Source/Editor/BlueprintGraph/Classes/BlueprintVariableNodeSpawner.h
@@ -40,10 +40,11 @@ public:
 	 * @param  NodeClass	The node type that you want the spawner to spawn.
 	 * @param  VarContext	The graph that the local variable belongs to.
 	 * @param  VarDesc		Details the local variable (name, type, etc.)
+	 * @param  VarProperty	The property that represents the local-variable you want nodes spawned for.
 	 * @param  Outer		Optional outer for the new spawner (if left null, the transient package will be used).
 	 * @return A newly allocated instance of this class.
 	 */
-	static UBlueprintVariableNodeSpawner* Create(TSubclassOf<UK2Node_Variable> NodeClass, UEdGraph* VarContext, FBPVariableDescription const& VarDesc, UObject* Outer = nullptr);
+	static UBlueprintVariableNodeSpawner* Create(TSubclassOf<UK2Node_Variable> NodeClass, UEdGraph* VarContext, FBPVariableDescription const& VarDesc, UProperty* VarProperty, UObject* Outer = nullptr);
 
 	// UBlueprintNodeSpawner interface
 	virtual void Prime() override;
diff --git a/Engine/Source/Editor/BlueprintGraph/Private/BlueprintActionDatabase.cpp b/Engine/Source/Editor/BlueprintGraph/Private/BlueprintActionDatabase.cpp
--- a/Engine/Source/Editor/BlueprintGraph/Private/BlueprintActionDatabase.cpp
+++ b/Engine/Source/Editor/BlueprintGraph/Private/BlueprintActionDatabase.cpp
@@ -737,9 +737,13 @@ static void BlueprintActionDatabaseImpl::AddBlueprintGraphActions(UBlueprint con
 		{
 			for (FBPVariableDescription const& LocalVar : FunctionEntry->LocalVariables)
 			{
-				UBlueprintNodeSpawner* GetVarSpawner = UBlueprintVariableNodeSpawner::Create(UK2Node_VariableGet::StaticClass(), FunctionGraph, LocalVar);
+				// Create a member reference so we can safely resolve the UProperty
+				FMemberReference Reference;
+				Reference.SetLocalMember(LocalVar.VarName, FunctionGraph->GetName(), LocalVar.VarGuid);
+
+				UBlueprintNodeSpawner* GetVarSpawner = UBlueprintVariableNodeSpawner::Create(UK2Node_VariableGet::StaticClass(), FunctionGraph, LocalVar, Reference.ResolveMember<UProperty>(Blueprint->SkeletonGeneratedClass));
 				ActionListOut.Add(GetVarSpawner);
-				UBlueprintNodeSpawner* SetVarSpawner = UBlueprintVariableNodeSpawner::Create(UK2Node_VariableSet::StaticClass(), FunctionGraph, LocalVar);
+				UBlueprintNodeSpawner* SetVarSpawner = UBlueprintVariableNodeSpawner::Create(UK2Node_VariableSet::StaticClass(), FunctionGraph, LocalVar, Reference.ResolveMember<UProperty>(Blueprint->SkeletonGeneratedClass));
 				ActionListOut.Add(SetVarSpawner);
 			}
 		}
diff --git a/Engine/Source/Editor/BlueprintGraph/Private/BlueprintVariableNodeSpawner.cpp b/Engine/Source/Editor/BlueprintGraph/Private/BlueprintVariableNodeSpawner.cpp
--- a/Engine/Source/Editor/BlueprintGraph/Private/BlueprintVariableNodeSpawner.cpp
+++ b/Engine/Source/Editor/BlueprintGraph/Private/BlueprintVariableNodeSpawner.cpp
@@ -91,7 +91,7 @@ UBlueprintVariableNodeSpawner* UBlueprintVariableNodeSpawner::Create(TSubclassOf
 }
 
 //------------------------------------------------------------------------------
-UBlueprintVariableNodeSpawner* UBlueprintVariableNodeSpawner::Create(TSubclassOf<UK2Node_Variable> NodeClass, UEdGraph* VarContext, FBPVariableDescription const& VarDesc, UObject* Outer/*= nullptr*/)
+UBlueprintVariableNodeSpawner* UBlueprintVariableNodeSpawner::Create(TSubclassOf<UK2Node_Variable> NodeClass, UEdGraph* VarContext, FBPVariableDescription const& VarDesc, UProperty* VarProperty, UObject* Outer/*= nullptr*/)
 {
 	check(VarContext != nullptr);
 	if (Outer == nullptr)
@@ -109,6 +109,7 @@ UBlueprintVariableNodeSpawner* UBlueprintVariableNodeSpawner::Create(TSubclassOf
 	NodeSpawner->NodeClass     = NodeClass;
 	NodeSpawner->LocalVarOuter = VarContext;
 	NodeSpawner->LocalVarDesc  = VarDesc;
+	NodeSpawner->Field = VarProperty;
 
 	//--------------------------------------
 	// Default UI Signature
