Various improvement in ConstructObjectFromClass K2Node

diff --git a/Engine/Source/Editor/BlueprintGraph/Classes/K2Node_ConstructObjectFromClass.h b/Engine/Source/Editor/BlueprintGraph/Classes/K2Node_ConstructObjectFromClass.h
--- a/Engine/Source/Editor/BlueprintGraph/Classes/K2Node_ConstructObjectFromClass.h
+++ b/Engine/Source/Editor/BlueprintGraph/Classes/K2Node_ConstructObjectFromClass.h
@@ -17,9 +17,9 @@ class BLUEPRINTGRAPH_API UK2Node_ConstructObjectFromClass : public UK2Node
 	virtual FText GetNodeTitle(ENodeTitleType::Type TitleType) const override;
 	virtual void PinDefaultValueChanged(UEdGraphPin* Pin) override;
 	virtual FText GetTooltipText() const override;
-	virtual void ExpandNode(class FKismetCompilerContext& CompilerContext, UEdGraph* SourceGraph) override;
 	virtual bool HasExternalDependencies(TArray<class UStruct*>* OptionalOutput) const override;
 	virtual bool IsCompatibleWithGraph(const UEdGraph* TargetGraph) const override;
+	virtual void PinConnectionListChanged(UEdGraphPin* Pin);
 	// End UEdGraphNode interface.
 
 	// Begin UK2Node interface
@@ -28,7 +28,6 @@ class BLUEPRINTGRAPH_API UK2Node_ConstructObjectFromClass : public UK2Node
 	virtual void GetNodeAttributes( TArray<TKeyValuePair<FString, FString>>& OutNodeAttributes ) const override;
 	virtual void GetMenuActions(FBlueprintActionDatabaseRegistrar& ActionRegistrar) const override;
 	virtual FText GetMenuCategory() const override;
-	virtual class FNodeHandlingFunctor* CreateNodeHandler(class FKismetCompilerContext& CompilerContext) const override;
 	// End UK2Node interface
 
 	/** Create new pins to show properties on archetype */
@@ -45,10 +44,18 @@ class BLUEPRINTGRAPH_API UK2Node_ConstructObjectFromClass : public UK2Node
 	UEdGraphPin* GetWorldContextPin() const;
 	/** Get the result output pin */
 	UEdGraphPin* GetResultPin() const;
+	/** Get the result input pin */
+	UEdGraphPin* GetOuterPin() const;
 
 	/** Get the class that we are going to spawn, if it's defined as default value */
 	UClass* GetClassToSpawn(const TArray<UEdGraphPin*>* InPinsToSearch=NULL) const;
 
+	/** Returns if the node uses World Object Context input */
+	virtual bool UseWorldContext() const;
+
+	/** Returns if the node uses Outer input */
+	virtual bool UseOuter() const { return false; }
+
 protected:
 	/** Gets the default node title when no class is selected */
 	virtual FText GetBaseNodeTitle() const;
@@ -66,6 +73,9 @@ protected:
 	 */
 	void SetPinToolTip(UEdGraphPin& MutatablePin, const FText& PinDescription) const;
 
+	/** Refresh pins when class was changed */
+	void OnClassPinChanged();
+
 	/** Tooltip text for this node. */
 	FText NodeTooltip;
 
diff --git a/Engine/Source/Editor/BlueprintGraph/Private/K2Node_ConstructObjectFromClass.cpp b/Engine/Source/Editor/BlueprintGraph/Private/K2Node_ConstructObjectFromClass.cpp
--- a/Engine/Source/Editor/BlueprintGraph/Private/K2Node_ConstructObjectFromClass.cpp
+++ b/Engine/Source/Editor/BlueprintGraph/Private/K2Node_ConstructObjectFromClass.cpp
@@ -11,10 +11,12 @@ struct FK2Node_ConstructObjectFromClassHelper
 {
 	static FString WorldContextPinName;
 	static FString ClassPinName;
+	static FString OuterPinName;
 };
 
 FString FK2Node_ConstructObjectFromClassHelper::WorldContextPinName(TEXT("WorldContextObject"));
 FString FK2Node_ConstructObjectFromClassHelper::ClassPinName(TEXT("Class"));
+FString FK2Node_ConstructObjectFromClassHelper::OuterPinName(TEXT("Outer"));
 
 #define LOCTEXT_NAMESPACE "K2Node_ConstructObjectFromClass"
 
@@ -29,6 +31,13 @@ UClass* UK2Node_ConstructObjectFromClass::GetClassPinBaseClass() const
 	return UObject::StaticClass();
 }
 
+bool UK2Node_ConstructObjectFromClass::UseWorldContext() const
+{
+	auto BP = GetBlueprint();
+	const UClass* ParentClass = BP ? BP->ParentClass : nullptr;
+	return ParentClass ? ParentClass->HasMetaData(FBlueprintMetadata::MD_ShowWorldContextPin) : false;
+}
+
 void UK2Node_ConstructObjectFromClass::AllocateDefaultPins()
 {
 	const UEdGraphSchema_K2* K2Schema = GetDefault<UEdGraphSchema_K2>();
@@ -38,7 +47,7 @@ void UK2Node_ConstructObjectFromClass::AllocateDefaultPins()
 	CreatePin(EGPD_Output, K2Schema->PC_Exec, TEXT(""), NULL, false, false, K2Schema->PN_Then);
 
 	// If required add the world context pin
-	if ( GetBlueprint()->ParentClass->HasMetaData(FBlueprintMetadata::MD_ShowWorldContextPin) )
+	if (UseWorldContext())
 	{
 		CreatePin(EGPD_Input, K2Schema->PC_Object, TEXT(""), UObject::StaticClass(), false, false, FK2Node_ConstructObjectFromClassHelper::WorldContextPinName);
 	}
@@ -51,9 +60,22 @@ void UK2Node_ConstructObjectFromClass::AllocateDefaultPins()
 	UEdGraphPin* ResultPin = CreatePin(EGPD_Output, K2Schema->PC_Object, TEXT(""), GetClassPinBaseClass(), false, false, K2Schema->PN_ReturnValue);
 	SetPinToolTip(*ResultPin, LOCTEXT("ResultPinDescription", "The spawned object"));
 
+	if (UseOuter())
+	{
+		UEdGraphPin* OuterPin = CreatePin(EGPD_Input, K2Schema->PC_Object, TEXT(""), UObject::StaticClass(), false, false, FK2Node_ConstructObjectFromClassHelper::OuterPinName);
+		SetPinToolTip(*OuterPin, LOCTEXT("OuterPinDescription", "Owner of the created object"));
+	}
+
 	Super::AllocateDefaultPins();
 }
 
+UEdGraphPin* UK2Node_ConstructObjectFromClass::GetOuterPin() const
+{
+	UEdGraphPin* Pin = FindPinChecked(FK2Node_ConstructObjectFromClassHelper::OuterPinName);
+	ensure(nullptr == Pin || Pin->Direction == EGPD_Input);
+	return Pin;
+}
+
 void UK2Node_ConstructObjectFromClass::SetPinToolTip(UEdGraphPin& MutatablePin, const FText& PinDescription) const
 {
 	MutatablePin.PinToolTip = UEdGraphSchema_K2::TypeToText(MutatablePin.PinType).ToString();
@@ -125,6 +147,11 @@ UClass* UK2Node_ConstructObjectFromClass::GetClassToSpawn(const TArray<UEdGraphP
 	{
 		UseSpawnClass = CastChecked<UClass>(ClassPin->DefaultObject);
 	}
+	else if (ClassPin && ClassPin->LinkedTo.Num())
+	{
+		auto ClassSource = ClassPin->LinkedTo[0];
+		UseSpawnClass = ClassSource ? Cast<UClass>(ClassSource->PinType.PinSubCategoryObject.Get()) : nullptr;
+	}
 
 	return UseSpawnClass;
 }
@@ -147,45 +174,59 @@ bool UK2Node_ConstructObjectFromClass::IsSpawnVarPin(UEdGraphPin* Pin)
 			Pin->PinName != K2Schema->PN_Then &&
 			Pin->PinName != K2Schema->PN_ReturnValue &&
 			Pin->PinName != FK2Node_ConstructObjectFromClassHelper::ClassPinName &&
-			Pin->PinName != FK2Node_ConstructObjectFromClassHelper::WorldContextPinName );
+			Pin->PinName != FK2Node_ConstructObjectFromClassHelper::WorldContextPinName &&
+			Pin->PinName != FK2Node_ConstructObjectFromClassHelper::OuterPinName);
 }
 
-void UK2Node_ConstructObjectFromClass::PinDefaultValueChanged(UEdGraphPin* ChangedPin) 
+void UK2Node_ConstructObjectFromClass::OnClassPinChanged()
 {
-	if (ChangedPin->PinName == FK2Node_ConstructObjectFromClassHelper::ClassPinName)
-	{
-		const UEdGraphSchema_K2* K2Schema = GetDefault<UEdGraphSchema_K2>();
+	const UEdGraphSchema_K2* K2Schema = GetDefault<UEdGraphSchema_K2>();
 
-		// Because the archetype has changed, we break the output link as the output pin type will change
-		//UEdGraphPin* ResultPin = GetResultPin();
-		//ResultPin->BreakAllPinLinks();
+	// Because the archetype has changed, we break the output link as the output pin type will change
+	//UEdGraphPin* ResultPin = GetResultPin();
+	//ResultPin->BreakAllPinLinks();
 
-		// Remove all pins related to archetype variables
-		TArray<UEdGraphPin*> OldPins = Pins;
-		for (int32 i = 0; i < OldPins.Num(); i++)
+	// Remove all pins related to archetype variables
+	TArray<UEdGraphPin*> OldPins = Pins;
+	for (int32 i = 0; i < OldPins.Num(); i++)
+	{
+		UEdGraphPin* OldPin = OldPins[i];
+		if (IsSpawnVarPin(OldPin))
 		{
-			UEdGraphPin* OldPin = OldPins[i];
-			if (IsSpawnVarPin(OldPin))
-			{
-				OldPin->BreakAllPinLinks();
-				Pins.Remove(OldPin);
-			}
+			OldPin->BreakAllPinLinks();
+			Pins.Remove(OldPin);
 		}
+	}
 
-		CachedNodeTitle.MarkDirty();
+	CachedNodeTitle.MarkDirty();
 
-		UClass* UseSpawnClass = GetClassToSpawn();
-		if (UseSpawnClass != NULL)
-		{
-			CreatePinsForClass(UseSpawnClass);
-		}
+	UClass* UseSpawnClass = GetClassToSpawn();
+	if (UseSpawnClass != NULL)
+	{
+		CreatePinsForClass(UseSpawnClass);
+	}
+
+	// Refresh the UI for the graph so the pin changes show up
+	UEdGraph* Graph = GetGraph();
+	Graph->NotifyGraphChanged();
 
-		// Refresh the UI for the graph so the pin changes show up
-		UEdGraph* Graph = GetGraph();
-		Graph->NotifyGraphChanged();
+	// Mark dirty
+	FBlueprintEditorUtils::MarkBlueprintAsModified(GetBlueprint());
+}
+
+void UK2Node_ConstructObjectFromClass::PinConnectionListChanged(UEdGraphPin* Pin)
+{
+	if (Pin && (Pin->PinName == FK2Node_ConstructObjectFromClassHelper::ClassPinName))
+	{
+		OnClassPinChanged();
+	}
+}
 
-		// Mark dirty
-		FBlueprintEditorUtils::MarkBlueprintAsModified(GetBlueprint());
+void UK2Node_ConstructObjectFromClass::PinDefaultValueChanged(UEdGraphPin* ChangedPin) 
+{
+	if (ChangedPin && (ChangedPin->PinName == FK2Node_ConstructObjectFromClassHelper::ClassPinName))
+	{
+		OnClassPinChanged();
 	}
 }
 
@@ -326,16 +367,6 @@ FText UK2Node_ConstructObjectFromClass::GetMenuCategory() const
 	return FEditorCategoryUtils::GetCommonCategory(FCommonEditorCategory::Gameplay);
 }
 
-FNodeHandlingFunctor* UK2Node_ConstructObjectFromClass::CreateNodeHandler(FKismetCompilerContext& CompilerContext) const
-{
-	return new FNodeHandlingFunctor(CompilerContext);
-}
-
-void UK2Node_ConstructObjectFromClass::ExpandNode(class FKismetCompilerContext& CompilerContext, UEdGraph* SourceGraph)
-{
-	Super::ExpandNode(CompilerContext, SourceGraph);	
-}
-
 bool UK2Node_ConstructObjectFromClass::HasExternalDependencies(TArray<class UStruct*>* OptionalOutput) const
 {
 	UClass* SourceClass = GetClassToSpawn();
