Proper class name in node title in K2Node_ConstructObjectFromClass

diff --git a/Engine/Source/Editor/BlueprintGraph/Private/K2Node_ConstructObjectFromClass.cpp b/Engine/Source/Editor/BlueprintGraph/Private/K2Node_ConstructObjectFromClass.cpp
--- a/Engine/Source/Editor/BlueprintGraph/Private/K2Node_ConstructObjectFromClass.cpp
+++ b/Engine/Source/Editor/BlueprintGraph/Private/K2Node_ConstructObjectFromClass.cpp
@@ -301,19 +301,20 @@ FText UK2Node_ConstructObjectFromClass::GetNodeTitle(ENodeTitleType::Type TitleT
 	}
 	else if (UEdGraphPin* ClassPin = GetClassPin())
 	{
+		auto DefaultValueClass = Cast<const UClass>(ClassPin->DefaultObject);
 		if (ClassPin->LinkedTo.Num() > 0)
 		{
 			// Blueprint will be determined dynamically, so we don't have the name in this case
 			return NSLOCTEXT("K2Node", "ConstructObject_Title_Unknown", "Construct");
 		}
-		else if (ClassPin->DefaultObject == nullptr)
+		else if (DefaultValueClass == nullptr)
 		{
 			return NSLOCTEXT("K2Node", "ConstructObject_Title_NONE", "Construct NONE");
 		}
 		else if (CachedNodeTitle.IsOutOfDate(this))
 		{
 			FFormatNamedArguments Args;
-			Args.Add(TEXT("ClassName"), FText::FromString(ClassPin->DefaultObject->GetName()));
+			Args.Add(TEXT("ClassName"), DefaultValueClass->GetDisplayNameText());
 			// FText::Format() is slow, so we cache this to save on performance
 			CachedNodeTitle.SetCachedText(FText::Format(GetNodeTitleFormat(), Args), this);
 		}
