Fixed uses of FindFunctionByName that did not use GET_FUNCTION_NAME_CHECKED.

diff --git a/Engine/Source/Editor/BlueprintGraph/Private/EdGraphSchema_K2.cpp b/Engine/Source/Editor/BlueprintGraph/Private/EdGraphSchema_K2.cpp
--- a/Engine/Source/Editor/BlueprintGraph/Private/EdGraphSchema_K2.cpp
+++ b/Engine/Source/Editor/BlueprintGraph/Private/EdGraphSchema_K2.cpp
@@ -5,6 +5,7 @@
 #include "Engine/LevelScriptBlueprint.h"
 #include "Kismet/KismetSystemLibrary.h"
 #include "Kismet/KismetArrayLibrary.h"
+#include "Kismet/KismetMathLibrary.h"
 #include "GraphEditorActions.h"
 #include "GraphEditorSettings.h"
 #include "ScopedTransaction.h"
@@ -5145,17 +5146,17 @@ void UEdGraphSchema_K2::BackwardCompatibilityNodeConversion(UEdGraph* Graph, boo
 
 				if(OldMakeStructNode->StructType->GetName() == TEXT("Rotator"))
 				{
-					MakeNodeFunction = FindObject<UClass>(ANY_PACKAGE, TEXT("KismetMathLibrary"))->FindFunctionByName(TEXT("MakeRot"));
+					MakeNodeFunction = UKismetMathLibrary::StaticClass()->FindFunctionByName(GET_FUNCTION_NAME_CHECKED(UKismetMathLibrary, MakeRot));
 					OldPinToNewPinMap.Add(TEXT("Rotator"), TEXT("ReturnValue"));
 				}
 				else if(OldMakeStructNode->StructType->GetName() == TEXT("Vector"))
 				{
-					MakeNodeFunction = FindObject<UClass>(ANY_PACKAGE, TEXT("KismetMathLibrary"))->FindFunctionByName(TEXT("MakeVector"));
+					MakeNodeFunction = UKismetMathLibrary::StaticClass()->FindFunctionByName(GET_FUNCTION_NAME_CHECKED(UKismetMathLibrary, MakeVector));
 					OldPinToNewPinMap.Add(TEXT("Vector"), TEXT("ReturnValue"));
 				}
 				else if(OldMakeStructNode->StructType->GetName() == TEXT("Vector2D"))
 				{
-					MakeNodeFunction = FindObject<UClass>(ANY_PACKAGE, TEXT("KismetMathLibrary"))->FindFunctionByName(TEXT("MakeVector2D"));
+					MakeNodeFunction = UKismetMathLibrary::StaticClass()->FindFunctionByName(GET_FUNCTION_NAME_CHECKED(UKismetMathLibrary, MakeVector2D));
 					OldPinToNewPinMap.Add(TEXT("Vector2D"), TEXT("ReturnValue"));
 				}
 
@@ -5226,17 +5227,17 @@ void UEdGraphSchema_K2::BackwardCompatibilityNodeConversion(UEdGraph* Graph, boo
 
 				if(OldBreakStructNode->StructType->GetName() == TEXT("Rotator"))
 				{
-					BreakNodeFunction = FindObject<UClass>(ANY_PACKAGE, TEXT("KismetMathLibrary"))->FindFunctionByName(TEXT("BreakRot"));
+					BreakNodeFunction = UKismetMathLibrary::StaticClass()->FindFunctionByName(GET_FUNCTION_NAME_CHECKED(UKismetMathLibrary, BreakRot));
 					OldPinToNewPinMap.Add(TEXT("Rotator"), TEXT("InRot"));
 				}
 				else if(OldBreakStructNode->StructType->GetName() == TEXT("Vector"))
 				{
-					BreakNodeFunction = FindObject<UClass>(ANY_PACKAGE, TEXT("KismetMathLibrary"))->FindFunctionByName(TEXT("BreakVector"));
+					BreakNodeFunction = UKismetMathLibrary::StaticClass()->FindFunctionByName(GET_FUNCTION_NAME_CHECKED(UKismetMathLibrary, BreakVector));
 					OldPinToNewPinMap.Add(TEXT("Vector"), TEXT("InVec"));
 				}
 				else if(OldBreakStructNode->StructType->GetName() == TEXT("Vector2D"))
 				{
-					BreakNodeFunction = FindObject<UClass>(ANY_PACKAGE, TEXT("KismetMathLibrary"))->FindFunctionByName(TEXT("BreakVector2D"));
+					BreakNodeFunction = UKismetMathLibrary::StaticClass()->FindFunctionByName(GET_FUNCTION_NAME_CHECKED(UKismetMathLibrary, BreakVector2D));
 					OldPinToNewPinMap.Add(TEXT("Vector2D"), TEXT("InVec"));
 				}
 
diff --git a/Engine/Source/Editor/BlueprintGraph/Private/K2Node_CallFunction.cpp b/Engine/Source/Editor/BlueprintGraph/Private/K2Node_CallFunction.cpp
--- a/Engine/Source/Editor/BlueprintGraph/Private/K2Node_CallFunction.cpp
+++ b/Engine/Source/Editor/BlueprintGraph/Private/K2Node_CallFunction.cpp
@@ -2048,14 +2048,14 @@ void UK2Node_CallFunction::CallForEachElementInArrayExpansion(UK2Node* Node, UEd
 
 	// Do loop condition
 	UK2Node_CallFunction* Condition = CompilerContext.SpawnIntermediateNode<UK2Node_CallFunction>(Node, SourceGraph); 
-	Condition->SetFromFunction(UKismetMathLibrary::StaticClass()->FindFunctionByName(TEXT("Less_IntInt")));
+	Condition->SetFromFunction(UKismetMathLibrary::StaticClass()->FindFunctionByName(GET_FUNCTION_NAME_CHECKED(UKismetMathLibrary, Less_IntInt)));
 	Condition->AllocateDefaultPins();
 	Schema->TryCreateConnection(Condition->GetReturnValuePin(), Branch->GetConditionPin());
 	Schema->TryCreateConnection(Condition->FindPinChecked(TEXT("A")), IteratorVar->GetVariablePin());
 
 	// Array size
 	UK2Node_CallArrayFunction* ArrayLength = CompilerContext.SpawnIntermediateNode<UK2Node_CallArrayFunction>(Node, SourceGraph); 
-	ArrayLength->SetFromFunction(UKismetArrayLibrary::StaticClass()->FindFunctionByName(TEXT("Array_Length")));
+	ArrayLength->SetFromFunction(UKismetArrayLibrary::StaticClass()->FindFunctionByName(GET_FUNCTION_NAME_CHECKED(UKismetArrayLibrary, Array_Length)));
 	ArrayLength->AllocateDefaultPins();
 	CompilerContext.CopyPinLinksToIntermediate(*MultiSelf, *ArrayLength->GetTargetArrayPin());
 	ArrayLength->PinConnectionListChanged(ArrayLength->GetTargetArrayPin());
@@ -2063,7 +2063,7 @@ void UK2Node_CallFunction::CallForEachElementInArrayExpansion(UK2Node* Node, UEd
 
 	// Get Element
 	UK2Node_CallArrayFunction* GetElement = CompilerContext.SpawnIntermediateNode<UK2Node_CallArrayFunction>(Node, SourceGraph); 
-	GetElement->SetFromFunction(UKismetArrayLibrary::StaticClass()->FindFunctionByName(TEXT("Array_Get")));
+	GetElement->SetFromFunction(UKismetArrayLibrary::StaticClass()->FindFunctionByName(GET_FUNCTION_NAME_CHECKED(UKismetArrayLibrary, Array_Get)));
 	GetElement->AllocateDefaultPins();
 	CompilerContext.CopyPinLinksToIntermediate(*MultiSelf, *GetElement->GetTargetArrayPin());
 	GetElement->PinConnectionListChanged(GetElement->GetTargetArrayPin());
@@ -2071,7 +2071,7 @@ void UK2Node_CallFunction::CallForEachElementInArrayExpansion(UK2Node* Node, UEd
 
 	// Iterator increment
 	UK2Node_CallFunction* Increment = CompilerContext.SpawnIntermediateNode<UK2Node_CallFunction>(Node, SourceGraph); 
-	Increment->SetFromFunction(UKismetMathLibrary::StaticClass()->FindFunctionByName(TEXT("Add_IntInt")));
+	Increment->SetFromFunction(UKismetMathLibrary::StaticClass()->FindFunctionByName(GET_FUNCTION_NAME_CHECKED(UKismetMathLibrary, Add_IntInt)));
 	Increment->AllocateDefaultPins();
 	Schema->TryCreateConnection(Increment->FindPinChecked(TEXT("A")), IteratorVar->GetVariablePin());
 	Increment->FindPinChecked(TEXT("B"))->DefaultValue = TEXT("1");
diff --git a/Engine/Source/Editor/BlueprintGraph/Private/K2Node_ForEachElementInEnum.cpp b/Engine/Source/Editor/BlueprintGraph/Private/K2Node_ForEachElementInEnum.cpp
--- a/Engine/Source/Editor/BlueprintGraph/Private/K2Node_ForEachElementInEnum.cpp
+++ b/Engine/Source/Editor/BlueprintGraph/Private/K2Node_ForEachElementInEnum.cpp
@@ -79,7 +79,7 @@ struct FForExpandNodeHelper
 
 		// Do loop condition
 		UK2Node_CallFunction* Condition = CompilerContext.SpawnIntermediateNode<UK2Node_CallFunction>(Node, SourceGraph); 
-		Condition->SetFromFunction(UKismetMathLibrary::StaticClass()->FindFunctionByName(TEXT("Less_IntInt")));
+		Condition->SetFromFunction(UKismetMathLibrary::StaticClass()->FindFunctionByName(GET_FUNCTION_NAME_CHECKED(UKismetMathLibrary, Less_IntInt)));
 		Condition->AllocateDefaultPins();
 		bResult &= Schema->TryCreateConnection(Condition->GetReturnValuePin(), Branch->GetConditionPin());
 		bResult &= Schema->TryCreateConnection(Condition->FindPinChecked(TEXT("A")), LoopCounterOutPin);
@@ -102,7 +102,7 @@ struct FForExpandNodeHelper
 
 		// Loop Counter increment
 		UK2Node_CallFunction* Increment = CompilerContext.SpawnIntermediateNode<UK2Node_CallFunction>(Node, SourceGraph); 
-		Increment->SetFromFunction(UKismetMathLibrary::StaticClass()->FindFunctionByName(TEXT("Add_IntInt")));
+		Increment->SetFromFunction(UKismetMathLibrary::StaticClass()->FindFunctionByName(GET_FUNCTION_NAME_CHECKED(UKismetMathLibrary, Add_IntInt)));
 		Increment->AllocateDefaultPins();
 		bResult &= Schema->TryCreateConnection(Increment->FindPinChecked(TEXT("A")), LoopCounterOutPin);
 		Increment->FindPinChecked(TEXT("B"))->DefaultValue = TEXT("1");
diff --git a/Engine/Source/Editor/BlueprintGraph/Private/K2Node_GetInputAxisKeyValue.cpp b/Engine/Source/Editor/BlueprintGraph/Private/K2Node_GetInputAxisKeyValue.cpp
--- a/Engine/Source/Editor/BlueprintGraph/Private/K2Node_GetInputAxisKeyValue.cpp
+++ b/Engine/Source/Editor/BlueprintGraph/Private/K2Node_GetInputAxisKeyValue.cpp
@@ -29,7 +29,7 @@ void UK2Node_GetInputAxisKeyValue::AllocateDefaultPins()
 void UK2Node_GetInputAxisKeyValue::Initialize(const FKey AxisKey)
 {
 	InputAxisKey = AxisKey;
-	SetFromFunction(AActor::StaticClass()->FindFunctionByName(TEXT("GetInputAxisKeyValue")));
+	SetFromFunction(AActor::StaticClass()->FindFunctionByName(GET_FUNCTION_NAME_CHECKED(AActor, GetInputAxisKeyValue)));
 	
 	CachedTooltip.MarkDirty();
 	CachedNodeTitle.MarkDirty();
diff --git a/Engine/Source/Editor/BlueprintGraph/Private/K2Node_GetInputAxisValue.cpp b/Engine/Source/Editor/BlueprintGraph/Private/K2Node_GetInputAxisValue.cpp
--- a/Engine/Source/Editor/BlueprintGraph/Private/K2Node_GetInputAxisValue.cpp
+++ b/Engine/Source/Editor/BlueprintGraph/Private/K2Node_GetInputAxisValue.cpp
@@ -29,7 +29,7 @@ void UK2Node_GetInputAxisValue::AllocateDefaultPins()
 void UK2Node_GetInputAxisValue::Initialize(const FName AxisName)
 {
 	InputAxisName = AxisName;
-	SetFromFunction(AActor::StaticClass()->FindFunctionByName(TEXT("GetInputAxisValue")));
+	SetFromFunction(AActor::StaticClass()->FindFunctionByName(GET_FUNCTION_NAME_CHECKED(AActor, GetInputAxisValue)));
 }
 
 FText UK2Node_GetInputAxisValue::GetNodeTitle(ENodeTitleType::Type TitleType) const
diff --git a/Engine/Source/Editor/BlueprintGraph/Private/K2Node_GetInputVectorAxisValue.cpp b/Engine/Source/Editor/BlueprintGraph/Private/K2Node_GetInputVectorAxisValue.cpp
--- a/Engine/Source/Editor/BlueprintGraph/Private/K2Node_GetInputVectorAxisValue.cpp
+++ b/Engine/Source/Editor/BlueprintGraph/Private/K2Node_GetInputVectorAxisValue.cpp
@@ -16,7 +16,7 @@ UK2Node_GetInputVectorAxisValue::UK2Node_GetInputVectorAxisValue(const FObjectIn
 void UK2Node_GetInputVectorAxisValue::Initialize(const FKey AxisKey)
 {
 	InputAxisKey = AxisKey;
-	SetFromFunction(AActor::StaticClass()->FindFunctionByName(TEXT("GetInputVectorAxisValue")));
+	SetFromFunction(AActor::StaticClass()->FindFunctionByName(GET_FUNCTION_NAME_CHECKED(AActor, GetInputVectorAxisValue)));
 }
 
 FText UK2Node_GetInputVectorAxisValue::GetTooltipText() const
diff --git a/Engine/Source/Editor/GameplayTagsEditor/Private/GameplayTagsK2Node_MultiCompareGameplayTagAssetInterface.cpp b/Engine/Source/Editor/GameplayTagsEditor/Private/GameplayTagsK2Node_MultiCompareGameplayTagAssetInterface.cpp
--- a/Engine/Source/Editor/GameplayTagsEditor/Private/GameplayTagsK2Node_MultiCompareGameplayTagAssetInterface.cpp
+++ b/Engine/Source/Editor/GameplayTagsEditor/Private/GameplayTagsK2Node_MultiCompareGameplayTagAssetInterface.cpp
@@ -47,7 +47,7 @@ void UGameplayTagsK2Node_MultiCompareGameplayTagAssetInterface::ExpandNode(class
 
 		// Create call function node for the Compare function HasAllMatchingGameplayTags
 		UK2Node_CallFunction* PinCallFunction = CompilerContext.SpawnIntermediateNode<UK2Node_CallFunction>(this, SourceGraph);
-		const UFunction* Function = UBlueprintGameplayTagLibrary::StaticClass()->FindFunctionByName(TEXT("HasAllMatchingGameplayTags"));
+		const UFunction* Function = UBlueprintGameplayTagLibrary::StaticClass()->FindFunctionByName(GET_FUNCTION_NAME_CHECKED(UBlueprintGameplayTagLibrary, HasAllMatchingGameplayTags));
 		PinCallFunction->SetFromFunction(Function);
 		PinCallFunction->AllocateDefaultPins();
 
diff --git a/Engine/Source/Editor/GameplayTagsEditor/Private/GameplayTagsK2Node_MultiCompareGameplayTagAssetInterfaceSingleTags.cpp b/Engine/Source/Editor/GameplayTagsEditor/Private/GameplayTagsK2Node_MultiCompareGameplayTagAssetInterfaceSingleTags.cpp
--- a/Engine/Source/Editor/GameplayTagsEditor/Private/GameplayTagsK2Node_MultiCompareGameplayTagAssetInterfaceSingleTags.cpp
+++ b/Engine/Source/Editor/GameplayTagsEditor/Private/GameplayTagsK2Node_MultiCompareGameplayTagAssetInterfaceSingleTags.cpp
@@ -50,7 +50,7 @@ void UGameplayTagsK2Node_MultiCompareGameplayTagAssetInterfaceSingleTags::Expand
 
 		// Create call function node for the Compare function HasAllMatchingGameplayTags
 		UK2Node_CallFunction* PinCallFunction = CompilerContext.SpawnIntermediateNode<UK2Node_CallFunction>(this, SourceGraph);
-		const UFunction* Function = UBlueprintGameplayTagLibrary::StaticClass()->FindFunctionByName(TEXT("DoesTagAssetInterfaceHaveTag"));
+		const UFunction* Function = UBlueprintGameplayTagLibrary::StaticClass()->FindFunctionByName(GET_FUNCTION_NAME_CHECKED(UBlueprintGameplayTagLibrary, DoesTagAssetInterfaceHaveTag));
 		PinCallFunction->SetFromFunction(Function);
 		PinCallFunction->AllocateDefaultPins();
 
diff --git a/Engine/Source/Editor/GameplayTagsEditor/Private/GameplayTagsK2Node_MultiCompareGameplayTagContainer.cpp b/Engine/Source/Editor/GameplayTagsEditor/Private/GameplayTagsK2Node_MultiCompareGameplayTagContainer.cpp
--- a/Engine/Source/Editor/GameplayTagsEditor/Private/GameplayTagsK2Node_MultiCompareGameplayTagContainer.cpp
+++ b/Engine/Source/Editor/GameplayTagsEditor/Private/GameplayTagsK2Node_MultiCompareGameplayTagContainer.cpp
@@ -47,7 +47,7 @@ void UGameplayTagsK2Node_MultiCompareGameplayTagContainer::ExpandNode(class FKis
 
 		// Create call function node for the Compare function HasAllMatchingGameplayTags
 		UK2Node_CallFunction* PinCallFunction = CompilerContext.SpawnIntermediateNode<UK2Node_CallFunction>(this, SourceGraph);
-		const UFunction* Function = UBlueprintGameplayTagLibrary::StaticClass()->FindFunctionByName(TEXT("DoesContainerMatchAllTagsInContainer"));
+		const UFunction* Function = UBlueprintGameplayTagLibrary::StaticClass()->FindFunctionByName(GET_FUNCTION_NAME_CHECKED(UBlueprintGameplayTagLibrary, DoesContainerMatchAllTagsInContainer));
 		PinCallFunction->SetFromFunction(Function);
 		PinCallFunction->AllocateDefaultPins();
 
diff --git a/Engine/Source/Editor/GameplayTagsEditor/Private/GameplayTagsK2Node_MultiCompareGameplayTagContainerSingleTags.cpp b/Engine/Source/Editor/GameplayTagsEditor/Private/GameplayTagsK2Node_MultiCompareGameplayTagContainerSingleTags.cpp
--- a/Engine/Source/Editor/GameplayTagsEditor/Private/GameplayTagsK2Node_MultiCompareGameplayTagContainerSingleTags.cpp
+++ b/Engine/Source/Editor/GameplayTagsEditor/Private/GameplayTagsK2Node_MultiCompareGameplayTagContainerSingleTags.cpp
@@ -50,7 +50,7 @@ void UGameplayTagsK2Node_MultiCompareGameplayTagContainerSingleTags::ExpandNode(
 
 		// Create call function node for the Compare function HasAllMatchingGameplayTags
 		UK2Node_CallFunction* PinCallFunction = CompilerContext.SpawnIntermediateNode<UK2Node_CallFunction>(this, SourceGraph);
-		const UFunction* Function = UBlueprintGameplayTagLibrary::StaticClass()->FindFunctionByName(TEXT("DoesContainerHaveTag"));
+		const UFunction* Function = UBlueprintGameplayTagLibrary::StaticClass()->FindFunctionByName(GET_FUNCTION_NAME_CHECKED(UBlueprintGameplayTagLibrary, DoesContainerHaveTag));
 		PinCallFunction->SetFromFunction(Function);
 		PinCallFunction->AllocateDefaultPins();
 
diff --git a/Engine/Source/Editor/UnrealEd/Private/Kismet2/Kismet2.cpp b/Engine/Source/Editor/UnrealEd/Private/Kismet2/Kismet2.cpp
--- a/Engine/Source/Editor/UnrealEd/Private/Kismet2/Kismet2.cpp
+++ b/Engine/Source/Editor/UnrealEd/Private/Kismet2/Kismet2.cpp
@@ -486,7 +486,7 @@ UBlueprint* FKismetEditorUtilities::CreateBlueprint(UClass* ParentClass, UObject
 
 				// add try get owner node
 				UK2Node_CallFunction* GetOwnerNode = NewObject<UK2Node_CallFunction>(Graph);
-				UFunction* MakeNodeFunction = FindObject<UClass>(ANY_PACKAGE, TEXT("AnimInstance"))->FindFunctionByName(TEXT("TryGetPawnOwner"));
+				UFunction* MakeNodeFunction = UAnimInstance::StaticClass()->FindFunctionByName(GET_FUNCTION_NAME_CHECKED(UAnimInstance, TryGetPawnOwner));
 				GetOwnerNode->CreateNewGuid();
 				GetOwnerNode->PostPlacedNewNode();
 				GetOwnerNode->SetFromFunction(MakeNodeFunction);
