K2Node_ConstructObjectFromClass - links to exposed properties are not broken when the class is changed.

diff --git a/Engine/Source/Editor/BlueprintGraph/Classes/K2Node_ConstructObjectFromClass.h b/Engine/Source/Editor/BlueprintGraph/Classes/K2Node_ConstructObjectFromClass.h
--- a/Engine/Source/Editor/BlueprintGraph/Classes/K2Node_ConstructObjectFromClass.h
+++ b/Engine/Source/Editor/BlueprintGraph/Classes/K2Node_ConstructObjectFromClass.h
@@ -31,7 +31,7 @@ class BLUEPRINTGRAPH_API UK2Node_ConstructObjectFromClass : public UK2Node
 	// End UK2Node interface
 
 	/** Create new pins to show properties on archetype */
-	void CreatePinsForClass(UClass* InClass);
+	void CreatePinsForClass(UClass* InClass, TArray<UEdGraphPin*>* OutClassPins = nullptr);
 
 	/** See if this is a spawn variable pin, or a 'default' pin */
 	virtual bool IsSpawnVarPin(UEdGraphPin* Pin);
diff --git a/Engine/Source/Editor/BlueprintGraph/Private/K2Node_ConstructObjectFromClass.cpp b/Engine/Source/Editor/BlueprintGraph/Private/K2Node_ConstructObjectFromClass.cpp
--- a/Engine/Source/Editor/BlueprintGraph/Private/K2Node_ConstructObjectFromClass.cpp
+++ b/Engine/Source/Editor/BlueprintGraph/Private/K2Node_ConstructObjectFromClass.cpp
@@ -90,7 +90,7 @@ void UK2Node_ConstructObjectFromClass::SetPinToolTip(UEdGraphPin& MutatablePin,
 	MutatablePin.PinToolTip += FString(TEXT("\n")) + PinDescription.ToString();
 }
 
-void UK2Node_ConstructObjectFromClass::CreatePinsForClass(UClass* InClass)
+void UK2Node_ConstructObjectFromClass::CreatePinsForClass(UClass* InClass, TArray<UEdGraphPin*>* OutClassPins)
 {
 	check(InClass != NULL);
 
@@ -116,6 +116,11 @@ void UK2Node_ConstructObjectFromClass::CreatePinsForClass(UClass* InClass)
 			UEdGraphPin* Pin = CreatePin(EGPD_Input, TEXT(""), TEXT(""), NULL, false, false, Property->GetName());
 			const bool bPinGood = (Pin != NULL) && K2Schema->ConvertPropertyToPinType(Property, /*out*/ Pin->PinType);
 			
+			if (OutClassPins && Pin)
+			{
+				OutClassPins->Add(Pin);
+			}
+
 			if (ClassDefaultObject && Pin != NULL && K2Schema->PinDefaultValueIsEditable(*Pin))
 			{
 				FString DefaultValueAsString;
@@ -182,30 +187,35 @@ void UK2Node_ConstructObjectFromClass::OnClassPinChanged()
 {
 	const UEdGraphSchema_K2* K2Schema = GetDefault<UEdGraphSchema_K2>();
 
-	// Because the archetype has changed, we break the output link as the output pin type will change
-	//UEdGraphPin* ResultPin = GetResultPin();
-	//ResultPin->BreakAllPinLinks();
-
 	// Remove all pins related to archetype variables
 	TArray<UEdGraphPin*> OldPins = Pins;
+	TArray<UEdGraphPin*> OldClassPins;
+
 	for (int32 i = 0; i < OldPins.Num(); i++)
 	{
 		UEdGraphPin* OldPin = OldPins[i];
 		if (IsSpawnVarPin(OldPin))
 		{
-			OldPin->BreakAllPinLinks();
 			Pins.Remove(OldPin);
+			OldClassPins.Add(OldPin);
 		}
 	}
 
 	CachedNodeTitle.MarkDirty();
 
 	UClass* UseSpawnClass = GetClassToSpawn();
+	TArray<UEdGraphPin*> NewClassPins;
 	if (UseSpawnClass != NULL)
 	{
-		CreatePinsForClass(UseSpawnClass);
+		CreatePinsForClass(UseSpawnClass, &NewClassPins);
 	}
 
+	// Rewire the old pins to the new pins so connections are maintained if possible
+	RewireOldPinsToNewPins(OldClassPins, NewClassPins);
+
+	// Destroy the old pins
+	DestroyPinList(OldClassPins);
+
 	// Refresh the UI for the graph so the pin changes show up
 	UEdGraph* Graph = GetGraph();
 	Graph->NotifyGraphChanged();
@@ -299,32 +309,18 @@ FText UK2Node_ConstructObjectFromClass::GetNodeTitle(ENodeTitleType::Type TitleT
 	{
 		return GetBaseNodeTitle();
 	}
-	else if (UEdGraphPin* ClassPin = GetClassPin())
+	else if (auto ClassToSpawn = GetClassToSpawn())
 	{
-		auto DefaultValueClass = Cast<const UClass>(ClassPin->DefaultObject);
-		if (ClassPin->LinkedTo.Num() > 0)
-		{
-			// Blueprint will be determined dynamically, so we don't have the name in this case
-			return NSLOCTEXT("K2Node", "ConstructObject_Title_Unknown", "Construct");
-		}
-		else if (DefaultValueClass == nullptr)
-		{
-			return NSLOCTEXT("K2Node", "ConstructObject_Title_NONE", "Construct NONE");
-		}
-		else if (CachedNodeTitle.IsOutOfDate(this))
+		if (CachedNodeTitle.IsOutOfDate(this))
 		{
 			FFormatNamedArguments Args;
-			Args.Add(TEXT("ClassName"), DefaultValueClass->GetDisplayNameText());
+			Args.Add(TEXT("ClassName"), ClassToSpawn->GetDisplayNameText());
 			// FText::Format() is slow, so we cache this to save on performance
 			CachedNodeTitle.SetCachedText(FText::Format(GetNodeTitleFormat(), Args), this);
 		}
+		return CachedNodeTitle;
 	}
-	else
-	{
-		return NSLOCTEXT("K2Node", "ConstructObject_Title_NONE", "Construct NONE");
-	}
-
-	return CachedNodeTitle;
+	return NSLOCTEXT("K2Node", "ConstructObject_Title_NONE", "Construct NONE");
 }
 
 bool UK2Node_ConstructObjectFromClass::IsCompatibleWithGraph(const UEdGraph* TargetGraph) const 
