Input action and key events will no longer be placeable in Blueprint functions.

Added checks to UK2Node_InputAction::IsCompatibleWithGraph and UK2Node_InputKey::IsCompatibleWithGraph to verify the graph type since these nodes do not directly inherit from K2Node_Event (in ExpandNode they replace themselves with actual events).

diff --git a/Engine/Source/Editor/BlueprintGraph/Private/K2Node_InputAction.cpp b/Engine/Source/Editor/BlueprintGraph/Private/K2Node_InputAction.cpp
--- a/Engine/Source/Editor/BlueprintGraph/Private/K2Node_InputAction.cpp
+++ b/Engine/Source/Editor/BlueprintGraph/Private/K2Node_InputAction.cpp
@@ -77,12 +77,20 @@ FText UK2Node_InputAction::GetTooltipText() const
 
 bool UK2Node_InputAction::IsCompatibleWithGraph(UEdGraph const* Graph) const
 {
-	UBlueprint* Blueprint = FBlueprintEditorUtils::FindBlueprintForGraph(Graph);
+	// This node expands into event nodes and must be placed in a Ubergraph
+	EGraphType const GraphType = Graph->GetSchema()->GetGraphType(Graph);
+	bool bIsCompatible = (GraphType == EGraphType::GT_Ubergraph);
 
-	UEdGraphSchema_K2 const* K2Schema = Cast<UEdGraphSchema_K2>(Graph->GetSchema());
-	bool const bIsConstructionScript = (K2Schema != nullptr) ? K2Schema->IsConstructionScript(Graph) : false;
+	if (bIsCompatible)
+	{
+		UBlueprint* Blueprint = FBlueprintEditorUtils::FindBlueprintForGraph(Graph);
+
+		UEdGraphSchema_K2 const* K2Schema = Cast<UEdGraphSchema_K2>(Graph->GetSchema());
+		bool const bIsConstructionScript = (K2Schema != nullptr) ? K2Schema->IsConstructionScript(Graph) : false;
 
-	return (Blueprint != nullptr) && Blueprint->SupportsInputEvents() && !bIsConstructionScript && Super::IsCompatibleWithGraph(Graph);
+		bIsCompatible = (Blueprint != nullptr) && Blueprint->SupportsInputEvents() && !bIsConstructionScript && Super::IsCompatibleWithGraph(Graph);
+	}
+	return bIsCompatible;
 }
 
 UEdGraphPin* UK2Node_InputAction::GetPressedPin() const
diff --git a/Engine/Source/Editor/BlueprintGraph/Private/K2Node_InputKey.cpp b/Engine/Source/Editor/BlueprintGraph/Private/K2Node_InputKey.cpp
--- a/Engine/Source/Editor/BlueprintGraph/Private/K2Node_InputKey.cpp
+++ b/Engine/Source/Editor/BlueprintGraph/Private/K2Node_InputKey.cpp
@@ -198,12 +198,20 @@ FName UK2Node_InputKey::GetPaletteIcon(FLinearColor& OutColor) const
 
 bool UK2Node_InputKey::IsCompatibleWithGraph(UEdGraph const* Graph) const
 {
-	UBlueprint* Blueprint = FBlueprintEditorUtils::FindBlueprintForGraph(Graph);
+	// This node expands into event nodes and must be placed in a Ubergraph
+	EGraphType const GraphType = Graph->GetSchema()->GetGraphType(Graph);
+	bool bIsCompatible = (GraphType == EGraphType::GT_Ubergraph);
 
-	UEdGraphSchema_K2 const* K2Schema = Cast<UEdGraphSchema_K2>(Graph->GetSchema());
-	bool const bIsConstructionScript = (K2Schema != nullptr) ? K2Schema->IsConstructionScript(Graph) : false;
+	if (bIsCompatible)
+	{
+		UBlueprint* Blueprint = FBlueprintEditorUtils::FindBlueprintForGraph(Graph);
+
+		UEdGraphSchema_K2 const* K2Schema = Cast<UEdGraphSchema_K2>(Graph->GetSchema());
+		bool const bIsConstructionScript = (K2Schema != nullptr) ? K2Schema->IsConstructionScript(Graph) : false;
 
-	return (Blueprint != nullptr) && Blueprint->SupportsInputEvents() && !bIsConstructionScript && Super::IsCompatibleWithGraph(Graph);
+		bIsCompatible = (Blueprint != nullptr) && Blueprint->SupportsInputEvents() && !bIsConstructionScript && Super::IsCompatibleWithGraph(Graph);
+	}
+	return bIsCompatible;
 }
 
 UEdGraphPin* UK2Node_InputKey::GetPressedPin() const
