UE-16513 [CrashReport] Adding a Return node to a function that contains a Local array variable then compiling will crash the engine

diff --git a/Engine/Source/Editor/BlueprintGraph/Private/K2Node_FunctionEntry.cpp b/Engine/Source/Editor/BlueprintGraph/Private/K2Node_FunctionEntry.cpp
--- a/Engine/Source/Editor/BlueprintGraph/Private/K2Node_FunctionEntry.cpp
+++ b/Engine/Source/Editor/BlueprintGraph/Private/K2Node_FunctionEntry.cpp
@@ -314,8 +314,8 @@ void UK2Node_FunctionEntry::ExpandNode(class FKismetCompilerContext& CompilerCon
 								const UArrayProperty* ArrayProperty = Cast<UArrayProperty>(Property);
 								check(ArrayProperty);
 
-								FScriptArrayHelper ArrayHelper(ArrayProperty, StructData->GetStructMemory());
-								FScriptArrayHelper DefaultArrayHelper(ArrayProperty, StructData->GetStructMemory());
+								FScriptArrayHelper_InContainer ArrayHelper(ArrayProperty, StructData->GetStructMemory());
+								FScriptArrayHelper_InContainer DefaultArrayHelper(ArrayProperty, StructData->GetStructMemory());
 
 								uint8* StructDefaults = NULL;
 								UStructProperty* StructProperty = dynamic_cast<UStructProperty*>(ArrayProperty->Inner);
