UE-11946 Support Objects in Blueprint

- GenericCreateObject node
- UGameplayStatics::SpawnObject function
- DontUseGenericSpawnObject class meta data

Missing: runtime mechanism to verify if an object of given class can be spawned using Generic Create Object

diff --git a/Engine/Source/Editor/BlueprintGraph/Classes/K2Node_GenericCreateObject.h b/Engine/Source/Editor/BlueprintGraph/Classes/K2Node_GenericCreateObject.h
new file mode 100644
--- /dev/null
+++ b/Engine/Source/Editor/BlueprintGraph/Classes/K2Node_GenericCreateObject.h
@@ -0,0 +1,22 @@
+// Copyright 1998-2015 Epic Games, Inc. All Rights Reserved.
+
+#pragma once
+
+#include "K2Node.h"
+#include "K2Node_ConstructObjectFromClass.h"
+#include "K2Node_GenericCreateObject.generated.h"
+
+UCLASS()
+class BLUEPRINTGRAPH_API UK2Node_GenericCreateObject : public UK2Node_ConstructObjectFromClass
+{
+	GENERATED_BODY()
+
+	// Begin UEdGraphNode interface.
+	virtual void ExpandNode(class FKismetCompilerContext& CompilerContext, UEdGraph* SourceGraph) override;
+	virtual void ValidateNodeDuringCompilation(class FCompilerResultsLog& MessageLog) const override;
+	virtual void EarlyValidation(class FCompilerResultsLog& MessageLog) const override;
+	// End UEdGraphNode interface.
+
+	virtual bool UseWorldContext() const override { return false; }
+	virtual bool UseOuter() const override { return true; }
+};
diff --git a/Engine/Source/Editor/BlueprintGraph/Private/K2Node_GenericCreateObject.cpp b/Engine/Source/Editor/BlueprintGraph/Private/K2Node_GenericCreateObject.cpp
new file mode 100644
--- /dev/null
+++ b/Engine/Source/Editor/BlueprintGraph/Private/K2Node_GenericCreateObject.cpp
@@ -0,0 +1,89 @@
+// Copyright 1998-2015 Epic Games, Inc. All Rights Reserved.
+
+#include "BlueprintGraphPrivatePCH.h"
+#include "KismetCompiler.h"
+#include "Runtime/Engine/Classes/Kismet/GameplayStatics.h"
+#include "K2Node_GenericCreateObject.h"
+
+#define LOCTEXT_NAMESPACE "K2Node_GenericCreateObject"
+
+void UK2Node_GenericCreateObject::ExpandNode(class FKismetCompilerContext& CompilerContext, UEdGraph* SourceGraph)
+{
+	UK2Node_CallFunction* CallCreateNode = CompilerContext.SpawnIntermediateNode<UK2Node_CallFunction>(this, SourceGraph);
+	CallCreateNode->FunctionReference.SetExternalMember(GET_FUNCTION_NAME_CHECKED(UGameplayStatics, SpawnObject), UGameplayStatics::StaticClass());
+	CallCreateNode->AllocateDefaultPins();
+
+	bool bSucceeded = true;
+	//connect exe
+	{
+		auto SpawnExecPin = GetExecPin();
+		auto CallExecPin = CallCreateNode->GetExecPin();
+		bSucceeded &= SpawnExecPin && CallExecPin && CompilerContext.MovePinLinksToIntermediate(*SpawnExecPin, *CallExecPin).CanSafeConnect();
+	}
+
+	//connect class
+	{
+		auto SpawnClassPin = GetClassPin();
+		auto CallClassPin = CallCreateNode->FindPin(TEXT("ObjectClass"));
+		bSucceeded &= SpawnClassPin && CallClassPin && CompilerContext.MovePinLinksToIntermediate(*SpawnClassPin, *CallClassPin).CanSafeConnect();
+	}
+		
+	//connect outer
+	{
+		auto SpawnOuterPin = GetOuterPin();
+		auto CallOuterPin = CallCreateNode->FindPin(TEXT("Outer"));
+		bSucceeded &= SpawnOuterPin && CallOuterPin && CompilerContext.MovePinLinksToIntermediate(*SpawnOuterPin, *CallOuterPin).CanSafeConnect();
+	}
+
+	UEdGraphPin* CallResultPin = nullptr;
+	//connect result
+	{
+		auto SpawnResultPin = GetResultPin();
+		CallResultPin = CallCreateNode->GetReturnValuePin();
+
+		// cast HACK. It should be safe. The only problem is native code generation.
+		if (SpawnResultPin && CallResultPin)
+		{
+			CallResultPin->PinType = SpawnResultPin->PinType;
+		}
+		bSucceeded &= SpawnResultPin && CallResultPin && CompilerContext.MovePinLinksToIntermediate(*SpawnResultPin, *CallResultPin).CanSafeConnect();
+	}
+
+	//assign exposed values and connect then
+	{
+		auto LastThen = FKismetCompilerUtilities::GenerateAssignmentNodes(CompilerContext, SourceGraph, CallCreateNode, this, CallResultPin, GetClassToSpawn());
+		auto SpawnNodeThen = GetThenPin();
+		bSucceeded &= SpawnNodeThen && LastThen && CompilerContext.MovePinLinksToIntermediate(*SpawnNodeThen, *LastThen).CanSafeConnect();
+	}
+
+	BreakAllNodeLinks();
+
+	if (!bSucceeded)
+	{
+		CompilerContext.MessageLog.Error(*LOCTEXT("GenericCreateObject_Error", "ICE: GenericCreateObject error @@").ToString(), this);
+	}
+}
+
+void UK2Node_GenericCreateObject::ValidateNodeDuringCompilation(class FCompilerResultsLog& MessageLog) const
+{
+	Super::ValidateNodeDuringCompilation(MessageLog);
+}
+
+void UK2Node_GenericCreateObject::EarlyValidation(class FCompilerResultsLog& MessageLog) const
+{
+	Super::EarlyValidation(MessageLog);
+
+	auto ClassToSpawn = GetClassToSpawn();
+	if (!UGameplayStatics::CanSpawnObjectOfClass(ClassToSpawn))
+	{
+		MessageLog.Error(*FString::Printf(*LOCTEXT("GenericCreateObject_WrongClass", "Wrong class to spawn '%s' in @@").ToString(), *GetPathNameSafe(ClassToSpawn)), this);
+	}
+
+	auto OuterPin = GetOuterPin();
+	if (!OuterPin || (!OuterPin->DefaultObject && !OuterPin->LinkedTo.Num()))
+	{
+		MessageLog.Error(*LOCTEXT("GenericCreateObject_NoOuter", "Outer object is required in @@").ToString(), this);
+	}
+}
+
+#undef LOCTEXT_NAMESPACE
\ No newline at end of file
diff --git a/Engine/Source/Runtime/CoreUObject/Public/UObject/ObjectBase.h b/Engine/Source/Runtime/CoreUObject/Public/UObject/ObjectBase.h
--- a/Engine/Source/Runtime/CoreUObject/Public/UObject/ObjectBase.h
+++ b/Engine/Source/Runtime/CoreUObject/Public/UObject/ObjectBase.h
@@ -979,6 +979,9 @@ namespace UM
 		/// [ClassMetadata] Indicates that when placing blueprint nodes in graphs owned by this class that the hidden world context pin should be visible because the self context of the class cannot
 		///                 provide the world context and it must be wired in manually
 		ShowWorldContextPin,
+
+		//[ClassMetadata] Do not spawn an object of the class using Generic Create Object node in Blueprint. It makes sense only for a BluprintType class, that is neither Actor, nor ActorComponent.
+		DontUseGenericSpawnObject,
 	};
 
 	// Metadata usable in USTRUCT
diff --git a/Engine/Source/Runtime/Engine/Classes/Kismet/GameplayStatics.h b/Engine/Source/Runtime/Engine/Classes/Kismet/GameplayStatics.h
--- a/Engine/Source/Runtime/Engine/Classes/Kismet/GameplayStatics.h
+++ b/Engine/Source/Runtime/Engine/Classes/Kismet/GameplayStatics.h
@@ -30,6 +30,12 @@ class ENGINE_API UGameplayStatics : public UBlueprintFunctionLibrary
 {
 	GENERATED_UCLASS_BODY()
 
+	// --- Create Object
+	UFUNCTION(BlueprintCallable, Category = "Spawning", meta = (UnsafeDuringActorConstruction = "true", BlueprintInternalUseOnly = "true"))
+	static UObject* SpawnObject(TSubclassOf<UObject> ObjectClass, UObject* Outer);
+
+	static bool CanSpawnObjectOfClass(TSubclassOf<UObject> ObjectClass);
+
 	// --- Spawning functions ------------------------------
 
 	/** Spawns an instance of a blueprint, but does not automatically run its construction script.  */
diff --git a/Engine/Source/Runtime/Engine/Private/GameplayStatics.cpp b/Engine/Source/Runtime/Engine/Private/GameplayStatics.cpp
--- a/Engine/Source/Runtime/Engine/Private/GameplayStatics.cpp
+++ b/Engine/Source/Runtime/Engine/Private/GameplayStatics.cpp
@@ -19,6 +19,9 @@
 #include "PhysicalMaterials/PhysicalMaterial.h"
 #include "Components/DecalComponent.h"
 #include "LandscapeProxy.h"
+#include "MessageLog.h"
+
+#define LOCTEXT_NAMESPACE "GameplayStatics"
 
 static const int UE4_SAVEGAME_FILE_TYPE_TAG = 0x53415647;		// "sAvG"
 static const int UE4_SAVEGAME_FILE_VERSION = 1;
@@ -282,6 +285,74 @@ void UGameplayStatics::ApplyDamage(AActor* DamagedActor, float BaseDamage, ACont
 	}
 }
 
+bool UGameplayStatics::CanSpawnObjectOfClass(TSubclassOf<UObject> ObjectClass)
+{
+	bool bBlueprintType = true;
+#if WITH_EDITOR
+	{
+		static const FName BlueprintTypeName(TEXT("BlueprintType"));
+		static const FName NotBlueprintTypeName(TEXT("NotBlueprintType"));
+		const UClass* ParentClass = ObjectClass;
+		while (ParentClass)
+		{
+			// Climb up the class hierarchy and look for "BlueprintType" and "NotBlueprintType" to see if this class is allowed.
+			if (ParentClass->GetBoolMetaData(BlueprintTypeName))
+			{
+				bBlueprintType = true;
+				break;
+			}
+			else if (ParentClass->GetBoolMetaData(NotBlueprintTypeName))
+			{
+				bBlueprintType = false;
+				break;
+			}
+			ParentClass = ParentClass->GetSuperClass();
+		}
+	}
+#endif // WITH_EDITOR
+
+	bool bForbiddenSpawn = false;
+#if WITH_EDITOR
+	static const FName DontUseGenericSpawnObjectName(TEXT("DontUseGenericSpawnObject"));
+	const UClass* ParentClass = ObjectClass;
+	while (ParentClass)
+	{
+		if (ParentClass->HasMetaData(DontUseGenericSpawnObjectName))
+		{
+			bForbiddenSpawn = true;
+			break;
+		}
+		ParentClass = ParentClass->GetSuperClass();
+	}
+#endif // WITH_EDITOR
+
+	return (nullptr != *ObjectClass)
+		&& bBlueprintType
+		&& !bForbiddenSpawn
+		&& !ObjectClass->HasAnyClassFlags(CLASS_Abstract | CLASS_Deprecated | CLASS_NewerVersionExists)
+		&& !ObjectClass->IsChildOf(AActor::StaticClass())
+		&& !ObjectClass->IsChildOf(UActorComponent::StaticClass());
+}
+
+UObject* UGameplayStatics::SpawnObject(TSubclassOf<UObject> ObjectClass, UObject* Outer)
+{
+	if (!CanSpawnObjectOfClass(ObjectClass))
+	{
+#if WITH_EDITOR
+		FMessageLog("PIE").Error(FText::Format(LOCTEXT("SpawnObjectWrongClass", "SpawnObject wrong class: {0}'"), FText::FromString(GetNameSafe(*ObjectClass))));
+#endif // WITH_EDITOR
+		UE_LOG(LogScript, Error, TEXT("UGameplayStatics::SpawnObject wrong class: %s"), *GetNameSafe(*ObjectClass));
+		return nullptr;
+	}
+
+	if (!Outer)
+	{
+		UE_LOG(LogScript, Warning, TEXT("UGameplayStatics::SpawnObject null outer"));
+		return nullptr;
+	}
+
+	return NewObject<UObject>(Outer, ObjectClass, NAME_None, RF_StrongRefOnFrame);
+}
 
 AActor* UGameplayStatics::BeginSpawningActorFromBlueprint(UObject* WorldContextObject, UBlueprint const* Blueprint, const FTransform& SpawnTransform, bool bNoCollisionFail)
 {
@@ -1508,4 +1579,4 @@ int32 UGameplayStatics::GrassOverlappingSphereCount(UObject* WorldContextObject,
 	return Count;
 }
 
-
+#undef LOCTEXT_NAMESPACE
\ No newline at end of file
