Adding a way to explicitly hide fields/nodes from Blueprints via config file.

diff --git a/Engine/Source/Editor/BlueprintGraph/Private/BlueprintActionFilter.cpp b/Engine/Source/Editor/BlueprintGraph/Private/BlueprintActionFilter.cpp
--- a/Engine/Source/Editor/BlueprintGraph/Private/BlueprintActionFilter.cpp
+++ b/Engine/Source/Editor/BlueprintGraph/Private/BlueprintActionFilter.cpp
@@ -30,6 +30,8 @@
 
 namespace BlueprintActionFilterImpl
 {
+	static FString const ConfigSection("BlueprintEditor.Menu");
+
 	/**
 	 * Blueprints have several classes associated with them (the skeleton, for
 	 * UI reflection, and the full generated class). This retrieves the 
@@ -62,6 +64,20 @@ namespace BlueprintActionFilterImpl
 	 * @return Null if the class doesn't implement the function, otherwise the super class that adds it (could be the class passed in).
 	 */
 	static UClass const* FindInheritedInterfaceClass(UClass const* SubClass, TSubclassOf<UInterface> Interface);
+
+	/**
+	 * 
+	 * 
+	 * @return 
+	 */
+	static const TArray<FString>& GetHiddenFieldPaths();
+
+	/**
+	 * 
+	 * 
+	 * @return 
+	 */
+	static const TArray< TSubclassOf<UEdGraphNode> >& GetHiddenNodeTypes();
 	
 	/**
 	 * Utility method to check and see if the specified node-spawner would 
@@ -206,6 +222,15 @@ namespace BlueprintActionFilterImpl
 	 */
 	static bool IsFieldCategoryHidden(FBlueprintActionFilter const& Filter, FBlueprintActionInfo& BlueprintAction);
 
+	/**
+	 * 
+	 * 
+	 * @param  Filter    
+	 * @param  BlueprintAction    
+	 * @return 
+	 */
+	static bool IsActionHiddenByConfig(FBlueprintActionFilter const& Filter, FBlueprintActionInfo& BlueprintAction);
+
 	/**
 	 * Rejection test that checks to see if the node-spawner would produce a 
 	 * node type that isn't white-listed. 
@@ -426,6 +451,51 @@ static UClass const* BlueprintActionFilterImpl::FindInheritedInterfaceClass(UCla
 	return ImplementsInterface;
 }
 
+//------------------------------------------------------------------------------
+static const TArray<FString>& BlueprintActionFilterImpl::GetHiddenFieldPaths()
+{
+	static TArray<FString> HiddenFields;
+
+	static bool bInited = false;
+	if (!bInited)
+	{
+		FString const HiddenFieldsId("BlueprintHiddenFields");
+		GConfig->GetArray(*ConfigSection, *HiddenFieldsId, HiddenFields, GEditorIni);
+
+		bInited = true;
+	}
+
+	return HiddenFields;
+}
+
+//------------------------------------------------------------------------------
+static const TArray< TSubclassOf<UEdGraphNode> >& BlueprintActionFilterImpl::GetHiddenNodeTypes()
+{
+	static TArray< TSubclassOf<UEdGraphNode> > HiddenNodeTypes;
+
+	static bool bInited = false;
+	if (!bInited)
+	{
+		TArray<FString> HiddenClassNames;
+
+		FString const HiddenFieldsId("BlueprintHiddenNodes");
+		GConfig->GetArray(*ConfigSection, *HiddenFieldsId, HiddenClassNames, GEditorIni);
+
+		HiddenNodeTypes.Reserve(HiddenClassNames.Num());
+		for (FString const& ClassName : HiddenClassNames)
+		{
+			if (UClass* FoundClass = FindObject<UClass>(ANY_PACKAGE, *ClassName))
+			{
+				HiddenNodeTypes.Add(FoundClass);
+			}
+		}
+
+		bInited = true;
+	}
+
+	return HiddenNodeTypes;
+}
+
 //------------------------------------------------------------------------------
 static bool BlueprintActionFilterImpl::IsImpure(FBlueprintActionInfo& BlueprintAction)
 {
@@ -796,6 +866,34 @@ static bool BlueprintActionFilterImpl::IsFieldCategoryHidden(FBlueprintActionFil
 	return bIsFilteredOut;
 }
 
+//------------------------------------------------------------------------------
+static bool BlueprintActionFilterImpl::IsActionHiddenByConfig(FBlueprintActionFilter const& Filter, FBlueprintActionInfo& BlueprintAction)
+{
+	bool bIsFilteredOut = false;
+
+	if (const UField* ActionField = BlueprintAction.GetAssociatedMemberField())
+	{
+		const TArray<FString>& HiddenFields = GetHiddenFieldPaths();
+
+		FString const FieldPath = ActionField->GetPathName();
+		if (HiddenFields.Contains(FieldPath))
+		{
+			bIsFilteredOut = true;
+		}
+	}
+
+	if (!bIsFilteredOut)
+	{
+		const TArray< TSubclassOf<UEdGraphNode> >& HiddenNodes = GetHiddenNodeTypes();
+		if (HiddenNodes.Contains(BlueprintAction.GetNodeClass()))
+		{
+			bIsFilteredOut = true;
+		}
+	}
+
+	return bIsFilteredOut;
+}
+
 //------------------------------------------------------------------------------
 static bool BlueprintActionFilterImpl::IsIncompatibleImpureNode(FBlueprintActionFilter const& Filter, FBlueprintActionInfo& BlueprintAction)
 {
@@ -1683,6 +1781,7 @@ FBlueprintActionFilter::FBlueprintActionFilter(uint32 Flags/*= 0x00*/)
 	AddRejectionTest(FRejectionTestDelegate::CreateStatic(IsIncompatibleLatentNode));
 	AddRejectionTest(FRejectionTestDelegate::CreateStatic(IsIncompatibleImpureNode));
 	
+	AddRejectionTest(FRejectionTestDelegate::CreateStatic(IsActionHiddenByConfig));
 	AddRejectionTest(FRejectionTestDelegate::CreateStatic(IsFieldCategoryHidden));
 	if (Flags & BPFILTER_RejectGlobalFields)
 	{
