Selecting a pin type for a Select node will not fail the first time.

diff --git a/Engine/Source/Editor/BlueprintGraph/Private/K2Node_Select.cpp b/Engine/Source/Editor/BlueprintGraph/Private/K2Node_Select.cpp
--- a/Engine/Source/Editor/BlueprintGraph/Private/K2Node_Select.cpp
+++ b/Engine/Source/Editor/BlueprintGraph/Private/K2Node_Select.cpp
@@ -217,6 +217,9 @@ void UK2Node_Select::AllocateDefaultPins()
 	// To refresh, just in case it changed
 	SetEnum(Enum, true);
 
+	// No need to reconstruct the node after force setting the enum, we are at the start of reconstruction already
+	bReconstructNode = false;
+
 	if (Enum)
 	{
 		NumOptionPins = EnumEntries.Num();
