Adding pins to a MakeArray node will give the new pin a default values.

diff --git a/Engine/Source/Editor/BlueprintGraph/Private/K2Node_MakeArray.cpp b/Engine/Source/Editor/BlueprintGraph/Private/K2Node_MakeArray.cpp
--- a/Engine/Source/Editor/BlueprintGraph/Private/K2Node_MakeArray.cpp
+++ b/Engine/Source/Editor/BlueprintGraph/Private/K2Node_MakeArray.cpp
@@ -327,8 +327,9 @@ void UK2Node_MakeArray::AddInputPin()
 
 	++NumInputs;
 	FEdGraphPinType OutputPinType = GetOutputPin()->PinType;
-	CreatePin(EGPD_Input, OutputPinType.PinCategory, OutputPinType.PinSubCategory, OutputPinType.PinSubCategoryObject.Get(), false, false, *FString::Printf(TEXT("[%d]"), (NumInputs-1)));
-	
+	UEdGraphPin* Pin = CreatePin(EGPD_Input, OutputPinType.PinCategory, OutputPinType.PinSubCategory, OutputPinType.PinSubCategoryObject.Get(), false, false, *FString::Printf(TEXT("[%d]"), (NumInputs-1)));
+	GetDefault<UEdGraphSchema_K2>()->SetPinDefaultValueBasedOnType(Pin);
+
 	const bool bIsCompiling = GetBlueprint()->bBeingCompiled;
 	if( !bIsCompiling )
 	{
