Game modules can now extend Blueprint menu filtering, and add in their own menu rejection tests.

diff --git a/Engine/Source/Editor/BlueprintGraph/Private/BlueprintActionFilter.cpp b/Engine/Source/Editor/BlueprintGraph/Private/BlueprintActionFilter.cpp
--- a/Engine/Source/Editor/BlueprintGraph/Private/BlueprintActionFilter.cpp
+++ b/Engine/Source/Editor/BlueprintGraph/Private/BlueprintActionFilter.cpp
@@ -8,6 +8,7 @@
 #include "BlueprintEventNodeSpawner.h"
 #include "BlueprintBoundEventNodeSpawner.h"
 #include "BlueprintBoundNodeSpawner.h"
+#include "BlueprintGraphModule.h"	// for GetExtendedActionMenuFilters
 #include "EdGraphSchema_K2.h"		// for FBlueprintMetadata
 #include "BlueprintEditorUtils.h"	// for FindBlueprintForGraph()
 #include "ObjectEditorUtils.h"		// for IsFunctionHiddenFromClass()/IsVariableCategoryHiddenFromClass()
@@ -1893,5 +1894,19 @@ bool FBlueprintActionFilter::IsFilteredByThis(FBlueprintActionInfo& BlueprintAct
 		}
 	}
 
+	if (!bIsFiltered)
+	{
+		FBlueprintGraphModule& BluprintGraphModule = FModuleManager::LoadModuleChecked<FBlueprintGraphModule>("BlueprintGraph");
+
+		for (auto& ExtraRejectionTest : BluprintGraphModule.GetExtendedActionMenuFilters())
+		{
+			if (ExtraRejectionTest.Execute(FilterRef, BlueprintAction))
+			{
+				bIsFiltered = true;
+				break;
+			}
+		}
+	}
+
 	return bIsFiltered;
 }
diff --git a/Engine/Source/Editor/BlueprintGraph/Public/BlueprintActionFilter.h b/Engine/Source/Editor/BlueprintGraph/Public/BlueprintActionFilter.h
--- a/Engine/Source/Editor/BlueprintGraph/Public/BlueprintActionFilter.h
+++ b/Engine/Source/Editor/BlueprintGraph/Public/BlueprintActionFilter.h
@@ -2,7 +2,8 @@
 
 #pragma once
 
-#include "BlueprintNodeBinder.h" // for IBlueprintNodeBinder::FBindingSet
+#include "BlueprintNodeBinder.h"  // for IBlueprintNodeBinder::FBindingSet
+#include "BlueprintGraphModule.h" // for FActionMenuRejectionTest
 
 // Forward declarations
 class UBlueprintNodeSpawner;
@@ -150,7 +151,7 @@ class BLUEPRINTGRAPH_API FBlueprintActionFilter
 {
 public:
 	/** The filter uses a series of rejection tests matching */
-	DECLARE_DELEGATE_RetVal_TwoParams(bool, FRejectionTestDelegate, FBlueprintActionFilter const&, FBlueprintActionInfo&);
+	typedef FBlueprintGraphModule::FActionMenuRejectionTest FRejectionTestDelegate;
 
 public:
 	enum EFlags // Flags, which configure certain rejection tests.
diff --git a/Engine/Source/Editor/BlueprintGraph/Public/BlueprintGraphModule.h b/Engine/Source/Editor/BlueprintGraph/Public/BlueprintGraphModule.h
--- a/Engine/Source/Editor/BlueprintGraph/Public/BlueprintGraphModule.h
+++ b/Engine/Source/Editor/BlueprintGraph/Public/BlueprintGraphModule.h
@@ -4,8 +4,27 @@
 
 #include "ModuleInterface.h"
 
+// Forward declarations
+class  FBlueprintActionFilter;
+struct FBlueprintActionInfo;
+
 class FBlueprintGraphModule : public IModuleInterface
 {
 public:
+	DECLARE_DELEGATE_RetVal_TwoParams(bool, FActionMenuRejectionTest, FBlueprintActionFilter const&, FBlueprintActionInfo&);
+	/**
+	 * Provides access to a list of additional rejection tests that will be  
+	 * utilized when constructing Blueprint context/action menus. Helpful to 
+	 * game modules, so they can inject their own custom filtering (and hide 
+	 * options that may be undesired for their specific project).
+	 *
+	 * These rejection tests will be ran after all the standard filter tests, 
+	 * and in the order that they are in this array.
+	 * 
+	 * @return An array of additional rejection tests for FBlueprintActionFilter.
+	 */
+	BLUEPRINTGRAPH_API TArray<FActionMenuRejectionTest>& GetExtendedActionMenuFilters() { return ExtendedMenuFilters; }
 
+private:
+	TArray<FActionMenuRejectionTest> ExtendedMenuFilters;
 };
\ No newline at end of file
