Math Expression nodes will no longer always appear as changed when doing a diff on a Blueprint.

The Math Expression graph will be completely missing from the diff display, it is regenerated each time and the important detail is the math expression the user inputs. The node's math expression will show up if there are differences between two Blueprints.

diff --git a/Engine/Source/Editor/BlueprintGraph/Classes/K2Node_MathExpression.h b/Engine/Source/Editor/BlueprintGraph/Classes/K2Node_MathExpression.h
--- a/Engine/Source/Editor/BlueprintGraph/Classes/K2Node_MathExpression.h
+++ b/Engine/Source/Editor/BlueprintGraph/Classes/K2Node_MathExpression.h
@@ -35,6 +35,7 @@ public:
 	virtual FText GetNodeTitle(ENodeTitleType::Type TitleType) const override;
 	virtual void PostPlacedNewNode() override;
 	virtual void ReconstructNode() override;
+	virtual void FindDiffs(class UEdGraphNode* OtherNode, struct FDiffResults& Results )  override;
 	// End of UEdGraphNode interface
 
 	// UK2Node interface
diff --git a/Engine/Source/Editor/BlueprintGraph/Private/K2Node_MathExpression.cpp b/Engine/Source/Editor/BlueprintGraph/Private/K2Node_MathExpression.cpp
--- a/Engine/Source/Editor/BlueprintGraph/Private/K2Node_MathExpression.cpp
+++ b/Engine/Source/Editor/BlueprintGraph/Private/K2Node_MathExpression.cpp
@@ -2608,6 +2608,11 @@ FText UK2Node_MathExpression::GetNodeTitle(ENodeTitleType::Type TitleType) const
 			{
 				CachedDisplayExpression.SetCachedText(FText::FromString(SanitizeDisplayExpression(ExpressionRoot->ToDisplayString(GetBlueprint()))), this);
 			}
+			else
+			{
+				// Fallback and display the expression in it's raw form
+				CachedDisplayExpression.SetCachedText(FText::FromString(Expression), this);
+			}
 		}
 		return CachedDisplayExpression;
 	}
@@ -2658,4 +2663,32 @@ FText UK2Node_MathExpression::GetFullTitle(FText InExpression) const
 	return FText::Format(LOCTEXT("MathExpressionSecondTitleLine", "{0}\nMath Expression"), InExpression);
 }
 
+void UK2Node_MathExpression::FindDiffs(class UEdGraphNode* OtherNode, struct FDiffResults& Results )
+{
+	UK2Node_MathExpression* MathExpression1 = this;
+	UK2Node_MathExpression* MathExpression2 = Cast<UK2Node_MathExpression>(OtherNode);
+
+	// Compare the visual display of a math expression (the visual display involves consolidating variable Guid's into readable parameters)
+	FText Expression1 = MathExpression1->GetNodeTitle(ENodeTitleType::EditableTitle);
+	FText Expression2 = MathExpression2->GetNodeTitle(ENodeTitleType::EditableTitle);
+	if (Expression1.CompareTo(Expression2) != 0)
+	{
+		FDiffSingleResult Diff;
+		Diff.Node1 = MathExpression2;
+		Diff.Node2 = MathExpression1;
+
+		Diff.Diff = EDiffType::NODE_PROPERTY;
+		FText NodeName = GetNodeTitle(ENodeTitleType::ListView);
+
+		FFormatNamedArguments Args;
+		Args.Add(TEXT("Expression1"), Expression1);
+		Args.Add(TEXT("Expression2"), Expression2);
+
+		Diff.ToolTip =  FText::Format(LOCTEXT("DIF_MathExpressionToolTip", "Math Expression '{Expression1}' changed to '{Expression2}'"), Args);
+		Diff.DisplayColor = FLinearColor(0.85f,0.71f,0.25f);
+		Diff.DisplayString = FText::Format(LOCTEXT("DIF_MathExpression", "Math Expression '{Expression1}' changed to '{Expression2}'"), Args);
+		Results.Add(Diff);
+	}
+}
+
 #undef LOCTEXT_NAMESPACE
diff --git a/Engine/Source/Editor/Kismet/Private/SBlueprintDiff.cpp b/Engine/Source/Editor/Kismet/Private/SBlueprintDiff.cpp
--- a/Engine/Source/Editor/Kismet/Private/SBlueprintDiff.cpp
+++ b/Engine/Source/Editor/Kismet/Private/SBlueprintDiff.cpp
@@ -668,7 +668,11 @@ void SBlueprintDiff::Construct( const FArguments& InArgs)
 				break;
 			}
 		}
-		CreateGraphEntry(GraphOld,GraphNew);
+		// Do not worry about graphs that are contained in MathExpression nodes, they are recreated each compile
+		if (IsGraphDiffNeeded(GraphOld))
+		{
+			CreateGraphEntry(GraphOld,GraphNew);
+		}
 	}
 
 	//Add graphs that only exist in 2nd(new) blueprint
@@ -677,7 +681,11 @@ void SBlueprintDiff::Construct( const FArguments& InArgs)
 		UEdGraph* GraphNew = *It2;
 		if(GraphNew != NULL)
 		{
-			CreateGraphEntry(NULL,GraphNew);
+			// Do not worry about graphs that are contained in MathExpression nodes, they are recreated each compile
+			if (IsGraphDiffNeeded(GraphNew))
+			{
+				CreateGraphEntry(NULL,GraphNew);
+			}
 		}
 	}
 
@@ -1389,5 +1397,11 @@ void SBlueprintDiff::SetCurrentMode(FName NewMode)
 	}
 }
 
+bool SBlueprintDiff::IsGraphDiffNeeded(class UEdGraph* InGraph) const
+{
+	// Do not worry about graphs that are contained in MathExpression nodes, they are recreated each compile
+	return !InGraph->GetOuter()->IsA<UK2Node_MathExpression>();
+}
+
 #undef LOCTEXT_NAMESPACE
 
diff --git a/Engine/Source/Editor/Kismet/Public/SBlueprintDiff.h b/Engine/Source/Editor/Kismet/Public/SBlueprintDiff.h
--- a/Engine/Source/Editor/Kismet/Public/SBlueprintDiff.h
+++ b/Engine/Source/Editor/Kismet/Public/SBlueprintDiff.h
@@ -175,6 +175,9 @@ protected:
 	/** Event handler that updates the graph view when user selects a new graph */
 	void HandleGraphChanged( const FName GraphName );
 
+	/** Checks if a graph is valid for diff'ing */
+	bool IsGraphDiffNeeded(class UEdGraph* InGraph) const;
+
 	struct FDiffControl
 	{
 		FDiffControl()
