Changing the Class of a "Spawn Actor from Class" node will no longer disconnect the Result pin's links if they are still valid.

diff --git a/Engine/Source/Editor/BlueprintGraph/Private/K2Node_SpawnActorFromClass.cpp b/Engine/Source/Editor/BlueprintGraph/Private/K2Node_SpawnActorFromClass.cpp
--- a/Engine/Source/Editor/BlueprintGraph/Private/K2Node_SpawnActorFromClass.cpp
+++ b/Engine/Source/Editor/BlueprintGraph/Private/K2Node_SpawnActorFromClass.cpp
@@ -191,10 +191,6 @@ void UK2Node_SpawnActorFromClass::OnClassPinChanged()
 {
  	const UEdGraphSchema_K2* K2Schema = GetDefault<UEdGraphSchema_K2>();
 
-	// Because the archetype has changed, we break the output link as the output pin type will change
-	UEdGraphPin* ResultPin = GetResultPin();
-	ResultPin->BreakAllPinLinks();
-
 	// Remove all pins related to archetype variables
 	TArray<UEdGraphPin*> OldPins = Pins;
 	TArray<UEdGraphPin*> OldClassPins;
@@ -217,6 +213,19 @@ void UK2Node_SpawnActorFromClass::OnClassPinChanged()
 	{
 		CreatePinsForClass(UseSpawnClass, NewClassPins);
 	}
+
+	UEdGraphPin* ResultPin = GetResultPin();
+	// Cache all the pin connections to the ResultPin, we will attempt to recreate them
+	TArray<UEdGraphPin*> ResultPinConnectionList = ResultPin->LinkedTo;
+	// Because the archetype has changed, we break the output link as the output pin type will change
+	ResultPin->BreakAllPinLinks();
+
+	// Recreate any pin links to the Result pin that are still valid
+	for (UEdGraphPin* Connections : ResultPinConnectionList)
+	{
+		K2Schema->TryCreateConnection(ResultPin, Connections);
+	}
+
 	K2Schema->ConstructBasicPinTooltip(*ResultPin, LOCTEXT("ResultPinDescription", "The spawned Actor"), ResultPin->PinToolTip);
 
 	// Rewire the old pins to the new pins so connections are maintained if possible
