Fixed non unity build.

diff --git a/Engine/Source/Editor/BlueprintGraph/Private/K2Node_MathExpression.cpp b/Engine/Source/Editor/BlueprintGraph/Private/K2Node_MathExpression.cpp
--- a/Engine/Source/Editor/BlueprintGraph/Private/K2Node_MathExpression.cpp
+++ b/Engine/Source/Editor/BlueprintGraph/Private/K2Node_MathExpression.cpp
@@ -9,6 +9,7 @@
 #include "UnrealMathUtility.h"
 #include "BlueprintEditorUtils.h"
 #include "BlueprintActionDatabaseRegistrar.h"
+#include "DiffResults.h"
 
 #define LOCTEXT_NAMESPACE "K2Node"
 
