Fixed GLEO when a Blueprint references a macro that references a function library that is not referenced by the original Blueprint.

diff --git a/Engine/Source/Editor/BlueprintGraph/Private/K2Node_MacroInstance.cpp b/Engine/Source/Editor/BlueprintGraph/Private/K2Node_MacroInstance.cpp
--- a/Engine/Source/Editor/BlueprintGraph/Private/K2Node_MacroInstance.cpp
+++ b/Engine/Source/Editor/BlueprintGraph/Private/K2Node_MacroInstance.cpp
@@ -471,6 +471,15 @@ bool UK2Node_MacroInstance::HasExternalDependencies(TArray<class UStruct*>* Opti
 		{
 			OptionalOutput->AddUnique(OtherClass);
 		}
+
+		for (UEdGraphNode* Node : GetMacroGraph()->Nodes)
+		{
+			UK2Node* K2Node = Cast<UK2Node>(Node);
+			if (K2Node)
+			{
+				K2Node->HasExternalDependencies(OptionalOutput);
+			}
+		}
 	}
 	const bool bSuperResult = Super::HasExternalDependencies(OptionalOutput);
 	return bSuperResult || bResult;
