[UE-16996] Properly removing pin watches for params that are being deleted from a custom event node.

diff --git a/Engine/Source/Editor/BlueprintGraph/Private/K2Node_EditablePinBase.cpp b/Engine/Source/Editor/BlueprintGraph/Private/K2Node_EditablePinBase.cpp
--- a/Engine/Source/Editor/BlueprintGraph/Private/K2Node_EditablePinBase.cpp
+++ b/Engine/Source/Editor/BlueprintGraph/Private/K2Node_EditablePinBase.cpp
@@ -1,6 +1,7 @@
 // Copyright 1998-2015 Epic Games, Inc. All Rights Reserved.
 
 #include "BlueprintGraphPrivatePCH.h"
+#include "KismetDebugUtilities.h"
 
 UK2Node_EditablePinBase::UK2Node_EditablePinBase(const FObjectInitializer& ObjectInitializer)
 	: Super(ObjectInitializer)
@@ -48,6 +49,11 @@ void UK2Node_EditablePinBase::RemoveUserDefinedPin(TSharedPtr<FUserPinInfo> PinT
 			Pin->BreakAllPinLinks();
 			Pins.Remove(Pin);
 			Pin->MarkPendingKill();
+
+			if (UBlueprint* Blueprint = GetBlueprint())
+			{
+				FKismetDebugUtilities::RemovePinWatch(Blueprint, Pin);
+			}
 		}
 	}
 
diff --git a/Engine/Source/Editor/GraphEditor/Private/KismetNodes/KismetNodeInfoContext.cpp b/Engine/Source/Editor/GraphEditor/Private/KismetNodes/KismetNodeInfoContext.cpp
--- a/Engine/Source/Editor/GraphEditor/Private/KismetNodes/KismetNodeInfoContext.cpp
+++ b/Engine/Source/Editor/GraphEditor/Private/KismetNodes/KismetNodeInfoContext.cpp
@@ -86,6 +86,10 @@ FKismetNodeInfoContext::FKismetNodeInfoContext(UEdGraph* SourceGraph)
 		for (auto WatchedPinIt = SourceBlueprint->PinWatches.CreateConstIterator(); WatchedPinIt; ++WatchedPinIt)
 		{
 			UEdGraphPin* WatchedPin = *WatchedPinIt;
+			if (!ensure(WatchedPin))
+			{
+				continue;
+			}
 
 			UEdGraphNode* OwningNode = Cast<UEdGraphNode>(WatchedPin->GetOuter());
 			if (!ensure(OwningNode != NULL)) // shouldn't happen, but just in case a dead pin was added to the PinWatches array
