[UE-16972] Could find no user action that aws improved (or even effected by) this logic that disconnects pins. This case fails because one pin is of type self and one is a normal object reference, but this function did not provide a class context. It's better to just not run the logic at all, rather than destroy the pin.

diff --git a/Engine/Source/Editor/BlueprintGraph/Private/K2Node_CallArrayFunction.cpp b/Engine/Source/Editor/BlueprintGraph/Private/K2Node_CallArrayFunction.cpp
--- a/Engine/Source/Editor/BlueprintGraph/Private/K2Node_CallArrayFunction.cpp
+++ b/Engine/Source/Editor/BlueprintGraph/Private/K2Node_CallArrayFunction.cpp
@@ -244,16 +244,6 @@ void UK2Node_CallArrayFunction::PropagateArrayTypeInfo(const UEdGraphPin* Source
 				{
 					CurrentPin->ResetDefaultValue();
 				}
-
-				// Verify that all previous connections to this pin are still valid with the new type
-				for (TArray<UEdGraphPin*>::TIterator ConnectionIt(CurrentPin->LinkedTo); ConnectionIt; ++ConnectionIt)
-				{
-					UEdGraphPin* ConnectedPin = *ConnectionIt;
-					if (!Schema->ArePinsCompatible(CurrentPin, ConnectedPin))
-					{
-						CurrentPin->BreakLinkTo(ConnectedPin);
-					}
-				}
 			}
 		}
 	}
