Fixed a hard to reproduce Blueprint Editor crash

diff --git a/Engine/Source/Editor/BlueprintGraph/Private/EdGraphSchema_K2.cpp b/Engine/Source/Editor/BlueprintGraph/Private/EdGraphSchema_K2.cpp
--- a/Engine/Source/Editor/BlueprintGraph/Private/EdGraphSchema_K2.cpp
+++ b/Engine/Source/Editor/BlueprintGraph/Private/EdGraphSchema_K2.cpp
@@ -2658,8 +2658,8 @@ bool UEdGraphSchema_K2::DoesSupportPinWatching() const
 bool UEdGraphSchema_K2::IsPinBeingWatched(UEdGraphPin const* Pin) const
 {
 	// Note: If you crash here; it is likely that you forgot to call Blueprint->OnBlueprintChanged.Broadcast(Blueprint) to invalidate the cached UI state
-	UBlueprint* Blueprint = FBlueprintEditorUtils::FindBlueprintForNodeChecked(Pin->GetOwningNode());
-	return FKismetDebugUtilities::IsPinBeingWatched(Blueprint, Pin);
+	UBlueprint* Blueprint = FBlueprintEditorUtils::FindBlueprintForNode(Pin ? Pin->GetOwningNodeUnchecked() : nullptr);
+	return ensure(Blueprint) ? FKismetDebugUtilities::IsPinBeingWatched(Blueprint, Pin) : false;
 }
 
 void UEdGraphSchema_K2::ClearPinWatch(UEdGraphPin const* Pin) const
