ME Optimization:
- FBPTerminal::InlineGeneratedParameter added - the result of the function is directly used as input parameter, it's not stroed in a lical variable
- FKismetFunctionContext::InlineGeneratedParameter
- MathExpressionHandler - compile (to out intermediate representation) a subgraph using "reverse polish notation" approach
- Math Expression node is not longer expanded
- VMBackend supports InlineGeneratedParameter and InlineGeneratedParameter

diff --git a/Engine/Source/Editor/BlueprintGraph/Classes/K2Node_MathExpression.h b/Engine/Source/Editor/BlueprintGraph/Classes/K2Node_MathExpression.h
--- a/Engine/Source/Editor/BlueprintGraph/Classes/K2Node_MathExpression.h
+++ b/Engine/Source/Editor/BlueprintGraph/Classes/K2Node_MathExpression.h
@@ -36,11 +36,14 @@ public:
 	virtual void PostPlacedNewNode() override;
 	virtual void ReconstructNode() override;
 	virtual void FindDiffs(class UEdGraphNode* OtherNode, struct FDiffResults& Results )  override;
+	virtual bool ShouldMergeChildGraphs() const override { return false; }
 	// End of UEdGraphNode interface
 
 	// UK2Node interface
 	virtual void ValidateNodeDuringCompilation(class FCompilerResultsLog& MessageLog) const override;
 	virtual void GetMenuActions(FBlueprintActionDatabaseRegistrar& ActionRegistrar) const override;
+	virtual class FNodeHandlingFunctor* CreateNodeHandler(class FKismetCompilerContext& CompilerContext) const override;
+	virtual bool IsNodePure() const { return true; }
 	// End of UK2Node interface
 
 	// Begin UK2Node_EditablePinBase interface
diff --git a/Engine/Source/Editor/BlueprintGraph/Private/K2Node_MathExpression.cpp b/Engine/Source/Editor/BlueprintGraph/Private/K2Node_MathExpression.cpp
--- a/Engine/Source/Editor/BlueprintGraph/Private/K2Node_MathExpression.cpp
+++ b/Engine/Source/Editor/BlueprintGraph/Private/K2Node_MathExpression.cpp
@@ -10,6 +10,7 @@
 #include "BlueprintEditorUtils.h"
 #include "BlueprintActionDatabaseRegistrar.h"
 #include "DiffResults.h"
+#include "MathExpressionHandler.h"
 
 #define LOCTEXT_NAMESPACE "K2Node"
 
@@ -2434,6 +2435,12 @@ void UK2Node_MathExpression::GetMenuActions(FBlueprintActionDatabaseRegistrar& A
 	}
 }
 
+//------------------------------------------------------------------------------
+FNodeHandlingFunctor* UK2Node_MathExpression::CreateNodeHandler(class FKismetCompilerContext& CompilerContext) const
+{
+	return new FKCHandler_MathExpression(CompilerContext);
+}
+
 //------------------------------------------------------------------------------
 TSharedPtr<class INameValidatorInterface> UK2Node_MathExpression::MakeNameValidator() const
 {
diff --git a/Engine/Source/Editor/BlueprintGraph/Private/MathExpressionHandler.cpp b/Engine/Source/Editor/BlueprintGraph/Private/MathExpressionHandler.cpp
new file mode 100644
--- /dev/null
+++ b/Engine/Source/Editor/BlueprintGraph/Private/MathExpressionHandler.cpp
@@ -0,0 +1,235 @@
+// Copyright 1998-2015 Epic Games, Inc. All Rights Reserved.
+
+#include "BlueprintGraphPrivatePCH.h"
+#include "K2Node_MathExpression.h"
+#include "EdGraphUtilities.h"
+#include "KismetCompiler.h"
+#include "MathExpressionHandler.h"
+
+#define LOCTEXT_NAMESPACE "KCHandler_MathExpression"
+
+struct KCHandler_MathExpressionHelper
+{
+	static bool CreateMap(UEdGraphNode& KeyNode, UEdGraphNode& ValueNode, EEdGraphPinDirection KeyPinDirection, UEdGraphSchema_K2* Schema, TMap<UEdGraphPin*, UEdGraphPin*>& OutMap)
+	{
+		check(Schema);
+		for (auto InnerInputPin : KeyNode.Pins)
+		{
+			if (!InnerInputPin || InnerInputPin->Direction != KeyPinDirection)
+			{
+				return false;
+			}
+
+			auto OuterInputPin = ValueNode.FindPin(InnerInputPin->PinName);
+			if (!OuterInputPin
+				|| !Schema->ArePinsCompatible(InnerInputPin, OuterInputPin, nullptr, false))
+			{
+				return false;
+			}
+			OutMap.Add(InnerInputPin, OuterInputPin);
+		}
+		return true;
+	}
+};
+
+bool FKCHandler_MathExpression::CanBeCalledByMathExpression(const UFunction* Function)
+{
+	bool bGoodFunction = (nullptr != Function)
+		&& Function->HasAllFunctionFlags(FUNC_Static | FUNC_BlueprintPure | FUNC_Final | FUNC_Native)
+		&& !Function->HasAnyFunctionFlags(FUNC_BlueprintAuthorityOnly | FUNC_BlueprintCosmetic)
+		&& !Function->GetOuterUClass()->IsChildOf(UInterface::StaticClass());
+	if (!bGoodFunction)
+	{
+		return false;
+	}
+
+	for (TFieldIterator<const UProperty> It(Function); It && (It->PropertyFlags & CPF_Parm); ++It)
+	{
+		const UProperty* Property = *It;
+		const bool bGoodProperty = Property && (!Property->HasAnyPropertyFlags(CPF_OutParm) || Property->HasAllPropertyFlags(CPF_ReturnParm));
+		if (!bGoodProperty)
+		{
+			return false;
+		}
+	}
+	return true;
+}
+
+FBlueprintCompiledStatement* FKCHandler_MathExpression::GenerateFunctionRPN(UEdGraphNode* CurrentNode, FKismetFunctionContext& Context, UK2Node_MathExpression& MENode, FBPTerminal* ResultTerm, TMap<UEdGraphPin*, UEdGraphPin*>& InnerToOuterInput)
+{
+	UK2Node_CallFunction* CallFunctionNode = Cast<UK2Node_CallFunction>(CurrentNode);
+	UFunction* Function = CallFunctionNode ? CallFunctionNode->GetTargetFunction() : nullptr;
+	if (!CanBeCalledByMathExpression(Function))
+	{
+		CompilerContext.MessageLog.Error(*FString::Printf(*LOCTEXT("WrongFunction_Error", "Function '%s' cannot be called inside Math Expression @@ - @@").ToString(), *GetNameSafe(Function)), CallFunctionNode, &MENode);
+		return nullptr;
+	}
+
+	FBlueprintCompiledStatement* NewDetachedStatement = new FBlueprintCompiledStatement();
+	FBlueprintCompiledStatement& Statement = *NewDetachedStatement;
+	Statement.FunctionToCall = Function;
+	Statement.FunctionContext = nullptr;
+	Statement.Type = KCST_CallFunction;
+	Statement.LHS = ResultTerm; // required only for the first node
+
+	TArray<FBPTerminal*> RHSTerms;
+	for (TFieldIterator<UProperty> It(Function); It && (It->PropertyFlags & CPF_Parm); ++It)
+	{
+		UProperty* Property = *It;
+		if (Property && !Property->HasAnyPropertyFlags(CPF_ReturnParm | CPF_OutParm))
+		{
+			UEdGraphPin* PinToTry = nullptr;
+			{
+				UEdGraphPin* PinMatch = CallFunctionNode->FindPin(Property->GetName());
+				const bool bGoodPin = PinMatch && FKismetCompilerUtilities::IsTypeCompatibleWithProperty(PinMatch, Property, CompilerContext.MessageLog, CompilerContext.GetSchema(), Context.NewClass);
+				PinToTry = bGoodPin ? FEdGraphUtilities::GetNetFromPin(PinMatch) : nullptr;
+			}
+
+			FBPTerminal* RHSTerm = nullptr;
+			{
+				UEdGraphPin** OuterInputPtr = PinToTry ? InnerToOuterInput.Find(PinToTry) : nullptr;
+				UEdGraphPin* OuterInputNet = (OuterInputPtr && *OuterInputPtr) ? FEdGraphUtilities::GetNetFromPin(*OuterInputPtr) : nullptr;
+				FBPTerminal** OuterTerm = OuterInputNet ? Context.NetMap.Find(OuterInputNet) : nullptr;
+				// Input is an outer term
+				if (OuterTerm && *OuterTerm)
+				{
+					RHSTerm = *OuterTerm;
+				}
+			}
+
+			if (!RHSTerm)
+			{
+				FBPTerminal** Term = PinToTry ? Context.NetMap.Find(PinToTry) : nullptr;
+				// Input is a literal term
+				if (Term && *Term && (*Term)->bIsLiteral)
+				{
+					RHSTerm = *Term;
+				}
+				// Input is an InlineGeneratedParameter
+				else if (Term && *Term)
+				{
+					ensure(!(*Term)->InlineGeneratedParameter);
+					UEdGraphNode* SourceNode = PinToTry ? PinToTry->GetOwningNodeUnchecked() : nullptr;
+					FBlueprintCompiledStatement* InlineGeneratedParameterStatement = GenerateFunctionRPN(SourceNode, Context, MENode, nullptr, InnerToOuterInput);
+					if (InlineGeneratedParameterStatement)
+					{
+						Context.AllGeneratedStatements.Add(InlineGeneratedParameterStatement);
+						RHSTerm = *Term;
+						RHSTerm->InlineGeneratedParameter = InlineGeneratedParameterStatement;
+					}
+				}
+			}
+
+			if (RHSTerm)
+			{
+				RHSTerms.Add(RHSTerm);
+			}
+			else
+			{
+				CompilerContext.MessageLog.Error(*FString::Printf(*LOCTEXT("FindPinParameter_Error", "Could not find a pin for the parameter %s of %s on @@").ToString(), *GetNameSafe(Property), *GetNameSafe(Function)), CallFunctionNode);
+			}
+		}
+	}
+	Statement.RHS = RHSTerms;
+	return &Statement;
+}
+
+void FKCHandler_MathExpression::RegisterNets(FKismetFunctionContext& Context, UEdGraphNode* InNode)
+{
+	FNodeHandlingFunctor::RegisterNets(Context, InNode);
+
+	auto Node_MathExpression = CastChecked<UK2Node_MathExpression>(InNode);
+	auto InnerEntryNode = Node_MathExpression->GetEntryNode();
+	auto InnerExitNode = Node_MathExpression->GetExitNode();
+	check(Node_MathExpression->BoundGraph && InnerEntryNode && InnerExitNode);
+
+	for (auto InnerGraph : Node_MathExpression->BoundGraph->Nodes)
+	{
+		if (!InnerGraph || (InnerGraph == InnerEntryNode) || (InnerGraph == InnerExitNode))
+		{
+			continue;
+		}
+
+		for (auto Pin : InnerGraph->Pins)
+		{
+			// Register fake terms for InlineGeneratedValues
+			if (Pin && (Pin->Direction == EEdGraphPinDirection::EGPD_Output) && Pin->LinkedTo.Num())
+			{
+				auto Linked = Pin->LinkedTo[0];
+				auto LinkedOwnerNode = Linked ? Linked->GetOwningNodeUnchecked() : nullptr;
+				if (LinkedOwnerNode && (InnerExitNode != LinkedOwnerNode))
+				{
+					FBPTerminal* Term = new (Context.InlineGeneratedValues) FBPTerminal();
+					Term->CopyFromPin(Pin, Context.NetNameMap->MakeValidName(Pin));
+					Context.NetMap.Add(Pin, Term);
+				}
+			}
+
+			// Register Literals
+			if (Pin && (Pin->Direction == EEdGraphPinDirection::EGPD_Input) && !Pin->LinkedTo.Num())
+			{
+				RegisterLiteral(Context, Pin);
+			}
+		}
+	}
+}
+
+void FKCHandler_MathExpression::RegisterNet(FKismetFunctionContext& Context, UEdGraphPin* Net) 
+{
+	FBPTerminal* Term = Context.CreateLocalTerminalFromPinAutoChooseScope(Net, Context.NetNameMap->MakeValidName(Net));
+	Context.NetMap.Add(Net, Term);
+}
+
+void FKCHandler_MathExpression::Compile(FKismetFunctionContext& Context, UEdGraphNode* Node)
+{
+	UK2Node_MathExpression* Node_MathExpression = CastChecked<UK2Node_MathExpression>(Node);
+	check(Context.Schema);
+
+	auto InnerExitNode = Node_MathExpression->GetExitNode();
+	auto InnerEntryNode = Node_MathExpression->GetEntryNode();
+
+	if (!InnerExitNode || !InnerEntryNode || (InnerExitNode->Pins.Num() != 1) || ((InnerExitNode->Pins.Num() + InnerEntryNode->Pins.Num()) != Node->Pins.Num()))
+	{
+		Context.MessageLog.Warning(*LOCTEXT("Compile_PinError", "ICE - wrong inner pins - @@").ToString(), Node);
+		return;
+	}
+
+	TMap<UEdGraphPin*, UEdGraphPin*> InnerToOuterInput;
+	bool bProperMap = KCHandler_MathExpressionHelper::CreateMap(*InnerEntryNode, *Node_MathExpression, EEdGraphPinDirection::EGPD_Output, Context.Schema, InnerToOuterInput);
+
+	TMap<UEdGraphPin*, UEdGraphPin*> InnerToOuterOutput;
+	bProperMap &= KCHandler_MathExpressionHelper::CreateMap(*InnerExitNode, *Node_MathExpression, EEdGraphPinDirection::EGPD_Input, Context.Schema, InnerToOuterOutput);
+
+	if (!bProperMap)
+	{
+		Context.MessageLog.Warning(*LOCTEXT("Compile_WrongMap", "ICE - cannot map pins - @@").ToString(), Node);
+		return;
+	}
+
+	auto InnerOutputPin = InnerExitNode->Pins[0];
+	auto OuterOutputPinPtr = InnerToOuterOutput.Find(InnerOutputPin);
+	if (!InnerOutputPin || (InnerOutputPin->LinkedTo.Num() != 1) || !InnerOutputPin->LinkedTo[0] || !OuterOutputPinPtr || !*OuterOutputPinPtr)
+	{
+		Context.MessageLog.Warning(*LOCTEXT("Compile_WrongOutputLink", "ICE - wrong output link - @@").ToString(), Node);
+		return;
+	}
+
+	FBPTerminal** OutputTermPtr = Context.NetMap.Find(FEdGraphUtilities::GetNetFromPin(*OuterOutputPinPtr));
+	FBPTerminal* OutputTerm = OutputTermPtr ? *OutputTermPtr : nullptr;
+	auto LastInnerNode = InnerOutputPin->LinkedTo[0]->GetOwningNodeUnchecked();
+
+	FBlueprintCompiledStatement* DetachedStatement = GenerateFunctionRPN(LastInnerNode, Context, *Node_MathExpression, OutputTerm, InnerToOuterInput);
+	if (DetachedStatement)
+	{
+		Context.AllGeneratedStatements.Add(DetachedStatement);
+		TArray<FBlueprintCompiledStatement*>& StatementList = Context.StatementsPerNode.FindOrAdd(Node);
+		StatementList.Add(DetachedStatement);
+	}
+	else
+	{
+		Context.MessageLog.Warning(*LOCTEXT("Compile_CannotGenerateFunction", "ICE - cannot generate function - @@").ToString(), Node);
+		return;
+	}
+}
+
+#undef LOCTEXT_NAMESPACE
diff --git a/Engine/Source/Editor/BlueprintGraph/Public/MathExpressionHandler.h b/Engine/Source/Editor/BlueprintGraph/Public/MathExpressionHandler.h
new file mode 100644
--- /dev/null
+++ b/Engine/Source/Editor/BlueprintGraph/Public/MathExpressionHandler.h
@@ -0,0 +1,29 @@
+// Copyright 1998-2015 Epic Games, Inc. All Rights Reserved.
+
+#pragma once
+
+#include "KismetCompiler.h"
+
+struct FBlueprintCompiledStatement;
+class UK2Node_MathExpression;
+struct FKismetFunctionContext;
+
+class FKCHandler_MathExpression : public FNodeHandlingFunctor
+{
+	FBlueprintCompiledStatement* GenerateFunctionRPN(UEdGraphNode* CurrentNode
+		, FKismetFunctionContext& Context
+		, UK2Node_MathExpression& MENode
+		, FBPTerminal* ResultTerm
+		, TMap<UEdGraphPin*, UEdGraphPin*>& InnerToOuterInput);
+
+public:
+	FKCHandler_MathExpression(FKismetCompilerContext& InCompilerContext)
+		: FNodeHandlingFunctor(InCompilerContext)
+	{}
+
+	virtual void RegisterNets(FKismetFunctionContext& Context, UEdGraphNode* InNode) override;
+	virtual void RegisterNet(FKismetFunctionContext& Context, UEdGraphPin* Net) override;
+	virtual void Compile(FKismetFunctionContext& Context, UEdGraphNode* Node) override;
+
+	static bool CanBeCalledByMathExpression(const UFunction* Function);
+};
diff --git a/Engine/Source/Editor/KismetCompiler/Private/KismetCompilerVMBackend.cpp b/Engine/Source/Editor/KismetCompiler/Private/KismetCompilerVMBackend.cpp
--- a/Engine/Source/Editor/KismetCompiler/Private/KismetCompilerVMBackend.cpp
+++ b/Engine/Source/Editor/KismetCompiler/Private/KismetCompilerVMBackend.cpp
@@ -667,7 +667,7 @@ public:
 		Writer << EX_EndStructConst;
 	}
 
-	void EmitFunctionCall(FKismetCompilerContext& CompilerContext, FKismetFunctionContext& FunctionContext, FBlueprintCompiledStatement& Statement)
+	void EmitFunctionCall(FKismetCompilerContext& CompilerContext, FKismetFunctionContext& FunctionContext, FBlueprintCompiledStatement& Statement, UEdGraphNode* SourceNode)
 	{
 		UFunction* FunctionToCall = Statement.FunctionToCall;
 		check(FunctionToCall);
@@ -729,7 +729,10 @@ public:
 			UProperty* FuncParamProperty = *PropIt;
 			if (FuncParamProperty->HasAnyPropertyFlags(CPF_ReturnParm))
 			{
-				EmitDestinationExpression(Statement.LHS);
+				if (Statement.LHS)
+				{
+					EmitDestinationExpression(Statement.LHS);
+				}
 				bHasOutputValue = true;
 			}
 			else if (FuncParamProperty->HasAnyPropertyFlags(CPF_OutParm) && !FuncParamProperty->HasAnyPropertyFlags(CPF_ConstParm))
@@ -808,8 +811,18 @@ public:
  				}
 				else
 				{
-					// Emit parameter term normally
-					EmitTerm(Term, FuncParamProperty);
+					if (Term->InlineGeneratedParameter)
+					{
+						ensure(!Term->InlineGeneratedParameter->bIsJumpTarget);
+						ensure(Term->InlineGeneratedParameter->Type == KCST_CallFunction);
+
+						GenerateCodeForStatement(CompilerContext, FunctionContext, *Term->InlineGeneratedParameter, SourceNode);
+					}
+					else
+					{
+						// Emit parameter term normally
+						EmitTerm(Term, FuncParamProperty);
+					}
 				}
 
 
@@ -1314,7 +1327,7 @@ public:
 			Writer << ((Statement.Type == KCST_DebugSite) ? EX_Tracepoint : EX_WireTracepoint);
 			break;
 		case KCST_CallFunction:
-			EmitFunctionCall(CompilerContext, FunctionContext, Statement);
+			EmitFunctionCall(CompilerContext, FunctionContext, Statement, SourceNode);
 			break;
 		case KCST_CallDelegate:
 			EmitCallDelegate(Statement);
diff --git a/Engine/Source/Editor/KismetCompiler/Public/BPTerminal.h b/Engine/Source/Editor/KismetCompiler/Public/BPTerminal.h
--- a/Engine/Source/Editor/KismetCompiler/Public/BPTerminal.h
+++ b/Engine/Source/Editor/KismetCompiler/Public/BPTerminal.h
@@ -4,6 +4,8 @@
 
 //////////////////////////////////////////////////////////////////////////
 
+struct FBlueprintCompiledStatement;
+
 /** A terminal in the graph (literal or variable reference) */
 struct FBPTerminal
 {
@@ -35,6 +37,9 @@ struct FBPTerminal
 	/** String representation of the default value of the property associated with this term (or path to object) */
 	FString PropertyDefault;
 
+	/** Used for MathExpression optimization. The parameter will be filled directly by a result of a function called inline. No local variable is necessary to pass the value. */
+	FBlueprintCompiledStatement* InlineGeneratedParameter;
+
 	FBPTerminal()
 		: bIsLiteral(false)
 		, bIsLocal(false)
@@ -42,10 +47,11 @@ struct FBPTerminal
 		, bIsSavePersistent(false)
 		, bIsStructContext(false)
 		, bPassedByReference(false)
-		, Source(NULL)
-		, Context(NULL)
-		, AssociatedVarProperty(NULL)
-		, ObjectLiteral(NULL)
+		, Source(nullptr)
+		, Context(nullptr)
+		, AssociatedVarProperty(nullptr)
+		, ObjectLiteral(nullptr)
+		, InlineGeneratedParameter(nullptr)
 	{
 	}
 
diff --git a/Engine/Source/Editor/KismetCompiler/Public/KismetCompiledFunctionContext.h b/Engine/Source/Editor/KismetCompiler/Public/KismetCompiledFunctionContext.h
--- a/Engine/Source/Editor/KismetCompiler/Public/KismetCompiledFunctionContext.h
+++ b/Engine/Source/Editor/KismetCompiler/Public/KismetCompiledFunctionContext.h
@@ -69,6 +69,7 @@ public:
 	TIndirectArray<FBPTerminal> Locals;
 	TIndirectArray<FBPTerminal> EventGraphLocals;
 	TIndirectArray<FBPTerminal>	LevelActorReferences;
+	TIndirectArray<FBPTerminal>	InlineGeneratedValues; // A function generating the parameter will be called inline. The value won't be stored in a local variable.
 	TMap<UEdGraphPin*, FBPTerminal*> NetMap;
 	TMap<UEdGraphPin*, FBPTerminal*> LiteralHackMap;
 
