Pin tooltips on Call Function nodes in Blueprints will no longer grow with each compile or change to the Blueprint.

diff --git a/Engine/Source/Editor/BlueprintGraph/Private/K2Node_CallFunction.cpp b/Engine/Source/Editor/BlueprintGraph/Private/K2Node_CallFunction.cpp
--- a/Engine/Source/Editor/BlueprintGraph/Private/K2Node_CallFunction.cpp
+++ b/Engine/Source/Editor/BlueprintGraph/Private/K2Node_CallFunction.cpp
@@ -508,6 +508,7 @@ void UK2Node_CallFunction::GetPinHoverText(const UEdGraphPin& Pin, FString& Hove
 	{
 		for (auto& P : Pins)
 		{
+			P->PinToolTip.Empty();
 			GeneratePinTooltip(*P);
 		}
 
