[UE-17011] If a function doesn't have any extra keywords, make sure we still include the compact title at the very least.

diff --git a/Engine/Source/Editor/BlueprintGraph/Private/K2Node_CallFunction.cpp b/Engine/Source/Editor/BlueprintGraph/Private/K2Node_CallFunction.cpp
--- a/Engine/Source/Editor/BlueprintGraph/Private/K2Node_CallFunction.cpp
+++ b/Engine/Source/Editor/BlueprintGraph/Private/K2Node_CallFunction.cpp
@@ -1418,6 +1418,10 @@ FText UK2Node_CallFunction::GetKeywordsForFunction(const UFunction* Function)
 		Args.Add(TEXT("MetadataKeywords"), MetadataKeywords);
 		ResultKeywords = FText::Format(FText::FromString("{Name} {MetadataKeywords}"), Args);
 	}
+	else
+	{
+		ResultKeywords = FText::FromString(Keywords);
+	}
 
 	return ResultKeywords;
 }
