Back out revision 47 from //depot/UE4/Engine/Source/Editor/BlueprintGraph/Private/K2Node_MacroInstance.cpp

This causes infinite recursion in AllCircleTests.umap in QAGame

diff --git a/Engine/Source/Editor/BlueprintGraph/Private/K2Node_MacroInstance.cpp b/Engine/Source/Editor/BlueprintGraph/Private/K2Node_MacroInstance.cpp
--- a/Engine/Source/Editor/BlueprintGraph/Private/K2Node_MacroInstance.cpp
+++ b/Engine/Source/Editor/BlueprintGraph/Private/K2Node_MacroInstance.cpp
@@ -471,15 +471,6 @@ bool UK2Node_MacroInstance::HasExternalDependencies(TArray<class UStruct*>* Opti
 		{
 			OptionalOutput->AddUnique(OtherClass);
 		}
-
-		for (UEdGraphNode* Node : GetMacroGraph()->Nodes)
-		{
-			UK2Node* K2Node = Cast<UK2Node>(Node);
-			if (K2Node)
-			{
-				K2Node->HasExternalDependencies(OptionalOutput);
-			}
-		}
 	}
 	const bool bSuperResult = Super::HasExternalDependencies(OptionalOutput);
 	return bSuperResult || bResult;
