Optimized math Expression can internally use GetVariable node.  UE-17139

diff --git a/Engine/Source/Editor/BlueprintGraph/Private/MathExpressionHandler.cpp b/Engine/Source/Editor/BlueprintGraph/Private/MathExpressionHandler.cpp
--- a/Engine/Source/Editor/BlueprintGraph/Private/MathExpressionHandler.cpp
+++ b/Engine/Source/Editor/BlueprintGraph/Private/MathExpressionHandler.cpp
@@ -5,6 +5,7 @@
 #include "EdGraphUtilities.h"
 #include "KismetCompiler.h"
 #include "MathExpressionHandler.h"
+#include "K2Node_VariableGet.h"
 
 #define LOCTEXT_NAMESPACE "KCHandler_MathExpression"
 
@@ -100,13 +101,15 @@ FBlueprintCompiledStatement* FKCHandler_MathExpression::GenerateFunctionRPN(UEdG
 			if (!RHSTerm)
 			{
 				FBPTerminal** Term = PinToTry ? Context.NetMap.Find(PinToTry) : nullptr;
+				const bool bValidTerm = Term && *Term;
 				// Input is a literal term
-				if (Term && *Term && (*Term)->bIsLiteral)
+				// Input is a variable
+				if (bValidTerm && ((*Term)->bIsLiteral || (*Term)->AssociatedVarProperty))
 				{
 					RHSTerm = *Term;
 				}
 				// Input is an InlineGeneratedParameter
-				else if (Term && *Term)
+				else if (bValidTerm)
 				{
 					ensure(!(*Term)->InlineGeneratedParameter);
 					UEdGraphNode* SourceNode = PinToTry ? PinToTry->GetOwningNodeUnchecked() : nullptr;
@@ -143,14 +146,24 @@ void FKCHandler_MathExpression::RegisterNets(FKismetFunctionContext& Context, UE
 	auto InnerExitNode = Node_MathExpression->GetExitNode();
 	check(Node_MathExpression->BoundGraph && InnerEntryNode && InnerExitNode);
 
-	for (auto InnerGraph : Node_MathExpression->BoundGraph->Nodes)
+	for (auto InnerGraphNode : Node_MathExpression->BoundGraph->Nodes)
 	{
-		if (!InnerGraph || (InnerGraph == InnerEntryNode) || (InnerGraph == InnerExitNode))
+		if (!InnerGraphNode || (InnerGraphNode == InnerEntryNode) || (InnerGraphNode == InnerExitNode))
 		{
 			continue;
 		}
 
-		for (auto Pin : InnerGraph->Pins)
+		if (auto GetVarNode = Cast<UK2Node_VariableGet>(InnerGraphNode))
+		{
+			auto VarHandler = GetVarNode->CreateNodeHandler(CompilerContext);
+			if (ensure(VarHandler))
+			{
+				VarHandler->RegisterNets(Context, GetVarNode);
+			}
+			continue;
+		}
+
+		for (auto Pin : InnerGraphNode->Pins)
 		{
 			// Register fake terms for InlineGeneratedValues
 			if (Pin && (Pin->Direction == EEdGraphPinDirection::EGPD_Output) && Pin->LinkedTo.Num())
@@ -190,7 +203,7 @@ void FKCHandler_MathExpression::Compile(FKismetFunctionContext& Context, UEdGrap
 
 	if (!InnerExitNode || !InnerEntryNode || (InnerExitNode->Pins.Num() != 1) || ((InnerExitNode->Pins.Num() + InnerEntryNode->Pins.Num()) != Node->Pins.Num()))
 	{
-		Context.MessageLog.Warning(*LOCTEXT("Compile_PinError", "ICE - wrong inner pins - @@").ToString(), Node);
+		Context.MessageLog.Error(*LOCTEXT("Compile_PinError", "ICE - wrong inner pins - @@").ToString(), Node);
 		return;
 	}
 
@@ -202,7 +215,7 @@ void FKCHandler_MathExpression::Compile(FKismetFunctionContext& Context, UEdGrap
 
 	if (!bProperMap)
 	{
-		Context.MessageLog.Warning(*LOCTEXT("Compile_WrongMap", "ICE - cannot map pins - @@").ToString(), Node);
+		Context.MessageLog.Error(*LOCTEXT("Compile_WrongMap", "ICE - cannot map pins - @@").ToString(), Node);
 		return;
 	}
 
@@ -210,7 +223,7 @@ void FKCHandler_MathExpression::Compile(FKismetFunctionContext& Context, UEdGrap
 	auto OuterOutputPinPtr = InnerToOuterOutput.Find(InnerOutputPin);
 	if (!InnerOutputPin || (InnerOutputPin->LinkedTo.Num() != 1) || !InnerOutputPin->LinkedTo[0] || !OuterOutputPinPtr || !*OuterOutputPinPtr)
 	{
-		Context.MessageLog.Warning(*LOCTEXT("Compile_WrongOutputLink", "ICE - wrong output link - @@").ToString(), Node);
+		Context.MessageLog.Error(*LOCTEXT("Compile_WrongOutputLink", "ICE - wrong output link - @@").ToString(), Node);
 		return;
 	}
 
@@ -227,7 +240,7 @@ void FKCHandler_MathExpression::Compile(FKismetFunctionContext& Context, UEdGrap
 	}
 	else
 	{
-		Context.MessageLog.Warning(*LOCTEXT("Compile_CannotGenerateFunction", "ICE - cannot generate function - @@").ToString(), Node);
+		Context.MessageLog.Error(*LOCTEXT("Compile_CannotGenerateFunction", "ICE - cannot generate function - @@").ToString(), Node);
 		return;
 	}
 }
