UK2Node_FunctionEntry::GetTooltipText was added.
https://udn.unrealengine.com/questions/249835/k2node-functionentry-result-tooltips.html

diff --git a/Engine/Source/Editor/BlueprintGraph/Classes/K2Node_FunctionEntry.h b/Engine/Source/Editor/BlueprintGraph/Classes/K2Node_FunctionEntry.h
--- a/Engine/Source/Editor/BlueprintGraph/Classes/K2Node_FunctionEntry.h
+++ b/Engine/Source/Editor/BlueprintGraph/Classes/K2Node_FunctionEntry.h
@@ -39,6 +39,7 @@ class UK2Node_FunctionEntry : public UK2Node_FunctionTerminator
 	virtual bool CanUserDeleteNode() const override { return false; }
 	virtual bool IsDeprecated() const override;
 	virtual FString GetDeprecationMessage() const override;
+	virtual FText GetTooltipText() const override;
 	// End UEdGraphNode interface
 
 	// Begin UK2Node interface
diff --git a/Engine/Source/Editor/BlueprintGraph/Private/K2Node_FunctionEntry.cpp b/Engine/Source/Editor/BlueprintGraph/Private/K2Node_FunctionEntry.cpp
--- a/Engine/Source/Editor/BlueprintGraph/Private/K2Node_FunctionEntry.cpp
+++ b/Engine/Source/Editor/BlueprintGraph/Private/K2Node_FunctionEntry.cpp
@@ -3,6 +3,7 @@
 
 #include "BlueprintGraphPrivatePCH.h"
 #include "KismetCompiler.h"
+#include "K2Node_CallFunction.h"
 
 #define LOCTEXT_NAMESPACE "K2Node_FunctionEntry"
 
@@ -261,6 +262,16 @@ FString UK2Node_FunctionEntry::GetDeprecationMessage() const
 	return Super::GetDeprecationMessage();
 }
 
+FText UK2Node_FunctionEntry::GetTooltipText() const
+{
+	UFunction* Function = FindField<UFunction>(SignatureClass, SignatureName);
+	if (Function)
+	{
+		return FText::FromString(UK2Node_CallFunction::GetDefaultTooltipForFunction(Function));
+	}
+	return Super::GetTooltipText();
+}
+
 void UK2Node_FunctionEntry::ExpandNode(class FKismetCompilerContext& CompilerContext, UEdGraph* SourceGraph)
 {
 	Super::ExpandNode(CompilerContext, SourceGraph);
