[UE-16444][UE-16915] Guarding against the use of editor-only properties in Blueprints (still allowed for things like Blutilities).

diff --git a/Engine/Source/Editor/BlueprintGraph/Private/K2Node_Variable.cpp b/Engine/Source/Editor/BlueprintGraph/Private/K2Node_Variable.cpp
--- a/Engine/Source/Editor/BlueprintGraph/Private/K2Node_Variable.cpp
+++ b/Engine/Source/Editor/BlueprintGraph/Private/K2Node_Variable.cpp
@@ -449,20 +449,32 @@ FText UK2Node_Variable::GetToolTipHeading() const
 {
 	FText Heading = Super::GetToolTipHeading();
 
-	UProperty const* VariableProperty = VariableReference.ResolveMember<UProperty>(GetBlueprintClassFromNode());
-	if (VariableProperty && VariableProperty->HasAllPropertyFlags(CPF_Net))
+	// attempt to reflect the node's GetCornerIcon() with some tooltip documentation 
+	FText IconTag;
+	if ( UProperty const* VariableProperty = VariableReference.ResolveMember<UProperty>(GetBlueprintClassFromNode()) )
 	{
-		FText ReplicatedTag = LOCTEXT("ReplicatedVar", "Replicated");
-		if (Heading.IsEmpty())
+		if (VariableProperty->HasAllPropertyFlags(CPF_Net | CPF_EditorOnly))
 		{
-			Heading = ReplicatedTag;
+			IconTag = LOCTEXT("ReplicatedEditorOnlyVar", "Editor-Only | Replicated");
 		}
-		else 
+		else if (VariableProperty->HasAnyPropertyFlags(CPF_Net))
 		{
-			Heading = FText::Format(FText::FromString("{0}\n{1}"), ReplicatedTag, Heading);
+			IconTag = LOCTEXT("ReplicatedVar", "Replicated");
+		}
+		else if (VariableProperty->HasAnyPropertyFlags(CPF_EditorOnly))
+		{
+			IconTag = LOCTEXT("EditorOnlyVar", "Editor-Only");
 		}
 	}
 
+	if (Heading.IsEmpty())
+	{
+		return IconTag;
+	}
+	else if (!IconTag.IsEmpty())
+	{
+		Heading = FText::Format(FText::FromString("{0}\n{1}"), IconTag, Heading);
+	}
 	return Heading;
 }
 
@@ -627,10 +639,16 @@ bool UK2Node_Variable::RemapRestrictedLinkReference(FName OldVariableName, FName
 
 FName UK2Node_Variable::GetCornerIcon() const
 {
-	const UProperty* VariableProperty = VariableReference.ResolveMember<UProperty>(GetBlueprintClassFromNode());
-	if (VariableProperty && VariableProperty->HasAllPropertyFlags(CPF_Net))
+	if (const UProperty* VariableProperty = VariableReference.ResolveMember<UProperty>(GetBlueprintClassFromNode()))
 	{
-		return TEXT("Graph.Replication.Replicated");
+		if (VariableProperty->HasAllPropertyFlags(CPF_Net))
+		{
+			return TEXT("Graph.Replication.Replicated");
+		}
+		else if (VariableProperty->HasAllPropertyFlags(CPF_EditorOnly))
+		{
+			return TEXT("Graph.Editor.EditorOnlyIcon");
+		}
 	}
 
 	return Super::GetCornerIcon();
diff --git a/Engine/Source/Editor/EditorStyle/Private/SlateEditorStyle.cpp b/Engine/Source/Editor/EditorStyle/Private/SlateEditorStyle.cpp
--- a/Engine/Source/Editor/EditorStyle/Private/SlateEditorStyle.cpp
+++ b/Engine/Source/Editor/EditorStyle/Private/SlateEditorStyle.cpp
@@ -3856,6 +3856,8 @@ void FSlateEditorStyle::FStyle::SetupGraphEditorStyles()
 		Set( "Graph.Replication.ClientEvent", new IMAGE_BRUSH( "/Graph/ClientEvent", FVector2D(32,32) ) );
 		Set( "Graph.Replication.Replicated", new IMAGE_BRUSH( "/Graph/Replicated", FVector2D(32,32) ) );
 
+		Set("Graph.Editor.EditorOnlyIcon", new IMAGE_BRUSH("/Graph/EditorOnly", FVector2D(32, 32)));
+
 		Set( "Graph.Event.InterfaceEventIcon", new IMAGE_BRUSH("/Graph/InterfaceEventIcon", FVector2D(32,32) ) );
 
 		Set( "Graph.Latent.LatentIcon", new IMAGE_BRUSH("/Graph/LatentIcon", FVector2D(32,32) ) );
diff --git a/Engine/Source/Runtime/CoreUObject/Private/UObject/ScriptCore.cpp b/Engine/Source/Runtime/CoreUObject/Private/UObject/ScriptCore.cpp
--- a/Engine/Source/Runtime/CoreUObject/Private/UObject/ScriptCore.cpp
+++ b/Engine/Source/Runtime/CoreUObject/Private/UObject/ScriptCore.cpp
@@ -1097,13 +1097,27 @@ IMPLEMENT_VM_FUNCTION( EX_LocalVariable, execLocalVariable );
 
 void UObject::execInstanceVariable(FFrame& Stack, RESULT_DECL)
 {
-	UProperty* VarProperty = Stack.ReadProperty();
-	Stack.MostRecentPropertyAddress = VarProperty->ContainerPtrToValuePtr<uint8>(this);
+	UProperty* VarProperty = (UProperty*)Stack.ReadObject();
+	Stack.MostRecentProperty = VarProperty;
 
-	if (RESULT_PARAM)
+	if (VarProperty == nullptr)
 	{
-		VarProperty->CopyCompleteValueToScriptVM(RESULT_PARAM, Stack.MostRecentPropertyAddress);
+		static FBlueprintExceptionInfo ExceptionInfo(EBlueprintExceptionType::AccessViolation, TEXT("Attempt to access missing property. If this is a packaged/cooked build, are you attempting to use an editor-only property?"));
+		FBlueprintCoreDelegates::ThrowScriptException(this, Stack, ExceptionInfo);
+
+		Stack.MostRecentPropertyAddress = nullptr;
+	}
+	else
+	{
+		Stack.MostRecentPropertyAddress = VarProperty->ContainerPtrToValuePtr<uint8>(this);
+
+		if (RESULT_PARAM)
+		{
+			VarProperty->CopyCompleteValueToScriptVM(RESULT_PARAM, Stack.MostRecentPropertyAddress);
+		}
 	}
+
+	
 }
 IMPLEMENT_VM_FUNCTION( EX_InstanceVariable, execInstanceVariable );
 
@@ -1616,6 +1630,7 @@ void UObject::ProcessContextOpcode( FFrame& Stack, RESULT_DECL, bool bCanFailSil
 				// Stack.MostRecentProperty will be NULL under the following conditions:
 				//   1. the context expression was a function call which returned an object
 				//   2. the context expression was a literal object reference
+				//   3. the context expression was an instance variable that no longer exists (it was editor-only, etc.)
 				static FBlueprintExceptionInfo ExceptionInfo(EBlueprintExceptionType::AccessViolation, TEXT("Accessed None"));
 				FBlueprintCoreDelegates::ThrowScriptException(this, Stack, ExceptionInfo);
 			}
