[UE-17186] No longer destroying pins on reconstruction when the types dont match on load, instead relying on compile errors/warnings. This bug was caused by code changes to a node that made the pin type more restrictive. In this case we definitely don't want to destroy connections, and we could not find a case where destroying connections on load was advisable.

diff --git a/Engine/Source/Editor/BlueprintGraph/Private/K2Node.cpp b/Engine/Source/Editor/BlueprintGraph/Private/K2Node.cpp
--- a/Engine/Source/Editor/BlueprintGraph/Private/K2Node.cpp
+++ b/Engine/Source/Editor/BlueprintGraph/Private/K2Node.cpp
@@ -463,7 +463,7 @@ UK2Node::ERedirectType UK2Node::DoPinsMatchForReconstruction(const UEdGraphPin*
 		if( OuterGraph && OuterGraph->Schema )
 		{
 			const UEdGraphSchema_K2* K2Schema = Cast<const UEdGraphSchema_K2>(GetSchema());
-			if( !K2Schema || K2Schema->IsSelfPin(*NewPin) || K2Schema->ArePinTypesCompatible(OldPin->PinType, NewPin->PinType) )
+			if (!K2Schema || K2Schema->IsSelfPin(*NewPin) || K2Schema->ArePinTypesCompatible(OldPin->PinType, NewPin->PinType) || GetBlueprint()->bIsRegeneratingOnLoad)
 			{
 				RedirectType = ERedirectType_Name;
 			}
