For the simplest cases MathExpresion will be expanded the old way. The RPE optimization cannot handle the exoression without any function call.
Cases:
- literal - "3.14"
- variable access - "MyVar"
- new local variavle - "a"

UE-17139 Math Expressions fail to compile when output pin connected

diff --git a/Engine/Source/Editor/BlueprintGraph/Classes/K2Node_MathExpression.h b/Engine/Source/Editor/BlueprintGraph/Classes/K2Node_MathExpression.h
--- a/Engine/Source/Editor/BlueprintGraph/Classes/K2Node_MathExpression.h
+++ b/Engine/Source/Editor/BlueprintGraph/Classes/K2Node_MathExpression.h
@@ -36,14 +36,14 @@ public:
 	virtual void PostPlacedNewNode() override;
 	virtual void ReconstructNode() override;
 	virtual void FindDiffs(class UEdGraphNode* OtherNode, struct FDiffResults& Results )  override;
-	virtual bool ShouldMergeChildGraphs() const override { return false; }
+	virtual bool ShouldMergeChildGraphs() const override { return ShouldExpandInsteadCompile(); }
 	// End of UEdGraphNode interface
 
 	// UK2Node interface
 	virtual void ValidateNodeDuringCompilation(class FCompilerResultsLog& MessageLog) const override;
 	virtual void GetMenuActions(FBlueprintActionDatabaseRegistrar& ActionRegistrar) const override;
 	virtual class FNodeHandlingFunctor* CreateNodeHandler(class FKismetCompilerContext& CompilerContext) const override;
-	virtual bool IsNodePure() const { return true; }
+	virtual bool IsNodePure() const { return !ShouldExpandInsteadCompile(); }
 	// End of UK2Node interface
 
 	// Begin UK2Node_EditablePinBase interface
@@ -51,6 +51,9 @@ public:
 	// End UK2Node_EditablePinBase interface
 
 private:
+	/* Returns true, when the node can/should not be optimized.*/
+	bool ShouldExpandInsteadCompile() const;
+
 	/**
 	* Clears this node's sub-graph, and then takes the supplied string and
 	* parses, and converts it into a series of new graph nodes.
diff --git a/Engine/Source/Editor/BlueprintGraph/Private/K2Node_MathExpression.cpp b/Engine/Source/Editor/BlueprintGraph/Private/K2Node_MathExpression.cpp
--- a/Engine/Source/Editor/BlueprintGraph/Private/K2Node_MathExpression.cpp
+++ b/Engine/Source/Editor/BlueprintGraph/Private/K2Node_MathExpression.cpp
@@ -2435,12 +2435,39 @@ void UK2Node_MathExpression::GetMenuActions(FBlueprintActionDatabaseRegistrar& A
 	}
 }
 
+
 //------------------------------------------------------------------------------
 FNodeHandlingFunctor* UK2Node_MathExpression::CreateNodeHandler(class FKismetCompilerContext& CompilerContext) const
 {
 	return new FKCHandler_MathExpression(CompilerContext);
 }
 
+//------------------------------------------------------------------------------
+bool UK2Node_MathExpression::ShouldExpandInsteadCompile() const
+{
+	const int32 TunnelNodesNum = 2;
+	if (!BoundGraph || (TunnelNodesNum >= BoundGraph->Nodes.Num()))
+	{
+		return true;
+	}
+
+	if ((TunnelNodesNum + 1) == BoundGraph->Nodes.Num())
+	{
+		TArray<UEdGraphNode*> InnerNodes = BoundGraph->Nodes;
+		InnerNodes.RemoveSingleSwap(GetEntryNode(), false);
+		InnerNodes.RemoveSingleSwap(GetExitNode(), false);
+		const bool bTheOnlyNodeIsNotAFunctionCall = (1 == InnerNodes.Num())
+			&& (nullptr != InnerNodes[0])
+			&& !InnerNodes[0]->IsA<UK2Node_CallFunction>();
+		if (bTheOnlyNodeIsNotAFunctionCall)
+		{
+			return true;
+		}
+	}
+
+	return false;
+}
+
 //------------------------------------------------------------------------------
 TSharedPtr<class INameValidatorInterface> UK2Node_MathExpression::MakeNameValidator() const
 {
