[UE-15480] The output pin subtype for a GetClass function call node is now linked to the input pin's subtype.

change summary:
- modified FDynamicOutputHelper::GetPinClass() to include object reference pin types
- modified FDynamicOutputHelper::IsTypePickerPin() to include the PC_Object category
- modified the UGameplayStatics::GetObjectClass() UFUNCTION declaration to include 'DeterminesOutputType' metadata

diff --git a/Engine/Source/Editor/BlueprintGraph/Private/K2Node_CallFunction.cpp b/Engine/Source/Editor/BlueprintGraph/Private/K2Node_CallFunction.cpp
--- a/Engine/Source/Editor/BlueprintGraph/Private/K2Node_CallFunction.cpp
+++ b/Engine/Source/Editor/BlueprintGraph/Private/K2Node_CallFunction.cpp
@@ -238,8 +238,8 @@ UClass* FDynamicOutputHelper::GetPinClass(UEdGraphPin* Pin)
 {
 	UClass* PinClass = UObject::StaticClass();
 
-	bool const bIsClassPin = (Pin->PinType.PinCategory == UEdGraphSchema_K2::PC_Class);
-	if (bIsClassPin)
+	bool const bIsClassOrObjectPin = (Pin->PinType.PinCategory == UEdGraphSchema_K2::PC_Class || Pin->PinType.PinCategory == UEdGraphSchema_K2::PC_Object);
+	if (bIsClassOrObjectPin)
 	{
 		if (UClass* DefaultClass = Cast<UClass>(Pin->DefaultObject))
 		{
@@ -319,7 +319,8 @@ bool FDynamicOutputHelper::IsTypePickerPin(UEdGraphPin* Pin) const
 	}
 
 	bool const bPinIsClassPicker = (Pin->PinType.PinCategory == UEdGraphSchema_K2::PC_Class);
-	return bIsTypeDeterminingPin && bPinIsClassPicker && (Pin->Direction == EGPD_Input);
+	bool const bPinIsObjectPicker = (Pin->PinType.PinCategory == UEdGraphSchema_K2::PC_Object);
+	return bIsTypeDeterminingPin && (bPinIsClassPicker || bPinIsObjectPicker) && (Pin->Direction == EGPD_Input);
 }
 
 UEdGraphPin* FDynamicOutputHelper::GetDynamicOutPin(const UK2Node_CallFunction* FuncNode)
diff --git a/Engine/Source/Runtime/Engine/Classes/Kismet/GameplayStatics.h b/Engine/Source/Runtime/Engine/Classes/Kismet/GameplayStatics.h
--- a/Engine/Source/Runtime/Engine/Classes/Kismet/GameplayStatics.h
+++ b/Engine/Source/Runtime/Engine/Classes/Kismet/GameplayStatics.h
@@ -157,7 +157,7 @@ class ENGINE_API UGameplayStatics : public UBlueprintFunctionLibrary
 	UFUNCTION(BlueprintPure, Category="Game", meta=(WorldContext="WorldContextObject") )
 	static class AGameState* GetGameState(UObject* WorldContextObject);
 
-	UFUNCTION(BlueprintPure, meta=(DisplayName = "GetClass"), Category="Utilities")
+	UFUNCTION(BlueprintPure, meta=(DisplayName = "GetClass", DeterminesOutputType="Object"), Category="Utilities")
 	static class UClass *GetObjectClass(const UObject *Object);
 
 	/**
