[UEBP-59] Optimize Blueprint loading done subsequently to initializing the blueprint action database, by decreasing the number of TObjectIterators we use to build the database.

diff --git a/Engine/Source/Editor/AnimGraph/Private/AnimGraphNode_BlendListByEnum.cpp b/Engine/Source/Editor/AnimGraph/Private/AnimGraphNode_BlendListByEnum.cpp
--- a/Engine/Source/Editor/AnimGraph/Private/AnimGraphNode_BlendListByEnum.cpp
+++ b/Engine/Source/Editor/AnimGraph/Private/AnimGraphNode_BlendListByEnum.cpp
@@ -58,37 +58,26 @@ void UAnimGraphNode_BlendListByEnum::PostPlacedNewNode()
 
 void UAnimGraphNode_BlendListByEnum::GetMenuActions(FBlueprintActionDatabaseRegistrar& ActionRegistrar) const
 {
-	auto CustomizeBlendListEnumNodeLambda = [](UEdGraphNode* NewNode, bool bIsTemplateNode, UEnum* EnumValue)
+	struct GetMenuActions_Utils
 	{
-		UAnimGraphNode_BlendListByEnum* BlendListEnumNode = CastChecked<UAnimGraphNode_BlendListByEnum>(NewNode);
-		BlendListEnumNode->BoundEnum = EnumValue;
+		static void SetNodeEnum(UEdGraphNode* NewNode, bool /*bIsTemplateNode*/, TWeakObjectPtr<UEnum> NonConstEnumPtr)
+		{
+			UAnimGraphNode_BlendListByEnum* BlendListEnumNode = CastChecked<UAnimGraphNode_BlendListByEnum>(NewNode);
+			BlendListEnumNode->BoundEnum = NonConstEnumPtr.Get();
+		}
 	};
 
+	UClass* NodeClass = GetClass();
 	// add all blendlist enum entries
-	for (TObjectIterator<UEnum> EnumIt; EnumIt; ++EnumIt)
+	ActionRegistrar.RegisterEnumActions( FBlueprintActionDatabaseRegistrar::FMakeEnumSpawnerDelegate::CreateLambda([NodeClass](const UEnum* Enum)->UBlueprintNodeSpawner*
 	{
-		UEnum* CurrentEnum = *EnumIt;
-
-		// to keep from needlessly instantiating a UBlueprintNodeSpawner, first   
-		// check to make sure that the registrar is looking for actions of this type
-		// (could be regenerating actions for a specific asset, and therefore the 
-		// registrar would only accept actions corresponding to that asset)
-		if (!ActionRegistrar.IsOpenForRegistration(CurrentEnum))
-		{
-			continue;
-		}
+		UBlueprintNodeSpawner* NodeSpawner = UBlueprintNodeSpawner::Create(NodeClass);
+		check(NodeSpawner != nullptr);
+		TWeakObjectPtr<UEnum> NonConstEnumPtr = Enum;
+		NodeSpawner->CustomizeNodeDelegate = UBlueprintNodeSpawner::FCustomizeNodeDelegate::CreateStatic(GetMenuActions_Utils::SetNodeEnum, NonConstEnumPtr);
 
-		const bool bIsBlueprintType = UEdGraphSchema_K2::IsAllowableBlueprintVariableType(CurrentEnum);
-		if (bIsBlueprintType)
-		{
-			UBlueprintNodeSpawner* NodeSpawner = UBlueprintNodeSpawner::Create(GetClass());
-			check(NodeSpawner != nullptr);
-
-			NodeSpawner->CustomizeNodeDelegate = UBlueprintNodeSpawner::FCustomizeNodeDelegate::CreateStatic(CustomizeBlendListEnumNodeLambda, CurrentEnum);
-
-			ActionRegistrar.AddBlueprintAction(CurrentEnum, NodeSpawner);
-		}
-	}
+		return NodeSpawner;
+	}) );
 }
 
 void UAnimGraphNode_BlendListByEnum::GetContextMenuActions(const FGraphNodeContextMenuBuilder& Context) const
diff --git a/Engine/Source/Editor/AnimGraph/Private/AnimGraphNode_BlendSpacePlayer.cpp b/Engine/Source/Editor/AnimGraph/Private/AnimGraphNode_BlendSpacePlayer.cpp
--- a/Engine/Source/Editor/AnimGraph/Private/AnimGraphNode_BlendSpacePlayer.cpp
+++ b/Engine/Source/Editor/AnimGraph/Private/AnimGraphNode_BlendSpacePlayer.cpp
@@ -109,34 +109,53 @@ void UAnimGraphNode_BlendSpacePlayer::GetContextMenuActions(const FGraphNodeCont
 
 void UAnimGraphNode_BlendSpacePlayer::GetMenuActions(FBlueprintActionDatabaseRegistrar& ActionRegistrar) const
 {
-	auto PostSpawnSetupLambda = [](UEdGraphNode* NewNode, bool /*bIsTemplateNode*/, TWeakObjectPtr<UBlendSpaceBase> BlendSpace)
+	struct GetMenuActions_Utils
 	{
-		UAnimGraphNode_BlendSpacePlayer* BlendSpaceNode = CastChecked<UAnimGraphNode_BlendSpacePlayer>(NewNode);
-		BlendSpaceNode->Node.BlendSpace = BlendSpace.Get();
+		static void SetNodeBlendSpace(UEdGraphNode* NewNode, bool /*bIsTemplateNode*/, TWeakObjectPtr<UBlendSpaceBase> BlendSpace)
+		{
+			UAnimGraphNode_BlendSpacePlayer* BlendSpaceNode = CastChecked<UAnimGraphNode_BlendSpacePlayer>(NewNode);
+			BlendSpaceNode->Node.BlendSpace = BlendSpace.Get();
+		}
+
+		static UBlueprintNodeSpawner* MakeBlendSpaceAction(TSubclassOf<UEdGraphNode> const NodeClass, const UBlendSpaceBase* BlendSpace)
+		{
+			UBlueprintNodeSpawner* NodeSpawner = nullptr;
+
+			bool const bIsAimOffset = BlendSpace->IsA(UAimOffsetBlendSpace::StaticClass()) ||
+				BlendSpace->IsA(UAimOffsetBlendSpace1D::StaticClass());
+			if (!bIsAimOffset)
+			{
+				NodeSpawner = UBlueprintNodeSpawner::Create(NodeClass);
+				check(NodeSpawner != nullptr);
+
+				TWeakObjectPtr<UBlendSpaceBase> BlendSpacePtr = BlendSpace;
+				NodeSpawner->CustomizeNodeDelegate = UBlueprintNodeSpawner::FCustomizeNodeDelegate::CreateStatic(GetMenuActions_Utils::SetNodeBlendSpace, BlendSpacePtr);
+			}	
+			return NodeSpawner;
+		}
 	};
 
-	for (TObjectIterator<UBlendSpaceBase> BlendSpaceIt; BlendSpaceIt; ++BlendSpaceIt)
+	if (const UObject* RegistrarTarget = ActionRegistrar.GetActionKeyFilter())
 	{
-		UBlendSpaceBase* BlendSpace = *BlendSpaceIt;
-		// to keep from needlessly instantiating a UBlueprintNodeSpawner, first   
-		// check to make sure that the registrar is looking for actions of this type
-		// (could be regenerating actions for a specific asset, and therefore the 
-		// registrar would only accept actions corresponding to that asset)
-		if (!ActionRegistrar.IsOpenForRegistration(BlendSpace))
+		if (const UBlendSpaceBase* TargetBlendSpace = Cast<UBlendSpaceBase>(RegistrarTarget))
 		{
-			continue;
+			if (UBlueprintNodeSpawner* NodeSpawner = GetMenuActions_Utils::MakeBlendSpaceAction(GetClass(), TargetBlendSpace))
+			{
+				ActionRegistrar.AddBlueprintAction(TargetBlendSpace, NodeSpawner);
+			}
 		}
-
-		bool const bIsAimOffset = BlendSpace->IsA(UAimOffsetBlendSpace::StaticClass()) ||
-			BlendSpace->IsA(UAimOffsetBlendSpace1D::StaticClass());
-		if (!bIsAimOffset)
+		// else, the Blueprint database is specifically looking for actions pertaining to something different (not a BlendSpace asset)
+	}
+	else
+	{
+		UClass* NodeClass = GetClass();
+		for (TObjectIterator<UBlendSpaceBase> BlendSpaceIt; BlendSpaceIt; ++BlendSpaceIt)
 		{
-			UBlueprintNodeSpawner* NodeSpawner = UBlueprintNodeSpawner::Create(GetClass());
-			check(NodeSpawner != nullptr);
-			ActionRegistrar.AddBlueprintAction(BlendSpace, NodeSpawner);
-
-			TWeakObjectPtr<UBlendSpaceBase> BlendSpacePtr = BlendSpace;
-			NodeSpawner->CustomizeNodeDelegate = UBlueprintNodeSpawner::FCustomizeNodeDelegate::CreateStatic(PostSpawnSetupLambda, BlendSpacePtr);
+			UBlendSpaceBase* BlendSpace = *BlendSpaceIt;
+			if (UBlueprintNodeSpawner* NodeSpawner = GetMenuActions_Utils::MakeBlendSpaceAction(NodeClass, BlendSpace))
+			{
+				ActionRegistrar.AddBlueprintAction(BlendSpace, NodeSpawner);
+			}
 		}
 	}
 }
diff --git a/Engine/Source/Editor/AnimGraph/Private/AnimGraphNode_RotationOffsetBlendSpace.cpp b/Engine/Source/Editor/AnimGraph/Private/AnimGraphNode_RotationOffsetBlendSpace.cpp
--- a/Engine/Source/Editor/AnimGraph/Private/AnimGraphNode_RotationOffsetBlendSpace.cpp
+++ b/Engine/Source/Editor/AnimGraph/Private/AnimGraphNode_RotationOffsetBlendSpace.cpp
@@ -60,34 +60,53 @@ FText UAnimGraphNode_RotationOffsetBlendSpace::GetNodeTitle(ENodeTitleType::Type
 
 void UAnimGraphNode_RotationOffsetBlendSpace::GetMenuActions(FBlueprintActionDatabaseRegistrar& ActionRegistrar) const
 {
-	auto PostSpawnSetupLambda = [](UEdGraphNode* NewNode, bool /*bIsTemplateNode*/, TWeakObjectPtr<UBlendSpaceBase> BlendSpace)
+	struct GetMenuActions_Utils
 	{
-		UAnimGraphNode_RotationOffsetBlendSpace* BlendSpaceNode = CastChecked<UAnimGraphNode_RotationOffsetBlendSpace>(NewNode);
-		BlendSpaceNode->Node.BlendSpace = BlendSpace.Get();
+		static void SetNodeBlendSpace(UEdGraphNode* NewNode, bool /*bIsTemplateNode*/, TWeakObjectPtr<UBlendSpaceBase> BlendSpace)
+		{
+			UAnimGraphNode_RotationOffsetBlendSpace* BlendSpaceNode = CastChecked<UAnimGraphNode_RotationOffsetBlendSpace>(NewNode);
+			BlendSpaceNode->Node.BlendSpace = BlendSpace.Get();
+		}
+
+		static UBlueprintNodeSpawner* MakeBlendSpaceAction(TSubclassOf<UEdGraphNode> const NodeClass, const UBlendSpaceBase* BlendSpace)
+		{
+			UBlueprintNodeSpawner* NodeSpawner = nullptr;
+
+			bool const bIsAimOffset = BlendSpace->IsA(UAimOffsetBlendSpace::StaticClass()) ||
+				BlendSpace->IsA(UAimOffsetBlendSpace1D::StaticClass());
+			if (bIsAimOffset)
+			{
+				NodeSpawner = UBlueprintNodeSpawner::Create(NodeClass);
+				check(NodeSpawner != nullptr);
+
+				TWeakObjectPtr<UBlendSpaceBase> BlendSpacePtr = BlendSpace;
+				NodeSpawner->CustomizeNodeDelegate = UBlueprintNodeSpawner::FCustomizeNodeDelegate::CreateStatic(GetMenuActions_Utils::SetNodeBlendSpace, BlendSpacePtr);
+			}
+			return NodeSpawner;
+		}
 	};
 
-	for (TObjectIterator<UBlendSpaceBase> BlendSpaceIt; BlendSpaceIt; ++BlendSpaceIt)
+	if (const UObject* RegistrarTarget = ActionRegistrar.GetActionKeyFilter())
 	{
-		UBlendSpaceBase* BlendSpace = *BlendSpaceIt;
-		// to keep from needlessly instantiating a UBlueprintNodeSpawner, first   
-		// check to make sure that the registrar is looking for actions of this type
-		// (could be regenerating actions for a specific asset, and therefore the 
-		// registrar would only accept actions corresponding to that asset)
-		if (!ActionRegistrar.IsOpenForRegistration(BlendSpace))
+		if (const UBlendSpaceBase* TargetBlendSpace = Cast<UBlendSpaceBase>(RegistrarTarget))
 		{
-			continue;
+			if (UBlueprintNodeSpawner* NodeSpawner = GetMenuActions_Utils::MakeBlendSpaceAction(GetClass(), TargetBlendSpace))
+			{
+				ActionRegistrar.AddBlueprintAction(TargetBlendSpace, NodeSpawner);
+			}
 		}
-
-		bool const bIsAimOffset = BlendSpace->IsA(UAimOffsetBlendSpace::StaticClass()) ||
-			BlendSpace->IsA(UAimOffsetBlendSpace1D::StaticClass());
-		if (bIsAimOffset)
+		// else, the Blueprint database is specifically looking for actions pertaining to something different (not a BlendSpace asset)
+	}
+	else
+	{
+		UClass* NodeClass = GetClass();
+		for (TObjectIterator<UBlendSpaceBase> BlendSpaceIt; BlendSpaceIt; ++BlendSpaceIt)
 		{
-			UBlueprintNodeSpawner* NodeSpawner = UBlueprintNodeSpawner::Create(GetClass());
-			check(NodeSpawner != nullptr);
-			ActionRegistrar.AddBlueprintAction(BlendSpace, NodeSpawner);
-
-			TWeakObjectPtr<UBlendSpaceBase> BlendSpacePtr = BlendSpace;
-			NodeSpawner->CustomizeNodeDelegate = UBlueprintNodeSpawner::FCustomizeNodeDelegate::CreateStatic(PostSpawnSetupLambda, BlendSpacePtr);
+			UBlendSpaceBase* BlendSpace = *BlendSpaceIt;
+			if (UBlueprintNodeSpawner* NodeSpawner = GetMenuActions_Utils::MakeBlendSpaceAction(NodeClass, BlendSpace))
+			{
+				ActionRegistrar.AddBlueprintAction(BlendSpace, NodeSpawner);
+			}
 		}
 	}
 }
diff --git a/Engine/Source/Editor/BlueprintGraph/Private/BlueprintActionDatabaseRegistrar.cpp b/Engine/Source/Editor/BlueprintGraph/Private/BlueprintActionDatabaseRegistrar.cpp
--- a/Engine/Source/Editor/BlueprintGraph/Private/BlueprintActionDatabaseRegistrar.cpp
+++ b/Engine/Source/Editor/BlueprintGraph/Private/BlueprintActionDatabaseRegistrar.cpp
@@ -185,6 +185,167 @@ bool FBlueprintActionDatabaseRegistrar::IsOpenForRegistration(FAssetData const&
 	return IsOpenForRegistration(OwnerKey);
 }
 
+//------------------------------------------------------------------------------
+int32 FBlueprintActionDatabaseRegistrar::RegisterStructActions(const FMakeStructSpawnerDelegate& MakeActionCallback)
+{
+	int32 RegisteredCount = 0;
+
+	// to keep from needlessly looping through UScriptStructs, first check to 
+	// see if the registrar is looking for only certain actions of this type
+	// (could be regenerating actions for a specific asset, and therefore the 
+	// registrar would only accept actions corresponding to that asset)
+	if (const UObject* RegistrarTarget = GetActionKeyFilter())
+	{
+		if (const UScriptStruct* StructAsset = Cast<UScriptStruct>(RegistrarTarget))
+		{
+			check(IsOpenForRegistration(StructAsset));
+			if (UBlueprintNodeSpawner* NewAction = MakeActionCallback.Execute(StructAsset))
+			{
+				RegisteredCount += (int32)AddBlueprintAction(StructAsset, NewAction);
+			}
+		}
+		// else, the target is a class or a different asset type... not something pertaining to a struct
+	}
+	else
+	{
+		for (TObjectIterator<UScriptStruct> StructIt; StructIt; ++StructIt)
+		{
+			UScriptStruct const* Struct = (*StructIt);
+			if (UBlueprintNodeSpawner* NewAction = MakeActionCallback.Execute(Struct))
+			{
+				RegisteredCount += (int32)AddBlueprintAction(Struct, NewAction);
+			}
+		}
+	}
+	return RegisteredCount;
+}
+
+//------------------------------------------------------------------------------
+int32 FBlueprintActionDatabaseRegistrar::RegisterEnumActions(const FMakeEnumSpawnerDelegate& MakeActionCallback)
+{
+	int32 RegisteredCount = 0;
+	if (const UObject* RegistrarTarget = GetActionKeyFilter())
+	{
+		if (const UClass* TargetClass = Cast<UClass>(RegistrarTarget))
+		{
+			for (TFieldIterator<UEnum> EnumIt(TargetClass, EFieldIteratorFlags::ExcludeSuper); EnumIt; ++EnumIt)
+			{
+				UEnum const* EnumToConsider = (*EnumIt);
+				if (!UEdGraphSchema_K2::IsAllowableBlueprintVariableType(EnumToConsider))
+				{
+					continue;
+				}
+
+				if (UBlueprintNodeSpawner* NewAction = MakeActionCallback.Execute(EnumToConsider))
+				{
+					RegisteredCount += (int32)AddBlueprintAction(EnumToConsider, NewAction);
+				}
+			}
+		}
+		else if (const UEnum* TargetEnum = Cast<UEnum>(RegistrarTarget))
+		{
+			if (UEdGraphSchema_K2::IsAllowableBlueprintVariableType(TargetEnum))
+			{
+				if (UBlueprintNodeSpawner* NewAction = MakeActionCallback.Execute(TargetEnum))
+				{
+					RegisteredCount += (int32)AddBlueprintAction(TargetEnum, NewAction);
+				}
+			}
+		}
+	}
+	else
+	{
+		for (TObjectIterator<UEnum> EnumIt; EnumIt; ++EnumIt)
+		{
+			UEnum const* EnumToConsider = (*EnumIt);
+			if (!UEdGraphSchema_K2::IsAllowableBlueprintVariableType(EnumToConsider))
+			{
+				continue;
+			}
+
+			if (UBlueprintNodeSpawner* NewAction = MakeActionCallback.Execute(EnumToConsider))
+			{
+				RegisteredCount += (int32)AddBlueprintAction(EnumToConsider, NewAction);
+			}
+		}
+	}
+
+	return RegisteredCount;
+}
+
+//------------------------------------------------------------------------------
+int32 FBlueprintActionDatabaseRegistrar::RegisterClassFactoryActions(const UClass* TargetType, const FMakeFuncSpawnerDelegate& MakeActionCallback)
+{
+	struct RegisterClassFactoryActions_Utils
+	{
+		static bool IsFactoryMethod(const UFunction* Function, const UClass* TargetType)
+		{
+			if (!Function->HasAnyFunctionFlags(FUNC_Static))
+			{
+				return false;
+			}
+
+			UObjectProperty* ReturnProperty = Cast<UObjectProperty>(Function->GetReturnProperty());
+			// see if the function is a static factory method
+			bool const bIsFactoryMethod = (ReturnProperty != nullptr) && ReturnProperty->PropertyClass->IsChildOf(TargetType);
+
+			return bIsFactoryMethod;
+		}
+	};
+
+	int32 RegisteredCount = 0;
+	if (const UObject* RegistrarTarget = GetActionKeyFilter())
+	{
+		if (const UClass* TargetClass = Cast<UClass>(RegistrarTarget))
+		{
+			if (!TargetClass->HasAnyClassFlags(CLASS_Abstract) && !TargetClass->IsChildOf(TargetType))
+			{
+				for (TFieldIterator<UFunction> FuncIt(TargetClass, EFieldIteratorFlags::ExcludeSuper); FuncIt; ++FuncIt)
+				{
+					UFunction* Function = *FuncIt;
+					if (!RegisterClassFactoryActions_Utils::IsFactoryMethod(Function, TargetType))
+					{
+						continue;
+					}
+					else if (UBlueprintNodeSpawner* NewAction = MakeActionCallback.Execute(Function))
+					{
+						RegisteredCount += (int32)AddBlueprintAction(Function, NewAction);
+					}
+				}
+			}
+		}
+	}
+	else
+	{
+		// these nested loops are combing over the same classes/functions the
+		// FBlueprintActionDatabase does; ideally we save on perf and fold this in
+		// with FBlueprintActionDatabase, but we want to give separate modules
+		// the opportunity to add their own actions per class func
+		for (TObjectIterator<UClass> ClassIt; ClassIt; ++ClassIt)
+		{
+			UClass* Class = *ClassIt;
+			if (Class->HasAnyClassFlags(CLASS_Abstract) || !Class->IsChildOf(TargetType))
+			{
+				continue;
+			}
+			
+			for (TFieldIterator<UFunction> FuncIt(Class, EFieldIteratorFlags::ExcludeSuper); FuncIt; ++FuncIt)
+			{
+				UFunction* Function = *FuncIt;
+				if (!RegisterClassFactoryActions_Utils::IsFactoryMethod(Function, TargetType))
+				{
+					continue;
+				}
+				else if (UBlueprintNodeSpawner* NewAction = MakeActionCallback.Execute(Function))
+				{
+					RegisteredCount += (int32)AddBlueprintAction(Function, NewAction);
+				}
+			}
+		}
+	}
+	return 0;
+}
+
 //------------------------------------------------------------------------------
 bool FBlueprintActionDatabaseRegistrar::AddActionToDatabase(UObject const* ActionKey, UBlueprintNodeSpawner* NodeSpawner)
 {
diff --git a/Engine/Source/Editor/BlueprintGraph/Private/K2Node_BreakStruct.cpp b/Engine/Source/Editor/BlueprintGraph/Private/K2Node_BreakStruct.cpp
--- a/Engine/Source/Editor/BlueprintGraph/Private/K2Node_BreakStruct.cpp
+++ b/Engine/Source/Editor/BlueprintGraph/Private/K2Node_BreakStruct.cpp
@@ -320,53 +320,44 @@ FNodeHandlingFunctor* UK2Node_BreakStruct::CreateNodeHandler(class FKismetCompil
 
 void UK2Node_BreakStruct::GetMenuActions(FBlueprintActionDatabaseRegistrar& ActionRegistrar) const
 {
-	auto SetNodeStructLambda = [](UEdGraphNode* NewNode, UField const* /*StructField*/, TWeakObjectPtr<UScriptStruct> NonConstStructPtr)
+	struct GetMenuActions_Utils
 	{
-		UK2Node_BreakStruct* BreakNode = CastChecked<UK2Node_BreakStruct>(NewNode);
-		BreakNode->StructType = NonConstStructPtr.Get();
-	};
+		static void SetNodeStruct(UEdGraphNode* NewNode, UField const* /*StructField*/, TWeakObjectPtr<UScriptStruct> NonConstStructPtr)
+		{
+			UK2Node_BreakStruct* BreakNode = CastChecked<UK2Node_BreakStruct>(NewNode);
+			BreakNode->StructType = NonConstStructPtr.Get();
+		}
 
-	auto CategoryOverrideLambda = [](FBlueprintActionContext const& Context, IBlueprintNodeBinder::FBindingSet const& /*Bindings*/, FBlueprintActionUiSpec* UiSpecOut, TWeakObjectPtr<UScriptStruct> StructPtr)
-	{
-		for (UEdGraphPin* Pin : Context.Pins)
+		static void OverrideCategory(FBlueprintActionContext const& Context, IBlueprintNodeBinder::FBindingSet const& /*Bindings*/, FBlueprintActionUiSpec* UiSpecOut, TWeakObjectPtr<UScriptStruct> StructPtr)
 		{
-			UScriptStruct* PinStruct = Cast<UScriptStruct>(Pin->PinType.PinSubCategoryObject.Get());
-			if ((PinStruct != nullptr) && (StructPtr.Get() == PinStruct) && (Pin->Direction == EGPD_Output))
+			for (UEdGraphPin* Pin : Context.Pins)
 			{
-				UiSpecOut->Category = LOCTEXT("EmptyCategory", "|"); 
-				break;
+				UScriptStruct* PinStruct = Cast<UScriptStruct>(Pin->PinType.PinSubCategoryObject.Get());
+				if ((PinStruct != nullptr) && (StructPtr.Get() == PinStruct) && (Pin->Direction == EGPD_Output))
+				{
+					UiSpecOut->Category = LOCTEXT("EmptyCategory", "|");
+					break;
+				}
 			}
 		}
 	};
 
-	for (TObjectIterator<UScriptStruct> StructIt; StructIt; ++StructIt)
+	UClass* NodeClass = GetClass();
+	ActionRegistrar.RegisterStructActions( FBlueprintActionDatabaseRegistrar::FMakeStructSpawnerDelegate::CreateLambda([NodeClass](const UScriptStruct* Struct)->UBlueprintNodeSpawner*
 	{
-		UScriptStruct const* Struct = (*StructIt);
-		if (!CanBeBroken(Struct, false))
+		UBlueprintFieldNodeSpawner* NodeSpawner = nullptr;
+		
+		if (UK2Node_BreakStruct::CanBeBroken(Struct, /*bIncludeEditAnywhere =*/false))
 		{
-			continue;
-		}
+			NodeSpawner = UBlueprintFieldNodeSpawner::Create(NodeClass, Struct);
+			check(NodeSpawner != nullptr);
+			TWeakObjectPtr<UScriptStruct> NonConstStructPtr = Struct;
+			NodeSpawner->SetNodeFieldDelegate     = UBlueprintFieldNodeSpawner::FSetNodeFieldDelegate::CreateStatic(GetMenuActions_Utils::SetNodeStruct, NonConstStructPtr);
+			NodeSpawner->DynamicUiSignatureGetter = UBlueprintFieldNodeSpawner::FUiSpecOverrideDelegate::CreateStatic(GetMenuActions_Utils::OverrideCategory, NonConstStructPtr);
 
-		// to keep from needlessly instantiating a UBlueprintNodeSpawners, first   
-		// check to make sure that the registrar is looking for actions of this type
-		// (could be regenerating actions for a specific asset, and therefore the 
-		// registrar would only accept actions corresponding to that asset)
-		if (!ActionRegistrar.IsOpenForRegistration(Struct))
-		{
-			continue;
 		}
-
-		UBlueprintFieldNodeSpawner* NodeSpawner = UBlueprintFieldNodeSpawner::Create(GetClass(), Struct);
-		check(NodeSpawner != nullptr);
-		TWeakObjectPtr<UScriptStruct> NonConstStructPtr = Struct;
-		NodeSpawner->SetNodeFieldDelegate = UBlueprintFieldNodeSpawner::FSetNodeFieldDelegate::CreateStatic(SetNodeStructLambda, NonConstStructPtr);
-		NodeSpawner->DynamicUiSignatureGetter = UBlueprintFieldNodeSpawner::FUiSpecOverrideDelegate::CreateStatic(CategoryOverrideLambda, NonConstStructPtr);
-
-		// this struct could belong to a class, or is a user defined struct 
-		// (asset), that's why we want to make sure to register it along with 
-		// the action (so the action knows to refresh when the class/asset is).
-		ActionRegistrar.AddBlueprintAction(Struct, NodeSpawner);
-	}	
+		return NodeSpawner;
+	}) );
 }
 
 FText UK2Node_BreakStruct::GetMenuCategory() const
diff --git a/Engine/Source/Editor/BlueprintGraph/Private/K2Node_CastByteToEnum.cpp b/Engine/Source/Editor/BlueprintGraph/Private/K2Node_CastByteToEnum.cpp
--- a/Engine/Source/Editor/BlueprintGraph/Private/K2Node_CastByteToEnum.cpp
+++ b/Engine/Source/Editor/BlueprintGraph/Private/K2Node_CastByteToEnum.cpp
@@ -173,40 +173,26 @@ bool UK2Node_CastByteToEnum::IsConnectionDisallowed(const UEdGraphPin* MyPin, co
 
 void UK2Node_CastByteToEnum::GetMenuActions(FBlueprintActionDatabaseRegistrar& ActionRegistrar) const
 {
-	auto SetNodeEnumLambda = [](UEdGraphNode* NewNode, UField const* /*EnumField*/, TWeakObjectPtr<UEnum> NonConstEnumPtr)
+	struct GetMenuActions_Utils
 	{
-		UK2Node_CastByteToEnum* EnumNode = CastChecked<UK2Node_CastByteToEnum>(NewNode);
-		EnumNode->Enum  = NonConstEnumPtr.Get();
-		EnumNode->bSafe = true;
-	};
-
-	for (TObjectIterator<UEnum> EnumIt; EnumIt; ++EnumIt)
-	{
-		UEnum const* EnumToConsider = (*EnumIt);
-		if (!UEdGraphSchema_K2::IsAllowableBlueprintVariableType(EnumToConsider))
+		static void SetNodeEnum(UEdGraphNode* NewNode, UField const* /*EnumField*/, TWeakObjectPtr<UEnum> NonConstEnumPtr)
 		{
-			continue;
-		}
-
-		// to keep from needlessly instantiating a UBlueprintNodeSpawners, first   
-		// check to make sure that the registrar is looking for actions of this type
-		// (could be regenerating actions for a specific asset, and therefore the 
-		// registrar would only accept actions corresponding to that asset)
-		if (!ActionRegistrar.IsOpenForRegistration(EnumToConsider))
-		{
-			continue;
+			UK2Node_CastByteToEnum* EnumNode = CastChecked<UK2Node_CastByteToEnum>(NewNode);
+			EnumNode->Enum  = NonConstEnumPtr.Get();
+			EnumNode->bSafe = true;
 		}
+	};
 
-		UBlueprintFieldNodeSpawner* NodeSpawner = UBlueprintFieldNodeSpawner::Create(GetClass(), EnumToConsider);
+	UClass* NodeClass = GetClass();
+	ActionRegistrar.RegisterEnumActions( FBlueprintActionDatabaseRegistrar::FMakeEnumSpawnerDelegate::CreateLambda([NodeClass](const UEnum* Enum)->UBlueprintNodeSpawner*
+	{
+		UBlueprintFieldNodeSpawner* NodeSpawner = UBlueprintFieldNodeSpawner::Create(NodeClass, Enum);
 		check(NodeSpawner != nullptr);
-		TWeakObjectPtr<UEnum> NonConstEnumPtr = EnumToConsider;
-		NodeSpawner->SetNodeFieldDelegate = UBlueprintFieldNodeSpawner::FSetNodeFieldDelegate::CreateStatic(SetNodeEnumLambda, NonConstEnumPtr);
+		TWeakObjectPtr<UEnum> NonConstEnumPtr = Enum;
+		NodeSpawner->SetNodeFieldDelegate = UBlueprintFieldNodeSpawner::FSetNodeFieldDelegate::CreateStatic(GetMenuActions_Utils::SetNodeEnum, NonConstEnumPtr);
 
-		// this enum could belong to a class, or is a user defined enum (asset), 
-		// that's why we want to make sure to register it along with the action 
-		// (so the action can be refreshed when the class/asset is).
-		ActionRegistrar.AddBlueprintAction(EnumToConsider, NodeSpawner);
-	}
+		return NodeSpawner;
+	}) );
 }
 
 FText UK2Node_CastByteToEnum::GetMenuCategory() const
diff --git a/Engine/Source/Editor/BlueprintGraph/Private/K2Node_EnumLiteral.cpp b/Engine/Source/Editor/BlueprintGraph/Private/K2Node_EnumLiteral.cpp
--- a/Engine/Source/Editor/BlueprintGraph/Private/K2Node_EnumLiteral.cpp
+++ b/Engine/Source/Editor/BlueprintGraph/Private/K2Node_EnumLiteral.cpp
@@ -109,39 +109,25 @@ FNodeHandlingFunctor* UK2Node_EnumLiteral::CreateNodeHandler(FKismetCompilerCont
 
 void UK2Node_EnumLiteral::GetMenuActions(FBlueprintActionDatabaseRegistrar& ActionRegistrar) const
 {
-	auto SetNodeEnumLambda = [](UEdGraphNode* NewNode, UField const* /*EnumField*/, TWeakObjectPtr<UEnum> NonConstEnumPtr)
+	struct GetMenuActions_Utils
 	{
-		UK2Node_EnumLiteral* EnumNode = CastChecked<UK2Node_EnumLiteral>(NewNode);
-		EnumNode->Enum = NonConstEnumPtr.Get();
-	};
-
-	for (TObjectIterator<UEnum> EnumIt; EnumIt; ++EnumIt)
-	{
-		UEnum const* EnumToConsider = (*EnumIt);
-		if (!UEdGraphSchema_K2::IsAllowableBlueprintVariableType(EnumToConsider))
+		static void SetNodeEnum(UEdGraphNode* NewNode, UField const* /*EnumField*/, TWeakObjectPtr<UEnum> NonConstEnumPtr)
 		{
-			continue;
-		}
-
-		// to keep from needlessly instantiating a UBlueprintNodeSpawners, first   
-		// check to make sure that the registrar is looking for actions of this type
-		// (could be regenerating actions for a specific asset, and therefore the 
-		// registrar would only accept actions corresponding to that asset)
-		if (!ActionRegistrar.IsOpenForRegistration(EnumToConsider))
-		{
-			continue;
+			UK2Node_EnumLiteral* EnumNode = CastChecked<UK2Node_EnumLiteral>(NewNode);
+			EnumNode->Enum = NonConstEnumPtr.Get();
 		}
+	};
 
-		UBlueprintFieldNodeSpawner* NodeSpawner = UBlueprintFieldNodeSpawner::Create(GetClass(), EnumToConsider);
+	UClass* NodeClass = GetClass();
+	ActionRegistrar.RegisterEnumActions( FBlueprintActionDatabaseRegistrar::FMakeEnumSpawnerDelegate::CreateLambda([NodeClass](const UEnum* Enum)->UBlueprintNodeSpawner*
+	{
+		UBlueprintFieldNodeSpawner* NodeSpawner = UBlueprintFieldNodeSpawner::Create(NodeClass, Enum);
 		check(NodeSpawner != nullptr);
-		TWeakObjectPtr<UEnum> NonConstEnumPtr = EnumToConsider;
-		NodeSpawner->SetNodeFieldDelegate = UBlueprintFieldNodeSpawner::FSetNodeFieldDelegate::CreateStatic(SetNodeEnumLambda, NonConstEnumPtr);
+		TWeakObjectPtr<UEnum> NonConstEnumPtr = Enum;
+		NodeSpawner->SetNodeFieldDelegate = UBlueprintFieldNodeSpawner::FSetNodeFieldDelegate::CreateStatic(GetMenuActions_Utils::SetNodeEnum, NonConstEnumPtr);
 
-		// this enum could belong to a class, or is a user defined enum (asset), 
-		// that's why we want to make sure to register it along with the action 
-		// (so the action can be refreshed when the class/asset is).
-		ActionRegistrar.AddBlueprintAction(EnumToConsider, NodeSpawner);
-	}
+		return NodeSpawner;
+	}) );
 }
 
 FText UK2Node_EnumLiteral::GetMenuCategory() const
diff --git a/Engine/Source/Editor/BlueprintGraph/Private/K2Node_ForEachElementInEnum.cpp b/Engine/Source/Editor/BlueprintGraph/Private/K2Node_ForEachElementInEnum.cpp
--- a/Engine/Source/Editor/BlueprintGraph/Private/K2Node_ForEachElementInEnum.cpp
+++ b/Engine/Source/Editor/BlueprintGraph/Private/K2Node_ForEachElementInEnum.cpp
@@ -225,39 +225,25 @@ void UK2Node_ForEachElementInEnum::ExpandNode(class FKismetCompilerContext& Comp
 
 void UK2Node_ForEachElementInEnum::GetMenuActions(FBlueprintActionDatabaseRegistrar& ActionRegistrar) const
 {
-	auto SetNodeEnumLambda = [](UEdGraphNode* NewNode, UField const* /*EnumField*/, TWeakObjectPtr<UEnum> NonConstEnumPtr)
+	struct GetMenuActions_Utils
 	{
-		UK2Node_ForEachElementInEnum* EnumNode = CastChecked<UK2Node_ForEachElementInEnum>(NewNode);
-		EnumNode->Enum = NonConstEnumPtr.Get();
-	};
-
-	for (TObjectIterator<UEnum> EnumIt; EnumIt; ++EnumIt)
-	{
-		UEnum const* EnumToConsider = (*EnumIt);
-		if (!UEdGraphSchema_K2::IsAllowableBlueprintVariableType(EnumToConsider))
+		static void SetNodeEnum(UEdGraphNode* NewNode, UField const* /*EnumField*/, TWeakObjectPtr<UEnum> NonConstEnumPtr)
 		{
-			continue;
-		}
-
-		// to keep from needlessly instantiating a UBlueprintNodeSpawners, first   
-		// check to make sure that the registrar is looking for actions of this type
-		// (could be regenerating actions for a specific asset, and therefore the 
-		// registrar would only accept actions corresponding to that asset)
-		if (!ActionRegistrar.IsOpenForRegistration(EnumToConsider))
-		{
-			continue;
+			UK2Node_ForEachElementInEnum* EnumNode = CastChecked<UK2Node_ForEachElementInEnum>(NewNode);
+			EnumNode->Enum = NonConstEnumPtr.Get();
 		}
+	};
 
-		UBlueprintFieldNodeSpawner* NodeSpawner = UBlueprintFieldNodeSpawner::Create(GetClass(), EnumToConsider);
+	UClass* NodeClass = GetClass();
+	ActionRegistrar.RegisterEnumActions( FBlueprintActionDatabaseRegistrar::FMakeEnumSpawnerDelegate::CreateLambda([NodeClass](const UEnum* Enum)->UBlueprintNodeSpawner*
+	{
+		UBlueprintFieldNodeSpawner* NodeSpawner = UBlueprintFieldNodeSpawner::Create(NodeClass, Enum);
 		check(NodeSpawner != nullptr);
-		TWeakObjectPtr<UEnum> NonConstEnumPtr = EnumToConsider;
-		NodeSpawner->SetNodeFieldDelegate = UBlueprintFieldNodeSpawner::FSetNodeFieldDelegate::CreateStatic(SetNodeEnumLambda, NonConstEnumPtr);
+		TWeakObjectPtr<UEnum> NonConstEnumPtr = Enum;
+		NodeSpawner->SetNodeFieldDelegate = UBlueprintFieldNodeSpawner::FSetNodeFieldDelegate::CreateStatic(GetMenuActions_Utils::SetNodeEnum, NonConstEnumPtr);
 
-		// this enum could belong to a class, or is a user defined enum (asset), 
-		// that's why we want to make sure to register it along with the action 
-		// (so the action can be refreshed when the class/asset is).
-		ActionRegistrar.AddBlueprintAction(EnumToConsider, NodeSpawner);
-	}
+		return NodeSpawner;
+	}) );
 }
 
 FText UK2Node_ForEachElementInEnum::GetMenuCategory() const
diff --git a/Engine/Source/Editor/BlueprintGraph/Private/K2Node_GetInputAxisKeyValue.cpp b/Engine/Source/Editor/BlueprintGraph/Private/K2Node_GetInputAxisKeyValue.cpp
--- a/Engine/Source/Editor/BlueprintGraph/Private/K2Node_GetInputAxisKeyValue.cpp
+++ b/Engine/Source/Editor/BlueprintGraph/Private/K2Node_GetInputAxisKeyValue.cpp
@@ -131,27 +131,26 @@ void UK2Node_GetInputAxisKeyValue::GetMenuActions(FBlueprintActionDatabaseRegist
 	// type disappears, then the action should go with it)
 	UClass* ActionKey = GetClass();
 
-	for (FKey const Key : AllKeys)
+	// to keep from needlessly instantiating a UBlueprintNodeSpawner (and 
+	// iterating over keys), first check to make sure that the registrar is 
+	// looking for actions of this type (could be regenerating actions for a 
+	// specific asset, and therefore the registrar would only accept actions 
+	// corresponding to that asset)
+	if (ActionRegistrar.IsOpenForRegistration(ActionKey))
 	{
-		if (!Key.IsBindableInBlueprints() || !Key.IsFloatAxis())
+		for (FKey const Key : AllKeys)
 		{
-			continue;
-		}
+			if (!Key.IsBindableInBlueprints() || !Key.IsFloatAxis())
+			{
+				continue;
+			}
 
-		// to keep from needlessly instantiating a UBlueprintNodeSpawner, first   
-		// check to make sure that the registrar is looking for actions of this type
-		// (could be regenerating actions for a specific asset, and therefore the 
-		// registrar would only accept actions corresponding to that asset)
-		if (!ActionRegistrar.IsOpenForRegistration(ActionKey))
-		{
-			continue;
+			UBlueprintNodeSpawner* NodeSpawner = UBlueprintNodeSpawner::Create(GetClass());
+			check(NodeSpawner != nullptr);
+
+			NodeSpawner->CustomizeNodeDelegate = UBlueprintNodeSpawner::FCustomizeNodeDelegate::CreateStatic(CustomizeInputNodeLambda, Key);
+			ActionRegistrar.AddBlueprintAction(ActionKey, NodeSpawner);
 		}
-		
-		UBlueprintNodeSpawner* NodeSpawner = UBlueprintNodeSpawner::Create(GetClass());
-		check(NodeSpawner != nullptr);
-		
-		NodeSpawner->CustomizeNodeDelegate = UBlueprintNodeSpawner::FCustomizeNodeDelegate::CreateStatic(CustomizeInputNodeLambda, Key);
-		ActionRegistrar.AddBlueprintAction(ActionKey, NodeSpawner);
 	}
 }
 
diff --git a/Engine/Source/Editor/BlueprintGraph/Private/K2Node_GetInputVectorAxisValue.cpp b/Engine/Source/Editor/BlueprintGraph/Private/K2Node_GetInputVectorAxisValue.cpp
--- a/Engine/Source/Editor/BlueprintGraph/Private/K2Node_GetInputVectorAxisValue.cpp
+++ b/Engine/Source/Editor/BlueprintGraph/Private/K2Node_GetInputVectorAxisValue.cpp
@@ -79,27 +79,26 @@ void UK2Node_GetInputVectorAxisValue::GetMenuActions(FBlueprintActionDatabaseReg
 	// type disappears, then the action should go with it)
 	UClass* ActionKey = GetClass();
 
-	for (FKey const Key : AllKeys)
+	// to keep from needlessly instantiating a UBlueprintNodeSpawner (and 
+	// iterating over keys), first check to make sure that the registrar is 
+	// looking for actions of this type (could be regenerating actions for a 
+	// specific asset, and therefore the registrar would only accept actions 
+	// corresponding to that asset)
+	if (ActionRegistrar.IsOpenForRegistration(ActionKey))
 	{
-		if (!Key.IsBindableInBlueprints() || !Key.IsVectorAxis())
+		for (FKey const Key : AllKeys)
 		{
-			continue;
-		}
+			if (!Key.IsBindableInBlueprints() || !Key.IsVectorAxis())
+			{
+				continue;
+			}
 
-		// to keep from needlessly instantiating a UBlueprintNodeSpawner, first   
-		// check to make sure that the registrar is looking for actions of this type
-		// (could be regenerating actions for a specific asset, and therefore the 
-		// registrar would only accept actions corresponding to that asset)
-		if (!ActionRegistrar.IsOpenForRegistration(ActionKey))
-		{
-			continue;
-		}
+			UBlueprintNodeSpawner* NodeSpawner = UBlueprintNodeSpawner::Create(GetClass());
+			check(NodeSpawner != nullptr);
+			NodeSpawner->DefaultMenuSignature.MenuName = FText::Format(NSLOCTEXT("K2Node_GetInputVectorAxisValue", "MenuName", "Get {0}"), Key.GetDisplayName());
 
-		UBlueprintNodeSpawner* NodeSpawner = UBlueprintNodeSpawner::Create(GetClass());
-		check(NodeSpawner != nullptr);
-		NodeSpawner->DefaultMenuSignature.MenuName = FText::Format(NSLOCTEXT("K2Node_GetInputVectorAxisValue", "MenuName", "Get {0}"), Key.GetDisplayName());
-
-		NodeSpawner->CustomizeNodeDelegate = UBlueprintNodeSpawner::FCustomizeNodeDelegate::CreateStatic(CustomizeInputNodeLambda, Key);
-		ActionRegistrar.AddBlueprintAction(ActionKey, NodeSpawner);
+			NodeSpawner->CustomizeNodeDelegate = UBlueprintNodeSpawner::FCustomizeNodeDelegate::CreateStatic(CustomizeInputNodeLambda, Key);
+			ActionRegistrar.AddBlueprintAction(ActionKey, NodeSpawner);
+		}
 	}
 }
diff --git a/Engine/Source/Editor/BlueprintGraph/Private/K2Node_GetNumEnumEntries.cpp b/Engine/Source/Editor/BlueprintGraph/Private/K2Node_GetNumEnumEntries.cpp
--- a/Engine/Source/Editor/BlueprintGraph/Private/K2Node_GetNumEnumEntries.cpp
+++ b/Engine/Source/Editor/BlueprintGraph/Private/K2Node_GetNumEnumEntries.cpp
@@ -95,39 +95,25 @@ void UK2Node_GetNumEnumEntries::ExpandNode(class FKismetCompilerContext& Compile
 
 void UK2Node_GetNumEnumEntries::GetMenuActions(FBlueprintActionDatabaseRegistrar& ActionRegistrar) const
 {
-	auto SetNodeEnumLambda = [](UEdGraphNode* NewNode, UField const* /*EnumField*/, TWeakObjectPtr<UEnum> NonConstEnumPtr)
+	struct GetMenuActions_Utils
 	{
-		UK2Node_GetNumEnumEntries* EnumNode = CastChecked<UK2Node_GetNumEnumEntries>(NewNode);
-		EnumNode->Enum = NonConstEnumPtr.Get();
-	};
-
-	for (TObjectIterator<UEnum> EnumIt; EnumIt; ++EnumIt)
-	{
-		UEnum const* EnumToConsider = (*EnumIt);
-		if (!UEdGraphSchema_K2::IsAllowableBlueprintVariableType(EnumToConsider))
+		static void SetNodeEnum(UEdGraphNode* NewNode, UField const* /*EnumField*/, TWeakObjectPtr<UEnum> NonConstEnumPtr)
 		{
-			continue;
-		}
-
-		// to keep from needlessly instantiating a UBlueprintNodeSpawners, first   
-		// check to make sure that the registrar is looking for actions of this type
-		// (could be regenerating actions for a specific asset, and therefore the 
-		// registrar would only accept actions corresponding to that asset)
-		if (!ActionRegistrar.IsOpenForRegistration(EnumToConsider))
-		{
-			continue;
+			UK2Node_GetNumEnumEntries* EnumNode = CastChecked<UK2Node_GetNumEnumEntries>(NewNode);
+			EnumNode->Enum = NonConstEnumPtr.Get();
 		}
+	};
 
-		UBlueprintFieldNodeSpawner* NodeSpawner = UBlueprintFieldNodeSpawner::Create(GetClass(), EnumToConsider);
+	UClass* NodeClass = GetClass();
+	ActionRegistrar.RegisterEnumActions( FBlueprintActionDatabaseRegistrar::FMakeEnumSpawnerDelegate::CreateLambda([NodeClass](const UEnum* Enum)->UBlueprintNodeSpawner*
+	{
+		UBlueprintFieldNodeSpawner* NodeSpawner = UBlueprintFieldNodeSpawner::Create(NodeClass, Enum);
 		check(NodeSpawner != nullptr);
-		TWeakObjectPtr<UEnum> NonConstEnumPtr = EnumToConsider;
-		NodeSpawner->SetNodeFieldDelegate = UBlueprintFieldNodeSpawner::FSetNodeFieldDelegate::CreateStatic(SetNodeEnumLambda, NonConstEnumPtr);
+		TWeakObjectPtr<UEnum> NonConstEnumPtr = Enum;
+		NodeSpawner->SetNodeFieldDelegate = UBlueprintFieldNodeSpawner::FSetNodeFieldDelegate::CreateStatic(GetMenuActions_Utils::SetNodeEnum, NonConstEnumPtr);
 
-		// this enum could belong to a class, or is a user defined enum (asset), 
-		// that's why we want to make sure to register it along with the action 
-		// (so the action can be refreshed when the class/asset is).
-		ActionRegistrar.AddBlueprintAction(EnumToConsider, NodeSpawner);
-	}
+		return NodeSpawner;
+	}) );
 }
 
 FText UK2Node_GetNumEnumEntries::GetMenuCategory() const
diff --git a/Engine/Source/Editor/BlueprintGraph/Private/K2Node_InputAxisKeyEvent.cpp b/Engine/Source/Editor/BlueprintGraph/Private/K2Node_InputAxisKeyEvent.cpp
--- a/Engine/Source/Editor/BlueprintGraph/Private/K2Node_InputAxisKeyEvent.cpp
+++ b/Engine/Source/Editor/BlueprintGraph/Private/K2Node_InputAxisKeyEvent.cpp
@@ -134,27 +134,26 @@ void UK2Node_InputAxisKeyEvent::GetMenuActions(FBlueprintActionDatabaseRegistrar
 	// type disappears, then the action should go with it)
 	UClass* ActionKey = GetClass();
 
-	for (FKey const Key : AllKeys)
+	// to keep from needlessly instantiating a UBlueprintNodeSpawner (and 
+	// iterating over keys), first check to make sure that the registrar is 
+	// looking for actions of this type (could be regenerating actions for a 
+	// specific asset, and therefore the registrar would only accept actions 
+	// corresponding to that asset)
+	if (ActionRegistrar.IsOpenForRegistration(ActionKey))
 	{
-		if (!Key.IsBindableInBlueprints() || !Key.IsFloatAxis())
+		for (FKey const Key : AllKeys)
 		{
-			continue;
-		}
+			if (!Key.IsBindableInBlueprints() || !Key.IsFloatAxis())
+			{
+				continue;
+			}
 
-		// to keep from needlessly instantiating a UBlueprintNodeSpawner, first   
-		// check to make sure that the registrar is looking for actions of this type
-		// (could be regenerating actions for a specific asset, and therefore the 
-		// registrar would only accept actions corresponding to that asset)
-		if (!ActionRegistrar.IsOpenForRegistration(ActionKey))
-		{
-			continue;
-		}
+			UBlueprintNodeSpawner* NodeSpawner = UBlueprintNodeSpawner::Create(GetClass());
+			check(NodeSpawner != nullptr);
 
-		UBlueprintNodeSpawner* NodeSpawner = UBlueprintNodeSpawner::Create(GetClass());
-		check(NodeSpawner != nullptr);
-
-		NodeSpawner->CustomizeNodeDelegate = UBlueprintNodeSpawner::FCustomizeNodeDelegate::CreateStatic(CustomizeInputNodeLambda, Key);
-		ActionRegistrar.AddBlueprintAction(ActionKey, NodeSpawner);
+			NodeSpawner->CustomizeNodeDelegate = UBlueprintNodeSpawner::FCustomizeNodeDelegate::CreateStatic(CustomizeInputNodeLambda, Key);
+			ActionRegistrar.AddBlueprintAction(ActionKey, NodeSpawner);
+		}
 	}
 }
 
diff --git a/Engine/Source/Editor/BlueprintGraph/Private/K2Node_InputKey.cpp b/Engine/Source/Editor/BlueprintGraph/Private/K2Node_InputKey.cpp
--- a/Engine/Source/Editor/BlueprintGraph/Private/K2Node_InputKey.cpp
+++ b/Engine/Source/Editor/BlueprintGraph/Private/K2Node_InputKey.cpp
@@ -304,28 +304,27 @@ void UK2Node_InputKey::GetMenuActions(FBlueprintActionDatabaseRegistrar& ActionR
 	// type disappears, then the action should go with it)
 	UClass* ActionKey = GetClass();
 
-	for (FKey const Key : AllKeys)
+	// to keep from needlessly instantiating a UBlueprintNodeSpawner (and 
+	// iterating over keys), first check to make sure that the registrar is 
+	// looking for actions of this type (could be regenerating actions for a 
+	// specific asset, and therefore the registrar would only accept actions 
+	// corresponding to that asset)
+	if (ActionRegistrar.IsOpenForRegistration(ActionKey))
 	{
-		// these will be handled by UK2Node_GetInputAxisKeyValue and UK2Node_GetInputVectorAxisValue respectively
-		if (!Key.IsBindableInBlueprints() || Key.IsFloatAxis() || Key.IsVectorAxis())
+		for (FKey const Key : AllKeys)
 		{
-			continue;
-		}
+			// these will be handled by UK2Node_GetInputAxisKeyValue and UK2Node_GetInputVectorAxisValue respectively
+			if (!Key.IsBindableInBlueprints() || Key.IsFloatAxis() || Key.IsVectorAxis())
+			{
+				continue;
+			}
 
-		// to keep from needlessly instantiating a UBlueprintNodeSpawner, first   
-		// check to make sure that the registrar is looking for actions of this type
-		// (could be regenerating actions for a specific asset, and therefore the 
-		// registrar would only accept actions corresponding to that asset)
-		if (!ActionRegistrar.IsOpenForRegistration(ActionKey))
-		{
-			continue;
-		}
+			UBlueprintNodeSpawner* NodeSpawner = UBlueprintNodeSpawner::Create(GetClass());
+			check(NodeSpawner != nullptr);
 
-		UBlueprintNodeSpawner* NodeSpawner = UBlueprintNodeSpawner::Create(GetClass());
-		check(NodeSpawner != nullptr);
-
-		NodeSpawner->CustomizeNodeDelegate = UBlueprintNodeSpawner::FCustomizeNodeDelegate::CreateStatic(CustomizeInputNodeLambda, Key);
-		ActionRegistrar.AddBlueprintAction(ActionKey, NodeSpawner);
+			NodeSpawner->CustomizeNodeDelegate = UBlueprintNodeSpawner::FCustomizeNodeDelegate::CreateStatic(CustomizeInputNodeLambda, Key);
+			ActionRegistrar.AddBlueprintAction(ActionKey, NodeSpawner);
+		}
 	}
 }
 
diff --git a/Engine/Source/Editor/BlueprintGraph/Private/K2Node_InputVectorAxisEvent.cpp b/Engine/Source/Editor/BlueprintGraph/Private/K2Node_InputVectorAxisEvent.cpp
--- a/Engine/Source/Editor/BlueprintGraph/Private/K2Node_InputVectorAxisEvent.cpp
+++ b/Engine/Source/Editor/BlueprintGraph/Private/K2Node_InputVectorAxisEvent.cpp
@@ -80,27 +80,26 @@ void UK2Node_InputVectorAxisEvent::GetMenuActions(FBlueprintActionDatabaseRegist
 	// mutated (or removed)... here we use the node's class (so if the node 
 	// type disappears, then the action should go with it)
 	UClass* ActionKey = GetClass();
-	
-	for (FKey const Key : AllKeys)
+
+	// to keep from needlessly instantiating a UBlueprintNodeSpawner (and 
+	// iterating over keys), first check to make sure that the registrar is 
+	// looking for actions of this type (could be regenerating actions for a 
+	// specific asset, and therefore the registrar would only accept actions 
+	// corresponding to that asset)
+	if (ActionRegistrar.IsOpenForRegistration(ActionKey))
 	{
-		if (!Key.IsBindableInBlueprints() || !Key.IsVectorAxis())
+		for (FKey const Key : AllKeys)
 		{
-			continue;
-		}
+			if (!Key.IsBindableInBlueprints() || !Key.IsVectorAxis())
+			{
+				continue;
+			}
 
-		// to keep from needlessly instantiating a UBlueprintNodeSpawner, first   
-		// check to make sure that the registrar is looking for actions of this type
-		// (could be regenerating actions for a specific asset, and therefore the 
-		// registrar would only accept actions corresponding to that asset)
-		if (!ActionRegistrar.IsOpenForRegistration(ActionKey))
-		{
-			continue;
+			UBlueprintNodeSpawner* NodeSpawner = UBlueprintNodeSpawner::Create(GetClass());
+			check(NodeSpawner != nullptr);
+
+			NodeSpawner->CustomizeNodeDelegate = UBlueprintNodeSpawner::FCustomizeNodeDelegate::CreateStatic(CustomizeInputNodeLambda, Key);
+			ActionRegistrar.AddBlueprintAction(ActionKey, NodeSpawner);
 		}
-		
-		UBlueprintNodeSpawner* NodeSpawner = UBlueprintNodeSpawner::Create(GetClass());
-		check(NodeSpawner != nullptr);
-		
-		NodeSpawner->CustomizeNodeDelegate = UBlueprintNodeSpawner::FCustomizeNodeDelegate::CreateStatic(CustomizeInputNodeLambda, Key);
-		ActionRegistrar.AddBlueprintAction(ActionKey, NodeSpawner);
 	}
 }
diff --git a/Engine/Source/Editor/BlueprintGraph/Private/K2Node_MakeStruct.cpp b/Engine/Source/Editor/BlueprintGraph/Private/K2Node_MakeStruct.cpp
--- a/Engine/Source/Editor/BlueprintGraph/Private/K2Node_MakeStruct.cpp
+++ b/Engine/Source/Editor/BlueprintGraph/Private/K2Node_MakeStruct.cpp
@@ -287,53 +287,43 @@ UK2Node::ERedirectType UK2Node_MakeStruct::DoPinsMatchForReconstruction(const UE
 
 void UK2Node_MakeStruct::GetMenuActions(FBlueprintActionDatabaseRegistrar& ActionRegistrar) const
 {
-	auto SetNodeStructLambda = [](UEdGraphNode* NewNode, UField const* /*StructField*/, TWeakObjectPtr<UScriptStruct> NonConstStructPtr)
+	struct GetMenuActions_Utils
 	{
-		UK2Node_MakeStruct* MakeNode = CastChecked<UK2Node_MakeStruct>(NewNode);
-		MakeNode->StructType = NonConstStructPtr.Get();
-	};
+		static void SetNodeStruct(UEdGraphNode* NewNode, UField const* /*StructField*/, TWeakObjectPtr<UScriptStruct> NonConstStructPtr)
+		{
+			UK2Node_MakeStruct* MakeNode = CastChecked<UK2Node_MakeStruct>(NewNode);
+			MakeNode->StructType = NonConstStructPtr.Get();
+		}
 
-	auto CategoryOverrideLambda = [](FBlueprintActionContext const& Context, IBlueprintNodeBinder::FBindingSet const& /*Bindings*/, FBlueprintActionUiSpec* UiSpecOut, TWeakObjectPtr<UScriptStruct> StructPtr)
-	{
-		for (UEdGraphPin* Pin : Context.Pins)
+		static void OverrideCategory(FBlueprintActionContext const& Context, IBlueprintNodeBinder::FBindingSet const& /*Bindings*/, FBlueprintActionUiSpec* UiSpecOut, TWeakObjectPtr<UScriptStruct> StructPtr)
 		{
-			UScriptStruct* PinStruct = Cast<UScriptStruct>(Pin->PinType.PinSubCategoryObject.Get());
-			if ((PinStruct != nullptr) && (StructPtr.Get() == PinStruct) && (Pin->Direction == EGPD_Input))
+			for (UEdGraphPin* Pin : Context.Pins)
 			{
-				UiSpecOut->Category = LOCTEXT("EmptyCategory", "|");
-				break;
+				UScriptStruct* PinStruct = Cast<UScriptStruct>(Pin->PinType.PinSubCategoryObject.Get());
+				if ((PinStruct != nullptr) && (StructPtr.Get() == PinStruct) && (Pin->Direction == EGPD_Input))
+				{
+					UiSpecOut->Category = LOCTEXT("EmptyCategory", "|");
+					break;
+				}
 			}
 		}
 	};
 
-	for (TObjectIterator<UScriptStruct> StructIt; StructIt; ++StructIt)
+	UClass* NodeClass = GetClass();
+	ActionRegistrar.RegisterStructActions( FBlueprintActionDatabaseRegistrar::FMakeStructSpawnerDelegate::CreateLambda([NodeClass](const UScriptStruct* Struct)->UBlueprintNodeSpawner*
 	{
-		UScriptStruct const* Struct = (*StructIt);
-		if (!CanBeMade(Struct, false))
-		{
-			continue;
-		}
-
-		// to keep from needlessly instantiating a UBlueprintNodeSpawners, first   
-		// check to make sure that the registrar is looking for actions of this type
-		// (could be regenerating actions for a specific asset, and therefore the 
-		// registrar would only accept actions corresponding to that asset)
-		if (!ActionRegistrar.IsOpenForRegistration(Struct))
+		UBlueprintFieldNodeSpawner* NodeSpawner = nullptr;
+		
+		if (UK2Node_MakeStruct::CanBeMade(Struct, /*bIncludeEditOnly =*/false))
 		{
-			continue;
+			NodeSpawner = UBlueprintFieldNodeSpawner::Create(NodeClass, Struct);
+			check(NodeSpawner != nullptr);
+			TWeakObjectPtr<UScriptStruct> NonConstStructPtr = Struct;
+			NodeSpawner->SetNodeFieldDelegate     = UBlueprintFieldNodeSpawner::FSetNodeFieldDelegate::CreateStatic(GetMenuActions_Utils::SetNodeStruct, NonConstStructPtr);
+			NodeSpawner->DynamicUiSignatureGetter = UBlueprintFieldNodeSpawner::FUiSpecOverrideDelegate::CreateStatic(GetMenuActions_Utils::OverrideCategory, NonConstStructPtr);
 		}
-
-		UBlueprintFieldNodeSpawner* NodeSpawner = UBlueprintFieldNodeSpawner::Create(GetClass(), Struct);
-		check(NodeSpawner != nullptr);
-		TWeakObjectPtr<UScriptStruct> NonConstStructPtr = Struct;
-		NodeSpawner->SetNodeFieldDelegate = UBlueprintFieldNodeSpawner::FSetNodeFieldDelegate::CreateStatic(SetNodeStructLambda, NonConstStructPtr);
-		NodeSpawner->DynamicUiSignatureGetter = UBlueprintFieldNodeSpawner::FUiSpecOverrideDelegate::CreateStatic(CategoryOverrideLambda, NonConstStructPtr);
-
-		// this struct could belong to a class, or is a user defined struct 
-		// (asset), that's why we want to make sure to register it along with 
-		// the action (so the action knows to refresh when the class/asset is).
-		ActionRegistrar.AddBlueprintAction(Struct, NodeSpawner);
-	}
+		return NodeSpawner;
+	}) );
 }
 
 FText UK2Node_MakeStruct::GetMenuCategory() const
diff --git a/Engine/Source/Editor/BlueprintGraph/Private/K2Node_SwitchEnum.cpp b/Engine/Source/Editor/BlueprintGraph/Private/K2Node_SwitchEnum.cpp
--- a/Engine/Source/Editor/BlueprintGraph/Private/K2Node_SwitchEnum.cpp
+++ b/Engine/Source/Editor/BlueprintGraph/Private/K2Node_SwitchEnum.cpp
@@ -85,39 +85,25 @@ bool UK2Node_SwitchEnum::IsConnectionDisallowed(const UEdGraphPin* MyPin, const
 
 void UK2Node_SwitchEnum::GetMenuActions(FBlueprintActionDatabaseRegistrar& ActionRegistrar) const
 {
-	auto SetNodeEnumLambda = [](UEdGraphNode* NewNode, UField const* /*EnumField*/, TWeakObjectPtr<UEnum> NonConstEnumPtr)
+	struct GetMenuActions_Utils
 	{
-		UK2Node_SwitchEnum* EnumNode = CastChecked<UK2Node_SwitchEnum>(NewNode);
-		EnumNode->Enum = NonConstEnumPtr.Get();
-	};
-
-	for (TObjectIterator<UEnum> EnumIt; EnumIt; ++EnumIt)
-	{
-		UEnum const* EnumToConsider = (*EnumIt);
-		if (!UEdGraphSchema_K2::IsAllowableBlueprintVariableType(EnumToConsider))
+		static void SetNodeEnum(UEdGraphNode* NewNode, UField const* /*EnumField*/, TWeakObjectPtr<UEnum> NonConstEnumPtr)
 		{
-			continue;
-		}
-
-		// to keep from needlessly instantiating a UBlueprintNodeSpawners, first   
-		// check to make sure that the registrar is looking for actions of this type
-		// (could be regenerating actions for a specific asset, and therefore the 
-		// registrar would only accept actions corresponding to that asset)
-		if (!ActionRegistrar.IsOpenForRegistration(EnumToConsider))
-		{
-			continue;
+			UK2Node_SwitchEnum* EnumNode = CastChecked<UK2Node_SwitchEnum>(NewNode);
+			EnumNode->Enum = NonConstEnumPtr.Get();
 		}
+	};
 
-		UBlueprintFieldNodeSpawner* NodeSpawner = UBlueprintFieldNodeSpawner::Create(GetClass(), EnumToConsider);
+	UClass* NodeClass = GetClass();
+	ActionRegistrar.RegisterEnumActions( FBlueprintActionDatabaseRegistrar::FMakeEnumSpawnerDelegate::CreateLambda([NodeClass](const UEnum* Enum)->UBlueprintNodeSpawner*
+	{
+		UBlueprintFieldNodeSpawner* NodeSpawner = UBlueprintFieldNodeSpawner::Create(NodeClass, Enum);
 		check(NodeSpawner != nullptr);
-		TWeakObjectPtr<UEnum> NonConstEnumPtr = EnumToConsider;
-		NodeSpawner->SetNodeFieldDelegate = UBlueprintFieldNodeSpawner::FSetNodeFieldDelegate::CreateStatic(SetNodeEnumLambda, NonConstEnumPtr);
+		TWeakObjectPtr<UEnum> NonConstEnumPtr = Enum;
+		NodeSpawner->SetNodeFieldDelegate = UBlueprintFieldNodeSpawner::FSetNodeFieldDelegate::CreateStatic(GetMenuActions_Utils::SetNodeEnum, NonConstEnumPtr);
 
-		// this enum could belong to a class, or is a user defined enum (asset), 
-		// that's why we want to make sure to register it along with the action 
-		// (so the action can be refreshed when the class/asset is).
-		ActionRegistrar.AddBlueprintAction(EnumToConsider, NodeSpawner);
-	}
+		return NodeSpawner;
+	}) );
 }
 
 void UK2Node_SwitchEnum::CreateSelectionPin()
diff --git a/Engine/Source/Editor/BlueprintGraph/Public/BlueprintActionDatabaseRegistrar.h b/Engine/Source/Editor/BlueprintGraph/Public/BlueprintActionDatabaseRegistrar.h
--- a/Engine/Source/Editor/BlueprintGraph/Public/BlueprintActionDatabaseRegistrar.h
+++ b/Engine/Source/Editor/BlueprintGraph/Public/BlueprintActionDatabaseRegistrar.h
@@ -14,18 +14,18 @@ class BLUEPRINTGRAPH_API FBlueprintActionDatabaseRegistrar
 public:
 	/**
 	 * Attempts to suss out the key that this action should be registered under;
-	 * if it doesn't find a better one then it associates the action with the 
+	 * if it doesn't find a better one then it associates the action with the
 	 * node filling this out.
-	 * 
+	 *
 	 * @param  NodeSpawner	The new node spawner that you wish to register.
 	 */
 	bool AddBlueprintAction(UBlueprintNodeSpawner* NodeSpawner);
 
 	/**
-	 * Each action should be recorded under a specific UField key; primarily to 
+	 * Each action should be recorded under a specific UField key; primarily to
 	 * refresh those actions when the corresponding asset is updated (Blueprint
 	 * regenerated, struct added/deleted, etc.).
-	 * 
+	 *
 	 * @param  FieldOwner	A field object that the node is associated with (flags when this action should be updated).
 	 * @param  NodeSpawner	The new node spawner that you wish to register.
 	 */
@@ -36,23 +36,23 @@ public:
 	bool AddBlueprintAction(FAssetData const& AssetDataOwner, UBlueprintNodeSpawner* NodeSpawner);
 
 	/**
-	 * Special case for asset bound actions (we want to clean-up/refresh these 
+	 * Special case for asset bound actions (we want to clean-up/refresh these
 	 * when the corresponding asset is updated). AssetOwner must be an asset!
-	 * 
+	 *
 	 * @param  AssetOwner	An asset object that the node is associated with (flags when this action should be updated).
 	 * @param  NodeSpawner	The new node spawner that you wish to register.
 	 */
 	bool AddBlueprintAction(UObject const* AssetOwner, UBlueprintNodeSpawner* NodeSpawner);
 
 	/**
-	 * Occasionally (when an asset is added/refreshed), this registrar will be 
-	 * passed around to gather only specific keyed actions (see ActionKeyFilter). 
-	 * In that case, it will block registration of all unwanted keys.  
-	 * Functionality wise this doesn't matter to UK2Node, but UK2Node may be  
+	 * Occasionally (when an asset is added/refreshed), this registrar will be
+	 * passed around to gather only specific keyed actions (see ActionKeyFilter).
+	 * In that case, it will block registration of all unwanted keys.
+	 * Functionality wise this doesn't matter to UK2Node, but UK2Node may be
 	 * able to save on some work/allocations if it knew this beforehand.
-	 * 
+	 *
 	 * @param  OwnerKey		The key you wish to register your action(s) under.
-	 * @return True if the OwnerKey would is allowed to register actions, false if it would be blocked. 
+	 * @return True if the OwnerKey would is allowed to register actions, false if it would be blocked.
 	 */
 	bool IsOpenForRegistration(UObject const* OwnerKey);
 	bool IsOpenForRegistration(FAssetData const& AssetDataOwner);
@@ -63,6 +63,40 @@ public:
 		return ActionKeyFilter;
 	}
 
+	/// 
+	DECLARE_DELEGATE_RetVal_OneParam(UBlueprintNodeSpawner*, FMakeStructSpawnerDelegate, const UScriptStruct*);
+	/**
+	*
+	*
+	* @param  MakeActionCallback
+	* @return
+	*/
+	int32 RegisterStructActions(const FMakeStructSpawnerDelegate& MakeActionCallback);
+
+	/// 
+	DECLARE_DELEGATE_RetVal_OneParam(UBlueprintNodeSpawner*, FMakeEnumSpawnerDelegate, const UEnum*);
+	/**
+	 *
+	 *
+	 * @param  MakeActionCallback
+	 * @return
+	 */
+	int32 RegisterEnumActions(const FMakeEnumSpawnerDelegate& MakeActionCallback);
+
+	/// 
+	DECLARE_DELEGATE_RetVal_OneParam(UBlueprintNodeSpawner*, FMakeFuncSpawnerDelegate, const UFunction*);
+	/**
+	 *
+	 *
+	 * @param  MakeActionCallback
+	 * @return
+	 */
+	template<class T>
+	int32 RegisterClassFactoryActions(const FMakeFuncSpawnerDelegate& MakeActionCallback)
+	{
+		return RegisterClassFactoryActions(T::StaticClass(), MakeActionCallback);
+	}
+
 private:
 	typedef FBlueprintActionDatabase::FActionRegistry			FActionRegistry;
 	typedef FBlueprintActionDatabase::FUnloadedActionRegistry	FUnloadedActionRegistry;
@@ -72,6 +106,15 @@ private:
 	friend class FBlueprintActionDatabase;
 	FBlueprintActionDatabaseRegistrar(FActionRegistry& Database, FUnloadedActionRegistry& UnloadedDatabase, FPrimingQueue& PrimingQueue, TSubclassOf<UEdGraphNode> DefaultKey = nullptr);
 
+	/**
+	 * 
+	 * 
+	 * @param  TargetType    
+	 * @param  MakeActionCallback    
+	 * @return 
+	 */
+	int32 RegisterClassFactoryActions(const UClass* TargetType, const FMakeFuncSpawnerDelegate& MakeActionCallback);
+
 	/**
 	 * Internal method that actually adds the action to the database.
 	 * 
diff --git a/Engine/Source/Editor/GameplayAbilitiesEditor/Private/K2Node_LatentAbilityCall.cpp b/Engine/Source/Editor/GameplayAbilitiesEditor/Private/K2Node_LatentAbilityCall.cpp
--- a/Engine/Source/Editor/GameplayAbilitiesEditor/Private/K2Node_LatentAbilityCall.cpp
+++ b/Engine/Source/Editor/GameplayAbilitiesEditor/Private/K2Node_LatentAbilityCall.cpp
@@ -54,67 +54,35 @@ bool UK2Node_LatentAbilityCall::IsCompatibleWithGraph(UEdGraph const* TargetGrap
 
 void UK2Node_LatentAbilityCall::GetMenuActions(FBlueprintActionDatabaseRegistrar& ActionRegistrar) const
 {
-	// these nested loops are combing over the same classes/functions the
-	// FBlueprintActionDatabase does; ideally we save on perf and fold this in
-	// with FBlueprintActionDatabase, but we want to keep the modules separate
-	for (TObjectIterator<UClass> ClassIt; ClassIt; ++ClassIt)
+	struct GetMenuActions_Utils
 	{
-		UClass* Class = *ClassIt;
-		if (!Class->IsChildOf<UAbilityTask>() || Class->HasAnyClassFlags(CLASS_Abstract))
+		static void SetNodeFunc(UEdGraphNode* NewNode, bool /*bIsTemplateNode*/, TWeakObjectPtr<UFunction> FunctionPtr)
 		{
-			continue;
-		}
-		
-		for (TFieldIterator<UFunction> FuncIt(Class, EFieldIteratorFlags::ExcludeSuper); FuncIt; ++FuncIt)
-		{
-			UFunction* Function = *FuncIt;
-			if (!Function->HasAnyFunctionFlags(FUNC_Static))
+			UK2Node_LatentAbilityCall* AsyncTaskNode = CastChecked<UK2Node_LatentAbilityCall>(NewNode);
+			if (FunctionPtr.IsValid())
 			{
-				continue;
-			}
-
-			// to keep from needlessly instantiating a UBlueprintNodeSpawner, first   
-			// check to make sure that the registrar is looking for actions of this type
-			// (could be regenerating actions for a specific asset, and therefore the 
-			// registrar would only accept actions corresponding to that asset)
-			if (!ActionRegistrar.IsOpenForRegistration(Function))
-			{
-				continue;
-			}
-
-			UObjectProperty* ReturnProperty = Cast<UObjectProperty>(Function->GetReturnProperty());
-			// see if the function is a static factory method for online proxies
-			bool const bIsProxyFactoryMethod = (ReturnProperty != nullptr) && ReturnProperty->PropertyClass->IsChildOf<UAbilityTask>();
-			
-			if (bIsProxyFactoryMethod)
-			{
-				UBlueprintNodeSpawner* NodeSpawner = UBlueprintFunctionNodeSpawner::Create(Function);
-				check(NodeSpawner != nullptr);
-				NodeSpawner->NodeClass = GetClass();
-				
-				auto CustomizeAcyncNodeLambda = [](UEdGraphNode* NewNode, bool bIsTemplateNode, TWeakObjectPtr<UFunction> FunctionPtr)
-				{
-					UK2Node_LatentAbilityCall* AsyncTaskNode = CastChecked<UK2Node_LatentAbilityCall>(NewNode);
-					if (FunctionPtr.IsValid())
-					{
-						UFunction* Func = FunctionPtr.Get();
-						UObjectProperty* ReturnProp = CastChecked<UObjectProperty>(Func->GetReturnProperty());
+				UFunction* Func = FunctionPtr.Get();
+				UObjectProperty* ReturnProp = CastChecked<UObjectProperty>(Func->GetReturnProperty());
 						
-						AsyncTaskNode->ProxyFactoryFunctionName = Func->GetFName();
-						AsyncTaskNode->ProxyFactoryClass        = Func->GetOuterUClass();
-						AsyncTaskNode->ProxyClass               = ReturnProp->PropertyClass;
-					}
-				};
-				
-				TWeakObjectPtr<UFunction> FunctionPtr = Function;
-				NodeSpawner->CustomizeNodeDelegate = UBlueprintNodeSpawner::FCustomizeNodeDelegate::CreateStatic(CustomizeAcyncNodeLambda, FunctionPtr);
-				
-				// @TODO: since this can't be folded into FBlueprintActionDatabase, we
-				//        need a way to associate these spawners with a certain class
-				ActionRegistrar.AddBlueprintAction(Function, NodeSpawner);
+				AsyncTaskNode->ProxyFactoryFunctionName = Func->GetFName();
+				AsyncTaskNode->ProxyFactoryClass        = Func->GetOuterUClass();
+				AsyncTaskNode->ProxyClass               = ReturnProp->PropertyClass;
 			}
 		}
-	}
+	};
+
+	UClass* NodeClass = GetClass();
+	ActionRegistrar.RegisterClassFactoryActions<UAbilityTask>( FBlueprintActionDatabaseRegistrar::FMakeFuncSpawnerDelegate::CreateLambda([NodeClass](const UFunction* FactoryFunc)->UBlueprintNodeSpawner*
+	{
+		UBlueprintNodeSpawner* NodeSpawner = UBlueprintFunctionNodeSpawner::Create(FactoryFunc);
+		check(NodeSpawner != nullptr);
+		NodeSpawner->NodeClass = NodeClass;
+
+		TWeakObjectPtr<UFunction> FunctionPtr = FactoryFunc;
+		NodeSpawner->CustomizeNodeDelegate = UBlueprintNodeSpawner::FCustomizeNodeDelegate::CreateStatic(GetMenuActions_Utils::SetNodeFunc, FunctionPtr);
+
+		return NodeSpawner;
+	}) );
 }
 
 #undef LOCTEXT_NAMESPACE
diff --git a/Engine/Source/Editor/GameplayTasksEditor/Private/K2Node_LatentGameplayTaskCall.cpp b/Engine/Source/Editor/GameplayTasksEditor/Private/K2Node_LatentGameplayTaskCall.cpp
--- a/Engine/Source/Editor/GameplayTasksEditor/Private/K2Node_LatentGameplayTaskCall.cpp
+++ b/Engine/Source/Editor/GameplayTasksEditor/Private/K2Node_LatentGameplayTaskCall.cpp
@@ -53,78 +53,41 @@ bool UK2Node_LatentGameplayTaskCall::CanCreateUnderSpecifiedSchema(const UEdGrap
 
 void UK2Node_LatentGameplayTaskCall::GetMenuActions(FBlueprintActionDatabaseRegistrar& ActionRegistrar) const
 {
-	// these nested loops are combing over the same classes/functions the
-	// FBlueprintActionDatabase does; ideally we save on perf and fold this in
-	// with FBlueprintActionDatabase, but we want to keep the modules separate
-	for (TObjectIterator<UClass> ClassIt; ClassIt; ++ClassIt)
+	struct GetMenuActions_Utils
 	{
-		UClass* Class = *ClassIt;
-		if (!Class->IsChildOf<UGameplayTask>() || Class->HasAnyClassFlags(CLASS_Abstract) 
-			|| HasDedicatedNodeClass(Class))
+		static void SetNodeFunc(UEdGraphNode* NewNode, bool /*bIsTemplateNode*/, TWeakObjectPtr<UFunction> FunctionPtr)
 		{
-			continue;
+			UK2Node_LatentGameplayTaskCall* AsyncTaskNode = CastChecked<UK2Node_LatentGameplayTaskCall>(NewNode);
+			if (FunctionPtr.IsValid())
+			{
+				UFunction* Func = FunctionPtr.Get();
+				UObjectProperty* ReturnProp = CastChecked<UObjectProperty>(Func->GetReturnProperty());
+						
+				AsyncTaskNode->ProxyFactoryFunctionName = Func->GetFName();
+				AsyncTaskNode->ProxyFactoryClass        = Func->GetOuterUClass();
+				AsyncTaskNode->ProxyClass               = ReturnProp->PropertyClass;
+			}
 		}
+	};
+
+	UClass* NodeClass = GetClass();
+	ActionRegistrar.RegisterClassFactoryActions<UGameplayTask>( FBlueprintActionDatabaseRegistrar::FMakeFuncSpawnerDelegate::CreateLambda([NodeClass](const UFunction* FactoryFunc)->UBlueprintNodeSpawner*
+	{
+		UBlueprintNodeSpawner* NodeSpawner = nullptr;
 		
-		for (TFieldIterator<UFunction> FuncIt(Class, EFieldIteratorFlags::ExcludeSuper); FuncIt; ++FuncIt)
+		UClass* FuncClass = FactoryFunc->GetOwnerClass();
+		if (!UK2Node_LatentGameplayTaskCall::HasDedicatedNodeClass(FuncClass))
 		{
-			UFunction* Function = *FuncIt;
-			if (!Function->HasAnyFunctionFlags(FUNC_Static))
-			{
-				continue;
-			}
+			NodeSpawner = UBlueprintFunctionNodeSpawner::Create(FactoryFunc);
+			check(NodeSpawner != nullptr);
+			NodeSpawner->NodeClass = NodeClass;
 
-			// to keep from needlessly instantiating a UBlueprintNodeSpawner, first   
-			// check to make sure that the registrar is looking for actions of this type
-			// (could be regenerating actions for a specific asset, and therefore the 
-			// registrar would only accept actions corresponding to that asset)
-			if (!ActionRegistrar.IsOpenForRegistration(Function))
-			{
-				continue;
-			}
+			TWeakObjectPtr<UFunction> FunctionPtr = FactoryFunc;
+			NodeSpawner->CustomizeNodeDelegate = UBlueprintNodeSpawner::FCustomizeNodeDelegate::CreateStatic(GetMenuActions_Utils::SetNodeFunc, FunctionPtr);
 
-			//// check if function's expected world context object's class is of this BP's type
-			//// we can also accept ones that are not hidden
-			//UObjectProperty* WorldContextProperty = Cast<UObjectProperty>(Function->PropertyLink);
-			//if (ActionRegistrar.GetActionKeyFilter() != nullptr
-			//	&& WorldContextProperty != nullptr
-			//	&& ActionRegistrar.GetActionKeyFilter()->GetClass()->IsChildOf(WorldContextProperty->GetClass()) == false)
-			//{
-			//	continue;
-			//}
-
-			UObjectProperty* ReturnProperty = Cast<UObjectProperty>(Function->GetReturnProperty());
-			// see if the function is a static factory method for online proxies
-			bool const bIsProxyFactoryMethod = (ReturnProperty != nullptr) && ReturnProperty->PropertyClass->IsChildOf<UGameplayTask>();
-			
-			if (bIsProxyFactoryMethod)
-			{
-				UBlueprintNodeSpawner* NodeSpawner = UBlueprintFunctionNodeSpawner::Create(Function);
-				check(NodeSpawner != nullptr);
-				NodeSpawner->NodeClass = GetClass();
-				
-				auto CustomizeAcyncNodeLambda = [](UEdGraphNode* NewNode, bool bIsTemplateNode, TWeakObjectPtr<UFunction> FunctionPtr)
-				{
-					UK2Node_LatentGameplayTaskCall* AsyncTaskNode = CastChecked<UK2Node_LatentGameplayTaskCall>(NewNode);
-					if (FunctionPtr.IsValid())
-					{
-						UFunction* Func = FunctionPtr.Get();
-						UObjectProperty* ReturnProp = CastChecked<UObjectProperty>(Func->GetReturnProperty());
-						
-						AsyncTaskNode->ProxyFactoryFunctionName = Func->GetFName();
-						AsyncTaskNode->ProxyFactoryClass        = Func->GetOuterUClass();
-						AsyncTaskNode->ProxyClass               = ReturnProp->PropertyClass;
-					}
-				};
-				
-				TWeakObjectPtr<UFunction> FunctionPtr = Function;
-				NodeSpawner->CustomizeNodeDelegate = UBlueprintNodeSpawner::FCustomizeNodeDelegate::CreateStatic(CustomizeAcyncNodeLambda, FunctionPtr);
-				
-				// @TODO: since this can't be folded into FBlueprintActionDatabase, we
-				//        need a way to associate these spawners with a certain class
-				ActionRegistrar.AddBlueprintAction(Function, NodeSpawner);
-			}
 		}
-	}
+		return NodeSpawner;
+	}) );
 }
 
 // -------------------------------------------------
diff --git a/Engine/Source/Editor/Kismet/Private/Nodes/K2Node_AsyncAction.cpp b/Engine/Source/Editor/Kismet/Private/Nodes/K2Node_AsyncAction.cpp
--- a/Engine/Source/Editor/Kismet/Private/Nodes/K2Node_AsyncAction.cpp
+++ b/Engine/Source/Editor/Kismet/Private/Nodes/K2Node_AsyncAction.cpp
@@ -18,67 +18,35 @@ UK2Node_AsyncAction::UK2Node_AsyncAction(const FObjectInitializer& ObjectInitial
 
 void UK2Node_AsyncAction::GetMenuActions(FBlueprintActionDatabaseRegistrar& ActionRegistrar) const
 {
-	// these nested loops are combing over the same classes/functions the
-	// FBlueprintActionDatabase does; ideally we save on perf and fold this in
-	// with FBlueprintActionDatabase, but we want to keep the modules separate
-	for (TObjectIterator<UClass> ClassIt; ClassIt; ++ClassIt)
+	struct GetMenuActions_Utils
 	{
-		UClass* Class = *ClassIt;
-		if ( !Class->IsChildOf<UBlueprintAsyncActionBase>() || Class->HasAnyClassFlags(CLASS_Abstract) )
+		static void SetNodeFunc(UEdGraphNode* NewNode, bool /*bIsTemplateNode*/, TWeakObjectPtr<UFunction> FunctionPtr)
 		{
-			continue;
-		}
-		
-		for (TFieldIterator<UFunction> FuncIt(Class, EFieldIteratorFlags::ExcludeSuper); FuncIt; ++FuncIt)
-		{
-			UFunction* Function = *FuncIt;
-			if (!Function->HasAnyFunctionFlags(FUNC_Static))
-			{
-				continue;
-			}
-
-			// to keep from needlessly instantiating a UBlueprintNodeSpawner, first   
-			// check to make sure that the registrar is looking for actions of this type
-			// (could be regenerating actions for a specific asset, and therefore the 
-			// registrar would only accept actions corresponding to that asset)
-			if (!ActionRegistrar.IsOpenForRegistration(Function))
+			UK2Node_AsyncAction* AsyncTaskNode = CastChecked<UK2Node_AsyncAction>(NewNode);
+			if (FunctionPtr.IsValid())
 			{
-				continue;
-			}
-			
-			UObjectProperty* ReturnProperty = Cast<UObjectProperty>(Function->GetReturnProperty());
-			// see if the function is a static factory method for online proxies
-			bool const bIsProxyFactoryMethod = ( ReturnProperty != nullptr ) && ReturnProperty->PropertyClass->IsChildOf<UBlueprintAsyncActionBase>();
-			
-			if (bIsProxyFactoryMethod)
-			{
-				UBlueprintNodeSpawner* NodeSpawner = UBlueprintFunctionNodeSpawner::Create(Function);
-				check(NodeSpawner != nullptr);
-				NodeSpawner->NodeClass = GetClass();
-				
-				auto CustomizeAcyncNodeLambda = [](UEdGraphNode* NewNode, bool bIsTemplateNode, TWeakObjectPtr<UFunction> FunctionPtr)
-				{
-					UK2Node_AsyncAction* AsyncTaskNode = CastChecked<UK2Node_AsyncAction>(NewNode);
-					if (FunctionPtr.IsValid())
-					{
-						UFunction* Func = FunctionPtr.Get();
-						UObjectProperty* ReturnProp = CastChecked<UObjectProperty>(Func->GetReturnProperty());
+				UFunction* Func = FunctionPtr.Get();
+				UObjectProperty* ReturnProp = CastChecked<UObjectProperty>(Func->GetReturnProperty());
 						
-						AsyncTaskNode->ProxyFactoryFunctionName = Func->GetFName();
-						AsyncTaskNode->ProxyFactoryClass        = Func->GetOuterUClass();
-						AsyncTaskNode->ProxyClass               = ReturnProp->PropertyClass;
-					}
-				};
-				
-				TWeakObjectPtr<UFunction> FunctionPtr = Function;
-				NodeSpawner->CustomizeNodeDelegate = UBlueprintNodeSpawner::FCustomizeNodeDelegate::CreateStatic(CustomizeAcyncNodeLambda, FunctionPtr);
-				
-				// @TODO: since this can't be folded into FBlueprintActionDatabase, we
-				//        need a way to associate these spawners with a certain class
-				ActionRegistrar.AddBlueprintAction(Function, NodeSpawner);
+				AsyncTaskNode->ProxyFactoryFunctionName = Func->GetFName();
+				AsyncTaskNode->ProxyFactoryClass        = Func->GetOuterUClass();
+				AsyncTaskNode->ProxyClass               = ReturnProp->PropertyClass;
 			}
 		}
-	}
+	};
+
+	UClass* NodeClass = GetClass();
+	ActionRegistrar.RegisterClassFactoryActions<UBlueprintAsyncActionBase>(FBlueprintActionDatabaseRegistrar::FMakeFuncSpawnerDelegate::CreateLambda([NodeClass](const UFunction* FactoryFunc)->UBlueprintNodeSpawner*
+	{
+		UBlueprintNodeSpawner* NodeSpawner = UBlueprintFunctionNodeSpawner::Create(FactoryFunc);
+		check(NodeSpawner != nullptr);
+		NodeSpawner->NodeClass = NodeClass;
+
+		TWeakObjectPtr<UFunction> FunctionPtr = FactoryFunc;
+		NodeSpawner->CustomizeNodeDelegate = UBlueprintNodeSpawner::FCustomizeNodeDelegate::CreateStatic(GetMenuActions_Utils::SetNodeFunc, FunctionPtr);
+
+		return NodeSpawner;
+	}) );
 }
 
 #undef LOCTEXT_NAMESPACE
diff --git a/Engine/Source/Editor/OnlineBlueprintSupport/Private/K2Node_LatentOnlineCall.cpp b/Engine/Source/Editor/OnlineBlueprintSupport/Private/K2Node_LatentOnlineCall.cpp
--- a/Engine/Source/Editor/OnlineBlueprintSupport/Private/K2Node_LatentOnlineCall.cpp
+++ b/Engine/Source/Editor/OnlineBlueprintSupport/Private/K2Node_LatentOnlineCall.cpp
@@ -16,67 +16,35 @@ UK2Node_LatentOnlineCall::UK2Node_LatentOnlineCall(const FObjectInitializer& Obj
 
 void UK2Node_LatentOnlineCall::GetMenuActions(FBlueprintActionDatabaseRegistrar& ActionRegistrar) const
 {
-	// these nested loops are combing over the same classes/functions the
-	// FBlueprintActionDatabase does; ideally we save on perf and fold this in
-	// with FBlueprintActionDatabase, but we want to keep the modules separate
-	for (TObjectIterator<UClass> ClassIt; ClassIt; ++ClassIt)
+	struct GetMenuActions_Utils
 	{
-		UClass* Class = *ClassIt;
-		if (!Class->IsChildOf<UOnlineBlueprintCallProxyBase>() || Class->HasAnyClassFlags(CLASS_Abstract))
+		static void SetNodeFunc(UEdGraphNode* NewNode, bool /*bIsTemplateNode*/, TWeakObjectPtr<UFunction> FunctionPtr)
 		{
-			continue;
-		}
-		
-		for (TFieldIterator<UFunction> FuncIt(Class, EFieldIteratorFlags::ExcludeSuper); FuncIt; ++FuncIt)
-		{
-			UFunction* Function = *FuncIt;
-			if (!Function->HasAnyFunctionFlags(FUNC_Static))
-			{
-				continue;
-			}
-
-			// to keep from needlessly instantiating a UBlueprintNodeSpawner, first   
-			// check to make sure that the registrar is looking for actions of this type
-			// (could be regenerating actions for a specific asset, and therefore the 
-			// registrar would only accept actions corresponding to that asset)
-			if (!ActionRegistrar.IsOpenForRegistration(Function))
+			UK2Node_LatentOnlineCall* AsyncTaskNode = CastChecked<UK2Node_LatentOnlineCall>(NewNode);
+			if (FunctionPtr.IsValid())
 			{
-				continue;
-			}
-			
-			UObjectProperty* ReturnProperty = Cast<UObjectProperty>(Function->GetReturnProperty());
-			// see if the function is a static factory method for online proxies
-			bool const bIsProxyFactoryMethod = (ReturnProperty != nullptr) && ReturnProperty->PropertyClass->IsChildOf<UOnlineBlueprintCallProxyBase>();
-			
-			if (bIsProxyFactoryMethod)
-			{
-				UBlueprintNodeSpawner* NodeSpawner = UBlueprintFunctionNodeSpawner::Create(Function);
-				check(NodeSpawner != nullptr);
-				NodeSpawner->NodeClass = GetClass();
-				
-				auto CustomizeAcyncNodeLambda = [](UEdGraphNode* NewNode, bool bIsTemplateNode, TWeakObjectPtr<UFunction> FunctionPtr)
-				{
-					UK2Node_LatentOnlineCall* AsyncTaskNode = CastChecked<UK2Node_LatentOnlineCall>(NewNode);
-					if (FunctionPtr.IsValid())
-					{
-						UFunction* Func = FunctionPtr.Get();
-						UObjectProperty* ReturnProp = CastChecked<UObjectProperty>(Func->GetReturnProperty());
+				UFunction* Func = FunctionPtr.Get();
+				UObjectProperty* ReturnProp = CastChecked<UObjectProperty>(Func->GetReturnProperty());
 						
-						AsyncTaskNode->ProxyFactoryFunctionName = Func->GetFName();
-						AsyncTaskNode->ProxyFactoryClass        = Func->GetOuterUClass();
-						AsyncTaskNode->ProxyClass               = ReturnProp->PropertyClass;
-					}
-				};
-				
-				TWeakObjectPtr<UFunction> FunctionPtr = Function;
-				NodeSpawner->CustomizeNodeDelegate = UBlueprintNodeSpawner::FCustomizeNodeDelegate::CreateStatic(CustomizeAcyncNodeLambda, FunctionPtr);
-				
-				// @TODO: since this can't be folded into FBlueprintActionDatabase, we
-				//        need a way to associate these spawners with a certain class
-				ActionRegistrar.AddBlueprintAction(Function, NodeSpawner);
+				AsyncTaskNode->ProxyFactoryFunctionName = Func->GetFName();
+				AsyncTaskNode->ProxyFactoryClass        = Func->GetOuterUClass();
+				AsyncTaskNode->ProxyClass               = ReturnProp->PropertyClass;
 			}
 		}
-	}
+	};
+
+	UClass* NodeClass = GetClass();
+	ActionRegistrar.RegisterClassFactoryActions<UOnlineBlueprintCallProxyBase>( FBlueprintActionDatabaseRegistrar::FMakeFuncSpawnerDelegate::CreateLambda([NodeClass](const UFunction* FactoryFunc)->UBlueprintNodeSpawner*
+	{
+		UBlueprintNodeSpawner* NodeSpawner = UBlueprintFunctionNodeSpawner::Create(FactoryFunc);
+		check(NodeSpawner != nullptr);
+		NodeSpawner->NodeClass = NodeClass;
+
+		TWeakObjectPtr<UFunction> FunctionPtr = FactoryFunc;
+		NodeSpawner->CustomizeNodeDelegate = UBlueprintNodeSpawner::FCustomizeNodeDelegate::CreateStatic(GetMenuActions_Utils::SetNodeFunc, FunctionPtr);
+
+		return NodeSpawner;
+	}) );
 }
 
 #undef LOCTEXT_NAMESPACE
