Buildwatcher - Fix shadow variables

diff --git a/Engine/Source/Editor/BlueprintGraph/Private/BlueprintActionDatabaseRegistrar.cpp b/Engine/Source/Editor/BlueprintGraph/Private/BlueprintActionDatabaseRegistrar.cpp
--- a/Engine/Source/Editor/BlueprintGraph/Private/BlueprintActionDatabaseRegistrar.cpp
+++ b/Engine/Source/Editor/BlueprintGraph/Private/BlueprintActionDatabaseRegistrar.cpp
@@ -278,7 +278,7 @@ int32 FBlueprintActionDatabaseRegistrar::RegisterClassFactoryActions(const UClas
 {
 	struct RegisterClassFactoryActions_Utils
 	{
-		static bool IsFactoryMethod(const UFunction* Function, const UClass* TargetType)
+		static bool IsFactoryMethod(const UFunction* Function, const UClass* InTargetType)
 		{
 			if (!Function->HasAnyFunctionFlags(FUNC_Static))
 			{
@@ -287,7 +287,7 @@ int32 FBlueprintActionDatabaseRegistrar::RegisterClassFactoryActions(const UClas
 
 			UObjectProperty* ReturnProperty = Cast<UObjectProperty>(Function->GetReturnProperty());
 			// see if the function is a static factory method
-			bool const bIsFactoryMethod = (ReturnProperty != nullptr) && ReturnProperty->PropertyClass->IsChildOf(TargetType);
+			bool const bIsFactoryMethod = (ReturnProperty != nullptr) && ReturnProperty->PropertyClass->IsChildOf(InTargetType);
 
 			return bIsFactoryMethod;
 		}
diff --git a/Engine/Source/Editor/BlueprintGraph/Private/K2Node_CastByteToEnum.cpp b/Engine/Source/Editor/BlueprintGraph/Private/K2Node_CastByteToEnum.cpp
--- a/Engine/Source/Editor/BlueprintGraph/Private/K2Node_CastByteToEnum.cpp
+++ b/Engine/Source/Editor/BlueprintGraph/Private/K2Node_CastByteToEnum.cpp
@@ -184,11 +184,11 @@ void UK2Node_CastByteToEnum::GetMenuActions(FBlueprintActionDatabaseRegistrar& A
 	};
 
 	UClass* NodeClass = GetClass();
-	ActionRegistrar.RegisterEnumActions( FBlueprintActionDatabaseRegistrar::FMakeEnumSpawnerDelegate::CreateLambda([NodeClass](const UEnum* Enum)->UBlueprintNodeSpawner*
+	ActionRegistrar.RegisterEnumActions( FBlueprintActionDatabaseRegistrar::FMakeEnumSpawnerDelegate::CreateLambda([NodeClass](const UEnum* InEnum)->UBlueprintNodeSpawner*
 	{
-		UBlueprintFieldNodeSpawner* NodeSpawner = UBlueprintFieldNodeSpawner::Create(NodeClass, Enum);
+		UBlueprintFieldNodeSpawner* NodeSpawner = UBlueprintFieldNodeSpawner::Create(NodeClass, InEnum);
 		check(NodeSpawner != nullptr);
-		TWeakObjectPtr<UEnum> NonConstEnumPtr = Enum;
+		TWeakObjectPtr<UEnum> NonConstEnumPtr = InEnum;
 		NodeSpawner->SetNodeFieldDelegate = UBlueprintFieldNodeSpawner::FSetNodeFieldDelegate::CreateStatic(GetMenuActions_Utils::SetNodeEnum, NonConstEnumPtr);
 
 		return NodeSpawner;
diff --git a/Engine/Source/Editor/BlueprintGraph/Private/K2Node_EnumLiteral.cpp b/Engine/Source/Editor/BlueprintGraph/Private/K2Node_EnumLiteral.cpp
--- a/Engine/Source/Editor/BlueprintGraph/Private/K2Node_EnumLiteral.cpp
+++ b/Engine/Source/Editor/BlueprintGraph/Private/K2Node_EnumLiteral.cpp
@@ -119,11 +119,11 @@ void UK2Node_EnumLiteral::GetMenuActions(FBlueprintActionDatabaseRegistrar& Acti
 	};
 
 	UClass* NodeClass = GetClass();
-	ActionRegistrar.RegisterEnumActions( FBlueprintActionDatabaseRegistrar::FMakeEnumSpawnerDelegate::CreateLambda([NodeClass](const UEnum* Enum)->UBlueprintNodeSpawner*
+	ActionRegistrar.RegisterEnumActions( FBlueprintActionDatabaseRegistrar::FMakeEnumSpawnerDelegate::CreateLambda([NodeClass](const UEnum* InEnum)->UBlueprintNodeSpawner*
 	{
-		UBlueprintFieldNodeSpawner* NodeSpawner = UBlueprintFieldNodeSpawner::Create(NodeClass, Enum);
+		UBlueprintFieldNodeSpawner* NodeSpawner = UBlueprintFieldNodeSpawner::Create(NodeClass, InEnum);
 		check(NodeSpawner != nullptr);
-		TWeakObjectPtr<UEnum> NonConstEnumPtr = Enum;
+		TWeakObjectPtr<UEnum> NonConstEnumPtr = InEnum;
 		NodeSpawner->SetNodeFieldDelegate = UBlueprintFieldNodeSpawner::FSetNodeFieldDelegate::CreateStatic(GetMenuActions_Utils::SetNodeEnum, NonConstEnumPtr);
 
 		return NodeSpawner;
diff --git a/Engine/Source/Editor/BlueprintGraph/Private/K2Node_ForEachElementInEnum.cpp b/Engine/Source/Editor/BlueprintGraph/Private/K2Node_ForEachElementInEnum.cpp
--- a/Engine/Source/Editor/BlueprintGraph/Private/K2Node_ForEachElementInEnum.cpp
+++ b/Engine/Source/Editor/BlueprintGraph/Private/K2Node_ForEachElementInEnum.cpp
@@ -235,11 +235,11 @@ void UK2Node_ForEachElementInEnum::GetMenuActions(FBlueprintActionDatabaseRegist
 	};
 
 	UClass* NodeClass = GetClass();
-	ActionRegistrar.RegisterEnumActions( FBlueprintActionDatabaseRegistrar::FMakeEnumSpawnerDelegate::CreateLambda([NodeClass](const UEnum* Enum)->UBlueprintNodeSpawner*
+	ActionRegistrar.RegisterEnumActions( FBlueprintActionDatabaseRegistrar::FMakeEnumSpawnerDelegate::CreateLambda([NodeClass](const UEnum* InEnum)->UBlueprintNodeSpawner*
 	{
-		UBlueprintFieldNodeSpawner* NodeSpawner = UBlueprintFieldNodeSpawner::Create(NodeClass, Enum);
+		UBlueprintFieldNodeSpawner* NodeSpawner = UBlueprintFieldNodeSpawner::Create(NodeClass, InEnum);
 		check(NodeSpawner != nullptr);
-		TWeakObjectPtr<UEnum> NonConstEnumPtr = Enum;
+		TWeakObjectPtr<UEnum> NonConstEnumPtr = InEnum;
 		NodeSpawner->SetNodeFieldDelegate = UBlueprintFieldNodeSpawner::FSetNodeFieldDelegate::CreateStatic(GetMenuActions_Utils::SetNodeEnum, NonConstEnumPtr);
 
 		return NodeSpawner;
diff --git a/Engine/Source/Editor/BlueprintGraph/Private/K2Node_GetNumEnumEntries.cpp b/Engine/Source/Editor/BlueprintGraph/Private/K2Node_GetNumEnumEntries.cpp
--- a/Engine/Source/Editor/BlueprintGraph/Private/K2Node_GetNumEnumEntries.cpp
+++ b/Engine/Source/Editor/BlueprintGraph/Private/K2Node_GetNumEnumEntries.cpp
@@ -105,11 +105,11 @@ void UK2Node_GetNumEnumEntries::GetMenuActions(FBlueprintActionDatabaseRegistrar
 	};
 
 	UClass* NodeClass = GetClass();
-	ActionRegistrar.RegisterEnumActions( FBlueprintActionDatabaseRegistrar::FMakeEnumSpawnerDelegate::CreateLambda([NodeClass](const UEnum* Enum)->UBlueprintNodeSpawner*
+	ActionRegistrar.RegisterEnumActions( FBlueprintActionDatabaseRegistrar::FMakeEnumSpawnerDelegate::CreateLambda([NodeClass](const UEnum* InEnum)->UBlueprintNodeSpawner*
 	{
-		UBlueprintFieldNodeSpawner* NodeSpawner = UBlueprintFieldNodeSpawner::Create(NodeClass, Enum);
+		UBlueprintFieldNodeSpawner* NodeSpawner = UBlueprintFieldNodeSpawner::Create(NodeClass, InEnum);
 		check(NodeSpawner != nullptr);
-		TWeakObjectPtr<UEnum> NonConstEnumPtr = Enum;
+		TWeakObjectPtr<UEnum> NonConstEnumPtr = InEnum;
 		NodeSpawner->SetNodeFieldDelegate = UBlueprintFieldNodeSpawner::FSetNodeFieldDelegate::CreateStatic(GetMenuActions_Utils::SetNodeEnum, NonConstEnumPtr);
 
 		return NodeSpawner;
diff --git a/Engine/Source/Editor/BlueprintGraph/Private/K2Node_SwitchEnum.cpp b/Engine/Source/Editor/BlueprintGraph/Private/K2Node_SwitchEnum.cpp
--- a/Engine/Source/Editor/BlueprintGraph/Private/K2Node_SwitchEnum.cpp
+++ b/Engine/Source/Editor/BlueprintGraph/Private/K2Node_SwitchEnum.cpp
@@ -95,11 +95,11 @@ void UK2Node_SwitchEnum::GetMenuActions(FBlueprintActionDatabaseRegistrar& Actio
 	};
 
 	UClass* NodeClass = GetClass();
-	ActionRegistrar.RegisterEnumActions( FBlueprintActionDatabaseRegistrar::FMakeEnumSpawnerDelegate::CreateLambda([NodeClass](const UEnum* Enum)->UBlueprintNodeSpawner*
+	ActionRegistrar.RegisterEnumActions( FBlueprintActionDatabaseRegistrar::FMakeEnumSpawnerDelegate::CreateLambda([NodeClass](const UEnum* InEnum)->UBlueprintNodeSpawner*
 	{
-		UBlueprintFieldNodeSpawner* NodeSpawner = UBlueprintFieldNodeSpawner::Create(NodeClass, Enum);
+		UBlueprintFieldNodeSpawner* NodeSpawner = UBlueprintFieldNodeSpawner::Create(NodeClass, InEnum);
 		check(NodeSpawner != nullptr);
-		TWeakObjectPtr<UEnum> NonConstEnumPtr = Enum;
+		TWeakObjectPtr<UEnum> NonConstEnumPtr = InEnum;
 		NodeSpawner->SetNodeFieldDelegate = UBlueprintFieldNodeSpawner::FSetNodeFieldDelegate::CreateStatic(GetMenuActions_Utils::SetNodeEnum, NonConstEnumPtr);
 
 		return NodeSpawner;
