Better validation in Result Node.

diff --git a/Engine/Source/Editor/BlueprintGraph/Classes/K2Node_FunctionResult.h b/Engine/Source/Editor/BlueprintGraph/Classes/K2Node_FunctionResult.h
--- a/Engine/Source/Editor/BlueprintGraph/Classes/K2Node_FunctionResult.h
+++ b/Engine/Source/Editor/BlueprintGraph/Classes/K2Node_FunctionResult.h
@@ -14,6 +14,7 @@ class UK2Node_FunctionResult : public UK2Node_FunctionTerminator
 	virtual void AllocateDefaultPins() override;
 	virtual FText GetNodeTitle(ENodeTitleType::Type TitleType) const override;
 	virtual FText GetTooltipText() const override;
+	virtual void ValidateNodeDuringCompilation(class FCompilerResultsLog& MessageLog) const override;
 	// End UEdGraphNode interface
 
 	// Begin UK2Node interface
diff --git a/Engine/Source/Editor/BlueprintGraph/Private/K2Node_FunctionResult.cpp b/Engine/Source/Editor/BlueprintGraph/Private/K2Node_FunctionResult.cpp
--- a/Engine/Source/Editor/BlueprintGraph/Private/K2Node_FunctionResult.cpp
+++ b/Engine/Source/Editor/BlueprintGraph/Private/K2Node_FunctionResult.cpp
@@ -247,3 +247,23 @@ void UK2Node_FunctionResult::PostPlacedNewNode()
 		}
 	}
 }
+
+void UK2Node_FunctionResult::ValidateNodeDuringCompilation(class FCompilerResultsLog& MessageLog) const
+{
+	Super::ValidateNodeDuringCompilation(MessageLog);
+
+	auto AllResultNodes = GetAllResultNodes();
+	UK2Node_FunctionResult* OtherResult = AllResultNodes.Num() ? AllResultNodes[0] : nullptr;
+	if (OtherResult && (OtherResult != this))
+	{
+		for (auto Pin : Pins)
+		{
+			auto OtherPin = OtherResult->FindPin(Pin->PinName);
+			if (!OtherPin || (OtherPin->PinType != Pin->PinType))
+			{
+				MessageLog.Error(*NSLOCTEXT("K2Node", "FunctionResult_DifferentReturnError", "Return nodes don't match each other: @@, @@").ToString(), this, OtherResult);
+				break;
+			}
+		}
+	}
+}
\ No newline at end of file
