Math expression functions without any parameters will still provide their end parenthesis in the Math Expression node.

diff --git a/Engine/Source/Editor/BlueprintGraph/Private/K2Node_MathExpression.cpp b/Engine/Source/Editor/BlueprintGraph/Private/K2Node_MathExpression.cpp
--- a/Engine/Source/Editor/BlueprintGraph/Private/K2Node_MathExpression.cpp
+++ b/Engine/Source/Editor/BlueprintGraph/Private/K2Node_MathExpression.cpp
@@ -478,36 +478,50 @@ public:
 	virtual FString ToString() const override
 	{
 		FString AsString("(");
-		for (TSharedRef<IFExpressionNode> Child : Children)
+		if (Children.Num() > 0)
 		{
-			AsString += Child->ToString();
-			if (Child == Children.Last())
-			{
-				AsString += ")";
-			}
-			else 
+			for (TSharedRef<IFExpressionNode> Child : Children)
 			{
-				AsString += ", ";
+				AsString += Child->ToString();
+				if (Child == Children.Last())
+				{
+					AsString += ")";
+				}
+				else 
+				{
+					AsString += ", ";
+				}
 			}
 		}
+		else
+		{
+			AsString += ")";
+		}
 		return AsString;
 	}
 
 	virtual FString ToDisplayString(UBlueprint* InBlueprint) const
 	{
 		FString AsString("(");
-		for (TSharedRef<IFExpressionNode> Child : Children)
+		if (Children.Num() > 0)
 		{
-			AsString += Child->ToDisplayString(InBlueprint);
-			if (Child == Children.Last())
-			{
-				AsString += ")";
-			}
-			else 
+			for (TSharedRef<IFExpressionNode> Child : Children)
 			{
-				AsString += ", ";
+				AsString += Child->ToDisplayString(InBlueprint);
+				if (Child == Children.Last())
+				{
+					AsString += ")";
+				}
+				else 
+				{
+					AsString += ", ";
+				}
 			}
 		}
+		else
+		{
+			AsString += ")";
+		}
 		return AsString;
 	}
 
