Support for streaming level communication via interfaces.

Changed the K2Node_Message to expand and attempt to cast the incoming reference as a ULevelStreaming reference, get the level script actor, and cast that to the interface.

Exposed for Blueprints (BP internal use only) ULevelStreaming::GetLevelScriptActor

All levels can now have a BP interface added to them.

diff --git a/Engine/Source/Editor/BlueprintGraph/Classes/K2Node_Message.h b/Engine/Source/Editor/BlueprintGraph/Classes/K2Node_Message.h
--- a/Engine/Source/Editor/BlueprintGraph/Classes/K2Node_Message.h
+++ b/Engine/Source/Editor/BlueprintGraph/Classes/K2Node_Message.h
@@ -29,6 +29,17 @@ class UK2Node_Message : public UK2Node_CallFunction
 	virtual UEdGraphPin* CreateSelfPin(const UFunction* Function) override;
 
 protected:
+	/**
+	 * Helper function for expand node to expand the message to handle the case that a ULevelStreaming may be passed which needs to have its LevelScriptActor pulled out
+	 *
+	 * @param InCompilerContext			CompilerContext from ExpandNode
+	 * @param InSourceGraph				Graph that the node are expanding from
+	 * @param InStartingExecPin			Exec pin that connects to the message node, to correctly be mapped into the exec pins of the expanded nodes
+	 * @param InMessageSelfPin			The message node's self pin, to correctly be mapped into the object pins of the expanded nodes
+	 * @param InCastToInterfaceNode		Previously generated node to cast to the message's interface, to be wired into the expanded nodes
+	 */
+	void ExpandLevelStreamingHandlers(class FKismetCompilerContext& InCompilerContext, UEdGraph* InSourceGraph, class UEdGraphPin* InStartingExecPin, class UEdGraphPin* InMessageSelfPin, class UK2Node_DynamicCast* InCastToInterfaceNode);
+
 	virtual void EnsureFunctionIsInBlueprint() override;
 	// End K2Node_CallFunction Interface.
 
diff --git a/Engine/Source/Editor/BlueprintGraph/Private/K2Node_Message.cpp b/Engine/Source/Editor/BlueprintGraph/Private/K2Node_Message.cpp
--- a/Engine/Source/Editor/BlueprintGraph/Private/K2Node_Message.cpp
+++ b/Engine/Source/Editor/BlueprintGraph/Private/K2Node_Message.cpp
@@ -88,6 +88,104 @@ FNodeHandlingFunctor* UK2Node_Message::CreateNodeHandler(FKismetCompilerContext&
 	return new FNodeHandlingFunctor(CompilerContext);
 }
 
+void UK2Node_Message::ExpandLevelStreamingHandlers(class FKismetCompilerContext& InCompilerContext, UEdGraph* InSourceGraph, UEdGraphPin* InStartingExecPin, UEdGraphPin* InMessageSelfPin, UK2Node_DynamicCast* InCastToInterfaceNode)
+{
+	const UEdGraphSchema_K2* Schema = InCompilerContext.GetSchema();
+	
+	/**
+	 * Create intermediate nodes
+	 */
+
+	// We need to Cast to a ULevelStreaming object, so we can pull the ALevelScriptActor if successful
+	UK2Node_DynamicCast* CastToLevelStreamingNode = InCompilerContext.SpawnIntermediateNode<UK2Node_DynamicCast>(this, InSourceGraph);
+	CastToLevelStreamingNode->TargetType = ULevelStreaming::StaticClass(); // I want this to be ULevel!
+	CastToLevelStreamingNode->SetPurity(false);
+	CastToLevelStreamingNode->AllocateDefaultPins();
+
+	// Create a  GetLevelScriptActor CallFunction node, this will be used if the cast to ULevelStreaming was successful
+	UK2Node_CallFunction* GetLevelScriptActorNode = InCompilerContext.SpawnIntermediateNode<UK2Node_CallFunction>(this, InSourceGraph);
+	GetLevelScriptActorNode->SetFromFunction(ULevelStreaming::StaticClass()->FindFunctionByName(GET_FUNCTION_NAME_CHECKED(ULevelStreaming, GetLevelScriptActor)));
+	GetLevelScriptActorNode->AllocateDefaultPins();
+
+	// We will store the UObject we want to cast to interface into a temporary variable node
+	UK2Node_TemporaryVariable* TempVariable = InCompilerContext.SpawnIntermediateNode<UK2Node_TemporaryVariable>(this, InSourceGraph);
+	TempVariable->VariableType = InMessageSelfPin->PinType;
+	TempVariable->AllocateDefaultPins();
+
+	// AssignmentStatement node for when the ULevelStreaming cast was successful
+	UK2Node_AssignmentStatement* AssignTempVariableForCastSuccess = InCompilerContext.SpawnIntermediateNode<UK2Node_AssignmentStatement>(this, InSourceGraph);
+	AssignTempVariableForCastSuccess->AllocateDefaultPins();
+
+	// AssignmentStatement node for failed ULevelStreaming casts
+	UK2Node_AssignmentStatement* AssignTempVariableForCastFail = InCompilerContext.SpawnIntermediateNode<UK2Node_AssignmentStatement>(this, InSourceGraph);
+	AssignTempVariableForCastFail->AllocateDefaultPins();
+
+	/**
+	 * Make pin connections to the generated nodes
+	 */
+
+	// Move the connections on the Message node's self pin to the ULevelStreaming Cast node's source pin
+	{
+		InCompilerContext.MovePinLinksToIntermediate(*InMessageSelfPin, *CastToLevelStreamingNode->GetCastSourcePin());
+		CastToLevelStreamingNode->NotifyPinConnectionListChanged(CastToLevelStreamingNode->GetCastSourcePin());
+	}
+
+	// Connect the incoming exec pin to the ULevelStreaming cast node's exec pin, which is the exec flow's entry into this
+	if (InStartingExecPin != nullptr)
+	{
+		// Wire up the connections
+		InCompilerContext.MovePinLinksToIntermediate(*InStartingExecPin, *CastToLevelStreamingNode->GetExecPin());
+	}
+
+	// Connect the GetLevelAcriptActor CallFunction node's self pin to the ULevelStreaming node's cast result pin (successfully cast to a ULevelStreaming object)
+	{
+		Schema->TryCreateConnection(Schema->FindSelfPin(*GetLevelScriptActorNode, EGPD_Input), CastToLevelStreamingNode->GetCastResultPin());
+	}
+
+	// Prepare left side pin connections of AssignTempVariableForCastSuccess
+	{
+		// Connect the TempVariable to the AssignmentStatement node
+		// For the successful cast of ULevelStreaming temporary assignment, hook-up the variable pin
+		Schema->TryCreateConnection(AssignTempVariableForCastSuccess->GetVariablePin(), TempVariable->GetVariablePin());
+
+		// The last pin on the function node is the ALevelScriptActor
+		UEdGraphPin* FuncResultPin = GetLevelScriptActorNode->Pins[GetLevelScriptActorNode->Pins.Num() - 1];
+		ensure(!FuncResultPin->PinType.PinSubCategoryObject->IsA(ALevelScriptActor::StaticClass()));
+
+		// We do not want to let NotifyPinConnectionListChanged get called, it'll change the assignment's type to LevelScript instead of UObject
+		FuncResultPin->MakeLinkTo(AssignTempVariableForCastSuccess->GetValuePin());
+
+		// Hook up the AssignTempVariableForCastSuccess node's exec pin to the valid cast result of the Cast to ULevelStreaming node
+		Schema->TryCreateConnection(AssignTempVariableForCastSuccess->GetExecPin(), CastToLevelStreamingNode->GetValidCastPin());
+	}
+
+	// Prepare left side pin connections of AssignTempVariableForCastFail
+	{
+		// For the non-level actor temporary assignment, hook-up the variable pin
+		Schema->TryCreateConnection(AssignTempVariableForCastFail->GetVariablePin(), TempVariable->GetVariablePin());
+
+		// Connect the value pin to the pin the Cast to ULevelStreaming node's source pin was attached to
+		if (CastToLevelStreamingNode->GetCastSourcePin()->LinkedTo[0])
+		{
+			AssignTempVariableForCastFail->GetValuePin()->MakeLinkTo(CastToLevelStreamingNode->GetCastSourcePin()->LinkedTo[0]);
+		}
+
+		// Connect the LevelStreaming cast failed pin to the appropriate Assignment node
+		AssignTempVariableForCastFail->GetExecPin()->MakeLinkTo(CastToLevelStreamingNode->GetInvalidCastPin());
+	}
+
+	// Connect the Interface cast node to the generated pins
+	{
+		// The source pin of the Interface cast node connects to the temporary variable (storing either an ALevelScriptActor or another UObject)
+		UEdGraphPin* CastToInterfaceSourceObjectPin = InCastToInterfaceNode->GetCastSourcePin();
+		Schema->TryCreateConnection(TempVariable->GetVariablePin(), CastToInterfaceSourceObjectPin);
+
+		// Both assignment nodes (for the temporary variable) get connected to the interface node's exec pin
+		InCastToInterfaceNode->GetExecPin()->MakeLinkTo(AssignTempVariableForCastSuccess->GetThenPin());
+		InCastToInterfaceNode->GetExecPin()->MakeLinkTo(AssignTempVariableForCastFail->GetThenPin());
+	}
+}
+
 void UK2Node_Message::ExpandNode(class FKismetCompilerContext& CompilerContext, UEdGraph* SourceGraph)
 {
 	Super::ExpandNode(CompilerContext, SourceGraph);
@@ -143,20 +241,29 @@ void UK2Node_Message::ExpandNode(class FKismetCompilerContext& CompilerContext,
 			return;
 		}
 
-		CastToInterfaceResultPin->PinType.PinSubCategoryObject = *CastToInterfaceNode->TargetType;
-
-		if (ExecPin != nullptr)
+		// If the message pin is linked and what it is linked to is not a ULevelStreaming pin reference but is a child of ULevelStreaming, 
+		// then we need to leave the possibility that it could be a ULevelStreaming and will need to make an attempt at casting to one
+		if (MessageSelfPin->LinkedTo.Num() > 0 && MessageSelfPin->LinkedTo[0]->PinType.PinSubCategoryObject.Get() != ULevelStreaming::StaticClass() && ULevelStreaming::StaticClass()->IsChildOf(Cast<UClass>(MessageSelfPin->LinkedTo[0]->PinType.PinSubCategoryObject.Get())))
+		{
+			ExpandLevelStreamingHandlers(CompilerContext, SourceGraph, ExecPin, MessageSelfPin, CastToInterfaceNode);
+		}
+		else
 		{
-			UEdGraphPin* CastExecInput = CastToInterfaceNode->GetExecPin();
-			check(CastExecInput != nullptr);
+			// Move the connections on the Message node's self pin to the ULevelStreaming Cast node's source pin
+			{
+				CompilerContext.MovePinLinksToIntermediate(*MessageSelfPin, *CastToInterfaceNode->GetCastSourcePin());
+				CastToInterfaceNode->NotifyPinConnectionListChanged(CastToInterfaceNode->GetCastSourcePin());
+			}
 
-			// Wire up the connections
-			CompilerContext.MovePinLinksToIntermediate(*ExecPin, *CastExecInput);
+			// Connect the incoming exec pin to the ULevelStreaming cast node's exec pin, which is the exec flow's entry into this
+			if (ExecPin != nullptr)
+			{
+				// Wire up the connections
+				CompilerContext.MovePinLinksToIntermediate(*ExecPin, *CastToInterfaceNode->GetExecPin());
+			}
 		}
 
-		UEdGraphPin* CastToInterfaceSourceObjectPin = CastToInterfaceNode->GetCastSourcePin();
-		CompilerContext.MovePinLinksToIntermediate(*MessageSelfPin, *CastToInterfaceSourceObjectPin);
-		CastToInterfaceNode->PinConnectionListChanged(CastToInterfaceSourceObjectPin);
+		CastToInterfaceResultPin->PinType.PinSubCategoryObject = *CastToInterfaceNode->TargetType;
 
 		// Next, create the function call node
 		UK2Node_CallFunction* FunctionCallNode = CompilerContext.SpawnIntermediateNode<UK2Node_CallFunction>(this, SourceGraph);
@@ -170,10 +277,12 @@ void UK2Node_Message::ExpandNode(class FKismetCompilerContext& CompilerContext,
 		UEdGraphPin* LastOutCastFaildPin   = CastToInterfaceNode->GetInvalidCastPin();
 		check(LastOutCastFaildPin != nullptr);
 		UEdGraphPin* LastOutCastSuccessPin = CastToInterfaceValidPin;
+
 		// Wire up the connections
 		if (UEdGraphPin* CallFunctionExecPin = Schema->FindExecutionPin(*FunctionCallNode, EGPD_Input))
 		{
-			CastToInterfaceValidPin->MakeLinkTo(CallFunctionExecPin);
+			// Connect the CallFunctionExecPin to both Assignment nodes for the assignment of the UObject cast
+			CallFunctionExecPin->MakeLinkTo(CastToInterfaceValidPin);
 			LastOutCastSuccessPin = Schema->FindExecutionPin(*FunctionCallNode, EGPD_Output);
 		}
 		
diff --git a/Engine/Source/Editor/Kismet/Private/BlueprintDetailsCustomization.cpp b/Engine/Source/Editor/Kismet/Private/BlueprintDetailsCustomization.cpp
--- a/Engine/Source/Editor/Kismet/Private/BlueprintDetailsCustomization.cpp
+++ b/Engine/Source/Editor/Kismet/Private/BlueprintDetailsCustomization.cpp
@@ -4507,7 +4507,7 @@ void FBlueprintGlobalOptionsDetails::CustomizeDetails(IDetailLayoutBuilder& Deta
 		const bool bIsMacroLibrary = Blueprint->BlueprintType == BPTYPE_MacroLibrary;
 		const bool bIsLevelScriptBP = FBlueprintEditorUtils::IsLevelScriptBlueprint(Blueprint);
 		const bool bIsFunctionLibrary = Blueprint->BlueprintType == BPTYPE_FunctionLibrary;
-		const bool bSupportsInterfaces = !bIsLevelScriptBP && !bIsInterfaceBP && !bIsMacroLibrary && !bIsFunctionLibrary;
+		const bool bSupportsInterfaces = !bIsInterfaceBP && !bIsMacroLibrary && !bIsFunctionLibrary;
 
 		if (bSupportsInterfaces)
 		{
diff --git a/Engine/Source/Runtime/Engine/Classes/Engine/LevelStreaming.h b/Engine/Source/Runtime/Engine/Classes/Engine/LevelStreaming.h
--- a/Engine/Source/Runtime/Engine/Classes/Engine/LevelStreaming.h
+++ b/Engine/Source/Runtime/Engine/Classes/Engine/LevelStreaming.h
@@ -259,6 +259,10 @@ class ULevelStreaming : public UObject
 	UFUNCTION(BlueprintCallable, Category="Game")
 	ULevelStreaming* CreateInstance(FString UniqueInstanceName);
 
+	/** Returns the Level Script Actor of the level if the level is loaded and valid */
+	UFUNCTION(BlueprintPure, meta = (BlueprintInternalUseOnly = "true"))
+	ALevelScriptActor* GetLevelScriptActor();
+
 	//==============================================================================================
 	// Delegates
 	
diff --git a/Engine/Source/Runtime/Engine/Private/LevelStreaming.cpp b/Engine/Source/Runtime/Engine/Private/LevelStreaming.cpp
--- a/Engine/Source/Runtime/Engine/Private/LevelStreaming.cpp
+++ b/Engine/Source/Runtime/Engine/Private/LevelStreaming.cpp
@@ -969,6 +969,15 @@ bool ULevelStreamingKismet::ShouldBeLoaded() const
 	return bShouldBeLoaded;
 }
 
+ALevelScriptActor* ULevelStreaming::GetLevelScriptActor()
+{
+	if (LoadedLevel)
+	{
+		return LoadedLevel->GetLevelScriptActor();
+	}
+	return nullptr;
+}
+
 /*-----------------------------------------------------------------------------
 	ULevelStreamingAlwaysLoaded implementation.
 -----------------------------------------------------------------------------*/
