Can no longer split pins on Blueprint nodes if the struct does not have any valid properties.

diff --git a/Engine/Source/Editor/BlueprintGraph/Classes/K2Node_BreakStruct.h b/Engine/Source/Editor/BlueprintGraph/Classes/K2Node_BreakStruct.h
--- a/Engine/Source/Editor/BlueprintGraph/Classes/K2Node_BreakStruct.h
+++ b/Engine/Source/Editor/BlueprintGraph/Classes/K2Node_BreakStruct.h
@@ -24,7 +24,7 @@ class UK2Node_BreakStruct : public UK2Node_StructMemberGet
 	 * a creation of a break node, although we could do so in the future. There are legacy break nodes that
 	 * rely on expansion of structs that neither have BlueprintVisible properties nor are tagged as BlueprintType
 	 */
-	static bool CanBeBroken(const UScriptStruct* Struct, bool bIncludeEditAnywhere = true);
+	static bool CanBeBroken(const UScriptStruct* Struct, bool bIncludeEditAnywhere = true, bool bMustHaveValidProperties = false);
 
 	// Begin UEdGraphNode interface
 	virtual void AllocateDefaultPins() override;
diff --git a/Engine/Source/Editor/BlueprintGraph/Classes/K2Node_MakeStruct.h b/Engine/Source/Editor/BlueprintGraph/Classes/K2Node_MakeStruct.h
--- a/Engine/Source/Editor/BlueprintGraph/Classes/K2Node_MakeStruct.h
+++ b/Engine/Source/Editor/BlueprintGraph/Classes/K2Node_MakeStruct.h
@@ -25,7 +25,7 @@ class UK2Node_MakeStruct : public UK2Node_StructMemberSet
 	* a creation of a break node, although we could do so in the future. There are legacy break nodes that
 	* rely on expansion of structs that neither have BlueprintVisible properties nor are tagged as BlueprintType
 	*/
-	BLUEPRINTGRAPH_API static bool CanBeMade(const UScriptStruct* Struct, bool bIncludeEditAnywhere = true);
+	BLUEPRINTGRAPH_API static bool CanBeMade(const UScriptStruct* Struct, bool bIncludeEditAnywhere = true, bool bMustHaveValidProperties = false);
 	
 	// Begin UEdGraphNode interface
 	virtual void AllocateDefaultPins() override;
diff --git a/Engine/Source/Editor/BlueprintGraph/Private/EdGraphSchema_K2.cpp b/Engine/Source/Editor/BlueprintGraph/Private/EdGraphSchema_K2.cpp
--- a/Engine/Source/Editor/BlueprintGraph/Private/EdGraphSchema_K2.cpp
+++ b/Engine/Source/Editor/BlueprintGraph/Private/EdGraphSchema_K2.cpp
@@ -1390,7 +1390,7 @@ bool UEdGraphSchema_K2::PinHasSplittableStructType(const UEdGraphPin* InGraphPin
 		UScriptStruct* StructType = CastChecked<UScriptStruct>(InGraphPin->PinType.PinSubCategoryObject.Get());
 		if (InGraphPin->Direction == EGPD_Input)
 		{
-			bCanSplit = UK2Node_MakeStruct::CanBeMade(StructType);
+			bCanSplit = UK2Node_MakeStruct::CanBeMade(StructType, true, true);
 			if (!bCanSplit)
 			{
 				const FString& MetaData = StructType->GetMetaData(TEXT("HasNativeMake"));
@@ -1400,7 +1400,7 @@ bool UEdGraphSchema_K2::PinHasSplittableStructType(const UEdGraphPin* InGraphPin
 		}
 		else
 		{
-			bCanSplit = UK2Node_BreakStruct::CanBeBroken(StructType);
+			bCanSplit = UK2Node_BreakStruct::CanBeBroken(StructType, true, true);
 			if (!bCanSplit)
 			{
 				const FString& MetaData = StructType->GetMetaData(TEXT("HasNativeBreak"));
@@ -1494,9 +1494,19 @@ void UEdGraphSchema_K2::GetContextMenuActions(const UEdGraph* CurrentGraph, cons
 					}
 				}
 	
-				if (PinHasSplittableStructType(InGraphPin) && InGraphNode->AllowSplitPins())
+				if (InGraphPin->PinType.PinCategory == PC_Struct && InGraphNode->AllowSplitPins())
 				{
-					MenuBuilder->AddMenuEntry( FGraphEditorCommands::Get().SplitStructPin );
+					// If the pin cannot be split, create an error tooltip to use
+					FText Tooltip;
+					if (PinHasSplittableStructType(InGraphPin))
+					{
+						Tooltip = FGraphEditorCommands::Get().SplitStructPin->GetDescription();
+					}
+					else
+					{
+						Tooltip = LOCTEXT("SplitStructPin_Error", "Cannot split the struct pin, it may be missing Blueprint exposed properties!");
+					}
+					MenuBuilder->AddMenuEntry( FGraphEditorCommands::Get().SplitStructPin, NAME_None, FGraphEditorCommands::Get().SplitStructPin->GetLabel(), Tooltip );
 				}
 
 				if (InGraphPin->ParentPin != NULL)
diff --git a/Engine/Source/Editor/BlueprintGraph/Private/K2Node_BreakStruct.cpp b/Engine/Source/Editor/BlueprintGraph/Private/K2Node_BreakStruct.cpp
--- a/Engine/Source/Editor/BlueprintGraph/Private/K2Node_BreakStruct.cpp
+++ b/Engine/Source/Editor/BlueprintGraph/Private/K2Node_BreakStruct.cpp
@@ -134,12 +134,12 @@ static bool CanCreatePinForProperty(const UProperty* Property, bool bIncludeEdit
 	return bVisible && bConvertable;
 }
 
-bool UK2Node_BreakStruct::CanBeBroken(const UScriptStruct* Struct, bool bIncludeEditAnywhere )
+bool UK2Node_BreakStruct::CanBeBroken(const UScriptStruct* Struct, bool bIncludeEditAnywhere, bool bMustHaveValidProperties )
 {
 	const UEdGraphSchema_K2* Schema = GetDefault<UEdGraphSchema_K2>();
 	if(Struct && Schema && !Struct->HasMetaData(TEXT("HasNativeBreak")))
 	{
-		if (UEdGraphSchema_K2::IsAllowableBlueprintVariableType(Struct))
+		if (!bMustHaveValidProperties && UEdGraphSchema_K2::IsAllowableBlueprintVariableType(Struct))
 		{
 			return true;
 		}
diff --git a/Engine/Source/Editor/BlueprintGraph/Private/K2Node_MakeStruct.cpp b/Engine/Source/Editor/BlueprintGraph/Private/K2Node_MakeStruct.cpp
--- a/Engine/Source/Editor/BlueprintGraph/Private/K2Node_MakeStruct.cpp
+++ b/Engine/Source/Editor/BlueprintGraph/Private/K2Node_MakeStruct.cpp
@@ -227,11 +227,11 @@ FLinearColor UK2Node_MakeStruct::GetNodeTitleColor() const
 	return UK2Node::GetNodeTitleColor();
 }
 
-bool UK2Node_MakeStruct::CanBeMade(const UScriptStruct* Struct, bool bIncludeEditOnly )
+bool UK2Node_MakeStruct::CanBeMade(const UScriptStruct* Struct, bool bIncludeEditOnly, bool bMustHaveValidProperties )
 {
 	if (Struct && !Struct->HasMetaData(TEXT("HasNativeMake")))
 	{
-		if (UEdGraphSchema_K2::IsAllowableBlueprintVariableType(Struct))
+		if (!bMustHaveValidProperties && UEdGraphSchema_K2::IsAllowableBlueprintVariableType(Struct))
 		{
 			return true;
 		}
