[UEBP-60] Skipping tooltip generation for transient pins

diff --git a/Engine/Source/Editor/BlueprintGraph/Private/EdGraphSchema_K2.cpp b/Engine/Source/Editor/BlueprintGraph/Private/EdGraphSchema_K2.cpp
--- a/Engine/Source/Editor/BlueprintGraph/Private/EdGraphSchema_K2.cpp
+++ b/Engine/Source/Editor/BlueprintGraph/Private/EdGraphSchema_K2.cpp
@@ -3030,6 +3030,11 @@ FText UEdGraphSchema_K2::GetPinDisplayName(const UEdGraphPin* Pin) const
 
 void UEdGraphSchema_K2::ConstructBasicPinTooltip(const UEdGraphPin& Pin, const FText& PinDescription, FString& TooltipOut) const
 {
+	if (Pin.HasAnyFlags(RF_Transient))
+	{
+		return;
+	}
+
 	if (bGeneratingDocumentation)
 	{
 		TooltipOut = PinDescription.ToString();
diff --git a/Engine/Source/Editor/BlueprintGraph/Private/K2Node_CallFunction.cpp b/Engine/Source/Editor/BlueprintGraph/Private/K2Node_CallFunction.cpp
--- a/Engine/Source/Editor/BlueprintGraph/Private/K2Node_CallFunction.cpp
+++ b/Engine/Source/Editor/BlueprintGraph/Private/K2Node_CallFunction.cpp
@@ -1185,6 +1185,11 @@ FText UK2Node_CallFunction::GetTooltipText() const
 
 void UK2Node_CallFunction::GeneratePinTooltipFromFunction(UEdGraphPin& Pin, const UFunction* Function)
 {
+	if (Pin.HasAnyFlags(RF_Transient))
+	{
+		return;
+	}
+
 	// figure what tag we should be parsing for (is this a return-val pin, or a parameter?)
 	FString ParamName;
 	FString TagStr = TEXT("@param");
