UE-17607 Crash attempting to vect != vect in Math Expression Node

diff --git a/Engine/Source/Editor/BlueprintGraph/Private/EdGraphSchema_K2.cpp b/Engine/Source/Editor/BlueprintGraph/Private/EdGraphSchema_K2.cpp
--- a/Engine/Source/Editor/BlueprintGraph/Private/EdGraphSchema_K2.cpp
+++ b/Engine/Source/Editor/BlueprintGraph/Private/EdGraphSchema_K2.cpp
@@ -2630,14 +2630,22 @@ bool UEdGraphSchema_K2::CreateAutomaticConversionNodeAndConnections(UEdGraphPin*
 
 FString UEdGraphSchema_K2::IsPinDefaultValid(const UEdGraphPin* Pin, const FString& NewDefaultValue, UObject* NewDefaultObject, const FText& InNewDefaultText) const
 {
-	const UBlueprint* OwningBP = FBlueprintEditorUtils::FindBlueprintForNodeChecked(Pin->GetOwningNode());
-	const bool bIsArray = Pin->PinType.bIsArray;
-	const bool bIsReference = Pin->PinType.bIsReference;
-	const bool bIsAutoCreateRefTerm = IsAutoCreateRefTerm(Pin);
+	check(Pin);
 
 	FFormatNamedArguments MessageArgs;
 	MessageArgs.Add(TEXT("PinName"), Pin->GetDisplayName());
 
+	const UBlueprint* OwningBP = FBlueprintEditorUtils::FindBlueprintForNode(Pin->GetOwningNodeUnchecked());
+	if (!OwningBP)
+	{
+		FText MsgFormat = LOCTEXT("NoBlueprintFoundForPin", "No Blueprint was found for the pin '{PinName}'.");
+		return FText::Format(MsgFormat, MessageArgs).ToString();
+	}
+
+	const bool bIsArray = Pin->PinType.bIsArray;
+	const bool bIsReference = Pin->PinType.bIsReference;
+	const bool bIsAutoCreateRefTerm = IsAutoCreateRefTerm(Pin);
+
 	if (OwningBP->BlueprintType != BPTYPE_Interface)
 	{
 		if( !bIsAutoCreateRefTerm )
diff --git a/Engine/Source/Editor/BlueprintGraph/Private/K2Node_MathExpression.cpp b/Engine/Source/Editor/BlueprintGraph/Private/K2Node_MathExpression.cpp
--- a/Engine/Source/Editor/BlueprintGraph/Private/K2Node_MathExpression.cpp
+++ b/Engine/Source/Editor/BlueprintGraph/Private/K2Node_MathExpression.cpp
@@ -1799,7 +1799,7 @@ private:
 									MessageLog.Error(*ErrorText.ToString(), FunctionCall);
 								}
 							}
-							else
+							else if (InputPin->DefaultValue.IsEmpty()) // there is an ErrorTolerance parameter with a default value in EqualEqual_VectorVector
 							{
 								// too many pins - shouldn't be possible due to the checking in FindMatchingFunction() above
 								FText ErrorText = LOCTEXT("ConnectPinError", "The '@@' function requires more parameters than were provided");
