[UE-16820] Provide ability to access class default property values at runtime in a Blueprint script.

diff --git a/Engine/Source/Developer/BlueprintCompilerCppBackend/Private/BlueprintCompilerCppBackend.cpp b/Engine/Source/Developer/BlueprintCompilerCppBackend/Private/BlueprintCompilerCppBackend.cpp
--- a/Engine/Source/Developer/BlueprintCompilerCppBackend/Private/BlueprintCompilerCppBackend.cpp
+++ b/Engine/Source/Developer/BlueprintCompilerCppBackend/Private/BlueprintCompilerCppBackend.cpp
@@ -166,7 +166,7 @@ FString FBlueprintCompilerCppBackend::TermToText(const FBPTerminal* Term, const
 		{
 			Prefix = TermToText(Term->Context);
 
-			if (Term->Context->bIsStructContext)
+			if (Term->Context->IsStructContextType())
 			{
 				Prefix += TEXT(".");
 			}
diff --git a/Engine/Source/Developer/BlueprintCompilerCppBackend/Private/BlueprintCompilerCppBackendUtils.h b/Engine/Source/Developer/BlueprintCompilerCppBackend/Private/BlueprintCompilerCppBackendUtils.h
--- a/Engine/Source/Developer/BlueprintCompilerCppBackend/Private/BlueprintCompilerCppBackendUtils.h
+++ b/Engine/Source/Developer/BlueprintCompilerCppBackend/Private/BlueprintCompilerCppBackendUtils.h
@@ -24,7 +24,7 @@ public:
 		TArray<FString> SafetyConditions;
 		for (; Term; Term = Term->Context)
 		{
-			if (!Term->bIsStructContext && (Term->Type.PinSubCategory != TEXT("self")))
+			if (!Term->IsStructContextType() && (Term->Type.PinSubCategory != TEXT("self")))
 			{
 				ensure(!Term->bIsLiteral);
 				SafetyConditions.Add(CppBackend.TermToText(Term));
diff --git a/Engine/Source/Editor/BlueprintGraph/Classes/K2Node_GetClassDefaults.h b/Engine/Source/Editor/BlueprintGraph/Classes/K2Node_GetClassDefaults.h
new file mode 100644
--- /dev/null
+++ b/Engine/Source/Editor/BlueprintGraph/Classes/K2Node_GetClassDefaults.h
@@ -0,0 +1,72 @@
+// Copyright 1998-2015 Epic Games, Inc. All Rights Reserved.
+
+#pragma once
+#include "K2Node.h"
+#include "K2Node_GetClassDefaults.generated.h"
+
+UCLASS(MinimalAPI)
+class UK2Node_GetClassDefaults : public UK2Node
+{
+	GENERATED_UCLASS_BODY()
+
+	// Begin UObject interface
+	virtual void PostEditChangeProperty(struct FPropertyChangedEvent& PropertyChangedEvent) override;
+	// End UObject interface
+
+	// Begin UEdGraphNode interface
+	virtual FText GetNodeTitle(ENodeTitleType::Type TitleType) const override;
+	virtual void AllocateDefaultPins() override;
+	virtual void PostPlacedNewNode() override;
+	virtual void PinDefaultValueChanged(UEdGraphPin* Pin) override;
+	virtual void PinConnectionListChanged(UEdGraphPin* Pin) override;
+	// End UEdGraphNode interface
+
+	// Begin UK2Node interface
+	virtual bool IsNodePure() const override { return true; }
+	virtual bool ShouldShowNodeProperties() const override { return true; }
+	virtual void ReallocatePinsDuringReconstruction(TArray<UEdGraphPin*>& OldPins) override;
+	virtual class FNodeHandlingFunctor* CreateNodeHandler(class FKismetCompilerContext& CompilerContext) const override;
+	virtual void GetMenuActions(FBlueprintActionDatabaseRegistrar& ActionRegistrar) const override;
+	virtual FText GetMenuCategory() const override;
+	// End UK2Node interface
+
+public:
+	/** Finds and returns the class input pin from the current set of pins. */
+	UEdGraphPin* FindClassPin() const
+	{
+		return FindClassPin(Pins);
+	}
+
+protected:
+	/**
+	 * Finds and returns the class input pin.
+	 *
+	 * @param FromPins	A list of pins to search.
+	 */
+	UEdGraphPin* FindClassPin(const TArray<UEdGraphPin*>& FromPins) const;
+
+	/**
+	 * Determines the input class type from the given pin.
+	 *
+	 * @param FromPin	Input class pin. If not set (default), then it will fall back to using FindClassPin().
+	 */
+	UClass* GetInputClass(const UEdGraphPin* FromPin = nullptr) const;
+
+	/**
+	 * Creates the full set of output pins (properties) from the given input class.
+	 *
+	 * @param InClass	Input class type.
+	 */
+	void CreateOutputPins(UClass* InClass);
+
+	/** Will be called whenever the class pin selector changes its value. */
+	void OnClassPinChanged();
+
+private:
+	/** Class pin name */
+	static FString ClassPinName;
+
+	/** Output pin visibility control */
+	UPROPERTY(EditAnywhere, Category=PinOptions, EditFixedSize)
+	TArray<FOptionalPinFromProperty> ShowPinForProperties;
+};
diff --git a/Engine/Source/Editor/BlueprintGraph/Private/K2Node_GetClassDefaults.cpp b/Engine/Source/Editor/BlueprintGraph/Private/K2Node_GetClassDefaults.cpp
new file mode 100644
--- /dev/null
+++ b/Engine/Source/Editor/BlueprintGraph/Private/K2Node_GetClassDefaults.cpp
@@ -0,0 +1,348 @@
+// Copyright 1998-2015 Epic Games, Inc. All Rights Reserved.
+
+#include "BlueprintGraphPrivatePCH.h"
+#include "EditorCategoryUtils.h"
+#include "BlueprintNodeSpawner.h"
+#include "BlueprintActionDatabaseRegistrar.h"
+#include "KismetCompiler.h"
+#include "K2Node_GetClassDefaults.h"
+
+#define LOCTEXT_NAMESPACE "UK2Node_GetClassDefaults"
+
+FString UK2Node_GetClassDefaults::ClassPinName(TEXT("Class"));
+
+namespace
+{
+	// Optional pin manager subclass.
+	struct FClassDefaultsOptionalPinManager : public FOptionalPinManager
+	{
+		FClassDefaultsOptionalPinManager(UClass* InClass)
+			:FOptionalPinManager()
+		{
+			SrcClass = InClass;
+		}
+
+		virtual void GetRecordDefaults(UProperty* TestProperty, FOptionalPinFromProperty& Record) const override
+		{
+			FOptionalPinManager::GetRecordDefaults(TestProperty, Record);
+
+			// Show pin unless the property is owned by a parent class.
+			Record.bShowPin = TestProperty->GetOwnerClass() == SrcClass;
+		}
+
+		virtual bool CanTreatPropertyAsOptional(UProperty* TestProperty) const override
+		{
+			// Don't expose object reference properties or anything not marked BlueprintReadOnly/BlueprintReadWrite.
+			// @TODO - Could potentially expose object reference values if/when we have support for 'const' input pins.
+			return TestProperty != nullptr
+				&& !TestProperty->IsA<UClassProperty>()
+				&& !TestProperty->IsA<UObjectProperty>()
+				&& !TestProperty->IsA<UInterfaceProperty>()
+				&& TestProperty->HasAllPropertyFlags(CPF_BlueprintVisible);
+		}
+
+		virtual void CustomizePinData(UEdGraphPin* Pin, FName SourcePropertyName, int32 ArrayIndex, UProperty* Property = nullptr) const override
+		{
+			check(Pin);
+
+			// Move into the advanced view if the property metadata is set.
+			Pin->bAdvancedView = Property && Property->HasAnyPropertyFlags(CPF_AdvancedDisplay);
+		}
+
+	private:
+		// Class type for which optional pins are being managed.
+		UClass* SrcClass;
+	};
+
+	// Compilation handler subclass.
+	class FKCHandler_GetClassDefaults : public FNodeHandlingFunctor
+	{
+	public:
+		FKCHandler_GetClassDefaults(FKismetCompilerContext& InCompilerContext)
+			: FNodeHandlingFunctor(InCompilerContext)
+		{
+		}
+
+		virtual void RegisterNets(FKismetFunctionContext& Context, UEdGraphNode* Node) override
+		{
+			// Cast to the correct node type
+			if(const UK2Node_GetClassDefaults* GetClassDefaultsNode = Cast<UK2Node_GetClassDefaults>(Node))
+			{
+				// Only if we have a valid class input pin
+				if(UEdGraphPin* ClassPin = GetClassDefaultsNode->FindClassPin())
+				{
+					// Redirect to a linked pin if necessary
+					UEdGraphPin* Net = FEdGraphUtilities::GetNetFromPin(ClassPin);
+					check(Net != nullptr);
+
+					// Register a literal if necessary (e.g. there are no linked pins)
+					if(ValidateAndRegisterNetIfLiteral(Context, Net))
+					{
+						// First check for a literal term in case one was created above
+						FBPTerminal** FoundTerm = Context.LiteralHackMap.Find(Net);
+						if(FoundTerm == nullptr)
+						{
+							// Otherwise, check for a linked term
+							FoundTerm = Context.NetMap.Find(Net);
+						}
+
+						// If we did not find an input term, make sure we create one here
+						FBPTerminal* ClassContextTerm = FoundTerm ? *FoundTerm : nullptr;
+						if(ClassContextTerm == nullptr)
+						{
+							ClassContextTerm = Context.CreateLocalTerminalFromPinAutoChooseScope(Net, Context.NetNameMap->MakeValidName(Net));
+							check(ClassContextTerm != nullptr);
+
+							Context.NetMap.Add(Net, ClassContextTerm);
+						}
+
+						// Flag this as a "class context" term
+						ClassContextTerm->SetContextTypeClass();
+
+						// Infer the class type from the context term
+						if(const UClass* ClassType = Cast<UClass>(ClassContextTerm->bIsLiteral ? ClassContextTerm->ObjectLiteral : ClassContextTerm->Type.PinSubCategoryObject.Get()))
+						{
+							// Create a local term for each output pin (class property)
+							for(int32 PinIndex = 0; PinIndex < Node->Pins.Num(); ++PinIndex)
+							{
+								UEdGraphPin* Pin = Node->Pins[PinIndex];
+								if(Pin != nullptr && Pin->Direction == EGPD_Output)
+								{
+									UProperty* BoundProperty = FindField<UProperty>(ClassType, *(Pin->PinName));
+									if(BoundProperty != nullptr)
+									{
+										FBPTerminal* OutputTerm = Context.CreateLocalTerminalFromPinAutoChooseScope(Pin, Pin->PinName);
+										check(OutputTerm != nullptr);
+
+										// Set as a variable within the class context
+										OutputTerm->AssociatedVarProperty = BoundProperty;
+										OutputTerm->Context = ClassContextTerm;
+
+										// Flag this as a "class default" variable term
+										OutputTerm->bIsConst = true;
+										OutputTerm->SetVarTypeDefault();
+
+										// Add it to the lookup table
+										Context.NetMap.Add(Pin, OutputTerm);
+									}
+									else
+									{
+										CompilerContext.MessageLog.Error(*LOCTEXT("UnmatchedOutputPinOnCompile", "Failed to find a class member to match @@").ToString(), Pin);
+									}
+								}
+							}
+						}
+						else
+						{
+							CompilerContext.MessageLog.Error(*LOCTEXT("InvalidClassTypeOnCompile", "Missing or invalid input class type for @@").ToString(), Node);
+						}
+					}
+				}
+			}
+		}
+	};
+}
+
+UK2Node_GetClassDefaults::UK2Node_GetClassDefaults(const FObjectInitializer& ObjectInitializer)
+	: Super(ObjectInitializer)
+{
+}
+
+void UK2Node_GetClassDefaults::PostEditChangeProperty(struct FPropertyChangedEvent& PropertyChangedEvent)
+{
+	// Rebuild the node if the "Show Pin" property was changed
+	FName PropertyName = (PropertyChangedEvent.Property != NULL) ? PropertyChangedEvent.Property->GetFName() : NAME_None;
+	if((PropertyName == TEXT("bShowPin")))
+	{
+		GetSchema()->ReconstructNode(*this);
+	}
+
+	Super::PostEditChangeProperty(PropertyChangedEvent);
+}
+
+FText UK2Node_GetClassDefaults::GetNodeTitle(ENodeTitleType::Type TitleType) const
+{
+	// @TODO - unique node icon needed as well?
+
+	return LOCTEXT("NodeTitle", "Get Class Defaults");
+}
+
+void UK2Node_GetClassDefaults::AllocateDefaultPins()
+{
+	const UEdGraphSchema_K2* K2Schema = GetDefault<UEdGraphSchema_K2>();
+
+	// Create the class input type selector pin
+	UEdGraphPin* ClassPin = CreatePin(EGPD_Input, K2Schema->PC_Class, TEXT(""), UObject::StaticClass(), false, false, ClassPinName);
+	K2Schema->ConstructBasicPinTooltip(*ClassPin, LOCTEXT("ClassPinDescription", "The class from which to access one or more default values."), ClassPin->PinToolTip);
+}
+
+void UK2Node_GetClassDefaults::PostPlacedNewNode()
+{
+	if(UEdGraphPin* ClassPin = FindClassPin(Pins))
+	{
+		// Default to the owner BP's generated class for "normal" BPs if this is a new node
+		const UBlueprint* OwnerBlueprint = GetBlueprint();
+		if(OwnerBlueprint != nullptr && OwnerBlueprint->BlueprintType == BPTYPE_Normal)
+		{
+			ClassPin->DefaultObject = OwnerBlueprint->GeneratedClass;
+		}
+
+		if(UClass* InputClass = GetInputClass(ClassPin))
+		{
+			CreateOutputPins(InputClass);
+		}
+	}
+}
+
+void UK2Node_GetClassDefaults::PinConnectionListChanged(UEdGraphPin* ChangedPin)
+{
+	if(ChangedPin != nullptr && ChangedPin->PinName == ClassPinName && ChangedPin->Direction == EGPD_Input)
+	{
+		OnClassPinChanged();
+	}
+}
+
+void UK2Node_GetClassDefaults::PinDefaultValueChanged(UEdGraphPin* ChangedPin) 
+{
+	if(ChangedPin != nullptr && ChangedPin->PinName == ClassPinName && ChangedPin->Direction == EGPD_Input)
+	{
+		OnClassPinChanged();
+	}
+}
+
+void UK2Node_GetClassDefaults::ReallocatePinsDuringReconstruction(TArray<UEdGraphPin*>& OldPins) 
+{
+	Super::ReallocatePinsDuringReconstruction(OldPins);
+
+	// Recreate output pins based on the previous input class
+	UEdGraphPin* OldClassPin = FindClassPin(OldPins);
+	if(UClass* InputClass = GetInputClass(OldClassPin))
+	{
+		CreateOutputPins(InputClass);
+	}
+}
+
+FNodeHandlingFunctor* UK2Node_GetClassDefaults::CreateNodeHandler(FKismetCompilerContext& CompilerContext) const
+{
+	return new FKCHandler_GetClassDefaults(CompilerContext);
+}
+
+void UK2Node_GetClassDefaults::GetMenuActions(FBlueprintActionDatabaseRegistrar& ActionRegistrar) const
+{
+	UClass* ActionKey = GetClass();
+	if (ActionRegistrar.IsOpenForRegistration(ActionKey))
+	{
+		UBlueprintNodeSpawner* NodeSpawner = UBlueprintNodeSpawner::Create(GetClass());
+		check(NodeSpawner != nullptr);
+
+		ActionRegistrar.AddBlueprintAction(ActionKey, NodeSpawner);
+	}
+}
+
+FText UK2Node_GetClassDefaults::GetMenuCategory() const
+{
+	return FEditorCategoryUtils::GetCommonCategory(FCommonEditorCategory::Class);
+}
+
+UEdGraphPin* UK2Node_GetClassDefaults::FindClassPin(const TArray<UEdGraphPin*>& FromPins) const
+{
+	UEdGraphPin* const* FoundPin = FromPins.FindByPredicate([](const UEdGraphPin* CurPin)
+	{
+		return CurPin && CurPin->Direction == EGPD_Input && CurPin->PinName == ClassPinName;
+	});
+
+	return FoundPin ? *FoundPin : nullptr;
+}
+
+UClass* UK2Node_GetClassDefaults::GetInputClass(const UEdGraphPin* FromPin) const
+{
+	UClass* InputClass = nullptr;
+
+	const UEdGraphPin* ClassPin = FromPin;
+	if(ClassPin == nullptr)
+	{
+		ClassPin = FindClassPin();
+	}
+	
+	if(ClassPin != nullptr)
+	{
+		check(ClassPin->Direction == EGPD_Input);
+
+		if(ClassPin->DefaultObject != nullptr && ClassPin->LinkedTo.Num() == 0)
+		{
+			InputClass = CastChecked<UClass>(ClassPin->DefaultObject);
+		}
+		else if(ClassPin->LinkedTo.Num() > 0)
+		{
+			if(UEdGraphPin* LinkedPin = ClassPin->LinkedTo[0])
+			{
+				InputClass = Cast<UClass>(LinkedPin->PinType.PinSubCategoryObject.Get());
+			}
+		}
+	}
+
+	return InputClass;
+}
+
+void UK2Node_GetClassDefaults::CreateOutputPins(UClass* InClass)
+{
+	// Create the set of output pins through the optional pin manager
+	FClassDefaultsOptionalPinManager OptionalPinManager(InClass);
+	OptionalPinManager.RebuildPropertyList(ShowPinForProperties, InClass);
+	OptionalPinManager.CreateVisiblePins(ShowPinForProperties, InClass, EGPD_Output, this);
+
+	// Check for any advanced properties (outputs)
+	bool bHasAdvancedPins = false;
+	for(int32 PinIndex = 0; PinIndex < Pins.Num() && !bHasAdvancedPins; ++PinIndex)
+	{
+		UEdGraphPin* Pin = Pins[PinIndex];
+		check(Pin != nullptr);
+
+		bHasAdvancedPins |= Pin->bAdvancedView;
+	}
+
+	// Toggle advanced display on/off based on whether or not we have any advanced outputs
+	if(bHasAdvancedPins && AdvancedPinDisplay == ENodeAdvancedPins::NoPins)
+	{
+		AdvancedPinDisplay = ENodeAdvancedPins::Hidden;
+	}
+	else if(!bHasAdvancedPins)
+	{
+		AdvancedPinDisplay = ENodeAdvancedPins::NoPins;
+	}
+}
+
+void UK2Node_GetClassDefaults::OnClassPinChanged()
+{
+	TArray<UEdGraphPin*> OldPins = Pins;
+	TArray<UEdGraphPin*> OldOutputPins;
+
+	// Gather all current output pins
+	for(int32 PinIndex = 0; PinIndex < OldPins.Num(); ++PinIndex)
+	{
+		UEdGraphPin* OldPin = OldPins[PinIndex];
+		if(OldPin->Direction == EGPD_Output)
+		{
+			Pins.Remove(OldPin);
+			OldOutputPins.Add(OldPin);
+		}
+	}
+
+	// Clear the current output pin settings (so they don't carry over to the new set)
+	ShowPinForProperties.Empty();
+
+	// Create output pins for the new class type
+	UClass* InputClass = GetInputClass();
+	CreateOutputPins(InputClass);
+
+	// Destroy the previous set of output pins
+	DestroyPinList(OldOutputPins);
+
+	// Notify the graph that the node has been changed
+	if(UEdGraph* Graph = GetGraph())
+	{
+		Graph->NotifyGraphChanged();
+	}
+}
+
+#undef LOCTEXT_NAMESPACE
diff --git a/Engine/Source/Editor/BlueprintGraph/Private/StructMemberNodeHandlers.cpp b/Engine/Source/Editor/BlueprintGraph/Private/StructMemberNodeHandlers.cpp
--- a/Engine/Source/Editor/BlueprintGraph/Private/StructMemberNodeHandlers.cpp
+++ b/Engine/Source/Editor/BlueprintGraph/Private/StructMemberNodeHandlers.cpp
@@ -22,7 +22,7 @@ static FBPTerminal* RegisterStructVar(FCompilerResultsLog& MessageLog, FKismetFu
 		Schema->ConvertPropertyToPinType(BoundProperty, /*out*/ Term->Type);
 		Term->Source = MemberSetNode;
 		Term->Name = MemberSetNode->GetVarNameString();
-		Term->bIsStructContext = true;
+		Term->SetContextTypeStruct();
 
 		Term->AssociatedVarProperty = BoundProperty;
 		//@TODO: Needed? Context.NetMap.Add(Net, Term);
diff --git a/Engine/Source/Editor/DetailCustomizations/Private/DetailCustomizations.cpp b/Engine/Source/Editor/DetailCustomizations/Private/DetailCustomizations.cpp
--- a/Engine/Source/Editor/DetailCustomizations/Private/DetailCustomizations.cpp
+++ b/Engine/Source/Editor/DetailCustomizations/Private/DetailCustomizations.cpp
@@ -243,6 +243,7 @@ void FDetailCustomizationsModule::RegisterObjectCustomizations()
 	//@TODO: A2REMOVAL: Rename FSkeletalControlNodeDetails to something more generic
 	RegisterCustomClassLayout("K2Node_StructMemberGet", FOnGetDetailCustomizationInstance::CreateStatic(&FSkeletalControlNodeDetails::MakeInstance));
 	RegisterCustomClassLayout("K2Node_StructMemberSet", FOnGetDetailCustomizationInstance::CreateStatic(&FSkeletalControlNodeDetails::MakeInstance));
+	RegisterCustomClassLayout("K2Node_GetClassDefaults", FOnGetDetailCustomizationInstance::CreateStatic(&FSkeletalControlNodeDetails::MakeInstance));
 
 	RegisterCustomClassLayout("AnimSequence", FOnGetDetailCustomizationInstance::CreateStatic(&FAnimSequenceDetails::MakeInstance));
 	RegisterCustomClassLayout("Rig", FOnGetDetailCustomizationInstance::CreateStatic(&FRigDetails::MakeInstance));
diff --git a/Engine/Source/Editor/KismetCompiler/Private/KismetCompiler.cpp b/Engine/Source/Editor/KismetCompiler/Private/KismetCompiler.cpp
--- a/Engine/Source/Editor/KismetCompiler/Private/KismetCompiler.cpp
+++ b/Engine/Source/Editor/KismetCompiler/Private/KismetCompiler.cpp
@@ -640,9 +640,8 @@ void FKismetCompilerContext::CreatePropertiesFromList(UStruct* Scope, UField**&
 		FBPTerminal& Term = Terms[i];
 
 		if(NULL != Term.AssociatedVarProperty)
-		{	
-			const bool bIsStructMember = (Term.Context && Term.Context->bIsStructContext);
-			if(bIsStructMember)
+		{
+			if(Term.Context && !Term.Context->IsObjectContextType())
 			{
 				continue;
 			}
@@ -717,7 +716,7 @@ void FKismetCompilerContext::CreatePropertiesFromList(UStruct* Scope, UField**&
 			PropertyStorageLocation = &(NewProperty->Next);
 
 			Term.AssociatedVarProperty = NewProperty;
-			Term.bIsLocal = bPropertiesAreLocal;
+			Term.SetVarTypeLocal(bPropertiesAreLocal);
 
 			// Record in the debugging information
 			//@TODO: Rename RegisterClassPropertyAssociation, etc..., to better match that indicate it works with locals
diff --git a/Engine/Source/Editor/KismetCompiler/Private/KismetCompilerMisc.cpp b/Engine/Source/Editor/KismetCompiler/Private/KismetCompilerMisc.cpp
--- a/Engine/Source/Editor/KismetCompiler/Private/KismetCompilerMisc.cpp
+++ b/Engine/Source/Editor/KismetCompiler/Private/KismetCompilerMisc.cpp
@@ -1198,7 +1198,7 @@ void FNodeHandlingFunctor::ResolveAndRegisterScopedTerm(FKismetFunctionContext&
 		// Check if the property is a local variable and mark it so
 		if( SearchScope == Context.Function && BoundProperty->GetOuter() == Context.Function)
 		{
-			Term->bIsLocal = true;
+			Term->SetVarTypeLocal(true);
 		}
 		else if (BoundProperty->HasAnyPropertyFlags(CPF_BlueprintReadOnly) || (Context.IsConstFunction() && Context.NewClass->IsChildOf(SearchScope)))
 		{
@@ -1799,7 +1799,7 @@ FBPTerminal* FKismetFunctionContext::CreateLocalTerminal(ETerminalSpecification
 	default:
 		const bool bIsLocal = !IsEventGraph();
 		Result = new (bIsLocal ? Locals : EventGraphLocals) FBPTerminal();
-		Result->bIsLocal = bIsLocal;
+		Result->SetVarTypeLocal(bIsLocal);
 		break;
 	}
 	return Result;
@@ -1842,5 +1842,5 @@ void FBPTerminal::CopyFromPin(UEdGraphPin* Net, const FString& NewName)
 	const UEdGraphSchema_K2* Schema = Cast<const UEdGraphSchema_K2>(Net->GetSchema());
 	const bool bStructCategory = Schema && (Schema->PC_Struct == Net->PinType.PinCategory);
 	const bool bStructSubCategoryObj = (NULL != Cast<UScriptStruct>(Net->PinType.PinSubCategoryObject.Get()));
-	bIsStructContext = bStructCategory && bStructSubCategoryObj;
+	SetContextTypeStruct(bStructCategory && bStructSubCategoryObj);
 }
diff --git a/Engine/Source/Editor/KismetCompiler/Private/KismetCompilerVMBackend.cpp b/Engine/Source/Editor/KismetCompiler/Private/KismetCompilerVMBackend.cpp
--- a/Engine/Source/Editor/KismetCompiler/Private/KismetCompilerVMBackend.cpp
+++ b/Engine/Source/Editor/KismetCompiler/Private/KismetCompilerVMBackend.cpp
@@ -253,19 +253,26 @@ protected:
 		{
 			bInContext = true;
 
-			static const FBoolConfigValueHelper CanSuppressAccessViolation(TEXT("Kismet"), TEXT("bCanSuppressAccessViolation"), GEngineIni);
-			if (bUnsafeToSkip || !CanSuppressAccessViolation)
+			if(Term->IsClassContextType())
 			{
-				Writer << EX_Context;
+				Writer << EX_ClassContext;
 			}
 			else
 			{
-				Writer << EX_Context_FailSilent;
-			}
+				static const FBoolConfigValueHelper CanSuppressAccessViolation(TEXT("Kismet"), TEXT("bCanSuppressAccessViolation"), GEngineIni);
+				if (bUnsafeToSkip || !CanSuppressAccessViolation)
+				{
+					Writer << EX_Context;
+				}
+				else
+				{
+					Writer << EX_Context_FailSilent;
+				}
 
-			if (bIsInterfaceContext)
-			{
-				Writer << EX_InterfaceContext;
+				if (bIsInterfaceContext)
+				{
+					Writer << EX_InterfaceContext;
+				}
 			}
 
 			// Function contexts must always be objects, so if we have a literal, give it the default object property so the compiler knows how to handle it
@@ -608,7 +615,11 @@ public:
 		else
 		{
 			check(Term->AssociatedVarProperty);
-			if (Term->bIsLocal)
+			if (Term->IsDefaultVarTerm())
+			{
+				Writer << EX_DefaultVariable;
+			}
+			else if (Term->IsLocalVarTerm())
 			{
 				Writer << (Term->AssociatedVarProperty->HasAnyPropertyFlags(CPF_OutParm) ? EX_LocalOutVariable : EX_LocalVariable);
 			}
@@ -879,7 +890,7 @@ public:
 		}
 		else
 		{
-			if (Term->Context->bIsStructContext)
+			if (Term->Context->IsStructContextType())
 			{
 				check(Term->AssociatedVarProperty);
 
diff --git a/Engine/Source/Editor/KismetCompiler/Public/BPTerminal.h b/Engine/Source/Editor/KismetCompiler/Public/BPTerminal.h
--- a/Engine/Source/Editor/KismetCompiler/Public/BPTerminal.h
+++ b/Engine/Source/Editor/KismetCompiler/Public/BPTerminal.h
@@ -11,12 +11,9 @@ struct FBPTerminal
 {
 	FString Name;
 	FEdGraphPinType Type;
-	bool bIsLiteral;
-	bool bIsLocal; // @TODO: Should be an enumeration or driven from the var property
+	bool bIsLiteral; 
 	bool bIsConst;
 	bool bIsSavePersistent;
-
-	bool bIsStructContext;
 	bool bPassedByReference;
 
 	// Source node (e.g., pin)
@@ -42,16 +39,16 @@ struct FBPTerminal
 
 	FBPTerminal()
 		: bIsLiteral(false)
-		, bIsLocal(false)
 		, bIsConst(false)
 		, bIsSavePersistent(false)
-		, bIsStructContext(false)
 		, bPassedByReference(false)
 		, Source(nullptr)
 		, Context(nullptr)
 		, AssociatedVarProperty(nullptr)
 		, ObjectLiteral(nullptr)
 		, InlineGeneratedParameter(nullptr)
+		, VarType(EVarType_Instanced)
+		, ContextType(EContextType_Object)
 	{
 	}
 
@@ -61,6 +58,79 @@ struct FBPTerminal
 	{
 		return !bIsLiteral && !bIsConst;
 	}
+
+	bool IsLocalVarTerm() const
+	{
+		return !bIsLiteral && VarType == EVarType_Local;
+	}
+
+	void SetVarTypeLocal(bool bIsLocal = true)
+	{
+		VarType = bIsLocal ? EVarType_Local : EVarType_Instanced;
+	}
+
+	bool IsDefaultVarTerm() const
+	{
+		return !bIsLiteral && VarType == EVarType_Default;
+	}
+
+	void SetVarTypeDefault(bool bIsDefault = true)
+	{
+		VarType = bIsDefault ? EVarType_Default : EVarType_Instanced;
+	}
+
+	bool IsInstancedVarTerm() const
+	{
+		return !bIsLiteral && VarType == EVarType_Instanced;
+	}
+
+	bool IsClassContextType() const
+	{
+		return ContextType == EContextType_Class;
+	}
+
+	void SetContextTypeClass(bool bIsClassContext = true)
+	{
+		ContextType = bIsClassContext ? EContextType_Class : EContextType_Object;
+	}
+
+	bool IsStructContextType() const
+	{
+		return ContextType == EContextType_Struct;
+	}
+
+	void SetContextTypeStruct(bool bIsStructContext = true)
+	{
+		ContextType = bIsStructContext ? EContextType_Struct : EContextType_Object;
+	}
+
+	bool IsObjectContextType() const
+	{
+		return ContextType == EContextType_Object;
+	}
+
+private:
+	// Variable reference types (mutually-exclusive)
+	enum EVarType
+	{
+		EVarType_Local,
+		EVarType_Default,
+		EVarType_Instanced
+	};
+
+	// For non-literal terms, this is set to the type of variable reference
+	EVarType VarType;
+
+	// Context types (mutually-exclusive)
+	enum EContextType
+	{
+		EContextType_Class,
+		EContextType_Struct,
+		EContextType_Object,
+	};
+
+	// If this term is also a context, this indicates which type of context it is
+	EContextType ContextType;
 };
 
 //////////////////////////////////////////////////////////////////////////
diff --git a/Engine/Source/Editor/UnrealEd/Private/ScriptDisassembler.cpp b/Engine/Source/Editor/UnrealEd/Private/ScriptDisassembler.cpp
--- a/Engine/Source/Editor/UnrealEd/Private/ScriptDisassembler.cpp
+++ b/Engine/Source/Editor/UnrealEd/Private/ScriptDisassembler.cpp
@@ -380,6 +380,12 @@ void FKismetBytecodeDisassembler::ProcessCommon(int32& ScriptIndex, EExprToken O
 			Ar.Logf(TEXT("%s $%X: Local variable named %s"), *Indents, (int32)Opcode, PropertyPtr ? *PropertyPtr->GetName() : TEXT("(null)"));
 			break;
 		}
+	case EX_DefaultVariable:
+		{
+			UProperty* PropertyPtr = ReadPointer<UProperty>(ScriptIndex);
+			Ar.Logf(TEXT("%s $%X: Default variable named %s"), *Indents, (int32)Opcode, PropertyPtr ? *PropertyPtr->GetName() : TEXT("(null)"));
+			break;
+		}
 	case EX_InstanceVariable:
 		{
 			UProperty* PropertyPtr = ReadPointer<UProperty>(ScriptIndex);
@@ -524,10 +530,11 @@ void FKismetBytecodeDisassembler::ProcessCommon(int32& ScriptIndex, EExprToken O
 			}
 			break;
 		}
+	case EX_ClassContext:
 	case EX_Context:
 	case EX_Context_FailSilent:
 		{
-			Ar.Logf(TEXT("%s $%X: %s"), *Indents, (int32)Opcode, TEXT("Context"));
+			Ar.Logf(TEXT("%s $%X: %s"), *Indents, (int32)Opcode, Opcode == EX_ClassContext ? TEXT("Class Context") : TEXT("Context"));
 			AddIndent();
 
 			// Object expression.
diff --git a/Engine/Source/Runtime/CoreUObject/Private/UObject/ScriptCore.cpp b/Engine/Source/Runtime/CoreUObject/Private/UObject/ScriptCore.cpp
--- a/Engine/Source/Runtime/CoreUObject/Private/UObject/ScriptCore.cpp
+++ b/Engine/Source/Runtime/CoreUObject/Private/UObject/ScriptCore.cpp
@@ -1121,6 +1121,46 @@ void UObject::execInstanceVariable(FFrame& Stack, RESULT_DECL)
 }
 IMPLEMENT_VM_FUNCTION( EX_InstanceVariable, execInstanceVariable );
 
+void UObject::execDefaultVariable(FFrame& Stack, RESULT_DECL)
+{
+	UProperty* VarProperty = (UProperty*)Stack.ReadObject();
+	Stack.MostRecentProperty = VarProperty;
+	Stack.MostRecentPropertyAddress = nullptr;
+
+	if(VarProperty == nullptr)
+	{
+		static FBlueprintExceptionInfo ExceptionInfo(EBlueprintExceptionType::AccessViolation, TEXT("Attempt to access a missing property. If this is a packaged/cooked build, are you attempting to use an editor-only property?"));
+		FBlueprintCoreDelegates::ThrowScriptException(this, Stack, ExceptionInfo);
+	}
+	else
+	{
+		UObject* DefaultObject = nullptr;
+		if(HasAnyFlags(RF_ClassDefaultObject))
+		{
+			DefaultObject = this;
+		}
+		else
+		{
+			// @todo - allow access to archetype properties through object references?
+		}
+
+		if(DefaultObject != nullptr)
+		{
+			Stack.MostRecentPropertyAddress = VarProperty->ContainerPtrToValuePtr<uint8>(DefaultObject);
+			if(RESULT_PARAM)
+			{
+				VarProperty->CopyCompleteValueToScriptVM(RESULT_PARAM, Stack.MostRecentPropertyAddress);
+			}
+		}
+		else
+		{
+			static FBlueprintExceptionInfo ExceptionInfo(EBlueprintExceptionType::AccessViolation, TEXT("Attempt to access a default property through an invalid or otherwise unsupported context."));
+			FBlueprintCoreDelegates::ThrowScriptException(this, Stack, ExceptionInfo);
+		}
+	}
+}
+IMPLEMENT_VM_FUNCTION( EX_DefaultVariable, execDefaultVariable );
+
 void UObject::execLocalOutVariable(FFrame& Stack, RESULT_DECL)
 {
 	checkSlow(Stack.Object == this);
@@ -1160,6 +1200,51 @@ void UObject::execInterfaceContext(FFrame& Stack, RESULT_DECL)
 }
 IMPLEMENT_VM_FUNCTION( EX_InterfaceContext, execInterfaceContext );
 
+void UObject::execClassContext(FFrame& Stack, RESULT_DECL)
+{
+	// Get class expression.
+	UClass* ClassContext = NULL;
+	Stack.Step(this, &ClassContext);
+
+	// Execute expression in class context.
+	if(IsValid(ClassContext))
+	{
+		UObject* DefaultObject = ClassContext->GetDefaultObject();
+		check(DefaultObject != NULL);
+
+		Stack.Code += sizeof(CodeSkipSizeType)	// Code offset for NULL expressions.
+			+ sizeof(ScriptPointerType);		// Property corresponding to the r-value data, in case the l-value needs to be cleared
+		Stack.Step(DefaultObject, RESULT_PARAM);
+	}
+	else
+	{
+		if (Stack.MostRecentProperty != NULL)
+		{
+			const FString Desc = FString::Printf(TEXT("Accessed null class context '%s'"), *Stack.MostRecentProperty->GetName());
+			FBlueprintExceptionInfo ExceptionInfo(EBlueprintExceptionType::AccessViolation, Desc);
+			FBlueprintCoreDelegates::ThrowScriptException(this, Stack, ExceptionInfo);
+		}
+		else
+		{
+			static FBlueprintExceptionInfo ExceptionInfo(EBlueprintExceptionType::AccessViolation, TEXT("Accessed null class context"));
+			FBlueprintCoreDelegates::ThrowScriptException(this, Stack, ExceptionInfo);
+		}
+
+		const CodeSkipSizeType wSkip = Stack.ReadCodeSkipCount(); // Code offset for NULL expressions. Code += sizeof(CodeSkipSizeType)
+		UProperty* RValueProperty = nullptr;
+		const VariableSizeType bSize = Stack.ReadVariableSize(&RValueProperty); // Code += sizeof(ScriptPointerType) + sizeof(uint8)
+		Stack.Code += wSkip;
+		Stack.MostRecentPropertyAddress = NULL;
+		Stack.MostRecentProperty = NULL;
+
+		if (RESULT_PARAM && RValueProperty)
+		{
+			RValueProperty->ClearValue(RESULT_PARAM);
+		}
+	}
+}
+IMPLEMENT_VM_FUNCTION( EX_ClassContext, execClassContext );
+
 void UObject::execEndOfScript( FFrame& Stack, RESULT_DECL )
 {
 #if WITH_EDITOR
diff --git a/Engine/Source/Runtime/CoreUObject/Public/UObject/Script.h b/Engine/Source/Runtime/CoreUObject/Public/UObject/Script.h
--- a/Engine/Source/Runtime/CoreUObject/Public/UObject/Script.h
+++ b/Engine/Source/Runtime/CoreUObject/Public/UObject/Script.h
@@ -137,7 +137,7 @@ enum EExprToken
 	// Variable references.
 	EX_LocalVariable		= 0x00,	// A local variable.
 	EX_InstanceVariable		= 0x01,	// An object variable.
-	//						= 0x02,
+	EX_DefaultVariable		= 0x02, // Default variable for a class context.
 	//						= 0x03,
 	EX_Return				= 0x04,	// Return from function.
 	//						= 0x05,
@@ -153,7 +153,7 @@ enum EExprToken
 	EX_Let					= 0x0F,	// Assign an arbitrary size value to a variable.
 	//						= 0x10,
 	//						= 0x11,
-	//						= 0x12,
+	EX_ClassContext			= 0x12,	// Class default object context.
 	EX_MetaCast             = 0x13, // Metaclass cast.
 	EX_LetBool				= 0x14, // Let boolean variable.
 	EX_EndParmValue			= 0x15,	// end of default value for optional function parameter
diff --git a/Engine/Source/Runtime/CoreUObject/Public/UObject/ScriptSerialization.h b/Engine/Source/Runtime/CoreUObject/Public/UObject/ScriptSerialization.h
--- a/Engine/Source/Runtime/CoreUObject/Public/UObject/ScriptSerialization.h
+++ b/Engine/Source/Runtime/CoreUObject/Public/UObject/ScriptSerialization.h
@@ -165,6 +165,7 @@
 		}
 		case EX_LocalVariable:
 		case EX_InstanceVariable:
+		case EX_DefaultVariable:
 		case EX_LocalOutVariable:
 		{
 			XFER_PROP_POINTER;
@@ -239,6 +240,7 @@
 			while( SerializeExpr( iCode, Ar ) != EX_EndFunctionParms ); // Parms.
 			break;
 		}
+		case EX_ClassContext:
 		case EX_Context:
 		case EX_Context_FailSilent:
 		{
