[UEBP-60] This saves 3s on editor load/PIE of large projects. Simple optimization for the most common case of RewireOldPinsToNewPins

diff --git a/Engine/Source/Editor/BlueprintGraph/Private/K2Node.cpp b/Engine/Source/Editor/BlueprintGraph/Private/K2Node.cpp
--- a/Engine/Source/Editor/BlueprintGraph/Private/K2Node.cpp
+++ b/Engine/Source/Editor/BlueprintGraph/Private/K2Node.cpp
@@ -634,8 +634,13 @@ void UK2Node::RewireOldPinsToNewPins(TArray<UEdGraphPin*>& InOldPins, TArray<UEd
 	{
 		UEdGraphPin* OldPin = InOldPins[OldPinIndex];
 
-		for (int32 NewPinIndex = 0; NewPinIndex < InNewPins.Num(); ++NewPinIndex)
+		// common case is for InOldPins and InNewPins to match, so we start searching from the current index:
+		const int32 NumNewPins = InNewPins.Num();
+		int32 NewPinIndex = OldPinIndex % InNewPins.Num();
+		for (int32 NewPinCount = 0; NewPinCount < InNewPins.Num(); ++NewPinCount)
 		{
+			// if InNewPins grows in this loop then we may skip entries and fail to find a match:
+			check(NumNewPins == InNewPins.Num());
 			UEdGraphPin* NewPin = InNewPins[NewPinIndex];
 
 			const ERedirectType RedirectType = DoPinsMatchForReconstruction(NewPin, NewPinIndex, OldPin, OldPinIndex);
@@ -644,6 +649,8 @@ void UK2Node::RewireOldPinsToNewPins(TArray<UEdGraphPin*>& InOldPins, TArray<UEd
 				ReconstructSinglePin(NewPin, OldPin, RedirectType);
 				break;
 			}
+
+			NewPinIndex = (NewPinIndex + 1) % InNewPins.Num();
 		}
 	}
 }
