Math Expression nodes will again correctly lookup by the display name of a function.

Changed metadata lookups that were still targetting FriendlyName to lookup the DisplayName.

diff --git a/Engine/Source/Editor/BlueprintGraph/Classes/EdGraphSchema_K2.h b/Engine/Source/Editor/BlueprintGraph/Classes/EdGraphSchema_K2.h
--- a/Engine/Source/Editor/BlueprintGraph/Classes/EdGraphSchema_K2.h
+++ b/Engine/Source/Editor/BlueprintGraph/Classes/EdGraphSchema_K2.h
@@ -91,7 +91,7 @@ public:
 	static const FName MD_CompactNodeTitle;
 
 	// [FunctionMetadata] Indicates that the function should be drawn with this title over the function name
-	static const FName MD_FriendlyName;
+	static const FName MD_DisplayName;
 
 	//    property metadata
 
diff --git a/Engine/Source/Editor/BlueprintGraph/Private/EdGraphSchema_K2.cpp b/Engine/Source/Editor/BlueprintGraph/Private/EdGraphSchema_K2.cpp
--- a/Engine/Source/Editor/BlueprintGraph/Private/EdGraphSchema_K2.cpp
+++ b/Engine/Source/Editor/BlueprintGraph/Private/EdGraphSchema_K2.cpp
@@ -56,7 +56,7 @@ const FName FBlueprintMetadata::MD_FunctionCategory(TEXT("Category"));
 const FName FBlueprintMetadata::MD_DeprecatedFunction(TEXT("DeprecatedFunction"));
 const FName FBlueprintMetadata::MD_DeprecationMessage(TEXT("DeprecationMessage"));
 const FName FBlueprintMetadata::MD_CompactNodeTitle(TEXT("CompactNodeTitle"));
-const FName FBlueprintMetadata::MD_FriendlyName(TEXT("FriendlyName"));
+const FName FBlueprintMetadata::MD_DisplayName(TEXT("DisplayName"));
 
 const FName FBlueprintMetadata::MD_ExposeOnSpawn(TEXT("ExposeOnSpawn"));
 const FName FBlueprintMetadata::MD_DefaultToSelf(TEXT("DefaultToSelf"));
diff --git a/Engine/Source/Editor/BlueprintGraph/Private/K2Node_MathExpression.cpp b/Engine/Source/Editor/BlueprintGraph/Private/K2Node_MathExpression.cpp
--- a/Engine/Source/Editor/BlueprintGraph/Private/K2Node_MathExpression.cpp
+++ b/Engine/Source/Editor/BlueprintGraph/Private/K2Node_MathExpression.cpp
@@ -805,9 +805,9 @@ public:
 						{
 							FunctionName = TestFunction->GetMetaData(FBlueprintMetadata::MD_CompactNodeTitle);
 						}
-						else if (TestFunction->HasMetaData(FBlueprintMetadata::MD_FriendlyName))
+						else if (TestFunction->HasMetaData(FBlueprintMetadata::MD_DisplayName))
 						{
-							FunctionName = TestFunction->GetMetaData(FBlueprintMetadata::MD_FriendlyName);
+							FunctionName = TestFunction->GetMetaData(FBlueprintMetadata::MD_DisplayName);
 						}
 						Add(FunctionName, TestFunction);
 					}
