Fix for interface communication with streaming levels not working unless converted ULevelStreaming pin to UObject pin first.

diff --git a/Engine/Source/Editor/BlueprintGraph/Private/K2Node_Message.cpp b/Engine/Source/Editor/BlueprintGraph/Private/K2Node_Message.cpp
--- a/Engine/Source/Editor/BlueprintGraph/Private/K2Node_Message.cpp
+++ b/Engine/Source/Editor/BlueprintGraph/Private/K2Node_Message.cpp
@@ -243,7 +243,7 @@ void UK2Node_Message::ExpandNode(class FKismetCompilerContext& CompilerContext,
 
 		// If the message pin is linked and what it is linked to is not a ULevelStreaming pin reference but is a child of ULevelStreaming, 
 		// then we need to leave the possibility that it could be a ULevelStreaming and will need to make an attempt at casting to one
-		if (MessageSelfPin->LinkedTo.Num() > 0 && MessageSelfPin->LinkedTo[0]->PinType.PinSubCategoryObject.Get() != ULevelStreaming::StaticClass() && ULevelStreaming::StaticClass()->IsChildOf(Cast<UClass>(MessageSelfPin->LinkedTo[0]->PinType.PinSubCategoryObject.Get())))
+		if (MessageSelfPin->LinkedTo.Num() > 0 && ULevelStreaming::StaticClass()->IsChildOf(Cast<UClass>(MessageSelfPin->LinkedTo[0]->PinType.PinSubCategoryObject.Get())))
 		{
 			ExpandLevelStreamingHandlers(CompilerContext, SourceGraph, ExecPin, MessageSelfPin, CastToInterfaceNode);
 		}
