Copy and pasting a K2Node_CallFunction with an interface pin for the SelfPin will no longer change the SelfPin into a normal Object Reference pin.

UK2Node_CallFunction::EnsureFunctionIsInBlueprint will check the function's owner class to see if it is a child of UInterface and will not allow changes to the FunctionReference if it is.

diff --git a/Engine/Source/Editor/BlueprintGraph/Private/K2Node_CallFunction.cpp b/Engine/Source/Editor/BlueprintGraph/Private/K2Node_CallFunction.cpp
--- a/Engine/Source/Editor/BlueprintGraph/Private/K2Node_CallFunction.cpp
+++ b/Engine/Source/Editor/BlueprintGraph/Private/K2Node_CallFunction.cpp
@@ -1604,14 +1604,19 @@ void UK2Node_CallFunction::EnsureFunctionIsInBlueprint()
 		UClass* FunctionOwnerClass = Function->GetOuterUClass();
 		UObject* FunctionGenerator = FunctionOwnerClass ? FunctionOwnerClass->ClassGeneratedBy : NULL;
 
-		// If function is generated from a blueprint object then dbl check self pin compatibility
-		UEdGraphPin* SelfPin = GetDefault<UEdGraphSchema_K2>()->FindSelfPin(*this, EGPD_Input);
-		if ((FunctionGenerator != NULL) && SelfPin)
+		// Never change the type if the function is an Interface function type, this only occurs when 
+		// the function is the interface function and the self pin will be PC_Interface type
+		if (!FunctionOwnerClass->IsChildOf(UInterface::StaticClass()))
 		{
-			UBlueprint* BlueprintObj = FBlueprintEditorUtils::FindBlueprintForNode(this);
-			if ((BlueprintObj != NULL) && !IsSelfPinCompatibleWithBlueprintContext(SelfPin, BlueprintObj))
+			// If function is generated from a blueprint object then dbl check self pin compatibility
+			UEdGraphPin* SelfPin = GetDefault<UEdGraphSchema_K2>()->FindSelfPin(*this, EGPD_Input);
+			if ((FunctionGenerator != NULL) && SelfPin)
 			{
-				FunctionReference.SetSelfMember(Function->GetFName());
+				UBlueprint* BlueprintObj = FBlueprintEditorUtils::FindBlueprintForNode(this);
+				if ((BlueprintObj != NULL) && !IsSelfPinCompatibleWithBlueprintContext(SelfPin, BlueprintObj))
+				{
+					FunctionReference.SetSelfMember(Function->GetFName());
+				}
 			}
 		}
 	}
