UE-17499 Could use a way to add custom type auto-conversions for blueprints

New FunctionMetadta "BlueprintAutocast" was added.

todo: add warning in UHT
todo: remove hardcoded functions from UEdGraphSchema_K2::SearchForAutocastFunction

diff --git a/Engine/Plugins/Experimental/LiveEditor/Source/LiveEditor/Private/K2Node_LiveEditObject.cpp b/Engine/Plugins/Experimental/LiveEditor/Source/LiveEditor/Private/K2Node_LiveEditObject.cpp
--- a/Engine/Plugins/Experimental/LiveEditor/Source/LiveEditor/Private/K2Node_LiveEditObject.cpp
+++ b/Engine/Plugins/Experimental/LiveEditor/Source/LiveEditor/Private/K2Node_LiveEditObject.cpp
@@ -526,10 +526,11 @@ void UK2Node_LiveEditObject::ExpandNode(class FKismetCompilerContext& CompilerCo
 		// The Convert function to go from (int)Delta to ULiveEditorKismetLibrary::ModifyPropertyByName(... float Delta ...)
 		//
 		FName ConvertFunctionName;
-		bool success = Schema->SearchForAutocastFunction( EventDeltaPin, MultiplyNodeFirstPin, ConvertFunctionName );
+		UClass* ConvertFunctionOwner = nullptr;
+		bool success = Schema->SearchForAutocastFunction(EventDeltaPin, MultiplyNodeFirstPin, ConvertFunctionName, ConvertFunctionOwner);
 		check( success );
 		UK2Node_CallFunction *ConvertDeltaNode = CompilerContext.SpawnIntermediateNode<UK2Node_CallFunction>(this,SourceGraph);
-		ConvertDeltaNode->FunctionReference.SetExternalMember( ConvertFunctionName, UKismetMathLibrary::StaticClass() );
+		ConvertDeltaNode->FunctionReference.SetExternalMember(ConvertFunctionName, ConvertFunctionOwner ? ConvertFunctionOwner : UKismetMathLibrary::StaticClass());
 		ConvertDeltaNode->AllocateDefaultPins();
 
 		FName PinName;
diff --git a/Engine/Source/Editor/AnimGraph/Classes/AnimationGraphSchema.h b/Engine/Source/Editor/AnimGraph/Classes/AnimationGraphSchema.h
--- a/Engine/Source/Editor/AnimGraph/Classes/AnimationGraphSchema.h
+++ b/Engine/Source/Editor/AnimGraph/Classes/AnimationGraphSchema.h
@@ -56,7 +56,7 @@ class UAnimationGraphSchema : public UEdGraphSchema_K2
 
 	// Begin UEdGraphSchema_K2 interface
 	virtual const FPinConnectionResponse DetermineConnectionResponseOfCompatibleTypedPins(const UEdGraphPin* PinA, const UEdGraphPin* PinB, const UEdGraphPin* InputPin, const UEdGraphPin* OutputPin) const override;
-	virtual bool SearchForAutocastFunction(const UEdGraphPin* OutputPin, const UEdGraphPin* InputPin, /*out*/ FName& TargetFunction) const override;
+	virtual bool SearchForAutocastFunction(const UEdGraphPin* OutputPin, const UEdGraphPin* InputPin, /*out*/ FName& TargetFunction, /*out*/ UClass*& FunctionOwner) const override;
 	virtual bool ArePinsCompatible(const UEdGraphPin* PinA, const UEdGraphPin* PinB, const UClass* CallingContext = NULL, bool bIgnoreArray = false) const override;
 	virtual bool DoesSupportAnimNotifyActions() const override;
 	// End UEdGraphSchema_K2 interface
diff --git a/Engine/Source/Editor/AnimGraph/Private/AnimationGraphSchema.cpp b/Engine/Source/Editor/AnimGraph/Private/AnimationGraphSchema.cpp
--- a/Engine/Source/Editor/AnimGraph/Private/AnimationGraphSchema.cpp
+++ b/Engine/Source/Editor/AnimGraph/Private/AnimationGraphSchema.cpp
@@ -169,7 +169,7 @@ bool UAnimationGraphSchema::DoesSupportAnimNotifyActions() const
 	return false;
 }
 
-bool UAnimationGraphSchema::SearchForAutocastFunction(const UEdGraphPin* OutputPin, const UEdGraphPin* InputPin, FName& TargetFunction) const
+bool UAnimationGraphSchema::SearchForAutocastFunction(const UEdGraphPin* OutputPin, const UEdGraphPin* InputPin, FName& TargetFunction, /*out*/ UClass*& FunctionOwner) const
 {
 	if (IsComponentSpacePosePin(OutputPin->PinType) && IsLocalSpacePosePin(InputPin->PinType))
 	{
@@ -183,7 +183,7 @@ bool UAnimationGraphSchema::SearchForAutocastFunction(const UEdGraphPin* OutputP
 	}
 	else
 	{
-		return Super::SearchForAutocastFunction(OutputPin, InputPin, TargetFunction);
+		return Super::SearchForAutocastFunction(OutputPin, InputPin, TargetFunction, FunctionOwner);
 	}
 }
 
diff --git a/Engine/Source/Editor/BlueprintGraph/Classes/EdGraphSchema_K2.h b/Engine/Source/Editor/BlueprintGraph/Classes/EdGraphSchema_K2.h
--- a/Engine/Source/Editor/BlueprintGraph/Classes/EdGraphSchema_K2.h
+++ b/Engine/Source/Editor/BlueprintGraph/Classes/EdGraphSchema_K2.h
@@ -871,7 +871,7 @@ public:
 	static UFunction* FindSetVariableByNameFunction(const FEdGraphPinType& PinType);
 
 	/** Find an appropriate function to call to perform an automatic cast operation */
-	virtual bool SearchForAutocastFunction(const UEdGraphPin* OutputPin, const UEdGraphPin* InputPin, /*out*/ FName& TargetFunction) const;
+	virtual bool SearchForAutocastFunction(const UEdGraphPin* OutputPin, const UEdGraphPin* InputPin, /*out*/ FName& TargetFunction, /*out*/ UClass*& FunctionOwner) const;
 
 	/** Find an appropriate node that can convert from one pin type to another (not a cast; e.g. "MakeLiteralArray" node) */
 	virtual bool FindSpecializedConversionNode(const UEdGraphPin* OutputPin, const UEdGraphPin* InputPin, bool bCreateNode, /*out*/ class UK2Node*& TargetNode) const;
diff --git a/Engine/Source/Editor/BlueprintGraph/Private/EdGraphSchema_K2.cpp b/Engine/Source/Editor/BlueprintGraph/Private/EdGraphSchema_K2.cpp
--- a/Engine/Source/Editor/BlueprintGraph/Private/EdGraphSchema_K2.cpp
+++ b/Engine/Source/Editor/BlueprintGraph/Private/EdGraphSchema_K2.cpp
@@ -2149,9 +2149,10 @@ const FPinConnectionResponse UEdGraphSchema_K2::CanCreateConnection(const UEdGra
 	{
 		// Autocasting
 		FName DummyName;
+		UClass* DummyClass;
 		UK2Node* DummyNode;
 
-		const bool bCanAutocast = SearchForAutocastFunction(OutputPin, InputPin, /*out*/ DummyName);
+		const bool bCanAutocast = SearchForAutocastFunction(OutputPin, InputPin, /*out*/ DummyName, DummyClass);
 		const bool bCanAutoConvert = FindSpecializedConversionNode(OutputPin, InputPin, false, /* out */ DummyNode);
 
 		if (bCanAutocast || bCanAutoConvert)
@@ -2180,7 +2181,128 @@ bool UEdGraphSchema_K2::TryCreateConnection(UEdGraphPin* PinA, UEdGraphPin* PinB
 	return bModified;
 }
 
-bool UEdGraphSchema_K2::SearchForAutocastFunction(const UEdGraphPin* OutputPin, const UEdGraphPin* InputPin, /*out*/ FName& TargetFunction) const
+struct FAutocastFunctionMap
+{
+private:
+	TMap<FString, TWeakObjectPtr<UFunction>> InnerMap;
+
+	static FString GenerateTypeData(const FEdGraphPinType& PinType)
+	{
+		auto Obj = PinType.PinSubCategoryObject.Get();
+		return FString::Printf(TEXT("%s;%s;%s"), *PinType.PinCategory, *PinType.PinSubCategory, Obj ? *Obj->GetPathName() : TEXT(""));
+	}
+
+	static FString GenerateCastData(const FEdGraphPinType& InputPinType, const FEdGraphPinType& OutputPinType)
+	{
+		return FString::Printf(TEXT("%s;%s"), *GenerateTypeData(InputPinType), *GenerateTypeData(OutputPinType));
+	}
+
+	static bool IsInputParam(const UProperty* Property)
+	{
+		const uint64 ConstOutParamFlag = CPF_OutParm | CPF_ConstParm;
+		const uint64 IsConstOut = Property ? (Property->PropertyFlags & ConstOutParamFlag) : 0;
+		return Property 
+			&& (CPF_Parm == (Property->PropertyFlags & (CPF_Parm | CPF_ReturnParm)))
+			&& ((0 == IsConstOut) || (ConstOutParamFlag == IsConstOut));
+	}
+
+	static int32 GetInputPropertiesNum(const UFunction* Function)
+	{
+		int32 InputPropNum = 0;
+		for (auto Property : TFieldRange<const UProperty>(Function))
+		{
+			InputPropNum += IsInputParam(Property) ? 1 : 0;
+		}
+		return InputPropNum;
+	}
+
+	static const UProperty* GetFirstInputProperty(const UFunction* Function)
+	{
+		for (auto Property : TFieldRange<const UProperty>(Function))
+		{
+			if (IsInputParam(Property))
+			{
+				return Property;
+			}
+		}
+		return nullptr;
+	}
+
+	void InsertFunction(UFunction* Function, const UEdGraphSchema_K2* Schema)
+	{
+		FEdGraphPinType InputPinType;
+		Schema->ConvertPropertyToPinType(GetFirstInputProperty(Function), InputPinType);
+
+		FEdGraphPinType OutputPinType;
+		Schema->ConvertPropertyToPinType(Function->GetReturnProperty(), OutputPinType);
+
+		InnerMap.Add(GenerateCastData(InputPinType, OutputPinType), Function);
+	}
+public:
+
+	static bool IsAutocastFunction(const UFunction* Function)
+	{
+		const FName BlueprintAutocast(TEXT("BlueprintAutocast"));
+		return Function
+			&& Function->HasMetaData(BlueprintAutocast)
+			&& Function->HasAllFunctionFlags(FUNC_Static | FUNC_Native | FUNC_Public | FUNC_BlueprintPure)
+			&& Function->GetReturnProperty()
+			&& (1 == GetInputPropertiesNum(Function));
+	}
+
+	void Refresh()
+	{
+		InnerMap.Empty();
+		auto Schema = GetDefault<UEdGraphSchema_K2>();
+
+		TArray<UClass*> Libraries;
+		GetDerivedClasses(UBlueprintFunctionLibrary::StaticClass(), Libraries);
+		for (auto Library : Libraries)
+		{
+			if (Library && (0 != (Library->ClassFlags & CLASS_Native)))
+			{
+				for (auto Function : TFieldRange<UFunction>(Library, EFieldIteratorFlags::ExcludeSuper, EFieldIteratorFlags::ExcludeDeprecated))
+				{
+					if (IsAutocastFunction(Function))
+					{
+						InsertFunction(Function, Schema);
+					}
+				}
+			}
+		}
+	}
+
+	UFunction* Find(const FEdGraphPinType& InputPinType, const FEdGraphPinType& OutputPinType)
+	{
+		TWeakObjectPtr<UFunction>* FuncPtr = InnerMap.Find(GenerateCastData(InputPinType, OutputPinType));
+		return FuncPtr ? FuncPtr->Get() : nullptr;
+	}
+
+	static FAutocastFunctionMap& Get()
+	{
+		static FAutocastFunctionMap* AutocastFunctionMap = nullptr;
+		if (AutocastFunctionMap == nullptr)
+		{
+			AutocastFunctionMap = new FAutocastFunctionMap();
+		}
+		return *AutocastFunctionMap;
+	}
+
+	static void OnProjectHotReloaded(bool bWasTriggeredAutomatically)
+	{
+		Get().Refresh();
+	}
+
+	FAutocastFunctionMap()
+	{
+		Refresh();
+
+		IHotReloadInterface& HotReloadSupport = FModuleManager::LoadModuleChecked<IHotReloadInterface>("HotReload");
+		HotReloadSupport.OnHotReload().AddStatic(&FAutocastFunctionMap::OnProjectHotReloaded);
+	}
+};
+
+bool UEdGraphSchema_K2::SearchForAutocastFunction(const UEdGraphPin* OutputPin, const UEdGraphPin* InputPin, /*out*/ FName& TargetFunction, /*out*/ UClass*& FunctionOwner) const
 {
 	// NOTE: Under no circumstances should anyone *ever* add a questionable cast to this function.
 	// If it could be at all confusing why a function is provided, to even a novice user, err on the side of do not cast!!!
@@ -2415,6 +2537,16 @@ bool UEdGraphSchema_K2::SearchForAutocastFunction(const UEdGraphPin* OutputPin,
 		}
 	}
 
+	if (TargetFunction == NAME_None)
+	{
+		auto AutocastFunctionMap = FAutocastFunctionMap::Get();
+		if (auto Func = AutocastFunctionMap.Find(OutputPin->PinType, InputPin->PinType))
+		{
+			TargetFunction = Func->GetFName();
+			FunctionOwner = Func->GetOwnerClass();
+		}
+	}
+
 	return TargetFunction != NAME_None;
 }
 
@@ -2592,15 +2724,14 @@ bool UEdGraphSchema_K2::CreateAutomaticConversionNodeAndConnections(UEdGraphPin*
 	}
 
 	FName TargetFunctionName;
+	UClass* ClassContainingConversionFunction = nullptr;
 	TSubclassOf<UK2Node> ConversionNodeClass;
 
 	UK2Node* TemplateConversionNode = NULL;
 
-	if (SearchForAutocastFunction(OutputPin, InputPin, /*out*/ TargetFunctionName))
+	if (SearchForAutocastFunction(OutputPin, InputPin, /*out*/ TargetFunctionName, /*out*/ClassContainingConversionFunction))
 	{
 		// Create a new call function node for the casting operator
-		UClass* ClassContainingConversionFunction = NULL; //@TODO: Should probably return this from the search function too
-
 		UK2Node_CallFunction* TemplateNode = NewObject<UK2Node_CallFunction>();
 		TemplateNode->FunctionReference.SetExternalMember(TargetFunctionName, ClassContainingConversionFunction);
 		//TemplateNode->bIsBeadFunction = true;
diff --git a/Engine/Source/Runtime/CoreUObject/Public/UObject/ObjectBase.h b/Engine/Source/Runtime/CoreUObject/Public/UObject/ObjectBase.h
--- a/Engine/Source/Runtime/CoreUObject/Public/UObject/ObjectBase.h
+++ b/Engine/Source/Runtime/CoreUObject/Public/UObject/ObjectBase.h
@@ -1203,6 +1203,8 @@ namespace UM
 		/// [FunctionMetadta] Used by BlueprintCallable functions to indicate which parameter is used to determine the World that the operation is occurring within.
 		WorldContext,
 
+		/// [FunctionMetadta] Used only by static BlueprintPure functions (with one return param and exactly one input param) from BlueprintLibrary.
+		BlueprintAutocast,
 	};
 
 	// Metadata usable in UINTERFACE
