Fixed nonunity compile error in EdGraphSchema_K2.cpp

diff --git a/Engine/Source/Editor/BlueprintGraph/Private/EdGraphSchema_K2.cpp b/Engine/Source/Editor/BlueprintGraph/Private/EdGraphSchema_K2.cpp
--- a/Engine/Source/Editor/BlueprintGraph/Private/EdGraphSchema_K2.cpp
+++ b/Engine/Source/Editor/BlueprintGraph/Private/EdGraphSchema_K2.cpp
@@ -24,6 +24,7 @@
 #include "BlueprintComponentNodeSpawner.h"
 #include "AssetRegistryModule.h"
 #include "Blueprint/AIBlueprintHelperLibrary.h"
+#include "HotReloadInterface.h"
 
 #include "K2Node_CastByteToEnum.h"
 #include "K2Node_ClassDynamicCast.h"
