UE-17499 Could use a way to add custom type auto-conversions for blueprints - improvements
Removed hardcoded functions from UEdGraphSchema_K2::SearchForAutocastFunction.

diff --git a/Engine/Source/Editor/BlueprintGraph/Private/EdGraphSchema_K2.cpp b/Engine/Source/Editor/BlueprintGraph/Private/EdGraphSchema_K2.cpp
--- a/Engine/Source/Editor/BlueprintGraph/Private/EdGraphSchema_K2.cpp
+++ b/Engine/Source/Editor/BlueprintGraph/Private/EdGraphSchema_K2.cpp
@@ -2182,10 +2182,14 @@ bool UEdGraphSchema_K2::TryCreateConnection(UEdGraphPin* PinA, UEdGraphPin* PinB
 	return bModified;
 }
 
-struct FAutocastFunctionMap
+struct FAutocastFunctionMap : private FNoncopyable
 {
 private:
+	static FAutocastFunctionMap* AutocastFunctionMap;
+
 	TMap<FString, TWeakObjectPtr<UFunction>> InnerMap;
+	FDelegateHandle OnHotReloadDelegateHandle;
+	FDelegateHandle OnModulesChangedDelegateHandle;
 
 	static FString GenerateTypeData(const FEdGraphPinType& PinType)
 	{
@@ -2198,30 +2202,19 @@ private:
 		return FString::Printf(TEXT("%s;%s"), *GenerateTypeData(InputPinType), *GenerateTypeData(OutputPinType));
 	}
 
-	static bool IsInputParam(const UProperty* Property)
+	static bool IsInputParam(uint64 PropertyFlags)
 	{
 		const uint64 ConstOutParamFlag = CPF_OutParm | CPF_ConstParm;
-		const uint64 IsConstOut = Property ? (Property->PropertyFlags & ConstOutParamFlag) : 0;
-		return Property 
-			&& (CPF_Parm == (Property->PropertyFlags & (CPF_Parm | CPF_ReturnParm)))
+		const uint64 IsConstOut = PropertyFlags & ConstOutParamFlag;
+		return (CPF_Parm == (PropertyFlags & (CPF_Parm | CPF_ReturnParm)))
 			&& ((0 == IsConstOut) || (ConstOutParamFlag == IsConstOut));
 	}
 
-	static int32 GetInputPropertiesNum(const UFunction* Function)
-	{
-		int32 InputPropNum = 0;
-		for (auto Property : TFieldRange<const UProperty>(Function))
-		{
-			InputPropNum += IsInputParam(Property) ? 1 : 0;
-		}
-		return InputPropNum;
-	}
-
 	static const UProperty* GetFirstInputProperty(const UFunction* Function)
 	{
 		for (auto Property : TFieldRange<const UProperty>(Function))
 		{
-			if (IsInputParam(Property))
+			if (Property && IsInputParam(Property->PropertyFlags))
 			{
 				return Property;
 			}
@@ -2248,11 +2241,19 @@ public:
 			&& Function->HasMetaData(BlueprintAutocast)
 			&& Function->HasAllFunctionFlags(FUNC_Static | FUNC_Native | FUNC_Public | FUNC_BlueprintPure)
 			&& Function->GetReturnProperty()
-			&& (1 == GetInputPropertiesNum(Function));
+			&& GetFirstInputProperty(Function);
 	}
 
 	void Refresh()
 	{
+#ifdef SCHEMA_K2_AUTOCASTFUNCTIONMAP_LOG_TIME
+		static_assert(false, "Macro redefinition.");
+#endif
+#define SCHEMA_K2_AUTOCASTFUNCTIONMAP_LOG_TIME 0
+#if SCHEMA_K2_AUTOCASTFUNCTIONMAP_LOG_TIME
+		const auto StartTime = FPlatformTime::Seconds();
+#endif //SCHEMA_K2_AUTOCASTFUNCTIONMAP_LOG_TIME
+
 		InnerMap.Empty();
 		auto Schema = GetDefault<UEdGraphSchema_K2>();
 
@@ -2260,7 +2261,7 @@ public:
 		GetDerivedClasses(UBlueprintFunctionLibrary::StaticClass(), Libraries);
 		for (auto Library : Libraries)
 		{
-			if (Library && (0 != (Library->ClassFlags & CLASS_Native)))
+			if (Library && (CLASS_Native == (Library->ClassFlags & (CLASS_Native | CLASS_Deprecated | CLASS_NewerVersionExists))))
 			{
 				for (auto Function : TFieldRange<UFunction>(Library, EFieldIteratorFlags::ExcludeSuper, EFieldIteratorFlags::ExcludeDeprecated))
 				{
@@ -2271,17 +2272,22 @@ public:
 				}
 			}
 		}
+
+#if SCHEMA_K2_AUTOCASTFUNCTIONMAP_LOG_TIME
+		const auto EndTime = FPlatformTime::Seconds();
+		UE_LOG(LogBlueprint, Warning, TEXT("FAutocastFunctionMap::Refresh took %fs"), EndTime - StartTime);
+#endif //SCHEMA_K2_AUTOCASTFUNCTIONMAP_LOG_TIME
+#undef SCHEMA_K2_AUTOCASTFUNCTIONMAP_LOG_TIME
 	}
 
-	UFunction* Find(const FEdGraphPinType& InputPinType, const FEdGraphPinType& OutputPinType)
+	UFunction* Find(const FEdGraphPinType& InputPinType, const FEdGraphPinType& OutputPinType) const
 	{
-		TWeakObjectPtr<UFunction>* FuncPtr = InnerMap.Find(GenerateCastData(InputPinType, OutputPinType));
+		const TWeakObjectPtr<UFunction>* FuncPtr = InnerMap.Find(GenerateCastData(InputPinType, OutputPinType));
 		return FuncPtr ? FuncPtr->Get() : nullptr;
 	}
 
 	static FAutocastFunctionMap& Get()
 	{
-		static FAutocastFunctionMap* AutocastFunctionMap = nullptr;
 		if (AutocastFunctionMap == nullptr)
 		{
 			AutocastFunctionMap = new FAutocastFunctionMap();
@@ -2291,7 +2297,18 @@ public:
 
 	static void OnProjectHotReloaded(bool bWasTriggeredAutomatically)
 	{
-		Get().Refresh();
+		if (AutocastFunctionMap)
+		{
+			AutocastFunctionMap->Refresh();
+		}
+	}
+
+	static void OnModulesChanged(FName ModuleThatChanged, EModuleChangeReason ReasonForChange)
+	{
+		if (AutocastFunctionMap)
+		{
+			AutocastFunctionMap->Refresh();
+		}
 	}
 
 	FAutocastFunctionMap()
@@ -2299,10 +2316,24 @@ public:
 		Refresh();
 
 		IHotReloadInterface& HotReloadSupport = FModuleManager::LoadModuleChecked<IHotReloadInterface>("HotReload");
-		HotReloadSupport.OnHotReload().AddStatic(&FAutocastFunctionMap::OnProjectHotReloaded);
+		OnHotReloadDelegateHandle = HotReloadSupport.OnHotReload().AddStatic(&FAutocastFunctionMap::OnProjectHotReloaded);
+
+		OnModulesChangedDelegateHandle = FModuleManager::Get().OnModulesChanged().AddStatic(&OnModulesChanged);
+	}
+
+	~FAutocastFunctionMap()
+	{
+		if (auto HotReloadSupport = FModuleManager::GetModulePtr<IHotReloadInterface>("HotReload"))
+		{
+			HotReloadSupport->OnHotReload().Remove(OnHotReloadDelegateHandle);
+		}
+
+		FModuleManager::Get().OnModulesChanged().Remove(OnModulesChangedDelegateHandle); 
 	}
 };
 
+FAutocastFunctionMap* FAutocastFunctionMap::AutocastFunctionMap = nullptr;
+
 bool UEdGraphSchema_K2::SearchForAutocastFunction(const UEdGraphPin* OutputPin, const UEdGraphPin* InputPin, /*out*/ FName& TargetFunction, /*out*/ UClass*& FunctionOwner) const
 {
 	// NOTE: Under no circumstances should anyone *ever* add a questionable cast to this function.
@@ -2310,206 +2341,15 @@ bool UEdGraphSchema_K2::SearchForAutocastFunction(const UEdGraphPin* OutputPin,
 	// This includes things like string->int (does it do length, atoi, or what?) that would be autocasts in a traditional scripting language
 
 	TargetFunction = NAME_None;
-
-	const UScriptStruct* InputStructType = Cast<const UScriptStruct>(InputPin->PinType.PinSubCategoryObject.Get());
-	const UScriptStruct* OutputStructType = Cast<const UScriptStruct>(OutputPin->PinType.PinSubCategoryObject.Get());
+	FunctionOwner = nullptr;
 
 	if (OutputPin->PinType.bIsArray != InputPin->PinType.bIsArray)
 	{
-		// We don't autoconvert between arrays and non-arrays.  Those are handled by specialized conversions
-	}
-	else if (OutputPin->PinType.PinCategory == PC_Int)
-	{
-		if (InputPin->PinType.PinCategory == PC_Float)
-		{
-			TargetFunction = TEXT("Conv_IntToFloat");
-		}
-		else if (InputPin->PinType.PinCategory == PC_String)
-		{
-			TargetFunction = TEXT("Conv_IntToString");
-		}
-		else if ((InputPin->PinType.PinCategory == PC_Byte) && (InputPin->PinType.PinSubCategoryObject == NULL))
-		{
-			TargetFunction = TEXT("Conv_IntToByte");
-		}
-		else if (InputPin->PinType.PinCategory == PC_Boolean)
-		{
-			TargetFunction = TEXT("Conv_IntToBool");
-		}
-		else if(InputPin->PinType.PinCategory == PC_Text)
-		{
-			TargetFunction = TEXT("Conv_IntToText");
-		}
-	}
-	else if (OutputPin->PinType.PinCategory == PC_Float)
-	{
-		if (InputPin->PinType.PinCategory == PC_Int)
-		{
-			TargetFunction = TEXT("FTrunc");
-		}
-		else if ((InputPin->PinType.PinCategory == PC_Struct) && (InputStructType == VectorStruct))
-		{
-			TargetFunction = TEXT("Conv_FloatToVector");
-		}
-		else if (InputPin->PinType.PinCategory == PC_String)
-		{
-			TargetFunction = TEXT("Conv_FloatToString");
-		}
-		else if ((InputPin->PinType.PinCategory == PC_Struct) && (InputStructType == LinearColorStruct))
-		{
-			TargetFunction = TEXT("Conv_FloatToLinearColor");
-		}
-		else if(InputPin->PinType.PinCategory == PC_Text)
-		{
-			TargetFunction = TEXT("Conv_FloatToText");
-		}
-	}
-	else if (OutputPin->PinType.PinCategory == PC_Struct)
-	{
-		if (OutputStructType == VectorStruct)
-		{
-			if ((InputPin->PinType.PinCategory == PC_Struct) && (InputStructType == TransformStruct))
-			{
-				TargetFunction = TEXT("Conv_VectorToTransform");
-			}
-			else if ((InputPin->PinType.PinCategory == PC_Struct) && (InputStructType == LinearColorStruct))
-			{
-				TargetFunction = TEXT("Conv_VectorToLinearColor");
-			}
-			else if ((InputPin->PinType.PinCategory == PC_Struct) && (InputStructType == RotatorStruct))
-			{
-				TargetFunction = TEXT("Conv_VectorToRotator");
-			}
-			else if (InputPin->PinType.PinCategory == PC_String)
-			{
-				TargetFunction = TEXT("Conv_VectorToString");
-			}
-			// NOTE: Did you see the note above about unsafe and unclear casts?
-		}
-		else if(OutputStructType == RotatorStruct)
-		{
-			if ((InputPin->PinType.PinCategory == PC_Struct) && (InputStructType == TransformStruct))
-			{
-				TargetFunction = TEXT("MakeTransform");
-			}
-			else if ((InputPin->PinType.PinCategory == PC_Struct) && (InputStructType == VectorStruct))
-			{
-				TargetFunction = TEXT("Conv_RotatorToVector");
-			}
-			else if (InputPin->PinType.PinCategory == PC_String)
-			{
-				TargetFunction = TEXT("Conv_RotatorToString");
-			}
-			// NOTE: Did you see the note above about unsafe and unclear casts?
-		}
-		else if(OutputStructType == LinearColorStruct)
-		{
-			if ((InputPin->PinType.PinCategory == PC_Struct) && (InputStructType == ColorStruct))
-			{
-				TargetFunction = TEXT("Conv_LinearColorToColor");
-			}
-			else if (InputPin->PinType.PinCategory == PC_String)
-			{
-				TargetFunction = TEXT("Conv_ColorToString");
-			}
-			else if ((InputPin->PinType.PinCategory == PC_Struct) && (InputStructType == VectorStruct))
-			{
-				TargetFunction = TEXT("Conv_LinearColorToVector");
-			}
-		}
-		else if(OutputStructType == ColorStruct)
-		{
-			if ((InputPin->PinType.PinCategory == PC_Struct) && (InputStructType == LinearColorStruct))
-			{
-				TargetFunction = TEXT("Conv_ColorToLinearColor");
-			}
-		}
-	}
-	else if (OutputPin->PinType.PinCategory == PC_Boolean)
-	{
-		if (InputPin->PinType.PinCategory == PC_String)
-		{
-			TargetFunction = TEXT("Conv_BoolToString");
-		}
-		else if (InputPin->PinType.PinCategory == PC_Int)
-		{
-			TargetFunction = TEXT("Conv_BoolToInt");
-		}
-		else if (InputPin->PinType.PinCategory == PC_Float)
-		{
-			TargetFunction = TEXT("Conv_BoolToFloat");
-		}
-		else if ((InputPin->PinType.PinCategory == PC_Byte) && (InputPin->PinType.PinSubCategoryObject == NULL))
-		{
-			TargetFunction = TEXT("Conv_BoolToByte");
-		}
-		else if ( InputPin->PinType.PinCategory == PC_Text )
-		{
-			TargetFunction = TEXT("Conv_BoolToText");
-		}
-	}
-	else if (OutputPin->PinType.PinCategory == PC_Byte &&
-			 (OutputPin->PinType.PinSubCategoryObject == NULL || !OutputPin->PinType.PinSubCategoryObject->IsA(UEnum::StaticClass())))
-	{
-		if (InputPin->PinType.PinCategory == PC_String)
-		{
-			TargetFunction = TEXT("Conv_ByteToString");
-		}
-		else if (InputPin->PinType.PinCategory == PC_Int)
-		{
-			TargetFunction = TEXT("Conv_ByteToInt");
-		}
-		else if (InputPin->PinType.PinCategory == PC_Float)
-		{
-			TargetFunction = TEXT("Conv_ByteToFloat");
-		}
-		else if ( InputPin->PinType.PinCategory == PC_Text )
-		{
-			TargetFunction = TEXT("Conv_ByteToText");
-		}
-	}
-	else if (OutputPin->PinType.PinCategory == PC_Name)
-	{
-		if (InputPin->PinType.PinCategory == PC_String)
-		{
-			TargetFunction = TEXT("Conv_NameToString");
-		}
-		else if (InputPin->PinType.PinCategory == PC_Text)
-		{
-			TargetFunction = TEXT("Conv_NameToText");
-		}
-	}
-	else if (OutputPin->PinType.PinCategory == PC_String)
-	{
-		if (InputPin->PinType.PinCategory == PC_Name)
-		{
-			TargetFunction = TEXT("Conv_StringToName");
-		}
-		else if (InputPin->PinType.PinCategory == PC_Int)
-		{
-			TargetFunction = TEXT("Conv_StringToInt");
-		}
-		else if (InputPin->PinType.PinCategory == PC_Float)
-		{
-			TargetFunction = TEXT("Conv_StringToFloat");
-		}
-		else if (InputPin->PinType.PinCategory == PC_Text)
-		{
-			TargetFunction = TEXT("Conv_StringToText");
-		}
-		else
-		{
-			// NOTE: Did you see the note above about unsafe and unclear casts?
-		}
-	}
-	else if (OutputPin->PinType.PinCategory == PC_Text)
-	{
-		if(InputPin->PinType.PinCategory == PC_String)
-		{
-			TargetFunction = TEXT("Conv_TextToString");
-		}
+		return false;
 	}
-	else if ((OutputPin->PinType.PinCategory == PC_Interface) && (InputPin->PinType.PinCategory == PC_Object))
+
+	// SPECIAL CASES, not supported by FAutocastFunctionMap
+	if ((OutputPin->PinType.PinCategory == PC_Interface) && (InputPin->PinType.PinCategory == PC_Object))
 	{
 		UClass const* InputClass = Cast<UClass const>(InputPin->PinType.PinSubCategoryObject.Get());
 
@@ -2537,10 +2377,22 @@ bool UEdGraphSchema_K2::SearchForAutocastFunction(const UEdGraphPin* OutputPin,
 			TargetFunction = TEXT("GetDisplayName");
 		}
 	}
+	else if (OutputPin->PinType.PinCategory == PC_Struct)
+	{
+		const UScriptStruct* OutputStructType = Cast<const UScriptStruct>(OutputPin->PinType.PinSubCategoryObject.Get());
+		if (OutputStructType == TBaseStructure<FRotator>::Get())
+		{
+			const UScriptStruct* InputStructType = Cast<const UScriptStruct>(InputPin->PinType.PinSubCategoryObject.Get());
+			if ((InputPin->PinType.PinCategory == PC_Struct) && (InputStructType == TBaseStructure<FTransform>::Get()))
+			{
+				TargetFunction = TEXT("MakeTransform");
+			}
+		}
+	}
 
 	if (TargetFunction == NAME_None)
 	{
-		auto AutocastFunctionMap = FAutocastFunctionMap::Get();
+		const auto& AutocastFunctionMap = FAutocastFunctionMap::Get();
 		if (auto Func = AutocastFunctionMap.Find(OutputPin->PinType, InputPin->PinType))
 		{
 			TargetFunction = Func->GetFName();
diff --git a/Engine/Source/Runtime/CoreUObject/Public/UObject/ObjectBase.h b/Engine/Source/Runtime/CoreUObject/Public/UObject/ObjectBase.h
--- a/Engine/Source/Runtime/CoreUObject/Public/UObject/ObjectBase.h
+++ b/Engine/Source/Runtime/CoreUObject/Public/UObject/ObjectBase.h
@@ -1203,7 +1203,7 @@ namespace UM
 		/// [FunctionMetadta] Used by BlueprintCallable functions to indicate which parameter is used to determine the World that the operation is occurring within.
 		WorldContext,
 
-		/// [FunctionMetadta] Used only by static BlueprintPure functions (with one return param and exactly one input param) from BlueprintLibrary.
+		/// [FunctionMetadta] Used only by static BlueprintPure functions from BlueprintLibrary. A cast node will be automatically added for the return type and the type of the first parameter of the function.
 		BlueprintAutocast,
 	};
 
diff --git a/Engine/Source/Runtime/Engine/Classes/Kismet/KismetMathLibrary.h b/Engine/Source/Runtime/Engine/Classes/Kismet/KismetMathLibrary.h
--- a/Engine/Source/Runtime/Engine/Classes/Kismet/KismetMathLibrary.h
+++ b/Engine/Source/Runtime/Engine/Classes/Kismet/KismetMathLibrary.h
@@ -488,7 +488,7 @@ class ENGINE_API UKismetMathLibrary : public UBlueprintFunctionLibrary
 	static int32 FFloor(float A);
 	
 	/* Rounds A to an integer with truncation towards zero.  (e.g. -1.7 truncated to -1, 2.8 truncated to 2) */
-	UFUNCTION(BlueprintPure, meta=(DisplayName = "Truncate"), Category="Math|Float")
+	UFUNCTION(BlueprintPure, meta=(DisplayName = "Truncate", BlueprintAutocast), Category="Math|Float")
 	static int32 FTrunc(float A);
 	
 	/* Rounds A to the smallest following integer */
@@ -1073,71 +1073,71 @@ class ENGINE_API UKismetMathLibrary : public UBlueprintFunctionLibrary
 	// -- Begin K2 utilities
 
 	/** Converts a byte to a float */
-	UFUNCTION(BlueprintPure, meta=(DisplayName = "ToFloat (byte)", CompactNodeTitle = "->", Keywords="cast convert"), Category="Math|Conversions")
+	UFUNCTION(BlueprintPure, meta=(DisplayName = "ToFloat (byte)", CompactNodeTitle = "->", Keywords="cast convert", BlueprintAutocast), Category="Math|Conversions")
 	static float Conv_ByteToFloat(uint8 InByte);
 
 	/** Converts an integer to a float */
-	UFUNCTION(BlueprintPure, meta=(DisplayName = "ToFloat (int)", CompactNodeTitle = "->", Keywords="cast convert"), Category="Math|Conversions")
+	UFUNCTION(BlueprintPure, meta=(DisplayName = "ToFloat (int)", CompactNodeTitle = "->", Keywords="cast convert", BlueprintAutocast), Category="Math|Conversions")
 	static float Conv_IntToFloat(int32 InInt);
 
 	/** Converts an integer to a byte (if the integer is too large, returns the low 8 bits) */
-	UFUNCTION(BlueprintPure, meta=(DisplayName = "ToByte (int)", CompactNodeTitle = "->", Keywords="cast convert"), Category="Math|Conversions")
+	UFUNCTION(BlueprintPure, meta=(DisplayName = "ToByte (int)", CompactNodeTitle = "->", Keywords="cast convert", BlueprintAutocast), Category="Math|Conversions")
 	static uint8 Conv_IntToByte(int32 InInt);
 
 	/** Converts a int to a bool*/
-	UFUNCTION(BlueprintPure, meta=(DisplayName = "ToBool (int)", CompactNodeTitle = "->", Keywords="cast convert"), Category="Math|Conversions")
+	UFUNCTION(BlueprintPure, meta=(DisplayName = "ToBool (int)", CompactNodeTitle = "->", Keywords="cast convert", BlueprintAutocast), Category="Math|Conversions")
 	static bool Conv_IntToBool(int32 InInt);
 
 	/** Converts a bool to an int */
-	UFUNCTION(BlueprintPure, meta=(DisplayName = "ToInt (bool)", CompactNodeTitle = "->", Keywords="cast convert"), Category="Math|Conversions")
+	UFUNCTION(BlueprintPure, meta=(DisplayName = "ToInt (bool)", CompactNodeTitle = "->", Keywords="cast convert", BlueprintAutocast), Category="Math|Conversions")
 	static int32 Conv_BoolToInt(bool InBool);
 
 	/** Converts a bool to a float (0.0f or 1.0f) */
-	UFUNCTION(BlueprintPure, meta=(DisplayName = "ToFloat (bool)", CompactNodeTitle = "->", Keywords="cast convert"), Category="Math|Conversions")
+	UFUNCTION(BlueprintPure, meta=(DisplayName = "ToFloat (bool)", CompactNodeTitle = "->", Keywords="cast convert", BlueprintAutocast), Category="Math|Conversions")
 	static float Conv_BoolToFloat(bool InBool);
 
 	/** Converts a bool to a byte */
-	UFUNCTION(BlueprintPure, meta=(DisplayName = "ToByte (bool)", CompactNodeTitle = "->", Keywords="cast convert"), Category="Math|Conversions")
+	UFUNCTION(BlueprintPure, meta=(DisplayName = "ToByte (bool)", CompactNodeTitle = "->", Keywords="cast convert", BlueprintAutocast), Category="Math|Conversions")
 	static uint8 Conv_BoolToByte(bool InBool);
 	
 	/** Converts a byte to an integer */
-	UFUNCTION(BlueprintPure, meta=(DisplayName = "ToInt (byte)", CompactNodeTitle = "->", Keywords="cast convert"), Category="Math|Conversions")
+	UFUNCTION(BlueprintPure, meta=(DisplayName = "ToInt (byte)", CompactNodeTitle = "->", Keywords="cast convert", BlueprintAutocast), Category="Math|Conversions")
 	static int32 Conv_ByteToInt(uint8 InByte);
 
 	/** Converts a vector to LinearColor */
-	UFUNCTION(BlueprintPure, meta=(DisplayName = "ToLinearColor (vector)", CompactNodeTitle = "->", Keywords="cast convert"), Category="Math|Conversions")
+	UFUNCTION(BlueprintPure, meta=(DisplayName = "ToLinearColor (vector)", CompactNodeTitle = "->", Keywords="cast convert", BlueprintAutocast), Category="Math|Conversions")
 	static FLinearColor Conv_VectorToLinearColor(FVector InVec);
 
 	/** Converts a LinearColor to a vector */
-	UFUNCTION(BlueprintPure, meta=(DisplayName = "ToVector (linear color)", CompactNodeTitle = "->", Keywords="cast convert"), Category="Math|Conversions")
+	UFUNCTION(BlueprintPure, meta=(DisplayName = "ToVector (linear color)", CompactNodeTitle = "->", Keywords="cast convert", BlueprintAutocast), Category="Math|Conversions")
 	static FVector Conv_LinearColorToVector(FLinearColor InLinearColor);
 
 	/** Converts a color to LinearColor */
-	UFUNCTION(BlueprintPure, meta=(DisplayName = "ToLinearColor (color)", CompactNodeTitle = "->", Keywords="cast convert"), Category="Math|Conversions")
+	UFUNCTION(BlueprintPure, meta=(DisplayName = "ToLinearColor (color)", CompactNodeTitle = "->", Keywords="cast convert", BlueprintAutocast), Category="Math|Conversions")
 	static FLinearColor Conv_ColorToLinearColor(FColor InColor);
 
 	/** Converts a LinearColor to a color*/
-	UFUNCTION(BlueprintPure, meta=(DisplayName = "ToColor (linear color)", CompactNodeTitle = "->", Keywords="cast convert"), Category="Math|Conversions")
+	UFUNCTION(BlueprintPure, meta=(DisplayName = "ToColor (linear color)", CompactNodeTitle = "->", Keywords="cast convert", BlueprintAutocast), Category="Math|Conversions")
 	static FColor Conv_LinearColorToColor(FLinearColor InLinearColor);
 
 	/** Convert a vector to a transform. Uses vector as location */
-	UFUNCTION(BlueprintPure, meta=(DisplayName = "ToTransform (vector)", CompactNodeTitle = "->", Keywords="cast convert"), Category="Math|Conversions")
+	UFUNCTION(BlueprintPure, meta=(DisplayName = "ToTransform (vector)", CompactNodeTitle = "->", Keywords="cast convert", BlueprintAutocast), Category="Math|Conversions")
 	static FTransform Conv_VectorToTransform(FVector InLocation);
 	
 	/** Convert a Vector to a Vector2D */
-	UFUNCTION(BlueprintPure, meta=(DisplayName = "ToVector2D (Vector)", CompactNodeTitle = "->", Keywords="cast convert"), Category="Math|Conversions")
+	UFUNCTION(BlueprintPure, meta=(DisplayName = "ToVector2D (Vector)", CompactNodeTitle = "->", Keywords="cast convert", BlueprintAutocast), Category="Math|Conversions")
 	static FVector2D Conv_VectorToVector2D(FVector InVector);
 
 	/** Convert a Vector2D to a Vector */
-	UFUNCTION(BlueprintPure, meta=(DisplayName = "ToVector (Vector2D)", CompactNodeTitle = "->", Keywords="cast convert"), Category="Math|Conversions")
+	UFUNCTION(BlueprintPure, meta=(DisplayName = "ToVector (Vector2D)", CompactNodeTitle = "->", Keywords="cast convert", BlueprintAutocast), Category="Math|Conversions")
 	static FVector Conv_Vector2DToVector(FVector2D InVector2D, float Z = 0);
 
 	/** Convert a float into a vector, where each element is that float */
-	UFUNCTION(BlueprintPure, meta=(DisplayName = "ToVector (float)", CompactNodeTitle = "->", Keywords="cast convert"), Category="Math|Conversions")
+	UFUNCTION(BlueprintPure, meta=(DisplayName = "ToVector (float)", CompactNodeTitle = "->", Keywords="cast convert", BlueprintAutocast), Category="Math|Conversions")
 	static FVector Conv_FloatToVector(float InFloat);
 
 	/** Convert a float into a LinearColor, where each element is that float */
-	UFUNCTION(BlueprintPure, meta=(DisplayName = "ToLinearColor (float)", CompactNodeTitle = "->", Keywords="cast convert"), Category="Math|Conversions")
+	UFUNCTION(BlueprintPure, meta=(DisplayName = "ToLinearColor (float)", CompactNodeTitle = "->", Keywords="cast convert", BlueprintAutocast), Category="Math|Conversions")
 	static FLinearColor Conv_FloatToLinearColor(float InFloat);
 
 	/** Makes a vector {X, Y, Z} */
@@ -1301,11 +1301,11 @@ class ENGINE_API UKismetMathLibrary : public UBlueprintFunctionLibrary
 	static FRotator MakeRotationFromAxes(FVector Forward, FVector Right, FVector Up);
 
 	/** Create a rotator which orients X along the supplied direction vector */
-	UFUNCTION(BlueprintPure, meta=(DisplayName = "RotationFromXVector", Keywords="rotation rotate cast convert"), Category="Math|Rotator")
+	UFUNCTION(BlueprintPure, meta=(DisplayName = "RotationFromXVector", Keywords="rotation rotate cast convert", BlueprintAutocast), Category="Math|Rotator")
 	static FRotator Conv_VectorToRotator(FVector InVec);
 
 	/** Get the X direction vector after this rotation */
-	UFUNCTION(BlueprintPure, meta=(DisplayName = "GetRotationXVector", Keywords="rotation rotate cast convert"), Category="Math|Rotator")
+	UFUNCTION(BlueprintPure, meta=(DisplayName = "GetRotationXVector", Keywords="rotation rotate cast convert", BlueprintAutocast), Category="Math|Rotator")
 	static FVector Conv_RotatorToVector(FRotator InRot);
 
 
diff --git a/Engine/Source/Runtime/Engine/Classes/Kismet/KismetStringLibrary.h b/Engine/Source/Runtime/Engine/Classes/Kismet/KismetStringLibrary.h
--- a/Engine/Source/Runtime/Engine/Classes/Kismet/KismetStringLibrary.h
+++ b/Engine/Source/Runtime/Engine/Classes/Kismet/KismetStringLibrary.h
@@ -10,59 +10,59 @@ class ENGINE_API UKismetStringLibrary : public UBlueprintFunctionLibrary
 	GENERATED_UCLASS_BODY()
 
 	/** Converts a float value to a string */
-	UFUNCTION(BlueprintPure, meta=(DisplayName = "ToString (float)", CompactNodeTitle = "->"), Category="Utilities|String")
+	UFUNCTION(BlueprintPure, meta=(DisplayName = "ToString (float)", CompactNodeTitle = "->", BlueprintAutocast), Category="Utilities|String")
 	static FString Conv_FloatToString(float InFloat);
 
 	/** Converts an integer value to a string */
-	UFUNCTION(BlueprintPure, meta=(DisplayName = "ToString (int)", CompactNodeTitle = "->"), Category="Utilities|String")
+	UFUNCTION(BlueprintPure, meta=(DisplayName = "ToString (int)", CompactNodeTitle = "->", BlueprintAutocast), Category="Utilities|String")
 	static FString Conv_IntToString(int32 InInt);
 
 	/** Converts a byte value to a string */
-	UFUNCTION(BlueprintPure, meta=(DisplayName = "ToString (byte)", CompactNodeTitle = "->"), Category="Utilities|String")
+	UFUNCTION(BlueprintPure, meta=(DisplayName = "ToString (byte)", CompactNodeTitle = "->", BlueprintAutocast), Category="Utilities|String")
 	static FString Conv_ByteToString(uint8 InByte);
 
 	/** Converts a boolean value to a string, either 'true' or 'false' */
-	UFUNCTION(BlueprintPure, meta=(DisplayName = "ToString (boolean)", CompactNodeTitle = "->"), Category="Utilities|String")
+	UFUNCTION(BlueprintPure, meta=(DisplayName = "ToString (boolean)", CompactNodeTitle = "->", BlueprintAutocast), Category="Utilities|String")
 	static FString Conv_BoolToString(bool InBool);
 
 	/** Converts a vector value to a string, in the form 'X= Y= Z=' */
-	UFUNCTION(BlueprintPure, meta=(DisplayName = "ToString (Vector)", CompactNodeTitle = "->"), Category="Utilities|String")
+	UFUNCTION(BlueprintPure, meta=(DisplayName = "ToString (Vector)", CompactNodeTitle = "->", BlueprintAutocast), Category="Utilities|String")
 	static FString Conv_VectorToString(FVector InVec);
 
 	/** Converts a vector2d value to a string, in the form 'X= Y=' */
-	UFUNCTION(BlueprintPure, meta=(DisplayName = "ToString (vector2d)", CompactNodeTitle = "->"), Category="Utilities|String")
+	UFUNCTION(BlueprintPure, meta=(DisplayName = "ToString (vector2d)", CompactNodeTitle = "->", BlueprintAutocast), Category="Utilities|String")
 	static FString Conv_Vector2dToString(FVector2D InVec);
 
 	/** Converts a rotator value to a string, in the form 'P= Y= R=' */
-	UFUNCTION(BlueprintPure, meta=(DisplayName = "ToString (rotator)", CompactNodeTitle = "->"), Category="Utilities|String")
+	UFUNCTION(BlueprintPure, meta=(DisplayName = "ToString (rotator)", CompactNodeTitle = "->", BlueprintAutocast), Category="Utilities|String")
 	static FString Conv_RotatorToString(FRotator InRot);
 
 	/** Converts a transform value to a string, in the form 'Translation: X= Y= Z= Rotation: P= Y= R= Scale: X= Y= Z=' */
-	UFUNCTION(BlueprintPure, meta=(DisplayName = "ToString (transform)", CompactNodeTitle = "->"), Category="Utilities|String")
+	UFUNCTION(BlueprintPure, meta=(DisplayName = "ToString (transform)", CompactNodeTitle = "->", BlueprintAutocast), Category="Utilities|String")
 	static FString Conv_TransformToString(const FTransform& InTrans);
 
 	/** Converts a UObject value to a string by calling the object's GetName method */
-	UFUNCTION(BlueprintPure, meta=(DisplayName = "ToString (object)", CompactNodeTitle = "->"), Category="Utilities|String")
+	UFUNCTION(BlueprintPure, meta=(DisplayName = "ToString (object)", CompactNodeTitle = "->", BlueprintAutocast), Category="Utilities|String")
 	static FString Conv_ObjectToString(class UObject* InObj);
 
 	/** Converts a linear color value to a string, in the form '(R=,G=,B=,A=)' */
-	UFUNCTION(BlueprintPure, meta=(DisplayName = "ToString (linear color)", CompactNodeTitle = "->"), Category="Utilities|String")
+	UFUNCTION(BlueprintPure, meta=(DisplayName = "ToString (linear color)", CompactNodeTitle = "->", BlueprintAutocast), Category="Utilities|String")
 	static FString Conv_ColorToString(FLinearColor InColor);
 
 	/** Converts a name value to a string */
-	UFUNCTION(BlueprintPure, meta=(DisplayName = "ToString (name)", CompactNodeTitle = "->"), Category="Utilities|String")
+	UFUNCTION(BlueprintPure, meta=(DisplayName = "ToString (name)", CompactNodeTitle = "->", BlueprintAutocast), Category="Utilities|String")
 	static FString Conv_NameToString(FName InName);
 
 	/** Converts a string to a name value */
-	UFUNCTION(BlueprintPure, meta=(DisplayName = "String To Name", CompactNodeTitle = "->"), Category="Utilities|String")
+	UFUNCTION(BlueprintPure, meta=(DisplayName = "String To Name", CompactNodeTitle = "->", BlueprintAutocast), Category="Utilities|String")
 	static FName Conv_StringToName(const FString& InString);
 
 	/** Converts a string to a int value */
-	UFUNCTION(BlueprintPure, meta=(DisplayName = "String To Int", CompactNodeTitle = "->"), Category="Utilities|String")
+	UFUNCTION(BlueprintPure, meta=(DisplayName = "String To Int", CompactNodeTitle = "->", BlueprintAutocast), Category="Utilities|String")
 	static int32 Conv_StringToInt(const FString& InString);
 
 	/** Converts a string to a float value */
-	UFUNCTION(BlueprintPure, meta=(DisplayName = "String To Float", CompactNodeTitle = "->"), Category="Utilities|String")
+	UFUNCTION(BlueprintPure, meta=(DisplayName = "String To Float", CompactNodeTitle = "->", BlueprintAutocast), Category="Utilities|String")
 	static float Conv_StringToFloat(const FString& InString);
 	
 
diff --git a/Engine/Source/Runtime/Engine/Classes/Kismet/KismetTextLibrary.h b/Engine/Source/Runtime/Engine/Classes/Kismet/KismetTextLibrary.h
--- a/Engine/Source/Runtime/Engine/Classes/Kismet/KismetTextLibrary.h
+++ b/Engine/Source/Runtime/Engine/Classes/Kismet/KismetTextLibrary.h
@@ -46,15 +46,15 @@ class ENGINE_API UKismetTextLibrary : public UBlueprintFunctionLibrary
 	GENERATED_UCLASS_BODY()
 
 	/** Converts localizable text to the string */
-	UFUNCTION(BlueprintPure, meta=(DisplayName = "ToString (text)", CompactNodeTitle = "->"), Category="Utilities|String")
+	UFUNCTION(BlueprintPure, meta = (DisplayName = "ToString (text)", CompactNodeTitle = "->", BlueprintAutocast), Category = "Utilities|String")
 	static FString Conv_TextToString(const FText& InText);
 
 	/** Converts string to localizable text */
-	UFUNCTION(BlueprintPure, meta=(DisplayName = "ToText (string)", CompactNodeTitle = "->"), Category="Utilities|Text")
+	UFUNCTION(BlueprintPure, meta=(DisplayName = "ToText (string)", CompactNodeTitle = "->", BlueprintAutocast), Category="Utilities|Text")
 	static FText Conv_StringToText(const FString& InString);
 
 	/** Converts string to localizable text */
-	UFUNCTION(BlueprintPure, meta=(DisplayName = "ToText (name)", CompactNodeTitle = "->"), Category="Utilities|Text")
+	UFUNCTION(BlueprintPure, meta=(DisplayName = "ToText (name)", CompactNodeTitle = "->", BlueprintAutocast), Category="Utilities|Text")
 	static FText Conv_NameToText(FName InName);
 
 	/* Returns true if text is empty. */
@@ -106,21 +106,21 @@ class ENGINE_API UKismetTextLibrary : public UBlueprintFunctionLibrary
 	static bool NotEqual_IgnoreCase_TextText(const FText& A, const FText& B);
 
 	/** Converts a boolean value to text, either 'true' or 'false' */
-	UFUNCTION(BlueprintPure, meta=(DisplayName = "ToText (boolean)", CompactNodeTitle = "->"), Category="Utilities|Text")
+	UFUNCTION(BlueprintPure, meta=(DisplayName = "ToText (boolean)", CompactNodeTitle = "->", BlueprintAutocast), Category="Utilities|Text")
 	static FText Conv_BoolToText(bool InBool);
 
 	/** Converts a byte value to text */
-	UFUNCTION(BlueprintPure, meta=(DisplayName = "ToText (byte)", CompactNodeTitle = "->"), Category="Utilities|Text")
+	UFUNCTION(BlueprintPure, meta=(DisplayName = "ToText (byte)", CompactNodeTitle = "->", BlueprintAutocast), Category="Utilities|Text")
 	static FText Conv_ByteToText(uint8 Value);
 
 	// Default values are duplicated from FNumberFormattingOptions and should be replicated in all functions and in the struct when changed!
 	/* Converts a passed in integer to a text based on formatting options */
-	UFUNCTION(BlueprintPure, meta=(DisplayName = "ToText (int)", AdvancedDisplay = "1"), Category="Utilities|Text")
+	UFUNCTION(BlueprintPure, meta=(DisplayName = "ToText (int)", AdvancedDisplay = "1", BlueprintAutocast), Category="Utilities|Text")
 	static FText Conv_IntToText(int32 Value, bool bUseGrouping = true, int32 MinimumIntegralDigits = 1, int32 MaximumIntegralDigits = 324);
 
 	// Default values are duplicated from FNumberFormattingOptions and should be replicated in all functions and in the struct when changed!
 	/* Converts a passed in float to a text based on formatting options */
-	UFUNCTION(BlueprintPure, meta=(DisplayName = "ToText (float)", AdvancedDisplay = "1"), Category="Utilities|Text")
+	UFUNCTION(BlueprintPure, meta=(DisplayName = "ToText (float)", AdvancedDisplay = "1", BlueprintAutocast), Category="Utilities|Text")
 	static FText Conv_FloatToText(float Value, TEnumAsByte<ERoundingMode> RoundingMode, bool bUseGrouping = true, int32 MinimumIntegralDigits = 1, int32 MaximumIntegralDigits = 324, int32 MinimumFractionalDigits = 0, int32 MaximumFractionalDigits = 3);
 
 	// Default values are duplicated from FNumberFormattingOptions and should be replicated in all functions and in the struct when changed!
