Input Touch event node cannot be placed in function or macro graphs.

diff --git a/Engine/Source/Editor/BlueprintGraph/Private/K2Node_InputTouch.cpp b/Engine/Source/Editor/BlueprintGraph/Private/K2Node_InputTouch.cpp
--- a/Engine/Source/Editor/BlueprintGraph/Private/K2Node_InputTouch.cpp
+++ b/Engine/Source/Editor/BlueprintGraph/Private/K2Node_InputTouch.cpp
@@ -102,12 +102,20 @@ FText UK2Node_InputTouch::GetMenuCategory() const
 
 bool UK2Node_InputTouch::IsCompatibleWithGraph(const UEdGraph* TargetGraph) const
 {
-	UBlueprint* Blueprint = FBlueprintEditorUtils::FindBlueprintForGraph(TargetGraph);
+	// This node expands into event nodes and must be placed in a Ubergraph
+	EGraphType const GraphType = TargetGraph->GetSchema()->GetGraphType(TargetGraph);
+	bool bIsCompatible = (GraphType == EGraphType::GT_Ubergraph);
 
-	UEdGraphSchema_K2 const* K2Schema = Cast<UEdGraphSchema_K2>(TargetGraph->GetSchema());
-	bool const bIsConstructionScript = (K2Schema != nullptr) ? K2Schema->IsConstructionScript(TargetGraph) : false;
+	if (bIsCompatible)
+	{
+		UBlueprint* Blueprint = FBlueprintEditorUtils::FindBlueprintForGraph(TargetGraph);
+
+		UEdGraphSchema_K2 const* K2Schema = Cast<UEdGraphSchema_K2>(TargetGraph->GetSchema());
+		bool const bIsConstructionScript = (K2Schema != nullptr) ? K2Schema->IsConstructionScript(TargetGraph) : false;
 
-	return (Blueprint != nullptr) && Blueprint->SupportsInputEvents() && !bIsConstructionScript && Super::IsCompatibleWithGraph(TargetGraph);
+		bIsCompatible = (Blueprint != nullptr) && Blueprint->SupportsInputEvents() && !bIsConstructionScript && Super::IsCompatibleWithGraph(TargetGraph);
+	}
+	return bIsCompatible;
 }
 
 UEdGraphPin* UK2Node_InputTouch::GetPressedPin() const
