[UE-15480] Restore ability to link GetClass() function call node output pin subtype to the input pin subtype (revised initial fix).

change summary:
- restored modification to FDynamicOutputHelper::GetPinClass() to include object reference pin types
- restored modification to FDynamicOutputHelper::IsTypePickerPin() to include the PC_Object category
- restored modification to the UGameplayStatics::GetObjectClass() UFUNCTION declaration to include 'DeterminesOutputType' metadata
- modified FDynamicOutputHelper::GetPinClass() to consider connections to a literal 'self' reference; resolves a potential BP compilation error that was introduced with the initial fix

diff --git a/Engine/Source/Editor/BlueprintGraph/Private/K2Node_CallFunction.cpp b/Engine/Source/Editor/BlueprintGraph/Private/K2Node_CallFunction.cpp
--- a/Engine/Source/Editor/BlueprintGraph/Private/K2Node_CallFunction.cpp
+++ b/Engine/Source/Editor/BlueprintGraph/Private/K2Node_CallFunction.cpp
@@ -238,8 +238,8 @@ UClass* FDynamicOutputHelper::GetPinClass(UEdGraphPin* Pin)
 {
 	UClass* PinClass = UObject::StaticClass();
 
-	bool const bIsClassPin = (Pin->PinType.PinCategory == UEdGraphSchema_K2::PC_Class);
-	if (bIsClassPin)
+	bool const bIsClassOrObjectPin = (Pin->PinType.PinCategory == UEdGraphSchema_K2::PC_Class || Pin->PinType.PinCategory == UEdGraphSchema_K2::PC_Object);
+	if (bIsClassOrObjectPin)
 	{
 		if (UClass* DefaultClass = Cast<UClass>(Pin->DefaultObject))
 		{
@@ -252,24 +252,33 @@ UClass* FDynamicOutputHelper::GetPinClass(UEdGraphPin* Pin)
 
 		if (Pin->LinkedTo.Num() > 0)
 		{
-			UClass* CommonInputClass = Cast<UClass>(Pin->LinkedTo[0]->PinType.PinSubCategoryObject.Get());
-			for (int32 LinkIndex = 1; LinkIndex < Pin->LinkedTo.Num(); ++LinkIndex)
+			UClass* CommonInputClass = nullptr;
+			for (UEdGraphPin* LinkedPin : Pin->LinkedTo)
 			{
-				UClass* LinkClass = Cast<UClass>(Pin->LinkedTo[LinkIndex]->PinType.PinSubCategoryObject.Get());
-				if (LinkClass == nullptr)
-				{
-					continue;
-				}
+				const FEdGraphPinType& LinkedPinType = LinkedPin->PinType;
 
-				if (CommonInputClass == nullptr)
+				UClass* LinkClass = Cast<UClass>(LinkedPinType.PinSubCategoryObject.Get());
+				if (LinkClass == nullptr && LinkedPinType.PinSubCategory == UEdGraphSchema_K2::PSC_Self)
 				{
-					CommonInputClass = LinkClass;
-					continue;
+					if (UK2Node* K2Node = Cast<UK2Node>(LinkedPin->GetOwningNode()))
+					{
+						LinkClass = K2Node->GetBlueprint()->GeneratedClass;
+					}
 				}
 
-				while (!LinkClass->IsChildOf(CommonInputClass))
+				if (LinkClass != nullptr)
 				{
-					CommonInputClass = CommonInputClass->GetSuperClass();
+					if (CommonInputClass != nullptr)
+					{
+						while (!LinkClass->IsChildOf(CommonInputClass))
+						{
+							CommonInputClass = CommonInputClass->GetSuperClass();
+						}
+					}
+					else
+					{
+						CommonInputClass = LinkClass;
+					}
 				}
 			}
 
@@ -319,7 +328,8 @@ bool FDynamicOutputHelper::IsTypePickerPin(UEdGraphPin* Pin) const
 	}
 
 	bool const bPinIsClassPicker = (Pin->PinType.PinCategory == UEdGraphSchema_K2::PC_Class);
-	return bIsTypeDeterminingPin && bPinIsClassPicker && (Pin->Direction == EGPD_Input);
+	bool const bPinIsObjectPicker = (Pin->PinType.PinCategory == UEdGraphSchema_K2::PC_Object);
+	return bIsTypeDeterminingPin && (bPinIsClassPicker || bPinIsObjectPicker) && (Pin->Direction == EGPD_Input);
 }
 
 UEdGraphPin* FDynamicOutputHelper::GetDynamicOutPin(const UK2Node_CallFunction* FuncNode)
diff --git a/Engine/Source/Runtime/Engine/Classes/Kismet/GameplayStatics.h b/Engine/Source/Runtime/Engine/Classes/Kismet/GameplayStatics.h
--- a/Engine/Source/Runtime/Engine/Classes/Kismet/GameplayStatics.h
+++ b/Engine/Source/Runtime/Engine/Classes/Kismet/GameplayStatics.h
@@ -157,7 +157,7 @@ class ENGINE_API UGameplayStatics : public UBlueprintFunctionLibrary
 	UFUNCTION(BlueprintPure, Category="Game", meta=(WorldContext="WorldContextObject") )
 	static class AGameState* GetGameState(UObject* WorldContextObject);
 
-	UFUNCTION(BlueprintPure, meta=(DisplayName = "GetClass"), Category="Utilities")
+	UFUNCTION(BlueprintPure, meta=(DisplayName = "GetClass", DeterminesOutputType = "Object"), Category="Utilities")
 	static class UClass *GetObjectClass(const UObject *Object);
 
 	/**
