[UE-14914] Prevent function parameters from being autowired when intended for internal use only.

change summary:
- added a new "InternalUseParam" UFUNCTION() meta to support tagging function parameters as internal use only (=> hidden | not connectable)
- added a UK2Node_CallFunction::IsConnectionDisallowed() override method; returns 'true' if the 'bNotConnectable' flag is set on the node's pin (and will emit a compile error in existing projects if something is connected to it)
- modified FBlueprintEditorUtils::GetHiddenPinsForFunction() to accept an optional 'InternalPins' TSet as an output; on return, this will contain the subset of hidden pins that are marked as "InternalUseParam"
- modified UK2Node_CallFunction::CreatePinsForFunctionCall() to set the 'bNotConnectable' pin attribute to 'true' if "InternalUseParam == '<PinName>'" is included within the UFunction metadata
- modified the AActor::AddComponent UFUNCTION() decl to switch the "HidePin" meta to "InternalUseParam" for the 'ComponentTemplateContext' parameter (because it should not be connectable)
- modified UK2Node_CallFunction::NotifyPinConnectionListChanged() to refresh the node to hide pins corresponding to internal-only params once the (invalid) connection in an existing project has been broken

diff --git a/Engine/Source/Editor/BlueprintGraph/Classes/EdGraphSchema_K2.h b/Engine/Source/Editor/BlueprintGraph/Classes/EdGraphSchema_K2.h
--- a/Engine/Source/Editor/BlueprintGraph/Classes/EdGraphSchema_K2.h
+++ b/Engine/Source/Editor/BlueprintGraph/Classes/EdGraphSchema_K2.h
@@ -93,6 +93,9 @@ public:
 	// [FunctionMetadata] Indicates that the function should be drawn with this title over the function name
 	static const FName MD_DisplayName;
 
+	// [FunctionMetadata] Indicates that a particular function parameter is for internal use only, which means it will be both hidden and not connectible.
+	static const FName MD_InternalUseParam;
+
 	//    property metadata
 
 	/** UPROPERTY will be exposed on "Spawn Blueprint" nodes as an input  */
diff --git a/Engine/Source/Editor/BlueprintGraph/Classes/K2Node_CallFunction.h b/Engine/Source/Editor/BlueprintGraph/Classes/K2Node_CallFunction.h
--- a/Engine/Source/Editor/BlueprintGraph/Classes/K2Node_CallFunction.h
+++ b/Engine/Source/Editor/BlueprintGraph/Classes/K2Node_CallFunction.h
@@ -102,6 +102,7 @@ public:
 	virtual void GetNodeAttributes( TArray<TKeyValuePair<FString, FString>>& OutNodeAttributes ) const override;
 	virtual FText GetMenuCategory() const override;
 	virtual bool IsActionFilteredOut(class FBlueprintActionFilter const& Filter) override;
+	virtual bool IsConnectionDisallowed(const UEdGraphPin* MyPin, const UEdGraphPin* OtherPin, FString& OutReason) const override;
 	// End of UK2Node interface
 
 	/** Returns the UFunction that this class is pointing to */
diff --git a/Engine/Source/Editor/BlueprintGraph/Private/EdGraphSchema_K2.cpp b/Engine/Source/Editor/BlueprintGraph/Private/EdGraphSchema_K2.cpp
--- a/Engine/Source/Editor/BlueprintGraph/Private/EdGraphSchema_K2.cpp
+++ b/Engine/Source/Editor/BlueprintGraph/Private/EdGraphSchema_K2.cpp
@@ -58,6 +58,7 @@ const FName FBlueprintMetadata::MD_DeprecatedFunction(TEXT("DeprecatedFunction")
 const FName FBlueprintMetadata::MD_DeprecationMessage(TEXT("DeprecationMessage"));
 const FName FBlueprintMetadata::MD_CompactNodeTitle(TEXT("CompactNodeTitle"));
 const FName FBlueprintMetadata::MD_DisplayName(TEXT("DisplayName"));
+const FName FBlueprintMetadata::MD_InternalUseParam(TEXT("InternalUseParam"));
 
 const FName FBlueprintMetadata::MD_ExposeOnSpawn(TEXT("ExposeOnSpawn"));
 const FName FBlueprintMetadata::MD_DefaultToSelf(TEXT("DefaultToSelf"));
diff --git a/Engine/Source/Editor/BlueprintGraph/Private/K2Node_CallFunction.cpp b/Engine/Source/Editor/BlueprintGraph/Private/K2Node_CallFunction.cpp
--- a/Engine/Source/Editor/BlueprintGraph/Private/K2Node_CallFunction.cpp
+++ b/Engine/Source/Editor/BlueprintGraph/Private/K2Node_CallFunction.cpp
@@ -852,7 +852,8 @@ bool UK2Node_CallFunction::CreatePinsForFunctionCall(const UFunction* Function)
 
 	// Build a list of the pins that should be hidden for this function (ones that are automagically filled in by the K2 compiler)
 	TSet<FString> PinsToHide;
-	FBlueprintEditorUtils::GetHiddenPinsForFunction(Graph, Function, PinsToHide);
+	TSet<FString> InternalPins;
+	FBlueprintEditorUtils::GetHiddenPinsForFunction(Graph, Function, PinsToHide, &InternalPins);
 
 	const bool bShowWorldContextPin = ((PinsToHide.Num() > 0) && BP && BP->ParentClass && BP->ParentClass->HasMetaData(FBlueprintMetadata::MD_ShowWorldContextPin));
 
@@ -892,6 +893,7 @@ bool UK2Node_CallFunction::CreatePinsForFunctionCall(const UFunction* Function)
 				if (!bShowWorldContextPin || !bIsSelfPin)
 				{
 					Pin->bHidden = true;
+					Pin->bNotConnectable = InternalPins.Contains(Pin->PinName);
 					K2Schema->SetPinDefaultValueBasedOnType(Pin);
 				}
 			}
@@ -984,6 +986,12 @@ void UK2Node_CallFunction::NotifyPinConnectionListChanged(UEdGraphPin* Pin)
 	if (Pin)
 	{
 		FCustomStructureParamHelper::UpdateCustomStructurePins(GetTargetFunction(), this, Pin);
+
+		// Refresh the node to hide internal-only pins once the [invalid] connection has been broken
+		if (Pin->bHidden && Pin->bNotConnectable && Pin->LinkedTo.Num() == 0)
+		{
+			GetGraph()->NotifyGraphChanged();
+		}
 	}
 
 	if (bIsBeadFunction)
@@ -2338,4 +2346,19 @@ void UK2Node_CallFunction::AddSearchMetaDataInfo(TArray<struct FSearchTagDataPai
 	}
 }
 
+bool UK2Node_CallFunction::IsConnectionDisallowed(const UEdGraphPin* MyPin, const UEdGraphPin* OtherPin, FString& OutReason) const
+{
+	bool bIsDisallowed = Super::IsConnectionDisallowed(MyPin, OtherPin, OutReason);
+	if (!bIsDisallowed && MyPin != nullptr)
+	{
+		if (MyPin->bNotConnectable)
+		{
+			bIsDisallowed = true;
+			OutReason = LOCTEXT("PinConnectionDisallowed", "This parameter is for internal use only.").ToString();
+		}
+	}
+
+	return bIsDisallowed;
+}
+
 #undef LOCTEXT_NAMESPACE
diff --git a/Engine/Source/Editor/UnrealEd/Private/Kismet2/BlueprintEditorUtils.cpp b/Engine/Source/Editor/UnrealEd/Private/Kismet2/BlueprintEditorUtils.cpp
--- a/Engine/Source/Editor/UnrealEd/Private/Kismet2/BlueprintEditorUtils.cpp
+++ b/Engine/Source/Editor/UnrealEd/Private/Kismet2/BlueprintEditorUtils.cpp
@@ -3815,7 +3815,7 @@ UEdGraph* FBlueprintEditorUtils::GetDelegateSignatureGraphByName(UBlueprint* Blu
 }
 
 // Gets a list of pins that should hidden for a given function
-void FBlueprintEditorUtils::GetHiddenPinsForFunction(UEdGraph const* Graph, UFunction const* Function, TSet<FString>& HiddenPins)
+void FBlueprintEditorUtils::GetHiddenPinsForFunction(UEdGraph const* Graph, UFunction const* Function, TSet<FString>& HiddenPins, TSet<FString>* OutInternalPins)
 {
 	check(Function != NULL);
 	TMap<FName, FString>* MetaData = UMetaData::GetMapForObject(Function);	
@@ -3836,6 +3836,15 @@ void FBlueprintEditorUtils::GetHiddenPinsForFunction(UEdGraph const* Graph, UFun
 			{
 				HiddenPins.Add(It.Value());
 			}
+			else if (Key == FBlueprintMetadata::MD_InternalUseParam)
+			{
+				HiddenPins.Add(It.Value());
+
+				if (OutInternalPins != nullptr)
+				{
+					OutInternalPins->Add(It.Value());
+				}
+			}
 			else if(Key == FBlueprintMetadata::MD_ExpandEnumAsExecs)
 			{
 				HiddenPins.Add(It.Value());
diff --git a/Engine/Source/Editor/UnrealEd/Public/Kismet2/BlueprintEditorUtils.h b/Engine/Source/Editor/UnrealEd/Public/Kismet2/BlueprintEditorUtils.h
--- a/Engine/Source/Editor/UnrealEd/Public/Kismet2/BlueprintEditorUtils.h
+++ b/Engine/Source/Editor/UnrealEd/Public/Kismet2/BlueprintEditorUtils.h
@@ -562,8 +562,9 @@ public:
 	 * @param			Graph			The graph that you're looking to call the function from (some functions hide different pins depending on the graph they're in)
 	 * @param			Function		The function to consider
 	 * @param [out]		HiddenPins		Set of pins that should be hidden
+	 * @param [out]		OutInternalPins	Subset of hidden pins that are marked for internal use only rather than marked as hidden (optional)
 	 */
-	static void GetHiddenPinsForFunction(UEdGraph const* Graph, UFunction const* Function, TSet<FString>& HiddenPins);
+	static void GetHiddenPinsForFunction(UEdGraph const* Graph, UFunction const* Function, TSet<FString>& HiddenPins, TSet<FString>* OutInternalPins = nullptr);
 
 	/** Makes sure that calls to parent functions are valid, and removes them if not */
 	static void ConformCallsToParentFunctions(UBlueprint* Blueprint);
diff --git a/Engine/Source/Runtime/Engine/Classes/GameFramework/Actor.h b/Engine/Source/Runtime/Engine/Classes/GameFramework/Actor.h
--- a/Engine/Source/Runtime/Engine/Classes/GameFramework/Actor.h
+++ b/Engine/Source/Runtime/Engine/Classes/GameFramework/Actor.h
@@ -889,7 +889,7 @@ public:
 	 * @param RelativeTransform				The relative transform between the new component and its attach parent (automatic only)
 	 * @param ComponentTemplateContext		Optional UBlueprintGeneratedClass reference to use to find the template in. If null (or not a BPGC), component is sought in this Actor's class
 	 */
-	UFUNCTION(BlueprintCallable, meta=(BlueprintInternalUseOnly = "true", DefaultToSelf="ComponentTemplateContext", HidePin="ComponentTemplateContext"))
+	UFUNCTION(BlueprintCallable, meta=(BlueprintInternalUseOnly = "true", DefaultToSelf="ComponentTemplateContext", InternalUseParam="ComponentTemplateContext"))
 	class UActorComponent* AddComponent(FName TemplateName, bool bManualAttachment, const FTransform& RelativeTransform, const UObject* ComponentTemplateContext);
 
 	/** DEPRECATED - Use Component::DestroyComponent */
