Can use unicode characters in Math Expression nodes. All unicode "letters" are now allowed for variable/input/output names. This does not extend to symbols and only the underscore symbol is still allowed.

Added FText::IsLetter to check a TCHAR to see if it is a letter according to ICU.

diff --git a/Engine/Source/Editor/BlueprintGraph/Private/BasicTokenParser.cpp b/Engine/Source/Editor/BlueprintGraph/Private/BasicTokenParser.cpp
--- a/Engine/Source/Editor/BlueprintGraph/Private/BasicTokenParser.cpp
+++ b/Engine/Source/Editor/BlueprintGraph/Private/BasicTokenParser.cpp
@@ -274,7 +274,7 @@ bool FBasicTokenParser::GetToken(FBasicToken& Token, bool bNoConsts/* = false*/)
 		Token.SetGuid(Token.Identifier);
 		return IsValid();
 	}
-	else if( (c>='A' && c<='Z') || (c>='a' && c<='z') || (c=='_') )
+	else if( FText::IsLetter(c)|| (c=='_') )
 	{
 		// Alphanumeric token.
 		int32 Length=0;
@@ -292,7 +292,7 @@ bool FBasicTokenParser::GetToken(FBasicToken& Token, bool bNoConsts/* = false*/)
 				break;
 			}
 			c = GetChar();
-		} while( ((c>='A')&&(c<='Z')) || ((c>='a')&&(c<='z')) || ((c>='0')&&(c<='9')) || (c=='_') );
+		} while( FText::IsLetter(c) || ((c>='0')&&(c<='9')) || (c=='_') );
 		UngetChar();
 		Token.Identifier[Length]=0;
 
diff --git a/Engine/Source/Runtime/Core/Private/Internationalization/ICUText.cpp b/Engine/Source/Runtime/Core/Private/Internationalization/ICUText.cpp
--- a/Engine/Source/Runtime/Core/Private/Internationalization/ICUText.cpp
+++ b/Engine/Source/Runtime/Core/Private/Internationalization/ICUText.cpp
@@ -13,6 +13,7 @@ PRAGMA_DISABLE_SHADOW_VARIABLE_WARNINGS
 #include <unicode/numfmt.h>
 #include <unicode/msgfmt.h>
 PRAGMA_ENABLE_SHADOW_VARIABLE_WARNINGS
+#include <unicode/uniset.h>
 
 #include "ICUUtilities.h"
 #include "ICUTextCharacterIterator.h"
@@ -222,4 +223,12 @@ bool FText::FSortPredicate::operator()(const FText& A, const FText& B) const
 	return Implementation->Compare(A, B);
 }
 
+bool FText::IsLetter( const TCHAR Char )
+{
+	icu::UnicodeString PatternString(TEXT("[\\p{L}]"));
+	UErrorCode ICUStatus = U_ZERO_ERROR;
+	icu::UnicodeSet Uniscode(PatternString, ICUStatus);
+	return Uniscode.contains(Char) != 0;
+}
+
 #endif
diff --git a/Engine/Source/Runtime/Core/Private/Internationalization/LegacyText.cpp b/Engine/Source/Runtime/Core/Private/Internationalization/LegacyText.cpp
--- a/Engine/Source/Runtime/Core/Private/Internationalization/LegacyText.cpp
+++ b/Engine/Source/Runtime/Core/Private/Internationalization/LegacyText.cpp
@@ -91,4 +91,9 @@ bool FText::FSortPredicate::operator()(const FText& A, const FText& B) const
 	return A.ToString() < B.ToString();
 }
 
+bool FText::IsLetter( const TCHAR Char )
+{
+	return (Char>='A' && Char<='Z') || (Char>='a' && Char<='z');
+}
+
 #endif
diff --git a/Engine/Source/Runtime/Core/Public/Internationalization/Text.h b/Engine/Source/Runtime/Core/Public/Internationalization/Text.h
--- a/Engine/Source/Runtime/Core/Public/Internationalization/Text.h
+++ b/Engine/Source/Runtime/Core/Public/Internationalization/Text.h
@@ -301,6 +301,9 @@ public:
 	 */
 	static bool IsWhitespace( const TCHAR Char );
 
+	/** Checks to see if the given character is considered a letter */
+	static bool IsLetter( const TCHAR Char );
+
 	static void GetFormatPatternParameters(const FText& Pattern, TArray<FString>& ParameterNames);
 
 	static FText Format(const FText& Pattern, const FFormatNamedArguments& Arguments);
