Call Function nodes will no longer change the function reference when duplicating the node or duplicating the entire Blueprint.

UK2Node_CallFunction::EnsureFunctionIsInBlueprint checks for the self pin being connected, if it is it will not change the function reference.

diff --git a/Engine/Source/Editor/BlueprintGraph/Private/K2Node_CallFunction.cpp b/Engine/Source/Editor/BlueprintGraph/Private/K2Node_CallFunction.cpp
--- a/Engine/Source/Editor/BlueprintGraph/Private/K2Node_CallFunction.cpp
+++ b/Engine/Source/Editor/BlueprintGraph/Private/K2Node_CallFunction.cpp
@@ -1615,25 +1615,30 @@ bool UK2Node_CallFunction::IsSelfPinCompatibleWithBlueprintContext(UEdGraphPin *
 
 void UK2Node_CallFunction::EnsureFunctionIsInBlueprint()
 {
-	// Ensure we're calling a function in a context related to our blueprint. If not, 
-	// reassigning the class and then calling ReconstructNodes will re-wire the pins correctly
-	if (UFunction* Function = GetTargetFunction())
+	// Do not mess with the function if there are pins connected to the target pin
+	UEdGraphPin* SelfPin = GetDefault<UEdGraphSchema_K2>()->FindSelfPin(*this, EGPD_Input);
+	if (SelfPin == nullptr || SelfPin->LinkedTo.Num() == 0)
 	{
-		UClass* FunctionOwnerClass = Function->GetOuterUClass();
-		UObject* FunctionGenerator = FunctionOwnerClass ? FunctionOwnerClass->ClassGeneratedBy : NULL;
-
-		// Never change the type if the function is an Interface function type, this only occurs when 
-		// the function is the interface function and the self pin will be PC_Interface type
-		if (!FunctionOwnerClass->IsChildOf(UInterface::StaticClass()))
+		// Ensure we're calling a function in a context related to our blueprint. If not, 
+		// reassigning the class and then calling ReconstructNodes will re-wire the pins correctly
+		if (UFunction* Function = GetTargetFunction())
 		{
-			// If function is generated from a blueprint object then dbl check self pin compatibility
-			UEdGraphPin* SelfPin = GetDefault<UEdGraphSchema_K2>()->FindSelfPin(*this, EGPD_Input);
-			if ((FunctionGenerator != NULL) && SelfPin)
+			UClass* FunctionOwnerClass = Function->GetOuterUClass();
+			UObject* FunctionGenerator = FunctionOwnerClass ? FunctionOwnerClass->ClassGeneratedBy : NULL;
+
+			// Never change the type if the function is an Interface function type, this only occurs when 
+			// the function is the interface function and the self pin will be PC_Interface type
+			if (!FunctionOwnerClass->IsChildOf(UInterface::StaticClass()))
 			{
-				UBlueprint* BlueprintObj = FBlueprintEditorUtils::FindBlueprintForNode(this);
-				if ((BlueprintObj != NULL) && !IsSelfPinCompatibleWithBlueprintContext(SelfPin, BlueprintObj))
+				// If function is generated from a blueprint object then dbl check self pin compatibility
+				UEdGraphPin* SelfPin = GetDefault<UEdGraphSchema_K2>()->FindSelfPin(*this, EGPD_Input);
+				if ((FunctionGenerator != NULL) && SelfPin)
 				{
-					FunctionReference.SetSelfMember(Function->GetFName());
+					UBlueprint* BlueprintObj = FBlueprintEditorUtils::FindBlueprintForNode(this);
+					if ((BlueprintObj != NULL) && !IsSelfPinCompatibleWithBlueprintContext(SelfPin, BlueprintObj))
+					{
+						FunctionReference.SetSelfMember(Function->GetFName());
+					}
 				}
 			}
 		}
