Fix shadow variable issue.

diff --git a/Engine/Source/Editor/BlueprintGraph/Private/K2Node_CallFunction.cpp b/Engine/Source/Editor/BlueprintGraph/Private/K2Node_CallFunction.cpp
--- a/Engine/Source/Editor/BlueprintGraph/Private/K2Node_CallFunction.cpp
+++ b/Engine/Source/Editor/BlueprintGraph/Private/K2Node_CallFunction.cpp
@@ -1617,7 +1617,7 @@ void UK2Node_CallFunction::EnsureFunctionIsInBlueprint()
 {
 	// Do not mess with the function if there are pins connected to the target pin
 	UEdGraphPin* SelfPin = GetDefault<UEdGraphSchema_K2>()->FindSelfPin(*this, EGPD_Input);
-	if (SelfPin == nullptr || SelfPin->LinkedTo.Num() == 0)
+	if (SelfPin && SelfPin->LinkedTo.Num() == 0)
 	{
 		// Ensure we're calling a function in a context related to our blueprint. If not, 
 		// reassigning the class and then calling ReconstructNodes will re-wire the pins correctly
@@ -1631,8 +1631,7 @@ void UK2Node_CallFunction::EnsureFunctionIsInBlueprint()
 			if (!FunctionOwnerClass->IsChildOf(UInterface::StaticClass()))
 			{
 				// If function is generated from a blueprint object then dbl check self pin compatibility
-				UEdGraphPin* SelfPin = GetDefault<UEdGraphSchema_K2>()->FindSelfPin(*this, EGPD_Input);
-				if ((FunctionGenerator != NULL) && SelfPin)
+				if (FunctionGenerator != NULL)
 				{
 					UBlueprint* BlueprintObj = FBlueprintEditorUtils::FindBlueprintForNode(this);
 					if ((BlueprintObj != NULL) && !IsSelfPinCompatibleWithBlueprintContext(SelfPin, BlueprintObj))
