Compiling a Blueprint with a validated Get node that is missing a property will no longer crash.

Added a check during ExpandNode to verify that the property exists to stop the ExpandNode from occuring. Normal error detection will report the missing property.

diff --git a/Engine/Source/Editor/BlueprintGraph/Private/K2Node_VariableGet.cpp b/Engine/Source/Editor/BlueprintGraph/Private/K2Node_VariableGet.cpp
--- a/Engine/Source/Editor/BlueprintGraph/Private/K2Node_VariableGet.cpp
+++ b/Engine/Source/Editor/BlueprintGraph/Private/K2Node_VariableGet.cpp
@@ -383,7 +383,8 @@ void UK2Node_VariableGet::ExpandNode(class FKismetCompilerContext& CompilerConte
 {
 	Super::ExpandNode(CompilerContext, SourceGraph);
 
-	if (!bIsPureGet)
+	// Do not attempt to expand the node when not a pure get nor when there is no property. Normal compilation error detection will detect the missing property.
+	if (!bIsPureGet && GetPropertyForVariable() != nullptr)
 	{
 		const UEdGraphSchema_K2* Schema = CompilerContext.GetSchema();
 		UEdGraphPin* ValuePin = GetValuePin();
