Object Reference, Class Reference, Asset ID, Asset Class ID pin types have been spun off into a sub-menu on the item to reduce duplication of these items in the listview. Selecting the item in the main tree will auto-select object reference as the pin type for the specified Blueprint type.

Spun out management of sub-menus in SMultiBox to the new class SMenuOwner for reusability.

Fixed potential issue with table views that would prevent them from stopping to be marked as hovered.

Added SSubMenuHandler to provide duplicate functionality to SMenuEntryBlock including deferred menu opening and close, and allowing the user's mouse to cross over other sibling items in order to get to the currently opened sub-menu.

diff --git a/Engine/Source/Editor/BlueprintGraph/Classes/EdGraphSchema_K2.h b/Engine/Source/Editor/BlueprintGraph/Classes/EdGraphSchema_K2.h
--- a/Engine/Source/Editor/BlueprintGraph/Classes/EdGraphSchema_K2.h
+++ b/Engine/Source/Editor/BlueprintGraph/Classes/EdGraphSchema_K2.h
@@ -194,6 +194,16 @@ struct FBlueprintCallableFunctionRedirect
 	FString ClassParamName;
 };
 
+enum class EObjectReferenceType : uint8
+{
+	NotAnObject		= 0x00,
+	ObjectReference = 0x01,
+	ClassReference	= 0x02,
+	AssetID			= 0x04,
+	ClassAssetID	= 0x08,
+	AllTypes		= 0x0f,
+};
+
 struct FTypesDatabase;
 
 UCLASS(config=Editor)
@@ -273,34 +283,29 @@ class BLUEPRINTGRAPH_API UEdGraphSchema_K2 : public UEdGraphSchema
 	// ID for checking dirty status of node titles against, increases every compile
 	static int32 CurrentCacheRefreshID;
 
+	// Pin Selector category for all object types
+	static const FString AllObjectTypes;
+
 	UPROPERTY(globalconfig)
 	TArray<FBlueprintCallableFunctionRedirect> EditoronlyBPFunctionRedirects;
 
 public:
+
 	//////////////////////////////////////////////////////////////////////////
 	// FPinTypeInfo
 	/** Class used for creating type tree selection info, which aggregates the various PC_* and PinSubtypes in the schema into a heirarchy */
 	class BLUEPRINTGRAPH_API FPinTypeTreeInfo
 	{
-	private:
-		FPinTypeTreeInfo()
-			: bReadOnly(false)
-		{
-		}
-
-		void Init(const FText& FriendlyCategoryName, const FString& CategoryName, const UEdGraphSchema_K2* Schema, const FText& InTooltip, bool bInReadOnly, FTypesDatabase* TypesDatabase);
-	
 	private:
 		/** The pin type corresponding to the schema type */
 		FEdGraphPinType PinType;
+		uint8 PossibleObjectReferenceTypes;
 
 		/** Asset Reference, used when PinType.PinSubCategoryObject is not loaded yet */
 		FStringAssetReference SubCategoryObjectAssetReference;
 
 		FText CachedDescription;
 
-		FText GenerateDescription();
-
 	public:
 		/** The children of this pin type */
 		TArray< TSharedPtr<FPinTypeTreeInfo> > Children;
@@ -322,9 +327,8 @@ public:
 		}
 
 		FPinTypeTreeInfo(const FText& InFriendlyName, const FString& CategoryName, const UEdGraphSchema_K2* Schema, const FText& InTooltip, bool bInReadOnly = false, FTypesDatabase* TypesDatabase = nullptr);
-		FPinTypeTreeInfo(const FString& CategoryName, UObject* SubCategoryObject, const FText& InTooltip, bool bInReadOnly = false);
-		FPinTypeTreeInfo(const FString& CategoryName, const FStringAssetReference& SubCategoryObject, const FText& InTooltip, bool bInReadOnly = false);
-		FPinTypeTreeInfo(const FText& InFriendlyName, const FString& CategoryName, const FStringAssetReference& SubCategoryObject, const FText& InTooltip, bool bInReadOnly = false);
+		FPinTypeTreeInfo(const FString& CategoryName, UObject* SubCategoryObject, const FText& InTooltip, bool bInReadOnly = false, uint8 InPossibleObjectReferenceTypes = 0);
+		FPinTypeTreeInfo(const FText& InFriendlyName, const FString& CategoryName, const FStringAssetReference& SubCategoryObject, const FText& InTooltip, bool bInReadOnly = false, uint8 InPossibleObjectReferenceTypes = 0);
 
 		FPinTypeTreeInfo(TSharedPtr<FPinTypeTreeInfo> InInfo)
 		{
@@ -334,6 +338,7 @@ public:
 			Tooltip = InInfo->Tooltip;
 			SubCategoryObjectAssetReference = InInfo->SubCategoryObjectAssetReference;
 			CachedDescription = InInfo->CachedDescription;
+			PossibleObjectReferenceTypes = InInfo->PossibleObjectReferenceTypes;
 		}
 		
 		/** Returns a succinct menu description of this type */
@@ -353,6 +358,22 @@ public:
 			}
 			return Tooltip;
 		}
+
+		uint8 GetPossibleObjectReferenceTypes() const
+		{
+			return PossibleObjectReferenceTypes;
+		}
+
+	private:
+
+		FPinTypeTreeInfo()
+			: PossibleObjectReferenceTypes(0)
+			, bReadOnly(false)
+		{}
+
+		void Init(const FText& FriendlyCategoryName, const FString& CategoryName, const UEdGraphSchema_K2* Schema, const FText& InTooltip, bool bInReadOnly, FTypesDatabase* TypesDatabase);
+
+		FText GenerateDescription();
 	};
 
 public:
@@ -802,7 +823,7 @@ public:
 	 * @param	bAllowExec		Whether or not to add the exec type to the type tree
 	 * @param	bAllowWildcard	Whether or not to add the wildcard type to the type tree
 	 */
-	void GetVariableTypeTree( TArray< TSharedPtr<FPinTypeTreeInfo> >& TypeTree, bool bAllowExec, bool bAllowWildcard ) const;
+	void GetVariableTypeTree(TArray< TSharedPtr<FPinTypeTreeInfo> >& TypeTree, bool bAllowExec, bool bAllowWildcard) const;
 
 	/**
 	 * Get the type tree for the index property types valid for this schema
@@ -973,6 +994,8 @@ public:
 
 private:
 
+	void GetVariableTypeTreeImpl(TArray< TSharedPtr<FPinTypeTreeInfo> >& TypeTree, bool bAllowExec, bool bAllowWildCard, bool bIndexTypesOnly) const;
+
 	/**
 	 * Returns true if the specified function has any out parameters
 	 * @param [in] Function	The function to check for out parameters
diff --git a/Engine/Source/Editor/BlueprintGraph/Private/EdGraphSchema_K2.cpp b/Engine/Source/Editor/BlueprintGraph/Private/EdGraphSchema_K2.cpp
--- a/Engine/Source/Editor/BlueprintGraph/Private/EdGraphSchema_K2.cpp
+++ b/Engine/Source/Editor/BlueprintGraph/Private/EdGraphSchema_K2.cpp
@@ -107,6 +107,7 @@ const FName FBlueprintMetadata::MD_ArrayDependentParam(TEXT("ArrayTypeDependentP
 #define LOCTEXT_NAMESPACE "KismetSchema"
 
 UEdGraphSchema_K2::FPinTypeTreeInfo::FPinTypeTreeInfo(const FText& InFriendlyName, const FString& CategoryName, const UEdGraphSchema_K2* Schema, const FText& InTooltip, bool bInReadOnly/*=false*/, FTypesDatabase* TypesDatabase /*=nullptr*/)
+	: PossibleObjectReferenceTypes(0)
 {
 	Init(InFriendlyName, CategoryName, Schema, InTooltip, bInReadOnly, TypesDatabase);
 }
@@ -116,13 +117,16 @@ struct FUnloadedAssetData
 	FStringAssetReference StringAssetReference;
 	FText AssetFriendlyName;
 	FText Tooltip;
+	uint8 PossibleObjectReferenceTypes;
 
 	FUnloadedAssetData()
+		: PossibleObjectReferenceTypes(0)
 	{}
 
-	FUnloadedAssetData(const FAssetData& InAsset)
+	FUnloadedAssetData(const FAssetData& InAsset, uint8 InPossibleObjectReferenceTypes = 0)
 		: StringAssetReference(InAsset.ToStringReference())
 		, AssetFriendlyName(FText::FromString(FName::NameToDisplayString(InAsset.AssetName.ToString(), false)))
+		, PossibleObjectReferenceTypes(InPossibleObjectReferenceTypes)
 	{
 		const FString* TooltipPtr = InAsset.TagsAndValues.Find(TEXT("Tooltip"));
 		Tooltip = FText::FromString((TooltipPtr && !TooltipPtr->IsEmpty()) ? *TooltipPtr : InAsset.ObjectPath.ToString());
@@ -133,10 +137,15 @@ struct FLoadedAssetData
 {
 	FText Tooltip;
 	UObject* Object;
+	uint8 PossibleObjectReferenceTypes;
 
-	FLoadedAssetData() : Object(nullptr) {}
-	FLoadedAssetData(UObject* InObject) 
+	FLoadedAssetData() 
+		: Object(nullptr)
+		, PossibleObjectReferenceTypes(0) {}
+
+	FLoadedAssetData(UObject* InObject, uint8 InPossibleObjectReferenceTypes = 0)
 		: Object(InObject) 
+		, PossibleObjectReferenceTypes(InPossibleObjectReferenceTypes)
 	{
 		UStruct* Struct = Cast<UStruct>(Object);
 		Tooltip = Struct ? Struct->GetToolTipText() : FText::GetEmpty();
@@ -165,166 +174,76 @@ private:
 		}
 	};
 
-	/** Helper function to add an unloaded asset to the children, does no validation on passed data */
-	static void AddUnloadedAsset(const FAssetData& InAsset, const FString& InCategoryName, TArray<FPinTypeTreeInfoPtr>& OutChildren)
-	{
-		const FString* TooltipPtr = InAsset.TagsAndValues.Find(TEXT("Tooltip"));
-		const FString Tooltip = (TooltipPtr && !TooltipPtr->IsEmpty()) ? *TooltipPtr : InAsset.ObjectPath.ToString();
-
-		FPinTypeTreeInfoPtr TypeTreeInfo = MakeShareable(new UEdGraphSchema_K2::FPinTypeTreeInfo(InCategoryName, InAsset.ToStringReference(), FText::FromString(Tooltip)));
-		TypeTreeInfo->FriendlyName = FText::FromString(FName::NameToDisplayString(InAsset.AssetName.ToString(), false));
-		OutChildren.Add(TypeTreeInfo);
-	}
-
-	/**
-	 * Gets a list of variable subtypes that are valid for the specified type
-	 *
-	 * @param	Type			The type to grab subtypes for
-	 * @param	SubtypesList	(out) Upon return, this will be a list of valid subtype objects for the specified type
-	 */
-	static void GetVariableSubtypes(const FString& Type, TArray<UObject*>& OutSubtypesList)
+public:
+	static void FillLoadedTypesDatabase(FTypesDatabase& TypesDatabase, bool bIndexTypesOnly)
 	{
-		OutSubtypesList.Empty();
+		// Loaded types
+		TypesDatabase.LoadedTypesMap.Reset();
 
-		if (Type == UEdGraphSchema_K2::PC_Struct)
-		{
-			// Find script structs marked with "BlueprintType=true" in their metadata, and add to the list
-			for (TObjectIterator<UScriptStruct> StructIt; StructIt; ++StructIt)
-			{
-				UScriptStruct* ScriptStruct = *StructIt;
-				if (UEdGraphSchema_K2::IsAllowableBlueprintVariableType(ScriptStruct))
-				{
-					OutSubtypesList.Add(ScriptStruct);
-				}
-			}
-		}
-		else if (Type == UEdGraphSchema_K2::PC_Class || Type == UEdGraphSchema_K2::PC_AssetClass)
-		{
-			// Generate a list of all potential objects which have "BlueprintType=true" in their metadata
-			for (TObjectIterator<UClass> ClassIt; ClassIt; ++ClassIt)
-			{
-				UClass* CurrentClass = *ClassIt;
-				if (UEdGraphSchema_K2::IsAllowableBlueprintVariableType(CurrentClass) && !CurrentClass->HasAnyClassFlags(CLASS_Deprecated))
-				{
-					OutSubtypesList.Add(CurrentClass);
-				}
-			}
-		}
-		else if (Type == UEdGraphSchema_K2::PC_Object || Type == UEdGraphSchema_K2::PC_Asset)
-		{
-			// Generate a list of all potential objects which have "BlueprintType=true" in their metadata (that aren't interfaces)
-			for (TObjectIterator<UClass> ClassIt; ClassIt; ++ClassIt)
-			{
-				UClass* CurrentClass = *ClassIt;
-				if (!CurrentClass->IsChildOf(UInterface::StaticClass()) && UEdGraphSchema_K2::IsAllowableBlueprintVariableType(CurrentClass) && !CurrentClass->HasAnyClassFlags(CLASS_Deprecated))
-				{
-					OutSubtypesList.Add(CurrentClass);
-				}
-			}
-		}
-		else if (Type == UEdGraphSchema_K2::PC_Interface)
-		{
-			// Generate a list of all potential objects which have "BlueprintType=true" in their metadata (only ones that are interfaces)
-			for (TObjectIterator<UClass> ClassIt; ClassIt; ++ClassIt)
-			{
-				UClass* CurrentClass = *ClassIt;
-				if (CurrentClass->IsChildOf(UInterface::StaticClass()) && UEdGraphSchema_K2::IsAllowableBlueprintVariableType(CurrentClass))
-				{
-					OutSubtypesList.Add(CurrentClass);
-				}
-			}
-		}
-		else if (Type == UEdGraphSchema_K2::PC_Enum)
+		//(Type == UEdGraphSchema_K2::PC_Enum)
 		{
+			FTypesDatabase::FLoadedTypesList LoadedTypesList = MakeShareable(new TArray<FLoadedAssetData>());
 			// Generate a list of all potential enums which have "BlueprintType=true" in their metadata
 			for (TObjectIterator<UEnum> EnumIt; EnumIt; ++EnumIt)
 			{
 				UEnum* CurrentEnum = *EnumIt;
 				if (UEdGraphSchema_K2::IsAllowableBlueprintVariableType(CurrentEnum))
 				{
-					OutSubtypesList.Add(CurrentEnum);
+					LoadedTypesList->Add(FLoadedAssetData(CurrentEnum));
 				}
 			}
+			TypesDatabase.LoadedTypesMap.Add(UEdGraphSchema_K2::PC_Enum, LoadedTypesList);
 		}
-	}
 
-public:
-	static void FillLoadedTypesDatabase(FTypesDatabase& TypesDatabase)
-	{
-		// Loaded types
-		TypesDatabase.LoadedTypesMap.Reset();
-
-		//(Type == UEdGraphSchema_K2::PC_Struct)
+		if (!bIndexTypesOnly)
 		{
-			FTypesDatabase::FLoadedTypesList LoadedTypesList = MakeShareable(new TArray<FLoadedAssetData>());
-			// Find script structs marked with "BlueprintType=true" in their metadata, and add to the list
-			for (TObjectIterator<UScriptStruct> StructIt; StructIt; ++StructIt)
+			//(Type == UEdGraphSchema_K2::PC_Struct)
 			{
-				UScriptStruct* ScriptStruct = *StructIt;
-				if (UEdGraphSchema_K2::IsAllowableBlueprintVariableType(ScriptStruct))
+				FTypesDatabase::FLoadedTypesList LoadedTypesList = MakeShareable(new TArray<FLoadedAssetData>());
+				// Find script structs marked with "BlueprintType=true" in their metadata, and add to the list
+				for (TObjectIterator<UScriptStruct> StructIt; StructIt; ++StructIt)
 				{
-					LoadedTypesList->Add(FLoadedAssetData(ScriptStruct));
+					UScriptStruct* ScriptStruct = *StructIt;
+					if (UEdGraphSchema_K2::IsAllowableBlueprintVariableType(ScriptStruct))
+					{
+						LoadedTypesList->Add(FLoadedAssetData(ScriptStruct));
+					}
 				}
+				TypesDatabase.LoadedTypesMap.Add(UEdGraphSchema_K2::PC_Struct, LoadedTypesList);
 			}
-			TypesDatabase.LoadedTypesMap.Add(UEdGraphSchema_K2::PC_Struct, LoadedTypesList);
-		}
 
-		//(Type == UEdGraphSchema_K2::PC_Class || Type == UEdGraphSchema_K2::PC_AssetClass) UEdGraphSchema_K2::PC_Interface
-		{
-			FTypesDatabase::FLoadedTypesList LoadedTypesList = MakeShareable(new TArray<FLoadedAssetData>());
-			FTypesDatabase::FLoadedTypesList InterfaceLoadedTypesList = MakeShareable(new TArray<FLoadedAssetData>());
-			// Generate a list of all potential objects which have "BlueprintType=true" in their metadata
-			for (TObjectIterator<UClass> ClassIt; ClassIt; ++ClassIt)
+			//(Type == UEdGraphSchema_K2::PC_Class || Type == UEdGraphSchema_K2::PC_AssetClass) UEdGraphSchema_K2::PC_Interface)
+			//(Type == UEdGraphSchema_K2::PC_Object || Type == UEdGraphSchema_K2::PC_Asset)
 			{
-				UClass* CurrentClass = *ClassIt;
-				if (UEdGraphSchema_K2::IsAllowableBlueprintVariableType(CurrentClass) && !CurrentClass->HasAnyClassFlags(CLASS_Deprecated))
-				{
-					LoadedTypesList->Add(FLoadedAssetData(CurrentClass));
-				}
-				
-				if (CurrentClass->IsChildOf(UInterface::StaticClass()) && UEdGraphSchema_K2::IsAllowableBlueprintVariableType(CurrentClass))
-				{
-					InterfaceLoadedTypesList->Add(FLoadedAssetData(CurrentClass));
-				}
-			}
-			TypesDatabase.LoadedTypesMap.Add(UEdGraphSchema_K2::PC_Class, LoadedTypesList);
-			TypesDatabase.LoadedTypesMap.Add(UEdGraphSchema_K2::PC_AssetClass, LoadedTypesList);
-			TypesDatabase.LoadedTypesMap.Add(UEdGraphSchema_K2::PC_Interface, InterfaceLoadedTypesList);
-		}
+				FTypesDatabase::FLoadedTypesList InterfaceLoadedTypesList = MakeShareable(new TArray<FLoadedAssetData>());
+				FTypesDatabase::FLoadedTypesList AllObjectLoadedTypesList = MakeShareable(new TArray<FLoadedAssetData>());
 
-		//(Type == UEdGraphSchema_K2::PC_Object || Type == UEdGraphSchema_K2::PC_Asset)
-		{
-			FTypesDatabase::FLoadedTypesList LoadedTypesList = MakeShareable(new TArray<FLoadedAssetData>());
-			// Generate a list of all potential objects which have "BlueprintType=true" in their metadata (that aren't interfaces)
-			for (TObjectIterator<UClass> ClassIt; ClassIt; ++ClassIt)
-			{
-				UClass* CurrentClass = *ClassIt;
-				if (!CurrentClass->IsChildOf(UInterface::StaticClass()) && UEdGraphSchema_K2::IsAllowableBlueprintVariableType(CurrentClass) && !CurrentClass->HasAnyClassFlags(CLASS_Deprecated))
+				// Generate a list of all potential objects which have "BlueprintType=true" in their metadata
+				for (TObjectIterator<UClass> ClassIt; ClassIt; ++ClassIt)
 				{
-					LoadedTypesList->Add(FLoadedAssetData(CurrentClass));
-				}
-			}
-			TypesDatabase.LoadedTypesMap.Add(UEdGraphSchema_K2::PC_Object, LoadedTypesList);
-			TypesDatabase.LoadedTypesMap.Add(UEdGraphSchema_K2::PC_Asset, LoadedTypesList);
-		}
-		
-		//(Type == UEdGraphSchema_K2::PC_Enum)
-		{
-			FTypesDatabase::FLoadedTypesList LoadedTypesList = MakeShareable(new TArray<FLoadedAssetData>());
-			// Generate a list of all potential enums which have "BlueprintType=true" in their metadata
-			for (TObjectIterator<UEnum> EnumIt; EnumIt; ++EnumIt)
-			{
-				UEnum* CurrentEnum = *EnumIt;
-				if (UEdGraphSchema_K2::IsAllowableBlueprintVariableType(CurrentEnum))
-				{
-					LoadedTypesList->Add(FLoadedAssetData(CurrentEnum));
+					UClass* CurrentClass = *ClassIt;
+					const bool bIsInterface = CurrentClass->IsChildOf(UInterface::StaticClass());
+					const bool bIsBlueprintType = UEdGraphSchema_K2::IsAllowableBlueprintVariableType(CurrentClass);
+					const bool bIsDeprecated = CurrentClass->HasAnyClassFlags(CLASS_Deprecated);
+					if (bIsBlueprintType && !bIsDeprecated)
+					{
+						if (bIsInterface)
+						{
+							InterfaceLoadedTypesList->Add(FLoadedAssetData(CurrentClass));
+						}
+						else
+						{
+							AllObjectLoadedTypesList->Add(FLoadedAssetData(CurrentClass, static_cast<uint8>(EObjectReferenceType::AllTypes)));
+						}
+					}
 				}
+				TypesDatabase.LoadedTypesMap.Add(UEdGraphSchema_K2::AllObjectTypes, AllObjectLoadedTypesList);
+				TypesDatabase.LoadedTypesMap.Add(UEdGraphSchema_K2::PC_Interface, InterfaceLoadedTypesList);
 			}
-			TypesDatabase.LoadedTypesMap.Add(UEdGraphSchema_K2::PC_Enum, LoadedTypesList);
 		}
 	}
 
-	static void FillUnLoadedTypesDatabase(FTypesDatabase& TypesDatabase)
+	static void FillUnLoadedTypesDatabase(FTypesDatabase& TypesDatabase, bool bIndexTypesOnly)
 	{
 		// Loaded types
 		TypesDatabase.UnLoadedTypesMap.Reset();
@@ -333,7 +252,7 @@ public:
 
 		{
 			TArray<FAssetData> AssetData;
-			AssetRegistryModule.Get().GetAssetsByClass(UUserDefinedStruct::StaticClass()->GetFName(), AssetData);
+			AssetRegistryModule.Get().GetAssetsByClass(UUserDefinedEnum::StaticClass()->GetFName(), AssetData);
 
 			FTypesDatabase::FUnLoadedTypesList UnLoadedTypesList = MakeShareable(new TArray<FUnloadedAssetData>());
 			for (int32 AssetIndex = 0; AssetIndex < AssetData.Num(); ++AssetIndex)
@@ -345,78 +264,85 @@ public:
 				}
 			}
 
-			TypesDatabase.UnLoadedTypesMap.Add(UEdGraphSchema_K2::PC_Struct, UnLoadedTypesList);
+			TypesDatabase.UnLoadedTypesMap.Add(UEdGraphSchema_K2::PC_Enum, UnLoadedTypesList);
 		}
 
+		if (!bIndexTypesOnly)
 		{
-			TArray<FAssetData> AssetData;
-			AssetRegistryModule.Get().GetAssetsByClass(UUserDefinedEnum::StaticClass()->GetFName(), AssetData);
-
-			FTypesDatabase::FUnLoadedTypesList UnLoadedTypesList = MakeShareable(new TArray<FUnloadedAssetData>());
-			for (int32 AssetIndex = 0; AssetIndex < AssetData.Num(); ++AssetIndex)
 			{
-				const FAssetData& Asset = AssetData[AssetIndex];
-				if (Asset.IsValid() && !Asset.IsAssetLoaded())
+				TArray<FAssetData> AssetData;
+				AssetRegistryModule.Get().GetAssetsByClass(UUserDefinedStruct::StaticClass()->GetFName(), AssetData);
+
+				FTypesDatabase::FUnLoadedTypesList UnLoadedTypesList = MakeShareable(new TArray<FUnloadedAssetData>());
+				for (int32 AssetIndex = 0; AssetIndex < AssetData.Num(); ++AssetIndex)
 				{
-					UnLoadedTypesList->Add(FUnloadedAssetData(Asset));
+					const FAssetData& Asset = AssetData[AssetIndex];
+					if (Asset.IsValid() && !Asset.IsAssetLoaded())
+					{
+						UnLoadedTypesList->Add(FUnloadedAssetData(Asset));
+					}
 				}
-			}
-
-			TypesDatabase.UnLoadedTypesMap.Add(UEdGraphSchema_K2::PC_Enum, UnLoadedTypesList);
-		}
 
-		//else if (Schema->PC_Object == CategoryName || Schema->PC_Class == CategoryName || Schema->PC_Interface == CategoryName || Schema->PC_Asset == CategoryName || Schema->PC_AssetClass == CategoryName)
-		{
-			TArray<FAssetData> AssetData;
-			AssetRegistryModule.Get().GetAssetsByClass(UBlueprint::StaticClass()->GetFName(), AssetData);
+				TypesDatabase.UnLoadedTypesMap.Add(UEdGraphSchema_K2::PC_Struct, UnLoadedTypesList);
+			}
 
-			const FString BPInterfaceTypeAllowed(TEXT("BPTYPE_Interface")); 
-			const FString BPNormalTypeAllowed(TEXT("BPTYPE_Normal"));
+			//else if (Schema->PC_Object == CategoryName || Schema->PC_Class == CategoryName || Schema->PC_Interface == CategoryName || Schema->PC_Asset == CategoryName || Schema->PC_AssetClass == CategoryName)
+			{
+				TArray<FAssetData> AssetData;
+				AssetRegistryModule.Get().GetAssetsByClass(UBlueprint::StaticClass()->GetFName(), AssetData);
 
-			FTypesDatabase::FUnLoadedTypesList UnLoadedInterfacesList = MakeShareable(new TArray<FUnloadedAssetData>());
-			FTypesDatabase::FUnLoadedTypesList UnLoadedClassesList = MakeShareable(new TArray<FUnloadedAssetData>());
+				const FString BPInterfaceTypeAllowed(TEXT("BPTYPE_Interface"));
+				const FString BPNormalTypeAllowed(TEXT("BPTYPE_Normal"));
 
-			for (int32 AssetIndex = 0; AssetIndex < AssetData.Num(); ++AssetIndex)
-			{
-				const FAssetData& Asset = AssetData[AssetIndex];
+				FTypesDatabase::FUnLoadedTypesList UnLoadedInterfacesList = MakeShareable(new TArray<FUnloadedAssetData>());
+				FTypesDatabase::FUnLoadedTypesList UnLoadedClassesList = MakeShareable(new TArray<FUnloadedAssetData>());
 
-				if (Asset.IsValid() && !Asset.IsAssetLoaded())
+				for (int32 AssetIndex = 0; AssetIndex < AssetData.Num(); ++AssetIndex)
 				{
-					const FString* BlueprintTypeStr = Asset.TagsAndValues.Find("BlueprintType");
-					const bool bNormalBP = BlueprintTypeStr && (*BlueprintTypeStr == BPNormalTypeAllowed);
-					const bool bInterfaceBP = BlueprintTypeStr && (*BlueprintTypeStr == BPInterfaceTypeAllowed);
+					const FAssetData& Asset = AssetData[AssetIndex];
 
-					if (bNormalBP || bInterfaceBP)
+					if (Asset.IsValid() && !Asset.IsAssetLoaded())
 					{
-						uint32 ClassFlags = 0;
-						const FString* ClassFlagsStr = Asset.TagsAndValues.Find("ClassFlags");
-						if (ClassFlagsStr)
-						{
-							ClassFlags = FCString::Atoi(**ClassFlagsStr);
-						}
+						const FString* BlueprintTypeStr = Asset.TagsAndValues.Find("BlueprintType");
+						const bool bNormalBP = BlueprintTypeStr && (*BlueprintTypeStr == BPNormalTypeAllowed);
+						const bool bInterfaceBP = BlueprintTypeStr && (*BlueprintTypeStr == BPInterfaceTypeAllowed);
 
-						if (!(ClassFlags & CLASS_Deprecated))
+						if (bNormalBP || bInterfaceBP)
 						{
-							if (bNormalBP)
+							uint32 ClassFlags = 0;
+							const FString* ClassFlagsStr = Asset.TagsAndValues.Find("ClassFlags");
+							if (ClassFlagsStr)
 							{
-								UnLoadedClassesList->Add(FUnloadedAssetData(Asset));
+								ClassFlags = FCString::Atoi(**ClassFlagsStr);
 							}
-							else if (bInterfaceBP)
+
+							if (!(ClassFlags & CLASS_Deprecated))
 							{
-								UnLoadedInterfacesList->Add(FUnloadedAssetData(Asset));
+								if (bNormalBP)
+								{
+									UnLoadedClassesList->Add(FUnloadedAssetData(Asset, static_cast<uint8>(EObjectReferenceType::AllTypes)));
+								}
+								else if (bInterfaceBP)
+								{
+									UnLoadedInterfacesList->Add(FUnloadedAssetData(Asset));
+								}
 							}
 						}
 					}
 				}
+				TypesDatabase.UnLoadedTypesMap.Add(UEdGraphSchema_K2::PC_Interface, UnLoadedInterfacesList);
+				TypesDatabase.UnLoadedTypesMap.Add(UEdGraphSchema_K2::AllObjectTypes, UnLoadedClassesList);
 			}
-			TypesDatabase.UnLoadedTypesMap.Add(UEdGraphSchema_K2::PC_Interface, UnLoadedInterfacesList);
-			TypesDatabase.UnLoadedTypesMap.Add(UEdGraphSchema_K2::PC_Object, UnLoadedClassesList);
-			TypesDatabase.UnLoadedTypesMap.Add(UEdGraphSchema_K2::PC_Class, UnLoadedClassesList);
-			TypesDatabase.UnLoadedTypesMap.Add(UEdGraphSchema_K2::PC_Asset, UnLoadedClassesList);
-			TypesDatabase.UnLoadedTypesMap.Add(UEdGraphSchema_K2::PC_AssetClass, UnLoadedClassesList);
 		}
 	}
 
+	/**
+	 * Gathers all valid sub-types (loaded and unloaded) of a passed category and sorts them alphabetically
+	 * @param FriendlyName		Friendly name to be used for the tooltip if there is no available data
+	 * @param CategoryName		Category (type) to find sub-types of
+	 * @param TypesDatabase		Types database
+	 * @param OutChildren		All the gathered children
+	 */
 	static void Gather(const FText& FriendlyName, const FString& CategoryName, FTypesDatabase& TypesDatabase, TArray<FPinTypeTreeInfoPtr>& OutChildren)
 	{
 		FEdGraphPinType LoadedPinSubtype;
@@ -429,7 +355,11 @@ public:
 		{
 			for (auto LoadedAssetData : *LoadedSubTypesPtr->Get())
 			{
-				OutChildren.Add(MakeShareable(new UEdGraphSchema_K2::FPinTypeTreeInfo(LoadedPinSubtype.PinCategory, LoadedAssetData.Object, LoadedAssetData.Tooltip.IsEmpty() ? FriendlyName : LoadedAssetData.Tooltip)));
+				OutChildren.Add(MakeShareable(new UEdGraphSchema_K2::FPinTypeTreeInfo(LoadedPinSubtype.PinCategory
+					, LoadedAssetData.Object
+					, LoadedAssetData.Tooltip.IsEmpty() ? FriendlyName : LoadedAssetData.Tooltip
+					, false
+					, LoadedAssetData.PossibleObjectReferenceTypes)));
 			}
 		}
 
@@ -438,7 +368,12 @@ public:
 		{
 			for (FUnloadedAssetData& It : *UnLoadedSubTypesPtr->Get())
 			{
-				FPinTypeTreeInfoPtr TypeTreeInfo = MakeShareable(new UEdGraphSchema_K2::FPinTypeTreeInfo(It.AssetFriendlyName, CategoryName, It.StringAssetReference, It.Tooltip));
+				FPinTypeTreeInfoPtr TypeTreeInfo = MakeShareable(new UEdGraphSchema_K2::FPinTypeTreeInfo(It.AssetFriendlyName
+					, CategoryName
+					, It.StringAssetReference
+					, It.Tooltip
+					, false
+					, It.PossibleObjectReferenceTypes));
 				OutChildren.Add(TypeTreeInfo);
 			}
 		}
@@ -446,101 +381,6 @@ public:
 		OutChildren.Sort(FCompareChildren());
 	}
 
-	/**
-	 * Gathers all valid sub-types (loaded and unloaded) of a passed category and sorts them alphabetically
-	 * @param FriendlyName		Friendly name to be used for the tooltip if there is no available data
-	 * @param CategoryName		Category (type) to find sub-types of
-	 * @param Schema			Schema to use
-	 * @param OutChildren		All the gathered children
-	 */
-	static void Gather(const FText& FriendlyName, const FString& CategoryName, const UEdGraphSchema_K2* Schema, TArray<FPinTypeTreeInfoPtr>& OutChildren)
-	{
-		// Gather any loaded subtype children first
-		TArray<UObject*> LoadedSubTypes;
-		GetVariableSubtypes(CategoryName, LoadedSubTypes);
-
-		FEdGraphPinType LoadedPinSubtype;
-		LoadedPinSubtype.PinCategory = (CategoryName == UEdGraphSchema_K2::PC_Enum ? UEdGraphSchema_K2::PC_Byte : CategoryName);
-		LoadedPinSubtype.PinSubCategory = TEXT("");
-		LoadedPinSubtype.PinSubCategoryObject = NULL;
-
-		// Add any loaded subtypes to the children list
-		for (auto it = LoadedSubTypes.CreateIterator(); it; ++it)
-		{
-			FText SubtypeTooltip;
-			UStruct* Struct = Cast<UStruct>(*it);
-			if(Struct != NULL)
-			{
-				SubtypeTooltip = (Struct ? Struct->GetToolTipText() : FriendlyName);
-			}
-			
-			OutChildren.Add( MakeShareable(new UEdGraphSchema_K2::FPinTypeTreeInfo(LoadedPinSubtype.PinCategory, *it, SubtypeTooltip)) );
-		}
-
-		UClass* AssetRegistrySearchClass = nullptr;
-		if(Schema->PC_Struct == CategoryName)
-		{
-			AssetRegistrySearchClass = UUserDefinedStruct::StaticClass();
-		}
-		else if(Schema->PC_Enum == CategoryName)
-		{
-			AssetRegistrySearchClass = UUserDefinedEnum::StaticClass();
-		}
-		else if (Schema->PC_Object == CategoryName || Schema->PC_Class == CategoryName || Schema->PC_Interface == CategoryName || Schema->PC_Asset == CategoryName || Schema->PC_AssetClass == CategoryName)
-		{
-			AssetRegistrySearchClass = UBlueprint::StaticClass();
-		}
-
-		// Gather all assets of the specified class
-		const FAssetRegistryModule& AssetRegistryModule = FModuleManager::LoadModuleChecked<FAssetRegistryModule>(TEXT("AssetRegistry"));
-		TArray<FAssetData> AssetData;
-		AssetRegistryModule.Get().GetAssetsByClass(AssetRegistrySearchClass->GetFName(), AssetData);
-
-		// Search through all the unloaded assets, filtered by certain rules based on struct/enum or object/class/interface
-		if (Schema->PC_Struct == CategoryName || Schema->PC_Enum == CategoryName)
-		{
-			for (int32 AssetIndex = 0; AssetIndex < AssetData.Num(); ++AssetIndex)
-			{
-				const FAssetData& Asset = AssetData[AssetIndex];
-				if (Asset.IsValid() && !Asset.IsAssetLoaded())
-				{
-					AddUnloadedAsset(Asset, CategoryName, OutChildren);
-				}
-			}
-		}
-		else if (Schema->PC_Object == CategoryName || Schema->PC_Class == CategoryName || Schema->PC_Interface == CategoryName || Schema->PC_Asset == CategoryName || Schema->PC_AssetClass == CategoryName)
-		{
-			const FString BPTypeAllowed = (Schema->PC_Interface == CategoryName)? TEXT("BPTYPE_Interface") : TEXT("BPTYPE_Normal");
-
-			for (int32 AssetIndex = 0; AssetIndex < AssetData.Num(); ++AssetIndex)
-			{
-				const FAssetData& Asset = AssetData[AssetIndex];
-
-				if (Asset.IsValid() && !Asset.IsAssetLoaded())
-				{
-					// Based on the category, only certain Blueprint types are available
-					const FString* BlueprintTypeStr = Asset.TagsAndValues.Find("BlueprintType");
-					if(BlueprintTypeStr && *BlueprintTypeStr == BPTypeAllowed)
-					{
-						uint32 ClassFlags = 0;
-						const FString* ClassFlagsStr = Asset.TagsAndValues.Find("ClassFlags");
-						if(ClassFlagsStr)
-						{
-							ClassFlags = FCString::Atoi(**ClassFlagsStr);
-						}
-
-						// Do not allow deprecated Blueprints to be added
-						if(!(ClassFlags & CLASS_Deprecated))
-						{
-							AddUnloadedAsset(Asset, CategoryName, OutChildren);
-						}
-					}
-				}
-			}
-		}
-		OutChildren.Sort(FCompareChildren());
-	}
-
 	/** Loads an asset based on the AssetReference through the asset registry */
 	static UObject* LoadAsset(const FStringAssetReference& AssetReference)
 	{
@@ -624,14 +464,11 @@ void UEdGraphSchema_K2::FPinTypeTreeInfo::Init(const FText& InFriendlyName, cons
 		{
 			FGatherTypesHelper::Gather(InFriendlyName, CategoryName, *TypesDatabase, Children);
 		}
-		else
-		{
-			FGatherTypesHelper::Gather(InFriendlyName, CategoryName, Schema, Children);
-		}
 	}
 }
 
-UEdGraphSchema_K2::FPinTypeTreeInfo::FPinTypeTreeInfo(const FString& CategoryName, UObject* SubCategoryObject, const FText& InTooltip, bool bInReadOnly/*=false*/)
+UEdGraphSchema_K2::FPinTypeTreeInfo::FPinTypeTreeInfo(const FString& CategoryName, UObject* SubCategoryObject, const FText& InTooltip, bool bInReadOnly/*=false*/, uint8 InPossibleObjectReferenceTypes)
+	: PossibleObjectReferenceTypes(InPossibleObjectReferenceTypes)
 {
 	check( !CategoryName.IsEmpty() );
 	check( SubCategoryObject );
@@ -644,22 +481,8 @@ UEdGraphSchema_K2::FPinTypeTreeInfo::FPinTypeTreeInfo(const FString& CategoryNam
 	CachedDescription = GenerateDescription();
 }
 
-UEdGraphSchema_K2::FPinTypeTreeInfo::FPinTypeTreeInfo(const FString& CategoryName, const FStringAssetReference& SubCategoryObject, const FText& InTooltip, bool bInReadOnly)
-{
-	check(!CategoryName.IsEmpty());
-	check(SubCategoryObject.IsValid());
-
-	Tooltip = InTooltip;
-	PinType.PinCategory = CategoryName;
-
-	SubCategoryObjectAssetReference = SubCategoryObject;
-	PinType.PinSubCategoryObject = SubCategoryObjectAssetReference.ResolveObject();
-
-	bReadOnly = bInReadOnly;
-	CachedDescription = GenerateDescription();
-}
-
-UEdGraphSchema_K2::FPinTypeTreeInfo::FPinTypeTreeInfo(const FText& InFriendlyName, const FString& CategoryName, const FStringAssetReference& SubCategoryObject, const FText& InTooltip, bool bInReadOnly)
+UEdGraphSchema_K2::FPinTypeTreeInfo::FPinTypeTreeInfo(const FText& InFriendlyName, const FString& CategoryName, const FStringAssetReference& SubCategoryObject, const FText& InTooltip, bool bInReadOnly, uint8 InPossibleObjectReferenceTypes)
+	: PossibleObjectReferenceTypes(InPossibleObjectReferenceTypes)
 {
 	FriendlyName = InFriendlyName;
 
@@ -765,6 +588,8 @@ const UScriptStruct* UEdGraphSchema_K2::ColorStruct = nullptr;
 bool UEdGraphSchema_K2::bGeneratingDocumentation = false;
 int32 UEdGraphSchema_K2::CurrentCacheRefreshID = 0;
 
+const FString UEdGraphSchema_K2::AllObjectTypes(TEXT("ObjectTypes"));
+
 UEdGraphSchema_K2::UEdGraphSchema_K2(const FObjectInitializer& ObjectInitializer)
 	: Super(ObjectInitializer)
 {
@@ -3486,6 +3311,7 @@ FText UEdGraphSchema_K2::GetCategoryText(const FString& Category, const bool bFo
 		CategoryDescriptions.Add(PC_Enum, LOCTEXT("EnumCategory","Enum"));
 		CategoryDescriptions.Add(PC_Asset, LOCTEXT("AssetCategory", "Asset ID"));
 		CategoryDescriptions.Add(PC_AssetClass, LOCTEXT("AssetClassCategory", "Class Asset ID"));
+		CategoryDescriptions.Add(AllObjectTypes, LOCTEXT("AllObjectTypes", "Object Types"));
 	}
 
 	if (bForMenu)
@@ -3584,7 +3410,12 @@ FText UEdGraphSchema_K2::TypeToText(const FEdGraphPinType& Type)
 	return PropertyText;
 }
 
-void UEdGraphSchema_K2::GetVariableTypeTree( TArray< TSharedPtr<FPinTypeTreeInfo> >& TypeTree, bool bAllowExec, bool bAllowWildCard ) const
+void UEdGraphSchema_K2::GetVariableTypeTree(TArray< TSharedPtr<FPinTypeTreeInfo> >& TypeTree, bool bAllowExec, bool bAllowWildcard) const
+{
+	GetVariableTypeTreeImpl(TypeTree, bAllowExec, bAllowWildcard, false);
+}
+
+void UEdGraphSchema_K2::GetVariableTypeTreeImpl(TArray< TSharedPtr<FPinTypeTreeInfo> >& TypeTree, bool bAllowExec, bool bAllowWildCard, bool bIndexTypesOnly) const
 {
 #ifdef SCHEMA_K2_GETVARIABLETYPETREE_LOG_TIME
 	static_assert(false, "Macro redefinition.");
@@ -3596,13 +3427,13 @@ void UEdGraphSchema_K2::GetVariableTypeTree( TArray< TSharedPtr<FPinTypeTreeInfo
 #endif //SCHEMA_K2_GETVARIABLETYPETREE_LOG_TIME
 
 	FTypesDatabase TypesDatabase;
-	FGatherTypesHelper::FillLoadedTypesDatabase(TypesDatabase);
+	FGatherTypesHelper::FillLoadedTypesDatabase(TypesDatabase, bIndexTypesOnly);
 
 #if SCHEMA_K2_GETVARIABLETYPETREE_LOG_TIME
 	const auto DatabaseLoadedTime = FPlatformTime::Seconds();
 #endif //SCHEMA_K2_GETVARIABLETYPETREE_LOG_TIME
 
-	FGatherTypesHelper::FillUnLoadedTypesDatabase(TypesDatabase);
+	FGatherTypesHelper::FillUnLoadedTypesDatabase(TypesDatabase, bIndexTypesOnly);
 
 #if SCHEMA_K2_GETVARIABLETYPETREE_LOG_TIME
 	const auto DatabaseUnLoadedTime = FPlatformTime::Seconds();
@@ -3619,16 +3450,18 @@ void UEdGraphSchema_K2::GetVariableTypeTree( TArray< TSharedPtr<FPinTypeTreeInfo
 	TypeTree.Add( MakeShareable( new FPinTypeTreeInfo(GetCategoryText(PC_Boolean, true), PC_Boolean, this, LOCTEXT("BooleanType", "True or false value")) ) );
 	TypeTree.Add( MakeShareable( new FPinTypeTreeInfo(GetCategoryText(PC_Byte, true), PC_Byte, this, LOCTEXT("ByteType", "8 bit number")) ) );
 	TypeTree.Add( MakeShareable( new FPinTypeTreeInfo(GetCategoryText(PC_Int, true), PC_Int, this, LOCTEXT("IntegerType", "Integer number")) ) );
-	TypeTree.Add( MakeShareable( new FPinTypeTreeInfo(GetCategoryText(PC_Float, true), PC_Float, this, LOCTEXT("FloatType", "Floating point number")) ) );
-	TypeTree.Add( MakeShareable( new FPinTypeTreeInfo(GetCategoryText(PC_Name, true), PC_Name, this, LOCTEXT("NameType", "A text name")) ) );
-	TypeTree.Add( MakeShareable( new FPinTypeTreeInfo(GetCategoryText(PC_String, true), PC_String, this, LOCTEXT("StringType", "A text string")) ) );
-	TypeTree.Add( MakeShareable( new FPinTypeTreeInfo(GetCategoryText(PC_Text, true), PC_Text, this, LOCTEXT("TextType", "A localizable text string")) ) );
-
-	// Add in special first-class struct types
-	TypeTree.Add( MakeShareable( new FPinTypeTreeInfo(PC_Struct, VectorStruct, LOCTEXT("VectorType", "A 3D vector")) ) );
-	TypeTree.Add( MakeShareable( new FPinTypeTreeInfo(PC_Struct, RotatorStruct, LOCTEXT("RotatorType", "A 3D rotation")) ) );
-	TypeTree.Add( MakeShareable( new FPinTypeTreeInfo(PC_Struct, TransformStruct, LOCTEXT("TransformType", "A 3D transformation, including translation, rotation and 3D scale.")) ) );
-	
+	if (!bIndexTypesOnly)
+	{
+		TypeTree.Add(MakeShareable(new FPinTypeTreeInfo(GetCategoryText(PC_Float, true), PC_Float, this, LOCTEXT("FloatType", "Floating point number"))));
+		TypeTree.Add(MakeShareable(new FPinTypeTreeInfo(GetCategoryText(PC_Name, true), PC_Name, this, LOCTEXT("NameType", "A text name"))));
+		TypeTree.Add(MakeShareable(new FPinTypeTreeInfo(GetCategoryText(PC_String, true), PC_String, this, LOCTEXT("StringType", "A text string"))));
+		TypeTree.Add(MakeShareable(new FPinTypeTreeInfo(GetCategoryText(PC_Text, true), PC_Text, this, LOCTEXT("TextType", "A localizable text string"))));
+
+		// Add in special first-class struct types
+		TypeTree.Add(MakeShareable(new FPinTypeTreeInfo(PC_Struct, TBaseStructure<FVector>::Get(), LOCTEXT("VectorType", "A 3D vector"))));
+		TypeTree.Add(MakeShareable(new FPinTypeTreeInfo(PC_Struct, TBaseStructure<FRotator>::Get(), LOCTEXT("RotatorType", "A 3D rotation"))));
+		TypeTree.Add(MakeShareable(new FPinTypeTreeInfo(PC_Struct, TBaseStructure<FTransform>::Get(), LOCTEXT("TransformType", "A 3D transformation, including translation, rotation and 3D scale."))));
+	}
 	// Add wildcard type
 	if (bAllowWildCard)
 	{
@@ -3637,49 +3470,30 @@ void UEdGraphSchema_K2::GetVariableTypeTree( TArray< TSharedPtr<FPinTypeTreeInfo
 
 	FTypesDatabase* TypesDatabasePtr = &TypesDatabase;
 	// Add the types that have subtrees
-	TypeTree.Add( MakeShareable( new FPinTypeTreeInfo(GetCategoryText(PC_Struct, true), PC_Struct, this, LOCTEXT("StructType", "Struct (value) types."), true, TypesDatabasePtr) ) );
-	TypeTree.Add( MakeShareable( new FPinTypeTreeInfo(GetCategoryText(PC_Object, true), PC_Object, this, LOCTEXT("ObjectType", "Object pointer."), true, TypesDatabasePtr) ) );
-	TypeTree.Add( MakeShareable( new FPinTypeTreeInfo(GetCategoryText(PC_Asset, true), PC_Asset, this, LOCTEXT("AssetType", "Asset pointer."), true, TypesDatabasePtr) ) );
-	TypeTree.Add( MakeShareable( new FPinTypeTreeInfo(GetCategoryText(PC_AssetClass, true), PC_AssetClass, this, LOCTEXT("ClassAssetType", "Class asset."), true, TypesDatabasePtr) ) );
-	TypeTree.Add( MakeShareable( new FPinTypeTreeInfo(GetCategoryText(PC_Interface, true), PC_Interface, this, LOCTEXT("InterfaceType", "Interface pointer."), true, TypesDatabasePtr) ) );
-	TypeTree.Add( MakeShareable( new FPinTypeTreeInfo(GetCategoryText(PC_Class, true), PC_Class, this, LOCTEXT("ClassType", "Class pointers."), true, TypesDatabasePtr) ) );
+	if (!bIndexTypesOnly)
+	{
+		TypeTree.Add(MakeShareable(new FPinTypeTreeInfo(GetCategoryText(PC_Struct, true), PC_Struct, this, LOCTEXT("StructType", "Struct (value) types."), true, TypesDatabasePtr)));
+		TypeTree.Add(MakeShareable(new FPinTypeTreeInfo(GetCategoryText(PC_Interface, true), PC_Interface, this, LOCTEXT("InterfaceType", "Interface pointer."), true, TypesDatabasePtr)));
+
+		TypeTree.Add(MakeShareable(new FPinTypeTreeInfo(GetCategoryText(AllObjectTypes, true), AllObjectTypes, this, LOCTEXT("ObjectType", "Object pointer."), true, TypesDatabasePtr)));
+	}
 	TypeTree.Add( MakeShareable( new FPinTypeTreeInfo(GetCategoryText(PC_Enum, true), PC_Enum, this, LOCTEXT("EnumType", "Enumeration types."), true, TypesDatabasePtr) ) );
 
 #if SCHEMA_K2_GETVARIABLETYPETREE_LOG_TIME
 	const auto EndTime = FPlatformTime::Seconds();
 	UE_LOG(LogBlueprint, Log, TEXT("UEdGraphSchema_K2::GetVariableTypeTree times - LoadedTypesDatabase: %f UnLoadedTypesDatabase: %f FPinTypeTreeInfo: %f"), DatabaseLoadedTime - StartTime, DatabaseUnLoadedTime - DatabaseLoadedTime, EndTime - DatabaseUnLoadedTime);
 #endif //SCHEMA_K2_GETVARIABLETYPETREE_LOG_TIME
-
 #undef SCHEMA_K2_GETVARIABLETYPETREE_LOG_TIME
 }
 
 void UEdGraphSchema_K2::GetVariableIndexTypeTree( TArray< TSharedPtr<FPinTypeTreeInfo> >& TypeTree, bool bAllowExec, bool bAllowWildcard ) const
 {
-	// Clear the list
-	TypeTree.Empty();
-
-	if( bAllowExec )
-	{
-		TypeTree.Add( MakeShareable( new FPinTypeTreeInfo(GetCategoryText(PC_Exec, true), PC_Exec, this, LOCTEXT("ExecIndexType", "Execution pin")) ) );
-	}
-
-	TypeTree.Add( MakeShareable( new FPinTypeTreeInfo(GetCategoryText(PC_Boolean, true), PC_Boolean, this, LOCTEXT("BooleanIndexType", "True or false value")) ) );
-	TypeTree.Add( MakeShareable( new FPinTypeTreeInfo(GetCategoryText(PC_Byte, true), PC_Byte, this, LOCTEXT("ByteIndexType", "8 bit number")) ) );
-	TypeTree.Add( MakeShareable( new FPinTypeTreeInfo(GetCategoryText(PC_Int, true), PC_Int, this, LOCTEXT("IntegerIndexType", "Integer number")) ) );
-
-	// Add wildcard type
-	if (bAllowWildcard)
-	{
-		TypeTree.Add( MakeShareable( new FPinTypeTreeInfo(GetCategoryText(PC_Wildcard, true), PC_Wildcard, this, LOCTEXT("WildcardIndexType", "Wildcard type (unspecified).")) ) );
-	}
-
-	// Add the types that have subtrees
-	TypeTree.Add( MakeShareable( new FPinTypeTreeInfo(GetCategoryText(PC_Enum, true), PC_Enum, this, LOCTEXT("EnumIndexType", "Enumeration types."), true) ) );
+	GetVariableTypeTreeImpl(TypeTree, bAllowExec, bAllowWildcard, true);
 }
 
 bool UEdGraphSchema_K2::DoesTypeHaveSubtypes(const FString& Category) const
 {
-	return (Category == PC_Struct) || (Category == PC_Object) || (Category == PC_Asset) || (Category == PC_AssetClass) || (Category == PC_Interface) || (Category == PC_Class) || (Category == PC_Enum);
+	return (Category == PC_Struct) || (Category == PC_Object) || (Category == PC_Asset) || (Category == PC_AssetClass) || (Category == PC_Interface) || (Category == PC_Class) || (Category == PC_Enum) || (Category == AllObjectTypes);
 }
 
 struct FWildcardArrayPinHelper
diff --git a/Engine/Source/Editor/KismetWidgets/Private/SPinTypeSelector.cpp b/Engine/Source/Editor/KismetWidgets/Private/SPinTypeSelector.cpp
--- a/Engine/Source/Editor/KismetWidgets/Private/SPinTypeSelector.cpp
+++ b/Engine/Source/Editor/KismetWidgets/Private/SPinTypeSelector.cpp
@@ -6,11 +6,60 @@
 #include "IDocumentation.h"
 #include "Editor/UnrealEd/Public/SListViewSelectorDropdownMenu.h"
 #include "SSearchBox.h"
+#include "SSubMenuHandler.h"
 
 #define LOCTEXT_NAMESPACE "PinTypeSelector"
 
 static const FString BigTooltipDocLink = TEXT("Shared/Editor/Blueprint/VariableTypes");
 
+class SPinTypeRow : public SComboRow<FPinTypeTreeItem>
+{
+public:
+	SLATE_BEGIN_ARGS( SPinTypeRow )
+	{}
+		SLATE_DEFAULT_SLOT( FArguments, Content )
+		SLATE_EVENT( FOnGetContent, OnGetMenuContent )
+	SLATE_END_ARGS()
+public:
+	void Construct(const FArguments& InArgs, const TSharedRef<STableViewBase>& InOwnerTable, TWeakPtr<SMenuOwner> InMenuOwner)
+	{
+		SComboRow<FPinTypeTreeItem>::Construct( SComboRow<FPinTypeTreeItem>::FArguments()
+			.ToolTip(InArgs._ToolTip)
+			[
+				SAssignNew(SubMenuHandler, SSubMenuHandler, InMenuOwner)
+				.OnGetMenuContent(InArgs._OnGetMenuContent)
+				.MenuContent(nullptr)
+				[
+					InArgs._Content.Widget
+				]
+			],
+			InOwnerTable);
+	}
+
+	// SWidget interface
+	virtual bool IsHovered() const override
+	{
+		return SComboRow<FPinTypeTreeItem>::IsHovered() || SubMenuHandler.Pin()->ShouldSubMenuAppearHovered();
+	}
+	// End of SWidget interface
+
+	/** Returns TRUE if there is a Sub-Menu to open */
+	bool HasSubMenu() const
+	{
+		return SubMenuHandler.Pin()->HasSubMenu();
+	}
+
+	/** Forces the sub-menu open, clobbering any other open ones in the process */
+	void RequestSubMenuToggle()
+	{
+		SubMenuHandler.Pin()->RequestSubMenuToggle(true, true);
+	}
+
+private:
+	/** The Sub-MenuHandler which is managing the sub-menu content so that mousing over other rows will not close the sub-menus immediately */
+	TWeakPtr<SSubMenuHandler> SubMenuHandler;
+};
+
 void SPinTypeSelector::Construct(const FArguments& InArgs, FGetPinTypeTree GetPinTypeTreeFunc)
 {
 	SearchText = FText::GetEmpty();
@@ -176,29 +225,148 @@ TSharedRef<ITableRow> SPinTypeSelector::GenerateTypeTreeRow(FPinTypeTreeItem InI
 
 	const FString PinTooltipExcerpt = ((PinType.PinCategory != UEdGraphSchema_K2::PC_Byte || PinType.PinSubCategoryObject == nullptr) ? PinType.PinCategory : TEXT("Enum")); 
 
-	return SNew( SComboRow<FPinTypeTreeItem>, OwnerTree )
+	// If there is a sub-menu for this pin type, we need to bind the function to handle the sub-menu
+	FOnGetContent OnGetContent;
+	if (InItem->GetPossibleObjectReferenceTypes() != static_cast<uint8>(EObjectReferenceType::NotAnObject))
+	{
+		OnGetContent.BindSP(this, &SPinTypeSelector::GetAllowedObjectTypes, InItem);
+	}
+
+	TSharedPtr< SHorizontalBox > HorizontalBox;
+	TSharedRef< ITableRow > ReturnWidget = SNew( SPinTypeRow, OwnerTree, MenuContent )
 		.ToolTip( IDocumentation::Get()->CreateToolTip( Tooltip, NULL, *BigTooltipDocLink, PinTooltipExcerpt) )
+		.OnGetMenuContent(OnGetContent)
 		[
-			SNew(SHorizontalBox)
+			SAssignNew(HorizontalBox, SHorizontalBox)
+			+SHorizontalBox::Slot()
+				.AutoWidth()
+				.Padding(1.f)
+			[
+				SNew(SImage)
+					.Image(IconBrush)
+					.ColorAndOpacity(Schema->GetPinTypeColor(PinType))
+					.Visibility( InItem->bReadOnly ? EVisibility::Collapsed : EVisibility::Visible )
+			]
+			+SHorizontalBox::Slot()
+				.AutoWidth()
+				.Padding(1.f)
+			[
+				SNew(STextBlock)
+					.Text(Description)
+					.HighlightText(SearchText)
+					.Font( bHasChildren ? FEditorStyle::GetFontStyle(TEXT("Kismet.TypePicker.CategoryFont")) : FEditorStyle::GetFontStyle(TEXT("Kismet.TypePicker.NormalFont")) )
+			]
+		];
+
+	// Add a sub-menu indicator arrow to inform the user that there are sub-items to be displayed
+	if (OnGetContent.IsBound())
+	{
+		HorizontalBox->AddSlot()
+			.FillWidth(1.0f)
+			.VAlign( VAlign_Center )
+			.HAlign( HAlign_Right )
+			[
+				SNew( SBox )
+				.Padding(FMargin(7,0,0,0))
+				[
+					SNew( SImage )
+					.Image( FEditorStyle::Get().GetBrush( "ToolBar.SubMenuIndicator" ) )
+				]
+			];
+	}
+
+	return ReturnWidget;
+}
+
+void SPinTypeSelector::AddObjectReferenceMenuEntry(FMenuBuilder& InOutMenuBuilder, FPinTypeTreeItem InItem, FEdGraphPinType& InPinType, const FSlateBrush* InIconBrush) const
+{
+	InOutMenuBuilder.AddMenuEntry(
+		FUIAction(
+		FExecuteAction::CreateSP(this, &SPinTypeSelector::OnSelectPinType, InItem, InPinType.PinCategory)
+		),
+		// Contents
+		SNew(SHorizontalBox)
 			+SHorizontalBox::Slot()
 			.AutoWidth()
 			.Padding(1.f)
 			[
 				SNew(SImage)
-				.Image(IconBrush)
-				.ColorAndOpacity(Schema->GetPinTypeColor(PinType))
-				.Visibility( InItem->bReadOnly ? EVisibility::Collapsed : EVisibility::Visible )
+				.Image(InIconBrush)
+				.ColorAndOpacity(Schema->GetPinTypeColor(InPinType))
 			]
 			+SHorizontalBox::Slot()
 			.AutoWidth()
 			.Padding(1.f)
 			[
 				SNew(STextBlock)
-				.Text(Description)
-				.HighlightText(SearchText)
-				.Font( bHasChildren ? FEditorStyle::GetFontStyle(TEXT("Kismet.TypePicker.CategoryFont")) : FEditorStyle::GetFontStyle(TEXT("Kismet.TypePicker.NormalFont")) )
+				.Text(UEdGraphSchema_K2::GetCategoryText(InPinType.PinCategory))
+				.Font(FEditorStyle::GetFontStyle(TEXT("Kismet.TypePicker.NormalFont")) )
 			]
-		];
+		);
+}
+
+TSharedRef< SWidget > SPinTypeSelector::GetAllowedObjectTypes(FPinTypeTreeItem InItem) const
+{
+	FMenuBuilder MenuBuilder( false, MakeShareable(new FUICommandList) );
+
+	// Do not force the pin type here, that causes a load of the Blueprint (if unloaded)
+	FEdGraphPinType PinType = InItem->GetPinType(false);
+	const FSlateBrush* IconBrush = GetIconFromPin(PinType);
+
+	uint8 PossibleObjectReferenceTypes = InItem->GetPossibleObjectReferenceTypes();
+
+	// Per each object reference type, change the category to the type and add a menu entry (this will get the color to be correct)
+
+	if (PossibleObjectReferenceTypes & static_cast<uint8>(EObjectReferenceType::ObjectReference))
+	{
+		PinType.PinCategory = UEdGraphSchema_K2::PC_Object;
+		AddObjectReferenceMenuEntry(MenuBuilder, InItem, PinType, IconBrush);
+	}
+
+	if (PossibleObjectReferenceTypes & static_cast<uint8>(EObjectReferenceType::ClassReference))
+	{
+		PinType.PinCategory = UEdGraphSchema_K2::PC_Class;
+		AddObjectReferenceMenuEntry(MenuBuilder, InItem, PinType, IconBrush);
+	}
+
+	if (PossibleObjectReferenceTypes & static_cast<uint8>(EObjectReferenceType::AssetID))
+	{
+		PinType.PinCategory = UEdGraphSchema_K2::PC_Asset;
+		AddObjectReferenceMenuEntry(MenuBuilder, InItem, PinType, IconBrush);
+	}
+
+	if (PossibleObjectReferenceTypes & static_cast<uint8>(EObjectReferenceType::ClassAssetID))
+	{
+		PinType.PinCategory = UEdGraphSchema_K2::PC_AssetClass;
+		AddObjectReferenceMenuEntry(MenuBuilder, InItem, PinType, IconBrush);
+	}
+
+	return MenuBuilder.MakeWidget();
+}
+
+void SPinTypeSelector::OnSelectPinType(FPinTypeTreeItem InItem, FString InPinCategory)
+{
+	const FScopedTransaction Transaction( LOCTEXT("ChangeParam", "Change Paramater Type") );
+
+	FEdGraphPinType NewTargetPinType = TargetPinType.Get();
+	//Call delegate in order to notify pin type change is about to happen
+	OnTypePreChanged.ExecuteIfBound(NewTargetPinType);
+
+	const FEdGraphPinType& SelectionPinType = InItem->GetPinType(true);
+
+	// Change the pin's type
+	NewTargetPinType.PinCategory = InPinCategory;
+	NewTargetPinType.PinSubCategory = SelectionPinType.PinSubCategory;
+	NewTargetPinType.PinSubCategoryObject = SelectionPinType.PinSubCategoryObject;
+
+	TypeComboButton->SetIsOpen(false);
+
+	if( NewTargetPinType.PinCategory == Schema->PC_Exec )
+	{
+		NewTargetPinType.bIsArray = false;
+	}
+
+	OnTypeChanged.ExecuteIfBound(NewTargetPinType);
 }
 
 void SPinTypeSelector::OnTypeSelectionChanged(FPinTypeTreeItem Selection, ESelectInfo::Type SelectInfo)
@@ -214,27 +382,16 @@ void SPinTypeSelector::OnTypeSelectionChanged(FPinTypeTreeItem Selection, ESelec
 	{
 		if( !Selection->bReadOnly )
 		{
-			const FScopedTransaction Transaction( LOCTEXT("ChangeParam", "Change Paramater Type") );
-			
-			FEdGraphPinType NewTargetPinType = TargetPinType.Get();
-			//Call delegate in order to notify pin type change is about to happen
-			OnTypePreChanged.ExecuteIfBound(NewTargetPinType);
-
-			const FEdGraphPinType& SelectionPinType = Selection->GetPinType(true);
-
-			// Change the pin's type
-			NewTargetPinType.PinCategory = SelectionPinType.PinCategory;
-			NewTargetPinType.PinSubCategory = SelectionPinType.PinSubCategory;
-			NewTargetPinType.PinSubCategoryObject = SelectionPinType.PinSubCategoryObject;
-
-			TypeComboButton->SetIsOpen(false);
-
-			if( NewTargetPinType.PinCategory == Schema->PC_Exec )
+			// Unless mouse clicking on an item with a sub-menu, all attempts to auto-select should open the sub-menu
+			TSharedPtr<SPinTypeRow> PinRow = StaticCastSharedPtr<SPinTypeRow>(TypeTreeView->WidgetFromItem(Selection));
+			if (SelectInfo != ESelectInfo::OnMouseClick && PinRow.IsValid() && PinRow->HasSubMenu())
 			{
-				NewTargetPinType.bIsArray = false;
+				PinRow->RequestSubMenuToggle();
+			}
+			else
+			{
+				OnSelectPinType(Selection, Selection->GetPossibleObjectReferenceTypes() == static_cast<uint8>(EObjectReferenceType::AllTypes)? Schema->PC_Object : Selection->GetPinType(false).PinCategory);
 			}
-
-			OnTypeChanged.ExecuteIfBound(NewTargetPinType);
 		}
 		else
 		{
@@ -295,26 +452,29 @@ TSharedRef<SWidget>	SPinTypeSelector::GetMenuContent()
 			.OnTextChanged( this, &SPinTypeSelector::OnFilterTextChanged )
 			.OnTextCommitted( this, &SPinTypeSelector::OnFilterTextCommitted );
 
-		MenuContent = SNew(SListViewSelectorDropdownMenu<FPinTypeTreeItem>, FilterTextBox, TypeTreeView)
+		MenuContent = SNew(SMenuOwner)
 			[
-				SNew( SVerticalBox )
-				+SVerticalBox::Slot()
-				.AutoHeight()
-				.Padding(4.f, 4.f, 4.f, 4.f)
+				SNew(SListViewSelectorDropdownMenu<FPinTypeTreeItem>, FilterTextBox, TypeTreeView)
 				[
-					FilterTextBox.ToSharedRef()
-				]
-				+SVerticalBox::Slot()
+					SNew( SVerticalBox )
+					+SVerticalBox::Slot()
 					.AutoHeight()
 					.Padding(4.f, 4.f, 4.f, 4.f)
 					[
-						SNew(SBox)
-						.HeightOverride(TreeViewHeight)
-						.WidthOverride(TreeViewWidth)
+						FilterTextBox.ToSharedRef()
+					]
+					+SVerticalBox::Slot()
+						.AutoHeight()
+						.Padding(4.f, 4.f, 4.f, 4.f)
 						[
-							TypeTreeView.ToSharedRef()
+							SNew(SBox)
+							.HeightOverride(TreeViewHeight)
+							.WidthOverride(TreeViewWidth)
+							[
+								TypeTreeView.ToSharedRef()
+							]
 						]
-					]
+				]
 			];
 			
 
@@ -511,6 +671,7 @@ FReply SPinTypeSelector::OnMouseButtonUp( const FGeometry& MyGeometry, const FPo
 
 void SPinTypeSelector::OnMouseLeave( const FPointerEvent& MouseEvent )
 {
+	SCompoundWidget::OnMouseLeave(MouseEvent);
 	bIsRightMousePressed = false;
 }
 
diff --git a/Engine/Source/Editor/KismetWidgets/Public/SPinTypeSelector.h b/Engine/Source/Editor/KismetWidgets/Public/SPinTypeSelector.h
--- a/Engine/Source/Editor/KismetWidgets/Public/SPinTypeSelector.h
+++ b/Engine/Source/Editor/KismetWidgets/Public/SPinTypeSelector.h
@@ -114,7 +114,7 @@ protected:
 	void GetTypeChildren(FPinTypeTreeItem InItem, TArray<FPinTypeTreeItem>& OutChildren);
 
 	/** Reference to the menu content that's displayed when the type button is clicked on */
-	TSharedPtr<SWidget> MenuContent;
+	TSharedPtr<SMenuOwner> MenuContent;
 	virtual TSharedRef<SWidget>	GetMenuContent();
 
 	/** Type searching support */
@@ -138,4 +138,25 @@ protected:
 
 	/** Callback to get the tooltip for the array button widget */
 	FText GetToolTipForArrayWidget() const;
+
+	/**
+	 * Helper function to add an object reference menu entry to a MenuBuilder
+	 *
+	 * @param InOutMenuBuilder		Menu builder to add the menu entry to
+	 * @param InItem				Tree item to use for the callback when a menu item is selected
+	 * @param InPinType				Pin type for generation of the widget to display for the menu entry
+	 * @param InIconBrush			Brush icon to use for the menu entry item
+	 */
+	void AddObjectReferenceMenuEntry(FMenuBuilder& InOutMenuBuilder, FPinTypeTreeItem InItem, FEdGraphPinType& InPinType, const FSlateBrush* InIconBrush) const;
+
+	/** Gets the allowable object types for an tree item, used for building the sub-menu */
+	TSharedRef< SWidget > GetAllowedObjectTypes(FPinTypeTreeItem InItem) const;
+	
+	/**
+	 * When a pin type is selected, handle it
+	 *
+	 * @param InItem				Item selected
+	 * @param InPinCategory			This is the PinType's category, must be provided separately as the PinType in the tree item is always Object Types for any object related type.
+	 */
+	void OnSelectPinType(FPinTypeTreeItem InItem, FString InPinCategory);
 };
diff --git a/Engine/Source/Runtime/Slate/Private/Framework/MultiBox/MultiBox.cpp b/Engine/Source/Runtime/Slate/Private/Framework/MultiBox/MultiBox.cpp
--- a/Engine/Source/Runtime/Slate/Private/Framework/MultiBox/MultiBox.cpp
+++ b/Engine/Source/Runtime/Slate/Private/Framework/MultiBox/MultiBox.cpp
@@ -773,46 +773,6 @@ void SMultiBoxWidget::BuildMultiBoxWidget()
 
 }
 
-
-void SMultiBoxWidget::SetSummonedMenu( TSharedRef< SMenuAnchor > InMenuAnchor )
-{
-	SummonedMenuAnchor = InMenuAnchor;
-}
-
-	
-/**
- * For menu bar multibox widgets, returns the currently open pull-down menu, if there is one open
- *
- * @return	Pull-down menu anchor, or null pointer
- */
-TSharedPtr< const SMenuAnchor > SMultiBoxWidget::GetOpenMenu() const
-{
-	if( SummonedMenuAnchor.IsValid() && SummonedMenuAnchor.Pin()->IsOpen() )
-	{
-		return SummonedMenuAnchor.Pin();
-	}
-
-	// No open menus
-	return TSharedPtr< SMenuAnchor >();
-}
-
-
-
-/**
- * For menu bar multibox widget, closes any open pull-down menus
- */
-void SMultiBoxWidget::CloseSummonedMenus()
-{
-	if( GetOpenMenu().IsValid() )
-	{
-		SummonedMenuAnchor.Pin()->SetIsOpen( false );
-
-		// Menu was closed, so we no longer need a weak reference to it
-		SummonedMenuAnchor = NULL;
-	}
-}
-
-
 TSharedRef<SWidget> SMultiBoxWidget::OnWrapButtonClicked()
 {
 	FMenuBuilder MenuBuilder(true, NULL, TSharedPtr<FExtender>(), false, GetStyleSet());
@@ -1090,3 +1050,4 @@ FReply SMultiBoxWidget::OnKeyDown( const FGeometry& MyGeometry, const FKeyEvent&
 
 	return FReply::Unhandled();
 }
+
diff --git a/Engine/Source/Runtime/Slate/Private/Widgets/Input/SSubMenuHandler.cpp b/Engine/Source/Runtime/Slate/Private/Widgets/Input/SSubMenuHandler.cpp
new file mode 100644
--- /dev/null
+++ b/Engine/Source/Runtime/Slate/Private/Widgets/Input/SSubMenuHandler.cpp
@@ -0,0 +1,158 @@
+// Copyright 1998-2015 Epic Games, Inc. All Rights Reserved.
+
+#include "SlatePrivatePCH.h"
+#include "SSubMenuHandler.h"
+
+void SSubMenuHandler::Construct(const FArguments& InArgs, TWeakPtr<SMenuOwner> InMenuOwner)
+{
+	MenuOwnerWidget = InMenuOwner;
+
+	TSharedPtr< SWidget > ChildSlotWidget;
+
+	if (InArgs._MenuAnchor.IsValid())
+	{
+		MenuAnchor = InArgs._MenuAnchor;
+		ChildSlotWidget = InArgs._Content.Widget;
+	}
+	else
+	{
+		// If no way for a sub-menu is provided, do not use a MenuAnchor
+		if (InArgs._OnGetMenuContent.IsBound() || InArgs._MenuContent.IsValid())
+		{
+			ChildSlotWidget = SAssignNew(MenuAnchor, SMenuAnchor)
+				.Placement(EMenuPlacement::MenuPlacement_MenuRight)
+				.OnGetMenuContent(InArgs._OnGetMenuContent)
+				.MenuContent(InArgs._MenuContent)
+				[
+					InArgs._Content.Widget
+				];
+		}
+		else
+		{
+			ChildSlotWidget = InArgs._Content.Widget;
+		}
+		
+	}
+
+	this->ChildSlot
+	[
+		ChildSlotWidget.ToSharedRef()
+	];
+}
+
+void SSubMenuHandler::OnMouseEnter( const FGeometry& MyGeometry, const FPointerEvent& MouseEvent )
+{
+	TSharedPtr< SMenuOwner > PinnedMenuOwnerWidget( MenuOwnerWidget.Pin() );
+	check( PinnedMenuOwnerWidget.IsValid() );
+
+	// Never dismiss another entry's submenu while the cursor is potentially moving toward that menu.  It's
+	// not fun to try to keep the mouse in the menu entry bounds while moving towards the actual menu!
+	TSharedPtr< const SMenuAnchor > OpenedMenuAnchor = PinnedMenuOwnerWidget->GetOpenMenu();
+	const bool bSubMenuAlreadyOpen = ( OpenedMenuAnchor.IsValid() && OpenedMenuAnchor->IsOpen() );
+	bool bMouseEnteredTowardSubMenu = false;
+	{
+		if( bSubMenuAlreadyOpen )
+		{
+			const FVector2D& SubMenuPosition = OpenedMenuAnchor->GetMenuPosition();
+			const bool bIsMenuTowardRight = MouseEvent.GetScreenSpacePosition().X < SubMenuPosition.X;
+			const bool bDidMouseEnterTowardRight = MouseEvent.GetCursorDelta().X >= 0.0f;	// NOTE: Intentionally inclusive of zero here.
+			bMouseEnteredTowardSubMenu = ( bIsMenuTowardRight == bDidMouseEnterTowardRight );
+		}
+	}
+
+	if (MenuAnchor.IsValid())
+	{
+		TSharedPtr< SMenuOwner > PinnedMenuOwnerWidget( MenuOwnerWidget.Pin() );
+		check( PinnedMenuOwnerWidget.IsValid() );
+
+		// Do we have a different pull-down menu open?
+		TSharedPtr< SMenuAnchor > PinnedMenuAnchor( MenuAnchor.Pin() );
+		if( PinnedMenuOwnerWidget->GetOpenMenu() != PinnedMenuAnchor )
+		{
+			const bool bClobber = bSubMenuAlreadyOpen && bMouseEnteredTowardSubMenu;
+			RequestSubMenuToggle( true, bClobber );
+		}
+	}
+	else
+	{
+		// Hovering over a menu item that is not a sub-menu, we need to close any sub-menus that are open
+		const bool bClobber = bSubMenuAlreadyOpen && bMouseEnteredTowardSubMenu;
+		RequestSubMenuToggle( false, bClobber );
+	}
+}
+
+void SSubMenuHandler::OnMouseLeave( const FPointerEvent& MouseEvent )
+{
+	SCompoundWidget::OnMouseLeave( MouseEvent );
+
+	// Reset any pending sub-menus that may be opening when we stop hovering over it
+	CancelPendingSubMenu();
+}
+
+bool SSubMenuHandler::ShouldSubMenuAppearHovered() const
+{
+	// The sub-menu entry should appear hovered if the sub-menu is open.  Except in the case that the user is actually interacting with this menu.  
+	// In that case we need to show what the user is selecting
+	return MenuAnchor.IsValid() && MenuAnchor.Pin()->IsOpen() && !MenuOwnerWidget.Pin()->IsHovered();
+}
+
+void SSubMenuHandler::RequestSubMenuToggle( bool bOpenMenu, const bool bClobber )
+{
+	if (MenuAnchor.IsValid())
+	{
+		// Reset the time before the menu opens
+		float TimeToSubMenuOpen = bClobber ? MultiBoxConstants::SubMenuClobberTime : MultiBoxConstants::SubMenuOpenTime;
+		if (!ActiveTimerHandle.IsValid())
+		{
+			ActiveTimerHandle = RegisterActiveTimer(TimeToSubMenuOpen, FWidgetActiveTimerDelegate::CreateSP(this, &SSubMenuHandler::UpdateSubMenuState, bOpenMenu));
+		}
+	}
+}
+
+void SSubMenuHandler::CancelPendingSubMenu()
+{
+	// Reset any pending sub-menu openings
+	//SubMenuRequestState = Idle;
+
+	auto PinnedActiveTimerHandle = ActiveTimerHandle.Pin();
+	if (PinnedActiveTimerHandle.IsValid())
+	{
+		UnRegisterActiveTimer(PinnedActiveTimerHandle.ToSharedRef());
+	}
+}
+
+EActiveTimerReturnType SSubMenuHandler::UpdateSubMenuState(double InCurrentTime, float InDeltaTime, bool bWantsOpen)
+{
+	TSharedPtr< SMenuOwner > PinnedMenuOwnerWidget(MenuOwnerWidget.Pin());
+	check(PinnedMenuOwnerWidget.IsValid());
+
+	if (bWantsOpen)
+	{
+		// For menu bar entries, we also need to handle mouse enter/leave events, so we can show and hide
+		// the pull-down menu appropriately
+		check(MenuAnchor.IsValid());
+
+		// Close other open pull-down menus from this menu bar
+		// Do we have a different pull-down menu open?
+		TSharedPtr< SMenuAnchor > PinnedMenuAnchor(MenuAnchor.Pin());
+		if (PinnedMenuOwnerWidget->GetOpenMenu() != PinnedMenuAnchor)
+		{
+			PinnedMenuOwnerWidget->CloseSummonedMenus();
+
+			// Summon the new pull-down menu!
+			if (PinnedMenuAnchor.IsValid())
+			{
+				PinnedMenuAnchor->SetIsOpen(true);
+			}
+
+			// Also tell the MenuOwner about this open pull-down menu, so it can be closed later if we need to
+			PinnedMenuOwnerWidget->SetSummonedMenu(PinnedMenuAnchor.ToSharedRef());
+		}
+	}
+	else
+	{
+		PinnedMenuOwnerWidget->CloseSummonedMenus();
+	}
+
+	return EActiveTimerReturnType::Stop;
+}
diff --git a/Engine/Source/Runtime/Slate/Private/Widgets/Layout/SMenuOwner.cpp b/Engine/Source/Runtime/Slate/Private/Widgets/Layout/SMenuOwner.cpp
new file mode 100644
--- /dev/null
+++ b/Engine/Source/Runtime/Slate/Private/Widgets/Layout/SMenuOwner.cpp
@@ -0,0 +1,39 @@
+// Copyright 1998-2015 Epic Games, Inc. All Rights Reserved.
+
+#include "SlatePrivatePCH.h"
+#include "SMenuOwner.h"
+
+void SMenuOwner::Construct(const FArguments& InArgs)
+{
+	this->ChildSlot
+		[
+			InArgs._Content.Widget
+		];
+}
+
+void SMenuOwner::SetSummonedMenu( TSharedRef< SMenuAnchor > InMenuAnchor )
+{
+	SummonedMenuAnchor = InMenuAnchor;
+}
+	
+TSharedPtr< const SMenuAnchor > SMenuOwner::GetOpenMenu() const
+{
+	if( SummonedMenuAnchor.IsValid() && SummonedMenuAnchor.Pin()->IsOpen() )
+	{
+		return SummonedMenuAnchor.Pin();
+	}
+
+	// No open menus
+	return TSharedPtr< SMenuAnchor >();
+}
+
+void SMenuOwner::CloseSummonedMenus()
+{
+	if( GetOpenMenu().IsValid() )
+	{
+		SummonedMenuAnchor.Pin()->SetIsOpen( false );
+
+		// Menu was closed, so we no longer need a weak reference to it
+		SummonedMenuAnchor = nullptr;
+	}
+}
\ No newline at end of file
diff --git a/Engine/Source/Runtime/Slate/Private/Widgets/Views/STableViewBase.cpp b/Engine/Source/Runtime/Slate/Private/Widgets/Views/STableViewBase.cpp
--- a/Engine/Source/Runtime/Slate/Private/Widgets/Views/STableViewBase.cpp
+++ b/Engine/Source/Runtime/Slate/Private/Widgets/Views/STableViewBase.cpp
@@ -453,6 +453,8 @@ FReply STableViewBase::OnMouseMove( const FGeometry& MyGeometry, const FPointerE
 
 void STableViewBase::OnMouseLeave( const FPointerEvent& MouseEvent )
 {
+	SCompoundWidget::OnMouseLeave(MouseEvent);
+
 	bStartedTouchInteraction = false;
 	if(this->HasMouseCapture() == false)
 	{
diff --git a/Engine/Source/Runtime/Slate/Public/Framework/MultiBox/MultiBox.h b/Engine/Source/Runtime/Slate/Public/Framework/MultiBox/MultiBox.h
--- a/Engine/Source/Runtime/Slate/Public/Framework/MultiBox/MultiBox.h
+++ b/Engine/Source/Runtime/Slate/Public/Framework/MultiBox/MultiBox.h
@@ -424,7 +424,7 @@ protected:
  * MultiBox Slate widget
  */
 class SLATE_API SMultiBoxWidget
-	: public SCompoundWidget
+	: public SMenuOwner
 {
 
 public:
@@ -470,28 +470,6 @@ public:
 	 */
 	void BuildMultiBoxWidget();
 
-	
-	/**
-	 * For menu bar multibox widgets, tells the multibox widget about a currently active pull-down menu
-	 *
-	 * @param	InMenuAnchor	Menu anchor for active pull-down menu or sub-menu
-	 */
-	void SetSummonedMenu( TSharedRef< SMenuAnchor > InMenuAnchor );
-	
-
-	/**
-	 * For menu bar or sub-menu multibox widgets, returns the currently open menu, if there is one open
-	 *
-	 * @return	Menu anchor, or null pointer
-	 */
-	TSharedPtr< const SMenuAnchor > GetOpenMenu() const;
-
-	/**
-	 * For menu bar multibox widget, closes any open pull-down or sub menus
-	 */
-	void CloseSummonedMenus();
-	
-
 	/** Generates the tiles for an STileView for button rows */
 	TSharedRef<ITableRow> GenerateTiles(TSharedPtr<SWidget> Item, const TSharedRef<STableViewBase>& OwnerTable);
 
@@ -585,9 +563,6 @@ private:
 	/** The MultiBox we're associated with */
 	TSharedPtr< FMultiBox > MultiBox;
 
-	/** For menu bar multibox widgets, this stores a weak reference to the last pull-down or sub-menu that was summoned. */
-	TWeakPtr< SMenuAnchor > SummonedMenuAnchor;
-
 	/** An array of widgets used for an STileView if used */
 	TArray< TSharedPtr<SWidget> > TileViewWidgets;
 
diff --git a/Engine/Source/Runtime/Slate/Public/SlateBasics.h b/Engine/Source/Runtime/Slate/Public/SlateBasics.h
--- a/Engine/Source/Runtime/Slate/Public/SlateBasics.h
+++ b/Engine/Source/Runtime/Slate/Public/SlateBasics.h
@@ -66,6 +66,7 @@ class FGenericCommands;
 #include "SUniformGridPanel.h"
 #include "SMenuAnchor.h"
 #include "MultiBoxDefs.h"
+#include "SMenuOwner.h"
 ////
 #include "MultiBox.h"
 ////
diff --git a/Engine/Source/Runtime/Slate/Public/Widgets/Input/SSubMenuHandler.h b/Engine/Source/Runtime/Slate/Public/Widgets/Input/SSubMenuHandler.h
new file mode 100644
--- /dev/null
+++ b/Engine/Source/Runtime/Slate/Public/Widgets/Input/SSubMenuHandler.h
@@ -0,0 +1,65 @@
+// Copyright 1998-2015 Epic Games, Inc. All Rights Reserved.
+
+#pragma once
+
+/**
+ * Wrapper for any widget that is used in a table view that wants to handle sub-menus with the same functionality as a normal multibox menu
+ */
+class SLATE_API SSubMenuHandler : public SCompoundWidget
+{
+public:
+	SLATE_BEGIN_ARGS( SSubMenuHandler )
+	{}
+		SLATE_DEFAULT_SLOT( FArguments, Content )
+		SLATE_ARGUMENT( TSharedPtr<SMenuAnchor>, MenuAnchor )
+		SLATE_ARGUMENT( TSharedPtr<SWidget>, MenuContent )
+		SLATE_EVENT( FOnGetContent, OnGetMenuContent )
+	SLATE_END_ARGS()
+
+	void Construct(const FArguments& InArgs, TWeakPtr<SMenuOwner> InMenuOwner);
+
+public:
+	// SWidget interface
+	virtual void OnMouseEnter( const FGeometry& MyGeometry, const FPointerEvent& MouseEvent ) override;
+	virtual void OnMouseLeave( const FPointerEvent& MouseEvent ) override;
+	// End of SWidget interface
+	
+	/**
+	 * Returns whether or the sub-menu entry should appear hovered.  If the sub-menu is open we will always show the menu as hovered to indicate which sub-menu is open
+	 * In the case that the user is interacting in this menu we do not show the menu as hovered because we need to show what the user is actually selecting
+	 */
+	bool ShouldSubMenuAppearHovered() const;
+
+	/**
+	 * Requests that the sub-menu associated with this widget be toggled on or off.
+	 * It does not immediately toggle the menu.  After a set amount of time is passed the menu will toggle
+	 *
+	 * @param bOpenMenu	true to open the menu, false to close the menu if one is currently open
+	 * @param bClobber true if we want to open a menu when another menu is already open
+	 */
+	void RequestSubMenuToggle( bool bOpenMenu, const bool bClobber );
+	
+	/**
+	 * Cancels any open requests to toggle a sub-menu       
+	 */
+	void CancelPendingSubMenu();
+
+	/** Returns TRUE if there is a sub-menu available for this item */
+	bool HasSubMenu() const
+	{
+		return MenuAnchor.IsValid();
+	}
+
+protected:
+	/** One-off delayed active timer to update the open/closed state of the sub menu. */
+	EActiveTimerReturnType UpdateSubMenuState(double InCurrentTime, float InDeltaTime, bool bWantsOpen);
+
+	/** The handle to the active timer to update the sub-menu state */
+	TWeakPtr<FActiveTimerHandle> ActiveTimerHandle;
+
+	/** For pull-down or sub-menu entries, this stores a weak reference to the menu anchor widget that we'll use to summon the menu */
+	TWeakPtr< SMenuAnchor > MenuAnchor;
+
+	/** Weak reference back to the MenuOwner widget that owns us */
+	TWeakPtr< class SMenuOwner > MenuOwnerWidget;
+};
\ No newline at end of file
diff --git a/Engine/Source/Runtime/Slate/Public/Widgets/Layout/SMenuOwner.h b/Engine/Source/Runtime/Slate/Public/Widgets/Layout/SMenuOwner.h
new file mode 100644
--- /dev/null
+++ b/Engine/Source/Runtime/Slate/Public/Widgets/Layout/SMenuOwner.h
@@ -0,0 +1,42 @@
+// Copyright 1998-2015 Epic Games, Inc. All Rights Reserved.
+
+#pragma once
+
+/**
+ * Any widget that wants to own lists of items that may have sub-menus that should be handled in the same way as multibox menus (mouse can cross other items without sub-menu closing)
+ */
+class SLATE_API SMenuOwner : public SCompoundWidget
+{
+public:
+	SLATE_BEGIN_ARGS( SMenuOwner )
+	{}
+		SLATE_DEFAULT_SLOT( FArguments, Content )
+	SLATE_END_ARGS()
+
+	void Construct(const FArguments& InArgs);
+
+public:
+	/**
+	 * For menu widgets, tells the owner widget about a currently active pull-down menu
+	 *
+	 * @param	InMenuAnchor	Menu anchor for active pull-down menu or sub-menu
+	 */
+	void SetSummonedMenu( TSharedRef< SMenuAnchor > InMenuAnchor );
+	
+
+	/**
+	 * For menu bar or sub-menu widgets, returns the currently open menu, if there is one open
+	 *
+	 * @return	Menu anchor, or null pointer
+	 */
+	TSharedPtr< const SMenuAnchor > GetOpenMenu() const;
+
+	/**
+	 * For menu bar widget, closes any open pull-down or sub menus
+	 */
+	void CloseSummonedMenus();
+
+protected:
+	/** For menu bar widgets, this stores a weak reference to the last pull-down or sub-menu that was summoned. */
+	TWeakPtr< SMenuAnchor > SummonedMenuAnchor;
+};
\ No newline at end of file
