Fixed crash when selecting for "None" class type on the Spawn Actor From Class node after having already selected another class.

Prevent attempts to auto-wire old pins to new when there are no new pins to prevent division by zero issue through modulus.

diff --git a/Engine/Source/Editor/BlueprintGraph/Private/K2Node.cpp b/Engine/Source/Editor/BlueprintGraph/Private/K2Node.cpp
--- a/Engine/Source/Editor/BlueprintGraph/Private/K2Node.cpp
+++ b/Engine/Source/Editor/BlueprintGraph/Private/K2Node.cpp
@@ -627,6 +627,12 @@ void UK2Node::ReconstructSinglePin(UEdGraphPin* NewPin, UEdGraphPin* OldPin, ERe
 
 void UK2Node::RewireOldPinsToNewPins(TArray<UEdGraphPin*>& InOldPins, TArray<UEdGraphPin*>& InNewPins)
 {
+	// No need to attempt anything if there are no new pins
+	if (InNewPins.Num() == 0)
+	{
+		return;
+	}
+
 	// Rewire any connection to pins that are matched by name (O(N^2) right now)
 	//@TODO: Can do moderately smart things here if only one pin changes name by looking at it's relative position, etc...,
 	// rather than just failing to map it and breaking the links
