UE-17069 Make/Break/SetMembersInStruct could be improved for PostProcessSettings
1st part: HiddenByDefault metadata. Structs with the MD has all pin disabled by default in Make/Break nodes.

diff --git a/Engine/Source/Editor/BlueprintGraph/Classes/K2Node_StructOperation.h b/Engine/Source/Editor/BlueprintGraph/Classes/K2Node_StructOperation.h
--- a/Engine/Source/Editor/BlueprintGraph/Classes/K2Node_StructOperation.h
+++ b/Engine/Source/Editor/BlueprintGraph/Classes/K2Node_StructOperation.h
@@ -33,7 +33,8 @@ protected:
 		virtual void GetRecordDefaults(UProperty* TestProperty, FOptionalPinFromProperty& Record) const override
 		{
 			Record.bCanToggleVisibility = true;
-			Record.bShowPin = true;
+			auto OwnerStruct = TestProperty ? TestProperty->GetOwnerStruct() : nullptr;
+			Record.bShowPin = OwnerStruct ? !OwnerStruct->HasMetaData(TEXT("HiddenByDefault")) : true;
 		}
 
 		virtual void CustomizePinData(UEdGraphPin* Pin, FName SourcePropertyName, int32 ArrayIndex, UProperty* Property) const override;
diff --git a/Engine/Source/Runtime/CoreUObject/Public/UObject/ObjectBase.h b/Engine/Source/Runtime/CoreUObject/Public/UObject/ObjectBase.h
--- a/Engine/Source/Runtime/CoreUObject/Public/UObject/ObjectBase.h
+++ b/Engine/Source/Runtime/CoreUObject/Public/UObject/ObjectBase.h
@@ -993,6 +993,9 @@ namespace UM
 
 		/// [StructMetadata] Indicates that the struct has a custom make node (and what the path to the BlueprintCallable UFunction is) that should be used instead of the default MakeStruct node.  
 		HasNativeMake,
+
+		/// [StructMetadata] Pins in Make and Break nodes are hidden by default.
+		HiddenByDefault,
 	};
 
 	// Metadata usable in UPROPERTY for customizing the behavior when displaying the property in a property panel
diff --git a/Engine/Source/Runtime/Engine/Classes/Engine/Scene.h b/Engine/Source/Runtime/Engine/Classes/Engine/Scene.h
--- a/Engine/Source/Runtime/Engine/Classes/Engine/Scene.h
+++ b/Engine/Source/Runtime/Engine/Classes/Engine/Scene.h
@@ -72,7 +72,7 @@ struct FWeightedBlendables
 // Each property consists of a bool to enable it (by default off),
 // the variable declaration and further down the default value for it.
 // The comment should include the meaning and usable range.
-USTRUCT(BlueprintType)
+USTRUCT(BlueprintType, meta=(HiddenByDefault))
 struct FPostProcessSettings
 {
 	GENERATED_USTRUCT_BODY()
