UE-12248 "Set members in [MyStruct]" node needs an output pin for the post-modified struct

diff --git a/Engine/Source/Editor/BlueprintGraph/Classes/K2Node_SetFieldsInStruct.h b/Engine/Source/Editor/BlueprintGraph/Classes/K2Node_SetFieldsInStruct.h
--- a/Engine/Source/Editor/BlueprintGraph/Classes/K2Node_SetFieldsInStruct.h
+++ b/Engine/Source/Editor/BlueprintGraph/Classes/K2Node_SetFieldsInStruct.h
@@ -16,6 +16,7 @@ class UK2Node_SetFieldsInStruct : public UK2Node_MakeStruct
 	virtual FText GetTooltipText() const override;
 	virtual FName GetPaletteIcon(FLinearColor& OutColor) const override;
 	virtual void ValidateNodeDuringCompilation(class FCompilerResultsLog& MessageLog) const override;
+	virtual void ExpandNode(class FKismetCompilerContext& CompilerContext, UEdGraph* SourceGraph);
 	// End  UEdGraphNode interface
 
 	// Begin K2Node interface
diff --git a/Engine/Source/Editor/BlueprintGraph/Private/K2Node_SetFieldsInStruct.cpp b/Engine/Source/Editor/BlueprintGraph/Private/K2Node_SetFieldsInStruct.cpp
--- a/Engine/Source/Editor/BlueprintGraph/Private/K2Node_SetFieldsInStruct.cpp
+++ b/Engine/Source/Editor/BlueprintGraph/Private/K2Node_SetFieldsInStruct.cpp
@@ -4,6 +4,7 @@
 #include "K2Node_SetFieldsInStruct.h"
 #include "MakeStructHandler.h"
 #include "CompilerResultsLog.h"
+#include "KismetCompiler.h"
 
 #define LOCTEXT_NAMESPACE "K2Node_MakeStruct"
 
@@ -13,6 +14,11 @@ struct SetFieldsInStructHelper
 	{
 		return TEXT("StructRef");
 	}
+
+	static const TCHAR* StructOutPinName()
+	{
+		return TEXT("StructOut");
+	}
 };
 
 class FKCHandler_SetFieldsInStruct : public FKCHandler_MakeStruct
@@ -44,6 +50,7 @@ void UK2Node_SetFieldsInStruct::AllocateDefaultPins()
 
 		CreatePin(EGPD_Input, Schema->PC_Struct, TEXT(""), StructType, false, true, SetFieldsInStructHelper::StructRefPinName());
 
+		CreatePin(EGPD_Output, Schema->PC_Struct, TEXT(""), StructType, false, true, SetFieldsInStructHelper::StructOutPinName());
 		{
 			FStructOnScope StructOnScope(Cast<UScriptStruct>(StructType));
 			FSetFieldsInStructPinManager OptionalPinManager(StructOnScope.GetStructMemory());
@@ -110,7 +117,7 @@ FNodeHandlingFunctor* UK2Node_SetFieldsInStruct::CreateNodeHandler(class FKismet
 
 bool UK2Node_SetFieldsInStruct::ShowCustomPinActions(const UEdGraphPin* Pin, bool bIgnorePinsNum)
 {
-	const int32 MinimalPinsNum = 4;
+	const int32 MinimalPinsNum = 5;
 	const UEdGraphSchema_K2* Schema = GetDefault<UEdGraphSchema_K2>();
 	const auto Node = Pin ? Cast<const UK2Node_SetFieldsInStruct>(Pin->GetOwningNodeUnchecked()) : NULL;
 	return Node
@@ -189,4 +196,24 @@ void UK2Node_SetFieldsInStruct::FSetFieldsInStructPinManager::GetRecordDefaults(
 	Record.bShowPin = false;
 }
 
+void UK2Node_SetFieldsInStruct::ExpandNode(class FKismetCompilerContext& CompilerContext, UEdGraph* SourceGraph)
+{
+	Super::ExpandNode(CompilerContext, SourceGraph);
+
+	UEdGraphPin* OutPin = FindPin(SetFieldsInStructHelper::StructOutPinName());
+	if (OutPin && OutPin->LinkedTo.Num())
+	{
+		UEdGraphPin* InPin = FindPin(SetFieldsInStructHelper::StructRefPinName());
+		UEdGraphPin* SourcePin = (InPin && (1 == InPin->LinkedTo.Num())) ? InPin->LinkedTo[0] : nullptr;
+		auto Schema = CompilerContext.GetSchema();
+		const bool bCopied = SourcePin && Schema->MovePinLinks(*OutPin, *SourcePin, true).CanSafeConnect();
+		
+		if (!bCopied)
+		{
+			CompilerContext.MessageLog.Error(*LOCTEXT("ExpansionError", "Cannot copy links from @@").ToString(), OutPin);
+		}
+	}
+	Pins.Remove(OutPin);
+}
+
 #undef LOCTEXT_NAMESPACE
\ No newline at end of file
