Tooltip for StructOutPin in SetFieldsInStruct

diff --git a/Engine/Source/Editor/BlueprintGraph/Private/K2Node_SetFieldsInStruct.cpp b/Engine/Source/Editor/BlueprintGraph/Private/K2Node_SetFieldsInStruct.cpp
--- a/Engine/Source/Editor/BlueprintGraph/Private/K2Node_SetFieldsInStruct.cpp
+++ b/Engine/Source/Editor/BlueprintGraph/Private/K2Node_SetFieldsInStruct.cpp
@@ -50,7 +50,8 @@ void UK2Node_SetFieldsInStruct::AllocateDefaultPins()
 
 		CreatePin(EGPD_Input, Schema->PC_Struct, TEXT(""), StructType, false, true, SetFieldsInStructHelper::StructRefPinName());
 
-		CreatePin(EGPD_Output, Schema->PC_Struct, TEXT(""), StructType, false, true, SetFieldsInStructHelper::StructOutPinName());
+		auto OutPin = CreatePin(EGPD_Output, Schema->PC_Struct, TEXT(""), StructType, false, true, SetFieldsInStructHelper::StructOutPinName());
+		OutPin->PinToolTip = LOCTEXT("SetFieldsInStruct_OutPinTooltip", "Reference to the input struct").ToString();
 		{
 			FStructOnScope StructOnScope(Cast<UScriptStruct>(StructType));
 			FSetFieldsInStructPinManager OptionalPinManager(StructOnScope.GetStructMemory());
