UE-17069 Make/Break/SetMembersInStruct could be improved for PostProcessSettings
3rd part: The override checkboxes should be sorted with the members they affect and ideally combined into one UI - maybe the show checkbox directly drives the override checkbox/bool. This could be done with another annotation or based on the existing annotations.

diff --git a/Engine/Source/Editor/BlueprintGraph/Classes/K2Node.h b/Engine/Source/Editor/BlueprintGraph/Classes/K2Node.h
--- a/Engine/Source/Editor/BlueprintGraph/Classes/K2Node.h
+++ b/Engine/Source/Editor/BlueprintGraph/Classes/K2Node.h
@@ -76,6 +76,7 @@ public:
 protected:
 	virtual void PostInitNewPin(UEdGraphPin* Pin, FOptionalPinFromProperty& Record, int32 ArrayIndex, UProperty* Property, uint8* PropertyAddress) const {}
 	virtual void PostRemovedOldPin(FOptionalPinFromProperty& Record, int32 ArrayIndex, UProperty* Property, uint8* PropertyAddress) const {}
+	void RebuildProperty(UProperty* TestProperty, FName CategoryName, TArray<FOptionalPinFromProperty>& Properties, UStruct* SourceStruct, TMap<FName, bool>& OldVisibility);
 };
 
 enum ERenamePinResult
diff --git a/Engine/Source/Editor/BlueprintGraph/Private/K2Node.cpp b/Engine/Source/Editor/BlueprintGraph/Private/K2Node.cpp
--- a/Engine/Source/Editor/BlueprintGraph/Private/K2Node.cpp
+++ b/Engine/Source/Editor/BlueprintGraph/Private/K2Node.cpp
@@ -840,30 +840,72 @@ void FOptionalPinManager::RebuildPropertyList(TArray<FOptionalPinFromProperty>&
 	// Rebuild the property list
 	Properties.Empty();
 
+	// find all "bOverride_" properties
+	TMap<FName, UProperty*> OverridesMap;
+	const FString OverridePrefix(TEXT("bOverride_"));
 	for (TFieldIterator<UProperty> It(SourceStruct, EFieldIteratorFlags::IncludeSuper); It; ++It)
 	{
 		UProperty* TestProperty = *It;
+		if (CanTreatPropertyAsOptional(TestProperty) && TestProperty->GetName().StartsWith(OverridePrefix))
+		{
+			FString OriginalName = TestProperty->GetName();
+			if (OriginalName.RemoveFromStart(OverridePrefix) && !OriginalName.IsEmpty())
+			{
+				OverridesMap.Add(FName(*OriginalName), TestProperty);
+			}
+		}
+	}
 
-		if (CanTreatPropertyAsOptional(TestProperty))
+	// handle regular properties
+	for (TFieldIterator<UProperty> It(SourceStruct, EFieldIteratorFlags::IncludeSuper); It; ++It)
+	{
+		UProperty* TestProperty = *It;
+		if (CanTreatPropertyAsOptional(TestProperty) && !TestProperty->GetName().StartsWith(OverridePrefix))
 		{
-			FOptionalPinFromProperty* Record = new (Properties) FOptionalPinFromProperty;
-			Record->PropertyName = TestProperty->GetFName();
-			Record->PropertyFriendlyName = UEditorEngine::GetFriendlyName(TestProperty, SourceStruct);
-			Record->PropertyTooltip = TestProperty->GetToolTipText();
+			FName CategoryName = NAME_None;
 #if WITH_EDITOR
-			Record->CategoryName = FObjectEditorUtils::GetCategoryFName(TestProperty);
+			CategoryName = FObjectEditorUtils::GetCategoryFName(TestProperty);
 #endif //WITH_EDITOR
-			// Get the defaults
-			GetRecordDefaults(TestProperty, *Record);
 
-			// If this is a refresh, propagate the old visibility
-			if (Record->bCanToggleVisibility)
+			UProperty* OverrideProperty = nullptr;
+			if (OverridesMap.RemoveAndCopyValue(TestProperty->GetFName(), OverrideProperty) && OverrideProperty)
 			{
-				if (bool* pShowHide = OldVisibility.Find(Record->PropertyName))
-				{
-					Record->bShowPin = *pShowHide;
-				}
+				RebuildProperty(OverrideProperty, CategoryName, Properties, SourceStruct, OldVisibility);
 			}
+			RebuildProperty(TestProperty, CategoryName, Properties, SourceStruct, OldVisibility);
+		}
+	}
+
+	// add remaining "bOverride_" properties
+	for (auto Pair : OverridesMap)
+	{
+		UProperty* TestProperty = Pair.Value;
+
+		FName CategoryName = NAME_None;
+#if WITH_EDITOR
+		CategoryName = FObjectEditorUtils::GetCategoryFName(TestProperty);
+#endif //WITH_EDITOR
+
+		RebuildProperty(TestProperty, CategoryName, Properties, SourceStruct, OldVisibility);
+	}
+}
+
+void FOptionalPinManager::RebuildProperty(UProperty* TestProperty, FName CategoryName, TArray<FOptionalPinFromProperty>& Properties, UStruct* SourceStruct, TMap<FName, bool>& OldVisibility)
+{
+	FOptionalPinFromProperty* Record = new (Properties)FOptionalPinFromProperty;
+	Record->PropertyName = TestProperty->GetFName();
+	Record->PropertyFriendlyName = UEditorEngine::GetFriendlyName(TestProperty, SourceStruct);
+	Record->PropertyTooltip = TestProperty->GetToolTipText();
+	Record->CategoryName = CategoryName;
+	// Get the defaults
+	GetRecordDefaults(TestProperty, *Record);
+
+	// If this is a refresh, propagate the old visibility
+	if (Record->bCanToggleVisibility)
+	{
+		if (bool* pShowHide = OldVisibility.Find(Record->PropertyName))
+		{
+			Record->bShowPin = *pShowHide;
 		}
 	}
 }
