UE-18147 Resolve nodes with no connections cause errors on compile

diff --git a/Engine/Source/Editor/BlueprintGraph/Private/K2Node_ConvertAsset.cpp b/Engine/Source/Editor/BlueprintGraph/Private/K2Node_ConvertAsset.cpp
--- a/Engine/Source/Editor/BlueprintGraph/Private/K2Node_ConvertAsset.cpp
+++ b/Engine/Source/Editor/BlueprintGraph/Private/K2Node_ConvertAsset.cpp
@@ -150,10 +150,10 @@ void UK2Node_ConvertAsset::ExpandNode(class FKismetCompilerContext& CompilerCont
 	const UEdGraphSchema_K2* Schema = CompilerContext.GetSchema();
 	
 	UClass* TargetType = GetTargetClass();
-	const bool bIsAssetClass = IsAssetClassType();
-	bool bIsErrorFree = TargetType && Schema && (2 == Pins.Num());
-	if (bIsErrorFree)
+	if (TargetType && Schema && (2 == Pins.Num()))
 	{
+		const bool bIsAssetClass = IsAssetClassType();
+
 		//Create Convert Function
 		auto ConvertToObjectFunc = CompilerContext.SpawnIntermediateNode<UK2Node_CallFunction>(this, SourceGraph);
 		const FName ConvertFunctionName = bIsAssetClass 
@@ -168,7 +168,7 @@ void UK2Node_ConvertAsset::ExpandNode(class FKismetCompilerContext& CompilerCont
 			? FString(TEXT("AssetClass"))
 			: FString(TEXT("Asset"));
 		auto ConvertInput = ConvertToObjectFunc->FindPin(ConvertInputName);
-		bIsErrorFree &= InputPin && ConvertInput && CompilerContext.MovePinLinksToIntermediate(*InputPin, *ConvertInput).CanSafeConnect();
+		bool bIsErrorFree = InputPin && ConvertInput && CompilerContext.MovePinLinksToIntermediate(*InputPin, *ConvertInput).CanSafeConnect();
 
 		auto ConvertOutput = ConvertToObjectFunc->GetReturnValuePin();
 		UEdGraphPin* InnerOutput = nullptr;
@@ -196,14 +196,14 @@ void UK2Node_ConvertAsset::ExpandNode(class FKismetCompilerContext& CompilerCont
 
 		auto OutputPin = FindPin(UK2Node_ConvertAssetImpl::OutputPinName);
 		bIsErrorFree &= OutputPin && InnerOutput && CompilerContext.MovePinLinksToIntermediate(*OutputPin, *InnerOutput).CanSafeConnect();
-	}
 
-	if (!bIsErrorFree)
-	{
-		CompilerContext.MessageLog.Error(*LOCTEXT("InternalConnectionError", "K2Node_ConvertAsset: Internal connection error. @@").ToString(), this);
-	}
+		if (!bIsErrorFree)
+		{
+			CompilerContext.MessageLog.Error(*LOCTEXT("InternalConnectionError", "K2Node_ConvertAsset: Internal connection error. @@").ToString(), this);
+		}
 
-	BreakAllNodeLinks();
+		BreakAllNodeLinks();
+	}
 }
 
 FText UK2Node_ConvertAsset::GetCompactNodeTitle() const
