UE-18223 Unable to copy/paste Return Node in functions
UE-18226 Multiple return nodes in implemented interface functions lack outputs

diff --git a/Engine/Source/Editor/BlueprintGraph/Classes/K2Node_FunctionResult.h b/Engine/Source/Editor/BlueprintGraph/Classes/K2Node_FunctionResult.h
--- a/Engine/Source/Editor/BlueprintGraph/Classes/K2Node_FunctionResult.h
+++ b/Engine/Source/Editor/BlueprintGraph/Classes/K2Node_FunctionResult.h
@@ -25,7 +25,8 @@ class UK2Node_FunctionResult : public UK2Node_FunctionTerminator
 	virtual void GetMenuActions(FBlueprintActionDatabaseRegistrar& ActionRegistrar) const override;
 	virtual bool IsCompatibleWithGraph(UEdGraph const* Graph) const;
 	virtual void PostPlacedNewNode();
-	virtual bool CanDuplicateNode() const { return false; }
+	virtual bool CanDuplicateNode() const { return true; }
+	virtual void PostPasteNode() override;
 	// End UK2Node interface
 
 	// Begin K2Node_FunctionTerminator interface
@@ -34,5 +35,8 @@ class UK2Node_FunctionResult : public UK2Node_FunctionTerminator
 	// End K2Node_FunctionTerminator interface
 
 	BLUEPRINTGRAPH_API TArray<UK2Node_FunctionResult*> GetAllResultNodes() const;
+
+protected:
+	void SyncWithPrimaryResultNode();
 };
 
diff --git a/Engine/Source/Editor/BlueprintGraph/Private/K2Node_FunctionResult.cpp b/Engine/Source/Editor/BlueprintGraph/Private/K2Node_FunctionResult.cpp
--- a/Engine/Source/Editor/BlueprintGraph/Private/K2Node_FunctionResult.cpp
+++ b/Engine/Source/Editor/BlueprintGraph/Private/K2Node_FunctionResult.cpp
@@ -209,6 +209,18 @@ void UK2Node_FunctionResult::PostPlacedNewNode()
 {
 	Super::PostPlacedNewNode();
 
+	SyncWithPrimaryResultNode();
+}
+
+void UK2Node_FunctionResult::PostPasteNode()
+{
+	Super::PostPasteNode();
+
+	SyncWithPrimaryResultNode();
+}
+
+void UK2Node_FunctionResult::SyncWithPrimaryResultNode()
+{
 	UK2Node_FunctionResult* PrimaryNode = nullptr;
 	TArray<UK2Node_FunctionResult*> AllResultNodes = GetAllResultNodes();
 	for (auto ResultNode : AllResultNodes)
@@ -222,14 +234,12 @@ void UK2Node_FunctionResult::PostPlacedNewNode()
 
 	if (PrimaryNode)
 	{
+		TArray< TSharedPtr<FUserPinInfo> > UDPinsCopy = UserDefinedPins;
+		for (auto UDPin : UDPinsCopy)
 		{
-			TArray< TSharedPtr<FUserPinInfo> > UDPinsCopy = UserDefinedPins;
-			for (auto UDPin : UDPinsCopy)
-			{
-				RemoveUserDefinedPin(UDPin);
-			}
-			UserDefinedPins.Empty();
+			RemoveUserDefinedPin(UDPin);
 		}
+		UserDefinedPins.Empty();
 
 		SignatureClass = PrimaryNode->SignatureClass;
 		SignatureName = PrimaryNode->SignatureName;
@@ -244,9 +254,10 @@ void UK2Node_FunctionResult::PostPlacedNewNode()
 				NewPinInfo->PinType = UDPin->PinType;
 				NewPinInfo->DesiredPinDirection = UDPin->DesiredPinDirection;
 				UserDefinedPins.Add(NewPinInfo);
-				CreatePinFromUserDefinition(NewPinInfo);
 			}
 		}
+
+		ReconstructNode();
 	}
 }
 
