Fixed incorrect tooltips between "Convert to validated Get" and "Convert to pure Get" options on Get variable nodes.

diff --git a/Engine/Source/Editor/BlueprintGraph/Private/K2Node_VariableGet.cpp b/Engine/Source/Editor/BlueprintGraph/Private/K2Node_VariableGet.cpp
--- a/Engine/Source/Editor/BlueprintGraph/Private/K2Node_VariableGet.cpp
+++ b/Engine/Source/Editor/BlueprintGraph/Private/K2Node_VariableGet.cpp
@@ -316,7 +316,7 @@ void UK2Node_VariableGet::GetContextMenuActions(const FGraphNodeContextMenuBuild
 			if (bIsPureGet)
 			{
 				MenuEntryTitle   = LOCTEXT("ConvertToImpureGetTitle",   "Convert to Validated Get");
-				MenuEntryTooltip = LOCTEXT("ConvertToImpureGetTooltip", "Removes the execution pins to make the node more versitile.");
+				MenuEntryTooltip = LOCTEXT("ConvertToImpureGetTooltip", "Adds in branching execution pins so that you can separatly handle when the returned value is valid/invalid.");
 
 				const UEdGraphSchema_K2* K2Schema = Cast<UEdGraphSchema_K2>(GetSchema());
 				check(K2Schema != nullptr);
@@ -330,7 +330,7 @@ void UK2Node_VariableGet::GetContextMenuActions(const FGraphNodeContextMenuBuild
 			else
 			{
 				MenuEntryTitle   = LOCTEXT("ConvertToPureGetTitle",   "Convert to pure Get");
-				MenuEntryTooltip = LOCTEXT("ConvertToPureGetTooltip", "Adds in branching execution pins so that you can separatly handle when the returned value is valid/invalid.");
+				MenuEntryTooltip = LOCTEXT("ConvertToPureGetTooltip", "Removes the execution pins to make the node more versitile.");
 			}
 
 			Context.MenuBuilder->AddMenuEntry(
