Blueprints: Added support for expressing a per-parameter DisplayName using the markup UPARAM(DisplayName="Something neat") before the parameter declaration

This allows renaming parameters without using a K2ParamRedirect, and to use names that are not legal C++ identifiers

UFUNCTION(BlueprintPure, ...)
static void DoSomethingAwesome(UPARAM(DisplayName="Awesome Param") float BoringParam);

diff --git a/Engine/Source/Editor/BlueprintGraph/Private/K2Node_CallFunction.cpp b/Engine/Source/Editor/BlueprintGraph/Private/K2Node_CallFunction.cpp
--- a/Engine/Source/Editor/BlueprintGraph/Private/K2Node_CallFunction.cpp
+++ b/Engine/Source/Editor/BlueprintGraph/Private/K2Node_CallFunction.cpp
@@ -872,6 +872,13 @@ bool UK2Node_CallFunction::CreatePinsForFunctionCall(const UFunction* Function)
 
 		if (bPinGood)
 		{
+			// Check for a display name override
+			const FString PinDisplayName = Param->GetMetaData(FBlueprintMetadata::MD_DisplayName);
+			if (!PinDisplayName.IsEmpty())
+			{
+				Pin->PinFriendlyName = FText::FromString(PinDisplayName);
+			}
+
 			//Flag pin as read only for const reference property
 			Pin->bDefaultValueIsIgnored = Param->HasAllPropertyFlags(CPF_ConstParm | CPF_ReferenceParm) && (!Function->HasMetaData(FBlueprintMetadata::MD_AutoCreateRefTerm) || Pin->PinType.bIsArray);
 
diff --git a/Engine/Source/Programs/UnrealHeaderTool/Private/HeaderParser.cpp b/Engine/Source/Programs/UnrealHeaderTool/Private/HeaderParser.cpp
--- a/Engine/Source/Programs/UnrealHeaderTool/Private/HeaderParser.cpp
+++ b/Engine/Source/Programs/UnrealHeaderTool/Private/HeaderParser.cpp
@@ -2567,7 +2567,7 @@ void FHeaderParser::GetVarType
 
 	TMap<FName, FString> MetaDataFromNewStyle;
 
-	bool bIsParamList = VariableCategory != EVariableCategory::Member && MatchIdentifier(TEXT("UPARAM"));
+	const bool bIsParamList = (VariableCategory != EVariableCategory::Member) && MatchIdentifier(TEXT("UPARAM"));
 
 	// No specifiers are allowed inside a TArray
 	if( (OuterPropertyType == NULL) || !OuterPropertyType->Matches(TEXT("TArray")) )
