Add HideSelfPin metadata for blueprint callable functions which will make the target pin hidden and not connectable

diff --git a/Engine/Source/Editor/BlueprintGraph/Classes/EdGraphSchema_K2.h b/Engine/Source/Editor/BlueprintGraph/Classes/EdGraphSchema_K2.h
--- a/Engine/Source/Editor/BlueprintGraph/Classes/EdGraphSchema_K2.h
+++ b/Engine/Source/Editor/BlueprintGraph/Classes/EdGraphSchema_K2.h
@@ -104,6 +104,9 @@ public:
 	/** UPROPERTY cannot be modified by other blueprints */
 	static const FName MD_Private;
 
+	/** If true, the self pin should not be shown or connectable regardless of purity, const, etc. similar to InternalUseParam */
+	static const FName MD_HideSelfPin;
+
 	/** If true, the specified UObject parameter will default to "self" if nothing is connected */
 	static const FName MD_DefaultToSelf;
 
diff --git a/Engine/Source/Editor/BlueprintGraph/Private/EdGraphSchema_K2.cpp b/Engine/Source/Editor/BlueprintGraph/Private/EdGraphSchema_K2.cpp
--- a/Engine/Source/Editor/BlueprintGraph/Private/EdGraphSchema_K2.cpp
+++ b/Engine/Source/Editor/BlueprintGraph/Private/EdGraphSchema_K2.cpp
@@ -61,6 +61,7 @@ const FName FBlueprintMetadata::MD_DisplayName(TEXT("DisplayName"));
 const FName FBlueprintMetadata::MD_InternalUseParam(TEXT("InternalUseParam"));
 
 const FName FBlueprintMetadata::MD_ExposeOnSpawn(TEXT("ExposeOnSpawn"));
+const FName FBlueprintMetadata::MD_HideSelfPin(TEXT("HideSelfPin"));
 const FName FBlueprintMetadata::MD_DefaultToSelf(TEXT("DefaultToSelf"));
 const FName FBlueprintMetadata::MD_WorldContext(TEXT("WorldContext"));
 const FName FBlueprintMetadata::MD_CallableWithoutWorldContext(TEXT("CallableWithoutWorldContext"));
diff --git a/Engine/Source/Editor/BlueprintGraph/Private/K2Node_CallFunction.cpp b/Engine/Source/Editor/BlueprintGraph/Private/K2Node_CallFunction.cpp
--- a/Engine/Source/Editor/BlueprintGraph/Private/K2Node_CallFunction.cpp
+++ b/Engine/Source/Editor/BlueprintGraph/Private/K2Node_CallFunction.cpp
@@ -836,7 +836,7 @@ bool UK2Node_CallFunction::CreatePinsForFunctionCall(const UFunction* Function)
 			// For static methods, wire up the self to the CDO of the class if it's not us
 			if (!bIsFunctionCompatibleWithSelf)
 			{
-				auto AuthoritativeClass = FunctionOwnerClass->GetAuthoritativeClass();
+				UClass* AuthoritativeClass = FunctionOwnerClass->GetAuthoritativeClass();
 				SelfPin->DefaultObject = AuthoritativeClass->GetDefaultObject();
 			}
 
@@ -845,8 +845,16 @@ bool UK2Node_CallFunction::CreatePinsForFunctionCall(const UFunction* Function)
 		}
 		else
 		{
-			// Hide the self pin if the function is compatible with the blueprint class and pure (the !bIsConstFunc portion should be going away soon too hopefully)
-			SelfPin->bHidden = bIsFunctionCompatibleWithSelf && (bIsPureFunc && !bIsConstFunc);
+			if (Function->GetBoolMetaData(FBlueprintMetadata::MD_HideSelfPin))
+			{
+				SelfPin->bHidden = true;
+				SelfPin->bNotConnectable = true;
+			}
+			else
+			{
+				// Hide the self pin if the function is compatible with the blueprint class and pure (the !bIsConstFunc portion should be going away soon too hopefully)
+				SelfPin->bHidden = (bIsFunctionCompatibleWithSelf && bIsPureFunc && !bIsConstFunc);
+			}
 		}
 	}
 
