Pins in Select node are dislayed as "ref". New Select node implementation enabled.

diff --git a/Engine/Config/BaseEngine.ini b/Engine/Config/BaseEngine.ini
--- a/Engine/Config/BaseEngine.ini
+++ b/Engine/Config/BaseEngine.ini
@@ -1653,7 +1653,7 @@ bDeferDependencyLoads=true
 bForceDisableCookedDependencyDeferring=false
 bExecutionAfterReturn=false
 bCanSuppressAccessViolation=false
-bUseSelectRef=false
+bUseSelectRef=true
 
 [/Script/Engine.Blueprint]
 bRecompileOnLoad=true
diff --git a/Engine/Source/Editor/BlueprintGraph/Private/K2Node_Select.cpp b/Engine/Source/Editor/BlueprintGraph/Private/K2Node_Select.cpp
--- a/Engine/Source/Editor/BlueprintGraph/Private/K2Node_Select.cpp
+++ b/Engine/Source/Editor/BlueprintGraph/Private/K2Node_Select.cpp
@@ -346,6 +346,8 @@ void UK2Node_Select::AllocateDefaultPins()
 		NumOptionPins = EnumEntries.Num();
 	}
 
+	static const FBoolConfigValueHelper UseSelectRef(TEXT("Kismet"), TEXT("bUseSelectRef"), GEngineIni);
+
 	// Create the option pins
 	for (int32 Idx = 0; Idx < NumOptionPins; Idx++)
 	{
@@ -368,6 +370,7 @@ void UK2Node_Select::AllocateDefaultPins()
 
 		if (NewPin)
 		{
+			NewPin->bDisplayAsMutableRef = UseSelectRef;
 			if (IndexPinType.PinCategory == UEdGraphSchema_K2::PC_Boolean)
 			{
 				NewPin->PinFriendlyName = (Idx == 0 ? GFalse : GTrue);
@@ -383,7 +386,8 @@ void UK2Node_Select::AllocateDefaultPins()
 	CreatePin(EGPD_Input, IndexPinType.PinCategory, IndexPinType.PinSubCategory, IndexPinType.PinSubCategoryObject.Get(), false, false, "Index");
 
 	// Create the return value
-	CreatePin(EGPD_Output, Schema->PC_Wildcard, TEXT(""), NULL, false, false, Schema->PN_ReturnValue);
+	auto ReturnPin = CreatePin(EGPD_Output, Schema->PC_Wildcard, TEXT(""), NULL, false, false, Schema->PN_ReturnValue);
+	ReturnPin->bDisplayAsMutableRef = UseSelectRef;
 
 	Super::AllocateDefaultPins();
 }
diff --git a/Engine/Source/Editor/GraphEditor/Private/SGraphPin.cpp b/Engine/Source/Editor/GraphEditor/Private/SGraphPin.cpp
--- a/Engine/Source/Editor/GraphEditor/Private/SGraphPin.cpp
+++ b/Engine/Source/Editor/GraphEditor/Private/SGraphPin.cpp
@@ -844,7 +844,7 @@ const FSlateBrush* SGraphPin::GetPinIcon() const
 			return IsHovered() ? CachedImg_DelegatePin_DisconnectedHovered : CachedImg_DelegatePin_Disconnected;
 		}
 	}
-	else if ( IsByMutableRef() )
+	else if (GraphPinObj->bDisplayAsMutableRef || IsByMutableRef())
 	{
 		if (IsConnected())
 		{
diff --git a/Engine/Source/Runtime/Engine/Classes/EdGraph/EdGraphPin.h b/Engine/Source/Runtime/Engine/Classes/EdGraph/EdGraphPin.h
--- a/Engine/Source/Runtime/Engine/Classes/EdGraph/EdGraphPin.h
+++ b/Engine/Source/Runtime/Engine/Classes/EdGraph/EdGraphPin.h
@@ -243,6 +243,10 @@ class UEdGraphPin : public UObject
 	UPROPERTY()
 	uint32 bAdvancedView:1;
 
+	/** If true, the pin is displayed as ref */
+	UPROPERTY(transient)
+	uint32 bDisplayAsMutableRef : 1;
+
 	/** Pin name could be changed, so whenever possible it's good to have a persistent GUID identifying Pin to reconstruct Node seamlessly */
 	UPROPERTY()
 	FGuid PersistentGuid;
