UE-18379 Copy/Paste of Select node loses type and values
UE-18380 Connecting index pin on Select node clears type of input pins

diff --git a/Engine/Source/Editor/BlueprintGraph/Private/K2Node_Select.cpp b/Engine/Source/Editor/BlueprintGraph/Private/K2Node_Select.cpp
--- a/Engine/Source/Editor/BlueprintGraph/Private/K2Node_Select.cpp
+++ b/Engine/Source/Editor/BlueprintGraph/Private/K2Node_Select.cpp
@@ -458,6 +458,7 @@ void UK2Node_Select::ReallocatePinsDuringReconstruction(TArray<UEdGraphPin*>& Ol
 	// See if this node was saved in the old version with a boolean as the condition
 	UEdGraphPin* OldConditionPin = NULL;
 	UEdGraphPin* OldIndexPin = NULL;
+	UEdGraphPin* OldReturnPin = NULL;
 	for (auto It = OldPins.CreateConstIterator(); It; It++)
 	{
 		if ((*It)->PinName == TEXT("bPickOption0"))
@@ -468,6 +469,26 @@ void UK2Node_Select::ReallocatePinsDuringReconstruction(TArray<UEdGraphPin*>& Ol
 		{
 			OldIndexPin = (*It);
 		}
+		else if ((*It)->PinName == Schema->PN_ReturnValue)
+		{
+			OldReturnPin = (*It);
+		}
+	}
+
+	bool bIsAnyOptionOrReturnConnected = false;
+	for (auto OldPin : OldPins)
+	{
+		if ((OldPin != OldConditionPin) && (OldPin != OldIndexPin) && OldPin->LinkedTo.Num())
+		{
+			bIsAnyOptionOrReturnConnected = true;
+			break;
+		}
+	}
+	auto NewReturn = GetReturnValuePin();
+	if (!bIsAnyOptionOrReturnConnected && OldReturnPin && NewReturn 
+		&& (NewReturn->PinType.PinCategory == Schema->PC_Wildcard))
+	{
+		NewReturn->PinType = OldReturnPin->PinType;
 	}
 
 	UEdGraphPin* IndexPin = GetIndexPin();
