UE-18350 Able to paste return nodes with outputs into construction script, deleting results in crash

diff --git a/Engine/Source/Editor/BlueprintGraph/Private/K2Node_FunctionResult.cpp b/Engine/Source/Editor/BlueprintGraph/Private/K2Node_FunctionResult.cpp
--- a/Engine/Source/Editor/BlueprintGraph/Private/K2Node_FunctionResult.cpp
+++ b/Engine/Source/Editor/BlueprintGraph/Private/K2Node_FunctionResult.cpp
@@ -190,9 +190,10 @@ void UK2Node_FunctionResult::GetMenuActions(FBlueprintActionDatabaseRegistrar& A
 
 bool UK2Node_FunctionResult::IsCompatibleWithGraph(UEdGraph const* Graph) const
 {
-	const EGraphType GraphType = Graph->GetSchema()->GetGraphType(Graph);
-	const bool bIsCompatible = GraphType == EGraphType::GT_Function;
-	return bIsCompatible && Super::IsCompatibleWithGraph(Graph);
+	auto K2Schema = Cast<const UEdGraphSchema_K2>(Graph ? Graph->GetSchema() : nullptr);
+	const bool bIsConstructionScript = (K2Schema != nullptr) ? K2Schema->IsConstructionScript(Graph) : false;
+	const bool bIsCompatible = (K2Schema != nullptr) ? (EGraphType::GT_Function == K2Schema->GetGraphType(Graph)) : false;
+	return bIsCompatible && !bIsConstructionScript && Super::IsCompatibleWithGraph(Graph);
 }
 
 TArray<UK2Node_FunctionResult*> UK2Node_FunctionResult::GetAllResultNodes() const
