Fix vector axis event output pin type caused by various refactoring changes

diff --git a/Engine/Source/Editor/BlueprintGraph/Private/K2Node_InputVectorAxisEvent.cpp b/Engine/Source/Editor/BlueprintGraph/Private/K2Node_InputVectorAxisEvent.cpp
--- a/Engine/Source/Editor/BlueprintGraph/Private/K2Node_InputVectorAxisEvent.cpp
+++ b/Engine/Source/Editor/BlueprintGraph/Private/K2Node_InputVectorAxisEvent.cpp
@@ -10,7 +10,7 @@
 UK2Node_InputVectorAxisEvent::UK2Node_InputVectorAxisEvent(const FObjectInitializer& ObjectInitializer)
 : Super(ObjectInitializer)
 {
-	EventReference.SetExternalDelegateMember(FName(TEXT("InputAxisHandlerDynamicSignature__DelegateSignature")));
+	EventReference.SetExternalDelegateMember(FName(TEXT("InputVectorAxisHandlerDynamicSignature__DelegateSignature")));
 }
 
 void UK2Node_InputVectorAxisEvent::Serialize(FArchive& Ar)
@@ -21,7 +21,7 @@ void UK2Node_InputVectorAxisEvent::Serialize(FArchive& Ar)
 	{
 		if(Ar.UE4Ver() < VER_UE4_K2NODE_EVENT_MEMBER_REFERENCE && EventSignatureName_DEPRECATED.IsNone() && EventSignatureClass_DEPRECATED == nullptr)
 		{
-			EventReference.SetExternalDelegateMember(FName(TEXT("InputAxisHandlerDynamicSignature__DelegateSignature")));
+			EventReference.SetExternalDelegateMember(FName(TEXT("InputVectorAxisHandlerDynamicSignature__DelegateSignature")));
 		}
 	}
 }
@@ -71,7 +71,7 @@ void UK2Node_InputVectorAxisEvent::GetMenuActions(FBlueprintActionDatabaseRegist
 
 	auto CustomizeInputNodeLambda = [](UEdGraphNode* NewNode, bool bIsTemplateNode, FKey Key)
 	{
-		UK2Node_InputAxisKeyEvent* InputNode = CastChecked<UK2Node_InputAxisKeyEvent>(NewNode);
+		UK2Node_InputVectorAxisEvent* InputNode = CastChecked<UK2Node_InputVectorAxisEvent>(NewNode);
 		InputNode->Initialize(Key);
 	};
 
