Hidden pins generated from native functions will correctly use their default values.

diff --git a/Engine/Source/Editor/BlueprintGraph/Private/K2Node_CallFunction.cpp b/Engine/Source/Editor/BlueprintGraph/Private/K2Node_CallFunction.cpp
--- a/Engine/Source/Editor/BlueprintGraph/Private/K2Node_CallFunction.cpp
+++ b/Engine/Source/Editor/BlueprintGraph/Private/K2Node_CallFunction.cpp
@@ -909,7 +909,6 @@ bool UK2Node_CallFunction::CreatePinsForFunctionCall(const UFunction* Function)
 				{
 					Pin->bHidden = true;
 					Pin->bNotConnectable = InternalPins.Contains(Pin->PinName);
-					K2Schema->SetPinDefaultValueBasedOnType(Pin);
 				}
 			}
 
