Select node creates local variable for index input.

diff --git a/Engine/Source/Editor/BlueprintGraph/Private/K2Node_CallFunction.cpp b/Engine/Source/Editor/BlueprintGraph/Private/K2Node_CallFunction.cpp
--- a/Engine/Source/Editor/BlueprintGraph/Private/K2Node_CallFunction.cpp
+++ b/Engine/Source/Editor/BlueprintGraph/Private/K2Node_CallFunction.cpp
@@ -2055,14 +2055,16 @@ UEdGraphPin* UK2Node_CallFunction::InnerHandleAutoCreateRef(UK2Node* Node, UEdGr
 		UK2Node_PureAssignmentStatement* AssignDefaultValue = CompilerContext.SpawnIntermediateNode<UK2Node_PureAssignmentStatement>(Node, SourceGraph);
 		AssignDefaultValue->AllocateDefaultPins();
 		const bool bVariableConnected = CompilerContext.GetSchema()->TryCreateConnection(AssignDefaultValue->GetVariablePin(), LocalVariable->GetVariablePin());
+		auto AssignInputPit = AssignDefaultValue->GetValuePin();
+		const bool bPreviousInputSaved = AssignInputPit && CompilerContext.MovePinLinksToIntermediate(*Pin, *AssignInputPit).CanSafeConnect();
 		const bool bOutputConnected = CompilerContext.GetSchema()->TryCreateConnection(AssignDefaultValue->GetOutputPin(), Pin);
-		if (!bVariableConnected || !bOutputConnected)
+		if (!bVariableConnected || !bOutputConnected || !bPreviousInputSaved)
 		{
 			CompilerContext.MessageLog.Error(*LOCTEXT("AutoCreateRefTermPin_AssignmentError", "AutoCreateRefTerm Expansion: Assignment Error @@").ToString(), AssignDefaultValue);
 			return NULL;
 		}
 		CompilerContext.GetSchema()->SetPinDefaultValueBasedOnType(AssignDefaultValue->GetValuePin());
-		return AssignDefaultValue->GetValuePin();
+		return AssignInputPit;
 	}
 	return NULL;
 }
diff --git a/Engine/Source/Editor/BlueprintGraph/Private/K2Node_Select.cpp b/Engine/Source/Editor/BlueprintGraph/Private/K2Node_Select.cpp
--- a/Engine/Source/Editor/BlueprintGraph/Private/K2Node_Select.cpp
+++ b/Engine/Source/Editor/BlueprintGraph/Private/K2Node_Select.cpp
@@ -1055,21 +1055,19 @@ void UK2Node_Select::ExpandNode(class FKismetCompilerContext& CompilerContext, U
 	auto Schema = CompilerContext.GetSchema();
 	for (auto Pin : Pins)
 	{
-		const bool bValidAutoRefPin = Pin && !Schema->IsMetaPin(*Pin) && (Pin->Direction == EGPD_Input) && !Pin->LinkedTo.Num();
+		const bool bValidAutoRefPin = Pin && !Schema->IsMetaPin(*Pin) && (Pin->Direction == EGPD_Input) && (!Pin->LinkedTo.Num() || (GetIndexPin() == Pin));
 		if (!bValidAutoRefPin)
 		{
 			continue;
 		}
 
-		const bool bHasDefaultValue = !Pin->DefaultValue.IsEmpty() || Pin->DefaultObject || !Pin->DefaultTextValue.IsEmpty();
-
 		//default values can be reset when the pin is connected
 		const auto DefaultValue = Pin->DefaultValue;
 		const auto DefaultObject = Pin->DefaultObject;
 		const auto DefaultTextValue = Pin->DefaultTextValue;
 		const auto AutogeneratedDefaultValue = Pin->AutogeneratedDefaultValue;
 
-		auto ValuePin = UK2Node_CallFunction::InnerHandleAutoCreateRef(this, Pin, CompilerContext, SourceGraph, bHasDefaultValue);
+		auto ValuePin = UK2Node_CallFunction::InnerHandleAutoCreateRef(this, Pin, CompilerContext, SourceGraph, true);
 		if (ValuePin)
 		{
 			if (!DefaultObject && DefaultTextValue.IsEmpty() && (DefaultValue == AutogeneratedDefaultValue))
