You can now bind the Any key
Input Action and Key events now return the key generating the event

diff --git a/Engine/Source/Editor/BlueprintGraph/Classes/K2Node_InputAction.h b/Engine/Source/Editor/BlueprintGraph/Classes/K2Node_InputAction.h
--- a/Engine/Source/Editor/BlueprintGraph/Classes/K2Node_InputAction.h
+++ b/Engine/Source/Editor/BlueprintGraph/Classes/K2Node_InputAction.h
@@ -55,8 +55,6 @@ class UK2Node_InputAction : public UK2Node
 	BLUEPRINTGRAPH_API UEdGraphPin* GetReleasedPin() const;
 
 private:
-	void CreateInputActionEvent(FKismetCompilerContext& CompilerContext, UEdGraph* SourceGraph, UEdGraphPin* InputActionPin, const EInputEvent InputKeyEvent);
-
 	/** Constructing FText strings can be costly, so we cache the node's title/tooltip */
 	FNodeTextCache CachedTooltip;
 	FNodeTextCache CachedNodeTitle;
diff --git a/Engine/Source/Editor/BlueprintGraph/Classes/K2Node_InputKey.h b/Engine/Source/Editor/BlueprintGraph/Classes/K2Node_InputKey.h
--- a/Engine/Source/Editor/BlueprintGraph/Classes/K2Node_InputKey.h
+++ b/Engine/Source/Editor/BlueprintGraph/Classes/K2Node_InputKey.h
@@ -79,8 +79,6 @@ class UK2Node_InputKey : public UK2Node
 	BLUEPRINTGRAPH_API UEdGraphPin* GetReleasedPin() const;
 
 private:
-	void CreateInputKeyEvent(FKismetCompilerContext& CompilerContext, UEdGraph* SourceGraph, UEdGraphPin* InputKeyPin, const EInputEvent KeyEvent);
-
 	/** Constructing FText strings can be costly, so we cache the node's title/tooltip */
 	FNodeTextCache CachedTooltip;
 	FNodeTextCache CachedNodeTitle;
diff --git a/Engine/Source/Editor/BlueprintGraph/Private/K2Node_InputAction.cpp b/Engine/Source/Editor/BlueprintGraph/Private/K2Node_InputAction.cpp
--- a/Engine/Source/Editor/BlueprintGraph/Private/K2Node_InputAction.cpp
+++ b/Engine/Source/Editor/BlueprintGraph/Private/K2Node_InputAction.cpp
@@ -37,6 +37,7 @@ void UK2Node_InputAction::AllocateDefaultPins()
 
 	CreatePin(EGPD_Output, K2Schema->PC_Exec, TEXT(""), NULL, false, false, TEXT("Pressed"));
 	CreatePin(EGPD_Output, K2Schema->PC_Exec, TEXT(""), NULL, false, false, TEXT("Released"));
+	CreatePin(EGPD_Output, K2Schema->PC_Struct, TEXT(""), FKey::StaticStruct(), false, false, TEXT("Key"));
 
 	Super::AllocateDefaultPins();
 }
@@ -115,40 +116,96 @@ void UK2Node_InputAction::ValidateNodeDuringCompilation(class FCompilerResultsLo
 	}
 }
 
-void UK2Node_InputAction::CreateInputActionEvent(FKismetCompilerContext& CompilerContext, UEdGraph* SourceGraph, UEdGraphPin* InputActionPin, const EInputEvent InputKeyEvent)
+void UK2Node_InputAction::ExpandNode(FKismetCompilerContext& CompilerContext, UEdGraph* SourceGraph)
 {
-	if (InputActionPin->LinkedTo.Num() > 0)
+	Super::ExpandNode(CompilerContext, SourceGraph);
+
+	UEdGraphPin* InputActionPressedPin = GetPressedPin();
+	UEdGraphPin* InputActionReleasedPin = GetReleasedPin();
+		
+	struct EventPinData
+	{
+		EventPinData(UEdGraphPin* InPin,TEnumAsByte<EInputEvent> InEvent ){	Pin=InPin;EventType=InEvent;};
+		UEdGraphPin* Pin;
+		TEnumAsByte<EInputEvent> EventType;
+	};
+
+	TArray<EventPinData> ActivePins;
+	if(( InputActionPressedPin != nullptr ) && (InputActionPressedPin->LinkedTo.Num() > 0 ))
+	{
+		ActivePins.Add(EventPinData(InputActionPressedPin,IE_Pressed));
+	}
+	if((InputActionReleasedPin != nullptr) && (InputActionReleasedPin->LinkedTo.Num() > 0 ))
+	{
+		ActivePins.Add(EventPinData(InputActionReleasedPin,IE_Released));
+	}
+	
+	const UEdGraphSchema_K2* Schema = CompilerContext.GetSchema();
+
+	// If more than one is linked we have to do more complicated behaviors
+	if( ActivePins.Num() > 1 )
+	{
+		// Create a temporary variable to copy Key in to
+		static UScriptStruct* KeyStruct = FKey::StaticStruct();
+		UK2Node_TemporaryVariable* ActionKeyVar = CompilerContext.SpawnIntermediateNode<UK2Node_TemporaryVariable>(this, SourceGraph);
+		ActionKeyVar->VariableType.PinCategory = Schema->PC_Struct;
+		ActionKeyVar->VariableType.PinSubCategoryObject = KeyStruct;
+		ActionKeyVar->AllocateDefaultPins();
+
+		for (auto PinIt = ActivePins.CreateIterator(); PinIt; ++PinIt)
+		{			
+			UEdGraphPin *EachPin = (*PinIt).Pin;
+			// Create the input touch event
+			UK2Node_InputActionEvent* InputActionEvent = CompilerContext.SpawnIntermediateNode<UK2Node_InputActionEvent>(this, SourceGraph);
+			InputActionEvent->CustomFunctionName = FName( *FString::Printf(TEXT("InpActEvt_%s_%s"), *InputActionName.ToString(), *InputActionEvent->GetName()));
+			InputActionEvent->InputActionName = InputActionName;
+			InputActionEvent->bConsumeInput = bConsumeInput;
+			InputActionEvent->bExecuteWhenPaused = bExecuteWhenPaused;
+			InputActionEvent->bOverrideParentBinding = bOverrideParentBinding;
+			InputActionEvent->InputKeyEvent = (*PinIt).EventType;
+			InputActionEvent->EventReference.SetExternalDelegateMember(FName(TEXT("InputActionHandlerDynamicSignature__DelegateSignature")));
+			InputActionEvent->bInternalEvent = true;
+			InputActionEvent->AllocateDefaultPins();
+
+			// Create assignment nodes to assign the key
+			UK2Node_AssignmentStatement* ActionKeyInitialize = CompilerContext.SpawnIntermediateNode<UK2Node_AssignmentStatement>(this, SourceGraph);
+			ActionKeyInitialize->AllocateDefaultPins();
+			Schema->TryCreateConnection(ActionKeyVar->GetVariablePin(), ActionKeyInitialize->GetVariablePin());
+			Schema->TryCreateConnection(ActionKeyInitialize->GetValuePin(), InputActionEvent->FindPinChecked(TEXT("Key")));
+			// Connect the events to the assign key nodes
+			Schema->TryCreateConnection(Schema->FindExecutionPin(*InputActionEvent, EGPD_Output), ActionKeyInitialize->GetExecPin());
+
+			// Move the original event connections to the then pin of the key assign
+			CompilerContext.MovePinLinksToIntermediate(*EachPin, *ActionKeyInitialize->GetThenPin());
+			
+			// Move the original event variable connections to the intermediate nodes
+			CompilerContext.MovePinLinksToIntermediate(*FindPin(TEXT("Key")), *ActionKeyVar->GetVariablePin());
+		}	
+	}
+	else if( ActivePins.Num() == 1 )
 	{
-		UK2Node_InputActionEvent* InputActionEvent = CompilerContext.SpawnIntermediateNode<UK2Node_InputActionEvent>(this, SourceGraph);
-		InputActionEvent->CustomFunctionName = FName( *FString::Printf(TEXT("InpActEvt_%s_%s"), *InputActionName.ToString(), *InputActionEvent->GetName()));
-		InputActionEvent->InputActionName = InputActionName;
-		InputActionEvent->bConsumeInput = bConsumeInput;
-		InputActionEvent->bExecuteWhenPaused = bExecuteWhenPaused;
-		InputActionEvent->bOverrideParentBinding = bOverrideParentBinding;
-		InputActionEvent->InputKeyEvent = InputKeyEvent;
-		InputActionEvent->EventReference.SetExternalDelegateMember(FName(TEXT("InputActionHandlerDynamicSignature__DelegateSignature")));
-		InputActionEvent->bInternalEvent = true;
-		InputActionEvent->AllocateDefaultPins();
-
-		// Move any exec links from the InputActionNode pin to the InputActionEvent node
-		UEdGraphPin* EventOutput = CompilerContext.GetSchema()->FindExecutionPin(*InputActionEvent, EGPD_Output);
-
-		if(EventOutput != NULL)
+		UEdGraphPin* InputActionPin = ActivePins[0].Pin;
+		EInputEvent InputEvent = ActivePins[0].EventType;
+	
+		if (InputActionPin->LinkedTo.Num() > 0)
 		{
-			CompilerContext.MovePinLinksToIntermediate(*InputActionPin, *EventOutput);
+			UK2Node_InputActionEvent* InputActionEvent = CompilerContext.SpawnIntermediateNode<UK2Node_InputActionEvent>(this, SourceGraph);
+			InputActionEvent->CustomFunctionName = FName( *FString::Printf(TEXT("InpActEvt_%s_%s"), *InputActionName.ToString(), *InputActionEvent->GetName()));
+			InputActionEvent->InputActionName = InputActionName;
+			InputActionEvent->bConsumeInput = bConsumeInput;
+			InputActionEvent->bExecuteWhenPaused = bExecuteWhenPaused;
+			InputActionEvent->bOverrideParentBinding = bOverrideParentBinding;
+			InputActionEvent->InputKeyEvent = InputEvent;
+			InputActionEvent->EventReference.SetExternalDelegateMember(FName(TEXT("InputActionHandlerDynamicSignature__DelegateSignature")));
+			InputActionEvent->bInternalEvent = true;
+			InputActionEvent->AllocateDefaultPins();
+
+			CompilerContext.MovePinLinksToIntermediate(*InputActionPin, *Schema->FindExecutionPin(*InputActionEvent, EGPD_Output));
+			CompilerContext.MovePinLinksToIntermediate(*FindPin(TEXT("Key")), *InputActionEvent->FindPin(TEXT("Key")));
 		}
 	}
 }
 
-
-void UK2Node_InputAction::ExpandNode(FKismetCompilerContext& CompilerContext, UEdGraph* SourceGraph)
-{
-	Super::ExpandNode(CompilerContext, SourceGraph);
-
-	CreateInputActionEvent(CompilerContext, SourceGraph, GetPressedPin(), IE_Pressed);
-	CreateInputActionEvent(CompilerContext, SourceGraph, GetReleasedPin(), IE_Released);
-}
-
 void UK2Node_InputAction::GetMenuActions(FBlueprintActionDatabaseRegistrar& ActionRegistrar) const
 {
 	TArray<FName> ActionNames;
diff --git a/Engine/Source/Editor/BlueprintGraph/Private/K2Node_InputKey.cpp b/Engine/Source/Editor/BlueprintGraph/Private/K2Node_InputKey.cpp
--- a/Engine/Source/Editor/BlueprintGraph/Private/K2Node_InputKey.cpp
+++ b/Engine/Source/Editor/BlueprintGraph/Private/K2Node_InputKey.cpp
@@ -56,6 +56,7 @@ void UK2Node_InputKey::AllocateDefaultPins()
 
 	CreatePin(EGPD_Output, K2Schema->PC_Exec, TEXT(""), NULL, false, false, TEXT("Pressed"));
 	CreatePin(EGPD_Output, K2Schema->PC_Exec, TEXT(""), NULL, false, false, TEXT("Released"));
+	CreatePin(EGPD_Output, K2Schema->PC_Struct, TEXT(""), FKey::StaticStruct(), false, false, TEXT("Key"));
 
 	Super::AllocateDefaultPins();
 }
@@ -242,51 +243,120 @@ void UK2Node_InputKey::ValidateNodeDuringCompilation(class FCompilerResultsLog&
 	}
 }
 
-void UK2Node_InputKey::CreateInputKeyEvent(FKismetCompilerContext& CompilerContext, UEdGraph* SourceGraph, UEdGraphPin* InputKeyPin, const EInputEvent KeyEvent)
+void UK2Node_InputKey::ExpandNode(FKismetCompilerContext& CompilerContext, UEdGraph* SourceGraph)
 {
-	if (InputKeyPin->LinkedTo.Num() > 0)
+	Super::ExpandNode(CompilerContext, SourceGraph);
+
+	UEdGraphPin* InputKeyPressedPin = GetPressedPin();
+	UEdGraphPin* InputKeyReleasedPin = GetReleasedPin();
+		
+	struct EventPinData
 	{
-		UK2Node_InputKeyEvent* InputKeyEvent = CompilerContext.SpawnIntermediateNode<UK2Node_InputKeyEvent>(this, SourceGraph);
-		const FName ModifierName = GetModifierName();
-		if ( ModifierName != NAME_None )
-		{
-			InputKeyEvent->CustomFunctionName = FName( *FString::Printf(TEXT("InpActEvt_%s_%s_%s"), *ModifierName.ToString(), *InputKey.ToString(), *InputKeyEvent->GetName()));
-		}
-		else
-		{
-			InputKeyEvent->CustomFunctionName = FName( *FString::Printf(TEXT("InpActEvt_%s_%s"), *InputKey.ToString(), *InputKeyEvent->GetName()));
-		}
-		InputKeyEvent->InputChord.Key = InputKey;
-		InputKeyEvent->InputChord.bCtrl = bControl;
-		InputKeyEvent->InputChord.bAlt = bAlt;
-		InputKeyEvent->InputChord.bShift = bShift;
-		InputKeyEvent->InputChord.bCmd = bCommand;
-		InputKeyEvent->bConsumeInput = bConsumeInput;
-		InputKeyEvent->bExecuteWhenPaused = bExecuteWhenPaused;
-		InputKeyEvent->bOverrideParentBinding = bOverrideParentBinding;
-		InputKeyEvent->InputKeyEvent = KeyEvent;
-		InputKeyEvent->EventReference.SetExternalDelegateMember(FName(TEXT("InputActionHandlerDynamicSignature__DelegateSignature")));
-		InputKeyEvent->bInternalEvent = true;
-		InputKeyEvent->AllocateDefaultPins();
-
-		// Move any exec links from the InputActionNode pin to the InputActionEvent node
-		UEdGraphPin* EventOutput = CompilerContext.GetSchema()->FindExecutionPin(*InputKeyEvent, EGPD_Output);
-
-		if(EventOutput != NULL)
+		EventPinData(UEdGraphPin* InPin,TEnumAsByte<EInputEvent> InEvent ){	Pin=InPin;EventType=InEvent;};
+		UEdGraphPin* Pin;
+		TEnumAsByte<EInputEvent> EventType;
+	};
+
+	TArray<EventPinData> ActivePins;
+	if(( InputKeyPressedPin != nullptr ) && (InputKeyPressedPin->LinkedTo.Num() > 0 ))
+	{
+		ActivePins.Add(EventPinData(InputKeyPressedPin,IE_Pressed));
+	}
+	if((InputKeyReleasedPin != nullptr) && (InputKeyReleasedPin->LinkedTo.Num() > 0 ))
+	{
+		ActivePins.Add(EventPinData(InputKeyReleasedPin,IE_Released));
+	}
+	
+	const UEdGraphSchema_K2* Schema = CompilerContext.GetSchema();
+
+	// If more than one is linked we have to do more complicated behaviors
+	if( ActivePins.Num() > 1 )
+	{
+		// Create a temporary variable to copy Key in to
+		static UScriptStruct* KeyStruct = FKey::StaticStruct();
+		UK2Node_TemporaryVariable* KeyVar = CompilerContext.SpawnIntermediateNode<UK2Node_TemporaryVariable>(this, SourceGraph);
+		KeyVar->VariableType.PinCategory = Schema->PC_Struct;
+		KeyVar->VariableType.PinSubCategoryObject = KeyStruct;
+		KeyVar->AllocateDefaultPins();
+
+		for (auto PinIt = ActivePins.CreateIterator(); PinIt; ++PinIt)
+		{			
+			UEdGraphPin *EachPin = (*PinIt).Pin;
+			// Create the input touch event
+			UK2Node_InputKeyEvent* InputKeyEvent = CompilerContext.SpawnIntermediateNode<UK2Node_InputKeyEvent>(this, SourceGraph);
+			const FName ModifierName = GetModifierName();
+			if ( ModifierName != NAME_None )
+			{
+				InputKeyEvent->CustomFunctionName = FName( *FString::Printf(TEXT("InpActEvt_%s_%s_%s"), *ModifierName.ToString(), *InputKey.ToString(), *InputKeyEvent->GetName()));
+			}
+			else
+			{
+				InputKeyEvent->CustomFunctionName = FName( *FString::Printf(TEXT("InpActEvt_%s_%s"), *InputKey.ToString(), *InputKeyEvent->GetName()));
+			}
+			InputKeyEvent->InputChord.Key = InputKey;
+			InputKeyEvent->InputChord.bCtrl = bControl;
+			InputKeyEvent->InputChord.bAlt = bAlt;
+			InputKeyEvent->InputChord.bShift = bShift;
+			InputKeyEvent->InputChord.bCmd = bCommand;
+			InputKeyEvent->bConsumeInput = bConsumeInput;
+			InputKeyEvent->bExecuteWhenPaused = bExecuteWhenPaused;
+			InputKeyEvent->bOverrideParentBinding = bOverrideParentBinding;
+			InputKeyEvent->InputKeyEvent = (*PinIt).EventType;
+			InputKeyEvent->EventReference.SetExternalDelegateMember(FName(TEXT("InputActionHandlerDynamicSignature__DelegateSignature")));
+			InputKeyEvent->bInternalEvent = true;
+			InputKeyEvent->AllocateDefaultPins();
+
+			// Create assignment nodes to assign the key
+			UK2Node_AssignmentStatement* KeyInitialize = CompilerContext.SpawnIntermediateNode<UK2Node_AssignmentStatement>(this, SourceGraph);
+			KeyInitialize->AllocateDefaultPins();
+			Schema->TryCreateConnection(KeyVar->GetVariablePin(), KeyInitialize->GetVariablePin());
+			Schema->TryCreateConnection(KeyInitialize->GetValuePin(), InputKeyEvent->FindPinChecked(TEXT("Key")));
+			// Connect the events to the assign key nodes
+			Schema->TryCreateConnection(Schema->FindExecutionPin(*InputKeyEvent, EGPD_Output), KeyInitialize->GetExecPin());
+
+			// Move the original event connections to the then pin of the key assign
+			CompilerContext.MovePinLinksToIntermediate(*EachPin, *KeyInitialize->GetThenPin());
+			
+			// Move the original event variable connections to the intermediate nodes
+			CompilerContext.MovePinLinksToIntermediate(*FindPin(TEXT("Key")), *KeyVar->GetVariablePin());
+		}	
+	}
+	else if( ActivePins.Num() == 1 )
+	{
+		UEdGraphPin* InputKeyPin = ActivePins[0].Pin;
+		EInputEvent InputEvent = ActivePins[0].EventType;
+	
+		if (InputKeyPin->LinkedTo.Num() > 0)
 		{
-			CompilerContext.MovePinLinksToIntermediate(*InputKeyPin, *EventOutput);
+			UK2Node_InputKeyEvent* InputKeyEvent = CompilerContext.SpawnIntermediateNode<UK2Node_InputKeyEvent>(this, SourceGraph);
+			const FName ModifierName = GetModifierName();
+			if ( ModifierName != NAME_None )
+			{
+				InputKeyEvent->CustomFunctionName = FName( *FString::Printf(TEXT("InpActEvt_%s_%s_%s"), *ModifierName.ToString(), *InputKey.ToString(), *InputKeyEvent->GetName()));
+			}
+			else
+			{
+				InputKeyEvent->CustomFunctionName = FName( *FString::Printf(TEXT("InpActEvt_%s_%s"), *InputKey.ToString(), *InputKeyEvent->GetName()));
+			}
+			InputKeyEvent->InputChord.Key = InputKey;
+			InputKeyEvent->InputChord.bCtrl = bControl;
+			InputKeyEvent->InputChord.bAlt = bAlt;
+			InputKeyEvent->InputChord.bShift = bShift;
+			InputKeyEvent->InputChord.bCmd = bCommand;
+			InputKeyEvent->bConsumeInput = bConsumeInput;
+			InputKeyEvent->bExecuteWhenPaused = bExecuteWhenPaused;
+			InputKeyEvent->bOverrideParentBinding = bOverrideParentBinding;
+			InputKeyEvent->InputKeyEvent = InputEvent;
+			InputKeyEvent->EventReference.SetExternalDelegateMember(FName(TEXT("InputActionHandlerDynamicSignature__DelegateSignature")));
+			InputKeyEvent->bInternalEvent = true;
+			InputKeyEvent->AllocateDefaultPins();
+
+			CompilerContext.MovePinLinksToIntermediate(*InputKeyPin, *Schema->FindExecutionPin(*InputKeyEvent, EGPD_Output));
+			CompilerContext.MovePinLinksToIntermediate(*FindPin(TEXT("Key")), *InputKeyEvent->FindPin(TEXT("Key")));
 		}
 	}
 }
 
-void UK2Node_InputKey::ExpandNode(FKismetCompilerContext& CompilerContext, UEdGraph* SourceGraph)
-{
-	Super::ExpandNode(CompilerContext, SourceGraph);
-
-	CreateInputKeyEvent(CompilerContext, SourceGraph, GetPressedPin(), IE_Pressed);
-	CreateInputKeyEvent(CompilerContext, SourceGraph, GetReleasedPin(), IE_Released);
-}
-
 void UK2Node_InputKey::GetMenuActions(FBlueprintActionDatabaseRegistrar& ActionRegistrar) const
 {
 	TArray<FKey> AllKeys;
diff --git a/Engine/Source/Runtime/Engine/Classes/Components/InputComponent.h b/Engine/Source/Runtime/Engine/Classes/Components/InputComponent.h
--- a/Engine/Source/Runtime/Engine/Classes/Components/InputComponent.h
+++ b/Engine/Source/Runtime/Engine/Classes/Components/InputComponent.h
@@ -64,11 +64,6 @@ struct TInputUnifiedDelegate
 		FuncDynDelegate.Unbind();
 	}
 
-	inline bool operator==(const TInputUnifiedDelegate<DelegateType, DynamicDelegateType>& Other) const
-	{
-		return ( FuncDelegate.IsBound() && (FuncDelegate == Other.FuncDelegate) ) || ( FuncDynDelegate.IsBound() && (FuncDynDelegate == Other.FuncDynDelegate) );
-	}
-
 protected:
 	/** Holds the delegate to call. */
 	DelegateType FuncDelegate;
@@ -94,23 +89,113 @@ struct FInputBinding
 
 /** Delegate signature for action events. */
 DECLARE_DELEGATE( FInputActionHandlerSignature );
-DECLARE_DYNAMIC_DELEGATE( FInputActionHandlerDynamicSignature );
+DECLARE_DELEGATE_OneParam( FInputActionHandlerWithKeySignature, FKey );
+DECLARE_DYNAMIC_DELEGATE_OneParam( FInputActionHandlerDynamicSignature, FKey, Key );
 
-/** Unified delegate specialization for Input Actions. */
-struct FInputActionUnifiedDelegate : public TInputUnifiedDelegate<FInputActionHandlerSignature, FInputActionHandlerDynamicSignature>
+struct FInputActionUnifiedDelegate
 {
+	FInputActionUnifiedDelegate() {};
+	FInputActionUnifiedDelegate(FInputActionHandlerSignature const& D) : FuncDelegate(D) {};
+	FInputActionUnifiedDelegate(FInputActionHandlerWithKeySignature const& D) : FuncDelegateWithKey(D) {};
+	FInputActionUnifiedDelegate(FInputActionHandlerDynamicSignature const& D) : FuncDynDelegate(D) {};
+
+	/** Returns if either the native or dynamic delegate is bound */
+	inline bool IsBound() const
+	{
+		return ( FuncDelegate.IsBound() || FuncDelegateWithKey.IsBound() || FuncDynDelegate.IsBound() );
+	}
+
+	/** Returns if either the native or dynamic delegate is bound to an object */
+	inline bool IsBoundToObject(void const* Object) const
+	{
+		if (FuncDelegate.IsBound())
+		{
+			return FuncDelegate.IsBoundToObject(Object);
+		}
+		else if (FuncDelegateWithKey.IsBound())
+		{
+			return FuncDelegateWithKey.IsBoundToObject(Object);
+		}
+		else if (FuncDynDelegate.IsBound())
+		{
+			return FuncDynDelegate.IsBoundToObject(Object);
+		}
+
+		return false;
+	}
+
+	/** Binds a native delegate and unbinds any bound dynamic delegate */
+	template< class UserClass >
+	inline void BindDelegate(UserClass* Object, typename FInputActionHandlerSignature::template TUObjectMethodDelegate< UserClass >::FMethodPtr Func)
+	{
+		FuncDynDelegate.Unbind();
+		FuncDelegateWithKey.Unbind();
+		FuncDelegate.BindUObject(Object, Func);
+	}
+
+	template< class UserClass >
+	inline void BindDelegate(UserClass* Object, typename FInputActionHandlerWithKeySignature::template TUObjectMethodDelegate< UserClass >::FMethodPtr Func)
+	{
+		FuncDynDelegate.Unbind();
+		FuncDelegate.Unbind();
+		FuncDelegateWithKey.BindUObject(Object, Func);
+	}
+
+	/** Binds a dynamic delegate and unbinds any bound native delegate */
+	inline void BindDelegate(UObject* Object, const FName FuncName)
+	{
+		FuncDelegate.Unbind();
+		FuncDelegateWithKey.Unbind();
+		FuncDynDelegate.BindUFunction(Object, FuncName);
+	}
+
+	/** Returns a reference to the native delegate and unbinds any bound dynamic delegate */
+	FInputActionHandlerSignature& GetDelegateForManualSet()
+	{
+		FuncDynDelegate.Unbind();
+		FuncDelegateWithKey.Unbind();
+		return FuncDelegate;
+	}
+
+	/** Returns a reference to the native delegate and unbinds any bound dynamic delegate */
+	FInputActionHandlerWithKeySignature& GetDelegateWithKeyForManualSet()
+	{
+		FuncDynDelegate.Unbind();
+		FuncDelegate.Unbind();
+		return FuncDelegateWithKey;
+	}
+
+	/** Unbinds any bound delegates */
+	inline void Unbind()
+	{
+		FuncDelegate.Unbind();
+		FuncDelegateWithKey.Unbind();
+		FuncDynDelegate.Unbind();
+	}
+
 	/** Execute function for the action unified delegate. */
-	inline void Execute() const
+	inline void Execute(const FKey Key) const
 	{
 		if (FuncDelegate.IsBound())
 		{
 			FuncDelegate.Execute();
 		}
+		else if (FuncDelegateWithKey.IsBound())
+		{
+			FuncDelegateWithKey.Execute(Key);
+		}
 		else if (FuncDynDelegate.IsBound())
 		{
-			FuncDynDelegate.Execute();
+			FuncDynDelegate.Execute(Key);
 		}
 	}
+protected:
+	/** Holds the delegate to call. */
+	FInputActionHandlerSignature FuncDelegate;
+	/** Holds the delegate that wants to know the key to call. */
+	FInputActionHandlerWithKeySignature FuncDelegateWithKey;
+	/** Holds the dynamic delegate to call. */
+	FInputActionHandlerDynamicSignature FuncDynDelegate;
 };
 
 /** Binds a delegate to an action. */
diff --git a/Engine/Source/Runtime/Engine/Classes/GameFramework/PlayerInput.h b/Engine/Source/Runtime/Engine/Classes/GameFramework/PlayerInput.h
--- a/Engine/Source/Runtime/Engine/Classes/GameFramework/PlayerInput.h
+++ b/Engine/Source/Runtime/Engine/Classes/GameFramework/PlayerInput.h
@@ -520,13 +520,24 @@ private:
 
 	/* Collects the chords and the delegates they invoke for an action binding
 	 * @param ActionBinding - the action to determine whether it occurred
+	 * @param bGamePaused - whether the game is currently paused
 	 * @param FoundChords - the list of chord/delegate pairs to add to
 	 * @param KeysToConsume - array to collect the keys associated with this binding that should be consumed
 	 */
 	void GetChordsForAction(const FInputActionBinding& ActionBinding, const bool bGamePaused, TArray<struct FDelegateDispatchDetails>& FoundChords, TArray<FKey>& KeysToConsume);
 
+	/* Helper function for GetChordsForAction to examine each keymapping that belongs to the ActionBinding
+	 * @param KeyMapping - the key mapping to determine whether it occured
+	 * @param ActionBinding - the action to determine whether it occurred
+	 * @param bGamePaused - whether the game is currently paused
+	 * @param FoundChords - the list of chord/delegate pairs to add to
+	 * @param KeysToConsume - array to collect the keys associated with this binding that should be consumed
+	 */
+	void GetChordsForKeyMapping(const FInputActionKeyMapping& KeyMapping, const FInputActionBinding& ActionBinding, const bool bGamePaused, TArray<FDelegateDispatchDetails>& FoundChords, TArray<FKey>& KeysToConsume);
+
 	/* Collects the chords and the delegates they invoke for a key binding
 	 * @param KeyBinding - the key to determine whether it occurred
+	 * @param bGamePaused - whether the game is currently paused
 	 * @param FoundChords - the list of chord/delegate pairs to add to
 	 * @param KeysToConsume - array to collect the keys associated with this binding that should be consumed
 	 */
diff --git a/Engine/Source/Runtime/Engine/Private/UserInterface/PlayerInput.cpp b/Engine/Source/Runtime/Engine/Private/UserInterface/PlayerInput.cpp
--- a/Engine/Source/Runtime/Engine/Private/UserInterface/PlayerInput.cpp
+++ b/Engine/Source/Runtime/Engine/Private/UserInterface/PlayerInput.cpp
@@ -598,6 +598,62 @@ void UPlayerInput::ConditionalBuildKeyMappings()
 	}
 }
 
+void UPlayerInput::GetChordsForKeyMapping(const FInputActionKeyMapping& KeyMapping, const FInputActionBinding& ActionBinding, const bool bGamePaused, TArray<FDelegateDispatchDetails>& FoundChords, TArray<FKey>& KeysToConsume)
+{
+	TArray<uint32> EventIndices;
+
+	// test modifier conditions and ignore the event if they failed
+	if (	(KeyMapping.bAlt == false || IsAltPressed())
+		&&	(KeyMapping.bCtrl == false || IsCtrlPressed())
+		&&	(KeyMapping.bShift == false || IsShiftPressed())
+		&&	(KeyMapping.bCmd == false || IsCmdPressed())
+		&& 	KeyEventOccurred(KeyMapping.Key, ActionBinding.KeyEvent, EventIndices))
+	{
+		bool bAddDelegate = true;
+
+		// look through the found chords and determine if this is masked (or masks) anything in the array
+		const FInputChord Chord(KeyMapping.Key, KeyMapping.bShift, KeyMapping.bCtrl, KeyMapping.bAlt, KeyMapping.bCmd);
+		for (int32 ChordIndex = FoundChords.Num() - 1; ChordIndex >= 0; --ChordIndex)
+		{
+			FInputChord::ERelationshipType ChordRelationship = Chord.GetRelationship(FoundChords[ChordIndex].Chord);
+
+			if (ChordRelationship == FInputChord::Masks)
+			{
+				// If we mask the found one, then remove it from the list
+				FoundChords.RemoveAtSwap(ChordIndex);
+			}
+			else if (ChordRelationship == FInputChord::Masked)
+			{
+				bAddDelegate = false;
+				break;
+			}
+		}
+
+		if (bAddDelegate)
+		{
+			check(EventIndices.Num() > 0);
+			FDelegateDispatchDetails FoundChord(  
+										EventIndices[0]
+										, FoundChords.Num()
+										, Chord
+										, ((!bGamePaused || ActionBinding.bExecuteWhenPaused) ? ActionBinding.ActionDelegate : FInputActionUnifiedDelegate())
+										, ActionBinding.KeyEvent
+										, &ActionBinding);
+			FoundChords.Add(FoundChord);
+
+			for (int32 EventsIndex = 1; EventsIndex < EventIndices.Num(); ++EventsIndex)
+			{
+				FoundChord.EventIndex = EventIndices[EventsIndex];
+				FoundChords.Add(FoundChord);
+			}
+			if (ActionBinding.bConsumeInput)
+			{
+				KeysToConsume.AddUnique(KeyMapping.Key);
+			}
+		}
+	}
+}
+
 void UPlayerInput::GetChordsForAction(const FInputActionBinding& ActionBinding, const bool bGamePaused, TArray<FDelegateDispatchDetails>& FoundChords, TArray<FKey>& KeysToConsume)
 {
 	ConditionalBuildKeyMappings();
@@ -605,64 +661,24 @@ void UPlayerInput::GetChordsForAction(const FInputActionBinding& ActionBinding,
 	FActionKeyDetails* KeyDetails = ActionKeyMap.Find(ActionBinding.ActionName);
 	if (KeyDetails)
 	{
-		for (int32 ActionIndex = 0; ActionIndex < KeyDetails->Actions.Num(); ++ActionIndex)
+		for (const FInputActionKeyMapping& KeyMapping : KeyDetails->Actions)
 		{
-			const FInputActionKeyMapping& KeyMapping = KeyDetails->Actions[ActionIndex];
-			if ( !IsKeyConsumed(KeyMapping.Key) )
+			if (KeyMapping.Key == EKeys::AnyKey)
 			{
-				TArray<uint32> EventIndices;
-
-				// test modifier conditions and ignore the event if they failed
-				if (	(KeyMapping.bAlt == false || IsAltPressed())
-					&&	(KeyMapping.bCtrl == false || IsCtrlPressed())
-					&&	(KeyMapping.bShift == false || IsShiftPressed())
-					&&	(KeyMapping.bCmd == false || IsCmdPressed())
-					&& 	KeyEventOccurred(KeyMapping.Key, ActionBinding.KeyEvent, EventIndices))
+				for (auto KeyStateIt(KeyStateMap.CreateConstIterator()); KeyStateIt; ++KeyStateIt)
 				{
-					bool bAddDelegate = true;
-
-					// look through the found chords and determine if this is masked (or masks) anything in the array
-					const FInputChord Chord(KeyMapping.Key, KeyMapping.bShift, KeyMapping.bCtrl, KeyMapping.bAlt, KeyMapping.bCmd);
-					for (int32 ChordIndex = FoundChords.Num() - 1; ChordIndex >= 0; --ChordIndex)
-					{
-						FInputChord::ERelationshipType ChordRelationship = Chord.GetRelationship(FoundChords[ChordIndex].Chord);
-
-						if (ChordRelationship == FInputChord::Masks)
-						{
-							// If we mask the found one, then remove it from the list
-							FoundChords.RemoveAtSwap(ChordIndex);
-						}
-						else if (ChordRelationship == FInputChord::Masked)
-						{
-							bAddDelegate = false;
-							break;
-						}
-					}
-
-					if (bAddDelegate)
+					if (!KeyStateIt.Key().IsFloatAxis() && !KeyStateIt.Key().IsVectorAxis())
 					{
-						check(EventIndices.Num() > 0);
-						FDelegateDispatchDetails FoundChord(  
-													EventIndices[0]
-													, FoundChords.Num()
-													, Chord
-													, ((!bGamePaused || ActionBinding.bExecuteWhenPaused) ? ActionBinding.ActionDelegate : FInputActionUnifiedDelegate())
-													, ActionBinding.KeyEvent
-													, &ActionBinding);
-						FoundChords.Add(FoundChord);
-
-						for (int32 EventsIndex = 1; EventsIndex < EventIndices.Num(); ++EventsIndex)
-						{
-							FoundChord.EventIndex = EventIndices[EventsIndex];
-							FoundChords.Add(FoundChord);
-						}
-						if (ActionBinding.bConsumeInput)
-						{
-							KeysToConsume.AddUnique(KeyMapping.Key);
-						}
+						FInputActionKeyMapping SubKeyMapping(KeyMapping);
+						SubKeyMapping.Key = KeyStateIt.Key();
+						GetChordsForKeyMapping(KeyMapping, ActionBinding, bGamePaused, FoundChords, KeysToConsume);
 					}
 				}
 			}
+			else if ( !IsKeyConsumed(KeyMapping.Key) )
+			{
+				GetChordsForKeyMapping(KeyMapping, ActionBinding, bGamePaused, FoundChords, KeysToConsume);
+			}
 		}
 	}
 }
@@ -671,7 +687,19 @@ void UPlayerInput::GetChordForKey(const FInputKeyBinding& KeyBinding, const bool
 {
 	bool bKeyOccurred = false;
 
-	if ( !IsKeyConsumed(KeyBinding.Chord.Key) )
+	if (KeyBinding.Chord.Key == EKeys::AnyKey)
+	{
+		for (auto KeyStateIt(KeyStateMap.CreateConstIterator()); KeyStateIt; ++KeyStateIt)
+		{
+			if (!KeyStateIt.Key().IsFloatAxis() && !KeyStateIt.Key().IsVectorAxis())
+			{
+				FInputKeyBinding SubKeyBinding(KeyBinding);
+				SubKeyBinding.Chord.Key = KeyStateIt.Key();
+				GetChordForKey(SubKeyBinding, bGamePaused, FoundChords, KeysToConsume);
+			}
+		}
+	}
+	else if ( !IsKeyConsumed(KeyBinding.Chord.Key) )
 	{
 		TArray<uint32> EventIndices;
 
@@ -1079,7 +1107,7 @@ void UPlayerInput::ProcessInputStack(const TArray<UInputComponent*>& InputCompon
 	{
 		if (Details.ActionDelegate.IsBound())
 		{
-			Details.ActionDelegate.Execute();
+			Details.ActionDelegate.Execute(Details.Chord.Key);
 		}
 		else if (Details.TouchDelegate.IsBound())
 		{
diff --git a/Engine/Source/Runtime/InputCore/Classes/InputCoreTypes.h b/Engine/Source/Runtime/InputCore/Classes/InputCoreTypes.h
--- a/Engine/Source/Runtime/InputCore/Classes/InputCoreTypes.h
+++ b/Engine/Source/Runtime/InputCore/Classes/InputCoreTypes.h
@@ -164,6 +164,8 @@ namespace EConsoleForGamepadLabels
 }
 struct INPUTCORE_API EKeys
 {
+	static const FKey AnyKey;
+
 	static const FKey MouseX;
 	static const FKey MouseY;
 	static const FKey MouseScrollUp;
diff --git a/Engine/Source/Runtime/InputCore/Private/InputCoreTypes.cpp b/Engine/Source/Runtime/InputCore/Private/InputCoreTypes.cpp
--- a/Engine/Source/Runtime/InputCore/Private/InputCoreTypes.cpp
+++ b/Engine/Source/Runtime/InputCore/Private/InputCoreTypes.cpp
@@ -8,6 +8,8 @@ DEFINE_LOG_CATEGORY(LogInput);
 
 #define LOCTEXT_NAMESPACE "InputKeys"
 
+const FKey EKeys::AnyKey("AnyKey");
+
 const FKey EKeys::MouseX("MouseX");
 const FKey EKeys::MouseY("MouseY");
 const FKey EKeys::MouseScrollUp("MouseScrollUp");
@@ -357,6 +359,8 @@ void EKeys::Initialize()
 	AddMenuCategoryDisplayInfo(NAME_MouseCategory, LOCTEXT("MouseSubCategory", "Mouse"), TEXT("GraphEditor.MouseEvent_16x"));
 	AddMenuCategoryDisplayInfo(NAME_KeyboardCategory, LOCTEXT("KeyboardSubCategory", "Keyboard"), TEXT("GraphEditor.KeyEvent_16x"));
 
+	AddKey(FKeyDetails(EKeys::AnyKey, LOCTEXT("AnyKey", "Any Key")));
+
 	AddKey(FKeyDetails(EKeys::MouseX, LOCTEXT("MouseX", "Mouse X"), FKeyDetails::FloatAxis | FKeyDetails::MouseButton | FKeyDetails::UpdateAxisWithoutSamples));
 	AddKey(FKeyDetails(EKeys::MouseY, LOCTEXT("MouseY", "Mouse Y"), FKeyDetails::FloatAxis | FKeyDetails::MouseButton | FKeyDetails::UpdateAxisWithoutSamples));
 	AddKey(FKeyDetails(EKeys::MouseWheelAxis, LOCTEXT("MouseWheelAxis", "Mouse Wheel Axis"), FKeyDetails::FloatAxis | FKeyDetails::MouseButton | FKeyDetails::UpdateAxisWithoutSamples));
