[UE-18422] Ensure that GetClassDefaults nodes are refreshed/recompiled when dependent Blueprint class assets are changed.

change summary:
- added a  UK2Node_GetClassDefaults::HasExternalDependencies() override

diff --git a/Engine/Source/Editor/BlueprintGraph/Classes/K2Node_GetClassDefaults.h b/Engine/Source/Editor/BlueprintGraph/Classes/K2Node_GetClassDefaults.h
--- a/Engine/Source/Editor/BlueprintGraph/Classes/K2Node_GetClassDefaults.h
+++ b/Engine/Source/Editor/BlueprintGraph/Classes/K2Node_GetClassDefaults.h
@@ -25,6 +25,7 @@ class UK2Node_GetClassDefaults : public UK2Node
 	virtual bool IsNodePure() const override { return true; }
 	virtual bool ShouldShowNodeProperties() const override { return true; }
 	virtual void ReallocatePinsDuringReconstruction(TArray<UEdGraphPin*>& OldPins) override;
+	virtual bool HasExternalDependencies(TArray<class UStruct*>* OptionalOutput) const override;
 	virtual class FNodeHandlingFunctor* CreateNodeHandler(class FKismetCompilerContext& CompilerContext) const override;
 	virtual void GetMenuActions(FBlueprintActionDatabaseRegistrar& ActionRegistrar) const override;
 	virtual FText GetMenuCategory() const override;
diff --git a/Engine/Source/Editor/BlueprintGraph/Private/K2Node_GetClassDefaults.cpp b/Engine/Source/Editor/BlueprintGraph/Private/K2Node_GetClassDefaults.cpp
--- a/Engine/Source/Editor/BlueprintGraph/Private/K2Node_GetClassDefaults.cpp
+++ b/Engine/Source/Editor/BlueprintGraph/Private/K2Node_GetClassDefaults.cpp
@@ -210,6 +210,20 @@ void UK2Node_GetClassDefaults::PinDefaultValueChanged(UEdGraphPin* ChangedPin)
 	}
 }
 
+bool UK2Node_GetClassDefaults::HasExternalDependencies(TArray<class UStruct*>* OptionalOutput) const
+{
+	UClass* SourceClass = GetInputClass();
+	UBlueprint* SourceBlueprint = GetBlueprint();
+	const bool bResult = (SourceClass && (SourceClass->ClassGeneratedBy != SourceBlueprint));
+	if (bResult && OptionalOutput)
+	{
+		OptionalOutput->AddUnique(SourceClass);
+	}
+
+	const bool bSuperResult = Super::HasExternalDependencies(OptionalOutput);
+	return bSuperResult || bResult;
+}
+
 void UK2Node_GetClassDefaults::ReallocatePinsDuringReconstruction(TArray<UEdGraphPin*>& OldPins) 
 {
 	Super::ReallocatePinsDuringReconstruction(OldPins);
