[UE-18321] Class default array properties are now copied on access to avoid mutating the CDO at runtime.

change summary:
- added UK2Node_GetClassDefaults::ExpandNode() to spawn intermediate nodes for copying array outputs to a local temporary if linked.

diff --git a/Engine/Source/Editor/BlueprintGraph/Classes/K2Node_GetClassDefaults.h b/Engine/Source/Editor/BlueprintGraph/Classes/K2Node_GetClassDefaults.h
--- a/Engine/Source/Editor/BlueprintGraph/Classes/K2Node_GetClassDefaults.h
+++ b/Engine/Source/Editor/BlueprintGraph/Classes/K2Node_GetClassDefaults.h
@@ -27,6 +27,7 @@ class UK2Node_GetClassDefaults : public UK2Node
 	virtual void ReallocatePinsDuringReconstruction(TArray<UEdGraphPin*>& OldPins) override;
 	virtual bool HasExternalDependencies(TArray<class UStruct*>* OptionalOutput) const override;
 	virtual class FNodeHandlingFunctor* CreateNodeHandler(class FKismetCompilerContext& CompilerContext) const override;
+	virtual void ExpandNode(class FKismetCompilerContext& CompilerContext, UEdGraph* SourceGraph) override;
 	virtual void GetMenuActions(FBlueprintActionDatabaseRegistrar& ActionRegistrar) const override;
 	virtual FText GetMenuCategory() const override;
 	// End UK2Node interface
diff --git a/Engine/Source/Editor/BlueprintGraph/Private/K2Node_GetClassDefaults.cpp b/Engine/Source/Editor/BlueprintGraph/Private/K2Node_GetClassDefaults.cpp
--- a/Engine/Source/Editor/BlueprintGraph/Private/K2Node_GetClassDefaults.cpp
+++ b/Engine/Source/Editor/BlueprintGraph/Private/K2Node_GetClassDefaults.cpp
@@ -5,6 +5,8 @@
 #include "BlueprintNodeSpawner.h"
 #include "BlueprintActionDatabaseRegistrar.h"
 #include "KismetCompiler.h"
+#include "K2Node_TemporaryVariable.h"
+#include "K2Node_PureAssignmentStatement.h"
 #include "K2Node_GetClassDefaults.h"
 
 #define LOCTEXT_NAMESPACE "UK2Node_GetClassDefaults"
@@ -241,6 +243,36 @@ FNodeHandlingFunctor* UK2Node_GetClassDefaults::CreateNodeHandler(FKismetCompile
 	return new FKCHandler_GetClassDefaults(CompilerContext);
 }
 
+void UK2Node_GetClassDefaults::ExpandNode(class FKismetCompilerContext& CompilerContext, UEdGraph* SourceGraph)
+{
+	const UClass* ClassType = GetInputClass();
+
+	// @TODO - Remove if/when we support 'const' input pins.
+	// For array properties, return a local copy of the array so that the original cannot be modified.
+	for(UEdGraphPin* OutputPin : Pins)
+	{
+		if(OutputPin != nullptr && OutputPin->Direction == EGPD_Output && OutputPin->LinkedTo.Num() > 0)
+		{
+			UProperty* BoundProperty = FindField<UProperty>(ClassType, *(OutputPin->PinName));
+			if(BoundProperty != nullptr && BoundProperty->IsA<UArrayProperty>())
+			{
+				UK2Node_TemporaryVariable* LocalVariable = CompilerContext.SpawnIntermediateNode<UK2Node_TemporaryVariable>(this, SourceGraph);
+				LocalVariable->VariableType = OutputPin->PinType;
+				LocalVariable->VariableType.bIsReference = false;
+				LocalVariable->AllocateDefaultPins();
+
+				UK2Node_PureAssignmentStatement* CopyDefaultValue = CompilerContext.SpawnIntermediateNode<UK2Node_PureAssignmentStatement>(this, SourceGraph);
+				CopyDefaultValue->AllocateDefaultPins();
+				CompilerContext.GetSchema()->TryCreateConnection(LocalVariable->GetVariablePin(), CopyDefaultValue->GetVariablePin());
+
+				// Note: This must be done AFTER connecting the variable input, which sets the pin type.
+				CompilerContext.MovePinLinksToIntermediate(*OutputPin, *CopyDefaultValue->GetOutputPin());
+				CompilerContext.GetSchema()->TryCreateConnection(OutputPin, CopyDefaultValue->GetValuePin());
+			}
+		}
+	}
+}
+
 void UK2Node_GetClassDefaults::GetMenuActions(FBlueprintActionDatabaseRegistrar& ActionRegistrar) const
 {
 	UClass* ActionKey = GetClass();
