"AsyncTaskProxy" pin is shown in AsyncTask node only when the proxy class has "ExposedAsyncProxy" meta data.

diff --git a/Engine/Source/Editor/BlueprintGraph/Private/K2Node_BaseAsyncTask.cpp b/Engine/Source/Editor/BlueprintGraph/Private/K2Node_BaseAsyncTask.cpp
--- a/Engine/Source/Editor/BlueprintGraph/Private/K2Node_BaseAsyncTask.cpp
+++ b/Engine/Source/Editor/BlueprintGraph/Private/K2Node_BaseAsyncTask.cpp
@@ -55,7 +55,15 @@ void UK2Node_BaseAsyncTask::AllocateDefaultPins()
 	CreatePin(EGPD_Input, K2Schema->PC_Exec, TEXT(""), NULL, false, false, K2Schema->PN_Execute);
 	CreatePin(EGPD_Output, K2Schema->PC_Exec, TEXT(""), NULL, false, false, K2Schema->PN_Then);
 
-	CreatePin(EGPD_Output, K2Schema->PC_Object, TEXT(""), ProxyClass, false, false, FBaseAsyncTaskHelper::GetAsyncTaskProxyName());
+	bool bExposeProxy = false;
+	for (const UStruct* TestStruct = ProxyClass; TestStruct && !bExposeProxy; TestStruct = TestStruct->GetSuperStruct())
+	{
+		bExposeProxy = TestStruct->HasMetaData(TEXT("ExposedAsyncProxy"));
+	}
+	if (bExposeProxy)
+	{
+		CreatePin(EGPD_Output, K2Schema->PC_Object, TEXT(""), ProxyClass, false, false, FBaseAsyncTaskHelper::GetAsyncTaskProxyName());
+	}
 
 	UFunction* DelegateSignatureFunction = NULL;
 	for (TFieldIterator<UProperty> PropertyIt(ProxyClass, EFieldIteratorFlags::ExcludeSuper); PropertyIt; ++PropertyIt)
@@ -284,8 +292,8 @@ void UK2Node_BaseAsyncTask::ExpandNode(class FKismetCompilerContext& CompilerCon
 
 	UEdGraphPin* const ProxyObjectPin = CallCreateProxyObjectNode->GetReturnValuePin();
 	check(ProxyObjectPin);
-	auto OutputAsyncTaskProxy = FindPinChecked(FBaseAsyncTaskHelper::GetAsyncTaskProxyName());
-	bIsErrorFree &= CompilerContext.MovePinLinksToIntermediate(*OutputAsyncTaskProxy, *ProxyObjectPin).CanSafeConnect();
+	UEdGraphPin* OutputAsyncTaskProxy = FindPin(FBaseAsyncTaskHelper::GetAsyncTaskProxyName());
+	bIsErrorFree &= !OutputAsyncTaskProxy || CompilerContext.MovePinLinksToIntermediate(*OutputAsyncTaskProxy, *ProxyObjectPin).CanSafeConnect();
 
 	// GATHER OUTPUT PARAMETERS AND PAIR THEM WITH LOCAL VARIABLES
 	TArray<FBaseAsyncTaskHelper::FOutputPinAndLocalVariable> VariableOutputs;
diff --git a/Engine/Source/Runtime/CoreUObject/Public/UObject/ObjectBase.h b/Engine/Source/Runtime/CoreUObject/Public/UObject/ObjectBase.h
--- a/Engine/Source/Runtime/CoreUObject/Public/UObject/ObjectBase.h
+++ b/Engine/Source/Runtime/CoreUObject/Public/UObject/ObjectBase.h
@@ -983,6 +983,9 @@ namespace UM
 
 		//[ClassMetadata] Do not spawn an object of the class using Generic Create Object node in Blueprint. It makes sense only for a BluprintType class, that is neither Actor, nor ActorComponent.
 		DontUseGenericSpawnObject,
+
+		//[ClassMetadata] Expose a proxy object of this class in Async Task node.
+		ExposedAsyncProxy,
 	};
 
 	// Metadata usable in USTRUCT
diff --git a/Engine/Source/Runtime/GameplayTasks/Classes/GameplayTask.h b/Engine/Source/Runtime/GameplayTasks/Classes/GameplayTask.h
--- a/Engine/Source/Runtime/GameplayTasks/Classes/GameplayTask.h
+++ b/Engine/Source/Runtime/GameplayTasks/Classes/GameplayTask.h
@@ -130,7 +130,7 @@ public:
 	FString GetDebugDescription() const;
 };
 
-UCLASS(Abstract)
+UCLASS(Abstract, meta = (ExposedAsyncProxy))
 class GAMEPLAYTASKS_API UGameplayTask : public UObject
 {
 	GENERATED_BODY()
