[NO NON-PII COMMIT MESSAGE]

diff --git a/Engine/Source/Editor/BlueprintGraph/Private/K2Node_LoadAsset.cpp b/Engine/Source/Editor/BlueprintGraph/Private/K2Node_LoadAsset.cpp
--- a/Engine/Source/Editor/BlueprintGraph/Private/K2Node_LoadAsset.cpp
+++ b/Engine/Source/Editor/BlueprintGraph/Private/K2Node_LoadAsset.cpp
@@ -228,12 +228,12 @@ FName UK2Node_LoadAssetClass::NativeFunctionName() const
 
 FText UK2Node_LoadAssetClass::GetTooltipText() const
 {
-	return FText(LOCTEXT("UK2Node_LoadAssetClassGetTooltipText", "Async Load Asset Class"));
+	return FText(LOCTEXT("UK2Node_LoadAssetClassGetTooltipText", "Async Load Class Asset"));
 }
 
 FText UK2Node_LoadAssetClass::GetNodeTitle(ENodeTitleType::Type TitleType) const
 {
-	return FText(LOCTEXT("UK2Node_LoadAssetClassGetNodeTitle", "Load Asset Class"));
+	return FText(LOCTEXT("UK2Node_LoadAssetClassGetNodeTitle", "Load Class Asset"));
 }
 
 const FString& UK2Node_LoadAssetClass::GetInputCategory() const
