[UE-17240] - Changing "[REDACTED_HANDLE]" to "See:"in tooltips.

diff --git a/Engine/Source/Editor/BlueprintGraph/Private/K2Node_CallFunction.cpp b/Engine/Source/Editor/BlueprintGraph/Private/K2Node_CallFunction.cpp
--- a/Engine/Source/Editor/BlueprintGraph/Private/K2Node_CallFunction.cpp
+++ b/Engine/Source/Editor/BlueprintGraph/Private/K2Node_CallFunction.cpp
@@ -12,7 +12,7 @@
 #include "GraphEditorSettings.h"
 #include "BlueprintActionFilter.h"
 #include "Editor/Kismet/Public/FindInBlueprintManager.h"
-
+#pragma optimize("",off)
 #define LOCTEXT_NAMESPACE "K2Node"
 
 /*******************************************************************************
@@ -1366,10 +1366,12 @@ FString UK2Node_CallFunction::GetDefaultTooltipForFunction(const UFunction* Func
 		static const FString DoxygenParam(TEXT("@param"));
 		static const FString DoxygenReturn(TEXT("@return"));
 		static const FString DoxygenSee(TEXT("@see"));
+		static const FString TooltipSee(TEXT("See:"));
 
 		Tooltip.Split(DoxygenParam, &Tooltip, nullptr, ESearchCase::IgnoreCase, ESearchDir::FromStart);
 		Tooltip.Split(DoxygenReturn, &Tooltip, nullptr, ESearchCase::IgnoreCase, ESearchDir::FromStart);
-		Tooltip.Split(DoxygenSee, &Tooltip, nullptr, ESearchCase::IgnoreCase, ESearchDir::FromStart);
+		Tooltip.ReplaceInline(*DoxygenSee, *TooltipSee);
+
 		Tooltip.Trim();
 		Tooltip.TrimTrailing();
 
diff --git a/Engine/Source/Runtime/CoreUObject/Private/UObject/Class.cpp b/Engine/Source/Runtime/CoreUObject/Private/UObject/Class.cpp
--- a/Engine/Source/Runtime/CoreUObject/Private/UObject/Class.cpp
+++ b/Engine/Source/Runtime/CoreUObject/Private/UObject/Class.cpp
@@ -220,7 +220,8 @@ FText UField::GetToolTipText(bool bShortTooltip) const
 		else
 		{
 			static const FString DoxygenSee(TEXT("@see"));
-			if (NativeToolTip.Split(DoxygenSee, &NativeToolTip, nullptr, ESearchCase::IgnoreCase, ESearchDir::FromStart))
+			static const FString TooltipSee(TEXT("See:"));
+			if (NativeToolTip.ReplaceInline(*DoxygenSee, *TooltipSee) > 0)
 			{
 				NativeToolTip.TrimTrailing();
 			}
diff --git a/Engine/Source/Runtime/CoreUObject/Private/UObject/Enum.cpp b/Engine/Source/Runtime/CoreUObject/Private/UObject/Enum.cpp
--- a/Engine/Source/Runtime/CoreUObject/Private/UObject/Enum.cpp
+++ b/Engine/Source/Runtime/CoreUObject/Private/UObject/Enum.cpp
@@ -494,7 +494,8 @@ FText UEnum::GetToolTipText(int32 NameIndex) const
 	if ( !FText::FindText( Namespace, Key, /*OUT*/LocalizedToolTip, &NativeToolTip ) )
 	{
 		static const FString DoxygenSee(TEXT("@see"));
-		if (NativeToolTip.Split(DoxygenSee, &NativeToolTip, nullptr, ESearchCase::IgnoreCase, ESearchDir::FromStart))
+		static const FString TooltipSee(TEXT("See:"));
+		if (NativeToolTip.ReplaceInline(*DoxygenSee, *TooltipSee) > 0)
 		{
 			NativeToolTip.TrimTrailing();
 		}
