Removing pragma optimize.

diff --git a/Engine/Source/Editor/BlueprintGraph/Private/K2Node_CallFunction.cpp b/Engine/Source/Editor/BlueprintGraph/Private/K2Node_CallFunction.cpp
--- a/Engine/Source/Editor/BlueprintGraph/Private/K2Node_CallFunction.cpp
+++ b/Engine/Source/Editor/BlueprintGraph/Private/K2Node_CallFunction.cpp
@@ -12,7 +12,7 @@
 #include "GraphEditorSettings.h"
 #include "BlueprintActionFilter.h"
 #include "Editor/Kismet/Public/FindInBlueprintManager.h"
-#pragma optimize("",off)
+
 #define LOCTEXT_NAMESPACE "K2Node"
 
 /*******************************************************************************
