Making sure the BlueprintActionDatabase gets properly registered for all Blueprint updates (in case a specific node type added actions for that Blueprint already).

diff --git a/Engine/Source/Editor/BlueprintGraph/Private/BlueprintActionDatabase.cpp b/Engine/Source/Editor/BlueprintGraph/Private/BlueprintActionDatabase.cpp
--- a/Engine/Source/Editor/BlueprintGraph/Private/BlueprintActionDatabase.cpp
+++ b/Engine/Source/Editor/BlueprintGraph/Private/BlueprintActionDatabase.cpp
@@ -1222,7 +1222,6 @@ void FBlueprintActionDatabase::RefreshAssetActions(UObject* const AssetObject)
 {
 	using namespace BlueprintActionDatabaseImpl;
 
-	bool const bHadExistingEntry = ActionRegistry.Contains(AssetObject);
 	FActionList& AssetActionList = ActionRegistry.FindOrAdd(AssetObject);
 	for (UBlueprintNodeSpawner* Action : AssetActionList)
 	{
@@ -1257,12 +1256,17 @@ void FBlueprintActionDatabase::RefreshAssetActions(UObject* const AssetObject)
 			AddAnimBlueprintGraphActions( AnimBlueprint, AssetActionList );
 		}
 
+		UBlueprint::FChangedEvent& OnBPChanged = BlueprintAsset->OnChanged();
+		UBlueprint::FCompiledEvent& OnBPCompiled = BlueprintAsset->OnCompiled();
 		// have to be careful not to register this callback twice for the 
 		// blueprint
-		if (!bHadExistingEntry)
+		if (!OnBPChanged.IsBoundToObject(this))
 		{
-			BlueprintAsset->OnChanged().AddRaw(this, &FBlueprintActionDatabase::OnBlueprintChanged);
-			BlueprintAsset->OnCompiled().AddRaw(this, &FBlueprintActionDatabase::OnBlueprintChanged);
+			OnBPChanged.AddRaw(this, &FBlueprintActionDatabase::OnBlueprintChanged);
+		}
+		if (!OnBPCompiled.IsBoundToObject(this))
+		{
+			OnBPCompiled.AddRaw(this, &FBlueprintActionDatabase::OnBlueprintChanged);
 		}
 	}
 
