[UE-17698] Better protecting ourselves from BlueprintNodeSpawners being deleted out from under the BlueprintActionDatabase.

diff --git a/Engine/Source/Editor/BlueprintGraph/Private/BlueprintActionDatabase.cpp b/Engine/Source/Editor/BlueprintGraph/Private/BlueprintActionDatabase.cpp
--- a/Engine/Source/Editor/BlueprintGraph/Private/BlueprintActionDatabase.cpp
+++ b/Engine/Source/Editor/BlueprintGraph/Private/BlueprintActionDatabase.cpp
@@ -1009,6 +1009,15 @@ void FBlueprintActionDatabase::AddReferencedObjects(FReferenceCollector& Collect
 	{
 		Collector.AddReferencedObjects(ActionListIt.Value);
 	}
+
+	// shouldn't have to do this, as the elements listed here should also be 
+	// accounted for in the regular ActionRegistry, but just in case we fail to 
+	// remove an element from here when we should.... this'll make sure these 
+	// elements stick around (so we don't crash in ClearUnloadedAssetActions)
+	for (auto& UnloadedActionListIt : UnloadedActionRegistry)
+	{
+		Collector.AddReferencedObjects(UnloadedActionListIt.Value);
+	}
 }
 
 //------------------------------------------------------------------------------
