Fix fallout from CL# 2635959

diff --git a/Engine/Source/Editor/BlueprintGraph/Private/EdGraphSchema_K2.cpp b/Engine/Source/Editor/BlueprintGraph/Private/EdGraphSchema_K2.cpp
--- a/Engine/Source/Editor/BlueprintGraph/Private/EdGraphSchema_K2.cpp
+++ b/Engine/Source/Editor/BlueprintGraph/Private/EdGraphSchema_K2.cpp
@@ -2993,7 +2993,7 @@ void UEdGraphSchema_K2::CreateFunctionGraphTerminators(UEdGraph& Graph, UClass*
 	if (InterfaceToImplement)
 	{
 		// Add modifier flags from the declaration
-		EntryNode->SetExtraFlags(EntryNode->GetFunctionFlags() | InterfaceToImplement->FunctionFlags & (FUNC_Const | FUNC_Static | FUNC_BlueprintPure));
+		EntryNode->SetExtraFlags(EntryNode->GetFunctionFlags() | (InterfaceToImplement->FunctionFlags & (FUNC_Const | FUNC_Static | FUNC_BlueprintPure)));
 
 		UK2Node* NextNode = EntryNode;
 		UEdGraphPin* NextExec = FindExecutionPin(*EntryNode, EGPD_Output);
