[NO NON-PII COMMIT MESSAGE]

diff --git a/Engine/Source/Editor/BlueprintGraph/Private/K2Node_Select.cpp b/Engine/Source/Editor/BlueprintGraph/Private/K2Node_Select.cpp
--- a/Engine/Source/Editor/BlueprintGraph/Private/K2Node_Select.cpp
+++ b/Engine/Source/Editor/BlueprintGraph/Private/K2Node_Select.cpp
@@ -490,6 +490,24 @@ void UK2Node_Select::ReallocatePinsDuringReconstruction(TArray<UEdGraphPin*>& Ol
 	{
 		NewReturn->PinType = OldReturnPin->PinType;
 	}
+	else if (OldReturnPin && NewReturn && OldReturnPin->LinkedTo.Num())
+	{
+		auto BP = GetBlueprint();
+		UClass* SelfClass = BP ? BP->GeneratedClass : nullptr;
+		bool OldTypeIsValid = true;
+		for (auto OutPin : OldReturnPin->LinkedTo)
+		{
+			if (OutPin && !Schema->ArePinTypesCompatible(OldReturnPin->PinType, OutPin->PinType, SelfClass))
+			{
+				OldTypeIsValid = false;
+				break;
+			}
+		}
+		if (OldTypeIsValid)
+		{
+			NewReturn->PinType = OldReturnPin->PinType;
+		}
+	}
 
 	UEdGraphPin* IndexPin = GetIndexPin();
 
