Deleting a map while it's level script is open will no longer crash.

Changed the PendingDelete list in the BlueprintActionDatabase to utilize weak pointers so the object can be correctly validated.

diff --git a/Engine/Source/Editor/BlueprintGraph/Private/BlueprintActionDatabase.cpp b/Engine/Source/Editor/BlueprintGraph/Private/BlueprintActionDatabase.cpp
--- a/Engine/Source/Editor/BlueprintGraph/Private/BlueprintActionDatabase.cpp
+++ b/Engine/Source/Editor/BlueprintGraph/Private/BlueprintActionDatabase.cpp
@@ -488,7 +488,7 @@ namespace BlueprintActionDatabaseImpl
 	 * way for a delete), but in-case the class wasn't deleted we need them 
 	 * tracked here so we can add them back in.
 	 */
-	TSet<UObject*> PendingDelete;
+	TSet<TWeakObjectPtr<UObject>> PendingDelete;
 
 	/** */
 	bool bIsInitializing = false;
@@ -883,8 +883,15 @@ static void BlueprintActionDatabaseImpl::OnAssetRemoved(UObject* AssetObject)
 	FBlueprintActionDatabase& ActionDatabase = FBlueprintActionDatabase::Get();
 	ActionDatabase.ClearAssetActions(AssetObject);
 
-	// the delete went through, so we don't need to track these for re-add
-	PendingDelete.Remove(AssetObject);
+	for (auto It(PendingDelete.CreateIterator()); It; ++It)
+	{
+		if ((*It).Get() == AssetObject)
+		{
+			// the delete went through, so we don't need to track these for re-add
+			It.RemoveCurrent();
+			break;
+		}
+	}
 }
 
 //------------------------------------------------------------------------------
@@ -1027,11 +1034,11 @@ void FBlueprintActionDatabase::Tick(float DeltaTime)
 
 	// entries that were removed from the database, in preparation for a delete
 	// (but the user ended up not deleting the object)
-	for (UObject* AssetObj : BlueprintActionDatabaseImpl::PendingDelete)
+	for (TWeakObjectPtr<UObject> AssetObj : BlueprintActionDatabaseImpl::PendingDelete)
 	{
-		if (IsValid(AssetObj))
+		if (AssetObj.IsValid())
 		{
-			RefreshAssetActions(AssetObj);
+			RefreshAssetActions(AssetObj.Get());
 		}
 	}
 	BlueprintActionDatabaseImpl::PendingDelete.Empty();
