[NO NON-PII COMMIT MESSAGE]

diff --git a/Engine/Config/BaseEngine.ini b/Engine/Config/BaseEngine.ini
--- a/Engine/Config/BaseEngine.ini
+++ b/Engine/Config/BaseEngine.ini
@@ -1254,6 +1254,9 @@ bUseStreamingPause=false
 +K2ParamRedirects=(NodeName="/Script/BlueprintGraph.K2Node_FunctionEntry", OldParamName="UserWidget.OnControllerButtonPressed.ControllerEvent", NewParamName="InKeyEvent")
 +K2ParamRedirects=(NodeName="/Script/BlueprintGraph.K2Node_FunctionEntry", OldParamName="UserWidget.OnControllerButtonReleased.ControllerEvent", NewParamName="InKeyEvent")
 
+; CreateEvent improvements:
++K2ParamRedirects=(NodeName="/Script/BlueprintGraph.K2Node_CreateDelegate", OldParamName="InputObject", NewParamName="self")
+
 ; UMG Checkbox State Rename
 
 +ActiveClassRedirects=(OldClassName="ESlateCheckBoxState", NewClassName="/Script/SlateCore.ECheckBoxState")
diff --git a/Engine/Source/Editor/BlueprintGraph/Private/DelegateNodeHandlers.cpp b/Engine/Source/Editor/BlueprintGraph/Private/DelegateNodeHandlers.cpp
--- a/Engine/Source/Editor/BlueprintGraph/Private/DelegateNodeHandlers.cpp
+++ b/Engine/Source/Editor/BlueprintGraph/Private/DelegateNodeHandlers.cpp
@@ -210,18 +210,24 @@ void FKCHandler_CreateDelegate::RegisterNets(FKismetFunctionContext& Context, UE
 	{
 		UEdGraphPin* InputPin = DelegateNode->GetObjectInPin();
 		check(NULL != InputPin);
-		if(!InputPin->LinkedTo.Num())
-		{
-			CompilerContext.MessageLog.Error(*LOCTEXT("NoDelegateObject", "No delegate object @@").ToString(), DelegateNode);
-			return;
-		}
+
 		UEdGraphPin* Net = FEdGraphUtilities::GetNetFromPin(InputPin);
 
 		FBPTerminal** FoundTerm = Context.NetMap.Find(Net);
 		FBPTerminal* InputObjTerm = FoundTerm ? *FoundTerm : NULL;
 		if(NULL == InputObjTerm)
 		{
-			InputObjTerm = Context.CreateLocalTerminalFromPinAutoChooseScope(Net, Context.NetNameMap->MakeValidName(Net));
+			if (InputPin->LinkedTo.Num() == 0)
+			{
+				InputObjTerm = Context.CreateLocalTerminal(ETerminalSpecification::TS_Literal);
+				InputObjTerm->Name = Context.NetNameMap->MakeValidName(Net);
+				InputObjTerm->Type.PinSubCategory = CompilerContext.GetSchema()->PN_Self;
+			}
+			else
+			{
+				InputObjTerm = Context.CreateLocalTerminalFromPinAutoChooseScope(Net, Context.NetNameMap->MakeValidName(Net));
+			}
+
 			Context.NetMap.Add(Net, InputObjTerm);
 		}
 	}
diff --git a/Engine/Source/Editor/BlueprintGraph/Private/EdGraphSchema_K2.cpp b/Engine/Source/Editor/BlueprintGraph/Private/EdGraphSchema_K2.cpp
--- a/Engine/Source/Editor/BlueprintGraph/Private/EdGraphSchema_K2.cpp
+++ b/Engine/Source/Editor/BlueprintGraph/Private/EdGraphSchema_K2.cpp
@@ -3982,6 +3982,10 @@ bool UEdGraphSchema_K2::ShouldShowAssetPickerForPin(UEdGraphPin* Pin) const
 						bShow = ( WorldContextPin != Pin );
 					}
 				}
+				else if (Cast<UK2Node_CreateDelegate>( Pin->GetOwningNode())) 
+				{
+					bShow = false;
+				}
 			}
 		}
 	}
diff --git a/Engine/Source/Editor/BlueprintGraph/Private/K2Node_CreateDelegate.cpp b/Engine/Source/Editor/BlueprintGraph/Private/K2Node_CreateDelegate.cpp
--- a/Engine/Source/Editor/BlueprintGraph/Private/K2Node_CreateDelegate.cpp
+++ b/Engine/Source/Editor/BlueprintGraph/Private/K2Node_CreateDelegate.cpp
@@ -11,7 +11,6 @@ struct FK2Node_CreateDelegate_Helper
 	static FString ObjectInputName;
 	static FString DelegateOutputName;
 };
-FString FK2Node_CreateDelegate_Helper::ObjectInputName(TEXT("InputObject"));
 FString FK2Node_CreateDelegate_Helper::DelegateOutputName(TEXT("OutputDelegate"));
 
 UK2Node_CreateDelegate::UK2Node_CreateDelegate(const FObjectInitializer& ObjectInitializer)
@@ -23,7 +22,7 @@ void UK2Node_CreateDelegate::AllocateDefaultPins()
 {
 	const UEdGraphSchema_K2* K2Schema = GetDefault<UEdGraphSchema_K2>();
 
-	if(UEdGraphPin* ObjPin = CreatePin(EGPD_Input, K2Schema->PC_Object, TEXT(""), UObject::StaticClass(), false, false, FK2Node_CreateDelegate_Helper::ObjectInputName))
+	if (UEdGraphPin* ObjPin = CreatePin(EGPD_Input, K2Schema->PC_Object, TEXT(""), UObject::StaticClass(), false, false, K2Schema->PN_Self))
 	{
 		ObjPin->PinFriendlyName = NSLOCTEXT("K2Node", "CreateDelegate_ObjectInputName", "Object");
 	}
@@ -280,7 +279,8 @@ UFunction* UK2Node_CreateDelegate::GetDelegateSignature() const
 
 UClass* UK2Node_CreateDelegate::GetScopeClass(bool bDontUseSkeletalClassForSelf/* = false*/) const
 {
-	UEdGraphPin* Pin = FindPin(FK2Node_CreateDelegate_Helper::ObjectInputName);
+	const UEdGraphSchema_K2* K2Schema = GetDefault<UEdGraphSchema_K2>();
+	UEdGraphPin* Pin = FindPin(K2Schema->PN_Self);
 	if (Pin == nullptr)
 	{
 		// The BlueprintNodeTemplateCache creates nodes but doesn't call allocate default pins.
@@ -289,18 +289,19 @@ UClass* UK2Node_CreateDelegate::GetScopeClass(bool bDontUseSkeletalClassForSelf/
 		return nullptr;
 	}
 	check(Pin->LinkedTo.Num() <= 1);
-	if(Pin->LinkedTo.Num())
+	bool bUseSelf = false;
+	if(Pin->LinkedTo.Num() == 0)
+	{
+		bUseSelf = true;
+	}
+	else
 	{
-		const UEdGraphSchema_K2* K2Schema = GetDefault<UEdGraphSchema_K2>();
 		if(UEdGraphPin* ResultPin = Pin->LinkedTo[0])
 		{
 			ensure(K2Schema->PC_Object == ResultPin->PinType.PinCategory);
 			if (K2Schema->PN_Self == ResultPin->PinType.PinSubCategory)
 			{
-				if (UBlueprint* ScopeClassBlueprint = GetBlueprint())
-				{
-					return bDontUseSkeletalClassForSelf ? ScopeClassBlueprint->GeneratedClass : ScopeClassBlueprint->SkeletonGeneratedClass;
-				}
+				bUseSelf = true;
 			}
 			if(UClass* TrueScopeClass = Cast<UClass>(ResultPin->PinType.PinSubCategoryObject.Get()))
 			{
@@ -315,6 +316,15 @@ UClass* UK2Node_CreateDelegate::GetScopeClass(bool bDontUseSkeletalClassForSelf/
 			}
 		}
 	}
+
+	if (bUseSelf)
+	{
+		if (UBlueprint* ScopeClassBlueprint = GetBlueprint())
+		{
+			return bDontUseSkeletalClassForSelf ? ScopeClassBlueprint->GeneratedClass : ScopeClassBlueprint->SkeletonGeneratedClass;
+		}
+	}
+
 	return NULL;
 }
 
@@ -330,7 +340,8 @@ UEdGraphPin* UK2Node_CreateDelegate::GetDelegateOutPin() const
 
 UEdGraphPin* UK2Node_CreateDelegate::GetObjectInPin() const
 {
-	return FindPinChecked(FK2Node_CreateDelegate_Helper::ObjectInputName);
+	const UEdGraphSchema_K2* K2Schema = GetDefault<UEdGraphSchema_K2>();
+	return FindPinChecked(K2Schema->PN_Self);
 }
 
 FText UK2Node_CreateDelegate::GetNodeTitle(ENodeTitleType::Type TitleType) const
@@ -340,9 +351,10 @@ FText UK2Node_CreateDelegate::GetNodeTitle(ENodeTitleType::Type TitleType) const
 
 UObject* UK2Node_CreateDelegate::GetJumpTargetForDoubleClick() const
 {
+	const UEdGraphSchema_K2* K2Schema = GetDefault<UEdGraphSchema_K2>();
 	UBlueprint* ScopeClassBlueprint = NULL;
 
-	UEdGraphPin* Pin = FindPinChecked(FK2Node_CreateDelegate_Helper::ObjectInputName);
+	UEdGraphPin* Pin = FindPinChecked(K2Schema->PN_Self);
 	UEdGraphPin* ResultPin = Pin->LinkedTo.Num() ? Pin->LinkedTo[0] : NULL;
 	if (ResultPin)
 	{
diff --git a/Engine/Source/Runtime/Engine/Classes/Kismet/KismetSystemLibrary.h b/Engine/Source/Runtime/Engine/Classes/Kismet/KismetSystemLibrary.h
--- a/Engine/Source/Runtime/Engine/Classes/Kismet/KismetSystemLibrary.h
+++ b/Engine/Source/Runtime/Engine/Classes/Kismet/KismetSystemLibrary.h
@@ -456,7 +456,7 @@ class ENGINE_API UKismetSystemLibrary : public UBlueprintFunctionLibrary
 	 * @param bLooping		true to keep executing the delegate every Time seconds, false to execute delegate only once.
 	 * @return				The timer handle to pass to other timer functions to manipulate this timer.
 	 */
-	UFUNCTION(BlueprintCallable, meta=(DeprecatedFunction, DeprecationMessage = "Use Set Timer by Event", DisplayName = "Set Timer by Function Name", DefaultToSelf = "Object"), Category="Utilities|Time")
+	UFUNCTION(BlueprintCallable, meta=(DisplayName = "Set Timer by Function Name", DefaultToSelf = "Object"), Category="Utilities|Time")
 	static FTimerHandle K2_SetTimer(UObject* Object, FString FunctionName, float Time, bool bLooping);
 
 	/**
