Fixed shadow variable

diff --git a/Engine/Source/Editor/BlueprintGraph/Private/K2Node_CreateDelegate.cpp b/Engine/Source/Editor/BlueprintGraph/Private/K2Node_CreateDelegate.cpp
--- a/Engine/Source/Editor/BlueprintGraph/Private/K2Node_CreateDelegate.cpp
+++ b/Engine/Source/Editor/BlueprintGraph/Private/K2Node_CreateDelegate.cpp
@@ -358,7 +358,6 @@ UObject* UK2Node_CreateDelegate::GetJumpTargetForDoubleClick() const
 	UEdGraphPin* ResultPin = Pin->LinkedTo.Num() ? Pin->LinkedTo[0] : NULL;
 	if (ResultPin)
 	{
-		const UEdGraphSchema_K2* K2Schema = GetDefault<UEdGraphSchema_K2>();
 		ensure(K2Schema->PC_Object == ResultPin->PinType.PinCategory);
 		if (UClass* TrueScopeClass = Cast<UClass>(ResultPin->PinType.PinSubCategoryObject.Get()))
 		{
