[UE-19293] Math expression now updates old calls to Rot/Rotator/MakeRot to use MakeRotator by swapping parameters

diff --git a/Engine/Source/Editor/BlueprintGraph/Classes/K2Node_MathExpression.h b/Engine/Source/Editor/BlueprintGraph/Classes/K2Node_MathExpression.h
--- a/Engine/Source/Editor/BlueprintGraph/Classes/K2Node_MathExpression.h
+++ b/Engine/Source/Editor/BlueprintGraph/Classes/K2Node_MathExpression.h
@@ -26,6 +26,7 @@ public:
 
 public:
 	// UObject interface
+	virtual void Serialize(FArchive& Ar) override;
 	virtual void PostEditChangeProperty(struct FPropertyChangedEvent& PropertyChangedEvent) override;
 	// End of UObject interface
 
diff --git a/Engine/Source/Editor/BlueprintGraph/Private/K2Node_MathExpression.cpp b/Engine/Source/Editor/BlueprintGraph/Private/K2Node_MathExpression.cpp
--- a/Engine/Source/Editor/BlueprintGraph/Private/K2Node_MathExpression.cpp
+++ b/Engine/Source/Editor/BlueprintGraph/Private/K2Node_MathExpression.cpp
@@ -970,7 +970,7 @@ private:
 			ADD_ALIAS("VECT2D")
 		FUNC_ALIASES_END
 	
-		FUNC_ALIASES_BEGIN("MakeRot")
+		FUNC_ALIASES_BEGIN("MakeRotator")
 			ADD_ALIAS("ROTATOR")
 			ADD_ALIAS("ROT")
 		FUNC_ALIASES_END
@@ -2416,6 +2416,66 @@ UK2Node_MathExpression::UK2Node_MathExpression(const FObjectInitializer& ObjectI
 	bCanRenameNode = true;
 }
 
+void UK2Node_MathExpression::Serialize(FArchive& Ar)
+{
+	UK2Node_Composite::Serialize(Ar);
+
+	if (Ar.IsLoading() && (Ar.UE4Ver() < VER_UE4_MAKE_ROT_RENAME_AND_REORDER))
+	{
+		// We need to reorder the parameters to MakeRot/MakeRotator/Rotator/Rot, to filter this expensive logic I'm just searching
+		// expressions for 'rot':
+		if (Expression.Contains(TEXT("Rot")))
+		{
+			// Now parse the expression and look for function expressions to the old MakeRot function:
+			FExpressionParser Parser;
+			TSharedPtr<IFExpressionNode> ExpressionRoot = Parser.ParseExpression(Expression);
+			
+			struct FMakeRotFixupVisitor : public FExpressionVisitor
+			{
+				virtual bool Visit(class FFunctionExpression& Node, EVisitPhase Phase)
+				{ 
+					if (Phase != VISIT_Pre)
+					{
+						return false;
+					}
+
+					const bool bIsMakeRot = (Node.FuncName == TEXT("makerot"));
+					if (bIsMakeRot ||
+						Node.FuncName == TEXT("rotator") ||
+						Node.FuncName == TEXT("rot"))
+					{
+						// reorder parameters to match new order of MakeRotator:
+						if (Node.ParamList->Children.Num() == 3)
+						{
+							TSharedRef<IFExpressionNode> OldPitch = Node.ParamList->Children[0];
+							TSharedRef<IFExpressionNode> OldYaw = Node.ParamList->Children[1];
+							TSharedRef<IFExpressionNode> OldRoll = Node.ParamList->Children[2];
+							
+							Node.ParamList->Children[0] = OldRoll;
+							Node.ParamList->Children[1] = OldPitch;
+							Node.ParamList->Children[2] = OldYaw;
+						}
+
+						// MakeRot also needs to be updated to the new name:
+						if (bIsMakeRot)
+						{
+							Node.FuncName = TEXT("MakeRotator");
+						}
+					}
+					return true;
+				}
+			};
+
+			// perform the update:
+			FMakeRotFixupVisitor Fixup;
+			ExpressionRoot->Accept(Fixup);
+
+			// reform the expression with the updated parameter order/function names:
+			Expression = ExpressionRoot->ToString();
+		}
+	}
+}
+
 //------------------------------------------------------------------------------
 void UK2Node_MathExpression::PostEditChangeProperty(struct FPropertyChangedEvent& PropertyChangedEvent)
 {
diff --git a/Engine/Source/Runtime/Core/Public/UObject/ObjectVersion.h b/Engine/Source/Runtime/Core/Public/UObject/ObjectVersion.h
--- a/Engine/Source/Runtime/Core/Public/UObject/ObjectVersion.h
+++ b/Engine/Source/Runtime/Core/Public/UObject/ObjectVersion.h
@@ -568,6 +568,8 @@ enum EUnrealEngineObjectUE4Version
 	VER_UE4_COOKED_ASSETS_IN_EDITOR_SUPPORT,
 	// Updated dialogue wave localization gathering logic.
 	VER_UE4_DIALOGUE_WAVE_NAMESPACE_AND_CONTEXT_CHANGES,
+	// Renamed MakeRot MakeRotator and rearranged parameters.
+	VER_UE4_MAKE_ROT_RENAME_AND_REORDER,
 	// -----<new versions can be added before this line>-------------------------------------------------
 	// - this needs to be the last line (see note below)
 	VER_UE4_AUTOMATIC_VERSION_PLUS_ONE,
