Adding in logging information to help with future instances of a crash.

diff --git a/Engine/Source/Editor/BlueprintGraph/Classes/K2Node_VariableGet.h b/Engine/Source/Editor/BlueprintGraph/Classes/K2Node_VariableGet.h
--- a/Engine/Source/Editor/BlueprintGraph/Classes/K2Node_VariableGet.h
+++ b/Engine/Source/Editor/BlueprintGraph/Classes/K2Node_VariableGet.h
@@ -11,6 +11,10 @@ class BLUEPRINTGRAPH_API UK2Node_VariableGet : public UK2Node_Variable
 {
 	GENERATED_UCLASS_BODY()
 
+	// Begin UObject interface
+	virtual void Serialize(FArchive& Ar) override;
+	// End of UObject interface
+
 	// Begin UEdGraphNode interface
 	virtual void AllocateDefaultPins() override;
 	virtual FText GetTooltipText() const override;
diff --git a/Engine/Source/Editor/BlueprintGraph/Private/K2Node_VariableGet.cpp b/Engine/Source/Editor/BlueprintGraph/Private/K2Node_VariableGet.cpp
--- a/Engine/Source/Editor/BlueprintGraph/Private/K2Node_VariableGet.cpp
+++ b/Engine/Source/Editor/BlueprintGraph/Private/K2Node_VariableGet.cpp
@@ -445,4 +445,31 @@ void UK2Node_VariableGet::ExpandNode(class FKismetCompilerContext& CompilerConte
 	}
 }
 
+void UK2Node_VariableGet::Serialize(FArchive& Ar)
+{
+	// The following code is to attempt to log info related to UE-19729
+	if (Ar.IsSaving() && !Ar.IsTransacting())
+	{
+		if (UEdGraph* Graph = Cast<UEdGraph>(GetOuter()))
+		{
+			if (UBlueprint* Blueprint = FBlueprintEditorUtils::FindBlueprintForGraph(Graph))
+			{
+				if (!Blueprint->bBeingCompiled)
+				{
+					FString VariableName = GetVarNameString();
+					FString BlueprintPath = Blueprint->GetPathName();
+					FString SetupStyle = bIsPureGet? TEXT("pure") : TEXT("validated");
+					UE_LOG(LogBlueprintUserMessages, Log, TEXT("Serialization for Get node for variable '%s' in Blueprint '%s' which is setup as %s"), *VariableName, *BlueprintPath, *SetupStyle);
+
+					// The following line may spur the crash noted in UE-19729 and will confirm that the crash happens before the FiB gather.
+					GetNodeTitle(ENodeTitleType::ListView);
+				}
+			}
+			
+		}
+	}
+	Super::Serialize(Ar);
+}
+
+
 #undef LOCTEXT_NAMESPACE
\ No newline at end of file
