Improved fix so that it no longer relies on changes to ObjectVersion.h. That file is in a somewhat indeterminate statel

diff --git a/Engine/Source/Editor/BlueprintGraph/Classes/K2Node_MathExpression.h b/Engine/Source/Editor/BlueprintGraph/Classes/K2Node_MathExpression.h
--- a/Engine/Source/Editor/BlueprintGraph/Classes/K2Node_MathExpression.h
+++ b/Engine/Source/Editor/BlueprintGraph/Classes/K2Node_MathExpression.h
@@ -24,6 +24,9 @@ public:
 	UPROPERTY(EditAnywhere, Category = Expression)
 	FString Expression;
 
+	UPROPERTY()
+	bool bMadeAfterRotChange;
+
 public:
 	// UObject interface
 	virtual void Serialize(FArchive& Ar) override;
diff --git a/Engine/Source/Editor/BlueprintGraph/Private/K2Node_MathExpression.cpp b/Engine/Source/Editor/BlueprintGraph/Private/K2Node_MathExpression.cpp
--- a/Engine/Source/Editor/BlueprintGraph/Private/K2Node_MathExpression.cpp
+++ b/Engine/Source/Editor/BlueprintGraph/Private/K2Node_MathExpression.cpp
@@ -2414,14 +2414,19 @@ UK2Node_MathExpression::UK2Node_MathExpression(const FObjectInitializer& ObjectI
 	// renaming the node rebuilds the expression (the node name is where they 
 	// specify the math equation)
 	bCanRenameNode = true;
+
+	bMadeAfterRotChange = false;
 }
 
 void UK2Node_MathExpression::Serialize(FArchive& Ar)
 {
 	UK2Node_Composite::Serialize(Ar);
 
-	if (Ar.IsLoading() && (Ar.UE4Ver() < VER_UE4_MAKE_ROT_RENAME_AND_REORDER))
+	if (Ar.IsLoading() && !bMadeAfterRotChange)
 	{
+		// remember that this logic has been run, we only want to run it once:
+		bMadeAfterRotChange = true;
+
 		// We need to reorder the parameters to MakeRot/MakeRotator/Rotator/Rot, to filter this expensive logic I'm just searching
 		// expressions for 'rot':
 		if (Expression.Contains(TEXT("Rot")))
@@ -2742,6 +2747,7 @@ FText UK2Node_MathExpression::GetNodeTitle(ENodeTitleType::Type TitleType) const
 //------------------------------------------------------------------------------
 void UK2Node_MathExpression::PostPlacedNewNode()
 {
+	bMadeAfterRotChange = true;
 	Super::PostPlacedNewNode();
 	FEdGraphUtilities::RenameGraphToNameOrCloseToName(BoundGraph, "MathExpression");
 }
