A warning when a deprecated BreakStruct output is used.
https://udn.unrealengine.com/questions/256628/suggested-change-warn-in-bp-on-breaking-a-struct-w.html

diff --git a/Engine/Source/Editor/BlueprintGraph/Private/K2Node_BreakStruct.cpp b/Engine/Source/Editor/BlueprintGraph/Private/K2Node_BreakStruct.cpp
--- a/Engine/Source/Editor/BlueprintGraph/Private/K2Node_BreakStruct.cpp
+++ b/Engine/Source/Editor/BlueprintGraph/Private/K2Node_BreakStruct.cpp
@@ -79,6 +79,14 @@ public:
 
 		if (BoundProperty != NULL)
 		{
+			if (BoundProperty->HasAnyPropertyFlags(CPF_Deprecated) && Net->LinkedTo.Num())
+			{
+				FText Message = FText::Format(LOCTEXT("BreakStruct_DeprecatedField_Warning", "@@ : Member '{0}' of struct '{1}' is deprecated.")
+					, BoundProperty->GetDisplayNameText()
+					, StructType->GetDisplayNameText());
+				CompilerContext.MessageLog.Warning(*Message.ToString(), Net->GetOuter());
+			}
+
 			FBPTerminal* Term = Context.CreateLocalTerminalFromPinAutoChooseScope(Net, Net->PinName);
 			Term->AssociatedVarProperty = BoundProperty;
 			Context.NetMap.Add(Net, Term);
