Fixed erroneous shadowed variable

(The variable is not shadowed - the lambda doesn't even have a capture mode!)

diff --git a/Engine/Source/Editor/BlueprintGraph/Private/EdGraphSchema_K2.cpp b/Engine/Source/Editor/BlueprintGraph/Private/EdGraphSchema_K2.cpp
--- a/Engine/Source/Editor/BlueprintGraph/Private/EdGraphSchema_K2.cpp
+++ b/Engine/Source/Editor/BlueprintGraph/Private/EdGraphSchema_K2.cpp
@@ -1488,21 +1488,21 @@ void UEdGraphSchema_K2::GetContextMenuActions(const UEdGraph* CurrentGraph, cons
 						MenuBuilder->AddMenuEntry( FGraphEditorCommands::Get().PromoteSelectionToMacro );
 					}
 
-					MenuBuilder->AddSubMenu(LOCTEXT("AlignmentHeader", "Alignment"), FText(), FNewMenuDelegate::CreateLambda([](FMenuBuilder& MenuBuilder){
+					MenuBuilder->AddSubMenu(LOCTEXT("AlignmentHeader", "Alignment"), FText(), FNewMenuDelegate::CreateLambda([](FMenuBuilder& InMenuBuilder){
 						
-						MenuBuilder.BeginSection("EdGraphSchemaAlignment", LOCTEXT("AlignHeader", "Align"));
-						MenuBuilder.AddMenuEntry( FGraphEditorCommands::Get().AlignNodesTop );
-						MenuBuilder.AddMenuEntry( FGraphEditorCommands::Get().AlignNodesMiddle );
-						MenuBuilder.AddMenuEntry( FGraphEditorCommands::Get().AlignNodesBottom );
-						MenuBuilder.AddMenuEntry( FGraphEditorCommands::Get().AlignNodesLeft );
-						MenuBuilder.AddMenuEntry( FGraphEditorCommands::Get().AlignNodesCenter );
-						MenuBuilder.AddMenuEntry( FGraphEditorCommands::Get().AlignNodesRight );
-						MenuBuilder.EndSection();
-
-						MenuBuilder.BeginSection("EdGraphSchemaDistribution", LOCTEXT("DistributionHeader", "Distribution"));
-						MenuBuilder.AddMenuEntry( FGraphEditorCommands::Get().DistributeNodesHorizontally );
-						MenuBuilder.AddMenuEntry( FGraphEditorCommands::Get().DistributeNodesVertically );
-						MenuBuilder.EndSection();
+						InMenuBuilder.BeginSection("EdGraphSchemaAlignment", LOCTEXT("AlignHeader", "Align"));
+						InMenuBuilder.AddMenuEntry( FGraphEditorCommands::Get().AlignNodesTop );
+						InMenuBuilder.AddMenuEntry( FGraphEditorCommands::Get().AlignNodesMiddle );
+						InMenuBuilder.AddMenuEntry( FGraphEditorCommands::Get().AlignNodesBottom );
+						InMenuBuilder.AddMenuEntry( FGraphEditorCommands::Get().AlignNodesLeft );
+						InMenuBuilder.AddMenuEntry( FGraphEditorCommands::Get().AlignNodesCenter );
+						InMenuBuilder.AddMenuEntry( FGraphEditorCommands::Get().AlignNodesRight );
+						InMenuBuilder.EndSection();
+
+						InMenuBuilder.BeginSection("EdGraphSchemaDistribution", LOCTEXT("DistributionHeader", "Distribution"));
+						InMenuBuilder.AddMenuEntry( FGraphEditorCommands::Get().DistributeNodesHorizontally );
+						InMenuBuilder.AddMenuEntry( FGraphEditorCommands::Get().DistributeNodesVertically );
+						InMenuBuilder.EndSection();
 						
 					}));
 				}
