Local Variable of type: Object, Class and Asset has  proper default value.

diff --git a/Engine/Source/Editor/BlueprintGraph/Private/EdGraphSchema_K2.cpp b/Engine/Source/Editor/BlueprintGraph/Private/EdGraphSchema_K2.cpp
--- a/Engine/Source/Editor/BlueprintGraph/Private/EdGraphSchema_K2.cpp
+++ b/Engine/Source/Editor/BlueprintGraph/Private/EdGraphSchema_K2.cpp
@@ -3845,13 +3845,19 @@ void UEdGraphSchema_K2::HandleGraphBeingDeleted(UEdGraph& GraphBeingRemoved) con
 void UEdGraphSchema_K2::TrySetDefaultValue(UEdGraphPin& Pin, const FString& NewDefaultValue) const
 {
 	FString UseDefaultValue;
-	UObject* UseDefaultObject = NULL;
+	UObject* UseDefaultObject = nullptr;
 	FText UseDefaultText;
 
-	if ((Pin.PinType.PinCategory == PC_Object) || (Pin.PinType.PinCategory == PC_Class) || (Pin.PinType.PinCategory == PC_Interface))
+	if ((Pin.PinType.PinCategory == PC_Object) 
+		|| (Pin.PinType.PinCategory == PC_Class) 
+		|| (Pin.PinType.PinCategory == PC_Interface)
+		|| (Pin.PinType.PinCategory == PC_Asset)
+		|| (Pin.PinType.PinCategory == PC_AssetClass))
 	{
-		UseDefaultObject = FindObject<UObject>(ANY_PACKAGE, *NewDefaultValue);
-		UseDefaultValue = NULL;
+		FString ObjectPathLocal = NewDefaultValue;
+		ConstructorHelpers::StripObjectClass(ObjectPathLocal);
+		UseDefaultObject = FindObject<UObject>(ANY_PACKAGE, *ObjectPathLocal);
+		UseDefaultValue.Empty();
 	}
 	else if(Pin.PinType.PinCategory == PC_Text)
 	{
@@ -3863,7 +3869,7 @@ void UEdGraphSchema_K2::TrySetDefaultValue(UEdGraphPin& Pin, const FString& NewD
 	}
 	else
 	{
-		UseDefaultObject = NULL;
+		UseDefaultObject = nullptr;
 		UseDefaultValue = NewDefaultValue;
 	}
 
