Proper error for an event node placed in a function graph.
(In Orion branch cl#2659623)

diff --git a/Engine/Source/Editor/BlueprintGraph/Private/EventEntryHandler.cpp b/Engine/Source/Editor/BlueprintGraph/Private/EventEntryHandler.cpp
--- a/Engine/Source/Editor/BlueprintGraph/Private/EventEntryHandler.cpp
+++ b/Engine/Source/Editor/BlueprintGraph/Private/EventEntryHandler.cpp
@@ -9,6 +9,11 @@ void FKCHandler_EventEntry::RegisterNet(FKismetFunctionContext& Context, UEdGrap
 	const bool bIsDelegateOutput = (NULL != Net) && (UK2Node_Event::DelegateOutputName == Net->PinName);
 	if(!bIsDelegateOutput)
 	{
+		if (!Context.IsEventGraph())
+		{
+			Context.MessageLog.Error(TEXT("Event node @@ registers net @@ in a non-event graph."), Net ? Net->GetOwningNodeUnchecked() : nullptr, Net);
+			return;
+		}
 		// This net is an event parameter; push to be a private class member variable
 		FBPTerminal* Term = Context.CreateLocalTerminal(ETerminalSpecification::TS_ForcedShared);
 		Term->CopyFromPin(Net, Context.NetNameMap->MakeValidName(Net));
diff --git a/Engine/Source/Editor/KismetCompiler/Private/KismetCompilerMisc.cpp b/Engine/Source/Editor/KismetCompiler/Private/KismetCompilerMisc.cpp
--- a/Engine/Source/Editor/KismetCompiler/Private/KismetCompilerMisc.cpp
+++ b/Engine/Source/Editor/KismetCompiler/Private/KismetCompilerMisc.cpp
@@ -1816,7 +1816,7 @@ FBPTerminal* FKismetFunctionContext::CreateLocalTerminal(ETerminalSpecification
 	switch (Spec)
 	{
 	case ETerminalSpecification::TS_ForcedShared:
-		// ensure(IsEventGraph()); it's used in function by UK2Node_AddComponent
+		ensure(IsEventGraph());
 		Result = new (EventGraphLocals)FBPTerminal();
 		break;
 	case ETerminalSpecification::TS_Literal:
