Cannot run this logic in Serialize because pins have not been completely serialized (FindPin always fails).

diff --git a/Engine/Source/Editor/BlueprintGraph/Classes/K2Node_Event.h b/Engine/Source/Editor/BlueprintGraph/Classes/K2Node_Event.h
--- a/Engine/Source/Editor/BlueprintGraph/Classes/K2Node_Event.h
+++ b/Engine/Source/Editor/BlueprintGraph/Classes/K2Node_Event.h
@@ -44,6 +44,7 @@ class UK2Node_Event : public UK2Node_EditablePinBase
 
 	// UObject interface
 	BLUEPRINTGRAPH_API virtual void Serialize(FArchive& Ar) override;
+	BLUEPRINTGRAPH_API virtual void PostLoad() override;
 	// End of UObject interface
 
 	// Begin UEdGraphNode interface
diff --git a/Engine/Source/Editor/BlueprintGraph/Private/K2Node_Event.cpp b/Engine/Source/Editor/BlueprintGraph/Private/K2Node_Event.cpp
--- a/Engine/Source/Editor/BlueprintGraph/Private/K2Node_Event.cpp
+++ b/Engine/Source/Editor/BlueprintGraph/Private/K2Node_Event.cpp
@@ -46,12 +46,6 @@ void UK2Node_Event::Serialize(FArchive& Ar)
 	// Fix up legacy nodes that may not yet have a delegate pin
 	if(Ar.IsLoading())
 	{
-		if(!FindPin(DelegateOutputName))
-		{
-			const UEdGraphSchema_K2* K2Schema = GetDefault<UEdGraphSchema_K2>();
-			CreatePin(EGPD_Output, K2Schema->PC_Delegate, TEXT(""), NULL, false, false, DelegateOutputName);
-		}
-
 		if(Ar.UE4Ver() < VER_UE4_K2NODE_EVENT_MEMBER_REFERENCE)
 		{
 			EventReference.SetExternalMember(EventSignatureName_DEPRECATED, EventSignatureClass_DEPRECATED);
@@ -59,6 +53,18 @@ void UK2Node_Event::Serialize(FArchive& Ar)
 	}
 }
 
+void UK2Node_Event::PostLoad()
+{
+	UK2Node_EditablePinBase::PostLoad();
+
+	// Fix up legacy nodes that may not yet have a delegate pin
+	if (!FindPin(DelegateOutputName))
+	{
+		const UEdGraphSchema_K2* K2Schema = GetDefault<UEdGraphSchema_K2>();
+		CreatePin(EGPD_Output, K2Schema->PC_Delegate, TEXT(""), NULL, false, false, DelegateOutputName);
+	}
+}
+
 FNodeHandlingFunctor* UK2Node_Event::CreateNodeHandler(FKismetCompilerContext& CompilerContext) const
 {
 	return new FKCHandler_EventEntry(CompilerContext);
