Fixed warnings that would sometimes appear when using level Blueprint communications.

When connecting the result of a Get Streaming Level node directly to a message call node it will no longer attempt to cast to a UStreamingLevel before getting the level Blueprint.

diff --git a/Engine/Source/Editor/BlueprintGraph/Private/K2Node_Message.cpp b/Engine/Source/Editor/BlueprintGraph/Private/K2Node_Message.cpp
--- a/Engine/Source/Editor/BlueprintGraph/Private/K2Node_Message.cpp
+++ b/Engine/Source/Editor/BlueprintGraph/Private/K2Node_Message.cpp
@@ -96,93 +96,122 @@ void UK2Node_Message::ExpandLevelStreamingHandlers(class FKismetCompilerContext&
 	 * Create intermediate nodes
 	 */
 
-	// We need to Cast to a ULevelStreaming object, so we can pull the ALevelScriptActor if successful
-	UK2Node_DynamicCast* CastToLevelStreamingNode = InCompilerContext.SpawnIntermediateNode<UK2Node_DynamicCast>(this, InSourceGraph);
-	CastToLevelStreamingNode->TargetType = ULevelStreaming::StaticClass(); // I want this to be ULevel!
-	CastToLevelStreamingNode->SetPurity(false);
-	CastToLevelStreamingNode->AllocateDefaultPins();
-
 	// Create a  GetLevelScriptActor CallFunction node, this will be used if the cast to ULevelStreaming was successful
 	UK2Node_CallFunction* GetLevelScriptActorNode = InCompilerContext.SpawnIntermediateNode<UK2Node_CallFunction>(this, InSourceGraph);
 	GetLevelScriptActorNode->SetFromFunction(ULevelStreaming::StaticClass()->FindFunctionByName(GET_FUNCTION_NAME_CHECKED(ULevelStreaming, GetLevelScriptActor)));
 	GetLevelScriptActorNode->AllocateDefaultPins();
 
-	// We will store the UObject we want to cast to interface into a temporary variable node
-	UK2Node_TemporaryVariable* TempVariable = InCompilerContext.SpawnIntermediateNode<UK2Node_TemporaryVariable>(this, InSourceGraph);
-	TempVariable->VariableType = InMessageSelfPin->PinType;
-	TempVariable->AllocateDefaultPins();
+	UEdGraphPin* PinToCastToInterface = nullptr;
 
-	// AssignmentStatement node for when the ULevelStreaming cast was successful
-	UK2Node_AssignmentStatement* AssignTempVariableForCastSuccess = InCompilerContext.SpawnIntermediateNode<UK2Node_AssignmentStatement>(this, InSourceGraph);
-	AssignTempVariableForCastSuccess->AllocateDefaultPins();
+	if (ULevelStreaming::StaticClass() != Cast<UClass>(InMessageSelfPin->LinkedTo[0]->PinType.PinSubCategoryObject.Get()))
+	{
+		/**
+		 * Create intermediate nodes
+		 */
+
+		// We need to Cast to a ULevelStreaming object, so we can pull the ALevelScriptActor if successful
+		UK2Node_DynamicCast* CastToLevelStreamingNode = InCompilerContext.SpawnIntermediateNode<UK2Node_DynamicCast>(this, InSourceGraph);
+		CastToLevelStreamingNode->TargetType = ULevelStreaming::StaticClass(); // I want this to be ULevel!
+		CastToLevelStreamingNode->SetPurity(false);
+		CastToLevelStreamingNode->AllocateDefaultPins();
+
+		// We will store the UObject we want to cast to interface into a temporary variable node
+		UK2Node_TemporaryVariable* TempVariable = InCompilerContext.SpawnIntermediateNode<UK2Node_TemporaryVariable>(this, InSourceGraph);
+		TempVariable->VariableType = InMessageSelfPin->PinType;
+		TempVariable->AllocateDefaultPins();
+
+		// AssignmentStatement node for when the ULevelStreaming cast was successful
+		UK2Node_AssignmentStatement* AssignTempVariableForCastSuccess = InCompilerContext.SpawnIntermediateNode<UK2Node_AssignmentStatement>(this, InSourceGraph);
+		AssignTempVariableForCastSuccess->AllocateDefaultPins();
+
+		// AssignmentStatement node for failed ULevelStreaming casts
+		UK2Node_AssignmentStatement* AssignTempVariableForCastFail = InCompilerContext.SpawnIntermediateNode<UK2Node_AssignmentStatement>(this, InSourceGraph);
+		AssignTempVariableForCastFail->AllocateDefaultPins();
+
+		/**
+		 * Make pin connections to the generated nodes
+		 */
+
+		// Move the connections on the Message node's self pin to the ULevelStreaming Cast node's source pin
+		{
+			InCompilerContext.MovePinLinksToIntermediate(*InMessageSelfPin, *CastToLevelStreamingNode->GetCastSourcePin());
+			CastToLevelStreamingNode->NotifyPinConnectionListChanged(CastToLevelStreamingNode->GetCastSourcePin());
+		}
 
-	// AssignmentStatement node for failed ULevelStreaming casts
-	UK2Node_AssignmentStatement* AssignTempVariableForCastFail = InCompilerContext.SpawnIntermediateNode<UK2Node_AssignmentStatement>(this, InSourceGraph);
-	AssignTempVariableForCastFail->AllocateDefaultPins();
+		// Connect the incoming exec pin to the ULevelStreaming cast node's exec pin, which is the exec flow's entry into this
+		if (InStartingExecPin != nullptr)
+		{
+			// Wire up the connections
+			InCompilerContext.MovePinLinksToIntermediate(*InStartingExecPin, *CastToLevelStreamingNode->GetExecPin());
+		}
 
-	/**
-	 * Make pin connections to the generated nodes
-	 */
+		{
+			Schema->TryCreateConnection(Schema->FindSelfPin(*GetLevelScriptActorNode, EGPD_Input), CastToLevelStreamingNode->GetCastResultPin());
+		}
 
-	// Move the connections on the Message node's self pin to the ULevelStreaming Cast node's source pin
-	{
-		InCompilerContext.MovePinLinksToIntermediate(*InMessageSelfPin, *CastToLevelStreamingNode->GetCastSourcePin());
-		CastToLevelStreamingNode->NotifyPinConnectionListChanged(CastToLevelStreamingNode->GetCastSourcePin());
-	}
+		// Prepare left side pin connections of AssignTempVariableForCastSuccess
+		{
+			// Connect the TempVariable to the AssignmentStatement node
+			// For the successful cast of ULevelStreaming temporary assignment, hook-up the variable pin
+			Schema->TryCreateConnection(AssignTempVariableForCastSuccess->GetVariablePin(), TempVariable->GetVariablePin());
 
-	// Connect the incoming exec pin to the ULevelStreaming cast node's exec pin, which is the exec flow's entry into this
-	if (InStartingExecPin != nullptr)
-	{
-		// Wire up the connections
-		InCompilerContext.MovePinLinksToIntermediate(*InStartingExecPin, *CastToLevelStreamingNode->GetExecPin());
-	}
+			// The last pin on the function node is the ALevelScriptActor
+			UEdGraphPin* FuncResultPin = GetLevelScriptActorNode->Pins[GetLevelScriptActorNode->Pins.Num() - 1];
+			ensure(!FuncResultPin->PinType.PinSubCategoryObject->IsA(ALevelScriptActor::StaticClass()));
 
-	// Connect the GetLevelAcriptActor CallFunction node's self pin to the ULevelStreaming node's cast result pin (successfully cast to a ULevelStreaming object)
-	{
-		Schema->TryCreateConnection(Schema->FindSelfPin(*GetLevelScriptActorNode, EGPD_Input), CastToLevelStreamingNode->GetCastResultPin());
-	}
+			// We do not want to let NotifyPinConnectionListChanged get called, it'll change the assignment's type to LevelScript instead of UObject
+			FuncResultPin->MakeLinkTo(AssignTempVariableForCastSuccess->GetValuePin());
 
-	// Prepare left side pin connections of AssignTempVariableForCastSuccess
-	{
-		// Connect the TempVariable to the AssignmentStatement node
-		// For the successful cast of ULevelStreaming temporary assignment, hook-up the variable pin
-		Schema->TryCreateConnection(AssignTempVariableForCastSuccess->GetVariablePin(), TempVariable->GetVariablePin());
+			// Hook up the AssignTempVariableForCastSuccess node's exec pin to the valid cast result of the Cast to ULevelStreaming node
+			Schema->TryCreateConnection(AssignTempVariableForCastSuccess->GetExecPin(), CastToLevelStreamingNode->GetValidCastPin());
+		}
 
-		// The last pin on the function node is the ALevelScriptActor
-		UEdGraphPin* FuncResultPin = GetLevelScriptActorNode->Pins[GetLevelScriptActorNode->Pins.Num() - 1];
-		ensure(!FuncResultPin->PinType.PinSubCategoryObject->IsA(ALevelScriptActor::StaticClass()));
+		// Prepare left side pin connections of AssignTempVariableForCastFail
+		{
+			// For the non-level actor temporary assignment, hook-up the variable pin
+			Schema->TryCreateConnection(AssignTempVariableForCastFail->GetVariablePin(), TempVariable->GetVariablePin());
 
-		// We do not want to let NotifyPinConnectionListChanged get called, it'll change the assignment's type to LevelScript instead of UObject
-		FuncResultPin->MakeLinkTo(AssignTempVariableForCastSuccess->GetValuePin());
+			// Connect the value pin to the pin the Cast to ULevelStreaming node's source pin was attached to
+			if (CastToLevelStreamingNode->GetCastSourcePin()->LinkedTo[0])
+			{
+				AssignTempVariableForCastFail->GetValuePin()->MakeLinkTo(CastToLevelStreamingNode->GetCastSourcePin()->LinkedTo[0]);
+			}
 
-		// Hook up the AssignTempVariableForCastSuccess node's exec pin to the valid cast result of the Cast to ULevelStreaming node
-		Schema->TryCreateConnection(AssignTempVariableForCastSuccess->GetExecPin(), CastToLevelStreamingNode->GetValidCastPin());
-	}
+			// Connect the LevelStreaming cast failed pin to the appropriate Assignment node
+			AssignTempVariableForCastFail->GetExecPin()->MakeLinkTo(CastToLevelStreamingNode->GetInvalidCastPin());
+		}
+
+		// Both assignment nodes (for the temporary variable) get connected to the interface node's exec pin
+		InCastToInterfaceNode->GetExecPin()->MakeLinkTo(AssignTempVariableForCastSuccess->GetThenPin());
+		InCastToInterfaceNode->GetExecPin()->MakeLinkTo(AssignTempVariableForCastFail->GetThenPin());
 
-	// Prepare left side pin connections of AssignTempVariableForCastFail
+		PinToCastToInterface = TempVariable->GetVariablePin();
+	}
+	else
 	{
-		// For the non-level actor temporary assignment, hook-up the variable pin
-		Schema->TryCreateConnection(AssignTempVariableForCastFail->GetVariablePin(), TempVariable->GetVariablePin());
+		PinToCastToInterface = InMessageSelfPin;
 
-		// Connect the value pin to the pin the Cast to ULevelStreaming node's source pin was attached to
-		if (CastToLevelStreamingNode->GetCastSourcePin()->LinkedTo[0])
+		// Move all pin connections from the message self pin to the GetLevelScriptActorNode's self pin
 		{
-			AssignTempVariableForCastFail->GetValuePin()->MakeLinkTo(CastToLevelStreamingNode->GetCastSourcePin()->LinkedTo[0]);
+			InCompilerContext.MovePinLinksToIntermediate(*InMessageSelfPin, *Schema->FindSelfPin(*GetLevelScriptActorNode, EGPD_Input));
+
+			// The last pin on the function node is the ALevelScriptActor
+			UEdGraphPin* FuncResultPin = GetLevelScriptActorNode->Pins[GetLevelScriptActorNode->Pins.Num() - 1];
+			ensure(!FuncResultPin->PinType.PinSubCategoryObject->IsA(ALevelScriptActor::StaticClass()));
+
+			// We will want to cast the resulting level Blueprint to the appropriate interface
+			PinToCastToInterface = FuncResultPin;
 		}
 
-		// Connect the LevelStreaming cast failed pin to the appropriate Assignment node
-		AssignTempVariableForCastFail->GetExecPin()->MakeLinkTo(CastToLevelStreamingNode->GetInvalidCastPin());
+		// Move all connections from the starting exec pin to the cast node
+		InCompilerContext.MovePinLinksToIntermediate(*InStartingExecPin, *InCastToInterfaceNode->GetExecPin());
 	}
 
 	// Connect the Interface cast node to the generated pins
 	{
 		// The source pin of the Interface cast node connects to the temporary variable (storing either an ALevelScriptActor or another UObject)
 		UEdGraphPin* CastToInterfaceSourceObjectPin = InCastToInterfaceNode->GetCastSourcePin();
-		Schema->TryCreateConnection(TempVariable->GetVariablePin(), CastToInterfaceSourceObjectPin);
-
-		// Both assignment nodes (for the temporary variable) get connected to the interface node's exec pin
-		InCastToInterfaceNode->GetExecPin()->MakeLinkTo(AssignTempVariableForCastSuccess->GetThenPin());
-		InCastToInterfaceNode->GetExecPin()->MakeLinkTo(AssignTempVariableForCastFail->GetThenPin());
+		Schema->TryCreateConnection(PinToCastToInterface, CastToInterfaceSourceObjectPin);
 	}
 }
 
