Improvements to the usability of PostProcessSettings structure in Blueprints.

Make/Break/Set nodes will no longer show the overrides variables that are used as conditionals for other variables.

These values will automatically be set to true when their corresponding variable is made visible in the node and false if they are not visible.

Added a tooltip for the checkbox to enable these pins to explain the situation.

This feature changes the functionality of existing nodes, all affected cases will give a Warning until the Blueprint is recompiled.

diff --git a/Engine/Source/Editor/BlueprintGraph/BlueprintGraph.Build.cs b/Engine/Source/Editor/BlueprintGraph/BlueprintGraph.Build.cs
--- a/Engine/Source/Editor/BlueprintGraph/BlueprintGraph.Build.cs
+++ b/Engine/Source/Editor/BlueprintGraph/BlueprintGraph.Build.cs
@@ -34,6 +34,7 @@ public class BlueprintGraph : ModuleRules
                 "GraphEditor",
 				"SlateCore",
                 "Kismet",
+                "PropertyEditor",
 			}
 		);
 
diff --git a/Engine/Source/Editor/BlueprintGraph/Classes/K2Node.h b/Engine/Source/Editor/BlueprintGraph/Classes/K2Node.h
--- a/Engine/Source/Editor/BlueprintGraph/Classes/K2Node.h
+++ b/Engine/Source/Editor/BlueprintGraph/Classes/K2Node.h
@@ -36,11 +36,14 @@ struct FOptionalPinFromProperty
 	UPROPERTY(EditAnywhere, BlueprintReadWrite, Category = Hi)
 	FName CategoryName;
 
+	UPROPERTY(EditAnywhere, Category = Hi)
+	bool bHasOverridePin;
+
 	FOptionalPinFromProperty()
 	{
 	}
 	
-	FOptionalPinFromProperty(FName InPropertyName, bool bInShowPin, bool bInCanToggleVisibility, const FString& InFriendlyName, const FText& InTooltip, bool bInPropertyIsCustomized, FName InCategoryName)
+	FOptionalPinFromProperty(FName InPropertyName, bool bInShowPin, bool bInCanToggleVisibility, const FString& InFriendlyName, const FText& InTooltip, bool bInPropertyIsCustomized, FName InCategoryName, bool bInHasOverridePin)
 		: PropertyName(InPropertyName)
 		, PropertyFriendlyName(InFriendlyName)
 		, PropertyTooltip(InTooltip)
@@ -48,6 +51,7 @@ struct FOptionalPinFromProperty
 		, bCanToggleVisibility(bInCanToggleVisibility)
 		, bPropertyIsCustomized(bInPropertyIsCustomized)
 		, CategoryName(InCategoryName)
+		, bHasOverridePin(bInHasOverridePin)
 	{
 	}
 };
diff --git a/Engine/Source/Editor/BlueprintGraph/Classes/K2Node_SetFieldsInStruct.h b/Engine/Source/Editor/BlueprintGraph/Classes/K2Node_SetFieldsInStruct.h
--- a/Engine/Source/Editor/BlueprintGraph/Classes/K2Node_SetFieldsInStruct.h
+++ b/Engine/Source/Editor/BlueprintGraph/Classes/K2Node_SetFieldsInStruct.h
@@ -10,6 +10,16 @@ UCLASS(MinimalAPI)
 class UK2Node_SetFieldsInStruct : public UK2Node_MakeStruct
 {
 	GENERATED_UCLASS_BODY()
+
+	/** Helper property to handle upgrades from an old system of displaying pins for 
+	 *	the override values that properties referenced as a conditional of being set in a struct */
+	UPROPERTY()
+	bool bMadeAfterOverridePinRemoval;
+
+	// UObject interface
+	virtual void Serialize(FArchive& Ar) override;
+	// End of UObject interface
+
 	// Begin UEdGraphNode interface
 	virtual void AllocateDefaultPins() override;
 	virtual FText GetNodeTitle(ENodeTitleType::Type TitleType) const override;
@@ -17,6 +27,7 @@ class UK2Node_SetFieldsInStruct : public UK2Node_MakeStruct
 	virtual FName GetPaletteIcon(FLinearColor& OutColor) const override;
 	virtual void ValidateNodeDuringCompilation(class FCompilerResultsLog& MessageLog) const override;
 	virtual void ExpandNode(class FKismetCompilerContext& CompilerContext, UEdGraph* SourceGraph);
+	virtual void PostPlacedNewNode() override;
 	// End  UEdGraphNode interface
 
 	// Begin K2Node interface
diff --git a/Engine/Source/Editor/BlueprintGraph/Private/K2Node.cpp b/Engine/Source/Editor/BlueprintGraph/Private/K2Node.cpp
--- a/Engine/Source/Editor/BlueprintGraph/Private/K2Node.cpp
+++ b/Engine/Source/Editor/BlueprintGraph/Private/K2Node.cpp
@@ -9,6 +9,7 @@
 #include "KismetCompiler.h"
 #include "GraphEditorSettings.h"
 #include "BlueprintEditorSettings.h"
+#include "Editor/PropertyEditor/Public/PropertyCustomizationHelpers.h"
 
 #include "ObjectEditorUtils.h"
 
@@ -868,11 +869,7 @@ void FOptionalPinManager::RebuildPropertyList(TArray<FOptionalPinFromProperty>&
 			CategoryName = FObjectEditorUtils::GetCategoryFName(TestProperty);
 #endif //WITH_EDITOR
 
-			UProperty* OverrideProperty = nullptr;
-			if (OverridesMap.RemoveAndCopyValue(TestProperty->GetFName(), OverrideProperty) && OverrideProperty)
-			{
-				RebuildProperty(OverrideProperty, CategoryName, Properties, SourceStruct, OldVisibility);
-			}
+			OverridesMap.Remove(TestProperty->GetFName());
 			RebuildProperty(TestProperty, CategoryName, Properties, SourceStruct, OldVisibility);
 		}
 	}
@@ -898,6 +895,10 @@ void FOptionalPinManager::RebuildProperty(UProperty* TestProperty, FName Categor
 	Record->PropertyFriendlyName = UEditorEngine::GetFriendlyName(TestProperty, SourceStruct);
 	Record->PropertyTooltip = TestProperty->GetToolTipText();
 	Record->CategoryName = CategoryName;
+
+	bool bNegate = false;
+	Record->bHasOverridePin = PropertyCustomizationHelpers::GetEditConditionProperty(TestProperty, bNegate) != nullptr;
+
 	// Get the defaults
 	GetRecordDefaults(TestProperty, *Record);
 
@@ -1147,7 +1148,6 @@ void UK2Node::GetPinHoverText(const UEdGraphPin& Pin, FString& HoverTextOut) con
 	// grab the debug value of the pin
 	FString WatchText;
 	const FKismetDebugUtilities::EWatchTextResult WatchStatus = FKismetDebugUtilities::GetWatchText(/*inout*/ WatchText, Blueprint, ActiveObject, &Pin);
-
 	// if this is an array pin, then we possibly have too many lines (too many entries)
 	if (Pin.PinType.bIsArray)
 	{
diff --git a/Engine/Source/Editor/BlueprintGraph/Private/K2Node_SetFieldsInStruct.cpp b/Engine/Source/Editor/BlueprintGraph/Private/K2Node_SetFieldsInStruct.cpp
--- a/Engine/Source/Editor/BlueprintGraph/Private/K2Node_SetFieldsInStruct.cpp
+++ b/Engine/Source/Editor/BlueprintGraph/Private/K2Node_SetFieldsInStruct.cpp
@@ -5,6 +5,7 @@
 #include "MakeStructHandler.h"
 #include "CompilerResultsLog.h"
 #include "KismetCompiler.h"
+#include "Editor/PropertyEditor/Public/PropertyCustomizationHelpers.h"
 
 #define LOCTEXT_NAMESPACE "K2Node_MakeStruct"
 
@@ -37,6 +38,7 @@ public:
 
 UK2Node_SetFieldsInStruct::UK2Node_SetFieldsInStruct(const FObjectInitializer& ObjectInitializer)
 	: Super(ObjectInitializer)
+	, bMadeAfterOverridePinRemoval(false)
 {
 }
 
@@ -109,6 +111,11 @@ void UK2Node_SetFieldsInStruct::ValidateNodeDuringCompilation(FCompilerResultsLo
 		FText ErrorMessage = LOCTEXT("SetStructFields_NoStructRefError", "The @@ pin must be connected to the struct that you wish to set.");
 		MessageLog.Error(*ErrorMessage.ToString(), FoundPin);
 	}
+
+	if (!bMadeAfterOverridePinRemoval)
+	{
+		MessageLog.Warning(*NSLOCTEXT("K2Node", "OverridePinRemoval", "Override pins have been removed from @@, please verify the Blueprint works as expected! See tooltips for enabling pin visibility for more details. This warning will go away after you resave the asset!").ToString(), this);
+	}
 }
 
 FNodeHandlingFunctor* UK2Node_SetFieldsInStruct::CreateNodeHandler(class FKismetCompilerContext& CompilerContext) const
@@ -197,6 +204,12 @@ void UK2Node_SetFieldsInStruct::FSetFieldsInStructPinManager::GetRecordDefaults(
 	Record.bShowPin = false;
 }
 
+void UK2Node_SetFieldsInStruct::PostPlacedNewNode()
+{
+	// New nodes automatically have this set.
+	bMadeAfterOverridePinRemoval = true;
+}
+
 void UK2Node_SetFieldsInStruct::ExpandNode(class FKismetCompilerContext& CompilerContext, UEdGraph* SourceGraph)
 {
 	Super::ExpandNode(CompilerContext, SourceGraph);
@@ -217,4 +230,42 @@ void UK2Node_SetFieldsInStruct::ExpandNode(class FKismetCompilerContext& Compile
 	Pins.Remove(OutPin);
 }
 
+void UK2Node_SetFieldsInStruct::Serialize(FArchive& Ar)
+{
+	UK2Node_StructOperation::Serialize(Ar);
+
+	if (Ar.IsLoading() && !bMadeAfterOverridePinRemoval)
+	{
+		// Check if this node actually requires warning the user that functionality has changed.
+
+		bMadeAfterOverridePinRemoval = true;
+		FOptionalPinManager PinManager;
+
+		// Have to check if this node is even in danger.
+		for (TFieldIterator<UProperty> It(StructType, EFieldIteratorFlags::IncludeSuper); It; ++It)
+		{
+			UProperty* TestProperty = *It;
+			if (PinManager.CanTreatPropertyAsOptional(TestProperty))
+			{
+				bool bNegate = false;
+				if (UProperty* OverrideProperty = PropertyCustomizationHelpers::GetEditConditionProperty(TestProperty, bNegate))
+				{
+					// We have confirmed that there is a property that uses an override variable to enable it, so set it to true.
+					bMadeAfterOverridePinRemoval = false;
+					break;
+				}
+			}
+		}
+	}
+	else if (Ar.IsSaving() && !Ar.IsTransacting())
+	{
+		UBlueprint* Blueprint = FBlueprintEditorUtils::FindBlueprintForNode(this);
+
+		if (Blueprint && !Blueprint->bBeingCompiled)
+		{
+			bMadeAfterOverridePinRemoval = true;
+		}
+	}
+}
+
 #undef LOCTEXT_NAMESPACE
\ No newline at end of file
diff --git a/Engine/Source/Editor/BlueprintGraph/Private/MakeStructHandler.cpp b/Engine/Source/Editor/BlueprintGraph/Private/MakeStructHandler.cpp
--- a/Engine/Source/Editor/BlueprintGraph/Private/MakeStructHandler.cpp
+++ b/Engine/Source/Editor/BlueprintGraph/Private/MakeStructHandler.cpp
@@ -4,6 +4,7 @@
 #include "MakeStructHandler.h"
 #include "K2Node_MakeStruct.h"
 #include "KismetCompiler.h"
+#include "Editor/PropertyEditor/Public/PropertyCustomizationHelpers.h"
 
 #define LOCTEXT_NAMESPACE "FKCHandler_MakeStruct"
 
@@ -110,6 +111,59 @@ void FKCHandler_MakeStruct::Compile(FKismetFunctionContext& Context, UEdGraphNod
 			FKismetCompilerUtilities::CreateObjectAssignmentStatement(Context, Node, SrcTerm, DstTerm);
 		}
 	}
+
+	// Search through all Properties for this node and set their override value (if any) to true or false based on if the value is being used
+	const UEdGraphSchema_K2* Schema = GetDefault<UEdGraphSchema_K2>();
+	TMap<UProperty*, FBPTerminal*> OverridePropertyToTerminalMap;
+	for (FOptionalPinFromProperty& PropertyEntry : Node->ShowPinForProperties)
+	{
+		if (UProperty* Property = FindFieldChecked<UProperty>(Node->StructType, PropertyEntry.PropertyName))
+		{
+			bool bNegate = false;
+			UProperty* OverrideProperty = PropertyCustomizationHelpers::GetEditConditionProperty(Property, bNegate);
+
+			if (OverrideProperty)
+			{
+				FBPTerminal** OverridePropertyTerminal = OverridePropertyToTerminalMap.Find(OverrideProperty);
+
+				// Setup a new terminal for the OverrideProperty if one hasn't been created
+				if (OverridePropertyTerminal == nullptr)
+				{
+					FEdGraphPinType PinType;
+					Schema->ConvertPropertyToPinType(OverrideProperty, /*out*/ PinType);
+
+					// Create the term in the list
+					FBPTerminal* Term = new (Context.VariableReferences) FBPTerminal();
+					Term->Type = PinType;
+					Term->AssociatedVarProperty = OverrideProperty;
+					Term->Context = OutputStructTerm;
+
+					FBlueprintCompiledStatement& AssignBoolStatement = Context.AppendStatementForNode(InNode);
+					AssignBoolStatement.Type = KCST_Assignment;
+
+					// Literal Bool Term to set the OverrideProperty to
+					FBPTerminal* BoolTerm = Context.CreateLocalTerminal(ETerminalSpecification::TS_Literal);
+					BoolTerm->Type.PinCategory = CompilerContext.GetSchema()->PC_Boolean;
+					BoolTerm->bIsLiteral = true;
+					// If we are showing the pin, then we are overriding the property
+					BoolTerm->Name = PropertyEntry.bShowPin? TEXT("true") : TEXT("false");
+
+					// Assigning the OverrideProperty to the literal bool term
+					AssignBoolStatement.LHS = Term;
+					AssignBoolStatement.RHS.Add(BoolTerm);
+
+					OverridePropertyToTerminalMap.Add(OverrideProperty, Term);
+				}
+				// When updating a terminal, we only want to set it to true
+				else if (PropertyEntry.bShowPin)
+				{
+					check(*OverridePropertyTerminal);
+					(*OverridePropertyTerminal)->Name = TEXT("true");
+				}
+			}
+		}
+	}
+
 	if (!Node->IsNodePure())
 	{
 		GenerateSimpleThenGoto(Context, *Node);
diff --git a/Engine/Source/Editor/DetailCustomizations/Private/SkeletalControlNodeDetails.cpp b/Engine/Source/Editor/DetailCustomizations/Private/SkeletalControlNodeDetails.cpp
--- a/Engine/Source/Editor/DetailCustomizations/Private/SkeletalControlNodeDetails.cpp
+++ b/Engine/Source/Editor/DetailCustomizations/Private/SkeletalControlNodeDetails.cpp
@@ -57,6 +57,27 @@ void FSkeletalControlNodeDetails::CustomizeDetails(class IDetailLayoutBuilder& D
 	}
 }
 
+ECheckBoxState FSkeletalControlNodeDetails::GetShowPinValueForProperty(TSharedRef<IPropertyHandle> InElementProperty) const
+{
+	FString Value;
+	InElementProperty->GetChildHandle("bShowPin")->GetValueAsFormattedString(Value);
+
+	if (Value == TEXT("true"))
+	{
+		return ECheckBoxState::Checked;
+	}
+	else if(Value == TEXT("false"))
+	{
+		return ECheckBoxState::Unchecked;
+	}
+	return ECheckBoxState::Undetermined;
+}
+
+void FSkeletalControlNodeDetails::OnShowPinChanged(ECheckBoxState InNewState, TSharedRef<IPropertyHandle> InElementProperty)
+{
+	InElementProperty->GetChildHandle("bShowPin")->SetValueFromFormattedString(InNewState == ECheckBoxState::Checked? TEXT("true") : TEXT("false"));
+}
+
 void FSkeletalControlNodeDetails::OnGenerateElementForPropertyPin(TSharedRef<IPropertyHandle> ElementProperty, int32 ElementIndex, IDetailChildrenBuilder& ChildrenBuilder, FName CategoryName)
 {
 	{
@@ -100,10 +121,47 @@ void FSkeletalControlNodeDetails::OnGenerateElementForPropertyPin(TSharedRef<IPr
 		}
 	}
 	
+	TSharedPtr<IPropertyHandle> HasOverrideValueHandle = ElementProperty->GetChildHandle("bHasOverridePin");
+	bool bHasOverrideValue;
+	HasOverrideValueHandle->GetValue(/*out*/ bHasOverrideValue);
+	FText Tooltip;
+
+	// Setup a tooltip based on whether the property has an override value or not.
+	if (bHasOverrideValue)
+	{
+		Tooltip = LOCTEXT("HasOverridePin", "Enabling this pin will make it visible for setting on the node and automatically enable the value for override when using the struct. Any updates to the resulting struct will require the value be set again or the override will be automatically disabled.");
+	}
+	else
+	{
+		Tooltip = LOCTEXT("HasNoOverridePin", "Enabling this pin will make it visible for setting on the node.");
+	}
+
+	TSharedRef<SWidget> PropertyNameWidget = ElementProperty->CreatePropertyNameWidget(PropertyFriendlyName);
+	PropertyNameWidget->SetToolTipText(Tooltip);
+
+	FString Value;
+	ElementProperty->GetChildHandle("bShowPin")->GetValueAsFormattedString(Value);
+
 	ChildrenBuilder.AddChildContent( PropertyFriendlyName )
 	[
-		SNew(SProperty, ElementProperty->GetChildHandle("bShowPin"))
-		.DisplayName(PropertyFriendlyName)
+		SNew( SHorizontalBox )
+		+ SHorizontalBox::Slot()
+		.VAlign(VAlign_Center)
+		.HAlign(HAlign_Right)
+		.Padding( 3.0f, 0.0f )
+		.FillWidth(1.0f)
+		[
+			PropertyNameWidget
+		]
+		+ SHorizontalBox::Slot()
+		.VAlign(VAlign_Center)
+		.FillWidth(1.0f)
+		[
+			SNew(SCheckBox)
+				.IsChecked(this, &FSkeletalControlNodeDetails::GetShowPinValueForProperty, ElementProperty)
+				.OnCheckStateChanged(this, &FSkeletalControlNodeDetails::OnShowPinChanged, ElementProperty)
+				.ToolTipText(Tooltip)
+		]
 	];
 }
 
diff --git a/Engine/Source/Editor/DetailCustomizations/Private/SkeletalControlNodeDetails.h b/Engine/Source/Editor/DetailCustomizations/Private/SkeletalControlNodeDetails.h
--- a/Engine/Source/Editor/DetailCustomizations/Private/SkeletalControlNodeDetails.h
+++ b/Engine/Source/Editor/DetailCustomizations/Private/SkeletalControlNodeDetails.h
@@ -16,6 +16,13 @@ public:
 
 protected:
 	void OnGenerateElementForPropertyPin(TSharedRef<IPropertyHandle> ElementProperty, int32 ElementIndex, IDetailChildrenBuilder& ChildrenBuilder, FName CategoryName);
+
+	/** Returns the current value of the checkbox being displayed for the bShowPin property */
+	ECheckBoxState GetShowPinValueForProperty(TSharedRef<IPropertyHandle> InElementProperty) const;
+
+	/** Helper function for changing the value of the bShowPin checkbox to update the property */
+	void OnShowPinChanged(ECheckBoxState InNewState, TSharedRef<IPropertyHandle> InElementProperty);
+
 private:
 	IDetailCategoryBuilder* DetailCategory;
 };
\ No newline at end of file
