Drag and dropping to create pins on function Return nodes when there are multiple of them will correctly update all Return nodes.

Adding the first Return node using the context menu will correctly make it editable.

Added error reporting when attempting to create a new pin via drag and dropping on an un-editable Return node.

diff --git a/Engine/Source/Editor/BlueprintGraph/Private/EdGraphSchema_K2.cpp b/Engine/Source/Editor/BlueprintGraph/Private/EdGraphSchema_K2.cpp
--- a/Engine/Source/Editor/BlueprintGraph/Private/EdGraphSchema_K2.cpp
+++ b/Engine/Source/Editor/BlueprintGraph/Private/EdGraphSchema_K2.cpp
@@ -5961,13 +5961,14 @@ UEdGraphPin* UEdGraphSchema_K2::DropPinOnNode(UEdGraphNode* InTargetNode, const
 	UEdGraphPin* ResultPin = nullptr;
 	if (UK2Node_EditablePinBase* EditablePinNode = Cast<UK2Node_EditablePinBase>(InTargetNode))
 	{
+		TArray<UK2Node_EditablePinBase*> EditablePinNodes;
 		EditablePinNode->Modify();
 
 		if (InSourcePinDirection == EGPD_Output && Cast<UK2Node_FunctionEntry>(InTargetNode))
 		{
-			if (UK2Node_EditablePinBase* ResultNode = FBlueprintEditorUtils::FindOrCreateFunctionResultNode(EditablePinNode))
+			if (UK2Node_FunctionResult* ResultNode = FBlueprintEditorUtils::FindOrCreateFunctionResultNode(EditablePinNode))
 			{
-				EditablePinNode = ResultNode;
+				EditablePinNodes.Add(ResultNode);
 			}
 			else
 			{
@@ -5982,7 +5983,7 @@ UEdGraphPin* UEdGraphSchema_K2::DropPinOnNode(UEdGraphNode* InTargetNode, const
 
 			if (FunctionEntryNode.Num() == 1)
 			{
-				EditablePinNode = FunctionEntryNode[0];
+				EditablePinNodes.Add(FunctionEntryNode[0]);
 			}
 			else
 			{
@@ -5990,9 +5991,31 @@ UEdGraphPin* UEdGraphSchema_K2::DropPinOnNode(UEdGraphNode* InTargetNode, const
 				return nullptr;
 			}
 		}
+		else
+		{
+			if (UK2Node_FunctionResult* ResultNode = Cast<UK2Node_FunctionResult>(EditablePinNode))
+			{
+				EditablePinNodes.Append(ResultNode->GetAllResultNodes());
+			}
+			else
+			{
+				EditablePinNodes.Add(EditablePinNode);
+			}
+		}
 
 		FString NewPinName = InSourcePinName;
-		ResultPin = EditablePinNode->CreateUserDefinedPin(NewPinName, InSourcePinType, (InSourcePinDirection == EGPD_Input)? EGPD_Output : EGPD_Input);
+		for (UK2Node_EditablePinBase* CurrentEditablePinNode : EditablePinNodes)
+		{
+			CurrentEditablePinNode->Modify();
+			UEdGraphPin* CreatedPin = nullptr;
+			CreatedPin = CurrentEditablePinNode->CreateUserDefinedPin(NewPinName, InSourcePinType, (InSourcePinDirection == EGPD_Input)? EGPD_Output : EGPD_Input);
+
+			// The final ResultPin is from the node the user dragged and dropped to
+			if (EditablePinNode == CurrentEditablePinNode)
+			{
+				ResultPin = CreatedPin;
+			}
+		}
 
 		FParamsChangedHelper ParamsChangedHelper;
 		ParamsChangedHelper.ModifiedBlueprints.Add(FBlueprintEditorUtils::FindBlueprintForNode(InTargetNode));
diff --git a/Engine/Source/Editor/BlueprintGraph/Private/K2Node_FunctionResult.cpp b/Engine/Source/Editor/BlueprintGraph/Private/K2Node_FunctionResult.cpp
--- a/Engine/Source/Editor/BlueprintGraph/Private/K2Node_FunctionResult.cpp
+++ b/Engine/Source/Editor/BlueprintGraph/Private/K2Node_FunctionResult.cpp
@@ -210,6 +210,18 @@ void UK2Node_FunctionResult::PostPlacedNewNode()
 {
 	Super::PostPlacedNewNode();
 
+	// If the entry is editable, so is the result
+	TArray<UK2Node_FunctionEntry*> AllEntryNodes;
+	if (auto Graph = GetGraph())
+	{
+		Graph->GetNodesOfClass(AllEntryNodes);
+
+		if (AllEntryNodes.Num() > 0)
+		{
+			bIsEditable = AllEntryNodes[0]->bIsEditable;
+		}
+	}
+
 	SyncWithPrimaryResultNode();
 }
 
@@ -217,6 +229,18 @@ void UK2Node_FunctionResult::PostPasteNode()
 {
 	Super::PostPasteNode();
 
+	// If the entry is editable, so is the result
+	TArray<UK2Node_FunctionEntry*> AllEntryNodes;
+	if (auto Graph = GetGraph())
+	{
+		Graph->GetNodesOfClass(AllEntryNodes);
+
+		if (AllEntryNodes.Num() > 0)
+		{
+			bIsEditable = AllEntryNodes[0]->bIsEditable;
+		}
+	}
+
 	SyncWithPrimaryResultNode();
 }
 
diff --git a/Engine/Source/Editor/BlueprintGraph/Private/K2Node_FunctionTerminator.cpp b/Engine/Source/Editor/BlueprintGraph/Private/K2Node_FunctionTerminator.cpp
--- a/Engine/Source/Editor/BlueprintGraph/Private/K2Node_FunctionTerminator.cpp
+++ b/Engine/Source/Editor/BlueprintGraph/Private/K2Node_FunctionTerminator.cpp
@@ -42,6 +42,10 @@ bool UK2Node_FunctionTerminator::CanCreateUserDefinedPin(const FEdGraphPinType&
 		OutErrorMessage = LOCTEXT("MultipleExecPinError", "Cannot support more exec pins!");
 		return false;
 	}
+	else if (!bIsEditable)
+	{
+		OutErrorMessage = LOCTEXT("NotEditableError", "Cannot edit this node!");
+	}
 
 	return bIsEditable;
 }
