[UE-12263] Eliminate potential FMemberReference::ResolveMember() log spam at startup/load time.

change summary:
- modified UK2Node_CallFunction::Serialize() to bypass enabled state validation when not necessary.

diff --git a/Engine/Source/Editor/BlueprintGraph/Private/K2Node_CallFunction.cpp b/Engine/Source/Editor/BlueprintGraph/Private/K2Node_CallFunction.cpp
--- a/Engine/Source/Editor/BlueprintGraph/Private/K2Node_CallFunction.cpp
+++ b/Engine/Source/Editor/BlueprintGraph/Private/K2Node_CallFunction.cpp
@@ -1815,20 +1815,25 @@ void UK2Node_CallFunction::Serialize(FArchive& Ar)
 			}
 		}
 
-		// Don't validate the enabled state if the user has explicitly set it.
-		if (!bUserSetEnabledState)
+		// Don't validate the enabled state if the user has explicitly set it. Also skip validation if we're just duplicating this node.
+		const bool bIsDuplicating = (Ar.GetPortFlags() & PPF_Duplicate) != 0;
+		if (!bIsDuplicating && !bUserSetEnabledState)
 		{
-			if (const UFunction* Function = GetTargetFunction())
+			UClass* SelfScope = GetBlueprintClassFromNode();
+			if (!FunctionReference.IsSelfContext() || SelfScope != nullptr)
 			{
-				// Enable as development-only if specified in metadata. This way existing functions that have the metadata added to them will get their enabled state fixed up on load.
-				if (EnabledState == ENodeEnabledState::Enabled && Function->HasMetaData(FBlueprintMetadata::MD_DevelopmentOnly))
+				if (const UFunction* Function = FunctionReference.ResolveMember<UFunction>(SelfScope))
 				{
-					EnabledState = ENodeEnabledState::DevelopmentOnly;
-				}
-				// Ensure that if the metadata is removed, we also fix up the enabled state to avoid leaving it set as development-only in that case.
-				else if (EnabledState == ENodeEnabledState::DevelopmentOnly && !Function->HasMetaData(FBlueprintMetadata::MD_DevelopmentOnly))
-				{
-					EnabledState = ENodeEnabledState::Enabled;
+					// Enable as development-only if specified in metadata. This way existing functions that have the metadata added to them will get their enabled state fixed up on load.
+					if (EnabledState == ENodeEnabledState::Enabled && Function->HasMetaData(FBlueprintMetadata::MD_DevelopmentOnly))
+					{
+						EnabledState = ENodeEnabledState::DevelopmentOnly;
+					}
+					// Ensure that if the metadata is removed, we also fix up the enabled state to avoid leaving it set as development-only in that case.
+					else if (EnabledState == ENodeEnabledState::DevelopmentOnly && !Function->HasMetaData(FBlueprintMetadata::MD_DevelopmentOnly))
+					{
+						EnabledState = ENodeEnabledState::Enabled;
+					}
 				}
 			}
 		}
