Spawn Actor From Class nodes can no longer be placed in Blueprints that do not support a world context.

diff --git a/Engine/Source/Editor/BlueprintGraph/Private/K2Node_SpawnActorFromClass.cpp b/Engine/Source/Editor/BlueprintGraph/Private/K2Node_SpawnActorFromClass.cpp
--- a/Engine/Source/Editor/BlueprintGraph/Private/K2Node_SpawnActorFromClass.cpp
+++ b/Engine/Source/Editor/BlueprintGraph/Private/K2Node_SpawnActorFromClass.cpp
@@ -509,7 +509,7 @@ FText UK2Node_SpawnActorFromClass::GetNodeTitle(ENodeTitleType::Type TitleType)
 bool UK2Node_SpawnActorFromClass::IsCompatibleWithGraph(const UEdGraph* TargetGraph) const 
 {
 	UBlueprint* Blueprint = FBlueprintEditorUtils::FindBlueprintForGraph(TargetGraph);
-	return Super::IsCompatibleWithGraph(TargetGraph) && (!Blueprint || FBlueprintEditorUtils::FindUserConstructionScript(Blueprint) != TargetGraph);
+	return Super::IsCompatibleWithGraph(TargetGraph) && (!Blueprint || (FBlueprintEditorUtils::FindUserConstructionScript(Blueprint) != TargetGraph && Blueprint->GeneratedClass->GetDefaultObject()->ImplementsGetWorld()));
 }
 
 void UK2Node_SpawnActorFromClass::GetNodeAttributes( TArray<TKeyValuePair<FString, FString>>& OutNodeAttributes ) const
