ForEachEnum will now correctly iterate over all enum values.

Converting an enum to string will report the name correctly.

Changed UKismetNodeHelperLibrary::GetValidIndex to UKismetNodeHelperLibrary::GetValidValue to properly analyze the value as being valid.

Added UKismetNodeHelperLibrary::GetEnumeratorValueFromIndex for converting an index for the enum elements to the actual value.

diff --git a/Engine/Source/Editor/BlueprintGraph/Private/K2Node_CastByteToEnum.cpp b/Engine/Source/Editor/BlueprintGraph/Private/K2Node_CastByteToEnum.cpp
--- a/Engine/Source/Editor/BlueprintGraph/Private/K2Node_CastByteToEnum.cpp
+++ b/Engine/Source/Editor/BlueprintGraph/Private/K2Node_CastByteToEnum.cpp
@@ -62,7 +62,7 @@ FText UK2Node_CastByteToEnum::GetCompactNodeTitle() const
 
 FName UK2Node_CastByteToEnum::GetFunctionName() const
 {
-	const FName FunctionName = GET_FUNCTION_NAME_CHECKED(UKismetNodeHelperLibrary, GetValidIndex);
+	const FName FunctionName = GET_FUNCTION_NAME_CHECKED(UKismetNodeHelperLibrary, GetValidValue);
 	return FunctionName;
 }
 
@@ -90,7 +90,7 @@ void UK2Node_CastByteToEnum::ExpandNode(class FKismetCompilerContext& CompilerCo
 
 		// FUNCTION INPUT BYTE PIN
 		UEdGraphPin* OrgInputPin = FindPinChecked(ByteInputPinName);
-		UEdGraphPin* FunctionIndexPin = CallValidation->FindPinChecked(TEXT("EnumeratorIndex"));
+		UEdGraphPin* FunctionIndexPin = CallValidation->FindPinChecked(TEXT("EnumeratorValue"));
 		check(EGPD_Input == FunctionIndexPin->Direction && Schema->PC_Byte == FunctionIndexPin->PinType.PinCategory);
 		CompilerContext.MovePinLinksToIntermediate(*OrgInputPin, *FunctionIndexPin);
 
diff --git a/Engine/Source/Editor/BlueprintGraph/Private/K2Node_ForEachElementInEnum.cpp b/Engine/Source/Editor/BlueprintGraph/Private/K2Node_ForEachElementInEnum.cpp
--- a/Engine/Source/Editor/BlueprintGraph/Private/K2Node_ForEachElementInEnum.cpp
+++ b/Engine/Source/Editor/BlueprintGraph/Private/K2Node_ForEachElementInEnum.cpp
@@ -34,7 +34,7 @@ struct FForExpandNodeHelper
 		, LoopCounterLimitInPin(nullptr)
 	{ }
 
-	bool BuildLoop(UK2Node* Node, FKismetCompilerContext& CompilerContext, UEdGraph* SourceGraph)
+	bool BuildLoop(UK2Node* Node, FKismetCompilerContext& CompilerContext, UEdGraph* SourceGraph, UEnum* Enum)
 	{
 		const UEdGraphSchema_K2* Schema = CompilerContext.GetSchema();
 		check(Node && SourceGraph && Schema);
@@ -86,12 +86,19 @@ struct FForExpandNodeHelper
 		LoopCounterLimitInPin = Condition->FindPinChecked(TEXT("B"));
 		check(LoopCounterLimitInPin);
 
+		// Convert the Enum index to a value
+		UK2Node_CallFunction* GetEnumeratorValueFromIndexCall = CompilerContext.SpawnIntermediateNode<UK2Node_CallFunction>(Node, SourceGraph); 
+		GetEnumeratorValueFromIndexCall->SetFromFunction(UKismetNodeHelperLibrary::StaticClass()->FindFunctionByName(GET_FUNCTION_NAME_CHECKED(UKismetNodeHelperLibrary, GetEnumeratorValueFromIndex)));
+		GetEnumeratorValueFromIndexCall->AllocateDefaultPins();
+		Schema->TrySetDefaultObject(*GetEnumeratorValueFromIndexCall->FindPinChecked(TEXT("Enum")), Enum);
+		bResult &= Schema->TryCreateConnection(GetEnumeratorValueFromIndexCall->FindPinChecked(TEXT("EnumeratorIndex")), LoopCounterOutPin);
+
 		// Array Index assigned
 		UK2Node_AssignmentStatement* ArrayIndexAssign = CompilerContext.SpawnIntermediateNode<UK2Node_AssignmentStatement>(Node, SourceGraph);
 		ArrayIndexAssign->AllocateDefaultPins();
 		bResult &= Schema->TryCreateConnection(Branch->GetThenPin(), ArrayIndexAssign->GetExecPin());
 		bResult &= Schema->TryCreateConnection(ArrayIndexAssign->GetVariablePin(), ArrayIndexOutPin);
-		bResult &= Schema->TryCreateConnection(ArrayIndexAssign->GetValuePin(), LoopCounterOutPin);
+		bResult &= Schema->TryCreateConnection(ArrayIndexAssign->GetValuePin(), GetEnumeratorValueFromIndexCall->GetReturnValuePin());
 
 		// body sequence
 		UK2Node_ExecutionSequence* Sequence = CompilerContext.SpawnIntermediateNode<UK2Node_ExecutionSequence>(Node, SourceGraph);
@@ -186,7 +193,7 @@ void UK2Node_ForEachElementInEnum::ExpandNode(class FKismetCompilerContext& Comp
 	}
 
 	FForExpandNodeHelper ForLoop;
-	if (!ForLoop.BuildLoop(this, CompilerContext, SourceGraph))
+	if (!ForLoop.BuildLoop(this, CompilerContext, SourceGraph, Enum))
 	{
 		CompilerContext.MessageLog.Error(*NSLOCTEXT("K2Node", "ForEachElementInEnum_ForError", "For Expand error in @@").ToString(), this);
 	}
diff --git a/Engine/Source/Runtime/Engine/Classes/Kismet/KismetNodeHelperLibrary.h b/Engine/Source/Runtime/Engine/Classes/Kismet/KismetNodeHelperLibrary.h
--- a/Engine/Source/Runtime/Engine/Classes/Kismet/KismetNodeHelperLibrary.h
+++ b/Engine/Source/Runtime/Engine/Classes/Kismet/KismetNodeHelperLibrary.h
@@ -129,9 +129,17 @@ class ENGINE_API UKismetNodeHelperLibrary : public UBlueprintFunctionLibrary
 
 	/**
 	 * @param Enum - Enumeration
-	 * @param EnumeratorIndex - Input index
+	 * @param EnumeratorIndex - Input value
 	 * @return - if EnumeratorIndex is valid return EnumeratorIndex, otherwise return MAX value of Enum
 	 */
 	UFUNCTION(BlueprintPure, meta=(BlueprintInternalUseOnly = "TRUE"))
-	static uint8 GetValidIndex(const UEnum* Enum, uint8 EnumeratorIndex);
+	static uint8 GetValidValue(const UEnum* Enum, uint8 EnumeratorValue);
+
+	/**
+	 * @param Enum - Enumeration
+	 * @param EnumeratorIndex - Input index
+	 * @return - The value of the enumerator, or INDEX_NONE
+	 */
+	UFUNCTION(BlueprintPure, meta=(BlueprintInternalUseOnly = "TRUE"))
+	static uint8 GetEnumeratorValueFromIndex(const UEnum* Enum, uint8 EnumeratorIndex);
 };
diff --git a/Engine/Source/Runtime/Engine/Private/KismetNodeHelperLibrary.cpp b/Engine/Source/Runtime/Engine/Private/KismetNodeHelperLibrary.cpp
--- a/Engine/Source/Runtime/Engine/Private/KismetNodeHelperLibrary.cpp
+++ b/Engine/Source/Runtime/Engine/Private/KismetNodeHelperLibrary.cpp
@@ -158,12 +158,23 @@ FString UKismetNodeHelperLibrary::GetEnumeratorUserFriendlyName(const UEnum* Enu
 	return FName().ToString();
 }
 
-uint8 UKismetNodeHelperLibrary::GetValidIndex(const UEnum* Enum, uint8 EnumeratorIndex)
+uint8 UKismetNodeHelperLibrary::GetValidValue(const UEnum* Enum, uint8 EnumeratorValue)
 {
-	const int32 EnumNum = Enum ? Enum->NumEnums() : 0;
-	if (ensure(EnumNum > 0))
+	if (NULL != Enum)
+	{
+		if (Enum->IsValidEnumValue(EnumeratorValue))
+		{
+			return EnumeratorValue;
+		}
+	}
+	return Enum->GetMaxEnumValue();
+}
+
+uint8 UKismetNodeHelperLibrary::GetEnumeratorValueFromIndex(const UEnum* Enum, uint8 EnumeratorIndex)
+{
+	if (NULL != Enum)
 	{
-		return (EnumeratorIndex < EnumNum) ? EnumeratorIndex : (EnumNum - 1);
+		return Enum->GetValueByIndex(EnumeratorIndex);
 	}
-	return 0;
+	return INDEX_NONE;
 }
\ No newline at end of file
