Fixes issues when creating a local variable using an orphaned variable node where other matching orphaned variable nodes would not update to point at the new local variable.

diff --git a/Engine/Source/Editor/BlueprintGraph/Private/EdGraphSchema_K2.cpp b/Engine/Source/Editor/BlueprintGraph/Private/EdGraphSchema_K2.cpp
--- a/Engine/Source/Editor/BlueprintGraph/Private/EdGraphSchema_K2.cpp
+++ b/Engine/Source/Editor/BlueprintGraph/Private/EdGraphSchema_K2.cpp
@@ -1615,7 +1615,17 @@ void UEdGraphSchema_K2::OnCreateNonExistentLocalVariable( UK2Node_Variable* Vari
 			FGuid LocalVarGuid = FBlueprintEditorUtils::FindLocalVariableGuidByName(OwnerBlueprint, Variable->GetGraph(), VarName);
 			if (LocalVarGuid.IsValid())
 			{
-				Variable->VariableReference.SetLocalMember(VarName, Variable->GetGraph()->GetName(), LocalVarGuid);
+				// Loop through every variable in the graph, check if the variable references are the same, and update them
+				FMemberReference OldReference = Variable->VariableReference;
+				TArray<UK2Node_Variable*> VariableNodeList;
+				Variable->GetGraph()->GetNodesOfClass(VariableNodeList);
+				for( UK2Node_Variable* VariableNode : VariableNodeList)
+				{
+					if (VariableNode->VariableReference.IsSameReference(OldReference))
+					{
+						VariableNode->VariableReference.SetLocalMember(VarName, Variable->GetGraph()->GetName(), LocalVarGuid);
+					}
+				}
 			}
 		}
 	}	
diff --git a/Engine/Source/Runtime/Engine/Classes/Engine/MemberReference.h b/Engine/Source/Runtime/Engine/Classes/Engine/MemberReference.h
--- a/Engine/Source/Runtime/Engine/Classes/Engine/MemberReference.h
+++ b/Engine/Source/Runtime/Engine/Classes/Engine/MemberReference.h
@@ -261,6 +261,17 @@ public:
 		return MemberScope;
 	}
 
+	/** Compares with another MemberReference to see if they are identical */
+	bool IsSameReference(FMemberReference& InReference)
+	{
+		return 
+			bSelfContext == InReference.bSelfContext &&
+			MemberParent == InReference.MemberParent &&
+			MemberName == InReference.MemberName &&
+			MemberGuid == InReference.MemberGuid &&
+			MemberScope == InReference.MemberScope;
+	}
+
 	/** Returns whether or not the variable has been deprecated */
 	bool IsDeprecated() const
 	{
