When using a pin connection to specify the class for "Spawn Actor From Class", the default value overrides from the node should be assigned to the new actor when it is spawned.

diff --git a/Engine/Source/Editor/BlueprintGraph/Private/K2Node_SpawnActorFromClass.cpp b/Engine/Source/Editor/BlueprintGraph/Private/K2Node_SpawnActorFromClass.cpp
--- a/Engine/Source/Editor/BlueprintGraph/Private/K2Node_SpawnActorFromClass.cpp
+++ b/Engine/Source/Editor/BlueprintGraph/Private/K2Node_SpawnActorFromClass.cpp
@@ -580,6 +580,9 @@ void UK2Node_SpawnActorFromClass::ExpandNode(class FKismetCompilerContext& Compi
 	UEdGraphPin* SpawnNodeThen = SpawnNode->GetThenPin();
 	UEdGraphPin* SpawnNodeResult = SpawnNode->GetResultPin();
 
+	// Cache the class to spawn. Note, this is the compile time class that the pin was set to or the variable type it was connected to. Runtime it could be a child.
+	UClass* ClassToSpawn = GetClassToSpawn();
+
 	UClass* SpawnClass = (SpawnClassPin != NULL) ? Cast<UClass>(SpawnClassPin->DefaultObject) : NULL;
 	if((0 == SpawnClassPin->LinkedTo.Num()) && (NULL == SpawnClass))
 	{
@@ -664,7 +667,7 @@ void UK2Node_SpawnActorFromClass::ExpandNode(class FKismetCompilerContext& Compi
 	// create 'set var' nodes
 
 	// Get 'result' pin from 'begin spawn', this is the actual actor we want to set properties on
-	UEdGraphPin* LastThen = FKismetCompilerUtilities::GenerateAssignmentNodes(CompilerContext, SourceGraph, CallBeginSpawnNode, SpawnNode, CallBeginResult, GetClassToSpawn() );
+	UEdGraphPin* LastThen = FKismetCompilerUtilities::GenerateAssignmentNodes(CompilerContext, SourceGraph, CallBeginSpawnNode, SpawnNode, CallBeginResult, ClassToSpawn );
 
 	// Make exec connection between 'then' on last node and 'finish'
 	LastThen->MakeLinkTo(CallFinishExec);
