Can no longer add reroute nodes to read only graphs such as math expression graphs.

Moved logic for deciding a read only graph from BlueprintEditor to BlueprintEditorUtils for more global access to the logic.

diff --git a/Engine/Source/Editor/BlueprintGraph/Private/EdGraphSchema_K2.cpp b/Engine/Source/Editor/BlueprintGraph/Private/EdGraphSchema_K2.cpp
--- a/Engine/Source/Editor/BlueprintGraph/Private/EdGraphSchema_K2.cpp
+++ b/Engine/Source/Editor/BlueprintGraph/Private/EdGraphSchema_K2.cpp
@@ -5912,17 +5912,20 @@ void UEdGraphSchema_K2::OnPinConnectionDoubleCicked(UEdGraphPin* PinA, UEdGraphP
 
 	// Create a new knot
 	UEdGraph* ParentGraph = PinA->GetOwningNode()->GetGraph();
-	UK2Node_Knot* NewKnot = FEdGraphSchemaAction_K2NewNode::SpawnNodeFromTemplate<UK2Node_Knot>(ParentGraph, NewObject<UK2Node_Knot>(), KnotTopLeft);
+	if (!FBlueprintEditorUtils::IsGraphReadOnly(ParentGraph))
+	{
+		UK2Node_Knot* NewKnot = FEdGraphSchemaAction_K2NewNode::SpawnNodeFromTemplate<UK2Node_Knot>(ParentGraph, NewObject<UK2Node_Knot>(), KnotTopLeft);
 
-	// Move the connections across (only notifying the knot, as the other two didn't really change)
-	PinA->BreakLinkTo(PinB);
-	PinA->MakeLinkTo((PinA->Direction == EGPD_Output) ? NewKnot->GetInputPin() : NewKnot->GetOutputPin());
-	PinB->MakeLinkTo((PinB->Direction == EGPD_Output) ? NewKnot->GetInputPin() : NewKnot->GetOutputPin());
-	NewKnot->PostReconstructNode();
+		// Move the connections across (only notifying the knot, as the other two didn't really change)
+		PinA->BreakLinkTo(PinB);
+		PinA->MakeLinkTo((PinA->Direction == EGPD_Output) ? NewKnot->GetInputPin() : NewKnot->GetOutputPin());
+		PinB->MakeLinkTo((PinB->Direction == EGPD_Output) ? NewKnot->GetInputPin() : NewKnot->GetOutputPin());
+		NewKnot->PostReconstructNode();
 
-	// Dirty the blueprint
-	UBlueprint* Blueprint = FBlueprintEditorUtils::FindBlueprintForGraphChecked(CastChecked<UEdGraph>(ParentGraph));
-	FBlueprintEditorUtils::MarkBlueprintAsModified(Blueprint);
+		// Dirty the blueprint
+		UBlueprint* Blueprint = FBlueprintEditorUtils::FindBlueprintForGraphChecked(CastChecked<UEdGraph>(ParentGraph));
+		FBlueprintEditorUtils::MarkBlueprintAsModified(Blueprint);
+	}
 }
 
 void UEdGraphSchema_K2::ConfigureVarNode(UK2Node_Variable* InVarNode, FName InVariableName, UStruct* InVariableSource, UBlueprint* InTargetBlueprint)
diff --git a/Engine/Source/Editor/Kismet/Private/BlueprintEditor.cpp b/Engine/Source/Editor/Kismet/Private/BlueprintEditor.cpp
--- a/Engine/Source/Editor/Kismet/Private/BlueprintEditor.cpp
+++ b/Engine/Source/Editor/Kismet/Private/BlueprintEditor.cpp
@@ -1235,7 +1235,6 @@ TSharedRef<SGraphEditor> FBlueprintEditor::CreateGraphEditorWidget(TSharedRef<FT
 		InEvents.OnCreateActionMenu = SGraphEditor::FOnCreateActionMenu::CreateSP(this, &FBlueprintEditor::OnCreateGraphActionMenu);
 	}
 
-
 	TSharedRef<SGraphEditor> Editor = SNew(SGraphEditor)
 		.AdditionalCommands(GraphEditorCommands)
 		.IsEditable(this, &FBlueprintEditor::IsEditable, InGraph)
@@ -8074,28 +8073,18 @@ UEdGraph* FBlueprintEditor::GetFocusedGraph() const
 
 bool FBlueprintEditor::IsEditable(UEdGraph* InGraph) const
 {
-	return InEditingMode() && !IsGraphReadOnly(InGraph);
+	return InEditingMode() && !FBlueprintEditorUtils::IsGraphReadOnly(InGraph);
 }
 
 bool FBlueprintEditor::IsGraphReadOnly(UEdGraph* InGraph) const
 {
-	return !(InGraph && InGraph->bEditable && IsGraphPanelEnabled(InGraph));
-}
-
-bool FBlueprintEditor::IsGraphPanelEnabled(UEdGraph* InGraph) const
-{
-	const UBlueprint* BlueprintForGraph = FBlueprintEditorUtils::FindBlueprintForGraph(InGraph);
-	bool const bIsInterface = ((BlueprintForGraph != NULL) && (BlueprintForGraph->BlueprintType == BPTYPE_Interface));
-	bool const bIsDelegate  = FBlueprintEditorUtils::IsDelegateSignatureGraph(InGraph);
-	bool const bIsMathExpression = FBlueprintEditorUtils::IsMathExpressionGraph(InGraph);
-
-	return !bIsInterface && !bIsDelegate && !bIsMathExpression;
+	return FBlueprintEditorUtils::IsGraphReadOnly(InGraph);
 }
 
 float FBlueprintEditor::GetInstructionTextOpacity(UEdGraph* InGraph) const
 {
 	UBlueprintEditorSettings const* Settings = GetDefault<UBlueprintEditorSettings>();
-	if ((InGraph == nullptr) || !IsEditable(InGraph) || !IsGraphPanelEnabled(InGraph) || !Settings->bShowGraphInstructionText)
+	if ((InGraph == nullptr) || !IsEditable(InGraph) || FBlueprintEditorUtils::IsGraphReadOnly(InGraph) || !Settings->bShowGraphInstructionText)
 	{
 		return 0.0f;
 	}
diff --git a/Engine/Source/Editor/Kismet/Public/BlueprintEditor.h b/Engine/Source/Editor/Kismet/Public/BlueprintEditor.h
--- a/Engine/Source/Editor/Kismet/Public/BlueprintEditor.h
+++ b/Engine/Source/Editor/Kismet/Public/BlueprintEditor.h
@@ -264,8 +264,6 @@ public:
 
 	/** Returns whether a graph is editable or not */
 	virtual bool IsEditable(UEdGraph* InGraph) const;
-	/** Determines if the graph's title bar should be the only interactable widget .*/
-	bool IsGraphPanelEnabled(UEdGraph* InGraph) const;
 	/** Determines if the graph is ReadOnly, this differs from editable in that it is never expected to be edited and is in a read-only state */
 	bool IsGraphReadOnly(UEdGraph* InGraph) const;
 
diff --git a/Engine/Source/Editor/UnrealEd/Private/Kismet2/BlueprintEditorUtils.cpp b/Engine/Source/Editor/UnrealEd/Private/Kismet2/BlueprintEditorUtils.cpp
--- a/Engine/Source/Editor/UnrealEd/Private/Kismet2/BlueprintEditorUtils.cpp
+++ b/Engine/Source/Editor/UnrealEd/Private/Kismet2/BlueprintEditorUtils.cpp
@@ -2736,6 +2736,26 @@ UEdGraph* FBlueprintEditorUtils::GetTopLevelGraph(const UEdGraph* InGraph)
 	return GraphToTest;
 }
 
+bool FBlueprintEditorUtils::IsGraphReadOnly(UEdGraph* InGraph)
+{
+	bool bGraphReadOnly = true;
+	if (InGraph)
+	{
+		bGraphReadOnly = !InGraph->bEditable;
+
+		if (!bGraphReadOnly)
+		{
+			const UBlueprint* BlueprintForGraph = FBlueprintEditorUtils::FindBlueprintForGraph(InGraph);
+			bool const bIsInterface = ((BlueprintForGraph != NULL) && (BlueprintForGraph->BlueprintType == BPTYPE_Interface));
+			bool const bIsDelegate  = FBlueprintEditorUtils::IsDelegateSignatureGraph(InGraph);
+			bool const bIsMathExpression = FBlueprintEditorUtils::IsMathExpressionGraph(InGraph);
+
+			bGraphReadOnly = bIsInterface || bIsDelegate || bIsMathExpression;
+		}
+	}
+	return bGraphReadOnly;
+}
+
 UK2Node_Event* FBlueprintEditorUtils::FindOverrideForFunction(const UBlueprint* Blueprint, const UClass* SignatureClass, FName SignatureName)
 {
 	TArray<UK2Node_Event*> AllEvents;
diff --git a/Engine/Source/Editor/UnrealEd/Public/Kismet2/BlueprintEditorUtils.h b/Engine/Source/Editor/UnrealEd/Public/Kismet2/BlueprintEditorUtils.h
--- a/Engine/Source/Editor/UnrealEd/Public/Kismet2/BlueprintEditorUtils.h
+++ b/Engine/Source/Editor/UnrealEd/Public/Kismet2/BlueprintEditorUtils.h
@@ -352,6 +352,9 @@ public:
 	 */
 	static UEdGraph* GetTopLevelGraph(const UEdGraph* InGraph);
 
+	/** Determines if the graph is ReadOnly, this differs from editable in that it is never expected to be edited and is in a read-only state */
+	static bool IsGraphReadOnly(UEdGraph* InGraph);
+
 	/** Look to see if an event already exists to override a particular function */
 	static class UK2Node_Event* FindOverrideForFunction(const UBlueprint* Blueprint, const UClass* SignatureClass, FName SignatureName);
 
