Improved logging of crash to gather information as to what might be the cause.

diff --git a/Engine/Source/Editor/BlueprintGraph/Private/K2Node_VariableGet.cpp b/Engine/Source/Editor/BlueprintGraph/Private/K2Node_VariableGet.cpp
--- a/Engine/Source/Editor/BlueprintGraph/Private/K2Node_VariableGet.cpp
+++ b/Engine/Source/Editor/BlueprintGraph/Private/K2Node_VariableGet.cpp
@@ -264,6 +264,20 @@ FText UK2Node_VariableGet::GetNodeTitle(ENodeTitleType::Type TitleType) const
 	for (int32 PinIndex = 0; PinIndex < Pins.Num(); ++PinIndex)
 	{
 		UEdGraphPin* Pin = Pins[PinIndex];
+
+		// The following code is to attempt to log info related to UE-19729
+		if (TitleType == ENodeTitleType::ListView)
+		{
+			if (UEdGraph* Graph = Cast<UEdGraph>(GetOuter()))
+			{
+				FString VariableName = GetVarNameString();
+				FString BlueprintPath = FBlueprintEditorUtils::FindBlueprintForGraph(Graph)->GetPathName();
+				FString SetupStyle = bIsPureGet? TEXT("pure") : TEXT("validated");
+				FString VariableResolves = (VariableReference.ResolveMember<UProperty>(GetBlueprintClassFromNode()) != nullptr)? TEXT("resolves") : TEXT("does not resolve");
+				checkf(Pin, TEXT("Get node for variable '%s' in Blueprint '%s' which is setup as %s and has %d pins. Variable %s"), *VariableName, *BlueprintPath, *SetupStyle, Pins.Num(), *VariableResolves);
+			}
+		}
+
 		if (Pin->Direction == EGPD_Output)
 		{
 			++NumOutputsFound;
@@ -459,11 +473,6 @@ void UK2Node_VariableGet::Serialize(FArchive& Ar)
 				{
 					if (!Blueprint->bBeingCompiled)
 					{
-						FString VariableName = GetVarNameString();
-						FString BlueprintPath = Blueprint->GetPathName();
-						FString SetupStyle = bIsPureGet? TEXT("pure") : TEXT("validated");
-						UE_LOG(LogBlueprintUserMessages, Log, TEXT("Serialization for Get node for variable '%s' in Blueprint '%s' which is setup as %s"), *VariableName, *BlueprintPath, *SetupStyle);
-
 						// The following line may spur the crash noted in UE-19729 and will confirm that the crash happens before the FiB gather.
 						GetNodeTitle(ENodeTitleType::ListView);
 					}
