Fixed crash from invalid macros after deleting source graph in MacroLibrary.

Added null check for Macro metadata return to stop attempts to access null pointer when retrieving the macro's title.

diff --git a/Engine/Source/Editor/BlueprintGraph/Private/K2Node_MacroInstance.cpp b/Engine/Source/Editor/BlueprintGraph/Private/K2Node_MacroInstance.cpp
--- a/Engine/Source/Editor/BlueprintGraph/Private/K2Node_MacroInstance.cpp
+++ b/Engine/Source/Editor/BlueprintGraph/Private/K2Node_MacroInstance.cpp
@@ -365,7 +365,12 @@ FName UK2Node_MacroInstance::GetPaletteIcon(FLinearColor& OutColor) const
 
 FText UK2Node_MacroInstance::GetCompactNodeTitle() const
 {
-	return GetAssociatedGraphMetadata(GetMacroGraph())->CompactNodeTitle;	
+	FText ResultText;
+	if (FKismetUserDeclaredFunctionMetadata* MacroGraphMetadata = GetAssociatedGraphMetadata(GetMacroGraph()))
+	{
+		ResultText = MacroGraphMetadata->CompactNodeTitle;
+	}
+	return ResultText;	
 }
 
 bool UK2Node_MacroInstance::ShouldDrawCompact() const
