Resolved Blueprint compiler errors with duplicating Matinee Actors and adding controllers.

MatineeControllerName in MatineeActor will be set when the node is spawned.

diff --git a/Engine/Source/Editor/BlueprintGraph/Private/K2Node_MatineeController.cpp b/Engine/Source/Editor/BlueprintGraph/Private/K2Node_MatineeController.cpp
--- a/Engine/Source/Editor/BlueprintGraph/Private/K2Node_MatineeController.cpp
+++ b/Engine/Source/Editor/BlueprintGraph/Private/K2Node_MatineeController.cpp
@@ -174,6 +174,7 @@ void UK2Node_MatineeController::GetMenuActions(FBlueprintActionDatabaseRegistrar
 	{
 		UK2Node_MatineeController* MatineeNode = CastChecked<UK2Node_MatineeController>(NewNode);
 		MatineeNode->MatineeActor = CastChecked<AMatineeActor>(BindObject);
+		MatineeNode->MatineeActor->MatineeControllerName = MatineeNode->GetFName();
 		return true;
 	};
 
