Copy and pasting local variable nodes for functions into collapsed graphs will no longer invalidate the node.

Pasted local variable nodes were targetting their current graph instead of the function graph.

diff --git a/Engine/Source/Editor/BlueprintGraph/Private/K2Node_Variable.cpp b/Engine/Source/Editor/BlueprintGraph/Private/K2Node_Variable.cpp
--- a/Engine/Source/Editor/BlueprintGraph/Private/K2Node_Variable.cpp
+++ b/Engine/Source/Editor/BlueprintGraph/Private/K2Node_Variable.cpp
@@ -816,10 +816,11 @@ void UK2Node_Variable::PostPasteNode()
 		// If the current graph is a Function graph, look to see if there is a compatible local variable (same name)
 		if (GetGraph()->GetSchema()->GetGraphType(GetGraph()) == GT_Function)
 		{
-			FBPVariableDescription* VariableDescription = FBlueprintEditorUtils::FindLocalVariable(Blueprint, GetGraph(), VariableReference.GetMemberName());
+			UEdGraph* FunctionGraph = FBlueprintEditorUtils::GetTopLevelGraph(GetGraph());
+			FBPVariableDescription* VariableDescription = FBlueprintEditorUtils::FindLocalVariable(Blueprint, FunctionGraph, VariableReference.GetMemberName());
 			if(VariableDescription)
 			{
-				VariableReference.SetLocalMember(VariableReference.GetMemberName(), GetGraph()->GetName(), VariableReference.GetMemberGuid());
+				VariableReference.SetLocalMember(VariableReference.GetMemberName(), FunctionGraph->GetName(), VariableReference.GetMemberGuid());
 			}
 		}
 		// If no variable was found, ResolveMember should automatically find a member variable with the same name in the current Blueprint and hook up to it as expected
