Toggling a function to pure will automatically remove all breakpoints from call sites.

diff --git a/Engine/Source/Editor/BlueprintGraph/Private/K2Node_CallFunction.cpp b/Engine/Source/Editor/BlueprintGraph/Private/K2Node_CallFunction.cpp
--- a/Engine/Source/Editor/BlueprintGraph/Private/K2Node_CallFunction.cpp
+++ b/Engine/Source/Editor/BlueprintGraph/Private/K2Node_CallFunction.cpp
@@ -975,6 +975,16 @@ void UK2Node_CallFunction::PostReconstructNode()
 	{
 		FDynamicOutputHelper(TypePickerPin).ConformOutputType();
 	}
+
+	if (IsNodePure())
+	{
+		// Remove any pre-existing breakpoint on this node since pure nodes cannot have breakpoints
+		if (UBreakpoint* ExistingBreakpoint = FKismetDebugUtilities::FindBreakpointForNode(GetBlueprint(), this))
+		{
+			// Remove the breakpoint
+			FKismetDebugUtilities::StartDeletingBreakpoint(ExistingBreakpoint, GetBlueprint());
+		}
+	}
 }
 
 void UK2Node_CallFunction::DestroyNode()
