Switched to using the original node guid to facilitate deterministic cooking

diff --git a/Engine/Source/Editor/BlueprintGraph/Private/K2Node_BaseAsyncTask.cpp b/Engine/Source/Editor/BlueprintGraph/Private/K2Node_BaseAsyncTask.cpp
--- a/Engine/Source/Editor/BlueprintGraph/Private/K2Node_BaseAsyncTask.cpp
+++ b/Engine/Source/Editor/BlueprintGraph/Private/K2Node_BaseAsyncTask.cpp
@@ -222,8 +222,8 @@ bool UK2Node_BaseAsyncTask::FBaseAsyncTaskHelper::HandleDelegateImplementation(
 	AddDelegateNode->AllocateDefaultPins();
 		bIsErrorFree &= Schema->TryCreateConnection(AddDelegateNode->FindPinChecked(Schema->PN_Self), ProxyObjectPin);
 		bIsErrorFree &= Schema->TryCreateConnection(InOutLastThenPin, AddDelegateNode->FindPinChecked(Schema->PN_Execute));
-		InOutLastThenPin = AddDelegateNode->FindPinChecked(Schema->PN_Then);
-		CurrentCENode->CustomFunctionName = *FString::Printf(TEXT("%s_%s"), *CurrentProperty->GetName(), *CurrentCENode->NodeGuid.ToString());
+		InOutLastThenPin = AddDelegateNode->FindPinChecked(Schema->PN_Then); 
+		CurrentCENode->CustomFunctionName = *FString::Printf(TEXT("%s_%s"), *CurrentProperty->GetName(), *CurrentNode->NodeGuid.ToString());
 		CurrentCENode->AllocateDefaultPins();
 
 		bIsErrorFree &= FBaseAsyncTaskHelper::CreateDelegateForNewFunction(AddDelegateNode->GetDelegatePin(), CurrentCENode->GetFunctionName(), CurrentNode, SourceGraph, CompilerContext);
