Fixed crash in FKCHandler_CallFunction::FindFunction when the calling context is null

diff --git a/Engine/Source/Editor/BlueprintGraph/Private/CallFunctionHandler.cpp b/Engine/Source/Editor/BlueprintGraph/Private/CallFunctionHandler.cpp
--- a/Engine/Source/Editor/BlueprintGraph/Private/CallFunctionHandler.cpp
+++ b/Engine/Source/Editor/BlueprintGraph/Private/CallFunctionHandler.cpp
@@ -597,9 +597,15 @@ void FKCHandler_CallFunction::RegisterNet(FKismetFunctionContext& Context, UEdGr
 UFunction* FKCHandler_CallFunction::FindFunction(FKismetFunctionContext& Context, UEdGraphNode* Node)
 {
 	UClass* CallingContext = GetCallingContext(Context, Node);
-	FString FunctionName = GetFunctionNameFromNode(Node);
 
-	return CallingContext->FindFunctionByName(*FunctionName);
+	if (CallingContext)
+	{
+		FString FunctionName = GetFunctionNameFromNode(Node);
+
+		return CallingContext->FindFunctionByName(*FunctionName);
+	}
+
+	return nullptr;
 }
 
 void FKCHandler_CallFunction::Transform(FKismetFunctionContext& Context, UEdGraphNode* Node)
