Missed file

diff --git a/Engine/Source/Editor/BlueprintGraph/Private/DynamicCastHandler.cpp b/Engine/Source/Editor/BlueprintGraph/Private/DynamicCastHandler.cpp
--- a/Engine/Source/Editor/BlueprintGraph/Private/DynamicCastHandler.cpp
+++ b/Engine/Source/Editor/BlueprintGraph/Private/DynamicCastHandler.cpp
@@ -114,7 +114,7 @@ void FKCHandler_DynamicCast::Compile(FKismetFunctionContext& Context, UEdGraphNo
 			return;
 		}
 		CastOpType = KCST_MetaCast;
-		//ClassTerm->Type.PinCategory = CompilerContext.GetSchema()->PC_Class;
+		ClassTerm->Type.PinCategory = CompilerContext.GetSchema()->PC_Class;
 	}
 
 	// Cast Statement
