Stoped kismet enums from reinitializing FText from strings.

diff --git a/Engine/Source/Editor/BlueprintGraph/Classes/K2Node_SwitchEnum.h b/Engine/Source/Editor/BlueprintGraph/Classes/K2Node_SwitchEnum.h
--- a/Engine/Source/Editor/BlueprintGraph/Classes/K2Node_SwitchEnum.h
+++ b/Engine/Source/Editor/BlueprintGraph/Classes/K2Node_SwitchEnum.h
@@ -22,7 +22,7 @@ class UK2Node_SwitchEnum : public UK2Node_Switch, public INodeDependingOnEnumInt
 
 	/** List of the current entries in the enum */
 	UPROPERTY()
-	TArray<FString> EnumFriendlyNames;
+	TArray<FText> EnumFriendlyNames;
 
 	// INodeDependingOnEnumInterface
 	virtual class UEnum* GetEnum() const override { return Enum; }
diff --git a/Engine/Source/Editor/BlueprintGraph/Private/K2Node_SwitchEnum.cpp b/Engine/Source/Editor/BlueprintGraph/Private/K2Node_SwitchEnum.cpp
--- a/Engine/Source/Editor/BlueprintGraph/Private/K2Node_SwitchEnum.cpp
+++ b/Engine/Source/Editor/BlueprintGraph/Private/K2Node_SwitchEnum.cpp
@@ -38,10 +38,10 @@ void UK2Node_SwitchEnum::SetEnum(UEnum* InEnum)
 				FString const EnumValueName = Enum->GetEnumName(EnumIndex);
 				EnumEntries.Add( FName(*EnumValueName) );
 
-				FString EnumFriendlyName = Enum->GetEnumText(EnumIndex).ToString();
-				if (EnumFriendlyName.Len() == 0)
+				FText EnumFriendlyName = Enum->GetEnumText(EnumIndex);
+				if (EnumFriendlyName.IsEmpty() )
 				{
-					EnumFriendlyName = EnumValueName;
+					EnumFriendlyName = FText::FromString(EnumValueName);
 				}
 				EnumFriendlyNames.Add( EnumFriendlyName );
 			}
@@ -156,7 +156,7 @@ void UK2Node_SwitchEnum::CreateCasePins()
 		int32 Index = EnumIt.GetIndex();
 		if (EnumFriendlyNames.IsValidIndex(Index))
 		{
-			NewPin->PinFriendlyName = FText::FromString(EnumFriendlyNames[Index]);
+			NewPin->PinFriendlyName = EnumFriendlyNames[Index];
 		}
 		
 		if(bShouldUseAdvancedView && (EnumIt.GetIndex() > 2))
