Fix compile - ambiguous UClass.

diff --git a/Engine/Source/Editor/BlueprintGraph/Private/K2Node_Variable.cpp b/Engine/Source/Editor/BlueprintGraph/Private/K2Node_Variable.cpp
--- a/Engine/Source/Editor/BlueprintGraph/Private/K2Node_Variable.cpp
+++ b/Engine/Source/Editor/BlueprintGraph/Private/K2Node_Variable.cpp
@@ -36,7 +36,7 @@ void UK2Node_Variable::Serialize(FArchive& Ar)
 			{
 				const bool bSelf = VariableReference.IsSelfContext();
 				UClass* MemberParentClass = VariableReference.GetMemberParentClass(nullptr);
-				if (UBlueprint::GetGuidFromClassByFieldName<UProperty>(bSelf? GetBlueprint()->GeneratedClass : MemberParentClass, VariableReference.GetMemberName(), VarGuid))
+				if (UBlueprint::GetGuidFromClassByFieldName<UProperty>(bSelf? *GetBlueprint()->GeneratedClass : MemberParentClass, VariableReference.GetMemberName(), VarGuid))
 				{
 					VariableReference.SetDirect(VariableReference.GetMemberName(), VarGuid, bSelf ? nullptr : MemberParentClass, bSelf);
 				}
