[UE-17698] Ensure that we clear the UnloadedActionRegistry table when rebuilding the Blueprint action database after a hot reload.

change summary:
- modified FBlueprintActionDatabase::RefreshAll() to clear out the UnloadedActionRegistry in addition to the "normal" ActionRegistry.
- modified FBlueprintActionDatabase::AddReferencedObjects() to ensure that there are no "orphaned" unloaded actions at GC time, but also still add them to the reference collector if they do happen to exist (so we don't crash) - this will hopefully better alert us to this situation in the future.

diff --git a/Engine/Source/Editor/BlueprintGraph/Private/BlueprintActionDatabase.cpp b/Engine/Source/Editor/BlueprintGraph/Private/BlueprintActionDatabase.cpp
--- a/Engine/Source/Editor/BlueprintGraph/Private/BlueprintActionDatabase.cpp
+++ b/Engine/Source/Editor/BlueprintGraph/Private/BlueprintActionDatabase.cpp
@@ -1012,8 +1012,10 @@ FBlueprintActionDatabase::FBlueprintActionDatabase()
 //------------------------------------------------------------------------------
 void FBlueprintActionDatabase::AddReferencedObjects(FReferenceCollector& Collector)
 {
+	TSet<UBlueprintNodeSpawner*> AllActions;
 	for (auto& ActionListIt : ActionRegistry)
 	{
+		AllActions.Append(ActionListIt.Value);
 		Collector.AddReferencedObjects(ActionListIt.Value);
 	}
 
@@ -1021,9 +1023,19 @@ void FBlueprintActionDatabase::AddReferencedObjects(FReferenceCollector& Collect
 	// accounted for in the regular ActionRegistry, but just in case we fail to 
 	// remove an element from here when we should.... this'll make sure these 
 	// elements stick around (so we don't crash in ClearUnloadedAssetActions)
-	for (auto& UnloadedActionListIt : UnloadedActionRegistry)
+	if (UnloadedActionRegistry.Num() > 0)
 	{
-		Collector.AddReferencedObjects(UnloadedActionListIt.Value);
+		TSet<UBlueprintNodeSpawner*> UnloadedActions;
+		for (auto& UnloadedActionListIt : UnloadedActionRegistry)
+		{
+			UnloadedActions.Append(UnloadedActionListIt.Value);
+		}
+
+		auto OrphanedUnloadedActions = UnloadedActions.Difference(AllActions.Intersect(UnloadedActions));
+		if (!ensureMsgf(OrphanedUnloadedActions.Num() == 0, TEXT("Found %d unloaded actions that were not also present in the Action Registry. This should be 0."), UnloadedActions.Num()))
+		{
+			Collector.AddReferencedObjects(OrphanedUnloadedActions);
+		}
 	}
 }
 
@@ -1102,6 +1114,7 @@ void FBlueprintActionDatabase::RefreshAll()
 	}
 
 	ActionRegistry.Empty();
+	UnloadedActionRegistry.Empty();
 	for (TObjectIterator<UClass> ClassIt; ClassIt; ++ClassIt)
 	{
 		UClass* const Class = (*ClassIt);
