Back out revision 43 from //depot/UE4/Engine/Source/Editor/BlueprintGraph/Private/K2Node_BaseAsyncTask.cpp
This causes UE-21565 - issues with duplicate function entries for async tasks used in macros

diff --git a/Engine/Source/Editor/BlueprintGraph/Private/K2Node_BaseAsyncTask.cpp b/Engine/Source/Editor/BlueprintGraph/Private/K2Node_BaseAsyncTask.cpp
--- a/Engine/Source/Editor/BlueprintGraph/Private/K2Node_BaseAsyncTask.cpp
+++ b/Engine/Source/Editor/BlueprintGraph/Private/K2Node_BaseAsyncTask.cpp
@@ -222,8 +222,8 @@ bool UK2Node_BaseAsyncTask::FBaseAsyncTaskHelper::HandleDelegateImplementation(
 	AddDelegateNode->AllocateDefaultPins();
 		bIsErrorFree &= Schema->TryCreateConnection(AddDelegateNode->FindPinChecked(Schema->PN_Self), ProxyObjectPin);
 		bIsErrorFree &= Schema->TryCreateConnection(InOutLastThenPin, AddDelegateNode->FindPinChecked(Schema->PN_Execute));
-		InOutLastThenPin = AddDelegateNode->FindPinChecked(Schema->PN_Then); 
-		CurrentCENode->CustomFunctionName = *FString::Printf(TEXT("%s_%s"), *CurrentProperty->GetName(), *CurrentNode->NodeGuid.ToString());
+		InOutLastThenPin = AddDelegateNode->FindPinChecked(Schema->PN_Then);
+		CurrentCENode->CustomFunctionName = *FString::Printf(TEXT("%s_%s"), *CurrentProperty->GetName(), *CurrentCENode->NodeGuid.ToString());
 		CurrentCENode->AllocateDefaultPins();
 
 		bIsErrorFree &= FBaseAsyncTaskHelper::CreateDelegateForNewFunction(AddDelegateNode->GetDelegatePin(), CurrentCENode->GetFunctionName(), CurrentNode, SourceGraph, CompilerContext);
