When creating a local variable from an orphaned variable node there will no longer be compiler errors until the node is refreshed.

Added a node reconstruction post re-assigning the local variable due to a rogue Target pin that was generated.

diff --git a/Engine/Source/Editor/BlueprintGraph/Private/EdGraphSchema_K2.cpp b/Engine/Source/Editor/BlueprintGraph/Private/EdGraphSchema_K2.cpp
--- a/Engine/Source/Editor/BlueprintGraph/Private/EdGraphSchema_K2.cpp
+++ b/Engine/Source/Editor/BlueprintGraph/Private/EdGraphSchema_K2.cpp
@@ -1643,6 +1643,7 @@ void UEdGraphSchema_K2::OnCreateNonExistentLocalVariable( UK2Node_Variable* Vari
 					if (VariableNode->VariableReference.IsSameReference(OldReference))
 					{
 						VariableNode->VariableReference.SetLocalMember(VarName, Variable->GetGraph()->GetName(), LocalVarGuid);
+						VariableNode->ReconstructNode();
 					}
 				}
 			}
