Mirroring 2922607 - Hiding the new array get by-ref node (bringing back the old array get node) until we can resolve ongoing issues with its use.

#codereveiw Dan.OConnor, Michael.Schoell

diff --git a/Engine/Source/Editor/BlueprintGraph/Private/EdGraphSchema_K2.cpp b/Engine/Source/Editor/BlueprintGraph/Private/EdGraphSchema_K2.cpp
--- a/Engine/Source/Editor/BlueprintGraph/Private/EdGraphSchema_K2.cpp
+++ b/Engine/Source/Editor/BlueprintGraph/Private/EdGraphSchema_K2.cpp
@@ -5423,295 +5423,6 @@ void UEdGraphSchema_K2::BackwardCompatibilityNodeConversion(UEdGraph* Graph, boo
 		{
 			return;
 		}
-
-		const FCustomVersion* CustomVersion = Graph->GetLinker()->GetCustomVersions().GetVersion(FBlueprintsObjectVersion::GUID);
-		if (CustomVersion == nullptr || CustomVersion->Version < FBlueprintsObjectVersion::ArrayGetByRefUpgrade)
-		{
-			// "ForEachLoop" and "ForEachLoopWithBreak" are special case macros that we do not want to add COPY nodes into, we do however want them outside the macro.
-			auto IsGraphForEachLoop = [](UEdGraph* InGraph) -> bool
-			{
-				check(InGraph);
-				return (InGraph->GetName() == TEXT("ForEachLoop") || InGraph->GetName() == TEXT("ForEachLoopWithBreak")) && InGraph->GetOutermost()->GetName() == TEXT("/Engine/EditorBlueprintResources/StandardMacros");
-			};
-
-			UBlueprint* GraphBlueprint = FBlueprintEditorUtils::FindBlueprintForGraph(Graph);
-			check(GraphBlueprint);
-
-			// For collecting upgrade notes and warnings from the upgrade process.
-			FCompilerResultsLog Results;
-			if (!GraphBlueprint->UpgradeNotesLog.IsValid())
-			{
-				GraphBlueprint->UpgradeNotesLog = TSharedPtr<FCompilerResultsLog>(new FCompilerResultsLog);
-				GraphBlueprint->UpgradeNotesLog->bSilentMode = false;
-				GraphBlueprint->UpgradeNotesLog->bAnnotateMentionedNodes = false;
-			}
-
-			// Finding all function nodes that use the "Array_Get" function and replacing them with the new by-ref node.
-			// We also inject Copy and Break nodes as needed to maintain previous functionality (in case by-value returns
-			// were being taken advantage of).
-			TArray<UK2Node_CallFunction*> CallFuncNodes;
-			Graph->GetNodesOfClass<UK2Node_CallFunction>(CallFuncNodes);
-			for (UK2Node_CallFunction* OldCallFuncNode : CallFuncNodes)
-			{
-				if (UFunction* TargetFunction = OldCallFuncNode->GetTargetFunction())
-				{
-					if (TargetFunction == UKismetArrayLibrary::StaticClass()->FindFunctionByName(GET_MEMBER_NAME_CHECKED(UKismetArrayLibrary, Array_Get)))
-					{
-						UK2Node_GetArrayItem* NewArrayGetNode = NewObject<UK2Node_GetArrayItem>(Graph);
-						check(NewArrayGetNode);
-						NewArrayGetNode->SetFlags(RF_Transactional);
-						Graph->AddNode(NewArrayGetNode, false, false);
-						NewArrayGetNode->CreateNewGuid();
-						NewArrayGetNode->PostPlacedNewNode();
-						NewArrayGetNode->AllocateDefaultPins();
-						NewArrayGetNode->NodePosX = OldCallFuncNode->NodePosX;
-						NewArrayGetNode->NodePosY = OldCallFuncNode->NodePosY;
-
-						bool bIsArrayObjectType = false;
-						for (UEdGraphPin* OldPin : OldCallFuncNode->Pins)
-						{
-							UEdGraphPin* TargetPin = nullptr;
-							if (OldPin->PinType.bIsArray)
-							{
-								TargetPin = NewArrayGetNode->GetTargetArrayPin();
-							}
-							else if (OldPin->Direction == EGPD_Input && OldPin->PinType.bIsArray == false && OldPin->PinType.PinCategory == PC_Int)
-							{
-								TargetPin = NewArrayGetNode->GetIndexPin();
-							}
-							else if (OldPin->Direction == EGPD_Output && OldPin->ParentPin == nullptr)
-							{
-								// The target pin is the Get-by-ref node's result pin
-								TargetPin = NewArrayGetNode->GetResultPin();
-
-								bIsArrayObjectType = TargetPin->PinType.PinCategory == PC_Object || TargetPin->PinType.PinCategory == PC_Class || 
-														TargetPin->PinType.PinCategory == PC_Asset || TargetPin->PinType.PinCategory == PC_AssetClass || 
-														TargetPin->PinType.PinCategory == PC_Interface;
-								// If the graph is from the "ForEachLoop" or "ForEachLoopWithBreak" macro, do not add the COPY node
-								if (!IsGraphForEachLoop(Graph) && !bIsArrayObjectType)
-								{
-									UK2Node_Copy* NewCopyNode = NewObject<UK2Node_Copy>(Graph);
-									check(NewCopyNode);
-									NewCopyNode->SetFlags(RF_Transactional);
-									Graph->AddNode(NewCopyNode, false, false);
-									NewCopyNode->CreateNewGuid();
-									NewCopyNode->PostPlacedNewNode();
-									NewCopyNode->AllocateDefaultPins();
-									NewCopyNode->NodePosX = NewArrayGetNode->NodePosX + 175.0f;
-									NewCopyNode->NodePosY = NewArrayGetNode->NodePosY;
-
-									// Add an upgrade note, informing users that the node was auto-injected
-									{
-										FFormatNamedArguments Args;
-										Args.Add(TEXT("CopyNodeName"), NewCopyNode->GetNodeTitle(ENodeTitleType::ListView));
-										Args.Add(TEXT("ArrayGetNodeName"), NewArrayGetNode->GetNodeTitle(ENodeTitleType::ListView));
-										FText UpgradeNoteText = FText::Format(LOCTEXT("ArrayGetByRef_CopyInjection", "This is a \"{CopyNodeName}\" node. It was injected to preserve functionality since the array \"{ArrayGetNodeName}\" node has changed. It outputs a copy of the input element (so modifications don't get applied to the original)."), Args);
-										NewCopyNode->AddNodeUpgradeNote(UpgradeNoteText);
-
-										// Upgrade note for the log
-										GraphBlueprint->UpgradeNotesLog->Note(TEXT("Auto-Upgrade Note: This is a @@ node. It was injected to preserve functionality since the array @@ node has changed. It outputs a copy of the input element (so modifications don't get applied to the original)."), NewCopyNode, NewArrayGetNode);
-									}
-
-									// Connect the Get-by-ref output pin to the copy node's input
-									TargetPin->MakeLinkTo(NewCopyNode->GetInputReferencePin());
-									NewCopyNode->NotifyPinConnectionListChanged(NewCopyNode->GetInputReferencePin());
-
-									if (OldPin->SubPins.Num())
-									{
-										// Make a separate break node!
-										UK2Node_BreakStruct* NewBreakNode = NewObject<UK2Node_BreakStruct>(Graph);
-										check(NewBreakNode);
-										NewBreakNode->SetFlags(RF_Transactional);
-										Graph->AddNode(NewBreakNode, false, false);
-										NewBreakNode->CreateNewGuid();
-										NewBreakNode->PostPlacedNewNode();
-										NewBreakNode->StructType = CastChecked<UScriptStruct>(OldPin->PinType.PinSubCategoryObject.Get());
-										NewBreakNode->AllocateDefaultPins();
-										NewBreakNode->NodePosX = NewCopyNode->NodePosX + 50.0f;
-										NewBreakNode->NodePosY = NewCopyNode->NodePosY;
-
-										// Add an upgrade note, informing users that the node was auto-injected
-										{
-											FFormatNamedArguments Args;
-											Args.Add(TEXT("ArrayGetNodeName"), NewArrayGetNode->GetNodeTitle(ENodeTitleType::ListView));
-											FText UpgradeNoteText = FText::Format(LOCTEXT("ArrayGetByRef_BreakInjection", "Since the array \"{ArrayGetNodeName}\" node has changed, this node was injected to preserve functionality."), Args);
-											NewBreakNode->AddNodeUpgradeNote(UpgradeNoteText);
-
-											// Upgrade note for the log
-											GraphBlueprint->UpgradeNotesLog->Note(TEXT("Auto-Upgrade Note: Since the array @@ node has changed, @@ was injected to preserve functionality."), NewArrayGetNode, NewBreakNode );
-										}
-
-										NewCopyNode->GetCopyResultPin()->MakeLinkTo(NewBreakNode->Pins[0]);
-										int SubPinIdx = 0;
-										for (UEdGraphPin* BreakNodePin : NewBreakNode->Pins)
-										{
-											if (BreakNodePin->Direction == EGPD_Output)
-											{
-												UEdGraphPin* SplitPin = OldPin->SubPins[SubPinIdx++];
-												if (!Graph->GetSchema()->MovePinLinks(*SplitPin, *BreakNodePin).CanSafeConnect())
-												{
-													GraphBlueprint->UpgradeNotesLog->Warning(TEXT("BackwardCompatibilityNodeConversion Error 'cannot safetly move pin @@ to @@' in blueprint: @@"), SplitPin, BreakNodePin, GraphBlueprint);
-												}
-											}
-										}
-									}
-									else
-									{
-										// Move all pins from the OldPin (the result of the Get node) to the output of the Copy Node
-										if (!Graph->GetSchema()->MovePinLinks(*OldPin, *NewCopyNode->GetCopyResultPin()).CanSafeConnect())
-										{
-											GraphBlueprint->UpgradeNotesLog->Warning(TEXT("BackwardCompatibilityNodeConversion Error 'cannot safetly move pin @@ to @@' in blueprint: @@"), OldPin, NewCopyNode, GraphBlueprint);
-										}
-
-										// We have made the connections already
-										TargetPin = nullptr;
-									}
-								}
-							}
-							else if (OldPin->PinName == PN_Self)
-							{
-								// Do not continue, there's just nothing to do
-							}
-							else if (OldPin->ParentPin)
-							{
-								continue;
-							}
-							else
-							{
-								check(0);
-							}
-
-							if (TargetPin && !Graph->GetSchema()->MovePinLinks(*OldPin, *TargetPin).CanSafeConnect())
-							{
-								GraphBlueprint->UpgradeNotesLog->Warning(TEXT("BackwardCompatibilityNodeConversion Error 'cannot safetly move pin @@ to @@' in blueprint: @@"), OldPin, TargetPin, GraphBlueprint);
-							}
-							else if (TargetPin)
-							{
-								NewArrayGetNode->NotifyPinConnectionListChanged(TargetPin);
-							}
-						}
-
-						if (!bIsArrayObjectType)
-						{
-							// Add an upgrade note to the node informing users that the functionality has changed.
-							NewArrayGetNode->AddNodeUpgradeNote(FText::FromString(TEXT("This node has changed! Its output is now a reference - which means modifying the element that this returns is the same as modifying the element stored in the array.")));
-
-							// Upgrade note for the log
-							GraphBlueprint->UpgradeNotesLog->Note(TEXT("Auto-Upgrade Note: @@ has changed! @@ is now a reference - which means modifying the element that this returns is the same as modifying the element stored in the array."), NewArrayGetNode, NewArrayGetNode->GetResultPin());
-						}
-
-						OldCallFuncNode->DestroyNode();
-					}
-				}
-			}
-
-			// We need to look at Macro Instances and upgrade all "ForEachLoop" and "ForEachLoopWithBreak" nodes because those graphs were prevented from having the Copy node injected
-			// because it is desirable for those Macros to return by-ref. However, to keep Blueprints running as they did prior to the upgrade, the instanced Macros need the Copy node.
-			TArray<UK2Node_MacroInstance*> MacroInstanceNodes;
-			Graph->GetNodesOfClass<UK2Node_MacroInstance>(MacroInstanceNodes);
-			for (UK2Node_MacroInstance* MacroInstanceNode : MacroInstanceNodes)
-			{
-				if (IsGraphForEachLoop(MacroInstanceNode->GetMacroGraph()))
-				{
-					UEdGraphPin* ArrayElementPin = MacroInstanceNode->FindPin(TEXT("Array Element"));
-
-					// Because we have a "ForEachLoop" or "ForEachLoopWithBreak" macro instance, we need to add a COPY node at the end, attached to the "Array Element" pin
-					if (ensure(ArrayElementPin) && 
-						ArrayElementPin->PinType.PinCategory != PC_Object && ArrayElementPin->PinType.PinCategory != PC_Class && 
-						ArrayElementPin->PinType.PinCategory != PC_Asset && ArrayElementPin->PinType.PinCategory != PC_AssetClass && 
-						ArrayElementPin->PinType.PinCategory != PC_Interface)
-					{
-						// Add an upgrade note, informing users that the node was auto-injected
-						{
-							FFormatNamedArguments Args;
-							Args.Add(TEXT("ArrayElementPin"), ArrayElementPin->PinFriendlyName);
-							FText UpgradeNoteText = FText::Format(LOCTEXT("ArrayGetByRef_ForEachLoopChanged", "This node has changed! Its output is now a reference - which means modifying the element that this returns in {ArrayElementPin} is the same as modifying the element stored in the array."), Args);
-							MacroInstanceNode->AddNodeUpgradeNote(UpgradeNoteText);
-
-							// Upgrade note for the log
-							GraphBlueprint->UpgradeNotesLog->Note(TEXT("Auto-Upgrade Note: @@ has changed! @@ is now a reference - which means modifying the element that this returns is the same as modifying the element stored in the array."), MacroInstanceNode, ArrayElementPin);
-						}
-
-						UK2Node_Copy* NewCopyNode = NewObject<UK2Node_Copy>(Graph);
-						check(NewCopyNode);
-						NewCopyNode->SetFlags(RF_Transactional);
-						Graph->AddNode(NewCopyNode, false, false);
-						NewCopyNode->CreateNewGuid();
-						NewCopyNode->PostPlacedNewNode();
-						NewCopyNode->AllocateDefaultPins();
-						NewCopyNode->NodePosX = MacroInstanceNode->NodePosX + 200.0f;
-						NewCopyNode->NodePosY = MacroInstanceNode->NodePosY + 50.0f;
-
-						// Add an upgrade note, informing users that the node was auto-injected
-						{
-							FFormatNamedArguments Args;
-							Args.Add(TEXT("CopyNodeName"), NewCopyNode->GetNodeTitle(ENodeTitleType::ListView));
-							Args.Add(TEXT("ForEachLoopName"), MacroInstanceNode->GetNodeTitle(ENodeTitleType::ListView));
-							FText UpgradeNoteText = FText::Format(LOCTEXT("ArrayGetByRef_ForLoopCopyInjection", "This is a \"{CopyNodeName}\" node. It was injected to preserve functionality since the array \"{ArrayGetNodeName}\" node has changed (inside {ForEachLoopName}). It outputs a copy of the input element (so modifications don't get applied to the original)."), Args);
-							NewCopyNode->AddNodeUpgradeNote(UpgradeNoteText);
-
-							// Upgrade note for the log
-							GraphBlueprint->UpgradeNotesLog->Note(TEXT("Auto-Upgrade Note: This is a @@ node. It was injected to preserve functionality since the array \"Get\" node (inside @@) has changed. It outputs a copy of the input element (so modifications don't get applied to the original)."), NewCopyNode, MacroInstanceNode);
-						}
-
-						if (ArrayElementPin->SubPins.Num())
-						{
-							// Make a separate break node!
-							UK2Node_BreakStruct* NewBreakNode = NewObject<UK2Node_BreakStruct>(Graph);
-							check(NewBreakNode);
-							NewBreakNode->SetFlags(RF_Transactional);
-							Graph->AddNode(NewBreakNode, false, false);
-							NewBreakNode->CreateNewGuid();
-							NewBreakNode->PostPlacedNewNode();
-							NewBreakNode->StructType = CastChecked<UScriptStruct>(ArrayElementPin->PinType.PinSubCategoryObject.Get());
-							NewBreakNode->AllocateDefaultPins();
-							NewBreakNode->NodePosX = NewCopyNode->NodePosX + 50.0f;
-							NewBreakNode->NodePosY = NewCopyNode->NodePosY;
-
-							// Add an upgrade note, informing users that the node was auto-injected
-							{
-								FFormatNamedArguments Args;
-								Args.Add(TEXT("ForEachLoopName"), MacroInstanceNode->GetNodeTitle(ENodeTitleType::ListView));
-								FText UpgradeNoteText = FText::Format(LOCTEXT("ArrayGetByRef_ForLoopBreakInjection", "Since the array \"Get\" node has changed (inside {ForEachLoopName}), this node was injected to preserve functionality."), Args);
-								NewBreakNode->AddNodeUpgradeNote(UpgradeNoteText);
-
-								// Upgrade note for the log
-								GraphBlueprint->UpgradeNotesLog->Note(TEXT("Auto-Upgrade Note: Since the array \"Get\" node has changed (inside @@), @@ was injected to preserve functionality. ."), MacroInstanceNode, NewBreakNode);
-							}
-
-							NewCopyNode->GetCopyResultPin()->MakeLinkTo(NewBreakNode->Pins[0]);
-							int SubPinIdx = 0;
-							for (UEdGraphPin* BreakNodePin : NewBreakNode->Pins)
-							{
-								if (BreakNodePin->Direction == EGPD_Output)
-								{
-									UEdGraphPin* SplitPin = ArrayElementPin->SubPins[SubPinIdx++];
-									if (!Graph->GetSchema()->MovePinLinks(*SplitPin, *BreakNodePin).CanSafeConnect())
-									{
-										GraphBlueprint->UpgradeNotesLog->Warning(TEXT("BackwardCompatibilityNodeConversion Error 'cannot safetly move pin @@ to @@' in blueprint: @@"), SplitPin, BreakNodePin, GraphBlueprint);
-									}
-								}
-							}
-
-							// Recombine the array element pin's children
-							RecombinePin(ArrayElementPin->SubPins[0]);
-						}
-						else
-						{
-							// Move all pins from the OldPin (the result of the Get node) to the output of the Copy Node
-							if (!Graph->GetSchema()->MovePinLinks(*ArrayElementPin, *NewCopyNode->GetCopyResultPin()).CanSafeConnect())
-							{
-								GraphBlueprint->UpgradeNotesLog->Warning(TEXT("BackwardCompatibilityNodeConversion Error 'cannot safetly move pin @@ to @@' in blueprint: @@"), ArrayElementPin, NewCopyNode, GraphBlueprint);
-							}
-						}
-
-						// Connect the Get-by-ref output pin to the copy node's input
-						ArrayElementPin->MakeLinkTo(NewCopyNode->GetInputReferencePin());
-						NewCopyNode->NotifyPinConnectionListChanged(NewCopyNode->GetInputReferencePin());
-					}
-				}
-			}
-		} // End FBlueprintsObjectVersion::ArrayGetByRefUpgrade upgrades
 	}
 }
 
diff --git a/Engine/Source/Editor/BlueprintGraph/Private/K2Node_GetArrayItem.cpp b/Engine/Source/Editor/BlueprintGraph/Private/K2Node_GetArrayItem.cpp
--- a/Engine/Source/Editor/BlueprintGraph/Private/K2Node_GetArrayItem.cpp
+++ b/Engine/Source/Editor/BlueprintGraph/Private/K2Node_GetArrayItem.cpp
@@ -118,26 +118,8 @@ class FNodeHandlingFunctor* UK2Node_GetArrayItem::CreateNodeHandler(class FKisme
 
 void UK2Node_GetArrayItem::GetMenuActions(FBlueprintActionDatabaseRegistrar& ActionRegistrar) const
 {
-	// actions get registered under specific object-keys; the idea is that 
-	// actions might have to be updated (or deleted) if their object-key is  
-	// mutated (or removed)... here we use the node's class (so if the node 
-	// type disappears, then the action should go with it)
-	// actions get registered under specific object-keys; the idea is that 
-	// actions might have to be updated (or deleted) if their object-key is  
-	// mutated (or removed)... here we use the node's class (so if the node 
-	// type disappears, then the action should go with it)
-	UClass* ActionKey = GetClass();
-	// to keep from needlessly instantiating a UBlueprintNodeSpawner, first   
-	// check to make sure that the registrar is looking for actions of this type
-	// (could be regenerating actions for a specific asset, and therefore the 
-	// registrar would only accept actions corresponding to that asset)
-	if (ActionRegistrar.IsOpenForRegistration(ActionKey))
-	{
-		UBlueprintNodeSpawner* NodeSpawner = UBlueprintNodeSpawner::Create(GetClass());
-		check(NodeSpawner != nullptr);
-
-		ActionRegistrar.AddBlueprintAction(ActionKey, NodeSpawner);
-	}
+	// temporarily don't register this slick new node with the database, gotta figure out how to provide old
+	// (get by value) behavior alongside new (get by ref) behavior
 }
 
 void UK2Node_GetArrayItem::NotifyPinConnectionListChanged(UEdGraphPin* Pin)
diff --git a/Engine/Source/Runtime/CoreUObject/Private/UObject/ScriptCore.cpp b/Engine/Source/Runtime/CoreUObject/Private/UObject/ScriptCore.cpp
--- a/Engine/Source/Runtime/CoreUObject/Private/UObject/ScriptCore.cpp
+++ b/Engine/Source/Runtime/CoreUObject/Private/UObject/ScriptCore.cpp
@@ -1808,10 +1808,16 @@ void UObject::execArrayGetByRef(FFrame& Stack, RESULT_DECL)
 	}
 	else
 	{
-		// The index is not valid. We will report this as an error (it is fatal in C++ conversions) but otherwise let the Blueprint handle the issue.
-		const int32 PropertySize = ArrayProperty->Inner->ElementSize * ArrayProperty->Inner->ArrayDim;
-		ArrayProperty->Inner->InitializeValue(Stack.MostRecentPropertyAddress);
-		FFrame::KismetExecutionMessage(*FString::Printf(TEXT("Attempted to get an item from array %s out of bounds [%d/%d]!"), *ArrayProperty->GetName(), ArrayIndex, ArrayHelper.Num()? ArrayHelper.Num() - 1 : 0), ELogVerbosity::Warning);
+		FBlueprintExceptionInfo ExceptionInfo(
+			EBlueprintExceptionType::AccessViolation,
+			FText::Format(
+				LOCTEXT("ArrayGetOutofBounds", "Attempted to get an item from array {0} out of bounds [{1}/{2}]!"),
+				FText::FromString(*ArrayProperty->GetName()), 
+				FText::AsNumber(ArrayIndex), 
+				FText::AsNumber(ArrayHelper.Num() ? ArrayHelper.Num() - 1 : 0)
+			)
+		);
+		FBlueprintCoreDelegates::ThrowScriptException(this, Stack, ExceptionInfo);
 	}
 }
 IMPLEMENT_VM_FUNCTION(EX_ArrayGetByRef, execArrayGetByRef);
diff --git a/Engine/Source/Runtime/Engine/Classes/Kismet/KismetArrayLibrary.h b/Engine/Source/Runtime/Engine/Classes/Kismet/KismetArrayLibrary.h
--- a/Engine/Source/Runtime/Engine/Classes/Kismet/KismetArrayLibrary.h
+++ b/Engine/Source/Runtime/Engine/Classes/Kismet/KismetArrayLibrary.h
@@ -118,7 +118,7 @@ class ENGINE_API UKismetArrayLibrary : public UBlueprintFunctionLibrary
 	 *@param	Index			The index in the array to get an item from
 	 *@return	The item stored at the index
 	*/
-	UFUNCTION(BlueprintPure, CustomThunk, meta=(DisplayName = "Get", CompactNodeTitle = "GET", ArrayParm = "TargetArray", ArrayTypeDependentParams = "Item", BlueprintInternalUseOnly = "True"))
+	UFUNCTION(BlueprintPure, CustomThunk, meta=(DisplayName = "Get", CompactNodeTitle = "GET", ArrayParm = "TargetArray", ArrayTypeDependentParams = "Item"), Category="Utilities|Array")
 	static void Array_Get(const TArray<int32>& TargetArray, int32 Index, int32& Item);
 
 	/* 
diff --git a/Engine/Source/Runtime/Engine/Private/Blueprint.cpp b/Engine/Source/Runtime/Engine/Private/Blueprint.cpp
--- a/Engine/Source/Runtime/Engine/Private/Blueprint.cpp
+++ b/Engine/Source/Runtime/Engine/Private/Blueprint.cpp
@@ -615,18 +615,6 @@ void UBlueprint::PostLoad()
 		}
 	}
 
-#if WITH_EDITOR
-	if (UpgradeNotesLog.IsValid() && GetLinker())
-	{
-		const FCustomVersion* CustomVersion = GetLinker()->GetCustomVersions().GetVersion(FBlueprintsObjectVersion::GUID);
-		if (CustomVersion == nullptr || CustomVersion->Version < FBlueprintsObjectVersion::ArrayGetByRefUpgrade)
-		{
-			UpgradeNotesLog->Note(TEXT("Auto-Upgrade Note: Updated all array \"Get\" nodes in Blueprint Class \'@@\' to return a reference. Nodes have been injected to ensure they continue to work as they always have."), this);
-			UpgradeNotesLog->Note(TEXT("Auto-Upgrade Note: To remove these upgrade notes, resave @@. Be sure to verify nothing has broken in the upgrade process."), this);
-		}
-	}
-#endif
-
 	FStructureEditorUtils::RemoveInvalidStructureMemberVariableFromBlueprint(this);
 
 
