Mirroring CL 2942481 from Dev-BP:

Fix for CIS BlueprintDocumentation commandlet crash - guarding "Outer" pin tooltip creation (making sure we have an outer pin first).

diff --git a/Engine/Source/Editor/BlueprintGraph/Private/K2Node_ConstructObjectFromClass.cpp b/Engine/Source/Editor/BlueprintGraph/Private/K2Node_ConstructObjectFromClass.cpp
--- a/Engine/Source/Editor/BlueprintGraph/Private/K2Node_ConstructObjectFromClass.cpp
+++ b/Engine/Source/Editor/BlueprintGraph/Private/K2Node_ConstructObjectFromClass.cpp
@@ -234,7 +234,10 @@ void UK2Node_ConstructObjectFromClass::GetPinHoverText(const UEdGraphPin& Pin, F
 {
 	SetPinToolTip(*GetClassPin(), LOCTEXT("ClassPinDescription", "The object class you want to construct"));
 	SetPinToolTip(*GetResultPin(), LOCTEXT("ResultPinDescription", "The constructed object"));
-	SetPinToolTip(*GetOuterPin(), LOCTEXT("OuterPinDescription", "Owner of the constructed object"));
+	if (UseOuter())
+	{
+		SetPinToolTip(*GetOuterPin(), LOCTEXT("OuterPinDescription", "Owner of the constructed object"));
+	}
 
 	return Super::GetPinHoverText(Pin, HoverTextOut);
 }
