Copying //UE4/Dev-Blueprints to //UE4/Dev-Main (Source: //UE4/Dev-Blueprints @ 3025888)

diff --git a/Engine/Source/Editor/BlueprintGraph/Private/BasicTokenParser.cpp b/Engine/Source/Editor/BlueprintGraph/Private/BasicTokenParser.cpp
--- a/Engine/Source/Editor/BlueprintGraph/Private/BasicTokenParser.cpp
+++ b/Engine/Source/Editor/BlueprintGraph/Private/BasicTokenParser.cpp
@@ -232,6 +232,30 @@
 	PrevComment.Empty( PrevComment.Len() );
 }
 
+// enumerated from logic in K2Node_MathExpression.cpp
+bool IsSymbolStart(TCHAR c)
+{
+	return
+		c == '|' ||
+		c == '&' ||
+		c == '~' ||
+		c == '^' ||
+		c == '!' ||
+		c == '<' ||
+		c == '>' ||
+		c == '=' ||
+		c == '+' ||
+		c == '-' ||
+		c == '*' ||
+		c == '/' ||
+		c == '%' ||
+		c == ':' ||
+		c == '(' ||
+		c == ')' ||
+		c == ',' ||
+		(c >= '0' && c <= '9'); // for the 'no consts' case numeric literals fall into the symbol branch
+}
+
 //------------------------------------------------------------------------------
 bool FBasicTokenParser::GetToken(FBasicToken& Token, bool bNoConsts/* = false*/)
 {
@@ -277,53 +301,6 @@
 		Token.SetGuid(Token.Identifier);
 		return IsValid();
 	}
-	else if( FText::IsLetter(c)|| (c=='_') )
-	{
-		// Alphanumeric token.
-		int32 Length=0;
-		do
-		{
-			Token.Identifier[Length++] = c;
-			if( Length >= NAME_SIZE )
-			{
-				Length = ((int32)NAME_SIZE) - 1;
-				Token.Identifier[Length]=0; // need this for the error description
-
-				FText ErrorDesc = FText::Format(LOCTEXT("IdTooLong", "Identifer ({0}...) exceeds maximum length of {1}"), FText::FromString(Token.Identifier), FText::AsNumber((int32)NAME_SIZE));
-				SetError(FErrorState::ParseError, ErrorDesc);
-
-				break;
-			}
-			c = GetChar();
-		} while( FText::IsLetter(c) || ((c>='0')&&(c<='9')) || (c=='_') );
-		UngetChar();
-		Token.Identifier[Length]=0;
-
-		// Assume this is an identifier unless we find otherwise.
-		Token.TokenType = FBasicToken::TOKEN_Identifier;
-
-		// Lookup the token's global name.
-		Token.TokenName = FName(Token.Identifier, FNAME_Find);
-
-		// If const values are allowed, determine whether the identifier represents a constant
-		if ( !bNoConsts )
-		{
-			// See if the identifier is part of a vector, rotation or other struct constant.
-			// boolean true/false
-			if( Token.Matches(TEXT("true")) )
-			{
-				Token.SetConstBool(true);
-				return true;
-			}
-			else if( Token.Matches(TEXT("false")) )
-			{
-				Token.SetConstBool(false);
-				return true;
-			}
-		}
-
-		return IsValid();
-	}
 	// if const values are allowed, determine whether the non-identifier token represents a const
 	else if ( !bNoConsts && ((c>='0' && c<='9') || ((c=='+' || c=='-') && (p>='0' && p<='9'))) )
 	{
@@ -425,7 +402,7 @@
 		Token.SetConstString(Temp);
 		return IsValid();
 	}
-	else
+	else if (IsSymbolStart(c))
 	{
 		// Symbol.
 		int32 Length=0;
@@ -475,6 +452,53 @@
 
 		return true;
 	}
+	else
+	{
+		// Alphanumeric token.
+		int32 Length = 0;
+		do
+		{
+			Token.Identifier[Length++] = c;
+			if (Length >= NAME_SIZE)
+			{
+				Length = ((int32)NAME_SIZE) - 1;
+				Token.Identifier[Length] = 0; // need this for the error description
+
+				FText ErrorDesc = FText::Format(LOCTEXT("IdTooLong", "Identifer ({0}...) exceeds maximum length of {1}"), FText::FromString(Token.Identifier), FText::AsNumber((int32)NAME_SIZE));
+				SetError(FErrorState::ParseError, ErrorDesc);
+
+				break;
+			}
+			c = GetChar();
+		} while (!IsSymbolStart(c) && c != '{' && c != '}' && c != '"' && !FText::IsWhitespace(c));
+		UngetChar();
+		Token.Identifier[Length] = 0;
+
+		// Assume this is an identifier unless we find otherwise.
+		Token.TokenType = FBasicToken::TOKEN_Identifier;
+
+		// Lookup the token's global name.
+		Token.TokenName = FName(Token.Identifier, FNAME_Find);
+
+		// If const values are allowed, determine whether the identifier represents a constant
+		if (!bNoConsts)
+		{
+			// See if the identifier is part of a vector, rotation or other struct constant.
+			// boolean true/false
+			if (Token.Matches(TEXT("true")))
+			{
+				Token.SetConstBool(true);
+				return true;
+			}
+			else if (Token.Matches(TEXT("false")))
+			{
+				Token.SetConstBool(false);
+				return true;
+			}
+		}
+
+		return IsValid();
+	}
 }
 
 //------------------------------------------------------------------------------
@@ -751,7 +775,7 @@
 			{
 				MultipleNewlines = true;
 			}
-		} while (IsWhitespace(c));
+		} while (FText::IsWhitespace(c));
 
 		if (c != TEXT('/') || PeekChar() != TEXT('/'))
 		{
@@ -807,11 +831,5 @@
 }
 
 //------------------------------------------------------------------------------
-bool FBasicTokenParser::IsWhitespace(TCHAR c)
-{
-	return c==TEXT(' ') || c==TEXT('\t') || c==TEXT('\r') || c==TEXT('\n');
-}
-
-//------------------------------------------------------------------------------
 bool FBasicTokenParser::MatchIdentifier(FName Match)
 {
diff --git a/Engine/Source/Editor/BlueprintGraph/Private/BasicTokenParser.h b/Engine/Source/Editor/BlueprintGraph/Private/BasicTokenParser.h
--- a/Engine/Source/Editor/BlueprintGraph/Private/BasicTokenParser.h
+++ b/Engine/Source/Editor/BlueprintGraph/Private/BasicTokenParser.h
@@ -325,14 +325,6 @@ protected:
 	 */
 	static bool IsEOL(TCHAR c);
 
-	/**
-	 * Tests if a character is a whitespace character.
-	 * 
-	 * @param  c	The character to test.
-	 * @return True if c is an whitespace character, false otherwise.
-	 */
-	static bool IsWhitespace(TCHAR c);
-
 	//--------------------------------------
 	// Error state 
 	//--------------------------------------
diff --git a/Engine/Source/Editor/BlueprintGraph/Private/BlueprintActionFilter.cpp b/Engine/Source/Editor/BlueprintGraph/Private/BlueprintActionFilter.cpp
--- a/Engine/Source/Editor/BlueprintGraph/Private/BlueprintActionFilter.cpp
+++ b/Engine/Source/Editor/BlueprintGraph/Private/BlueprintActionFilter.cpp
@@ -405,7 +405,20 @@ UClass* BlueprintActionFilterImpl::GetAuthoritativeBlueprintClass(UBlueprint con
 	UClass* AuthoritativeClass = BpClass;
 	if (BpClass != nullptr)
 	{
-		AuthoritativeClass = BpClass->GetAuthoritativeClass();
+		if (ensureMsgf(!BpClass->GetClass()->IsChildOf<UBlueprintGeneratedClass>() || BpClass->ClassGeneratedBy
+			, TEXT("Ambiguous Blueprint: '%s'; with skeleton: %s (%s), class: %s (%s), and parent: %s (%s) - authoratative class: %s (%s)")
+			, *Blueprint->GetPathName()
+			, Blueprint->SkeletonGeneratedClass ? *Blueprint->SkeletonGeneratedClass->GetName() : TEXT("[NULL]")
+			, Blueprint->SkeletonGeneratedClass ? *Blueprint->SkeletonGeneratedClass->GetClass()->GetName() : TEXT("N/A")
+			, Blueprint->GeneratedClass ? *Blueprint->GeneratedClass->GetName() : TEXT("[NULL]")
+			, Blueprint->GeneratedClass ? *Blueprint->GeneratedClass->GetClass()->GetName() : TEXT("N/A")
+			, Blueprint->ParentClass ? *Blueprint->ParentClass->GetName() : TEXT("[NULL]")
+			, Blueprint->ParentClass ? *Blueprint->ParentClass->GetClass()->GetName() : TEXT("N/A")
+			, BpClass ? *BpClass->GetName() : TEXT("[NULL]")
+			, BpClass ? *BpClass->GetClass()->GetName() : TEXT("N/A")) )
+		{
+			AuthoritativeClass = BpClass->GetAuthoritativeClass();
+		}		
 	}
 	return AuthoritativeClass;
 }
@@ -580,7 +593,15 @@ static bool BlueprintActionFilterImpl::IsEventUnimplementable(FBlueprintActionFi
 			for (UBlueprint const* Blueprint : FilterContext.Blueprints)
 			{
 				UClass const* BpClass = GetAuthoritativeBlueprintClass(Blueprint);
-				check(BpClass != nullptr);
+				if (!ensureMsgf(BpClass != nullptr
+					, TEXT("Unable to resolve IsEventUnimplementable() - Blueprint (%s) missing an authoratative class (skel: %s, generated: %s, parent: %s)")
+					, *Blueprint->GetName()
+					, Blueprint->SkeletonGeneratedClass ? *Blueprint->SkeletonGeneratedClass->GetName() : TEXT("[NULL]")
+					, Blueprint->GeneratedClass ? *Blueprint->GeneratedClass->GetName() : TEXT("[NULL]")
+					, Blueprint->ParentClass ? *Blueprint->ParentClass->GetName() : TEXT("[NULL]")) )
+				{
+					continue;
+				}
 
 				// if this function belongs directly to this blueprint, then it is
 				// already implemented here (this action however is valid for sub-
@@ -671,7 +692,15 @@ static bool BlueprintActionFilterImpl::IsFieldInaccessible(FBlueprintActionFilte
 			for (UBlueprint const* Blueprint : FilterContext.Blueprints)
 			{
 				UClass const* BpClass = GetAuthoritativeBlueprintClass(Blueprint);
-				check(BpClass != nullptr);
+				if (!ensureMsgf(BpClass != nullptr
+					, TEXT("Unable to resolve IsFieldInaccessible() - Blueprint (%s) missing an authoratative class (skel: %s, generated: %s, parent: %s)")
+					, *Blueprint->GetName()
+					, Blueprint->SkeletonGeneratedClass ? *Blueprint->SkeletonGeneratedClass->GetName() : TEXT("[NULL]")
+					, Blueprint->GeneratedClass ? *Blueprint->GeneratedClass->GetName() : TEXT("[NULL]")
+					, Blueprint->ParentClass ? *Blueprint->ParentClass->GetName() : TEXT("[NULL]")))
+				{
+					continue;
+				}
 			
 				// private functions are only accessible from the class they belong to
 				if (bIsPrivate && !IsClassOfType(BpClass, ActionOwner, /*bNeedsExactMatch =*/true))
diff --git a/Engine/Source/Editor/BlueprintGraph/Private/CallFunctionHandler.cpp b/Engine/Source/Editor/BlueprintGraph/Private/CallFunctionHandler.cpp
--- a/Engine/Source/Editor/BlueprintGraph/Private/CallFunctionHandler.cpp
+++ b/Engine/Source/Editor/BlueprintGraph/Private/CallFunctionHandler.cpp
@@ -334,7 +334,7 @@ void FKCHandler_CallFunction::CreateFunctionCallStatement(FKismetFunctionContext
 					bool bIsSelfTerm = true;
 					if(Target != nullptr)
 					{
-						const UEdGraphPin* SourcePin = Cast<UEdGraphPin>(Target->Source);
+						const UEdGraphPin* SourcePin = Target->SourcePin;
 						bIsSelfTerm = (SourcePin == nullptr || CompilerContext.GetSchema()->IsSelfPin(*SourcePin));
 					}
 
@@ -687,7 +687,7 @@ void FKCHandler_CallFunction::Transform(FKismetFunctionContext& Context, UEdGrap
 
 			if ((NewInPin != NULL) && (NewOutPin != NULL))
 			{
-				CompilerContext.MessageLog.NotifyIntermediateObjectCreation(NewOutPin, OldOutPin);
+				CompilerContext.MessageLog.NotifyIntermediatePinCreation(NewOutPin, OldOutPin);
 
 				while (OldOutPin->LinkedTo.Num() > 0)
 				{
diff --git a/Engine/Source/Editor/BlueprintGraph/Private/EdGraphSchema_K2.cpp b/Engine/Source/Editor/BlueprintGraph/Private/EdGraphSchema_K2.cpp
--- a/Engine/Source/Editor/BlueprintGraph/Private/EdGraphSchema_K2.cpp
+++ b/Engine/Source/Editor/BlueprintGraph/Private/EdGraphSchema_K2.cpp
@@ -930,18 +930,25 @@ bool UEdGraphSchema_K2::FunctionHasParamOfType(const UFunction* InFunction, UEdG
 		{
 			// See if this is the direction we want (input or output)
 			const bool bIsFunctionInput = !FuncParam->HasAnyPropertyFlags(CPF_OutParm) || FuncParam->HasAnyPropertyFlags(CPF_ReferenceParm);
-			if ((!bIsFunctionInput && bWantOutput) || (bIsFunctionInput && !bWantOutput))
+			if (bIsFunctionInput != bWantOutput)
 			{
 				// See if this pin has compatible types
 				FEdGraphPinType ParamPinType;
 				bool bConverted = ConvertPropertyToPinType(FuncParam, ParamPinType);
 				if (bConverted)
 				{
-					if (bIsFunctionInput && ArePinTypesCompatible(DesiredPinType, ParamPinType))
+					UClass* Context = nullptr;
+					UBlueprint* Blueprint = Cast<UBlueprint>(InGraph->GetOuter());
+					if (Blueprint)
+					{
+						Context = Blueprint->GeneratedClass;
+					}
+
+					if (bIsFunctionInput && ArePinTypesCompatible(DesiredPinType, ParamPinType, Context))
 					{
 						return true;
 					}
-					else if (!bIsFunctionInput && ArePinTypesCompatible(ParamPinType, DesiredPinType))
+					else if (!bIsFunctionInput && ArePinTypesCompatible(ParamPinType, DesiredPinType, Context))
 					{
 						return true;
 					}
@@ -1892,7 +1899,7 @@ void UEdGraphSchema_K2::GetJumpToConnectionSubMenuActions( class FMenuBuilder& M
 	TMap< FString, uint32 > LinkTitleCount;
 
 	// Add all the links we could break from
-	for(auto PinLink : InGraphPin->LinkedTo )
+	for(const UEdGraphPin* PinLink : InGraphPin->LinkedTo )
 	{
 		FText Title = PinLink->GetOwningNode()->GetNodeTitle(ENodeTitleType::ListView);
 		FString TitleString = Title.ToString();
@@ -1925,7 +1932,7 @@ void UEdGraphSchema_K2::GetJumpToConnectionSubMenuActions( class FMenuBuilder& M
 		++Count;
 
 		MenuBuilder.AddMenuEntry( Description, Description, FSlateIcon(), FUIAction(
-		FExecuteAction::CreateStatic(&FKismetEditorUtilities::BringKismetToFocusAttentionOnObject, Cast<const UObject>(PinLink), false)));
+		FExecuteAction::CreateStatic(&FKismetEditorUtilities::BringKismetToFocusAttentionOnPin, PinLink)));
 	}
 }
 
@@ -2192,7 +2199,7 @@ bool UEdGraphSchema_K2::TryCreateConnection(UEdGraphPin* PinA, UEdGraphPin* PinB
 
 	bool bModified = UEdGraphSchema::TryCreateConnection(PinA, PinB);
 
-	if (bModified && !PinA->HasAnyFlags(RF_Transient))
+	if (bModified && !PinA->IsPendingKill())
 	{
 		FBlueprintEditorUtils::MarkBlueprintAsModified(Blueprint);
 	}
@@ -2579,7 +2586,18 @@ void UEdGraphSchema_K2::AutowireConversionNode(UEdGraphPin* InputPin, UEdGraphPi
 	{
 		UEdGraphPin* TestPin = ConversionNode->Pins[PinIndex];
 
-		if ((TestPin->Direction == EGPD_Input) && (ArePinTypesCompatible(OutputPin->PinType, TestPin->PinType)))
+		UClass* Context = nullptr;
+		UK2Node* K2Node = Cast<UK2Node>(OutputPin->GetOwningNode());
+		if (K2Node != nullptr)
+		{
+			UBlueprint* Blueprint = K2Node->GetBlueprint();
+			if (Blueprint)
+			{
+				Context = Blueprint->GeneratedClass;
+			}
+		}
+
+		if ((TestPin->Direction == EGPD_Input) && (ArePinTypesCompatible(OutputPin->PinType, TestPin->PinType, Context)))
 		{
 			if(bAllowOutputConnections && TryCreateConnection(TestPin, OutputPin))
 			{
@@ -2587,7 +2605,7 @@ void UEdGraphSchema_K2::AutowireConversionNode(UEdGraphPin* InputPin, UEdGraphPi
 				bAllowOutputConnections = false;
 			}
 		}
-		else if ((TestPin->Direction == EGPD_Output) && (ArePinTypesCompatible(TestPin->PinType, InputPin->PinType)))
+		else if ((TestPin->Direction == EGPD_Output) && (ArePinTypesCompatible(TestPin->PinType, InputPin->PinType, Context)))
 		{
 			if(bAllowInputConnections && TryCreateConnection(TestPin, InputPin))
 			{
@@ -2834,7 +2852,7 @@ FText UEdGraphSchema_K2::GetPinDisplayName(const UEdGraphPin* Pin) const
 
 void UEdGraphSchema_K2::ConstructBasicPinTooltip(const UEdGraphPin& Pin, const FText& PinDescription, FString& TooltipOut) const
 {
-	if (Pin.HasAnyFlags(RF_Transient))
+	if (Pin.bWasTrashed)
 	{
 		return;
 	}
@@ -4080,24 +4098,6 @@ void UEdGraphSchema_K2::BreakSinglePinLink(UEdGraphPin* SourcePin, UEdGraphPin*
 
 void UEdGraphSchema_K2::ReconstructNode(UEdGraphNode& TargetNode, bool bIsBatchRequest/*=false*/) const
 {
-	{
-		TArray<UObject *> NodeChildren;
-		GetObjectsWithOuter(&TargetNode, NodeChildren, false);
-		for (int32 Iter = 0; Iter < NodeChildren.Num(); ++Iter)
-		{
-			UEdGraphPin* Pin = Cast<UEdGraphPin>(NodeChildren[Iter]);
-			const bool bIsValidPin = !Pin 
-				|| (Pin->IsPendingKill() && !Pin->LinkedTo.Num()) 
-				|| TargetNode.Pins.Contains(Pin);
-			if (!bIsValidPin)
-			{
-				UE_LOG(LogBlueprint, Warning,
-					TEXT("Invalid pin: '%s' thinks it belongs to a node ('%s' - '%s') which doesn't have record of it. Try refreshing the node, then compile and resave the Blueprint to hopefully aleviate the problem."),
-					*Pin->PinName, *TargetNode.GetNodeTitle(ENodeTitleType::MenuTitle).ToString(), *TargetNode.GetFullName());
-			}
-		}
-	}
-
 	Super::ReconstructNode(TargetNode, bIsBatchRequest);
 
 	// If the reconstruction is being handled by something doing a batch (i.e. the blueprint autoregenerating itself), defer marking the blueprint as modified to prevent multiple recompiles
@@ -5622,7 +5622,7 @@ UEdGraphNode* UEdGraphSchema_K2::CreateSubstituteNode(UEdGraphNode* Node, const
 				check(Pin);
 
 				// Reparent the pin to the new custom event node
-				Pin->Rename(*Pin->GetName(), CustomEventNode, RenameFlags | (Pin->HasAnyFlags(RF_NeedLoad | RF_NeedPostLoad | RF_NeedPostLoadSubobjects) ? REN_ForceNoResetLoaders : RF_NoFlags));
+				Pin->SetOwningNode(CustomEventNode);
 
 				// Don't include execution or delegate output pins as user-defined pins
 				if(!bOriginalWasCustomEvent && !IsExecPin(*Pin) && !IsDelegateCategory(Pin->PinType.PinCategory))
@@ -6181,7 +6181,7 @@ void UEdGraphSchema_K2::RecombinePin(UEdGraphPin* Pin) const
 		}
 
 		GraphNode->Pins.Remove(SubPin);
-		Blueprint->PinWatches.Remove(SubPin);
+		Blueprint->WatchedPins.Remove(SubPin);
 	}
 
 	if (Pin->Direction == EGPD_Input)
@@ -6216,7 +6216,18 @@ void UEdGraphSchema_K2::RecombinePin(UEdGraphPin* Pin) const
 		}
 	}
 
-	ParentPin->SubPins.Empty();	
+	// Clear out subpins:
+	TArray<UEdGraphPin*>& ParentSubPins = ParentPin->SubPins;
+	while (ParentSubPins.Num())
+	{
+		// To ensure that MarkPendingKill does not mutate ParentSubPins, we null out the ParentPin
+		// if we assume that MarkPendingKill *will* mutate ParentSubPins we could introduce an infinite
+		// loop. No known case of this being possible, but it would be trivial to write bad node logic
+		// that introduces this problem:
+		ParentSubPins.Last()->ParentPin = nullptr; 
+		ParentSubPins.Last()->MarkPendingKill();
+		ParentSubPins.RemoveAt(ParentSubPins.Num()-1);
+	}
 
 	Graph->NotifyGraphChanged();
 	FBlueprintEditorUtils::MarkBlueprintAsModified(Blueprint);
diff --git a/Engine/Source/Editor/BlueprintGraph/Private/K2Node.cpp b/Engine/Source/Editor/BlueprintGraph/Private/K2Node.cpp
--- a/Engine/Source/Editor/BlueprintGraph/Private/K2Node.cpp
+++ b/Engine/Source/Editor/BlueprintGraph/Private/K2Node.cpp
@@ -2,8 +2,9 @@
 
 
 #include "BlueprintGraphPrivatePCH.h"
-#include "Engine/Breakpoint.h"
 
+#include "EdGraph/EdGraphPin.h"
+#include "Engine/Breakpoint.h"
 #include "K2Node.h"
 #include "KismetDebugUtilities.h" // for HasDebuggingData(), GetWatchText()
 #include "KismetCompiler.h"
@@ -644,9 +645,9 @@ void UK2Node::ReconstructSinglePin(UEdGraphPin* NewPin, UEdGraphPin* OldPin, ERe
 	}
 
 	// Update the blueprints watched pins as the old pin will be going the way of the dodo
-	for (int32 WatchIndex = 0; WatchIndex < Blueprint->PinWatches.Num(); ++WatchIndex)
+	for (int32 WatchIndex = 0; WatchIndex < Blueprint->WatchedPins.Num(); ++WatchIndex)
 	{
-		UEdGraphPin*& WatchedPin = Blueprint->PinWatches[WatchIndex];
+		UEdGraphPin* WatchedPin = Blueprint->WatchedPins[WatchIndex].Get();
 		if( WatchedPin == OldPin )
 		{
 			WatchedPin = NewPin;
@@ -654,7 +655,7 @@ void UK2Node::ReconstructSinglePin(UEdGraphPin* NewPin, UEdGraphPin* OldPin, ERe
 		}
 	}
 
-	OldPin->Rename(NULL, GetTransientPackage(), (REN_DontCreateRedirectors|(Blueprint->bIsRegeneratingOnLoad ? REN_ForceNoResetLoaders : REN_None)));
+	OldPin->InvalidateAndTrash();;
 }
 
 void UK2Node::RewireOldPinsToNewPins(TArray<UEdGraphPin*>& InOldPins, TArray<UEdGraphPin*>& InNewPins)
@@ -702,15 +703,7 @@ void UK2Node::DestroyPinList(TArray<UEdGraphPin*>& InPins)
 		Pin->Modify();
 		Pin->BreakAllPinLinks();
 
-		// just in case this pin was set to watch (don't want to save PinWatches with dead pins)
-		Blueprint->PinWatches.Remove(Pin);
-#if 0
-		UEdGraphNode::ReturnPinToPool(Pin);
-#else
-		Pin->Rename(NULL, GetTransientPackage(), (Blueprint->bIsRegeneratingOnLoad ? REN_ForceNoResetLoaders : REN_None));
-		Pin->RemoveFromRoot();
-		Pin->MarkPendingKill();
-#endif
+		UEdGraphNode::DestroyPin(Pin);
 	}
 }
 
@@ -735,6 +728,12 @@ void UK2Node::ExpandSplitPin(FKismetCompilerContext* CompilerContext, UEdGraph*
 		{
 			if (ExpandedPin->Direction == Pin->Direction)
 			{
+				if (Pin->SubPins.Num() == SubPinIndex)
+				{
+					CompilerContext->MessageLog.Error(*LOCTEXT("PinExpansionError", "Failed to expand pin @@, likely due to bad logic in node @@").ToString(), Pin, Pin->GetOwningNode());
+					break;
+				}
+
 				UEdGraphPin* SubPin = Pin->SubPins[SubPinIndex++];
 				if (CompilerContext)
 				{
@@ -745,6 +744,8 @@ void UK2Node::ExpandSplitPin(FKismetCompilerContext* CompilerContext, UEdGraph*
 					Schema->MovePinLinks(*SubPin, *ExpandedPin);
 				}
 				Pins.Remove(SubPin);
+				SubPin->ParentPin = nullptr;
+				SubPin->MarkPendingKill();
 			}
 			else
 			{
diff --git a/Engine/Source/Editor/BlueprintGraph/Private/K2Node_AddComponent.cpp b/Engine/Source/Editor/BlueprintGraph/Private/K2Node_AddComponent.cpp
--- a/Engine/Source/Editor/BlueprintGraph/Private/K2Node_AddComponent.cpp
+++ b/Engine/Source/Editor/BlueprintGraph/Private/K2Node_AddComponent.cpp
@@ -20,6 +20,19 @@
 	bIsPureFunc = false;
 }
 
+void UK2Node_AddComponent::PostDuplicate(bool bDuplicateForPIE)
+{
+	Super::PostDuplicate(bDuplicateForPIE);
+
+	// Determine whether or not this node belongs to a Blueprint context.
+	UBlueprint* Blueprint = FBlueprintEditorUtils::FindBlueprintForNode(this);
+	if (!bDuplicateForPIE && Blueprint != nullptr && !Blueprint->bBeingCompiled)
+	{
+		// Create a unique component template for the duplicated node (this).
+		MakeNewComponentTemplate();
+	}
+}
+
 void UK2Node_AddComponent::AllocateDefaultPins()
 {
 	AllocateDefaultPinsWithoutExposedVariables();
@@ -235,65 +248,9 @@
 void UK2Node_AddComponent::PostPasteNode()
 {
 	Super::PostPasteNode();
-
-	// There is a template associated with this node that should be unique, but after a node is pasted, it either points to a
-	// template shared by the copied node, or to nothing (when pasting into a different blueprint)
-	const UEdGraphSchema_K2* K2Schema = GetDefault<UEdGraphSchema_K2>();
-	UBlueprint* Blueprint = GetBlueprint();
-
-	// Find the template name and return type pins
-	UEdGraphPin* TemplateNamePin = GetTemplateNamePin();
-	UEdGraphPin* ReturnPin = GetReturnValuePin();
-	if ((TemplateNamePin != NULL) && (ReturnPin != NULL))
-	{
-		// Find the current template if it exists
-		FString TemplateName = TemplateNamePin->DefaultValue;
-		UActorComponent* SourceTemplate = Blueprint->FindTemplateByName(FName(*TemplateName));
-
-		// Determine the type of the component needed
-		UClass* ComponentClass = Cast<UClass>(ReturnPin->PinType.PinSubCategoryObject.Get());
-			
-		if (ComponentClass)
-		{
-			ensure(NULL != Cast<UBlueprintGeneratedClass>(Blueprint->GeneratedClass));
-			// Create a new template object and update the template pin to point to it
-			UActorComponent* NewTemplate = NewObject<UActorComponent>(Blueprint->GeneratedClass, ComponentClass, NAME_None, RF_ArchetypeObject|RF_Public);
-			Blueprint->ComponentTemplates.Add(NewTemplate);
-
-			TemplateNamePin->DefaultValue = NewTemplate->GetName();
-
-			// Copy the old template data over to the new template if it's compatible
-			if ((SourceTemplate != NULL) && (SourceTemplate->GetClass()->IsChildOf(ComponentClass)))
-			{
-				TArray<uint8> SavedProperties;
-				FObjectWriter Writer(SourceTemplate, SavedProperties);
-				FObjectReader(NewTemplate, SavedProperties);
-			}
-			else if(TemplateBlueprint.Len() > 0)
-			{
-				// try to find/load our blueprint to copy the template
-				UBlueprint* SourceBlueprint = FindObject<UBlueprint>(NULL, *TemplateBlueprint);
-				if(SourceBlueprint != NULL)
-				{
-					SourceTemplate = SourceBlueprint->FindTemplateByName(FName(*TemplateName));
-					if ((SourceTemplate != NULL) && (SourceTemplate->GetClass()->IsChildOf(ComponentClass)))
-					{
-						TArray<uint8> SavedProperties;
-						FObjectWriter Writer(SourceTemplate, SavedProperties);
-						FObjectReader(NewTemplate, SavedProperties);
-					}
-				}
-
-				TemplateBlueprint.Empty();
-			}
-		}
-		else
-		{
-			// Clear the template connection; can't resolve the type of the component to create
-			ensure(false);
-			TemplateNamePin->DefaultValue = TEXT("");
-		}
-	}
+	
+	// Create a unique component template for the pasted node (this).
+	MakeNewComponentTemplate();
 }
 
 FText UK2Node_AddComponent::GetNodeTitle(ENodeTitleType::Type TitleType) const
@@ -443,4 +400,66 @@
 	}
 }
 
+void UK2Node_AddComponent::MakeNewComponentTemplate()
+{
+	// There is a template associated with this node that should be unique, but after a node is duplicated, it either points to a
+	// template associated with the original node, or to nothing (when pasting into a different blueprint)
+	const UEdGraphSchema_K2* K2Schema = GetDefault<UEdGraphSchema_K2>();
+	UBlueprint* Blueprint = GetBlueprint();
+
+	// Find the template name and return type pins
+	UEdGraphPin* TemplateNamePin = GetTemplateNamePin();
+	UEdGraphPin* ReturnPin = GetReturnValuePin();
+	if ((TemplateNamePin != NULL) && (ReturnPin != NULL))
+	{
+		// Find the current template if it exists
+		FString TemplateName = TemplateNamePin->DefaultValue;
+		UActorComponent* SourceTemplate = Blueprint->FindTemplateByName(FName(*TemplateName));
+
+		// Determine the type of the component needed
+		UClass* ComponentClass = Cast<UClass>(ReturnPin->PinType.PinSubCategoryObject.Get());
+
+		if (ComponentClass)
+		{
+			ensure(NULL != Cast<UBlueprintGeneratedClass>(Blueprint->GeneratedClass));
+			// Create a new template object and update the template pin to point to it
+			UActorComponent* NewTemplate = NewObject<UActorComponent>(Blueprint->GeneratedClass, ComponentClass, NAME_None, RF_ArchetypeObject | RF_Public);
+			Blueprint->ComponentTemplates.Add(NewTemplate);
+
+			TemplateNamePin->DefaultValue = NewTemplate->GetName();
+
+			// Copy the old template data over to the new template if it's compatible
+			if ((SourceTemplate != NULL) && (SourceTemplate->GetClass()->IsChildOf(ComponentClass)))
+			{
+				TArray<uint8> SavedProperties;
+				FObjectWriter Writer(SourceTemplate, SavedProperties);
+				FObjectReader(NewTemplate, SavedProperties);
+			}
+			else if (TemplateBlueprint.Len() > 0)
+			{
+				// try to find/load our blueprint to copy the template
+				UBlueprint* SourceBlueprint = FindObject<UBlueprint>(NULL, *TemplateBlueprint);
+				if (SourceBlueprint != NULL)
+				{
+					SourceTemplate = SourceBlueprint->FindTemplateByName(FName(*TemplateName));
+					if ((SourceTemplate != NULL) && (SourceTemplate->GetClass()->IsChildOf(ComponentClass)))
+					{
+						TArray<uint8> SavedProperties;
+						FObjectWriter Writer(SourceTemplate, SavedProperties);
+						FObjectReader(NewTemplate, SavedProperties);
+					}
+				}
+
+				TemplateBlueprint.Empty();
+			}
+		}
+		else
+		{
+			// Clear the template connection; can't resolve the type of the component to create
+			ensure(false);
+			TemplateNamePin->DefaultValue = TEXT("");
+		}
+	}
+}
+
 #undef LOCTEXT_NAMESPACE
diff --git a/Engine/Source/Editor/BlueprintGraph/Private/K2Node_CallFunction.cpp b/Engine/Source/Editor/BlueprintGraph/Private/K2Node_CallFunction.cpp
--- a/Engine/Source/Editor/BlueprintGraph/Private/K2Node_CallFunction.cpp
+++ b/Engine/Source/Editor/BlueprintGraph/Private/K2Node_CallFunction.cpp
@@ -1236,7 +1236,7 @@ FText UK2Node_CallFunction::GetTooltipText() const
 
 void UK2Node_CallFunction::GeneratePinTooltipFromFunction(UEdGraphPin& Pin, const UFunction* Function)
 {
-	if (Pin.HasAnyFlags(RF_Transient))
+	if (Pin.bWasTrashed)
 	{
 		return;
 	}
@@ -1971,7 +1971,7 @@ void UK2Node_CallFunction::ExpandNode(class FKismetCompilerContext& CompilerCont
 							AssignNode->AllocateDefaultPins();
 
 							// Move connections from fake 'enum exec' pint to this assignment node
-								CompilerContext.MovePinLinksToIntermediate(*Pin, *AssignNode->GetExecPin());
+							CompilerContext.MovePinLinksToIntermediate(*Pin, *AssignNode->GetExecPin());
 
 							// Connect this to out temp enum var
 							Schema->TryCreateConnection(AssignNode->GetVariablePin(), TempEnumVarOutput);
@@ -1983,6 +1983,7 @@ void UK2Node_CallFunction::ExpandNode(class FKismetCompilerContext& CompilerCont
 							AssignNode->GetValuePin()->DefaultValue = Pin->PinName;
 
 							// Finally remove this 'cosmetic' exec pin
+							Pins[PinIdx]->MarkPendingKill();
 							Pins.RemoveAt(PinIdx);
 						}
 					}
@@ -2013,10 +2014,11 @@ void UK2Node_CallFunction::ExpandNode(class FKismetCompilerContext& CompilerCont
 							Pin->Direction == EGPD_Output &&
 							Pin->PinType.PinCategory == Schema->PC_Exec )
 						{
-							// Move connections from fake 'enum exec' pint to this switch node
+							// Move connections from fake 'enum exec' pin to this switch node
 							CompilerContext.MovePinLinksToIntermediate(*Pin, *SwitchEnumNode->FindPinChecked(Pin->PinName));
 								
 							// Finally remove this 'cosmetic' exec pin
+							Pins[PinIdx]->MarkPendingKill();
 							Pins.RemoveAt(PinIdx);
 						}
 					}
@@ -2392,7 +2394,17 @@ UEdGraph* UK2Node_CallFunction::GetFunctionGraph(const UEdGraphNode*& OutGraphNo
 			UBlueprint* Blueprint = Cast<UBlueprint>(ParentClass->ClassGeneratedBy);
 			while(Blueprint != NULL)
 			{
-				UEdGraph* TargetGraph = FindObject<UEdGraph>(Blueprint, *(FunctionReference.GetMemberName().ToString()));
+				UEdGraph* TargetGraph = NULL;
+				FName FunctionName = FunctionReference.GetMemberName();
+				for (UEdGraph* Graph : Blueprint->FunctionGraphs) 
+				{
+					if (Graph->GetFName() == FunctionName)
+					{
+						TargetGraph = Graph;
+						break;
+					}
+				}
+
 				if((TargetGraph != NULL) && !TargetGraph->HasAnyFlags(RF_Transient))
 				{
 					// Found the function graph in a Blueprint, return that graph
diff --git a/Engine/Source/Editor/BlueprintGraph/Private/K2Node_ConstructObjectFromClass.cpp b/Engine/Source/Editor/BlueprintGraph/Private/K2Node_ConstructObjectFromClass.cpp
--- a/Engine/Source/Editor/BlueprintGraph/Private/K2Node_ConstructObjectFromClass.cpp
+++ b/Engine/Source/Editor/BlueprintGraph/Private/K2Node_ConstructObjectFromClass.cpp
@@ -194,6 +194,7 @@ void UK2Node_ConstructObjectFromClass::OnClassPinChanged()
 		UEdGraphPin* OldPin = OldPins[i];
 		if (IsSpawnVarPin(OldPin))
 		{
+			OldPin->MarkPendingKill();
 			Pins.Remove(OldPin);
 			OldClassPins.Add(OldPin);
 		}
diff --git a/Engine/Source/Editor/BlueprintGraph/Private/K2Node_CustomEvent.cpp b/Engine/Source/Editor/BlueprintGraph/Private/K2Node_CustomEvent.cpp
--- a/Engine/Source/Editor/BlueprintGraph/Private/K2Node_CustomEvent.cpp
+++ b/Engine/Source/Editor/BlueprintGraph/Private/K2Node_CustomEvent.cpp
@@ -132,8 +132,6 @@ void UK2Node_CustomEvent::Serialize(FArchive& Ar)
 		{
 			if (Pin)
 			{
-				Ar.Preload(Pin);
-
 				if (Pin->Direction == EGPD_Output
 					&& !Pin->PinType.bIsConst
 					&& !K2Schema->IsExecPin(*Pin)
diff --git a/Engine/Source/Editor/BlueprintGraph/Private/K2Node_DoOnceMultiInput.cpp b/Engine/Source/Editor/BlueprintGraph/Private/K2Node_DoOnceMultiInput.cpp
--- a/Engine/Source/Editor/BlueprintGraph/Private/K2Node_DoOnceMultiInput.cpp
+++ b/Engine/Source/Editor/BlueprintGraph/Private/K2Node_DoOnceMultiInput.cpp
@@ -244,8 +244,7 @@ void UK2Node_DoOnceMultiInput::RemoveInputPin(UEdGraphPin* Pin)
 		if (Pins.Find(Pin, /*out*/ PinRemovalIndex))
 		{
 			Pins.RemoveAt(PinRemovalIndex);
-			Pin->Modify();
-			Pin->BreakAllPinLinks();
+			Pin->MarkPendingKill();
 			--NumAdditionalInputs;
 
 			int32 NameIndex = 0;
diff --git a/Engine/Source/Editor/BlueprintGraph/Private/K2Node_EditablePinBase.cpp b/Engine/Source/Editor/BlueprintGraph/Private/K2Node_EditablePinBase.cpp
--- a/Engine/Source/Editor/BlueprintGraph/Private/K2Node_EditablePinBase.cpp
+++ b/Engine/Source/Editor/BlueprintGraph/Private/K2Node_EditablePinBase.cpp
@@ -46,7 +46,6 @@ void UK2Node_EditablePinBase::RemoveUserDefinedPin(TSharedPtr<FUserPinInfo> PinT
 		UEdGraphPin* Pin = Pins[i];
 		if( Pin->PinName == PinName )
 		{
-			Pin->BreakAllPinLinks();
 			Pins.Remove(Pin);
 			Pin->MarkPendingKill();
 
@@ -70,7 +69,6 @@ void UK2Node_EditablePinBase::RemoveUserDefinedPinByName(const FString& PinName)
 		{
 			Pin->Modify();
 
-			Pin->BreakAllPinLinks();
 			Pins.Remove(Pin);
 			Pin->MarkPendingKill();
 
@@ -92,6 +90,8 @@ void UK2Node_EditablePinBase::RemoveUserDefinedPinByName(const FString& PinName)
 
 void UK2Node_EditablePinBase::ExportCustomProperties(FOutputDevice& Out, uint32 Indent)
 {
+	Super::ExportCustomProperties(Out, Indent);
+
 	for (int32 PinIndex = 0; PinIndex < UserDefinedPins.Num(); ++PinIndex)
 	{
 		const FUserPinInfo& PinInfo = *UserDefinedPins[PinIndex].Get();
@@ -166,6 +166,10 @@ void UK2Node_EditablePinBase::ImportCustomProperties(const TCHAR* SourceText, FF
 
 		UserDefinedPins.Add(PinInfo);
 	}
+	else
+	{
+		Super::ImportCustomProperties(SourceText, Warn);
+	}
 }
 
 void UK2Node_EditablePinBase::Serialize(FArchive& Ar)
diff --git a/Engine/Source/Editor/BlueprintGraph/Private/K2Node_ExecutionSequence.cpp b/Engine/Source/Editor/BlueprintGraph/Private/K2Node_ExecutionSequence.cpp
--- a/Engine/Source/Editor/BlueprintGraph/Private/K2Node_ExecutionSequence.cpp
+++ b/Engine/Source/Editor/BlueprintGraph/Private/K2Node_ExecutionSequence.cpp
@@ -219,8 +219,8 @@ void UK2Node_ExecutionSequence::RemovePinFromExecutionNode(UEdGraphPin* TargetPi
 	UK2Node_ExecutionSequence* OwningSeq = Cast<UK2Node_ExecutionSequence>( TargetPin->GetOwningNode() );
 	if (OwningSeq)
 	{
-		TargetPin->BreakAllPinLinks();
 		OwningSeq->Pins.Remove(TargetPin);
+		TargetPin->MarkPendingKill();
 
 		// Renumber the pins so the numbering is compact
 		const UEdGraphSchema_K2* K2Schema = GetDefault<UEdGraphSchema_K2>();
diff --git a/Engine/Source/Editor/BlueprintGraph/Private/K2Node_FormatText.cpp b/Engine/Source/Editor/BlueprintGraph/Private/K2Node_FormatText.cpp
--- a/Engine/Source/Editor/BlueprintGraph/Private/K2Node_FormatText.cpp
+++ b/Engine/Source/Editor/BlueprintGraph/Private/K2Node_FormatText.cpp
@@ -99,7 +99,7 @@ void UK2Node_FormatText::PinConnectionListChanged(UEdGraphPin* Pin)
 			if(CheckPin != FormatPin && CheckPin->Direction == EGPD_Input)
 			{
 				CheckPin->Modify();
-				CheckPin->BreakAllPinLinks();
+				CheckPin->MarkPendingKill();
 				Pins.Remove(CheckPin);
 				--It;
 			}
@@ -138,7 +138,7 @@ void UK2Node_FormatText::PinDefaultValueChanged(UEdGraphPin* Pin)
 				int Index = 0;
 				if(!ArgumentParams.Find(CheckPin->PinName, Index))
 				{
-					CheckPin->BreakAllPinLinks();
+					CheckPin->MarkPendingKill();
 					Pins.Remove(CheckPin);
 					--It;
 				}
@@ -330,7 +330,11 @@ void UK2Node_FormatText::RemoveArgument(int32 InIndex)
 	const FScopedTransaction Transaction( NSLOCTEXT("Kismet", "RemoveArgumentPin", "Remove Argument Pin") );
 	Modify();
 
-	Pins.Remove(FindArgumentPin(PinNames[InIndex]));
+	if (UEdGraphPin* ArgumentPin = FindArgumentPin(PinNames[InIndex]))
+	{
+		Pins.Remove(ArgumentPin);
+		ArgumentPin->MarkPendingKill();
+	}
 	PinNames.RemoveAt(InIndex);
 
 	FBlueprintEditorUtils::MarkBlueprintAsStructurallyModified(GetBlueprint());
diff --git a/Engine/Source/Editor/BlueprintGraph/Private/K2Node_FunctionEntry.cpp b/Engine/Source/Editor/BlueprintGraph/Private/K2Node_FunctionEntry.cpp
--- a/Engine/Source/Editor/BlueprintGraph/Private/K2Node_FunctionEntry.cpp
+++ b/Engine/Source/Editor/BlueprintGraph/Private/K2Node_FunctionEntry.cpp
@@ -47,6 +47,61 @@ public:
 		UK2Node_FunctionEntry* EntryNode = CastChecked<UK2Node_FunctionEntry>(Node);
 
 		UFunction* Function = FindField<UFunction>(EntryNode->SignatureClass, EntryNode->SignatureName);
+		// if this function has a predefined signature (like for inherited/overridden 
+		// functions), then we want to make sure to account for the output 
+		// parameters - this is normally handled by the FunctionResult node, but 
+		// we're not guaranteed that one is connected to the entry node 
+		if (Function && Function->HasAnyFunctionFlags(FUNC_HasOutParms))
+		{
+			const UEdGraphSchema_K2* K2Schema = GetDefault<UEdGraphSchema_K2>();
+
+			for (TFieldIterator<UProperty> ParamIt(Function, EFieldIteratorFlags::ExcludeSuper); ParamIt; ++ParamIt)
+			{
+				UProperty* ParamProperty = *ParamIt;
+
+				// mirrored from UK2Node_FunctionResult::CreatePinsForFunctionEntryExit()
+				const bool bIsFunctionInput = !ParamProperty->HasAnyPropertyFlags(CPF_OutParm) || ParamProperty->HasAnyPropertyFlags(CPF_ReferenceParm);
+				if (bIsFunctionInput)
+				{
+					// 
+					continue;
+				}
+
+				FEdGraphPinType ParamType;
+				if (K2Schema->ConvertPropertyToPinType(ParamProperty, ParamType))
+				{
+					FString ParamName = ParamProperty->GetName();
+
+					bool bTermExists = false;
+					// check to see if this terminal already exists (most 
+					// likely added by a FunctionResult node) - if so, then 
+					// we don't need to add it ourselves
+					for (const FBPTerminal& ResultTerm : Context.Results)
+					{
+						if (ResultTerm.Name == ParamName && ResultTerm.Type == ParamType)
+						{
+							bTermExists = true;
+							break;
+						}
+					}
+
+					if (!bTermExists)
+					{
+						// create a terminal that represents a output param 
+						// for this function; if there is a FunctionResult 
+						// node wired into our function graph, know that it
+						// will first check to see if this already exists 
+						// for it to use (rather than creating one of its own)
+						FBPTerminal* ResultTerm = new (Context.Results) FBPTerminal();
+						ResultTerm->Name = ParamName;
+
+						ResultTerm->Type = ParamType;
+						ResultTerm->bPassedByReference = ParamType.bIsReference;
+						ResultTerm->SetContextTypeStruct(ParamType.PinCategory == UEdGraphSchema_K2::PC_Struct && Cast<UScriptStruct>(ParamType.PinSubCategoryObject.Get()));
+					}
+				}
+			}
+		}
 
 		for (UEdGraphPin* Pin : Node->Pins)
 		{
@@ -201,7 +256,7 @@ void UK2Node_FunctionEntry::RemoveOutputPin(UEdGraphPin* PinToRemove)
 	UK2Node_FunctionEntry* OwningSeq = Cast<UK2Node_FunctionEntry>( PinToRemove->GetOwningNode() );
 	if (OwningSeq)
 	{
-		PinToRemove->BreakAllPinLinks();
+		PinToRemove->MarkPendingKill();
 		OwningSeq->Pins.Remove(PinToRemove);
 	}
 }
diff --git a/Engine/Source/Editor/BlueprintGraph/Private/K2Node_FunctionResult.cpp b/Engine/Source/Editor/BlueprintGraph/Private/K2Node_FunctionResult.cpp
--- a/Engine/Source/Editor/BlueprintGraph/Private/K2Node_FunctionResult.cpp
+++ b/Engine/Source/Editor/BlueprintGraph/Private/K2Node_FunctionResult.cpp
@@ -210,18 +210,9 @@ void UK2Node_FunctionResult::PostPlacedNewNode()
 {
 	Super::PostPlacedNewNode();
 
-	// If the entry is editable, so is the result
-	TArray<UK2Node_FunctionEntry*> AllEntryNodes;
-	if (auto Graph = GetGraph())
-	{
-		Graph->GetNodesOfClass(AllEntryNodes);
-
-		if (AllEntryNodes.Num() > 0)
-		{
-			bIsEditable = AllEntryNodes[0]->bIsEditable;
-		}
-	}
-
+	// adhere to the function's inherited signature (if there is one)
+	SyncWithEntryNode();
+	// reflect any user added outputs (tracked by pre-existing result nodes)
 	SyncWithPrimaryResultNode();
 }
 
@@ -229,19 +220,71 @@ void UK2Node_FunctionResult::PostPasteNode()
 {
 	Super::PostPasteNode();
 
-	// If the entry is editable, so is the result
-	TArray<UK2Node_FunctionEntry*> AllEntryNodes;
-	if (auto Graph = GetGraph())
+	// adhere to the function's inherited signature (if there is one)
+	SyncWithEntryNode();
+	// reflect any user added outputs (tracked by pre-existing result nodes)
+	SyncWithPrimaryResultNode();
+}
+
+bool UK2Node_FunctionResult::CanUserDeleteNode() const
+{
+	bool bCanDelete = true;
+	if (!bIsEditable)
 	{
-		Graph->GetNodesOfClass(AllEntryNodes);
+		if (UEdGraph* Graph = GetGraph())
+		{
+			bCanDelete = false;
+			for (UEdGraphNode* Node : Graph->Nodes)
+			{
+				UK2Node_FunctionResult* ResultNode = Cast<UK2Node_FunctionResult>(Node);
+				if (ResultNode && ResultNode != this)
+				{
+					bCanDelete = true;
+					break;
+				}
+			}
+		}
+	}
+	return bCanDelete;
+}
 
-		if (AllEntryNodes.Num() > 0)
+void UK2Node_FunctionResult::SyncWithEntryNode()
+{
+	bool bWasSignatureMismatched = false;
+	if (UEdGraph* Graph = GetGraph())
+	{
+		for (UEdGraphNode* Node : Graph->Nodes)
 		{
-			bIsEditable = AllEntryNodes[0]->bIsEditable;
+			if (UK2Node_FunctionEntry* EntryNode = Cast<UK2Node_FunctionEntry>(Node))
+			{
+				bWasSignatureMismatched = (EntryNode->SignatureClass != SignatureClass) || 
+					(EntryNode->SignatureName != SignatureName) || (!EntryNode->bIsEditable && UserDefinedPins.Num() > 0);
+
+				// If the entry is editable, so is the result
+				bIsEditable    = EntryNode->bIsEditable;
+				SignatureClass = EntryNode->SignatureClass;
+				SignatureName  = EntryNode->SignatureName;
+				break;
+			}
 		}
 	}
 
-	SyncWithPrimaryResultNode();
+	if (bWasSignatureMismatched)
+	{
+		// to handle pasting of a result node from one function into another;
+		// if the new function is not editable (like for one that is overidden), 
+		// then we shouldn't have userdefined pins
+		if (!bIsEditable)
+		{
+			// iterate backwards so we can remove items from the list as we go
+			for (int32 UserPinIndex = UserDefinedPins.Num() - 1; UserPinIndex >= 0; --UserPinIndex)
+			{
+				RemoveUserDefinedPin(UserDefinedPins[UserPinIndex]);
+			}
+		}
+		
+		ReconstructNode();
+	}
 }
 
 void UK2Node_FunctionResult::SyncWithPrimaryResultNode()
@@ -259,29 +302,63 @@ void UK2Node_FunctionResult::SyncWithPrimaryResultNode()
 
 	if (PrimaryNode)
 	{
-		TArray< TSharedPtr<FUserPinInfo> > UDPinsCopy = UserDefinedPins;
-		for (auto UDPin : UDPinsCopy)
-		{
-			RemoveUserDefinedPin(UDPin);
-		}
-		UserDefinedPins.Empty();
-
 		SignatureClass = PrimaryNode->SignatureClass;
 		SignatureName = PrimaryNode->SignatureName;
 		bIsEditable = PrimaryNode->bIsEditable;
 
-		for (auto UDPin : PrimaryNode->UserDefinedPins)
+		// Temporary array that will contain our list of Old Pins that are no longer part of the return signature
+		TArray< TSharedPtr<FUserPinInfo> > OldPins = UserDefinedPins;
+
+		// Temporary array that will contain our list of Signature Pins that need to be added
+		TArray< TSharedPtr<FUserPinInfo> > SignaturePins = PrimaryNode->UserDefinedPins;
+
+		for (int OldIndex = OldPins.Num() - 1; OldIndex >= 0; --OldIndex)
 		{
-			if (UDPin.IsValid())
+			TSharedPtr<FUserPinInfo> OldPin = OldPins[OldIndex];
+
+			if (!OldPin.IsValid())
+			{
+				OldPins.RemoveAt(OldIndex);
+			}
+			else
 			{
-				TSharedPtr<FUserPinInfo> NewPinInfo = MakeShareable(new FUserPinInfo());
-				NewPinInfo->PinName = UDPin->PinName;
-				NewPinInfo->PinType = UDPin->PinType;
-				NewPinInfo->DesiredPinDirection = UDPin->DesiredPinDirection;
-				UserDefinedPins.Add(NewPinInfo);
+				for (int SignatureIndex = SignaturePins.Num() - 1; SignatureIndex >= 0; --SignatureIndex)
+				{
+					TSharedPtr<FUserPinInfo> SignaturePin = SignaturePins[SignatureIndex];
+					if (!SignaturePin.IsValid())
+					{
+						SignaturePins.RemoveAt(SignatureIndex);
+					}
+					else if (OldPin->PinName == SignaturePin->PinName &&
+						OldPin->PinType == SignaturePin->PinType &&
+						OldPin->DesiredPinDirection == SignaturePin->DesiredPinDirection)
+					{
+						// We have a match between our Signature pins and our Old Pins,
+						// so we can leave the old pin as is by removing it from both temporary lists.
+						OldPins.RemoveAt(OldIndex);
+						SignaturePins.RemoveAt(SignatureIndex);
+						break;
+					}
+				}
 			}
 		}
 
+		// Remove old pins that are not part of the primary node signature
+		for (TSharedPtr<FUserPinInfo> OldPinToRemove : OldPins)
+		{
+			RemoveUserDefinedPin(OldPinToRemove);
+		}
+
+		// Add pins that don't exist yet but are part of the primary node signature
+		for (TSharedPtr<FUserPinInfo> SignaturePinToAdd : SignaturePins)
+		{
+			TSharedPtr<FUserPinInfo> NewPinInfo = MakeShareable(new FUserPinInfo());
+			NewPinInfo->PinName = SignaturePinToAdd->PinName;
+			NewPinInfo->PinType = SignaturePinToAdd->PinType;
+			NewPinInfo->DesiredPinDirection = SignaturePinToAdd->DesiredPinDirection;
+			UserDefinedPins.Add(NewPinInfo);
+		}
+
 		ReconstructNode();
 	}
 }
diff --git a/Engine/Source/Editor/BlueprintGraph/Private/K2Node_GetClassDefaults.cpp b/Engine/Source/Editor/BlueprintGraph/Private/K2Node_GetClassDefaults.cpp
--- a/Engine/Source/Editor/BlueprintGraph/Private/K2Node_GetClassDefaults.cpp
+++ b/Engine/Source/Editor/BlueprintGraph/Private/K2Node_GetClassDefaults.cpp
@@ -366,12 +366,17 @@ UClass* UK2Node_GetClassDefaults::GetInputClass(const UEdGraphPin* FromPin) cons
 		}
 	}
 
-	// Switch Blueprint Class types to use the generated skeleton class.
+	// Switch Blueprint Class types to use the generated skeleton class (if valid).
 	if (InputClass)
 	{
 		if (UBlueprint* Blueprint = Cast<UBlueprint>(InputClass->ClassGeneratedBy))
 		{
-			InputClass = Blueprint->SkeletonGeneratedClass;
+			// Stick with the original (serialized) class if the skeleton class is not valid for some reason (e.g. the Blueprint hasn't been compiled on load yet).
+			// Note: There's not a need to force it to be preloaded here in that case, because once it is loaded, we'll end up reconstructing this node again anyway.
+			if (Blueprint->SkeletonGeneratedClass)
+			{
+				InputClass = Blueprint->SkeletonGeneratedClass;
+			}
 		}
 	}
 
diff --git a/Engine/Source/Editor/BlueprintGraph/Private/K2Node_MakeArray.cpp b/Engine/Source/Editor/BlueprintGraph/Private/K2Node_MakeArray.cpp
--- a/Engine/Source/Editor/BlueprintGraph/Private/K2Node_MakeArray.cpp
+++ b/Engine/Source/Editor/BlueprintGraph/Private/K2Node_MakeArray.cpp
@@ -182,8 +182,8 @@ void UK2Node_MakeArray::NotifyPinConnectionListChanged(UEdGraphPin* Pin)
 
 	if (Pin->LinkedTo.Num() > 0)
 	{
-		// Just made a connection, was it the first?
-		if (NumPinsWithLinks == 1)
+		// Just made a connection, was it the first and was it the parent pin?
+		if (Pin->ParentPin == nullptr && NumPinsWithLinks == 1)
 		{
 			// Update the types on all the pins
 			OutputPin->PinType = Pin->LinkedTo[0]->PinType;
@@ -200,7 +200,8 @@ void UK2Node_MakeArray::NotifyPinConnectionListChanged(UEdGraphPin* Pin)
 			bool bResetOutputPin = true;
 			for (int32 PinIndex = 0; PinIndex < InputPins.Num(); ++PinIndex)
 			{
-				if( InputPins[PinIndex]->GetDefaultAsString().IsEmpty() == false )
+				UEdGraphPin*& CurrentPin = InputPins[PinIndex];
+				if (CurrentPin->GetDefaultAsString() != CurrentPin->AutogeneratedDefaultValue)
 				{
 					bResetOutputPin = false;
 				}
@@ -314,6 +315,7 @@ UFunction* UK2Node_MakeArray::GetArrayAddFunction() const
 void UK2Node_MakeArray::PostReconstructNode()
 {
 	// Find a pin that has connections to use to jumpstart the wildcard process
+	bool bResetToWildcard = true;
 	for (int32 PinIndex = 0; PinIndex < Pins.Num(); ++PinIndex)
 	{
 		if (Pins[PinIndex]->LinkedTo.Num() > 0)
@@ -322,24 +324,47 @@ void UK2Node_MakeArray::PostReconstructNode()
 
 			// Update the types on all the pins
 			UEdGraphPin* OutputPin = GetOutputPin();
-			OutputPin->PinType = Pins[PinIndex]->LinkedTo[0]->PinType;
+			if (Pins[PinIndex]->ParentPin != nullptr)
+			{ 
+				// This is a split pin, so we want to base the pin type on the parent rather than the linked to pin
+				OutputPin->PinType = Pins[PinIndex]->ParentPin->PinType;
+			}
+			else
+			{
+				OutputPin->PinType = Pins[PinIndex]->LinkedTo[0]->PinType;
+			}
 			OutputPin->PinType.bIsArray = true;
 			PropagatePinType();
+			bResetToWildcard = false;
 			break;
 		}
-		else if(!Pins[PinIndex]->GetDefaultAsString().IsEmpty())
+		else if(Pins[PinIndex]->GetDefaultAsString() != Pins[PinIndex]->AutogeneratedDefaultValue)
 		{
 			// The pin has user data in it, continue to use its type as the type for all pins.
 
 			// Update the types on all the pins
 			UEdGraphPin* OutputPin = GetOutputPin();
-			OutputPin->PinType = Pins[PinIndex]->PinType;
+			if (Pins[PinIndex]->ParentPin != nullptr)
+			{
+				// This is a split pin, so we want to base the pin type on the parent instead of this pin
+				OutputPin->PinType = Pins[PinIndex]->ParentPin->PinType;
+			}
+			else
+			{
+				OutputPin->PinType = Pins[PinIndex]->PinType;
+			}
 			OutputPin->PinType.bIsArray = true;
 			PropagatePinType();
+			bResetToWildcard = false;
 			break;
 		}
 	}
 
+	if (bResetToWildcard)
+	{
+		ClearPinTypeToWildcard();
+	}
+
 	Super::PostReconstructNode();
 }
 
@@ -371,6 +396,44 @@ void UK2Node_MakeArray::AddInputPin()
 	}
 }
 
+void UK2Node_MakeArray::SyncPinNames()
+{
+	int32 CurrentNumParentPins = 0;
+	for (int32 PinIndex = 0; PinIndex < Pins.Num(); ++PinIndex)
+	{
+		UEdGraphPin*& CurrentPin = Pins[PinIndex];
+		if (CurrentPin->Direction == EGPD_Input &&
+			CurrentPin->ParentPin == nullptr)
+		{
+			FString OldName = CurrentPin->PinName;
+			FString ElementName = FString::Printf(TEXT("[%d]"), CurrentNumParentPins++);
+
+			CurrentPin->Modify();
+			CurrentPin->PinName = ElementName;
+
+			if (CurrentPin->SubPins.Num() > 0)
+			{
+				FString OldFriendlyName = OldName;
+				FString ElementFriendlyName = ElementName;
+
+				// SubPin Friendly Name has an extra space in it so we need to account for that
+				OldFriendlyName.InsertAt(1, " ");
+				ElementFriendlyName.InsertAt(1, " ");
+
+				for (UEdGraphPin*& SubPin : CurrentPin->SubPins)
+				{
+					FString SubPinFriendlyName = SubPin->PinFriendlyName.ToString();
+					SubPinFriendlyName = SubPinFriendlyName.Replace(*OldFriendlyName, *ElementFriendlyName);
+
+					SubPin->Modify();
+					SubPin->PinName = SubPin->PinName.Replace(*OldName, *ElementName);
+					SubPin->PinFriendlyName = FText::FromString(SubPinFriendlyName);
+				}
+			}
+		}
+	}
+}
+
 void UK2Node_MakeArray::RemoveInputPin(UEdGraphPin* Pin)
 {
 	FScopedTransaction Transaction( LOCTEXT("RemovePinTx", "RemovePin") );
@@ -381,17 +444,23 @@ void UK2Node_MakeArray::RemoveInputPin(UEdGraphPin* Pin)
 	int32 PinRemovalIndex = INDEX_NONE;
 	if (Pins.Find(Pin, /*out*/ PinRemovalIndex))
 	{
-		for (int32 PinIndex = PinRemovalIndex + 1; PinIndex < Pins.Num(); ++PinIndex)
+		UEdGraphPin* ParentPin = Pin->ParentPin;
+		Pins.RemoveAt(PinRemovalIndex);
+		Pin->MarkPendingKill();
+
+		if (!ParentPin)
 		{
-			Pins[PinIndex]->Modify();
-			Pins[PinIndex]->PinName = FString::Printf(TEXT("[%d]"), PinIndex - 2); // -1 to shift back one, -1 to account for the output pin at the 0th position
+			--NumInputs;
+		}
+		else if (ParentPin->SubPins.Num() == 0)
+		{
+			// We have removed all of the subpins, so we need to remove the parent as well
+			Pins.Remove(ParentPin);
+			ParentPin->MarkPendingKill();
+			--NumInputs;
 		}
 
-		Pins.RemoveAt(PinRemovalIndex);
-		Pin->Modify();
-		Pin->BreakAllPinLinks();
-
-		--NumInputs;
+		SyncPinNames();
 
 		FBlueprintEditorUtils::MarkBlueprintAsStructurallyModified(GetBlueprint());
 	}
@@ -449,7 +518,8 @@ void UK2Node_MakeArray::GetContextMenuActions(const FGraphNodeContextMenuBuilder
 
 bool UK2Node_MakeArray::IsConnectionDisallowed(const UEdGraphPin* MyPin, const UEdGraphPin* OtherPin, FString& OutReason) const
 {
-	if(OtherPin->PinType.bIsArray == true && MyPin->Direction == EGPD_Input)
+	// if MyPin has a ParentPin then we are dealing with a split pin and we should evaluate it with default behavior
+	if(MyPin->ParentPin == nullptr && OtherPin->PinType.bIsArray == true && MyPin->Direction == EGPD_Input)
 	{
 		OutReason = NSLOCTEXT("K2Node", "MakeArray_InputIsArray", "Cannot make an array with an input of an array!").ToString();
 		return true;
diff --git a/Engine/Source/Editor/BlueprintGraph/Private/K2Node_MakeStruct.cpp b/Engine/Source/Editor/BlueprintGraph/Private/K2Node_MakeStruct.cpp
--- a/Engine/Source/Editor/BlueprintGraph/Private/K2Node_MakeStruct.cpp
+++ b/Engine/Source/Editor/BlueprintGraph/Private/K2Node_MakeStruct.cpp
@@ -371,12 +371,6 @@ void UK2Node_MakeStruct::Serialize(FArchive& Ar)
 			{
 				FOptionalPinManager PinManager;
 
-				// Preload all pins so we can examine their values
-				for (UEdGraphPin* Pin : Pins)
-				{
-					Ar.Preload(Pin);
-				}
-
 				// Have to check if this node is even in danger.
 				for (FOptionalPinFromProperty& PropertyEntry : ShowPinForProperties)
 				{
diff --git a/Engine/Source/Editor/BlueprintGraph/Private/K2Node_SetFieldsInStruct.cpp b/Engine/Source/Editor/BlueprintGraph/Private/K2Node_SetFieldsInStruct.cpp
--- a/Engine/Source/Editor/BlueprintGraph/Private/K2Node_SetFieldsInStruct.cpp
+++ b/Engine/Source/Editor/BlueprintGraph/Private/K2Node_SetFieldsInStruct.cpp
@@ -133,7 +133,6 @@ void UK2Node_SetFieldsInStruct::AllocateDefaultPins()
 		UEdGraphPin* InPin = CreatePin(EGPD_Input, Schema->PC_Struct, TEXT(""), StructType, false, true, SetFieldsInStructHelper::StructRefPinName());
 
 		UEdGraphPin* OutPin = CreatePin(EGPD_Output, Schema->PC_Struct, TEXT(""), StructType, false, true, SetFieldsInStructHelper::StructOutPinName());
-		ByRefMatchupPins.Add(InPin, OutPin);
 
 		// Input pin will forward the ref to the output, if the input value is not a reference connection, a copy is made and modified instead and provided as a reference until the function is called again.
 		InPin->AssignByRefPassThroughConnection(OutPin);
diff --git a/Engine/Source/Editor/BlueprintGraph/Private/K2Node_SpawnActor.cpp b/Engine/Source/Editor/BlueprintGraph/Private/K2Node_SpawnActor.cpp
--- a/Engine/Source/Editor/BlueprintGraph/Private/K2Node_SpawnActor.cpp
+++ b/Engine/Source/Editor/BlueprintGraph/Private/K2Node_SpawnActor.cpp
@@ -145,7 +145,7 @@ void UK2Node_SpawnActor::PinDefaultValueChanged(UEdGraphPin* ChangedPin)
 			UEdGraphPin* OldPin = OldPins[i];
 			if (IsSpawnVarPin(OldPin))
 			{
-				OldPin->BreakAllPinLinks();
+				OldPin->MarkPendingKill();
 				Pins.Remove(OldPin);
 			}
 		}
diff --git a/Engine/Source/Editor/BlueprintGraph/Private/K2Node_SpawnActorFromClass.cpp b/Engine/Source/Editor/BlueprintGraph/Private/K2Node_SpawnActorFromClass.cpp
--- a/Engine/Source/Editor/BlueprintGraph/Private/K2Node_SpawnActorFromClass.cpp
+++ b/Engine/Source/Editor/BlueprintGraph/Private/K2Node_SpawnActorFromClass.cpp
@@ -611,7 +611,7 @@ void UK2Node_SpawnActorFromClass::ExpandNode(class FKismetCompilerContext& Compi
 	UClass* SpawnClass = (SpawnClassPin != NULL) ? Cast<UClass>(SpawnClassPin->DefaultObject) : NULL;
 	if((0 == SpawnClassPin->LinkedTo.Num()) && (NULL == SpawnClass))
 	{
-		CompilerContext.MessageLog.Error(*LOCTEXT("SpawnActorNodeMissingClass_Error", "Spawn node @@ must have a class specified.").ToString(), SpawnNode);
+		CompilerContext.MessageLog.Error(*LOCTEXT("SpawnActorNodeMissingClass_Error", "Spawn node @@ must have a @@ specified.").ToString(), SpawnNode, SpawnClassPin);
 		// we break exec links so this is the only error we get, don't want the SpawnActor node being considered and giving 'unexpected node' type warnings
 		SpawnNode->BreakAllNodeLinks();
 		return;
diff --git a/Engine/Source/Editor/BlueprintGraph/Private/K2Node_Switch.cpp b/Engine/Source/Editor/BlueprintGraph/Private/K2Node_Switch.cpp
--- a/Engine/Source/Editor/BlueprintGraph/Private/K2Node_Switch.cpp
+++ b/Engine/Source/Editor/BlueprintGraph/Private/K2Node_Switch.cpp
@@ -98,7 +98,7 @@ public:
 					FBPTerminal* CaseValueTerm = new (Context.Literals) FBPTerminal();
 					CaseValueTerm->Name = Pin->PinName;
 					CaseValueTerm->Type = SelectionPin->PinType;
-					CaseValueTerm->Source = Pin;
+					CaseValueTerm->SourcePin = Pin;
 					CaseValueTerm->bIsLiteral = true;
 
 					// Call the comparison function associated with this switch node
@@ -256,10 +256,8 @@ void UK2Node_Switch::RemovePinFromSwitchNode(UEdGraphPin* TargetPin)
 	}
 	else
 	{
-		TargetPin->BreakAllPinLinks();
+		TargetPin->MarkPendingKill();
 		Pins.Remove(TargetPin);
-
-		RemovePin(TargetPin);
 	}
 }
 
diff --git a/Engine/Source/Editor/BlueprintGraph/Private/K2Node_Variable.cpp b/Engine/Source/Editor/BlueprintGraph/Private/K2Node_Variable.cpp
--- a/Engine/Source/Editor/BlueprintGraph/Private/K2Node_Variable.cpp
+++ b/Engine/Source/Editor/BlueprintGraph/Private/K2Node_Variable.cpp
@@ -90,7 +90,17 @@ bool UK2Node_Variable::CreatePinForVariable(EEdGraphPinDirection Direction, FStr
 	{
 		if (!VariableReference.IsLocalScope())
 		{
-			Message_Warn(*FString::Printf(TEXT("CreatePinForVariable: '%s' variable not found. Base class was probably changed."), *GetVarNameString()));
+			FString WarningMsg = FString::Printf(TEXT("'%s' variable not found. Base class was probably changed."), *GetVarNameString());
+
+			UBlueprint* OwnerBP = GetBlueprint();
+			if (OwnerBP && OwnerBP->CurrentMessageLog)
+			{
+				OwnerBP->CurrentMessageLog->Warning(*FString::Printf(TEXT("@@: %s"), *WarningMsg), this);
+			}
+			else
+			{
+				Message_Warn(*WarningMsg);
+			}			
 		}
 		return false;
 	}
diff --git a/Engine/Source/Editor/BlueprintGraph/Private/K2Node_VariableSet.cpp b/Engine/Source/Editor/BlueprintGraph/Private/K2Node_VariableSet.cpp
--- a/Engine/Source/Editor/BlueprintGraph/Private/K2Node_VariableSet.cpp
+++ b/Engine/Source/Editor/BlueprintGraph/Private/K2Node_VariableSet.cpp
@@ -377,6 +377,7 @@ void UK2Node_VariableSet::ExpandNode(class FKismetCompilerContext& CompilerConte
 				}
 			}
 			Pins.Remove(Pin);
+			Pin->MarkPendingKill();
 		}
 		
 	}
diff --git a/Engine/Source/Editor/EditorStyle/Private/SlateEditorStyle.cpp b/Engine/Source/Editor/EditorStyle/Private/SlateEditorStyle.cpp
--- a/Engine/Source/Editor/EditorStyle/Private/SlateEditorStyle.cpp
+++ b/Engine/Source/Editor/EditorStyle/Private/SlateEditorStyle.cpp
@@ -4830,13 +4830,14 @@ void FSlateEditorStyle::FStyle::SetupLevelEditorStyle()
 		Set( "BlueprintProfiler.BPIcon_Macro", new IMAGE_BRUSH("Icons/AssetIcons/BlueprintMacroLibrary_16x", Icon16x16, FLinearColor( 0.46f, 0.54f, 0.81f )));
 		Set( "BlueprintProfiler.BPIcon_Interface", new IMAGE_BRUSH("Icons/AssetIcons/BlueprintInterface_16x", Icon16x16, FLinearColor( 0.46f, 0.54f, 0.81f )));
 		Set( "BlueprintProfiler.BPIcon_FunctionLibrary", new IMAGE_BRUSH("Icons/AssetIcons/BlueprintFunctionLibrary_16x", Icon16x16, FLinearColor( 0.46f, 0.54f, 0.81f )));
-		Set( "BlueprintProfiler.Actor", new IMAGE_BRUSH("Icons/ActorIcons/Sphere_16px", Icon16x16));
+		Set( "BlueprintProfiler.Actor", new IMAGE_BRUSH("Icons/AssetIcons/Actor_16x", Icon16x16, FLinearColor( 1.f, 1.f, 1.f )));
 		Set( "BlueprintProfiler.BPEventNode", new IMAGE_BRUSH("Icons/icon_Blueprint_Event_16x", Icon16x16, FLinearColor( 0.91f, 0.16f, 0.16f )));
 		Set( "BlueprintProfiler.BPFunctionNode", new IMAGE_BRUSH("Icons/icon_Blueprint_NewFunction_16x", Icon16x16, FLinearColor( 0.46f, 0.54f, 0.81f )));
 		Set( "BlueprintProfiler.BPMacroNode", new IMAGE_BRUSH("Icons/icon_Blueprint_Macro_16x", Icon16x16, FLinearColor( 0.46f, 0.54f, 0.81f )));
 		Set( "BlueprintProfiler.BPBranchNode", new IMAGE_BRUSH("Icons/icon_Blueprint_Sequence_16x", Icon16x16, FLinearColor( 0.46f, 0.54f, 0.81f )));
 		Set( "BlueprintProfiler.BPPinConnected", new IMAGE_BRUSH( "Icons/BPProfiler_ExecPin_Connected", Icon12x16 ) );
 		Set( "BlueprintProfiler.BPPinDisconnected", new IMAGE_BRUSH( "Icons/BPProfiler_ExecPin_Disconnected", Icon12x16 ) );
+		Set( "BlueprintProfiler.PureNode", new IMAGE_BRUSH( "/Icons/pill_16x", Icon16x16 ) );
 		Set( "BlueprintProfiler.BPNode", new IMAGE_BRUSH("Icons/Icon_Profiler_BPNode_16x", Icon16x16));
 	}
 #endif // WITH_EDITOR
@@ -5186,6 +5187,8 @@ void FSlateEditorStyle::FStyle::SetupPersonaStyle()
 		Set( "Kismet.Status.Unknown", new IMAGE_BRUSH( "Old/Kismet2/CompileStatus_Working", Icon40x40 ) );
 		Set( "Kismet.Status.Error", new IMAGE_BRUSH( "Old/Kismet2/CompileStatus_Fail", Icon40x40 ) );
 		Set( "Kismet.Status.Good", new IMAGE_BRUSH( "Old/Kismet2/CompileStatus_Good", Icon40x40 ) );
+		Set( "Kismet.Status.Instrumented", new IMAGE_BRUSH( "Old/Kismet2/CompileStatus_Instrumented", Icon40x40 ) );
+		Set( "Kismet.Status.NotInstrumented", new IMAGE_BRUSH( "Old/Kismet2/CompileStatus_NotInstrumented", Icon40x40 ) );
 		Set( "Kismet.Status.Warning", new IMAGE_BRUSH( "Old/Kismet2/CompileStatus_Warning", Icon40x40 ) );
 
 		Set( "BlueprintEditor.AddNewVariable", new IMAGE_BRUSH( "Icons/icon_Blueprint_AddVariable_40px", Icon40x40) );
diff --git a/Engine/Source/Editor/EnvironmentQueryEditor/Private/EnvironmentQueryGraph.cpp b/Engine/Source/Editor/EnvironmentQueryEditor/Private/EnvironmentQueryGraph.cpp
--- a/Engine/Source/Editor/EnvironmentQueryEditor/Private/EnvironmentQueryGraph.cpp
+++ b/Engine/Source/Editor/EnvironmentQueryEditor/Private/EnvironmentQueryGraph.cpp
@@ -267,6 +267,7 @@ void UEnvironmentQueryGraph::UpdateVersion_NestedNodes()
 		UEnvironmentQueryGraphNode_Option* OptionNode = Cast<UEnvironmentQueryGraphNode_Option>(Nodes[Idx]);
 		if (OptionNode && OptionNode->Pins.IsValidIndex(1))
 		{
+			OptionNode->Pins[1]->MarkPendingKill();
 			OptionNode->Pins.RemoveAt(1);
 		}
 	}
diff --git a/Engine/Source/Editor/GameplayTagsEditor/Private/SGameplayTagContainerGraphPin.cpp b/Engine/Source/Editor/GameplayTagsEditor/Private/SGameplayTagContainerGraphPin.cpp
--- a/Engine/Source/Editor/GameplayTagsEditor/Private/SGameplayTagContainerGraphPin.cpp
+++ b/Engine/Source/Editor/GameplayTagsEditor/Private/SGameplayTagContainerGraphPin.cpp
@@ -98,7 +98,7 @@ void SGameplayTagContainerGraphPin::ParseDefaultValueData()
 TSharedRef<SWidget> SGameplayTagContainerGraphPin::GetListContent()
 {
 	EditableContainers.Empty();
-	EditableContainers.Add( SGameplayTagWidget::FEditableGameplayTagContainerDatum( GraphPinObj, TagContainer.Get() ) );
+	EditableContainers.Add( SGameplayTagWidget::FEditableGameplayTagContainerDatum( GraphPinObj->GetOwningNode(), TagContainer.Get() ) );
 
 	return SNew( SVerticalBox )
 		+SVerticalBox::Slot()
diff --git a/Engine/Source/Editor/GameplayTagsEditor/Private/SGameplayTagGraphPin.cpp b/Engine/Source/Editor/GameplayTagsEditor/Private/SGameplayTagGraphPin.cpp
--- a/Engine/Source/Editor/GameplayTagsEditor/Private/SGameplayTagGraphPin.cpp
+++ b/Engine/Source/Editor/GameplayTagsEditor/Private/SGameplayTagGraphPin.cpp
@@ -45,7 +45,7 @@ void SGameplayTagGraphPin::ParseDefaultValueData()
 {
 	FString TagString = GraphPinObj->GetDefaultAsString();
 
-	UK2Node_CallFunction* CallFuncNode = Cast<UK2Node_CallFunction>(GraphPinObj->GetOuter());
+	UK2Node_CallFunction* CallFuncNode = Cast<UK2Node_CallFunction>(GraphPinObj->GetOwningNode());
 	
 	FilterString.Empty();
 	if (CallFuncNode)
@@ -82,7 +82,7 @@ void SGameplayTagGraphPin::ParseDefaultValueData()
 TSharedRef<SWidget> SGameplayTagGraphPin::GetListContent()
 {
 	EditableContainers.Empty();
-	EditableContainers.Add( SGameplayTagWidget::FEditableGameplayTagContainerDatum( GraphPinObj, TagContainer.Get() ) );
+	EditableContainers.Add( SGameplayTagWidget::FEditableGameplayTagContainerDatum( GraphPinObj->GetOwningNode(), TagContainer.Get() ) );
 
 	return SNew( SVerticalBox )
 		+SVerticalBox::Slot()
diff --git a/Engine/Source/Editor/GameplayTagsEditor/Private/SGameplayTagQueryGraphPin.cpp b/Engine/Source/Editor/GameplayTagsEditor/Private/SGameplayTagQueryGraphPin.cpp
--- a/Engine/Source/Editor/GameplayTagsEditor/Private/SGameplayTagQueryGraphPin.cpp
+++ b/Engine/Source/Editor/GameplayTagsEditor/Private/SGameplayTagQueryGraphPin.cpp
@@ -55,7 +55,7 @@ void SGameplayTagQueryGraphPin::ParseDefaultValueData()
 TSharedRef<SWidget> SGameplayTagQueryGraphPin::GetListContent()
 {
 	EditableQueries.Empty();
-	EditableQueries.Add(SGameplayTagQueryWidget::FEditableGameplayTagQueryDatum(GraphPinObj, TagQuery.Get(), &TagQueryExportText));
+	EditableQueries.Add(SGameplayTagQueryWidget::FEditableGameplayTagQueryDatum(GraphPinObj->GetOwningNode(), TagQuery.Get(), &TagQueryExportText));
 
 	return SNew( SVerticalBox )
 		+SVerticalBox::Slot()
diff --git a/Engine/Source/Editor/GameplayTasksEditor/Classes/K2Node_LatentGameplayTaskCall.h b/Engine/Source/Editor/GameplayTasksEditor/Classes/K2Node_LatentGameplayTaskCall.h
--- a/Engine/Source/Editor/GameplayTasksEditor/Classes/K2Node_LatentGameplayTaskCall.h
+++ b/Engine/Source/Editor/GameplayTasksEditor/Classes/K2Node_LatentGameplayTaskCall.h
@@ -33,7 +33,7 @@ public:
 	bool ValidateActorArraySpawning(class FKismetCompilerContext& CompilerContext, bool bGenerateErrors);
 
 	UPROPERTY()
-	TArray<UEdGraphPin*> SpawnParmPins;
+	TArray<FString> SpawnParmPins;
 
 	static void RegisterSpecializedTaskNodeClass(TSubclassOf<UK2Node_LatentGameplayTaskCall> NodeClass);
 protected:
diff --git a/Engine/Source/Editor/GameplayTasksEditor/Private/K2Node_LatentGameplayTaskCall.cpp b/Engine/Source/Editor/GameplayTasksEditor/Private/K2Node_LatentGameplayTaskCall.cpp
--- a/Engine/Source/Editor/GameplayTasksEditor/Private/K2Node_LatentGameplayTaskCall.cpp
+++ b/Engine/Source/Editor/GameplayTasksEditor/Private/K2Node_LatentGameplayTaskCall.cpp
@@ -207,7 +207,7 @@ void UK2Node_LatentGameplayTaskCall::CreatePinsForClass(UClass* InClass)
 
 			UEdGraphPin* Pin = CreatePin(EGPD_Input, TEXT(""), TEXT(""), NULL, false, false, Property->GetName());
 			const bool bPinGood = (Pin != NULL) && K2Schema->ConvertPropertyToPinType(Property, /*out*/ Pin->PinType);
-			SpawnParmPins.Add(Pin);
+			SpawnParmPins.Add(Pin->PinName);
 
 			if (ClassDefaultObject && Pin && K2Schema->PinDefaultValueIsEditable(*Pin))
 			{
@@ -237,10 +237,14 @@ void UK2Node_LatentGameplayTaskCall::PinDefaultValueChanged(UEdGraphPin* Changed
 		//ResultPin->BreakAllPinLinks();
 
 		// Remove all pins related to archetype variables
-		for (auto OldPin : SpawnParmPins)
+		for (const FString& OldPinReference : SpawnParmPins)
 		{
-			OldPin->BreakAllPinLinks();
-			Pins.Remove(OldPin);
+			UEdGraphPin* OldPin = FindPin(OldPinReference);
+			if(OldPin)
+			{
+				OldPin->MarkPendingKill();
+				Pins.Remove(OldPin);
+			}
 		}
 		SpawnParmPins.Empty();
 
@@ -404,9 +408,10 @@ bool UK2Node_LatentGameplayTaskCall::ValidateActorArraySpawning(class FKismetCom
 bool UK2Node_LatentGameplayTaskCall::ConnectSpawnProperties(UClass* ClassToSpawn, const UEdGraphSchema_K2* Schema, class FKismetCompilerContext& CompilerContext, UEdGraph* SourceGraph, UEdGraphPin*& LastThenPin, UEdGraphPin* SpawnedActorReturnPin)
 {
 	bool bIsErrorFree = true;
-	for (auto SpawnVarPin : SpawnParmPins)
+	for (const FString& OldPinReference : SpawnParmPins)
 	{
-		const bool bHasDefaultValue = !SpawnVarPin->DefaultValue.IsEmpty() || !SpawnVarPin->DefaultTextValue.IsEmpty() || SpawnVarPin->DefaultObject;
+		UEdGraphPin* SpawnVarPin = FindPin(OldPinReference);
+		const bool bHasDefaultValue = SpawnVarPin && (!SpawnVarPin->DefaultValue.IsEmpty() || !SpawnVarPin->DefaultTextValue.IsEmpty() || SpawnVarPin->DefaultObject);
 		if (SpawnVarPin->LinkedTo.Num() > 0 || bHasDefaultValue)
 		{
 			if (SpawnVarPin->LinkedTo.Num() == 0)
diff --git a/Engine/Source/Editor/GraphEditor/Private/ConnectionDrawingPolicy.cpp b/Engine/Source/Editor/GraphEditor/Private/ConnectionDrawingPolicy.cpp
--- a/Engine/Source/Editor/GraphEditor/Private/ConnectionDrawingPolicy.cpp
+++ b/Engine/Source/Editor/GraphEditor/Private/ConnectionDrawingPolicy.cpp
@@ -136,7 +136,7 @@ void FConnectionDrawingPolicy::DrawSplineWithArrow(const FGeometry& StartGeom, c
 }
 
 // Update the drawing policy with the set of hovered pins (which can be empty)
-void FConnectionDrawingPolicy::SetHoveredPins(const TSet< TWeakObjectPtr<UEdGraphPin> >& InHoveredPins, const TArray< TSharedPtr<SGraphPin> >& OverridePins, double HoverTime)
+void FConnectionDrawingPolicy::SetHoveredPins(const TSet< FEdGraphPinReference >& InHoveredPins, const TArray< TSharedPtr<SGraphPin> >& OverridePins, double HoverTime)
 {
 	HoveredPins.Empty();
 
diff --git a/Engine/Source/Editor/GraphEditor/Private/GraphEditorDragDropAction.cpp b/Engine/Source/Editor/GraphEditor/Private/GraphEditorDragDropAction.cpp
--- a/Engine/Source/Editor/GraphEditor/Private/GraphEditorDragDropAction.cpp
+++ b/Engine/Source/Editor/GraphEditor/Private/GraphEditorDragDropAction.cpp
@@ -35,7 +35,7 @@ UEdGraph* FGraphEditorDragDropAction::GetHoveredGraph() const
 
 void FGraphEditorDragDropAction::SetHoveredPin(UEdGraphPin* InPin)
 {
-	if (HoveredPin != InPin)
+	if (HoveredPin.Get() != InPin)
 	{
 		HoveredPin = InPin;
 		HoverTargetChanged();
diff --git a/Engine/Source/Editor/GraphEditor/Private/KismetNodes/KismetNodeInfoContext.cpp b/Engine/Source/Editor/GraphEditor/Private/KismetNodes/KismetNodeInfoContext.cpp
--- a/Engine/Source/Editor/GraphEditor/Private/KismetNodes/KismetNodeInfoContext.cpp
+++ b/Engine/Source/Editor/GraphEditor/Private/KismetNodes/KismetNodeInfoContext.cpp
@@ -83,16 +83,16 @@ FKismetNodeInfoContext::FKismetNodeInfoContext(UEdGraph* SourceGraph)
 		}
 
 		// Covert the watched pin array into a set
-		for (auto WatchedPinIt = SourceBlueprint->PinWatches.CreateConstIterator(); WatchedPinIt; ++WatchedPinIt)
+		for (auto WatchedPinIt = SourceBlueprint->WatchedPins.CreateConstIterator(); WatchedPinIt; ++WatchedPinIt)
 		{
-			UEdGraphPin* WatchedPin = *WatchedPinIt;
+			UEdGraphPin* WatchedPin = WatchedPinIt->Get();
 			if (!ensure(WatchedPin))
 			{
 				continue;
 			}
 
 			UEdGraphNode* OwningNode = Cast<UEdGraphNode>(WatchedPin->GetOuter());
-			if (!ensure(OwningNode != NULL)) // shouldn't happen, but just in case a dead pin was added to the PinWatches array
+			if (!ensure(OwningNode != NULL)) // shouldn't happen, but just in case a dead pin was added to the WatchedPins array
 			{
 				continue;
 			}
diff --git a/Engine/Source/Editor/GraphEditor/Private/SCommentBubble.cpp b/Engine/Source/Editor/GraphEditor/Private/SCommentBubble.cpp
--- a/Engine/Source/Editor/GraphEditor/Private/SCommentBubble.cpp
+++ b/Engine/Source/Editor/GraphEditor/Private/SCommentBubble.cpp
@@ -307,6 +307,15 @@ void SCommentBubble::UpdateBubble()
 	}
 }
 
+bool SCommentBubble::TextBlockHasKeyboardFocus() const
+{
+	if (TextBlock.IsValid())
+	{
+		return TextBlock->HasKeyboardFocus();
+	}
+	return false;
+}
+
 FVector2D SCommentBubble::GetOffset() const
 {
 	return FVector2D( 0.f, -GetDesiredSize().Y );
@@ -388,7 +397,7 @@ FSlateColor SCommentBubble::GetTextForegroundColor() const
 
 void SCommentBubble::OnCommentTextCommitted( const FText& NewText, ETextCommit::Type CommitInfo )
 {
-	if (CommitInfo == ETextCommit::OnCleared)
+	if (CommitInfo != ETextCommit::OnEnter)
 	{
 		// Don't respond to OnEnter, as it will be immediately followed by OnCleared anyway (due to loss of keyboard focus) and generate a second transaction
 		CachedComment = NewText.ToString();
@@ -425,14 +434,23 @@ ECheckBoxState SCommentBubble::GetToggleButtonCheck() const
 
 void SCommentBubble::OnCommentBubbleToggle( ECheckBoxState State )
 {
-	if( !IsReadOnly() )
+	const bool bNewVisibilityState = (State == ECheckBoxState::Checked);
+	if( !IsReadOnly() && bNewVisibilityState != GraphNode->bCommentBubbleVisible)
 	{
 		const FScopedTransaction Transaction( NSLOCTEXT( "CommentBubble", "BubbleVisibility", "Comment Bubble Visibility" ) );
 		GraphNode->Modify();
-		GraphNode->bCommentBubbleVisible = State == ECheckBoxState::Checked;
+		SetCommentBubbleVisibility(bNewVisibilityState);
+		OnToggledDelegate.ExecuteIfBound(GraphNode->bCommentBubbleVisible);
+	}
+}
+
+void SCommentBubble::SetCommentBubbleVisibility(bool bVisible)
+{
+	if (!IsReadOnly() && bVisible != GraphNode->bCommentBubbleVisible)
+	{
+		GraphNode->bCommentBubbleVisible = bVisible;
 		OpacityValue = 0.f;
 		UpdateBubble();
-		OnToggledDelegate.ExecuteIfBound(GraphNode->bCommentBubbleVisible);
 	}
 }
 
diff --git a/Engine/Source/Editor/GraphEditor/Private/SGraphNodeKnot.cpp b/Engine/Source/Editor/GraphEditor/Private/SGraphNodeKnot.cpp
--- a/Engine/Source/Editor/GraphEditor/Private/SGraphNodeKnot.cpp
+++ b/Engine/Source/Editor/GraphEditor/Private/SGraphNodeKnot.cpp
@@ -337,14 +337,13 @@ void SGraphNodeKnot::UpdateGraphNode()
 	.GraphNode(GraphNode)
 	.Text(this, &SGraphNode::GetNodeComment)
 	.OnTextCommitted(this, &SGraphNode::OnCommentTextCommitted)
-	.ToggleButtonCheck(this, &SGraphNodeKnot::GetBubbleCheckState)
 	.EnableTitleBarBubble(true)
 	.EnableBubbleCtrls(true)
 	.AllowPinning(true)
 	.ColorAndOpacity(CommentColor)
 	.GraphLOD(this, &SGraphNode::GetCurrentLOD)
-	.IsGraphNodeHovered(this, &SGraphNodeKnot::IsKnotHovered);
-	bHoveredCommentVisibility = false;
+	.IsGraphNodeHovered(this, &SGraphNode::IsHovered)
+	.OnToggled(this, &SGraphNode::OnCommentBubbleToggled);
 
 	GetOrAddSlot(ENodeZone::TopCenter)
 	.SlotOffset(TAttribute<FVector2D>(this, &SGraphNodeKnot::GetCommentOffset))
@@ -401,7 +400,7 @@ void SGraphNodeKnot::AddPin(const TSharedRef<SGraphPin>& PinToAdd)
 
 FVector2D SGraphNodeKnot::GetCommentOffset() const
 {
-	const bool bBubbleVisible = GraphNode->bCommentBubbleVisible || bHoveredCommentVisibility;
+	const bool bBubbleVisible = GraphNode->bCommentBubbleVisible || bAlwaysShowCommentBubble;
 	const float ZoomAmount = GraphNode->bCommentBubblePinned && OwnerGraphPanelPtr.IsValid() ? OwnerGraphPanelPtr.Pin()->GetZoomAmount() : 1.f;
 	const float NodeWidthOffset = bBubbleVisible ?	SKnotNodeDefinitions::KnotCenterBubbleAdjust * ZoomAmount :
 													SKnotNodeDefinitions::KnotCenterButtonAdjust * ZoomAmount;
@@ -409,36 +408,38 @@ FVector2D SGraphNodeKnot::GetCommentOffset() const
 	return FVector2D(NodeWidthOffset - CommentBubble->GetArrowCenterOffset(), -CommentBubble->GetDesiredSize().Y);
 }
 
-void SGraphNodeKnot::OnMouseEnter( const FGeometry& MyGeometry, const FPointerEvent& MouseEvent )
+void SGraphNodeKnot::OnCommentBubbleToggled(bool bInCommentBubbleVisible)
 {
-	SGraphNode::OnMouseEnter(MyGeometry, MouseEvent);
-	if (!GraphNode->bCommentBubbleVisible && !GraphNode->NodeComment.IsEmpty())
-	{
-		bHoveredCommentVisibility = true;
-		// Create the bubble widget while hovered
-		CommentBubble->OnCommentBubbleToggle(ECheckBoxState::Checked);
-	}
+	SGraphNode::OnCommentBubbleToggled(bInCommentBubbleVisible);
+	bAlwaysShowCommentBubble = bInCommentBubbleVisible;
 }
 
-void SGraphNodeKnot::OnMouseLeave(const FPointerEvent& MouseEvent)
+void SGraphNodeKnot::OnCommentTextCommitted(const FText& NewComment, ETextCommit::Type CommitInfo)
 {
-	SGraphNode::OnMouseLeave(MouseEvent);
-	if (bHoveredCommentVisibility)
+	SGraphNode::OnCommentTextCommitted(NewComment, CommitInfo);
+	if (!bAlwaysShowCommentBubble && !CommentBubble->TextBlockHasKeyboardFocus() && !CommentBubble->IsHovered())
 	{
-		bHoveredCommentVisibility = false;
-		// Destroy the comment is visibility hasn't changed
-		CommentBubble->OnCommentBubbleToggle(ECheckBoxState::Unchecked);
+		// Hide the comment bubble if visibility hasn't changed
+		CommentBubble->SetCommentBubbleVisibility(/*bVisible =*/false);
 	}
 }
 
-bool SGraphNodeKnot::IsKnotHovered() const
+void SGraphNodeKnot::OnMouseEnter(const FGeometry& MyGeometry, const FPointerEvent& MouseEvent)
 {
-	const bool bIsVisible = bHoveredCommentVisibility || GraphNode->bCommentBubbleVisible;
-	return bIsVisible ? false : IsHovered();
+	SGraphNode::OnMouseEnter(MyGeometry, MouseEvent);
+	if (!GraphNode->bCommentBubbleVisible && !GraphNode->NodeComment.IsEmpty())
+	{
+		// Show the bubble widget while hovered
+		CommentBubble->SetCommentBubbleVisibility(/*bVisible =*/true);
+	}
 }
 
-ECheckBoxState SGraphNodeKnot::GetBubbleCheckState() const
+void SGraphNodeKnot::OnMouseLeave(const FPointerEvent& MouseEvent)
 {
-	const bool bIsChecked = GraphNode->bCommentBubbleVisible && !bHoveredCommentVisibility;
-	return bIsChecked ? ECheckBoxState::Checked : ECheckBoxState::Unchecked;
-}
+	SGraphNode::OnMouseLeave(MouseEvent);
+	if (!bAlwaysShowCommentBubble && !CommentBubble->TextBlockHasKeyboardFocus())
+	{
+		// Hide the comment bubble if visibility hasn't changed;
+		CommentBubble->SetCommentBubbleVisibility(/*bVisible =*/false);
+	}
+}
\ No newline at end of file
diff --git a/Engine/Source/Editor/GraphEditor/Private/SGraphNodeKnot.h b/Engine/Source/Editor/GraphEditor/Private/SGraphNodeKnot.h
--- a/Engine/Source/Editor/GraphEditor/Private/SGraphNodeKnot.h
+++ b/Engine/Source/Editor/GraphEditor/Private/SGraphNodeKnot.h
@@ -29,21 +29,19 @@ public:
 	// End of SWidget interface
 
 protected:
-
-	/** Returns a custom hover state for the comment bubble */
-	bool IsKnotHovered() const;
-
 	/** Returns Offset to center comment on the node's only pin */
 	FVector2D GetCommentOffset() const;
+protected:
 
-	/** Called to override the comment bubble's current visibility */
-	ECheckBoxState GetBubbleCheckState() const;
+	/** Toggles the hovered visibility state */
+	virtual void OnCommentBubbleToggled(bool bInCommentBubbleVisible) override;
 
-protected:
+	/** If bHoveredCommentVisibility is true, hides the comment bubble after a change is committed */
+	virtual void OnCommentTextCommitted(const FText& NewComment, ETextCommit::Type CommitInfo) override;
+
+	/** The hovered visibility state. If false, comment bubble will only appear on hover. */
+	bool bAlwaysShowCommentBubble;
 
-	/** The hovered visibility state */
-	bool bHoveredCommentVisibility;
 	/** SharedPtr to comment bubble */
 	TSharedPtr<SCommentBubble> CommentBubble;
-
 };
diff --git a/Engine/Source/Editor/GraphEditor/Private/SGraphPanel.cpp b/Engine/Source/Editor/GraphEditor/Private/SGraphPanel.cpp
--- a/Engine/Source/Editor/GraphEditor/Private/SGraphPanel.cpp
+++ b/Engine/Source/Editor/GraphEditor/Private/SGraphPanel.cpp
@@ -1599,14 +1599,15 @@ void SGraphPanel::OnGraphChanged(const FEdGraphEditAction& EditAction)
 		// the widget immediately it won't be able to create its pins). There are lots of other examples, 
 		// and I can't be sure that I've found them all.... 
 		
-		// Minor note, the ugly little lambdas are just to deal with the  time values and return values 
+		// Minor note, the ugly little lambdas are just to deal with the time values and return values 
 		// that the timer system requires (and we don't leverage):
 		if (bWasRemoveAction)
 		{
 			const auto RemoveNodeDelegateWrapper = [](double, float, SGraphPanel* Parent, TWeakObjectPtr<UEdGraphNode> NodePtr) -> EActiveTimerReturnType
 			{
-				if (UEdGraphNode* Node = NodePtr.Get())
+				if (NodePtr.IsValid())
 				{
+					UEdGraphNode* Node = NodePtr.Get();
 					Parent->RemoveNode(Node);
 				}
 				return EActiveTimerReturnType::Stop;
@@ -1622,8 +1623,9 @@ void SGraphPanel::OnGraphChanged(const FEdGraphEditAction& EditAction)
 		{
 			const auto AddNodeDelegateWrapper = [](double, float, SGraphPanel* Parent, TWeakObjectPtr<UEdGraphNode> NodePtr, bool bForceUserAdded) -> EActiveTimerReturnType
 			{
-				if (UEdGraphNode* Node = NodePtr.Get())
+				if (NodePtr.IsValid())
 				{
+					UEdGraphNode* Node = NodePtr.Get();
 					Parent->RemoveNode(Node);
 					Parent->AddNode(Node, bForceUserAdded ? WasUserAdded : NotUserAdded);
 				}
@@ -1643,8 +1645,9 @@ void SGraphPanel::OnGraphChanged(const FEdGraphEditAction& EditAction)
 				Parent->DeferredSelectionTargetObjects.Empty();
 				for (TWeakObjectPtr<UEdGraphNode>& NodePtr : NodePtrs)
 				{
-					if (UEdGraphNode* Node = NodePtr.Get())
+					if (NodePtr.IsValid())
 					{
+						UEdGraphNode* Node = NodePtr.Get();
 						Parent->DeferredSelectionTargetObjects.Add(Node);
 					}
 				}
diff --git a/Engine/Source/Editor/GraphEditor/Private/SGraphPanel.h b/Engine/Source/Editor/GraphEditor/Private/SGraphPanel.h
--- a/Engine/Source/Editor/GraphEditor/Private/SGraphPanel.h
+++ b/Engine/Source/Editor/GraphEditor/Private/SGraphPanel.h
@@ -202,7 +202,7 @@ protected:
 	SGraphEditor::EPinVisibility PinVisibility;
 
 	/** List of pins currently being hovered over */
-	TSet< TWeakObjectPtr<UEdGraphPin> > CurrentHoveredPins;
+	TSet< FEdGraphPinReference > CurrentHoveredPins;
 
 	/** Time since the last mouse enter/exit on a pin */
 	double TimeWhenMouseEnteredPin;
diff --git a/Engine/Source/Editor/GraphEditor/Private/SGraphPin.cpp b/Engine/Source/Editor/GraphEditor/Private/SGraphPin.cpp
--- a/Engine/Source/Editor/GraphEditor/Private/SGraphPin.cpp
+++ b/Engine/Source/Editor/GraphEditor/Private/SGraphPin.cpp
@@ -577,7 +577,7 @@ void SGraphPin::OnMouseLeave( const FPointerEvent& MouseEvent )
 {	
 	TSharedPtr<SGraphPanel> Panel = OwnerNodePtr.Pin()->GetOwnerPanel();
 
-	for (TWeakObjectPtr<UEdGraphPin> WeakPin : HoverPinSet)
+	for (const FEdGraphPinReference& WeakPin : HoverPinSet)
 	{
 		if (UEdGraphPin* PinInNet = WeakPin.Get())
 		{
diff --git a/Engine/Source/Editor/GraphEditor/Public/ConnectionDrawingPolicy.h b/Engine/Source/Editor/GraphEditor/Public/ConnectionDrawingPolicy.h
--- a/Engine/Source/Editor/GraphEditor/Public/ConnectionDrawingPolicy.h
+++ b/Engine/Source/Editor/GraphEditor/Public/ConnectionDrawingPolicy.h
@@ -84,7 +84,7 @@ protected:
 	const FSlateRect& ClippingRect;
 	FSlateWindowElementList& DrawElementsList;
 	TMap< UEdGraphPin*, TSharedRef<SGraphPin> > PinToPinWidgetMap;
-	TSet< UEdGraphPin* > HoveredPins;
+	TSet< FEdGraphPinReference > HoveredPins;
 	TMap<TSharedRef<SWidget>, FArrangedWidget>* PinGeometries;
 	double LastHoverTimeEvent;
 	FVector2D LocalMousePosition;
@@ -94,7 +94,7 @@ public:
 	FConnectionDrawingPolicy(int32 InBackLayerID, int32 InFrontLayerID, float InZoomFactor, const FSlateRect& InClippingRect, FSlateWindowElementList& InDrawElements);
 
 	// Update the drawing policy with the set of hovered pins (which can be empty)
-	void SetHoveredPins(const TSet< TWeakObjectPtr<UEdGraphPin> >& InHoveredPins, const TArray< TSharedPtr<SGraphPin> >& OverridePins, double HoverTime);
+	void SetHoveredPins(const TSet< FEdGraphPinReference >& InHoveredPins, const TArray< TSharedPtr<SGraphPin> >& OverridePins, double HoverTime);
 
 	void SetMousePosition(const FVector2D& InMousePos);
 
diff --git a/Engine/Source/Editor/GraphEditor/Public/GraphEditorDragDropAction.h b/Engine/Source/Editor/GraphEditor/Public/GraphEditorDragDropAction.h
--- a/Engine/Source/Editor/GraphEditor/Public/GraphEditorDragDropAction.h
+++ b/Engine/Source/Editor/GraphEditor/Public/GraphEditorDragDropAction.h
@@ -47,7 +47,7 @@ private:
 	EVisibility GetErrorIconVisible() const;
 
 	// The pin that the drag action is currently hovering over
-	TWeakObjectPtr<class UEdGraphPin> HoveredPin;
+	FEdGraphPinReference HoveredPin;
 
 	// The node that the drag action is currently hovering over
 	TSharedPtr<class SGraphNode> HoveredNode;
diff --git a/Engine/Source/Editor/GraphEditor/Public/SCommentBubble.h b/Engine/Source/Editor/GraphEditor/Public/SCommentBubble.h
--- a/Engine/Source/Editor/GraphEditor/Public/SCommentBubble.h
+++ b/Engine/Source/Editor/GraphEditor/Public/SCommentBubble.h
@@ -89,12 +89,18 @@ class GRAPHEDITOR_API SCommentBubble : public SCompoundWidget
 	/** Called to display/hide the comment bubble */
 	void OnCommentBubbleToggle( ECheckBoxState State );
 
+	/** Directly sets the bubble's visibility, without assuming it was from a user action (no undo transaction, or OnToggled callback)*/
+	void SetCommentBubbleVisibility(bool bVisible);
+
 	/** Called when a node's comment bubble pinned state is changed */
 	void OnPinStateToggle( ECheckBoxState State ) const;
 
 	/** Called to update the bubble widget layout */
 	void UpdateBubble();
 
+	/** Returns if the text block currently has keyboard focus*/
+	bool TextBlockHasKeyboardFocus() const;
+
 protected:
 
 	/** Returns the current scale button tooltip */
diff --git a/Engine/Source/Editor/GraphEditor/Public/SGraphNode.h b/Engine/Source/Editor/GraphEditor/Public/SGraphNode.h
--- a/Engine/Source/Editor/GraphEditor/Public/SGraphNode.h
+++ b/Engine/Source/Editor/GraphEditor/Public/SGraphNode.h
@@ -130,7 +130,7 @@ public:
 	virtual FSlateColor GetNodeIndicatorOverlayColor() const { return FLinearColor::Transparent; }
 
 	/** @return visibility state of the indicator overlay */
-	virtual EVisibility GetNodeIndicatorOverlayVisibility() const { return EVisibility::Hidden; }
+	virtual EVisibility GetNodeIndicatorOverlayVisibility() const { return EVisibility::Collapsed; }
 
 	/** @return the tooltip to display when over the node */
 	FText GetNodeTooltip() const;
@@ -190,9 +190,9 @@ public:
 	/** Set event when the user generates a warning tooltip because a connection was invalid */
 	void SetDisallowedPinConnectionEvent(SGraphEditor::FOnDisallowedPinConnection InOnDisallowedPinConnection);
 	/** called to replace this nodes comment text */
-	void OnCommentTextCommitted(const FText& NewComment, ETextCommit::Type CommitInfo);
+	virtual void OnCommentTextCommitted(const FText& NewComment, ETextCommit::Type CommitInfo);
 	/** called when the node's comment bubble is toggled */
-	void OnCommentBubbleToggled(bool bInCommentBubbleVisible);
+	virtual void OnCommentBubbleToggled(bool bInCommentBubbleVisible);
 	/** returns true if a rename is pending on this node */
 	bool IsRenamePending() const { return bRenameIsPending; }
 
diff --git a/Engine/Source/Editor/GraphEditor/Public/SGraphPin.h b/Engine/Source/Editor/GraphEditor/Public/SGraphPin.h
--- a/Engine/Source/Editor/GraphEditor/Public/SGraphPin.h
+++ b/Engine/Source/Editor/GraphEditor/Public/SGraphPin.h
@@ -194,7 +194,7 @@ protected:
 	/** Cached offset from owning node to approximate position of culled pins */
 	FVector2D CachedNodeOffset;
 
-	TSet< TWeakObjectPtr<UEdGraphPin> > HoverPinSet;
+	TSet< FEdGraphPinReference > HoverPinSet;
 
 	/** TRUE if the pin should use the Pin's color for the text */
 	bool bUsePinColorForText;
diff --git a/Engine/Source/Editor/Kismet/Private/BlueprintDetailsCustomization.cpp b/Engine/Source/Editor/Kismet/Private/BlueprintDetailsCustomization.cpp
--- a/Engine/Source/Editor/Kismet/Private/BlueprintDetailsCustomization.cpp
+++ b/Engine/Source/Editor/Kismet/Private/BlueprintDetailsCustomization.cpp
@@ -1900,29 +1900,33 @@ TSharedPtr<FString> FBlueprintVarActionDetails::GetVariableReplicationType() con
 	uint64 PropFlags = 0;
 	UProperty* VariableProperty = CachedVariableProperty.Get();
 
-	if (VariableProperty && IsVariableInBlueprint())
+	if (VariableProperty && (IsVariableInBlueprint() || IsVariableInheritedByBlueprint()))
 	{
-		uint64 *PropFlagPtr = FBlueprintEditorUtils::GetBlueprintVariablePropertyFlags(GetPropertyOwnerBlueprint(), VariableProperty->GetFName());
-		
-		if (PropFlagPtr != NULL)
+		UBlueprint* BlueprintObj = GetPropertyOwnerBlueprint();
+		if (BlueprintObj != nullptr)
 		{
-			PropFlags = *PropFlagPtr;
-			bool IsReplicated = (PropFlags & CPF_Net) > 0;
-			bool bHasRepNotify = FBlueprintEditorUtils::GetBlueprintVariableRepNotifyFunc(GetPropertyOwnerBlueprint(), VariableProperty->GetFName()) != NAME_None;
-			if (bHasRepNotify)
+			uint64 *PropFlagPtr = FBlueprintEditorUtils::GetBlueprintVariablePropertyFlags(BlueprintObj, VariableProperty->GetFName());
+
+			if (PropFlagPtr != NULL)
 			{
-				// Verify they actually have a valid rep notify function still
-				UClass* GenClass = GetPropertyOwnerBlueprint()->SkeletonGeneratedClass;
-				UFunction* OnRepFunc = GenClass->FindFunctionByName(FBlueprintEditorUtils::GetBlueprintVariableRepNotifyFunc(GetPropertyOwnerBlueprint(), VariableProperty->GetFName()));
-				if( OnRepFunc == NULL || OnRepFunc->NumParms != 0 || OnRepFunc->GetReturnProperty() != NULL )
+				PropFlags = *PropFlagPtr;
+				bool IsReplicated = (PropFlags & CPF_Net) > 0;
+				bool bHasRepNotify = FBlueprintEditorUtils::GetBlueprintVariableRepNotifyFunc(BlueprintObj, VariableProperty->GetFName()) != NAME_None;
+				if (bHasRepNotify)
 				{
-					bHasRepNotify = false;
-					ReplicationOnRepFuncChanged(FName(NAME_None).ToString());	
+					// Verify they actually have a valid rep notify function still
+					UClass* GenClass = GetPropertyOwnerBlueprint()->SkeletonGeneratedClass;
+					UFunction* OnRepFunc = GenClass->FindFunctionByName(FBlueprintEditorUtils::GetBlueprintVariableRepNotifyFunc(BlueprintObj, VariableProperty->GetFName()));
+					if (OnRepFunc == NULL || OnRepFunc->NumParms != 0 || OnRepFunc->GetReturnProperty() != NULL)
+					{
+						bHasRepNotify = false;
+						ReplicationOnRepFuncChanged(FName(NAME_None).ToString());
+					}
 				}
-			}
 
-			VariableReplication = !IsReplicated ? EVariableReplication::None : 
-				bHasRepNotify ? EVariableReplication::RepNotify : EVariableReplication::Replicated;
+				VariableReplication = !IsReplicated ? EVariableReplication::None :
+					bHasRepNotify ? EVariableReplication::RepNotify : EVariableReplication::Replicated;
+			}
 		}
 	}
 
diff --git a/Engine/Source/Editor/Kismet/Private/BlueprintEditor.cpp b/Engine/Source/Editor/Kismet/Private/BlueprintEditor.cpp
--- a/Engine/Source/Editor/Kismet/Private/BlueprintEditor.cpp
+++ b/Engine/Source/Editor/Kismet/Private/BlueprintEditor.cpp
@@ -2297,7 +2297,7 @@ void FBlueprintEditor::CreateSCSEditors()
 
 void FBlueprintEditor::OnLogTokenClicked(const TSharedRef<IMessageToken>& Token)
 {
-	if( Token->GetType() == EMessageToken::Object )
+	if( Token->GetType() == EMessageToken::Object)
 	{
 		const TSharedRef<FUObjectToken> UObjectToken = StaticCastSharedRef<FUObjectToken>(Token);
 		if(UObjectToken->GetObject().IsValid())
@@ -2305,6 +2305,20 @@ void FBlueprintEditor::OnLogTokenClicked(const TSharedRef<IMessageToken>& Token)
 			JumpToHyperlink(UObjectToken->GetObject().Get());
 		}
 	}
+	else if (Token->GetType() == EMessageToken::EdGraph)
+	{
+		const TSharedRef<FEdGraphToken> EdGraphToken = StaticCastSharedRef<FEdGraphToken>(Token);
+		const UEdGraphPin* PinBeingReferenced = EdGraphToken->GetPin();
+		const UObject* ObjectBeingReferenced = EdGraphToken->GetGraphObject();
+		if (PinBeingReferenced)
+		{
+			JumpToPin(PinBeingReferenced);
+		}
+		else if(ObjectBeingReferenced)
+		{
+			JumpToHyperlink(ObjectBeingReferenced);
+		}
+	}
 }
 
 /** Create Default Commands **/
@@ -3348,7 +3362,7 @@ bool FBlueprintEditor::HasAnyDisabledBreakpoints() const
 
 bool FBlueprintEditor::HasAnyWatches() const
 {
-	return GetBlueprintObj() && GetBlueprintObj()->PinWatches.Num() > 0;
+	return GetBlueprintObj() && GetBlueprintObj()->WatchedPins.Num() > 0;
 }
 
 // Jumps to a hyperlinked node, pin, or graph, if it belongs to this blueprint
@@ -3367,10 +3381,6 @@ void FBlueprintEditor::JumpToHyperlink(const UObject* ObjectReference, bool bReq
 			JumpToNode(Node, false);
 		}
 	}
-	else if (const UEdGraphPin* Pin = Cast<const UEdGraphPin>(ObjectReference))
-	{
-		JumpToPin(Pin);
-	}
 	else if (const UEdGraph* Graph = Cast<const UEdGraph>(ObjectReference))
 	{
 		// Navigating into things should re-use the current tab when it makes sense
@@ -3448,6 +3458,19 @@ void FBlueprintEditor::JumpToHyperlink(const UObject* ObjectReference, bool bReq
 	}
 }
 
+void FBlueprintEditor::JumpToPin(const UEdGraphPin* Pin)
+{
+	if (!Pin->IsPendingKill())
+	{
+		// Open a graph editor and jump to the pin
+		TSharedPtr<SGraphEditor> GraphEditor = OpenGraphAndBringToFront(Pin->GetOwningNode()->GetGraph());
+		if (GraphEditor.IsValid())
+		{
+			GraphEditor->JumpToPin(Pin);
+		}
+	}
+}
+
 void FBlueprintEditor::AddReferencedObjects( FReferenceCollector& Collector )
 {
 	for (int32 i = 0; i < GetEditingObjects().Num(); ++i)
@@ -3523,19 +3546,6 @@ void FBlueprintEditor::JumpToNode(const UEdGraphNode* Node, bool bRequestRename)
 	}
 }
 
-void FBlueprintEditor::JumpToPin(const UEdGraphPin* Pin)
-{
-	if( !Pin->IsPendingKill() )
-	{
-		// Open a graph editor and jump to the pin
-		TSharedPtr<SGraphEditor> GraphEditor = OpenGraphAndBringToFront(Pin->GetOwningNode()->GetGraph());
-		if (GraphEditor.IsValid())
-		{
-			GraphEditor->JumpToPin(Pin);
-		}
-	}
-}
-
 UBlueprint* FBlueprintEditor::GetBlueprintObj() const
 {
 	return GetEditingObjects().Num() == 1 ? Cast<UBlueprint>(GetEditingObjects()[0]) : NULL;
@@ -7316,7 +7326,7 @@ void FBlueprintEditor::NotifyPostChange(const FPropertyChangedEvent& PropertyCha
 	if( IsEditingSingleBlueprint() && PropertyChangedEvent.ChangeType != EPropertyChangeType::Interactive )
 	{
 		UBlueprint* Blueprint = GetBlueprintObj();
-		FBlueprintEditorUtils::MarkBlueprintAsModified(Blueprint);
+		FBlueprintEditorUtils::MarkBlueprintAsModified(Blueprint, PropertyChangedEvent);
 
 		// Call PostEditChange() on any Actors that might be based on this Blueprint
 		FBlueprintEditorUtils::PostEditChangeBlueprintActors(Blueprint);
diff --git a/Engine/Source/Editor/Kismet/Private/BlueprintEditorModule.cpp b/Engine/Source/Editor/Kismet/Private/BlueprintEditorModule.cpp
--- a/Engine/Source/Editor/Kismet/Private/BlueprintEditorModule.cpp
+++ b/Engine/Source/Editor/Kismet/Private/BlueprintEditorModule.cpp
@@ -83,9 +83,9 @@ void FixSubObjectReferencesPostUndoRedo(UObject* InObject)
 	GetObjectsWithOuter(InObject, SubObjects, false);
 
 	// Post undo/redo, these may have the in-correct Outer but are incorrectly referenced by the CDO's UProperties
-	TSet<UObject*> PropertySubObjectReferences;
+	TSet<FInstancedSubObjRef> PropertySubObjectReferences;
 	UClass* ObjectClass = InObject->GetClass();
-	FFindInstancedReferenceSubobjectHelper::Get(ObjectClass, reinterpret_cast<uint8*>(InObject), PropertySubObjectReferences);
+	FFindInstancedReferenceSubobjectHelper::GetInstancedSubObjects(InObject, PropertySubObjectReferences);
 
 	TMap<UObject*, UObject*> OldToNewInstanceMap;
 	for (UObject* PropertySubObject : PropertySubObjectReferences)
@@ -224,7 +224,6 @@ void FBlueprintEditorModule::StartupModule()
 	if (ISettingsModule* SettingsModule = FModuleManager::GetModulePtr<ISettingsModule>("Settings"))
 	{
 	}
-
 }
 
 void FBlueprintEditorModule::ShutdownModule()
@@ -261,6 +260,8 @@ void FBlueprintEditorModule::ShutdownModule()
 
 	// Unregister internal SCS editor customizations
 	UnregisterSCSEditorCustomization("InstancedStaticMeshComponent");
+
+	UEdGraphPin::ShutdownVerification();
 }
 
 
diff --git a/Engine/Source/Editor/Kismet/Private/Debugging/SKismetDebuggingView.cpp b/Engine/Source/Editor/Kismet/Private/Debugging/SKismetDebuggingView.cpp
--- a/Engine/Source/Editor/Kismet/Private/Debugging/SKismetDebuggingView.cpp
+++ b/Engine/Source/Editor/Kismet/Private/Debugging/SKismetDebuggingView.cpp
@@ -257,7 +257,7 @@ struct FWatchLineItem : public FDebugLineItem
 {
 protected:
 	TWeakObjectPtr< UObject > ParentObjectRef;
-	TWeakObjectPtr< UObject > ObjectRef;
+	FEdGraphPinReference ObjectRef;
 public:
 	FWatchLineItem(UEdGraphPin* PinToWatch, UObject* ParentObject)
 		: FDebugLineItem(DLT_Watch)
@@ -275,12 +275,12 @@ public:
 
 	virtual FDebugLineItem* Duplicate() const override
 	{
-		return new FWatchLineItem(Cast<UEdGraphPin>(ObjectRef.Get()), ParentObjectRef.Get());
+		return new FWatchLineItem(ObjectRef.Get(), ParentObjectRef.Get());
 	}	
 
 	virtual void MakeMenu(class FMenuBuilder& MenuBuilder) override
 	{
-		if (UEdGraphPin* WatchedPin = Cast<UEdGraphPin>(ObjectRef.Get()))
+		if (UEdGraphPin* WatchedPin = ObjectRef.Get())
 		{
 			FUIAction ClearThisWatch(
 				FExecuteAction::CreateStatic( &FDebuggingActionCallbacks::ClearWatch, WatchedPin )
@@ -303,7 +303,7 @@ protected:
 
 FText FWatchLineItem::GetDisplayName() const
 {
-	if (UEdGraphPin* PinToWatch = Cast<UEdGraphPin>(ObjectRef.Get()))
+	if (UEdGraphPin* PinToWatch = ObjectRef.Get())
 	{
 		if (UBlueprint* Blueprint = GetBlueprintForObject(ParentObjectRef.Get()))
 		{
@@ -325,7 +325,7 @@ FText FWatchLineItem::GetDisplayName() const
 
 FText FWatchLineItem::GetDescription() const
 {
-	if (UEdGraphPin* PinToWatch = Cast<UEdGraphPin>(ObjectRef.Get()))
+	if (UEdGraphPin* PinToWatch = ObjectRef.Get())
 	{
 		// Try to determine the blueprint that generated the watch
 		UBlueprint* ParentBlueprint = GetBlueprintForObject(ParentObjectRef.Get());
@@ -382,9 +382,9 @@ TSharedRef<SWidget> FWatchLineItem::GenerateNameWidget()
 
 void FWatchLineItem::OnNavigateToWatchLocation( )
 {
-	if (UObject* ObjectToFocus = ObjectRef.Get())
+	if (UEdGraphPin* ObjectToFocus = ObjectRef.Get())
 	{
-		FKismetEditorUtilities::BringKismetToFocusAttentionOnObject(ObjectToFocus);
+		FKismetEditorUtilities::BringKismetToFocusAttentionOnPin(ObjectToFocus);
 	}
 }
 
@@ -603,9 +603,9 @@ public:
 			if ((ParentBP != NULL) && (ParentBP == ParentObject))
 			{
 				// Create children for each watch
-				for (int32 WatchIndex = 0; WatchIndex < ParentBP->PinWatches.Num(); ++WatchIndex)
+				for (int32 WatchIndex = 0; WatchIndex < ParentBP->WatchedPins.Num(); ++WatchIndex)
 				{
-					UEdGraphPin* WatchedPin = ParentBP->PinWatches[WatchIndex];
+					UEdGraphPin* WatchedPin = ParentBP->WatchedPins[WatchIndex].Get();
 
 					EnsureChildIsAdded(ChildrenMirrors, OutChildren, FWatchLineItem(WatchedPin, ParentObject));
 				}
@@ -629,9 +629,9 @@ public:
 				if (ParentBP != NULL)
 				{
 					// Create children for each watch
-					for (int32 WatchIndex = 0; WatchIndex < ParentBP->PinWatches.Num(); ++WatchIndex)
+					for (int32 WatchIndex = 0; WatchIndex < ParentBP->WatchedPins.Num(); ++WatchIndex)
 					{
-						UEdGraphPin* WatchedPin = ParentBP->PinWatches[WatchIndex];
+						UEdGraphPin* WatchedPin = ParentBP->WatchedPins[WatchIndex].Get();
 
 						EnsureChildIsAdded(ChildrenMirrors, OutChildren, FWatchLineItem(WatchedPin, ParentObject));
 					}
@@ -703,7 +703,7 @@ protected:
 	{
 		if (UBlueprint* BP = Cast<UBlueprint>(ObjectRef.Get()))
 		{
-			if (BP->PinWatches.Num() > 0)
+			if (BP->WatchedPins.Num() > 0)
 			{
 				FUIAction ClearAllWatches(
 					FExecuteAction::CreateStatic( &FDebuggingActionCallbacks::ClearWatches, BP )
diff --git a/Engine/Source/Editor/Kismet/Private/FiBSearchInstance.cpp b/Engine/Source/Editor/Kismet/Private/FiBSearchInstance.cpp
--- a/Engine/Source/Editor/Kismet/Private/FiBSearchInstance.cpp
+++ b/Engine/Source/Editor/Kismet/Private/FiBSearchInstance.cpp
@@ -225,11 +225,11 @@ void FFiBSearchInstance::BuildFunctionTargets(TSharedPtr<FImaginaryFiBData> InRo
 {
 	for (TSharedPtr<FImaginaryFiBData> ChildData : InRootData->GetAllParsedChildData())
 	{
-		if (ChildData->IsCompatibleWithFilter(InSearchQueryFilter))
+		if (!ChildData->IsCategory() && ChildData->IsCompatibleWithFilter(InSearchQueryFilter))
 		{
 			OutTargetPendingSearchables.Add(ChildData);
 		}
-		else if (ChildData->CanCallFilter(InSearchQueryFilter))
+		else if (ChildData->IsCategory() || ChildData->CanCallFilter(InSearchQueryFilter))
 		{
 			ChildData->ParseAllChildData();
 			BuildFunctionTargets(ChildData, InSearchQueryFilter, OutTargetPendingSearchables);
diff --git a/Engine/Source/Editor/Kismet/Private/FindInBlueprints.cpp b/Engine/Source/Editor/Kismet/Private/FindInBlueprints.cpp
--- a/Engine/Source/Editor/Kismet/Private/FindInBlueprints.cpp
+++ b/Engine/Source/Editor/Kismet/Private/FindInBlueprints.cpp
@@ -13,6 +13,8 @@
 #include "GenericCommands.h"
 #include "ImaginaryBlueprintData.h"
 #include "FiBSearchInstance.h"
+#include "BlueprintEditorTabs.h"
+#include "ToolkitManager.h"
 
 #define LOCTEXT_NAMESPACE "FindInBlueprints"
 
@@ -392,6 +394,50 @@ FFindInBlueprintsProperty::FFindInBlueprintsProperty(const FText& InValue, TShar
 {
 }
 
+FReply FFindInBlueprintsProperty::OnClick()
+{
+	if (bIsSCSComponent)
+	{
+		UBlueprint* Blueprint = GetParentBlueprint();
+		if (Blueprint)
+		{
+			TSharedPtr<IBlueprintEditor> BlueprintEditor = FKismetEditorUtilities::GetIBlueprintEditorForObject(Blueprint, false);
+
+			if (BlueprintEditor.IsValid())
+			{
+				// Open Viewport Tab
+				BlueprintEditor->FocusWindow();
+				BlueprintEditor->GetTabManager()->InvokeTab(FBlueprintEditorTabs::SCSViewportID);
+
+				// Find and Select the Component in the Viewport tab view
+				const TArray<USCS_Node*>& Nodes = Blueprint->SimpleConstructionScript->GetAllNodes();
+				for (USCS_Node* Node : Nodes)
+				{
+					if (Node->VariableName.ToString() == DisplayText.ToString())
+					{
+						UBlueprintGeneratedClass* GeneratedClass = Cast<UBlueprintGeneratedClass>(Blueprint->GeneratedClass);
+						if (GeneratedClass)
+						{
+							UActorComponent* Component = Node->GetActualComponentTemplate(GeneratedClass);
+							if (Component)
+							{
+								BlueprintEditor->FindAndSelectSCSEditorTreeNode(Component, false);
+							}
+						}
+						break;
+					}
+				}
+			}
+		}
+	}
+	else
+	{
+		return FFindInBlueprintsResult::OnClick();
+	}
+
+	return FReply::Handled();
+}
+
 TSharedRef<SWidget> FFindInBlueprintsProperty::CreateIcon() const
 {
 	FLinearColor IconColor = FLinearColor::White;
