Mirror of 3132793

diff --git a/Engine/Source/Editor/BlueprintGraph/Private/EdGraphSchema_K2.cpp b/Engine/Source/Editor/BlueprintGraph/Private/EdGraphSchema_K2.cpp
--- a/Engine/Source/Editor/BlueprintGraph/Private/EdGraphSchema_K2.cpp
+++ b/Engine/Source/Editor/BlueprintGraph/Private/EdGraphSchema_K2.cpp
@@ -3285,7 +3285,12 @@ bool UEdGraphSchema_K2::ConvertPropertyToPinType(const UProperty* Property, /*ou
 		bool bIsWeakPtr = false;
 		bool bResult = GetPropertyCategoryInfo(MapProperty->ValueProp, TypeOut.PinValueType.TerminalCategory, TypeOut.PinValueType.TerminalSubCategory, SubCategoryObject, bIsWeakPtr);
 		TypeOut.PinValueType.TerminalSubCategoryObject = SubCategoryObject;
-		ensure(!bIsWeakPtr);
+
+		if (bIsWeakPtr)
+		{
+			return false;
+		}
+
 		if (!bResult)
 		{
 			return false;
