Removing innocuous warning that was exposed by CIS (due to a recent Dev-BP change) - no longer warning about MakeStruct nodes w/out properties (sometimes you need to create a default struct to return).

diff --git a/Engine/Source/Editor/BlueprintGraph/Private/K2Node_MakeStruct.cpp b/Engine/Source/Editor/BlueprintGraph/Private/K2Node_MakeStruct.cpp
--- a/Engine/Source/Editor/BlueprintGraph/Private/K2Node_MakeStruct.cpp
+++ b/Engine/Source/Editor/BlueprintGraph/Private/K2Node_MakeStruct.cpp
@@ -206,11 +206,6 @@ void UK2Node_MakeStruct::ValidateNodeDuringCompilation(class FCompilerResultsLog
 			}
 		}
 
-		if (!bHasAnyBlueprintVisibleProperty)
-		{
-			MessageLog.Warning(*LOCTEXT("StructHasNoBPVisibleProperties_Warning", "@@ has no property tagged as BlueprintReadWrite. The node will be removed in a future release.").ToString(), this);
-		}
-
 		if (!bMadeAfterOverridePinRemoval)
 		{
 			MessageLog.Note(*NSLOCTEXT("K2Node", "OverridePinRemoval_SetFieldsInStruct", "Override pins have been removed from @@ in @@, it functions the same as it did but some functionality may be deprecated! This note will go away after you resave the asset!").ToString(), this, GetBlueprint());
