Fixing up a few remaining non-unity issues that I missed from the Dev-BP integrate.

diff --git a/Engine/Source/Developer/BlueprintCompilerCppBackend/Private/BlueprintCompilerCppBackendBase.h b/Engine/Source/Developer/BlueprintCompilerCppBackend/Private/BlueprintCompilerCppBackendBase.h
--- a/Engine/Source/Developer/BlueprintCompilerCppBackend/Private/BlueprintCompilerCppBackendBase.h
+++ b/Engine/Source/Developer/BlueprintCompilerCppBackend/Private/BlueprintCompilerCppBackendBase.h
@@ -4,6 +4,7 @@
 #include "CoreMinimal.h"
 #include "KismetCompiledFunctionContext.h"
 #include "BlueprintCompilerCppBackendInterface.h"
+#include "Engine/Blueprint.h" // for FCompilerNativizationOptions
 
 class UUserDefinedEnum;
 class UUserDefinedStruct;
diff --git a/Engine/Source/Developer/BlueprintCompilerCppBackend/Private/BlueprintCompilerCppBackendUtils.h b/Engine/Source/Developer/BlueprintCompilerCppBackend/Private/BlueprintCompilerCppBackendUtils.h
--- a/Engine/Source/Developer/BlueprintCompilerCppBackend/Private/BlueprintCompilerCppBackendUtils.h
+++ b/Engine/Source/Developer/BlueprintCompilerCppBackend/Private/BlueprintCompilerCppBackendUtils.h
@@ -7,6 +7,7 @@
 #include "UObject/UnrealType.h"
 #include "IBlueprintCompilerCppBackendModule.h"
 #include "BlueprintCompilerCppBackendGatherDependencies.h"
+#include "Engine/Blueprint.h" // for FCompilerNativizationOptions
 
 class USCS_Node;
 class UUserDefinedEnum;
diff --git a/Engine/Source/Developer/BlueprintNativeCodeGen/Private/BlueprintNativeCodeGenUtils.h b/Engine/Source/Developer/BlueprintNativeCodeGen/Private/BlueprintNativeCodeGenUtils.h
--- a/Engine/Source/Developer/BlueprintNativeCodeGen/Private/BlueprintNativeCodeGenUtils.h
+++ b/Engine/Source/Developer/BlueprintNativeCodeGen/Private/BlueprintNativeCodeGenUtils.h
@@ -5,6 +5,7 @@
 #include "CoreMinimal.h"
 #include "IBlueprintCompilerCppBackendModule.h"
 #include "Misc/FeedbackContext.h"
+#include "Engine/Blueprint.h" // for FCompilerNativizationOptions
 
 class SBuildProgressWidget;
 struct FBlueprintNativeCodeGenManifest;
diff --git a/Engine/Source/Editor/BlueprintGraph/Private/K2Node_MakeVariable.cpp b/Engine/Source/Editor/BlueprintGraph/Private/K2Node_MakeVariable.cpp
--- a/Engine/Source/Editor/BlueprintGraph/Private/K2Node_MakeVariable.cpp
+++ b/Engine/Source/Editor/BlueprintGraph/Private/K2Node_MakeVariable.cpp
@@ -11,6 +11,7 @@
 #include "KismetCompiledFunctionContext.h"
 #include "KismetCompiler.h"
 #include "KismetCompilerMisc.h"
+#include "UObject/StructOnScope.h"
 
 static const TCHAR* MakeVariableOutputPinName = TEXT("MakeVariableOutput");
 
@@ -98,7 +99,6 @@ void UK2Node_MakeVariable::SetupVariable(const FBPVariableDescription& InVariabl
 	UEdGraphPin* VariableOutputPin = CreatePin(EGPD_Output, InVariableType.VarType, MakeVariableOutputPinName);
 
 	// make input pins for every value in the default value, for map pins we'll make a pair of inputs:
-	
 	TSharedPtr<FStructOnScope> StructData = MakeShareable(new FStructOnScope(Scope));
 	FBlueprintEditorUtils::PropertyValueFromString(Property, VariableType.DefaultValue, StructData->GetStructMemory());
 
diff --git a/Engine/Source/Editor/UnrealEd/Private/Settings/SettingsClasses.cpp b/Engine/Source/Editor/UnrealEd/Private/Settings/SettingsClasses.cpp
--- a/Engine/Source/Editor/UnrealEd/Private/Settings/SettingsClasses.cpp
+++ b/Engine/Source/Editor/UnrealEd/Private/Settings/SettingsClasses.cpp
@@ -30,7 +30,7 @@
 #include "UnrealEdMisc.h"
 #include "CrashReporterSettings.h"
 #include "AutoReimport/AutoReimportUtilities.h"
-
+#include "Misc/ConfigCacheIni.h" // for FConfigCacheIni::GetString()
 #include "SourceCodeNavigation.h"
 #include "Developer/BlueprintProfiler/Public/BlueprintProfilerModule.h"
 
