Copy 3240405 to main to fix documentation warning

diff --git a/Engine/Source/Editor/BlueprintGraph/Private/EdGraphSchema_K2.cpp b/Engine/Source/Editor/BlueprintGraph/Private/EdGraphSchema_K2.cpp
--- a/Engine/Source/Editor/BlueprintGraph/Private/EdGraphSchema_K2.cpp
+++ b/Engine/Source/Editor/BlueprintGraph/Private/EdGraphSchema_K2.cpp
@@ -3958,7 +3958,7 @@ bool UEdGraphSchema_K2::DefaultValueSimpleValidation(const FEdGraphPinType& PinT
 		const UEnum* EnumPtr = Cast<const UEnum>(PinSubCategoryObject);
 		if (EnumPtr)
 		{
-			if (EnumPtr->FindEnumIndex(*NewDefaultValue) == INDEX_NONE)
+			if ( NewDefaultValue == TEXT("(INVALID)") || EnumPtr->FindEnumIndex(*NewDefaultValue) == INDEX_NONE)
 			{
 				DVSV_RETURN_MSG(FString::Printf(TEXT("'%s' is not a valid enumerant of '<%s>'"), *NewDefaultValue, *(EnumPtr->GetName())));
 			}
