Change summary:
- Backed out fixes for UE-35332 and UE-43565 from most recent Dev-Framwork integration for now.

diff --git a/Engine/Source/Editor/BlueprintGraph/Classes/K2Node_CallArrayFunction.h b/Engine/Source/Editor/BlueprintGraph/Classes/K2Node_CallArrayFunction.h
--- a/Engine/Source/Editor/BlueprintGraph/Classes/K2Node_CallArrayFunction.h
+++ b/Engine/Source/Editor/BlueprintGraph/Classes/K2Node_CallArrayFunction.h
@@ -55,14 +55,7 @@ class UK2Node_CallArrayFunction : public UK2Node_CallFunction
 	 */
 	BLUEPRINTGRAPH_API static bool IsWildcardProperty(UFunction* InArrayFunction, const UProperty* InProperty);
 
-	/** Propagates the given pin's type to the set of dynamically-typed pins. */
-	void PropagateArrayTypeInfo(const UEdGraphPin* SourcePin);
-
-protected:
-	/** Returns the set of pins tagged in UFUNCTION metadata as being dependent on the array type. */
 	void GetArrayTypeDependentPins(TArray<UEdGraphPin*>& OutPins) const;
-
-	/** Returns the full set of dynamically-typed pins (those that default to a wildcard pin type). */
-	void GetDynamicallyTypedPins(TArray<UEdGraphPin*>& OutPins) const;
+	void PropagateArrayTypeInfo(const UEdGraphPin* SourcePin);
 };
 
diff --git a/Engine/Source/Editor/BlueprintGraph/Private/K2Node_CallArrayFunction.cpp b/Engine/Source/Editor/BlueprintGraph/Private/K2Node_CallArrayFunction.cpp
--- a/Engine/Source/Editor/BlueprintGraph/Private/K2Node_CallArrayFunction.cpp
+++ b/Engine/Source/Editor/BlueprintGraph/Private/K2Node_CallArrayFunction.cpp
@@ -6,7 +6,6 @@
 #include "K2Node_GetArrayItem.h"
 #include "BlueprintsObjectVersion.h"
 #include "Kismet/KismetArrayLibrary.h" // for Array_Get()
-#include "Kismet2/BlueprintEditorUtils.h"
 
 UK2Node_CallArrayFunction::UK2Node_CallArrayFunction(const FObjectInitializer& ObjectInitializer)
 	: Super(ObjectInitializer)
@@ -64,12 +63,9 @@ void UK2Node_CallArrayFunction::NotifyPinConnectionListChanged(UEdGraphPin* Pin)
 {
 	Super::NotifyPinConnectionListChanged(Pin);
 
-	// Get the set of dynamically-typed pins.
-	TArray<UEdGraphPin*> DynamicPins;
-	GetDynamicallyTypedPins(DynamicPins);
+	TArray<UEdGraphPin*> PinsToCheck;
+	GetArrayTypeDependentPins(PinsToCheck);
 
-	// We need to also check subpins here in case one of them was the target of the connection event.
-	TArray<UEdGraphPin*> PinsToCheck = DynamicPins;
 	for (int32 Index = 0; Index < PinsToCheck.Num(); ++Index)
 	{
 		UEdGraphPin* PinToCheck = PinsToCheck[Index];
@@ -79,51 +75,53 @@ void UK2Node_CallArrayFunction::NotifyPinConnectionListChanged(UEdGraphPin* Pin)
 		}
 	}
 
-	// If the event target is one of the dynamically-typed pins, we may need to conform the node to match the linked pin's type.
+	PinsToCheck.Add(GetTargetArrayPin());
+
 	if (PinsToCheck.Contains(Pin))
 	{
 		const UEdGraphSchema_K2* Schema = GetDefault<UEdGraphSchema_K2>();
+		bool bNeedToPropagate = false;
 
-		// Determine whether or not the target pin is connected and/or has a valid pin type.
-		const bool bTargetPinHasValidType = Pin->PinType.PinCategory != Schema->PC_Wildcard;
-		const bool bTargetPinHasConnection = Pin->LinkedTo.Num() > 0;
-
-		// If the target pin has a connection, but does not yet have a valid pin type, we need to set it based on the connected pin's type.
-		if (bTargetPinHasConnection && !bTargetPinHasValidType)
+		if( Pin->LinkedTo.Num() > 0 )
 		{
-			// Find a linked pin with an authoritative pin type (it might be linked to one of the other pins, depending on which category it is).
-			const UEdGraphPin* LinkedPin = FBlueprintEditorUtils::FindLinkedPinWithAuthoritativePinType(Pin, PinsToCheck);
+			if (Pin->PinType.PinCategory == Schema->PC_Wildcard)
+			{
+				UEdGraphPin* LinkedTo = Pin->LinkedTo[0];
+				check(LinkedTo);
+				check(Pin->PinType.bIsArray == LinkedTo->PinType.bIsArray);
+
+				Pin->PinType.PinCategory = LinkedTo->PinType.PinCategory;
+				Pin->PinType.PinSubCategory = LinkedTo->PinType.PinSubCategory;
+				Pin->PinType.PinSubCategoryObject = LinkedTo->PinType.PinSubCategoryObject;
 
-			// Propagate the linked pin's type to the set of dynamically-typed pins (which will also include the target pin).
-			check(LinkedPin);
-			FBlueprintEditorUtils::PropagatePinTypeInfo(LinkedPin, DynamicPins);
+				bNeedToPropagate = true;
+			}
 		}
-		else if (!bTargetPinHasConnection && bTargetPinHasValidType)
+		else
 		{
-			// If the target pin does not have a connection, but does have a valid type, we may need to reset it back to the wildcard type.
-			bool bNeedToResetType = true;
+			bNeedToPropagate = true;
 
-			// Check to see if any of the other dynamically-typed pins still have a connection. If any of them do, then we can't reset the type.
 			for (UEdGraphPin* PinToCheck : PinsToCheck)
 			{
 				if (PinToCheck->LinkedTo.Num() > 0)
 				{
-					bNeedToResetType = false;
+					bNeedToPropagate = false;
 					break;
 				}
 			}
 
-			if (bNeedToResetType)
+			if (bNeedToPropagate)
 			{
-				// Reset the pin's type.
 				Pin->PinType.PinCategory = Schema->PC_Wildcard;
 				Pin->PinType.PinSubCategory = TEXT("");
 				Pin->PinType.PinSubCategoryObject = NULL;
-
-				// Reset all other dynamically-typed pins to match.
-				FBlueprintEditorUtils::PropagatePinTypeInfo(Pin, DynamicPins);
 			}
 		}
+
+		if (bNeedToPropagate)
+		{
+			PropagateArrayTypeInfo(Pin);
+		}
 	}
 }
 
@@ -251,13 +249,6 @@ bool UK2Node_CallArrayFunction::IsWildcardProperty(UFunction* InArrayFunction, c
 	return false;
 }
 
-void UK2Node_CallArrayFunction::PropagateArrayTypeInfo(const UEdGraphPin* SourcePin)
-{
-	TArray<UEdGraphPin*> DynamicPins;
-	GetDynamicallyTypedPins(DynamicPins);
-	FBlueprintEditorUtils::PropagatePinTypeInfo(SourcePin, DynamicPins);
-}
-
 void UK2Node_CallArrayFunction::GetArrayTypeDependentPins(TArray<UEdGraphPin*>& OutPins) const
 {
 	OutPins.Empty();
@@ -281,11 +272,49 @@ void UK2Node_CallArrayFunction::GetArrayTypeDependentPins(TArray<UEdGraphPin*>&
 	}
 }
 
-void UK2Node_CallArrayFunction::GetDynamicallyTypedPins(TArray<UEdGraphPin *>& OutPins) const
+void UK2Node_CallArrayFunction::PropagateArrayTypeInfo(const UEdGraphPin* SourcePin)
 {
-	// Include all pins that are listed as type-dependent in the UFUNCTION metadata. This must be done first, as it will clear the array.
-	GetArrayTypeDependentPins(OutPins);
+	if( SourcePin )
+	{
+		const UEdGraphSchema_K2* Schema = CastChecked<UEdGraphSchema_K2>(GetSchema());
+		const FEdGraphPinType& SourcePinType = SourcePin->PinType;
+
+		TArray<UEdGraphPin*> DependentPins;
+		GetArrayTypeDependentPins(DependentPins);
+		DependentPins.Add(GetTargetArrayPin());
+	
+		// Propagate pin type info (except for array info!) to pins with dependent types
+		for (UEdGraphPin* CurrentPin : DependentPins)
+		{
+			if (CurrentPin != SourcePin)
+			{
+				CA_SUPPRESS(6011); // warning C6011: Dereferencing NULL pointer 'CurrentPin'.
+				FEdGraphPinType& CurrentPinType = CurrentPin->PinType;
+
+				bool const bHasTypeMismatch = (CurrentPinType.PinCategory != SourcePinType.PinCategory) ||
+					(CurrentPinType.PinSubCategory != SourcePinType.PinSubCategory) ||
+					(CurrentPinType.PinSubCategoryObject != SourcePinType.PinSubCategoryObject);
+
+				if (!bHasTypeMismatch)
+				{
+					continue;
+				}
+
+				if (CurrentPin->SubPins.Num() > 0)
+				{
+					Schema->RecombinePin(CurrentPin->SubPins[0]);
+				}
 
-	// Include the target array pin.
-	OutPins.Add(GetTargetArrayPin());
-}
\ No newline at end of file
+				CurrentPinType.PinCategory          = SourcePinType.PinCategory;
+				CurrentPinType.PinSubCategory       = SourcePinType.PinSubCategory;
+				CurrentPinType.PinSubCategoryObject = SourcePinType.PinSubCategoryObject;
+
+				// Reset default values
+				if (!Schema->IsPinDefaultValid(CurrentPin, CurrentPin->DefaultValue, CurrentPin->DefaultObject, CurrentPin->DefaultTextValue).IsEmpty())
+				{
+					CurrentPin->ResetDefaultValue();
+				}
+			}
+		}
+	}
+}
diff --git a/Engine/Source/Editor/UnrealEd/Private/Kismet2/BlueprintEditorUtils.cpp b/Engine/Source/Editor/UnrealEd/Private/Kismet2/BlueprintEditorUtils.cpp
--- a/Engine/Source/Editor/UnrealEd/Private/Kismet2/BlueprintEditorUtils.cpp
+++ b/Engine/Source/Editor/UnrealEd/Private/Kismet2/BlueprintEditorUtils.cpp
@@ -5448,132 +5448,6 @@ bool FBlueprintEditorUtils::ShouldNativizeImplicitly(const UBlueprint* Blueprint
 	return false;
 }
 
-const UEdGraphPin* FBlueprintEditorUtils::FindLinkedPinWithMostDerivedPinType(const UEdGraphPin* Pin)
-{
-	const UEdGraphPin* MostDerivedLinkedTo = nullptr;
-	const UEdGraphSchema_K2* Schema = GetDefault<UEdGraphSchema_K2>();
-
-	check(Pin);
-
-	for (const UEdGraphPin* LinkedTo : Pin->LinkedTo)
-	{
-		if (MostDerivedLinkedTo == nullptr)
-		{
-			// The first link starts out as the most-derived type.
-			MostDerivedLinkedTo = LinkedTo;
-		}
-		else if (LinkedTo->PinType.PinCategory == Schema->PC_Object
-			|| LinkedTo->PinType.PinCategory == Schema->PC_Struct
-			|| LinkedTo->PinType.PinCategory == Schema->PC_Class
-			|| LinkedTo->PinType.PinCategory == Schema->PC_Interface)
-		{
-			// If the current linked pin's type (first param below) is a subclass of the current most-derived linked pin's type (second param below), then reset the most-derived linked pin to the current one.
-			if (LinkedTo->PinType.PinSubCategoryObject != MostDerivedLinkedTo->PinType.PinSubCategoryObject
-				&& Schema->ArePinTypesCompatible(LinkedTo->PinType, MostDerivedLinkedTo->PinType, nullptr, true))
-			{
-				MostDerivedLinkedTo = LinkedTo;
-			}
-		}
-		else
-		{
-			// For all other pin type categories, we don't need to iterate any further.
-			break;
-		}
-	}
-
-	return MostDerivedLinkedTo;
-}
-
-const UEdGraphPin* FBlueprintEditorUtils::FindLinkedPinWithAuthoritativePinType(const UEdGraphPin* Pin, const TArray<UEdGraphPin *>& PinsToCheck)
-{
-	const UEdGraphSchema_K2* Schema = GetDefault<UEdGraphSchema_K2>();
-
-	check(Pin);
-
-	// In case of multiple links, make sure we're using the link with the most-derived pin type.
-	const UEdGraphPin* LinkedTo = FindLinkedPinWithMostDerivedPinType(Pin);
-	if (LinkedTo)
-	{
-		// Make sure the pin type array attribute matches up, or they won't be compatible.
-		check(Pin->PinType.bIsArray == LinkedTo->PinType.bIsArray);
-
-		// For certain type categories, we check the other pins to see if any of them might be linked to a more-compatible pin type.
-		if (LinkedTo->PinType.PinCategory == Schema->PC_Object
-			|| LinkedTo->PinType.PinCategory == Schema->PC_Struct
-			|| LinkedTo->PinType.PinCategory == Schema->PC_Class
-			|| LinkedTo->PinType.PinCategory == Schema->PC_Interface)
-		{
-			for (UEdGraphPin* PinToCheck : PinsToCheck)
-			{
-				// Make sure we don't check against itself.
-				if (PinToCheck != Pin)
-				{
-					check(PinToCheck);
-
-					// If the other pin is also linked, check against that pin's most-derived linked pin type (in case of multiple links).
-					if (PinToCheck->LinkedTo.Num() > 0)
-					{
-						const UEdGraphPin* LinkedPinToCheck = FindLinkedPinWithMostDerivedPinType(PinToCheck);
-						check(LinkedPinToCheck);
-
-						// If the linked pin's type (first param below) is a subclass of the other linked pin's type (second param below), then return the other pin as the authoritative pin type.
-						if (LinkedTo->PinType.PinSubCategoryObject != LinkedPinToCheck->PinType.PinSubCategoryObject
-							&& Schema->ArePinTypesCompatible(LinkedTo->PinType, LinkedPinToCheck->PinType, nullptr, true))
-						{
-							return LinkedPinToCheck;
-						}
-					}
-				}
-			}
-		}
-	}
-
-	return LinkedTo;
-}
-
-void FBlueprintEditorUtils::PropagatePinTypeInfo(const UEdGraphPin* SourcePin, TArray<UEdGraphPin*>& TargetPins)
-{
-	if (SourcePin)
-	{
-		const UEdGraphSchema_K2* Schema = GetDefault<UEdGraphSchema_K2>();
-		const FEdGraphPinType& SourcePinType = SourcePin->PinType;
-
-		// Propagate pin type info to the set of target pins
-		for (UEdGraphPin* CurrentPin : TargetPins)
-		{
-			if (CurrentPin != SourcePin)
-			{
-				CA_SUPPRESS(6011); // warning C6011: Dereferencing NULL pointer 'CurrentPin'.
-				FEdGraphPinType& CurrentPinType = CurrentPin->PinType;
-
-				bool const bHasTypeMismatch = (CurrentPinType.PinCategory != SourcePinType.PinCategory) ||
-					(CurrentPinType.PinSubCategory != SourcePinType.PinSubCategory) ||
-					(CurrentPinType.PinSubCategoryObject != SourcePinType.PinSubCategoryObject);
-
-				if (!bHasTypeMismatch)
-				{
-					continue;
-				}
-
-				if (CurrentPin->SubPins.Num() > 0)
-				{
-					Schema->RecombinePin(CurrentPin->SubPins[0]);
-				}
-
-				CurrentPinType.PinCategory = SourcePinType.PinCategory;
-				CurrentPinType.PinSubCategory = SourcePinType.PinSubCategory;
-				CurrentPinType.PinSubCategoryObject = SourcePinType.PinSubCategoryObject;
-
-				// Reset default values
-				if (!Schema->IsPinDefaultValid(CurrentPin, CurrentPin->DefaultValue, CurrentPin->DefaultObject, CurrentPin->DefaultTextValue).IsEmpty())
-				{
-					CurrentPin->ResetDefaultValue();
-				}
-			}
-		}
-	}
-}
-
 //////////////////////////////////////////////////////////////////////////
 // Interfaces
 
diff --git a/Engine/Source/Editor/UnrealEd/Public/Kismet2/BlueprintEditorUtils.h b/Engine/Source/Editor/UnrealEd/Public/Kismet2/BlueprintEditorUtils.h
--- a/Engine/Source/Editor/UnrealEd/Public/Kismet2/BlueprintEditorUtils.h
+++ b/Engine/Source/Editor/UnrealEd/Public/Kismet2/BlueprintEditorUtils.h
@@ -1112,15 +1112,6 @@ public:
 	/** Returns whether or not the given Blueprint should be nativized implicitly, regardless of whether or not the user has explicitly enabled it */
 	static bool ShouldNativizeImplicitly(const UBlueprint* Blueprint);
 
-	/** Checks the given pin for linked pins, and returns the linked pin with the most-derived pin type. */
-	static const UEdGraphPin* FindLinkedPinWithMostDerivedPinType(const UEdGraphPin* Pin);
-
-	/** Checks the given pin against the given set of dynamically-typed pins, and returns the linked pin that has the authoritative pin type. Assumes that the given pin is linked. */
-	static const UEdGraphPin* FindLinkedPinWithAuthoritativePinType(const UEdGraphPin* Pin, const TArray<UEdGraphPin*>& PinsToCheck);
-
-	/** Propagates the source pin's type to the given set of target pins. */
-	static void PropagatePinTypeInfo(const UEdGraphPin* SourcePin, TArray<UEdGraphPin*>& TargetPins);
-
 	//////////////////////////////////////////////////////////////////////////
 	// Interface
 
