Copying //UE4/Dev-Framework to //UE4/Dev-Main (Source: //UE4/Dev-Framework @ 3431384)

diff --git a/Engine/Source/Developer/FunctionalTesting/Public/AutomationScreenshotOptions.h b/Engine/Source/Developer/FunctionalTesting/Public/AutomationScreenshotOptions.h
--- a/Engine/Source/Developer/FunctionalTesting/Public/AutomationScreenshotOptions.h
+++ b/Engine/Source/Developer/FunctionalTesting/Public/AutomationScreenshotOptions.h
@@ -140,7 +140,7 @@ public:
 	 * essentially the same.  Generally this is necessary as modern rendering techniques tend to introduce
 	 * noise constantly to hide aliasing.
 	 */
-	UPROPERTY(EditAnywhere, BlueprintReadWrite, Category="Comparison")
+	UPROPERTY(EditAnywhere, Category="Comparison")
 	FComparisonToleranceAmount ToleranceAmount;
 
 	/**
diff --git a/Engine/Source/Developer/Merge/Private/MergeUtils.cpp b/Engine/Source/Developer/Merge/Private/MergeUtils.cpp
--- a/Engine/Source/Developer/Merge/Private/MergeUtils.cpp
+++ b/Engine/Source/Developer/Merge/Private/MergeUtils.cpp
@@ -42,11 +42,8 @@ UObject const* FMergeToolUtils::LoadRevision(const FString& AssetName, const ISo
 	FString TempFileName;
 	if (DesiredRevision.Get(TempFileName))
 	{
-		// Forcibly disable compile on load in case we are loading old blueprints that might try to update/compile
-		TGuardValue<bool> DisableCompileOnLoad(GForceDisableBlueprintCompileOnLoad, true);
-
 		// Try and load that package
-		UPackage* TempPackage = LoadPackage(NULL, *TempFileName, LOAD_None);
+		UPackage* TempPackage = LoadPackage(NULL, *TempFileName, LOAD_DisableCompileOnLoad);
 		if (TempPackage != NULL)
 		{
 			// Grab the old asset from that old package
diff --git a/Engine/Source/Developer/MeshUtilities/Private/MeshUtilities.cpp b/Engine/Source/Developer/MeshUtilities/Private/MeshUtilities.cpp
--- a/Engine/Source/Developer/MeshUtilities/Private/MeshUtilities.cpp
+++ b/Engine/Source/Developer/MeshUtilities/Private/MeshUtilities.cpp
@@ -6257,7 +6257,7 @@ bool FMeshUtilities::ConstructRawMesh(
 	const FMeshBuildSettings& BuildSettings = bImportedMesh ? SourceStaticMeshModel.BuildSettings : SourceStaticMesh->SourceModels[0].BuildSettings;
 
 	// Transform raw mesh to world space
-	FTransform ComponentToWorldTransform = InMeshComponent->ComponentToWorld;
+	FTransform ComponentToWorldTransform = InMeshComponent->GetComponentTransform();
 	// Take into account build scale settings only for meshes imported from raw data
 	// meshes reconstructed from render data already have build scale applied
 	if (bImportedMesh)
@@ -6582,7 +6582,7 @@ static void ExtractPhysicsGeometry(UStaticMeshComponent* InMeshComponent, FKAggr
 	}
 
 	// Transform geometry to world space
-	FTransform CtoM = InMeshComponent->ComponentToWorld;
+	FTransform CtoM = InMeshComponent->GetComponentTransform();
 	TransformPhysicsGeometry(CtoM, OutAggGeom);
 }
 
@@ -6969,7 +6969,7 @@ void FMeshUtilities::MergeStaticMeshComponents(const TArray<UStaticMeshComponent
 		if (MeshId == 0)
 		{
 			// Mesh component pivot point
-			MergedAssetPivot = InSettings.bPivotPointAtZero ? FVector::ZeroVector : MeshComponent->ComponentToWorld.GetLocation();
+			MergedAssetPivot = InSettings.bPivotPointAtZero ? FVector::ZeroVector : MeshComponent->GetComponentTransform().GetLocation();
 			// Source mesh asset package name
 			MergedAssetPackageName = MeshComponent->GetStaticMesh()->GetOutermost()->GetName();
 		}
diff --git a/Engine/Source/Developer/SettingsEditor/Private/SettingsEditorModule.cpp b/Engine/Source/Developer/SettingsEditor/Private/SettingsEditorModule.cpp
--- a/Engine/Source/Developer/SettingsEditor/Private/SettingsEditorModule.cpp
+++ b/Engine/Source/Developer/SettingsEditor/Private/SettingsEditorModule.cpp
@@ -230,10 +230,10 @@ private:
 			if ( UDeveloperSettings* Settings = *SettingsIt )
 			{
 				// Only Add the CDO of any UDeveloperSettings objects.
-				if ( Settings->HasAnyFlags(RF_ClassDefaultObject) && !Settings->GetClass()->HasAnyCastFlag(CLASS_Deprecated) )
+				if ( Settings->HasAnyFlags(RF_ClassDefaultObject) && !Settings->GetClass()->HasAnyClassFlags(CLASS_Deprecated | CLASS_Abstract) )
 				{
 					// Ignore the setting if it's specifically the UDeveloperSettings or other abstract settings classes
-					if ( Settings->GetClass()->HasAnyClassFlags(CLASS_Abstract) || !Settings->SupportsAutoRegistration())
+					if ( Settings->GetClass()->HasAnyClassFlags(CLASS_Abstract) || !Settings->SupportsAutoRegistration() )
 					{
 						continue;
 					}
diff --git a/Engine/Source/Editor/AnimGraph/Classes/AnimGraphNode_AssetPlayerBase.h b/Engine/Source/Editor/AnimGraph/Classes/AnimGraphNode_AssetPlayerBase.h
--- a/Engine/Source/Editor/AnimGraph/Classes/AnimGraphNode_AssetPlayerBase.h
+++ b/Engine/Source/Editor/AnimGraph/Classes/AnimGraphNode_AssetPlayerBase.h
@@ -30,15 +30,4 @@ public:
 	ANIMGRAPH_API virtual void PinDefaultValueChanged(UEdGraphPin* Pin) override;
 
 	virtual void SetAnimationAsset(UAnimationAsset* Asset) { check(false); /*Base function called*/ }
-
-	/** Store off a string asset reference for later restoration */
-	void SetAssetReferenceForPinRestoration(UObject* InAsset);
-
-	/** Find an asset reference when restoring pins */
-	UObject* GetAssetReferenceForPinRestoration();
-
-private:
-	/** Non-concrete reference to asset that was being used, now potentially coming from a pin */
-	UPROPERTY()
-	FStringAssetReference AssetReferenceForPinRestoration;
 };
diff --git a/Engine/Source/Editor/AnimGraph/Classes/AnimGraphNode_Base.h b/Engine/Source/Editor/AnimGraph/Classes/AnimGraphNode_Base.h
--- a/Engine/Source/Editor/AnimGraph/Classes/AnimGraphNode_Base.h
+++ b/Engine/Source/Editor/AnimGraph/Classes/AnimGraphNode_Base.h
@@ -124,7 +124,7 @@ class ANIMGRAPH_API UAnimGraphNode_Base : public UK2Node
 {
 	GENERATED_UCLASS_BODY()
 
-	UPROPERTY(EditAnywhere, BlueprintReadWrite, Category=PinOptions, EditFixedSize)
+	UPROPERTY(EditAnywhere, Category=PinOptions, EditFixedSize)
 	TArray<FOptionalPinFromProperty> ShowPinForProperties;
 
 	UPROPERTY(Transient)
diff --git a/Engine/Source/Editor/AnimGraph/Private/AnimBlueprintNodeOptionalPinManager.cpp b/Engine/Source/Editor/AnimGraph/Private/AnimBlueprintNodeOptionalPinManager.cpp
--- a/Engine/Source/Editor/AnimGraph/Private/AnimBlueprintNodeOptionalPinManager.cpp
+++ b/Engine/Source/Editor/AnimGraph/Private/AnimBlueprintNodeOptionalPinManager.cpp
@@ -45,90 +45,84 @@ void FAnimBlueprintNodeOptionalPinManager::GetRecordDefaults(UProperty* TestProp
 
 void FAnimBlueprintNodeOptionalPinManager::CustomizePinData(UEdGraphPin* Pin, FName SourcePropertyName, int32 ArrayIndex, UProperty* Property) const
 {
-	if (BaseNode != NULL)
+	if (BaseNode != nullptr)
 	{
 		BaseNode->CustomizePinData(Pin, SourcePropertyName, ArrayIndex);
 	}
 }
 
-void FAnimBlueprintNodeOptionalPinManager::PostInitNewPin(UEdGraphPin* Pin, FOptionalPinFromProperty& Record, int32 ArrayIndex, UProperty* Property, uint8* PropertyAddress) const
+void FAnimBlueprintNodeOptionalPinManager::PostInitNewPin(UEdGraphPin* Pin, FOptionalPinFromProperty& Record, int32 ArrayIndex, UProperty* Property, uint8* PropertyAddress, uint8* DefaultPropertyAddress) const
 {
-	check(PropertyAddress != NULL);
+	check(PropertyAddress != nullptr);
 	check(Record.bShowPin);
-		
+	
+	const UAnimationGraphSchema* Schema = GetDefault<UAnimationGraphSchema>();
+
+	// In all cases set autogenerated default value from node defaults
+	if (DefaultPropertyAddress)
+	{
+		FString LiteralValue;
+		FBlueprintEditorUtils::PropertyValueToString_Direct(Property, DefaultPropertyAddress, LiteralValue);
 
+		Schema->SetPinAutogeneratedDefaultValue(Pin, LiteralValue);
+	}
+	else
+	{
+		Schema->SetPinAutogeneratedDefaultValueBasedOnType(Pin);
+	}
 
-	if (OldPins == NULL)
+	if (OldPins == nullptr)
 	{
 		// Initial construction of a visible pin; copy values from the struct
-		FBlueprintEditorUtils::ExportPropertyToKismetDefaultValue(Pin, Property, PropertyAddress, BaseNode);
+		FString LiteralValue;
+		FBlueprintEditorUtils::PropertyValueToString_Direct(Property, PropertyAddress, LiteralValue);
+
+		Schema->SetPinDefaultValueAtConstruction(Pin, LiteralValue);
 	}
 	else if (Record.bCanToggleVisibility)
 	{
 		if (UEdGraphPin* OldPin = OldPinMap.FindRef(Pin->PinName))
 		{
-			// Was already visible
+			// Was already visible, values will get copied over in pin reconstruction code
 		}
 		else
 		{
-			// Showing a pin that was previously hidden, during a reconstruction
-
-			UObjectProperty* ObjectProperty = Cast<UObjectProperty>(Property);
-			UAnimGraphNode_AssetPlayerBase* AssetPlayerNode = Cast<UAnimGraphNode_AssetPlayerBase>(BaseNode);
-			const bool bIsAnimationAsset = AssetPlayerNode != nullptr && ObjectProperty != nullptr && ObjectProperty->PropertyClass->IsChildOf<UAnimationAsset>();
-
-			// Store the old reference to an animation asset in the node's asset reference
-			if (bIsAnimationAsset)
-			{
-				AssetPlayerNode->SetAssetReferenceForPinRestoration(ObjectProperty->GetObjectPropertyValue(PropertyAddress));
-			}
-		
 			// Convert the struct property into DefaultValue/DefaultValueObject
-			FBlueprintEditorUtils::ExportPropertyToKismetDefaultValue(Pin, Property, PropertyAddress, BaseNode);
+			FString LiteralValue;
+			FBlueprintEditorUtils::PropertyValueToString_Direct(Property, PropertyAddress, LiteralValue);
 
-			// clear the asset reference on the node
-			if (bIsAnimationAsset)
-			{
-				ObjectProperty->SetObjectPropertyValue(PropertyAddress, nullptr);
-			}
+			Schema->SetPinDefaultValueAtConstruction(Pin, LiteralValue);
+		}
+
+		// Clear the asset reference on the node if the pin exists
+		// In theory this is only needed for pins that are newly created but there are a lot of existing nodes that have dead asset references
+		UObjectProperty* ObjectProperty = Cast<UObjectProperty>(Property);
+		if (ObjectProperty)
+		{
+			ObjectProperty->SetObjectPropertyValue(PropertyAddress, nullptr);
 		}
 	}
 }
 
-void FAnimBlueprintNodeOptionalPinManager::PostRemovedOldPin(FOptionalPinFromProperty& Record, int32 ArrayIndex, UProperty* Property, uint8* PropertyAddress) const
+void FAnimBlueprintNodeOptionalPinManager::PostRemovedOldPin(FOptionalPinFromProperty& Record, int32 ArrayIndex, UProperty* Property, uint8* PropertyAddress, uint8* DefaultPropertyAddress) const
 {
-	check(PropertyAddress != NULL);
+	check(PropertyAddress != nullptr);
 	check(!Record.bShowPin);
 
-	if (Record.bCanToggleVisibility && (OldPins != NULL))
+	if (Record.bCanToggleVisibility && (OldPins != nullptr))
 	{
 		const FString OldPinName = (ArrayIndex != INDEX_NONE) ? FString::Printf(TEXT("%s_%d"), *(Record.PropertyName.ToString()), ArrayIndex) : Record.PropertyName.ToString();
+		// Pin was visible but it's now hidden
 		if (UEdGraphPin* OldPin = OldPinMap.FindRef(OldPinName))
 		{
-			// Pin was visible but it's now hidden
-			UObjectProperty* ObjectProperty = Cast<UObjectProperty>(Property);
-			UAnimGraphNode_AssetPlayerBase* AssetPlayerNode = Cast<UAnimGraphNode_AssetPlayerBase>(BaseNode);
-			if (AssetPlayerNode && ObjectProperty && ObjectProperty->PropertyClass->IsChildOf<UAnimationAsset>())
-			{
-				// If this is an anim asset pin, dont store a hard reference to the asset in the node, as this ends up referencing things we might not want to load any more
-				UObject* AssetReference = AssetPlayerNode->GetAssetReferenceForPinRestoration();
-				if (AssetReference)
-				{
-					ObjectProperty->SetObjectPropertyValue(PropertyAddress, AssetReference);
-				}
-					
-			}
-			else
-			{
-				// Convert DefaultValue/DefaultValueObject and push back into the struct
-				FBlueprintEditorUtils::ImportKismetDefaultValueToProperty(OldPin, Property, PropertyAddress, BaseNode);
-			}
+			// Convert DefaultValue/DefaultValueObject and push back into the struct
+			FBlueprintEditorUtils::PropertyValueFromString_Direct(Property, OldPin->GetDefaultAsString(), PropertyAddress);
 		}
 	}
 }
 
-void FAnimBlueprintNodeOptionalPinManager::AllocateDefaultPins(UStruct* SourceStruct, uint8* StructBasePtr)
+void FAnimBlueprintNodeOptionalPinManager::AllocateDefaultPins(UStruct* SourceStruct, uint8* StructBasePtr, uint8* DefaultsPtr)
 {
 	RebuildPropertyList(BaseNode->ShowPinForProperties, SourceStruct);
-	CreateVisiblePins(BaseNode->ShowPinForProperties, SourceStruct, EGPD_Input, BaseNode, StructBasePtr);
+	CreateVisiblePins(BaseNode->ShowPinForProperties, SourceStruct, EGPD_Input, BaseNode, StructBasePtr, DefaultsPtr);
 }
diff --git a/Engine/Source/Editor/AnimGraph/Private/AnimBlueprintNodeOptionalPinManager.h b/Engine/Source/Editor/AnimGraph/Private/AnimBlueprintNodeOptionalPinManager.h
--- a/Engine/Source/Editor/AnimGraph/Private/AnimBlueprintNodeOptionalPinManager.h
+++ b/Engine/Source/Editor/AnimGraph/Private/AnimBlueprintNodeOptionalPinManager.h
@@ -21,8 +21,8 @@ public:
 	/** FOptionalPinManager interface */
 	virtual void GetRecordDefaults(UProperty* TestProperty, FOptionalPinFromProperty& Record) const override;
 	virtual void CustomizePinData(UEdGraphPin* Pin, FName SourcePropertyName, int32 ArrayIndex, UProperty* Property) const override;
-	virtual void PostInitNewPin(UEdGraphPin* Pin, FOptionalPinFromProperty& Record, int32 ArrayIndex, UProperty* Property, uint8* PropertyAddress) const override;
-	virtual void PostRemovedOldPin(FOptionalPinFromProperty& Record, int32 ArrayIndex, UProperty* Property, uint8* PropertyAddress) const override;
+	virtual void PostInitNewPin(UEdGraphPin* Pin, FOptionalPinFromProperty& Record, int32 ArrayIndex, UProperty* Property, uint8* PropertyAddress, uint8* DefaultPropertyAddress) const override;
+	virtual void PostRemovedOldPin(FOptionalPinFromProperty& Record, int32 ArrayIndex, UProperty* Property, uint8* PropertyAddress, uint8* DefaultPropertyAddress) const override;
 
-	void AllocateDefaultPins(UStruct* SourceStruct, uint8* StructBasePtr);
+	void AllocateDefaultPins(UStruct* SourceStruct, uint8* StructBasePtr, uint8* DefaultsPtr);
 };
diff --git a/Engine/Source/Editor/AnimGraph/Private/AnimGraphNode_AssetPlayerBase.cpp b/Engine/Source/Editor/AnimGraph/Private/AnimGraphNode_AssetPlayerBase.cpp
--- a/Engine/Source/Editor/AnimGraph/Private/AnimGraphNode_AssetPlayerBase.cpp
+++ b/Engine/Source/Editor/AnimGraph/Private/AnimGraphNode_AssetPlayerBase.cpp
@@ -23,16 +23,6 @@ void UAnimGraphNode_AssetPlayerBase::PinConnectionListChanged(UEdGraphPin* Pin)
 
 	if (Pin->PinType.PinCategory == UEdGraphSchema_K2::PC_Object)
 	{
-		// if pin has no connections, we may need to restore the default
-		if (Pin->LinkedTo.Num() == 0)
-		{
-			UObject* Asset = GetAssetReferenceForPinRestoration();
-			if (Asset != nullptr)
-			{
-				Pin->DefaultObject = Asset;
-			}
-		}
-
 		// recache visualization now an asset pin's connection is changed
 		if (const UEdGraphSchema* Schema = GetSchema())
 		{
@@ -55,18 +45,6 @@ void UAnimGraphNode_AssetPlayerBase::PinDefaultValueChanged(UEdGraphPin* Pin)
 	}
 }
 
-void UAnimGraphNode_AssetPlayerBase::SetAssetReferenceForPinRestoration(UObject* InAsset)
-{
-	AssetReferenceForPinRestoration = InAsset;
-}
-
-UObject* UAnimGraphNode_AssetPlayerBase::GetAssetReferenceForPinRestoration()
-{
-	return AssetReferenceForPinRestoration.TryLoad();
-}
-
-
-
 UClass* GetNodeClassForAsset(const UClass* AssetClass)
 {
 	UClass* NodeClass = nullptr;
diff --git a/Engine/Source/Editor/AnimGraph/Private/AnimGraphNode_Base.cpp b/Engine/Source/Editor/AnimGraph/Private/AnimGraphNode_Base.cpp
--- a/Engine/Source/Editor/AnimGraph/Private/AnimGraphNode_Base.cpp
+++ b/Engine/Source/Editor/AnimGraph/Private/AnimGraphNode_Base.cpp
@@ -40,7 +40,7 @@ void UAnimGraphNode_Base::CreateOutputPins()
 	if (!IsSinkNode())
 	{
 		const UAnimationGraphSchema* Schema = GetDefault<UAnimationGraphSchema>();
-		CreatePin(EGPD_Output, Schema->PC_Struct, TEXT(""), FPoseLink::StaticStruct(), /*bIsArray=*/ false, /*bIsReference=*/ false, TEXT("Pose"));
+		CreatePin(EGPD_Output, Schema->PC_Struct, FString(), FPoseLink::StaticStruct(), TEXT("Pose"));
 	}
 }
 
@@ -54,8 +54,9 @@ void UAnimGraphNode_Base::InternalPinCreation(TArray<UEdGraphPin*>* OldPins)
 	{
 		// Display any currently visible optional pins
 		{
+			UObject* NodeDefaults = GetArchetype();
 			FAnimBlueprintNodeOptionalPinManager OptionalPinManager(this, OldPins);
-			OptionalPinManager.AllocateDefaultPins(NodeStruct->Struct, NodeStruct->ContainerPtrToValuePtr<uint8>(this));
+			OptionalPinManager.AllocateDefaultPins(NodeStruct->Struct, NodeStruct->ContainerPtrToValuePtr<uint8>(this), NodeDefaults ? NodeStruct->ContainerPtrToValuePtr<uint8>(NodeDefaults) : nullptr);
 		}
 
 		// Create the output pin, if needed
@@ -229,7 +230,7 @@ void UAnimGraphNode_Base::CreatePinsForPoseLink(UProperty* PoseProperty, int32 A
 
 	// pose input
 	const FString NewPinName = (ArrayIndex == INDEX_NONE) ? PoseProperty->GetName() : FString::Printf(TEXT("%s_%d"), *(PoseProperty->GetName()), ArrayIndex);
-	CreatePin(EGPD_Input, Schema->PC_Struct, TEXT(""), A2PoseStruct, /*bIsArray=*/ false, /*bIsReference=*/ false, NewPinName);
+	CreatePin(EGPD_Input, Schema->PC_Struct, FString(), A2PoseStruct, NewPinName);
 }
 
 void UAnimGraphNode_Base::PostProcessPinName(const UEdGraphPin* Pin, FString& DisplayName) const
@@ -238,7 +239,7 @@ void UAnimGraphNode_Base::PostProcessPinName(const UEdGraphPin* Pin, FString& Di
 	{
 		if (Pin->PinName == TEXT("Pose"))
 		{
-			DisplayName = TEXT("");
+			DisplayName.Reset();
 		}
 	}
 }
diff --git a/Engine/Source/Editor/AnimGraph/Private/AnimGraphNode_BoneDrivenController.cpp b/Engine/Source/Editor/AnimGraph/Private/AnimGraphNode_BoneDrivenController.cpp
--- a/Engine/Source/Editor/AnimGraph/Private/AnimGraphNode_BoneDrivenController.cpp
+++ b/Engine/Source/Editor/AnimGraph/Private/AnimGraphNode_BoneDrivenController.cpp
@@ -168,8 +168,8 @@ void UAnimGraphNode_BoneDrivenController::Draw(FPrimitiveDrawInterface* PDI, USk
 
 	if ((SourceIdx != INDEX_NONE) && (TargetIdx != INDEX_NONE))
 	{
-		const FTransform SourceTM = SkelMeshComp->GetComponentSpaceTransforms()[SourceIdx] * SkelMeshComp->ComponentToWorld;
-		const FTransform TargetTM = SkelMeshComp->GetComponentSpaceTransforms()[TargetIdx] * SkelMeshComp->ComponentToWorld;
+		const FTransform SourceTM = SkelMeshComp->GetComponentSpaceTransforms()[SourceIdx] * SkelMeshComp->GetComponentTransform();
+		const FTransform TargetTM = SkelMeshComp->GetComponentSpaceTransforms()[TargetIdx] * SkelMeshComp->GetComponentTransform();
 
 		PDI->DrawLine(TargetTM.GetLocation(), SourceTM.GetLocation(), FLinearColor(0.0f, 0.0f, 1.0f), SDPG_Foreground, 0.5f);
 
diff --git a/Engine/Source/Editor/AnimGraph/Private/AnimGraphNode_LocalToComponentSpace.cpp b/Engine/Source/Editor/AnimGraph/Private/AnimGraphNode_LocalToComponentSpace.cpp
--- a/Engine/Source/Editor/AnimGraph/Private/AnimGraphNode_LocalToComponentSpace.cpp
+++ b/Engine/Source/Editor/AnimGraph/Private/AnimGraphNode_LocalToComponentSpace.cpp
@@ -36,12 +36,12 @@ FString UAnimGraphNode_LocalToComponentSpace::GetNodeCategory() const
 void UAnimGraphNode_LocalToComponentSpace::CreateOutputPins()
 {
 	const UAnimationGraphSchema* Schema = GetDefault<UAnimationGraphSchema>();
-	CreatePin(EGPD_Output, Schema->PC_Struct, TEXT(""), FComponentSpacePoseLink::StaticStruct(), /*bIsArray=*/ false, /*bIsReference=*/ false, TEXT("ComponentPose"));
+	CreatePin(EGPD_Output, Schema->PC_Struct, FString(), FComponentSpacePoseLink::StaticStruct(), TEXT("ComponentPose"));
 }
 
 void UAnimGraphNode_LocalToComponentSpace::PostProcessPinName(const UEdGraphPin* Pin, FString& DisplayName) const
 {
-	DisplayName = TEXT("");
+	DisplayName.Reset();
 }
 
 #undef LOCTEXT_NAMESPACE
diff --git a/Engine/Source/Editor/AnimGraph/Private/AnimGraphNode_MeshRefPose.cpp b/Engine/Source/Editor/AnimGraph/Private/AnimGraphNode_MeshRefPose.cpp
--- a/Engine/Source/Editor/AnimGraph/Private/AnimGraphNode_MeshRefPose.cpp
+++ b/Engine/Source/Editor/AnimGraph/Private/AnimGraphNode_MeshRefPose.cpp
@@ -36,7 +36,7 @@ FText UAnimGraphNode_MeshRefPose::GetNodeTitle(ENodeTitleType::Type TitleType) c
 void UAnimGraphNode_MeshRefPose::CreateOutputPins()
 {
 	const UAnimationGraphSchema* Schema = GetDefault<UAnimationGraphSchema>();
-	CreatePin(EGPD_Output, Schema->PC_Struct, TEXT(""), FComponentSpacePoseLink::StaticStruct(), /*bIsArray=*/ false, /*bIsReference=*/ false, TEXT("ComponentPose"));
+	CreatePin(EGPD_Output, Schema->PC_Struct, FString(), FComponentSpacePoseLink::StaticStruct(), TEXT("ComponentPose"));
 }
 
 #undef LOCTEXT_NAMESPACE
diff --git a/Engine/Source/Editor/AnimGraph/Private/AnimGraphNode_SkeletalControlBase.cpp b/Engine/Source/Editor/AnimGraph/Private/AnimGraphNode_SkeletalControlBase.cpp
--- a/Engine/Source/Editor/AnimGraph/Private/AnimGraphNode_SkeletalControlBase.cpp
+++ b/Engine/Source/Editor/AnimGraph/Private/AnimGraphNode_SkeletalControlBase.cpp
@@ -73,7 +73,7 @@ FText UAnimGraphNode_SkeletalControlBase::GetTooltipText() const
 void UAnimGraphNode_SkeletalControlBase::CreateOutputPins()
 {
 	const UAnimationGraphSchema* Schema = GetDefault<UAnimationGraphSchema>();
-	CreatePin(EGPD_Output, Schema->PC_Struct, TEXT(""), FComponentSpacePoseLink::StaticStruct(), /*bIsArray=*/ false, /*bIsReference=*/ false, TEXT("Pose"));
+	CreatePin(EGPD_Output, Schema->PC_Struct, FString(), FComponentSpacePoseLink::StaticStruct(), TEXT("Pose"));
 }
 
 
@@ -86,7 +86,7 @@ void UAnimGraphNode_SkeletalControlBase::ConvertToComponentSpaceTransform(const
 	case BCS_WorldSpace:
 	{
 		OutCSTransform = InTransform;
-		OutCSTransform.SetToRelativeTransform(SkelComp->ComponentToWorld);
+		OutCSTransform.SetToRelativeTransform(SkelComp->GetComponentTransform());
 	}
 		break;
 
@@ -254,7 +254,7 @@ FVector UAnimGraphNode_SkeletalControlBase::ConvertWidgetLocation(const USkeleta
 
 		switch (Space)
 		{
-			// ComponentToWorld must be Identity in preview window so same as ComponentSpace
+			// GetComponentTransform() must be Identity in preview window so same as ComponentSpace
 		case BCS_WorldSpace:
 		case BCS_ComponentSpace:
 		{
diff --git a/Engine/Source/Editor/AnimGraph/Private/AnimStateConduitNode.cpp b/Engine/Source/Editor/AnimGraph/Private/AnimStateConduitNode.cpp
--- a/Engine/Source/Editor/AnimGraph/Private/AnimStateConduitNode.cpp
+++ b/Engine/Source/Editor/AnimGraph/Private/AnimStateConduitNode.cpp
@@ -22,15 +22,15 @@ UAnimStateConduitNode::UAnimStateConduitNode(const FObjectInitializer& ObjectIni
 
 void UAnimStateConduitNode::AllocateDefaultPins()
 {
-	CreatePin(EGPD_Input, TEXT("Transition"), TEXT(""), NULL, false, false, TEXT("In"));
-	CreatePin(EGPD_Output, TEXT("Transition"), TEXT(""), NULL, false, false, TEXT("Out"));
+	CreatePin(EGPD_Input, TEXT("Transition"), FString(), nullptr, TEXT("In"));
+	CreatePin(EGPD_Output, TEXT("Transition"), FString(), nullptr, TEXT("Out"));
 }
 
 void UAnimStateConduitNode::AutowireNewNode(UEdGraphPin* FromPin)
 {
 	Super::AutowireNewNode(FromPin);
 
-	if (FromPin != NULL)
+	if (FromPin)
 	{
 		if (GetSchema()->TryCreateConnection(FromPin, GetInputPin()))
 		{
@@ -135,6 +135,8 @@ void UAnimStateConduitNode::DestroyNode()
 
 void UAnimStateConduitNode::ValidateNodeDuringCompilation(class FCompilerResultsLog& MessageLog) const
 {
+	Super::ValidateNodeDuringCompilation(MessageLog);
+
 	UAnimationTransitionGraph* TransGraph = CastChecked<UAnimationTransitionGraph>(BoundGraph);
 	UAnimGraphNode_TransitionResult* ResultNode = TransGraph->GetResultNode();
 	check(ResultNode);
diff --git a/Engine/Source/Editor/AnimGraph/Private/AnimStateEntryNode.cpp b/Engine/Source/Editor/AnimGraph/Private/AnimStateEntryNode.cpp
--- a/Engine/Source/Editor/AnimGraph/Private/AnimStateEntryNode.cpp
+++ b/Engine/Source/Editor/AnimGraph/Private/AnimStateEntryNode.cpp
@@ -21,7 +21,7 @@ UAnimStateEntryNode::UAnimStateEntryNode(const FObjectInitializer& ObjectInitial
 void UAnimStateEntryNode::AllocateDefaultPins()
 {
 	const UAnimationStateMachineSchema* Schema = GetDefault<UAnimationStateMachineSchema>();
-	UEdGraphPin* Outputs = CreatePin(EGPD_Output, Schema->PC_Exec, TEXT(""), NULL, false, false, TEXT("Entry"));
+	UEdGraphPin* Outputs = CreatePin(EGPD_Output, Schema->PC_Exec, FString(), nullptr, TEXT("Entry"));
 }
 
 FText UAnimStateEntryNode::GetNodeTitle(ENodeTitleType::Type TitleType) const
diff --git a/Engine/Source/Editor/AnimGraph/Private/AnimStateNode.cpp b/Engine/Source/Editor/AnimGraph/Private/AnimStateNode.cpp
--- a/Engine/Source/Editor/AnimGraph/Private/AnimStateNode.cpp
+++ b/Engine/Source/Editor/AnimGraph/Private/AnimStateNode.cpp
@@ -25,8 +25,8 @@ UAnimStateNode::UAnimStateNode(const FObjectInitializer& ObjectInitializer)
 
 void UAnimStateNode::AllocateDefaultPins()
 {
-	UEdGraphPin* Inputs = CreatePin(EGPD_Input, TEXT("Transition"), TEXT(""), NULL, false, false, TEXT("In"));
-	UEdGraphPin* Outputs = CreatePin(EGPD_Output, TEXT("Transition"), TEXT(""), NULL, false, false, TEXT("Out"));
+	UEdGraphPin* Inputs = CreatePin(EGPD_Input, TEXT("Transition"), FString(), nullptr, TEXT("In"));
+	UEdGraphPin* Outputs = CreatePin(EGPD_Output, TEXT("Transition"), FString(), nullptr, TEXT("Out"));
 }
 
 void UAnimStateNode::AutowireNewNode(UEdGraphPin* FromPin)
@@ -34,7 +34,7 @@ void UAnimStateNode::AutowireNewNode(UEdGraphPin* FromPin)
 	Super::AutowireNewNode(FromPin);
 
 	//@TODO: If the FromPin is a state, create a transition between us
-	if (FromPin != NULL)
+	if (FromPin)
 	{
 		if (GetSchema()->TryCreateConnection(FromPin, GetInputPin()))
 		{
diff --git a/Engine/Source/Editor/AnimGraph/Private/AnimStateTransitionNode.cpp b/Engine/Source/Editor/AnimGraph/Private/AnimStateTransitionNode.cpp
--- a/Engine/Source/Editor/AnimGraph/Private/AnimStateTransitionNode.cpp
+++ b/Engine/Source/Editor/AnimGraph/Private/AnimStateTransitionNode.cpp
@@ -82,9 +82,9 @@ UAnimStateTransitionNode::UAnimStateTransitionNode(const FObjectInitializer& Obj
 
 void UAnimStateTransitionNode::AllocateDefaultPins()
 {
-	UEdGraphPin* Inputs = CreatePin(EGPD_Input, TEXT("Transition"), TEXT(""), NULL, false, false, TEXT("In"));
+	UEdGraphPin* Inputs = CreatePin(EGPD_Input, TEXT("Transition"), FString(), nullptr, TEXT("In"));
 	Inputs->bHidden = true;
-	UEdGraphPin* Outputs = CreatePin(EGPD_Output, TEXT("Transition"), TEXT(""), NULL, false, false, TEXT("Out"));
+	UEdGraphPin* Outputs = CreatePin(EGPD_Output, TEXT("Transition"), FString(), nullptr, TEXT("Out"));
 	Outputs->bHidden = true;
 }
 
@@ -522,6 +522,8 @@ bool UAnimStateTransitionNode::IsBoundGraphShared()
 
 void UAnimStateTransitionNode::ValidateNodeDuringCompilation(class FCompilerResultsLog& MessageLog) const
 {
+	Super::ValidateNodeDuringCompilation(MessageLog);
+
 	if (UAnimationTransitionGraph* TransGraph = Cast<UAnimationTransitionGraph>(BoundGraph))
 	{
 		UAnimGraphNode_TransitionResult* ResultNode = TransGraph->GetResultNode();
diff --git a/Engine/Source/Editor/AnimGraph/Private/EditModes/AnimNodeEditMode.cpp b/Engine/Source/Editor/AnimGraph/Private/EditModes/AnimNodeEditMode.cpp
--- a/Engine/Source/Editor/AnimGraph/Private/EditModes/AnimNodeEditMode.cpp
+++ b/Engine/Source/Editor/AnimGraph/Private/EditModes/AnimNodeEditMode.cpp
@@ -364,7 +364,7 @@ void FAnimNodeEditMode::ConvertToComponentSpaceTransform(const USkeletalMeshComp
 	case BCS_WorldSpace:
 	{
 		OutCSTransform = InTransform;
-		OutCSTransform.SetToRelativeTransform(SkelComp->ComponentToWorld);
+		OutCSTransform.SetToRelativeTransform(SkelComp->GetComponentTransform());
 	}
 	break;
 
@@ -433,7 +433,7 @@ void FAnimNodeEditMode::ConvertToBoneSpaceTransform(const USkeletalMeshComponent
 	{
 		case BCS_WorldSpace:
 		{
-			OutBSTransform = InCSTransform * SkelComp->ComponentToWorld;
+			OutBSTransform = InCSTransform * SkelComp->GetComponentTransform();
 			break;
 		}
 		
@@ -609,7 +609,7 @@ FVector FAnimNodeEditMode::ConvertWidgetLocation(const USkeletalMeshComponent* S
 
 	switch (Space)
 	{
-		// ComponentToWorld must be Identity in preview window so same as ComponentSpace
+		// GetComponentTransform() must be Identity in preview window so same as ComponentSpace
 	case BCS_WorldSpace:
 	case BCS_ComponentSpace:
 	{
diff --git a/Engine/Source/Editor/AnimGraph/Private/EditModes/SplineIKEditMode.cpp b/Engine/Source/Editor/AnimGraph/Private/EditModes/SplineIKEditMode.cpp
--- a/Engine/Source/Editor/AnimGraph/Private/EditModes/SplineIKEditMode.cpp
+++ b/Engine/Source/Editor/AnimGraph/Private/EditModes/SplineIKEditMode.cpp
@@ -62,7 +62,7 @@ void FSplineIKEditMode::Render(const FSceneView* View, FViewport* Viewport, FPri
 {
 	UDebugSkelMeshComponent* SkelComp = GetAnimPreviewScene().GetPreviewMeshComponent();
 
-	USplineComponent::Draw(PDI, View, SplineIKRuntimeNode->GetTransformedSplineCurves().Position, SkelComp->ComponentToWorld.ToMatrixWithScale(), FLinearColor::Yellow, SDPG_Foreground);
+	USplineComponent::Draw(PDI, View, SplineIKRuntimeNode->GetTransformedSplineCurves().Position, SkelComp->GetComponentTransform().ToMatrixWithScale(), FLinearColor::Yellow, SDPG_Foreground);
 
 	for (int32 SplineHandleIndex = 0; SplineHandleIndex < SplineIKRuntimeNode->GetNumControlPoints(); SplineHandleIndex++)
 	{
@@ -83,7 +83,7 @@ FVector FSplineIKEditMode::GetWidgetLocation() const
 		FVector Location = SplineIKRuntimeNode->GetTransformedSplinePoint(SelectedSplinePoint).GetLocation();
 
 		UDebugSkelMeshComponent* SkelComp = GetAnimPreviewScene().GetPreviewMeshComponent();
-		return SkelComp->ComponentToWorld.TransformPosition(Location);
+		return SkelComp->GetComponentTransform().TransformPosition(Location);
 	}
 
 	return FVector::ZeroVector;
@@ -188,7 +188,7 @@ bool FSplineIKEditMode::GetCustomDrawingCoordinateSystem(FMatrix& InMatrix, void
 		if (SelectedSplinePoint != INDEX_NONE)
 		{
 			FTransform Transform = SplineIKRuntimeNode->GetTransformedSplinePoint(SelectedSplinePoint);
-			FTransform WorldTransform = Transform * SkelMeshComp->ComponentToWorld;
+			FTransform WorldTransform = Transform * SkelMeshComp->GetComponentTransform();
 			InMatrix = WorldTransform.ToMatrixNoScale().RemoveTranslation();
 		}
 
diff --git a/Engine/Source/Editor/AnimGraph/Private/K2Node_TransitionRuleGetter.cpp b/Engine/Source/Editor/AnimGraph/Private/K2Node_TransitionRuleGetter.cpp
--- a/Engine/Source/Editor/AnimGraph/Private/K2Node_TransitionRuleGetter.cpp
+++ b/Engine/Source/Editor/AnimGraph/Private/K2Node_TransitionRuleGetter.cpp
@@ -28,7 +28,7 @@ void UK2Node_TransitionRuleGetter::AllocateDefaultPins()
 {
 	const UEdGraphSchema_K2* Schema = GetDefault<UEdGraphSchema_K2>();
 
-	UEdGraphPin* OutputPin = CreatePin(EGPD_Output, Schema->PC_Float, /*PSC=*/ TEXT(""), /*PSC object=*/ NULL, /*bIsArray=*/ false, /*bIsReference=*/ false, TEXT("Output"));
+	UEdGraphPin* OutputPin = CreatePin(EGPD_Output, Schema->PC_Float, /*PSC=*/ FString(), /*PSC object=*/ nullptr, TEXT("Output"));
 	OutputPin->PinFriendlyName = GetFriendlyName(GetterType);
 
 	PreloadObject(AssociatedAnimAssetPlayerNode);
diff --git a/Engine/Source/Editor/AudioEditor/Private/Factories/SoundFactoryUtility.h b/Engine/Source/Editor/AudioEditor/Private/Factories/SoundFactoryUtility.h
--- a/Engine/Source/Editor/AudioEditor/Private/Factories/SoundFactoryUtility.h
+++ b/Engine/Source/Editor/AudioEditor/Private/Factories/SoundFactoryUtility.h
@@ -13,7 +13,7 @@ public:
 	TSet< const UClass* > AllowedChildrenOfClasses;
 
 	/** Disallowed class flags. */
-	uint32 DisallowedClassFlags;
+	EClassFlags DisallowedClassFlags;
 
 	virtual bool IsClassAllowed(const FClassViewerInitializationOptions& InInitOptions, const UClass* InClass, TSharedRef< FClassViewerFilterFuncs > InFilterFuncs) override
 	{
diff --git a/Engine/Source/Editor/AudioEditor/Private/SoundClassGraphNode.cpp b/Engine/Source/Editor/AudioEditor/Private/SoundClassGraphNode.cpp
--- a/Engine/Source/Editor/AudioEditor/Private/SoundClassGraphNode.cpp
+++ b/Engine/Source/Editor/AudioEditor/Private/SoundClassGraphNode.cpp
@@ -71,13 +71,13 @@ void USoundClassGraphNode::AllocateDefaultPins()
 {
 	check(Pins.Num() == 0);
 
-	ChildPin = CreatePin(EGPD_Output, TEXT("SoundClass"), TEXT(""), NULL, /*bIsArray=*/ false, /*bIsReference=*/ false, LOCTEXT("SoundClassChildren", "Children").ToString());
-	ParentPin = CreatePin(EGPD_Input, TEXT("SoundClass"), TEXT(""), NULL, /*bIsArray=*/ false, /*bIsReference=*/ false, TEXT(""));
+	ChildPin = CreatePin(EGPD_Output, TEXT("SoundClass"), FString(), nullptr, LOCTEXT("SoundClassChildren", "Children").ToString());
+	ParentPin = CreatePin(EGPD_Input, TEXT("SoundClass"), FString(), nullptr, FString());
 }
 
 void USoundClassGraphNode::AutowireNewNode(UEdGraphPin* FromPin)
 {
-	if (FromPin != NULL)
+	if (FromPin)
 	{
 		const USoundClassGraphSchema* Schema = CastChecked<USoundClassGraphSchema>(GetSchema());
 
diff --git a/Engine/Source/Editor/AudioEditor/Private/SoundCueGraphNode.cpp b/Engine/Source/Editor/AudioEditor/Private/SoundCueGraphNode.cpp
--- a/Engine/Source/Editor/AudioEditor/Private/SoundCueGraphNode.cpp
+++ b/Engine/Source/Editor/AudioEditor/Private/SoundCueGraphNode.cpp
@@ -60,7 +60,7 @@ void USoundCueGraphNode::SetSoundNode(USoundNode* InSoundNode)
 
 void USoundCueGraphNode::CreateInputPin()
 {
-	UEdGraphPin* NewPin = CreatePin(EGPD_Input, TEXT("SoundNode"), TEXT(""), NULL, /*bIsArray=*/ false, /*bIsReference=*/ false, SoundNode->GetInputPinName(GetInputCount()).ToString());
+	UEdGraphPin* NewPin = CreatePin(EGPD_Input, TEXT("SoundNode"), FString(), nullptr, SoundNode->GetInputPinName(GetInputCount()).ToString());
 	if (NewPin->PinName.IsEmpty())
 	{
 		// Makes sure pin has a name for lookup purposes but user will never see it
diff --git a/Engine/Source/Editor/AudioEditor/Private/SoundCueGraphNode_Base.cpp b/Engine/Source/Editor/AudioEditor/Private/SoundCueGraphNode_Base.cpp
--- a/Engine/Source/Editor/AudioEditor/Private/SoundCueGraphNode_Base.cpp
+++ b/Engine/Source/Editor/AudioEditor/Private/SoundCueGraphNode_Base.cpp
@@ -124,7 +124,7 @@ void USoundCueGraphNode_Base::AllocateDefaultPins()
 
 	if (!IsRootNode())
 	{
-		CreatePin(EGPD_Output, TEXT("SoundNode"), TEXT(""), NULL, /*bIsArray=*/ false, /*bIsReference=*/ false, TEXT("Output"));
+		CreatePin(EGPD_Output, TEXT("SoundNode"), FString(), nullptr, TEXT("Output"));
 	}
 }
 
diff --git a/Engine/Source/Editor/AudioEditor/Private/SoundCueGraphNode_Root.cpp b/Engine/Source/Editor/AudioEditor/Private/SoundCueGraphNode_Root.cpp
--- a/Engine/Source/Editor/AudioEditor/Private/SoundCueGraphNode_Root.cpp
+++ b/Engine/Source/Editor/AudioEditor/Private/SoundCueGraphNode_Root.cpp
@@ -36,7 +36,7 @@ FText USoundCueGraphNode_Root::GetTooltipText() const
 
 void USoundCueGraphNode_Root::CreateInputPins()
 {
-	CreatePin(EGPD_Input, TEXT("SoundNode"), TEXT("Root"), NULL, /*bIsArray=*/ false, /*bIsReference=*/ false, TEXT(""));
+	CreatePin(EGPD_Input, TEXT("SoundNode"), TEXT("Root"), nullptr, FString());
 }
 
 void USoundCueGraphNode_Root::GetContextMenuActions(const FGraphNodeContextMenuBuilder& Context) const
diff --git a/Engine/Source/Editor/AudioEditor/Private/SoundSubmixGraphNode.cpp b/Engine/Source/Editor/AudioEditor/Private/SoundSubmixGraphNode.cpp
--- a/Engine/Source/Editor/AudioEditor/Private/SoundSubmixGraphNode.cpp
+++ b/Engine/Source/Editor/AudioEditor/Private/SoundSubmixGraphNode.cpp
@@ -61,13 +61,13 @@ void USoundSubmixGraphNode::AllocateDefaultPins()
 {
 	check(Pins.Num() == 0);
 
-	ChildPin = CreatePin(EGPD_Output, TEXT("SoundSubmix"), TEXT(""), NULL, /*bIsArray=*/ false, /*bIsReference=*/ false, LOCTEXT("SoundSubmixChildren", "Children").ToString());
-	ParentPin = CreatePin(EGPD_Input, TEXT("SoundSubmix"), TEXT(""), NULL, /*bIsArray=*/ false, /*bIsReference=*/ false, TEXT(""));
+	ChildPin = CreatePin(EGPD_Output, TEXT("SoundSubmix"), FString(), nullptr, LOCTEXT("SoundSubmixChildren", "Children").ToString());
+	ParentPin = CreatePin(EGPD_Input, TEXT("SoundSubmix"), FString(), nullptr, FString());
 }
 
 void USoundSubmixGraphNode::AutowireNewNode(UEdGraphPin* FromPin)
 {
-	if (FromPin != NULL)
+	if (FromPin)
 	{
 		const USoundSubmixGraphSchema* Schema = CastChecked<USoundSubmixGraphSchema>(GetSchema());
 
diff --git a/Engine/Source/Editor/BehaviorTreeEditor/Private/BehaviorTreeDecoratorGraphNode_Decorator.cpp b/Engine/Source/Editor/BehaviorTreeEditor/Private/BehaviorTreeDecoratorGraphNode_Decorator.cpp
--- a/Engine/Source/Editor/BehaviorTreeEditor/Private/BehaviorTreeDecoratorGraphNode_Decorator.cpp
+++ b/Engine/Source/Editor/BehaviorTreeEditor/Private/BehaviorTreeDecoratorGraphNode_Decorator.cpp
@@ -15,7 +15,7 @@ UBehaviorTreeDecoratorGraphNode_Decorator::UBehaviorTreeDecoratorGraphNode_Decor
 
 void UBehaviorTreeDecoratorGraphNode_Decorator::AllocateDefaultPins()
 {
-	CreatePin(EGPD_Output, TEXT("Transition"), TEXT(""), NULL, false, false, TEXT("Out"));
+	CreatePin(EGPD_Output, TEXT("Transition"), FString(), nullptr, TEXT("Out"));
 }
 
 void UBehaviorTreeDecoratorGraphNode_Decorator::PostPlacedNewNode()
diff --git a/Engine/Source/Editor/BehaviorTreeEditor/Private/BehaviorTreeDecoratorGraphNode_Logic.cpp b/Engine/Source/Editor/BehaviorTreeEditor/Private/BehaviorTreeDecoratorGraphNode_Logic.cpp
--- a/Engine/Source/Editor/BehaviorTreeEditor/Private/BehaviorTreeDecoratorGraphNode_Logic.cpp
+++ b/Engine/Source/Editor/BehaviorTreeEditor/Private/BehaviorTreeDecoratorGraphNode_Logic.cpp
@@ -23,7 +23,7 @@ void UBehaviorTreeDecoratorGraphNode_Logic::AllocateDefaultPins()
 
 	if (LogicMode != EDecoratorLogicMode::Sink)
 	{
-		CreatePin(EGPD_Output, TEXT("Transition"), TEXT(""), NULL, false, false, TEXT("Out"));
+		CreatePin(EGPD_Output, TEXT("Transition"), FString(), nullptr, TEXT("Out"));
 	}
 }
 
@@ -73,7 +73,7 @@ bool UBehaviorTreeDecoratorGraphNode_Logic::CanRemovePins() const
 
 UEdGraphPin* UBehaviorTreeDecoratorGraphNode_Logic::AddInputPin()
 {
-	return CreatePin(EGPD_Input, TEXT("Transition"), TEXT(""), NULL, false, false, TEXT("In"));
+	return CreatePin(EGPD_Input, TEXT("Transition"), FString(), nullptr, TEXT("In"));
 }
 
 void UBehaviorTreeDecoratorGraphNode_Logic::RemoveInputPin(class UEdGraphPin* Pin)
diff --git a/Engine/Source/Editor/BehaviorTreeEditor/Private/BehaviorTreeGraphNode.cpp b/Engine/Source/Editor/BehaviorTreeEditor/Private/BehaviorTreeGraphNode.cpp
--- a/Engine/Source/Editor/BehaviorTreeEditor/Private/BehaviorTreeGraphNode.cpp
+++ b/Engine/Source/Editor/BehaviorTreeEditor/Private/BehaviorTreeGraphNode.cpp
@@ -44,8 +44,8 @@ UBehaviorTreeGraphNode::UBehaviorTreeGraphNode(const FObjectInitializer& ObjectI
 
 void UBehaviorTreeGraphNode::AllocateDefaultPins()
 {
-	CreatePin(EGPD_Input, UBehaviorTreeEditorTypes::PinCategory_MultipleNodes, TEXT(""), NULL, false, false, TEXT("In"));
-	CreatePin(EGPD_Output, UBehaviorTreeEditorTypes::PinCategory_MultipleNodes, TEXT(""), NULL, false, false, TEXT("Out"));
+	CreatePin(EGPD_Input, UBehaviorTreeEditorTypes::PinCategory_MultipleNodes, FString(), nullptr, TEXT("In"));
+	CreatePin(EGPD_Output, UBehaviorTreeEditorTypes::PinCategory_MultipleNodes, FString(), nullptr, TEXT("Out"));
 }
 
 void UBehaviorTreeGraphNode::InitializeInstance()
diff --git a/Engine/Source/Editor/BehaviorTreeEditor/Private/BehaviorTreeGraphNode_Root.cpp b/Engine/Source/Editor/BehaviorTreeEditor/Private/BehaviorTreeGraphNode_Root.cpp
--- a/Engine/Source/Editor/BehaviorTreeEditor/Private/BehaviorTreeGraphNode_Root.cpp
+++ b/Engine/Source/Editor/BehaviorTreeEditor/Private/BehaviorTreeGraphNode_Root.cpp
@@ -31,7 +31,7 @@ void UBehaviorTreeGraphNode_Root::PostPlacedNewNode()
 
 void UBehaviorTreeGraphNode_Root::AllocateDefaultPins()
 {
-	CreatePin(EGPD_Output, UBehaviorTreeEditorTypes::PinCategory_SingleComposite, TEXT(""), NULL, false, false, TEXT("In"));
+	CreatePin(EGPD_Output, UBehaviorTreeEditorTypes::PinCategory_SingleComposite, FString(), nullptr, TEXT("In"));
 }
 
 FText UBehaviorTreeGraphNode_Root::GetNodeTitle(ENodeTitleType::Type TitleType) const
diff --git a/Engine/Source/Editor/BehaviorTreeEditor/Private/BehaviorTreeGraphNode_SimpleParallel.cpp b/Engine/Source/Editor/BehaviorTreeEditor/Private/BehaviorTreeGraphNode_SimpleParallel.cpp
--- a/Engine/Source/Editor/BehaviorTreeEditor/Private/BehaviorTreeGraphNode_SimpleParallel.cpp
+++ b/Engine/Source/Editor/BehaviorTreeEditor/Private/BehaviorTreeGraphNode_SimpleParallel.cpp
@@ -11,10 +11,10 @@ UBehaviorTreeGraphNode_SimpleParallel::UBehaviorTreeGraphNode_SimpleParallel(con
 
 void UBehaviorTreeGraphNode_SimpleParallel::AllocateDefaultPins()
 {
-	CreatePin(EGPD_Input, UBehaviorTreeEditorTypes::PinCategory_MultipleNodes, TEXT(""), NULL, false, false, TEXT("In"));
+	CreatePin(EGPD_Input, UBehaviorTreeEditorTypes::PinCategory_MultipleNodes, FString(), nullptr, TEXT("In"));
 	
-	CreatePin(EGPD_Output, UBehaviorTreeEditorTypes::PinCategory_SingleTask, TEXT(""), NULL, false, false, TEXT("Task"));
-	CreatePin(EGPD_Output, UBehaviorTreeEditorTypes::PinCategory_SingleNode, TEXT(""), NULL, false, false, TEXT("Out"));
+	CreatePin(EGPD_Output, UBehaviorTreeEditorTypes::PinCategory_SingleTask, FString(), nullptr, TEXT("Task"));
+	CreatePin(EGPD_Output, UBehaviorTreeEditorTypes::PinCategory_SingleNode, FString(), nullptr, TEXT("Out"));
 }
 
 void UBehaviorTreeGraphNode_SimpleParallel::GetPinHoverText(const UEdGraphPin& Pin, FString& HoverTextOut) const
diff --git a/Engine/Source/Editor/BehaviorTreeEditor/Private/BehaviorTreeGraphNode_Task.cpp b/Engine/Source/Editor/BehaviorTreeEditor/Private/BehaviorTreeGraphNode_Task.cpp
--- a/Engine/Source/Editor/BehaviorTreeEditor/Private/BehaviorTreeGraphNode_Task.cpp
+++ b/Engine/Source/Editor/BehaviorTreeEditor/Private/BehaviorTreeGraphNode_Task.cpp
@@ -10,7 +10,7 @@ UBehaviorTreeGraphNode_Task::UBehaviorTreeGraphNode_Task(const FObjectInitialize
 
 void UBehaviorTreeGraphNode_Task::AllocateDefaultPins()
 {
-	CreatePin(EGPD_Input, UBehaviorTreeEditorTypes::PinCategory_SingleComposite, TEXT(""), NULL, false, false, TEXT("In"));
+	CreatePin(EGPD_Input, UBehaviorTreeEditorTypes::PinCategory_SingleComposite, FString(), nullptr, TEXT("In"));
 }
 
 FText UBehaviorTreeGraphNode_Task::GetNodeTitle(ENodeTitleType::Type TitleType) const
diff --git a/Engine/Source/Editor/BlueprintGraph/Classes/EdGraphSchema_K2.h b/Engine/Source/Editor/BlueprintGraph/Classes/EdGraphSchema_K2.h
--- a/Engine/Source/Editor/BlueprintGraph/Classes/EdGraphSchema_K2.h
+++ b/Engine/Source/Editor/BlueprintGraph/Classes/EdGraphSchema_K2.h
@@ -536,8 +536,11 @@ public:
 		If the CompilerContext is specified this will be created as an intermediate node */
 	class UK2Node* CreateSplitPinNode(UEdGraphPin* Pin, class FKismetCompilerContext* CompilerContext = NULL, UEdGraph* SourceGraph = NULL) const;
 
-	// Do validation, that doesn't require a knowledge about actual pin. 
-	virtual bool DefaultValueSimpleValidation(const FEdGraphPinType& PinType, const FString& PinName, const FString& NewDefaultValue, UObject* NewDefaultObject, const FText& InText, FString* OutMsg = NULL) const;
+	/** Reads in a FString and gets the values of the pin defaults for that type. This can be passed to DefaultValueSimpleValidation to validate. OwningObject can be null */
+	virtual void GetPinDefaultValuesFromString(const FEdGraphPinType& PinType, UObject* OwningObject, const FString& NewValue, FString& UseDefaultValue, UObject*& UseDefaultObject, FText& UseDefaultText) const;
+
+	/** Do validation, that doesn't require a knowledge about actual pin */
+	virtual bool DefaultValueSimpleValidation(const FEdGraphPinType& PinType, const FString& PinName, const FString& NewDefaultValue, UObject* NewDefaultObject, const FText& InText, FString* OutMsg = nullptr) const;
 
 	/** Returns true if the owning node is a function with AutoCreateRefTerm meta data */
 	bool IsAutoCreateRefTerm(const UEdGraphPin* Pin) const;
@@ -584,7 +587,7 @@ public:
 	 * @param	Pin	The pin to check.
 	 * @return	true if it is an execution pin.
 	 */
-	inline bool IsExecPin(const UEdGraphPin& Pin) const
+	static inline bool IsExecPin(const UEdGraphPin& Pin)
 	{
 		return Pin.PinType.PinCategory == PC_Exec;
 	}
@@ -605,7 +608,7 @@ public:
 	 */
 	inline bool IsMetaPin(const UEdGraphPin& Pin) const
 	{
-		return IsSelfPin(Pin) || IsExecPin(Pin);
+		return IsExecPin(Pin) || IsSelfPin(Pin);
 	}
 
 	/** Is given string a delegate category name ? */
@@ -614,7 +617,7 @@ public:
 	/** Returns whether a pin category is compatible with an Index Wildcard (PC_Wildcard and PSC_Index) */
 	inline bool IsIndexWildcardCompatible(const FEdGraphPinType& PinType) const
 	{
-		return (!PinType.bIsArray) && 
+		return (!PinType.IsContainer()) && 
 			(
 				PinType.PinCategory == PC_Boolean || 
 				PinType.PinCategory == PC_Int || 
@@ -943,14 +946,30 @@ public:
 	 */
 	virtual bool ArePinTypesCompatible(const FEdGraphPinType& Output, const FEdGraphPinType& Input, const UClass* CallingContext = NULL, bool bIgnoreArray = false) const;
 
-	// Sets the Pin default property potentially using the default value set in metadata of the param
-	virtual void SetPinDefaultValue(UEdGraphPin* Pin, const UFunction* Function = NULL, const UProperty* Param = NULL) const;
+	/** Sets the autogenerated default value for a pin, optionally using the passed in function and parameter. This will also reset the current default value to the autogenerated one */
+	virtual void SetPinAutogeneratedDefaultValue(UEdGraphPin* Pin, const FString& NewValue) const;
 
-	/**
-	 * Sets the default value of a pin based on the type of the pin (0 for int, false for boolean, etc...)
-	 */
+	/** Sets the autogenerated default value for a pin using the default for that type. This will also reset the current default value to the autogenerated one */
+	virtual void SetPinAutogeneratedDefaultValueBasedOnType(UEdGraphPin* Pin) const;
+
+	/** Resets a pin back to it's autogenerated default value, optionally calling the default value change callbacks */
+	virtual void ResetPinToAutogeneratedDefaultValue(UEdGraphPin* Pin, bool bCallModifyCallbacks = true) const;
+
+	/** Sets the pin defaults, but not autogenerated defaults, at pin construction time. This is like TrySetDefaultValue but does not do validation or callbacks */
+	virtual void SetPinDefaultValueAtConstruction(UEdGraphPin* Pin, const FString& DefaultValueString) const;
+
+	/** Call to let blueprint and UI know that parameters have changed for a function/macro/etc */
+	virtual void HandleParameterDefaultValueChanged(UK2Node* TargetNode) const;
+
+	DEPRECATED(4.17, "SetPinDefaultValue is deprecated due to confusing name, call SetPinAutogeneratedDefaultValue")
+	virtual void SetPinDefaultValue(UEdGraphPin* Pin, const UFunction* Function = nullptr, const UProperty* Param = nullptr) const;
+
+	DEPRECATED(4.17, "SetPinDefaultValueBasedOnType is deprecated due to confusing name, call SetPinAutogeneratedDefaultValue")
 	virtual void SetPinDefaultValueBasedOnType(UEdGraphPin* Pin) const;
 
+	/** Given a function and property, return the default value */
+	static bool FindFunctionParameterDefaultValue(const UFunction* Function, const UProperty* Param, FString& OutString);
+
 	/** Utility that makes sure existing connections are valid, breaking any that are now illegal. */
 	static void ValidateExistingConnections(UEdGraphPin* Pin);
 
@@ -1068,6 +1087,7 @@ public:
 	static void GetReplaceVariableMenu(class FMenuBuilder& MenuBuilder, class UK2Node_Variable* Variable, UBlueprint* OwnerBlueprint, bool bReplaceExistingVariable = false);
 
 private:
+
 	/**
 	 * Returns true if the specified function has any out parameters
 	 * @param [in] Function	The function to check for out parameters
diff --git a/Engine/Source/Editor/BlueprintGraph/Classes/K2Node.h b/Engine/Source/Editor/BlueprintGraph/Classes/K2Node.h
--- a/Engine/Source/Editor/BlueprintGraph/Classes/K2Node.h
+++ b/Engine/Source/Editor/BlueprintGraph/Classes/K2Node.h
@@ -43,44 +43,44 @@ struct FOptionalPinFromProperty
 {
 	GENERATED_USTRUCT_BODY()
 
-	UPROPERTY(EditAnywhere, Category= OptionalPin, BlueprintReadOnly)
+	UPROPERTY(EditAnywhere, Category= OptionalPin)
 	FName PropertyName;
 
-	UPROPERTY(EditAnywhere, Category= OptionalPin, BlueprintReadOnly)
+	UPROPERTY(EditAnywhere, Category= OptionalPin)
 	FString PropertyFriendlyName;
 
-	UPROPERTY(EditAnywhere, Category= OptionalPin, BlueprintReadOnly)
+	UPROPERTY(EditAnywhere, Category= OptionalPin)
 	FText PropertyTooltip;
 
-	UPROPERTY(EditAnywhere, BlueprintReadWrite, Category=OptionalPin)
-	bool bShowPin;
+	UPROPERTY(EditAnywhere, Category=OptionalPin)
+	FName CategoryName;
 
-	UPROPERTY(EditAnywhere, BlueprintReadWrite, Category=OptionalPin)
-	bool bCanToggleVisibility;
+	UPROPERTY(EditAnywhere, Category=OptionalPin)
+	uint8 bShowPin:1;
 
-	UPROPERTY(EditAnywhere, BlueprintReadWrite, Category=OptionalPin)
-	bool bPropertyIsCustomized;
-	
-	UPROPERTY(EditAnywhere, BlueprintReadWrite, Category=OptionalPin)
-	FName CategoryName;
+	UPROPERTY(EditAnywhere, Category=OptionalPin)
+	uint8 bCanToggleVisibility:1;
 
 	UPROPERTY(EditAnywhere, Category=OptionalPin)
-	bool bHasOverridePin;
+	uint8 bPropertyIsCustomized:1;
+	
+	UPROPERTY(EditAnywhere, Category=OptionalPin)
+	uint8 bHasOverridePin:1;
 
 	UPROPERTY(EditAnywhere, Category=OptionalPin)
-	bool bIsMarkedForAdvancedDisplay;
+	uint8 bIsMarkedForAdvancedDisplay:1;
 
 	/** TRUE if the override value is enabled for use */
 	UPROPERTY(EditAnywhere, Category = OptionalPin)
-	bool bIsOverrideEnabled;
+	uint8 bIsOverrideEnabled:1;
 
 	/** TRUE if the override value should be set through this pin */
 	UPROPERTY(EditAnywhere, Category = OptionalPin)
-	bool bIsSetValuePinVisible;
+	uint8 bIsSetValuePinVisible:1;
 
 	/** TRUE if the override pin is visible */
 	UPROPERTY(EditAnywhere, Category = OptionalPin)
-	bool bIsOverridePinVisible;
+	uint8 bIsOverridePinVisible:1;
 
 	FOptionalPinFromProperty()
 		: bIsMarkedForAdvancedDisplay(false)
@@ -96,10 +96,10 @@ struct FOptionalPinFromProperty
 #if WITH_EDITORONLY_DATA
 		, PropertyTooltip(InTooltip)
 #endif
+		, CategoryName(InCategoryName)
 		, bShowPin(bInShowPin)
 		, bCanToggleVisibility(bInCanToggleVisibility)
 		, bPropertyIsCustomized(bInPropertyIsCustomized)
-		, CategoryName(InCategoryName)
 		, bHasOverridePin(bInHasOverridePin)
 		, bIsMarkedForAdvancedDisplay(false)
 		, bIsOverrideEnabled(true)
@@ -115,26 +115,36 @@ struct BLUEPRINTGRAPH_API FOptionalPinManager
 public:
 	virtual ~FOptionalPinManager() { }
 
-	// Should the specified property be displayed by default
+	/** Should the specified property be displayed by default */
 	virtual void GetRecordDefaults(UProperty* TestProperty, FOptionalPinFromProperty& Record) const;
 
-	// Can this property be managed as an optional pin (with the ability to be shown or hidden)
+	/** Can this property be managed as an optional pin (with the ability to be shown or hidden) */
 	virtual bool CanTreatPropertyAsOptional(UProperty* TestProperty) const;
 
-	// Reconstructs the specified property array using the SourceStruct
-	// @param [in,out] Properties	The property array
-	// @param SourceStruct			The source structure to update the properties array from
-	// @param bDefaultVisibility
+	/** 
+	 * Reconstructs the specified property array using the SourceStruct
+	 * @param [in,out] Properties	The property array
+	 * @param SourceStruct			The source structure to update the properties array from
+	 * @param bDefaultVisibility
+	 */
 	void RebuildPropertyList(TArray<FOptionalPinFromProperty>& Properties, UStruct* SourceStruct);
 
-	// Creates a pin for each visible property on the specified node
-	void CreateVisiblePins(TArray<FOptionalPinFromProperty>& Properties, UStruct* SourceStruct, EEdGraphPinDirection Direction, class UK2Node* TargetNode, uint8* StructBasePtr = NULL);
+	/**
+	 * Creates a pin for each visible property on the specified node
+	 * @param Properties	The property array
+	 * @param SourceStruct	The source structure to update the properties array from
+	 * @param Direction		Direction of pins to create
+	 * @param TargetNode	Node to generate pins for
+	 * @param StructBasePtr	Struct that is iterated over, to represent the current version of the data
+	 * @param DefaultsPtr	Struct that is iterated over, to represent the default version of the data, used for delta serialization
+	 */
+	void CreateVisiblePins(TArray<FOptionalPinFromProperty>& Properties, UStruct* SourceStruct, EEdGraphPinDirection Direction, class UK2Node* TargetNode, uint8* StructBasePtr = nullptr, uint8* DefaultsPtr = nullptr);
 
 	// Customize automatically created pins if desired
 	virtual void CustomizePinData(UEdGraphPin* Pin, FName SourcePropertyName, int32 ArrayIndex, UProperty* Property = NULL) const {}
 protected:
-	virtual void PostInitNewPin(UEdGraphPin* Pin, FOptionalPinFromProperty& Record, int32 ArrayIndex, UProperty* Property, uint8* PropertyAddress) const {}
-	virtual void PostRemovedOldPin(FOptionalPinFromProperty& Record, int32 ArrayIndex, UProperty* Property, uint8* PropertyAddress) const {}
+	virtual void PostInitNewPin(UEdGraphPin* Pin, FOptionalPinFromProperty& Record, int32 ArrayIndex, UProperty* Property, uint8* PropertyAddress, uint8* DefaultPropertyAddress) const {}
+	virtual void PostRemovedOldPin(FOptionalPinFromProperty& Record, int32 ArrayIndex, UProperty* Property, uint8* PropertyAddress, uint8* DefaultPropertyAddress) const {}
 	void RebuildProperty(UProperty* TestProperty, FName CategoryName, TArray<FOptionalPinFromProperty>& Properties, UStruct* SourceStruct, TMap<FName, FOldOptionalPinSettings>& OldSettings);
 };
 
diff --git a/Engine/Source/Editor/BlueprintGraph/Classes/K2Node_CustomEvent.h b/Engine/Source/Editor/BlueprintGraph/Classes/K2Node_CustomEvent.h
--- a/Engine/Source/Editor/BlueprintGraph/Classes/K2Node_CustomEvent.h
+++ b/Engine/Source/Editor/BlueprintGraph/Classes/K2Node_CustomEvent.h
@@ -52,6 +52,7 @@ class UK2Node_CustomEvent : public UK2Node_Event
 	//~ Begin UK2Node_EditablePinBase Interface
 	virtual UEdGraphPin* CreatePinFromUserDefinition(const TSharedPtr<FUserPinInfo> NewPinInfo) override;
 	virtual bool CanCreateUserDefinedPin(const FEdGraphPinType& InPinType, EEdGraphPinDirection InDesiredDirection, FText& OutErrorMessage) override;
+	virtual bool ModifyUserDefinedPinDefaultValue(TSharedPtr<FUserPinInfo> PinInfo, const FString& NewDefaultValue) override;
 	//~ Begin UK2Node_EditablePinBase Interface
 
 	virtual bool IsUsedByAuthorityOnlyDelegate() const override;
diff --git a/Engine/Source/Editor/BlueprintGraph/Classes/K2Node_EditablePinBase.h b/Engine/Source/Editor/BlueprintGraph/Classes/K2Node_EditablePinBase.h
--- a/Engine/Source/Editor/BlueprintGraph/Classes/K2Node_EditablePinBase.h
+++ b/Engine/Source/Editor/BlueprintGraph/Classes/K2Node_EditablePinBase.h
@@ -33,31 +33,8 @@ struct FUserPinInfo
 	UPROPERTY()
 	FString PinDefaultValue;
 
-	friend FArchive& operator <<(FArchive& Ar, FUserPinInfo& Info)
-	{
-		Ar << Info.PinName;
-
-		if (Ar.UE4Ver() >= VER_UE4_SERIALIZE_PINTYPE_CONST)
-		{
-			Info.PinType.Serialize(Ar);
-			Ar << Info.DesiredPinDirection;
-		}
-		else
-		{
-			Ar << Info.PinType.bIsArray;
-			Ar << Info.PinType.bIsReference;
-
-			Ar << Info.PinType.PinCategory;
-			Ar << Info.PinType.PinSubCategory;
-
-			Ar << Info.PinType.PinSubCategoryObject;
-		}
-		
-		Ar << Info.PinDefaultValue;
-
-		return Ar;
-	}
-	
+	friend FArchive& operator <<(FArchive& Ar, FUserPinInfo& Info);
+
 };
 
 // This structure describes metadata associated with a user declared function or macro
@@ -120,6 +97,7 @@ class UK2Node_EditablePinBase : public UK2Node
 
 	// UEdGraphNode interface
 	BLUEPRINTGRAPH_API virtual void AllocateDefaultPins() override;
+	BLUEPRINTGRAPH_API virtual void PinDefaultValueChanged(UEdGraphPin* Pin) override;
 	// End of UEdGraphNode interface
 
 	// UK2Node interface
diff --git a/Engine/Source/Editor/BlueprintGraph/Classes/K2Node_FunctionEntry.h b/Engine/Source/Editor/BlueprintGraph/Classes/K2Node_FunctionEntry.h
--- a/Engine/Source/Editor/BlueprintGraph/Classes/K2Node_FunctionEntry.h
+++ b/Engine/Source/Editor/BlueprintGraph/Classes/K2Node_FunctionEntry.h
@@ -54,6 +54,7 @@ class UK2Node_FunctionEntry : public UK2Node_FunctionTerminator
 
 	//~ Begin UK2Node_EditablePinBase Interface
 	virtual bool CanUseRefParams() const override { return true; }
+	virtual bool ModifyUserDefinedPinDefaultValue(TSharedPtr<FUserPinInfo> PinInfo, const FString& NewDefaultValue) override;
 	//~ End UK2Node_EditablePinBase Interface
 
 	//~ Begin K2Node_FunctionTerminator Interface
diff --git a/Engine/Source/Editor/BlueprintGraph/Classes/K2Node_GenericCreateObject.h b/Engine/Source/Editor/BlueprintGraph/Classes/K2Node_GenericCreateObject.h
--- a/Engine/Source/Editor/BlueprintGraph/Classes/K2Node_GenericCreateObject.h
+++ b/Engine/Source/Editor/BlueprintGraph/Classes/K2Node_GenericCreateObject.h
@@ -16,7 +16,6 @@ class BLUEPRINTGRAPH_API UK2Node_GenericCreateObject : public UK2Node_ConstructO
 
 	//~ Begin UEdGraphNode Interface.
 	virtual void ExpandNode(class FKismetCompilerContext& CompilerContext, UEdGraph* SourceGraph) override;
-	virtual void ValidateNodeDuringCompilation(class FCompilerResultsLog& MessageLog) const override;
 	virtual void EarlyValidation(class FCompilerResultsLog& MessageLog) const override;
 	virtual bool IsCompatibleWithGraph(const UEdGraph* TargetGraph) const override;
 	//~ End UEdGraphNode Interface.
diff --git a/Engine/Source/Editor/BlueprintGraph/Classes/K2Node_StructOperation.h b/Engine/Source/Editor/BlueprintGraph/Classes/K2Node_StructOperation.h
--- a/Engine/Source/Editor/BlueprintGraph/Classes/K2Node_StructOperation.h
+++ b/Engine/Source/Editor/BlueprintGraph/Classes/K2Node_StructOperation.h
@@ -25,8 +25,8 @@ class UK2Node_StructOperation : public UK2Node_Variable
 	//~ Begin UK2Node Interface
 	//virtual bool DrawNodeAsVariable() const override { return true; }
 	virtual bool ShouldShowNodeProperties() const override { return true; }
-	virtual void ValidateNodeDuringCompilation(class FCompilerResultsLog& MessageLog) const override {}
-	virtual bool HasExternalDependencies(TArray<class UStruct*>* OptionalOutput = NULL) const override;
+	virtual void ValidateNodeDuringCompilation(class FCompilerResultsLog& MessageLog) const override;
+	virtual bool HasExternalDependencies(TArray<class UStruct*>* OptionalOutput = nullptr) const override;
 	virtual FString GetFindReferenceSearchString() const override;
 	virtual bool IsActionFilteredOut(class FBlueprintActionFilter const& Filter) override;
 	//~ End UK2Node Interface
diff --git a/Engine/Source/Editor/BlueprintGraph/Classes/K2Node_Switch.h b/Engine/Source/Editor/BlueprintGraph/Classes/K2Node_Switch.h
--- a/Engine/Source/Editor/BlueprintGraph/Classes/K2Node_Switch.h
+++ b/Engine/Source/Editor/BlueprintGraph/Classes/K2Node_Switch.h
@@ -67,6 +67,9 @@ class UK2Node_Switch : public UK2Node
 	 */
 	BLUEPRINTGRAPH_API virtual void RemovePinFromSwitchNode(UEdGraphPin* TargetPin);
 
+	/** Whether an execution pin can be removed from the node or not */
+	BLUEPRINTGRAPH_API virtual bool CanRemoveExecutionPin(UEdGraphPin* TargetPin) const;
+
 	/** Getting pin access */
 	BLUEPRINTGRAPH_API UEdGraphPin* GetSelectionPin() const;
 	BLUEPRINTGRAPH_API UEdGraphPin* GetDefaultPin() const;
diff --git a/Engine/Source/Editor/BlueprintGraph/Classes/K2Node_SwitchInteger.h b/Engine/Source/Editor/BlueprintGraph/Classes/K2Node_SwitchInteger.h
--- a/Engine/Source/Editor/BlueprintGraph/Classes/K2Node_SwitchInteger.h
+++ b/Engine/Source/Editor/BlueprintGraph/Classes/K2Node_SwitchInteger.h
@@ -39,6 +39,7 @@ class UK2Node_SwitchInteger : public UK2Node_Switch
 	virtual FString GetPinNameGivenIndex(int32 Index) override;
 	virtual void ReallocatePinsDuringReconstruction(TArray<UEdGraphPin*>& OldPins) override;
 	virtual FEdGraphPinType GetPinType() const override;
+	virtual bool CanRemoveExecutionPin(UEdGraphPin* TargetPin) const override;
 	// End of UK2Node_Switch Interface
 
 protected:
diff --git a/Engine/Source/Editor/BlueprintGraph/Classes/K2Node_SwitchString.h b/Engine/Source/Editor/BlueprintGraph/Classes/K2Node_SwitchString.h
--- a/Engine/Source/Editor/BlueprintGraph/Classes/K2Node_SwitchString.h
+++ b/Engine/Source/Editor/BlueprintGraph/Classes/K2Node_SwitchString.h
@@ -30,6 +30,7 @@ class UK2Node_SwitchString : public UK2Node_Switch
 	virtual FText GetTooltipText() const override;
 	virtual FText GetNodeTitle(ENodeTitleType::Type TitleType) const override;
 	virtual bool ShouldShowNodeProperties() const override { return true; }
+	virtual void ValidateNodeDuringCompilation(class FCompilerResultsLog& MessageLog) const override;
 	// End of UEdGraphNode interface
 
 	// UK2Node interface
diff --git a/Engine/Source/Editor/BlueprintGraph/Classes/K2Node_TunnelBoundary.h b/Engine/Source/Editor/BlueprintGraph/Classes/K2Node_TunnelBoundary.h
--- a/Engine/Source/Editor/BlueprintGraph/Classes/K2Node_TunnelBoundary.h
+++ b/Engine/Source/Editor/BlueprintGraph/Classes/K2Node_TunnelBoundary.h
@@ -65,10 +65,10 @@ protected:
 	void WireUpTunnelExit(UK2Node_Tunnel* TunnelInstance, UEdGraphPin* TunnelPin, FCompilerResultsLog& MessageLog);
 
 	/** Wires up the entry boundries */
-	void WireUpEntry(UEdGraphNode* SourceNode, UEdGraphPin* SourcePin, TArray<UEdGraphPin*>& EntryPins, FCompilerResultsLog& MessageLog);
+	void WireUpEntry(UEdGraphNode* SourceNode, UEdGraphPin* SourcePin, const TArray<UEdGraphPin*>& EntryPins, FCompilerResultsLog& MessageLog);
 
 	/** Wires up the exit boundries */
-	void WireUpExit(UEdGraphNode* SourceNode, UEdGraphPin* SourcePin, TArray<UEdGraphPin*>& ExitPins, FCompilerResultsLog& MessageLog);
+	void WireUpExit(UEdGraphNode* SourceNode, UEdGraphPin* SourcePin, const TArray<UEdGraphPin*>& ExitPins, FCompilerResultsLog& MessageLog);
 
 	/** Create base node name */
 	void CreateBaseNodeName(UEdGraphNode* SourceNode);
diff --git a/Engine/Source/Editor/BlueprintGraph/Private/BlueprintActionDatabase.cpp b/Engine/Source/Editor/BlueprintGraph/Private/BlueprintActionDatabase.cpp
--- a/Engine/Source/Editor/BlueprintGraph/Private/BlueprintActionDatabase.cpp
+++ b/Engine/Source/Editor/BlueprintGraph/Private/BlueprintActionDatabase.cpp
@@ -755,12 +755,7 @@ static void BlueprintActionDatabaseImpl::AddBlueprintGraphActions(UBlueprint con
 
 		for (UK2Node_FunctionEntry* FunctionEntry : GraphEntryNodes)
 		{
-			// Find the initial place where the function was defined, so we use the most generous scope
 			UFunction* SkeletonFunction = FindField<UFunction>(Blueprint->SkeletonGeneratedClass, FunctionGraph->GetFName());
-			for (UFunction* ParentFunction = SkeletonFunction; ParentFunction != nullptr; ParentFunction = ParentFunction->GetSuperFunction())
-			{
-				SkeletonFunction = ParentFunction;
-			}
 
 			// Create entries for function parameters
 			if (SkeletonFunction != nullptr)
diff --git a/Engine/Source/Editor/BlueprintGraph/Private/BlueprintActionFilter.cpp b/Engine/Source/Editor/BlueprintGraph/Private/BlueprintActionFilter.cpp
--- a/Engine/Source/Editor/BlueprintGraph/Private/BlueprintActionFilter.cpp
+++ b/Engine/Source/Editor/BlueprintGraph/Private/BlueprintActionFilter.cpp
@@ -1260,7 +1260,7 @@ static bool BlueprintActionFilterImpl::IsPinCompatibleWithTargetSelf(UEdGraphPin
 			if (IsClassOfType(PinObjClass, TargetClass))
 			{
 				bIsCompatible = true;
-				if (PinType.bIsArray)
+				if (PinType.IsArray())
 				{
 					if (UFunction const* Function = BlueprintAction.GetAssociatedFunction())
 					{
@@ -1276,7 +1276,7 @@ static bool BlueprintActionFilterImpl::IsPinCompatibleWithTargetSelf(UEdGraphPin
 					}
 				}
 			}
-			else if (!PinType.bIsArray && (BlueprintAction.GetNodeClass() == UK2Node_CallFunction::StaticClass()))
+			else if (!PinType.IsArray() && (BlueprintAction.GetNodeClass() == UK2Node_CallFunction::StaticClass()))
 			{
 				// if this is a bound CallFunction action, then we make the 
 				// assumption that it will be turned into a UK2Node_CallFunctionOnMember
@@ -1359,7 +1359,7 @@ static bool BlueprintActionFilterImpl::ArrayFunctionHasParamOfType(const UFuncti
 	FBlueprintEditorUtils::GetHiddenPinsForFunction(InGraph, ArrayFunction, HiddenPins);
 
 	FName ParamTag = FBlueprintMetadata::MD_ArrayDependentParam;
-	if (DesiredPinType.bIsArray)
+	if (DesiredPinType.IsArray())
 	{
 		ParamTag = FBlueprintMetadata::MD_ArrayParam;
 	}
@@ -1594,8 +1594,7 @@ static bool BlueprintActionFilterImpl::IsExtraneousInterfaceCall(FBlueprintActio
 		UClass* InterfaceClass = Function->GetOwnerClass();
 		checkSlow(InterfaceClass->IsChildOf<UInterface>());
 
-		bool const bIsAbstractCppClass = InterfaceClass->GetCppTypeInfo()->IsAbstract();
-		bool const bCanBeAddedToBlueprints = !bIsAbstractCppClass && !InterfaceClass->HasMetaData(FBlueprintMetadata::MD_CannotImplementInterfaceInBlueprint);
+		bool const bCanBeAddedToBlueprints = !InterfaceClass->HasMetaData(FBlueprintMetadata::MD_CannotImplementInterfaceInBlueprint);
 
 		bIsFilteredOut = (Filter.TargetClasses.Num() > 0);
 		for (const auto& ClassData : Filter.TargetClasses)
@@ -1762,7 +1761,7 @@ FBlueprintActionInfo::FBlueprintActionInfo(FBlueprintActionInfo const& Rhs, IBlu
 //------------------------------------------------------------------------------
 UObject const* FBlueprintActionInfo::GetActionOwner()
 {
-	return ActionOwner;
+	return ActionOwner.Get();
 }
 
 //------------------------------------------------------------------------------
@@ -1776,19 +1775,22 @@ UClass const* FBlueprintActionInfo::GetOwnerClass()
 {
 	if ((CacheFlags & EBlueprintActionInfoFlags::CachedClass) == 0)
 	{
-		CachedOwnerClass = Cast<UClass>(ActionOwner);
+		CachedOwnerClass = Cast<UClass>(ActionOwner.Get());
 		if (CachedOwnerClass == GetNodeClass())
 		{
 			CachedOwnerClass = nullptr;
 		}
-		else if (const UBlueprint* AsBlueprint = Cast<UBlueprint>(ActionOwner))
+		else if (const UBlueprint* AsBlueprint = Cast<UBlueprint>(ActionOwner.Get()))
 		{
 			CachedOwnerClass = AsBlueprint->SkeletonGeneratedClass;
 		}
 
 		if (CachedOwnerClass == nullptr)
 		{
-			CachedOwnerClass = GetAssociatedMemberField()->GetOwnerClass();
+			if (UField const* AssociatedMemberField = GetAssociatedMemberField())
+			{
+				CachedOwnerClass = AssociatedMemberField->GetOwnerClass();
+			}
 		}
 
 		CacheFlags |= EBlueprintActionInfoFlags::CachedClass;
diff --git a/Engine/Source/Editor/BlueprintGraph/Private/CallFunctionHandler.cpp b/Engine/Source/Editor/BlueprintGraph/Private/CallFunctionHandler.cpp
--- a/Engine/Source/Editor/BlueprintGraph/Private/CallFunctionHandler.cpp
+++ b/Engine/Source/Editor/BlueprintGraph/Private/CallFunctionHandler.cpp
@@ -120,31 +120,7 @@ void FKCHandler_CallFunction::CreateFunctionCallStatement(FKismetFunctionContext
 						{
 							check((*Term)->bIsLiteral);
 						
-							int32 LatentUUID = INDEX_NONE;
-							UEdGraphNode* AssociatedNode = LatentInfoPin->GetOwningNode();
-							if (AssociatedNode && AssociatedNode->NodeGuid.IsValid())
-							{
-								LatentUUID = GetTypeHash(AssociatedNode->NodeGuid);
-
-								UEdGraphNode* ResultNode = AssociatedNode;
-								UEdGraphNode* SourceNode = Cast<UEdGraphNode>(Context.MessageLog.GetIntermediateTunnelInstance(AssociatedNode));
-								while (SourceNode && SourceNode != ResultNode)
-								{
-									if (SourceNode->NodeGuid.IsValid())
-									{
-										LatentUUID = HashCombine(LatentUUID, GetTypeHash(SourceNode->NodeGuid));
-									}
-									ResultNode = SourceNode;
-									SourceNode = Cast<UEdGraphNode>(Context.MessageLog.GetIntermediateTunnelInstance(ResultNode));
-								}
-							}
-							else
-							{
-								static int32 FallbackUUID = 0;
-								LatentUUID = FallbackUUID++;
-
-								CompilerContext.MessageLog.Warning(*LOCTEXT("UUIDDeterministicCookWarn", "Failed to produce a deterministic UUID for a node's latent action: @@. Cooking this Blueprint (@@) asset will non-deterministic.").ToString(), LatentInfoPin, Context.Blueprint);
-							}							
+							int32 LatentUUID = CompilerContext.MessageLog.CalculateStableIdentifierForLatentActionManager(LatentInfoPin->GetOwningNode());
 
 							const FString ExecutionFunctionName = CompilerContext.GetSchema()->FN_ExecuteUbergraphBase.ToString() + TEXT("_") + Context.Blueprint->GetName();
 							(*Term)->Name = FString::Printf(TEXT("(Linkage=%s,UUID=%s,ExecutionFunction=%s,CallbackTarget=None)"), *FString::FromInt(INDEX_NONE), *FString::FromInt(LatentUUID), *ExecutionFunctionName);
diff --git a/Engine/Source/Editor/BlueprintGraph/Private/EdGraphSchema_K2.cpp b/Engine/Source/Editor/BlueprintGraph/Private/EdGraphSchema_K2.cpp
--- a/Engine/Source/Editor/BlueprintGraph/Private/EdGraphSchema_K2.cpp
+++ b/Engine/Source/Editor/BlueprintGraph/Private/EdGraphSchema_K2.cpp
@@ -1493,6 +1493,11 @@ void UEdGraphSchema_K2::GetContextMenuActions(const UEdGraph* CurrentGraph, cons
 					MenuBuilder->AddMenuEntry(FGraphEditorCommands::Get().RemoveThisStructVarPin);
 					MenuBuilder->AddMenuEntry(FGraphEditorCommands::Get().RemoveOtherStructVarPins);
 				}
+
+				if (InGraphPin->PinType.PinCategory != PC_Exec && InGraphPin->Direction == EGPD_Input && InGraphPin->LinkedTo.Num() == 0 && !ShouldHidePinDefaultValue(const_cast<UEdGraphPin*>(InGraphPin)))
+				{
+					MenuBuilder->AddMenuEntry(FGraphEditorCommands::Get().ResetPinToDefaultValue);
+				}
 			}
 		}
 		MenuBuilder->EndSection();
@@ -2252,11 +2257,9 @@ const FPinConnectionResponse UEdGraphSchema_K2::CanCreateConnection(const UEdGra
 	if(const UK2Node* OwningNode = Cast<UK2Node>(InputPin->GetOwningNode()))
 	{
 		const bool bAllowMultipleSelfs = OwningNode->AllowMultipleSelfs(true); // it applies also to ForEachCall
-		const bool bNotAnArrayFunction = !InputPin->PinType.bIsArray;
-		const bool bNotAMapFunction = !InputPin->PinType.bIsMap;
-		const bool bNotASetFunction = !InputPin->PinType.bIsSet;
+		const bool bNotAContainer = !InputPin->PinType.IsContainer();
 		const bool bSelfPin = IsSelfPin(*InputPin);
-		bIgnoreArray = bAllowMultipleSelfs && bNotAnArrayFunction && bNotAMapFunction && bNotASetFunction  && bSelfPin;
+		bIgnoreArray = bAllowMultipleSelfs && bNotAContainer && bSelfPin;
 	}
 
 	// Find the calling context in case one of the pins is of type object and has a value of Self
@@ -2493,11 +2496,9 @@ bool UEdGraphSchema_K2::SearchForAutocastFunction(const UEdGraphPin* OutputPin,
 	TargetFunction = NAME_None;
 	FunctionOwner = nullptr;
 
-	if (OutputPin->PinType.bIsArray != InputPin->PinType.bIsArray ||
-		OutputPin->PinType.bIsMap != InputPin->PinType.bIsMap ||
-		OutputPin->PinType.bIsSet != InputPin->PinType.bIsSet)
+	if (OutputPin->PinType.ContainerType != InputPin->PinType.ContainerType)
 	{
-		if (OutputPin->PinType.bIsSet && InputPin->PinType.bIsArray)
+		if (OutputPin->PinType.IsSet() && InputPin->PinType.IsArray())
 		{
 			UFunction* Function = UBlueprintSetLibrary::StaticClass()->FindFunctionByName(GET_MEMBER_NAME_CHECKED(UBlueprintSetLibrary, Set_ToArray));
 			TargetFunction = Function->GetFName();
@@ -2512,7 +2513,7 @@ bool UEdGraphSchema_K2::SearchForAutocastFunction(const UEdGraphPin* OutputPin,
 	{
 		UClass const* InputClass = Cast<UClass const>(InputPin->PinType.PinSubCategoryObject.Get());
 
-		bool const bInputIsUObject = ((InputClass != NULL) && (InputClass == UObject::StaticClass()));
+		bool const bInputIsUObject = (InputClass && (InputClass == UObject::StaticClass()));
 		if (bInputIsUObject)
 		{
 			UFunction* Function = UKismetSystemLibrary::StaticClass()->FindFunctionByName(GET_MEMBER_NAME_CHECKED(UKismetSystemLibrary, Conv_InterfaceToObject));
@@ -2576,7 +2577,7 @@ bool UEdGraphSchema_K2::FindSpecializedConversionNode(const UEdGraphPin* OutputP
 	TargetNode = nullptr;
 
 	// Conversion for scalar -> array
-	if( (!OutputPin->PinType.IsContainer() && InputPin->PinType.bIsArray) && ArePinTypesCompatible(OutputPin->PinType, InputPin->PinType, NULL, true))
+	if( (!OutputPin->PinType.IsContainer() && InputPin->PinType.IsArray()) && ArePinTypesCompatible(OutputPin->PinType, InputPin->PinType, nullptr, true))
 	{
 		bCanConvert = true;
 		if(bCreateNode)
@@ -2646,8 +2647,8 @@ bool UEdGraphSchema_K2::FindSpecializedConversionNode(const UEdGraphPin* OutputP
 
 		// CHECK BYTE TO ENUM CAST
 		UEnum* Enum = Cast<UEnum>(InputType.PinSubCategoryObject.Get());
-		const bool bInputIsEnum = !InputType.bIsArray && (PC_Byte == InputType.PinCategory) && Enum;
-		const bool bOutputIsByte = !OutputType.bIsArray && (PC_Byte == OutputType.PinCategory);
+		const bool bInputIsEnum = !InputType.IsContainer() && (PC_Byte == InputType.PinCategory) && Enum;
+		const bool bOutputIsByte = !OutputType.IsContainer() && (PC_Byte == OutputType.PinCategory);
 		if (bInputIsEnum && bOutputIsByte)
 		{
 			bCanConvert = true;
@@ -2804,9 +2805,9 @@ FString UEdGraphSchema_K2::IsPinDefaultValid(const UEdGraphPin* Pin, const FStri
 		return FText::Format(MsgFormat, MessageArgs).ToString();
 	}
 
-	const bool bIsArray = Pin->PinType.bIsArray;
-	const bool bIsSet = Pin->PinType.bIsSet;
-	const bool bIsMap = Pin->PinType.bIsMap;
+	const bool bIsArray = Pin->PinType.IsArray();
+	const bool bIsSet = Pin->PinType.IsSet();
+	const bool bIsMap = Pin->PinType.IsMap();
 	const bool bIsReference = Pin->PinType.bIsReference;
 	const bool bIsAutoCreateRefTerm = IsAutoCreateRefTerm(Pin);
 
@@ -2975,7 +2976,7 @@ FLinearColor UEdGraphSchema_K2::GetPinTypeColor(const FEdGraphPinType& PinType)
 
 FLinearColor UEdGraphSchema_K2::GetSecondaryPinTypeColor(const FEdGraphPinType& PinType) const
 {
-	if (PinType.bIsMap)
+	if (PinType.IsMap())
 	{
 		FEdGraphPinType FakePrimary = PinType;
 		FakePrimary.PinCategory = FakePrimary.PinValueType.TerminalCategory;
@@ -3320,11 +3321,11 @@ bool UEdGraphSchema_K2::GetPropertyCategoryInfo(const UProperty* TestProperty, F
 		OutCategory = PC_Struct;
 		OutSubCategoryObject = StructProperty->Struct;
 	}
-	else if (Cast<const UFloatProperty>(TestProperty) != NULL)
+	else if (TestProperty->IsA<UFloatProperty>())
 	{
 		OutCategory = PC_Float;
 	}
-	else if (Cast<const UIntProperty>(TestProperty) != NULL)
+	else if (TestProperty->IsA<UIntProperty>())
 	{
 		OutCategory = PC_Int;
 
@@ -3366,19 +3367,19 @@ bool UEdGraphSchema_K2::GetPropertyCategoryInfo(const UProperty* TestProperty, F
 			OutSubCategoryObject = EnumProperty->GetEnum();
 		}
 	}
-	else if (Cast<const UNameProperty>(TestProperty) != NULL)
+	else if (TestProperty->IsA<UNameProperty>())
 	{
 		OutCategory = PC_Name;
 	}
-	else if (Cast<const UBoolProperty>(TestProperty) != NULL)
+	else if (TestProperty->IsA<UBoolProperty>())
 	{
 		OutCategory = PC_Boolean;
 	}
-	else if (Cast<const UStrProperty>(TestProperty) != NULL)
+	else if (TestProperty->IsA<UStrProperty>())
 	{
 		OutCategory = PC_String;
 	}
-	else if (Cast<const UTextProperty>(TestProperty) != NULL)
+	else if (TestProperty->IsA<UTextProperty>())
 	{
 		OutCategory = PC_Text;
 	}
@@ -3434,9 +3435,7 @@ bool UEdGraphSchema_K2::ConvertPropertyToPinType(const UProperty* Property, /*ou
 	{
 		TestProperty = ArrayProperty->Inner;
 	}
-	TypeOut.bIsMap = (MapProperty != NULL);
-	TypeOut.bIsSet = (SetProperty != NULL);
-	TypeOut.bIsArray = (ArrayProperty != NULL);
+	TypeOut.ContainerType = FEdGraphPinType::ToPinContainerType(ArrayProperty != nullptr, SetProperty != nullptr, MapProperty != nullptr);
 	TypeOut.bIsReference = Property->HasAllPropertyFlags(CPF_OutParm|CPF_ReferenceParm);
 	TypeOut.bIsConst     = Property->HasAllPropertyFlags(CPF_ConstParm);
 
@@ -3462,7 +3461,9 @@ bool UEdGraphSchema_K2::ConvertPropertyToPinType(const UProperty* Property, /*ou
 	else
 	{
 		UObject* SubCategoryObject = nullptr;
-		bool bResult = GetPropertyCategoryInfo(TestProperty, TypeOut.PinCategory, TypeOut.PinSubCategory, SubCategoryObject, TypeOut.bIsWeakPointer);
+		bool bIsWeakPointer = false;
+		bool bResult = GetPropertyCategoryInfo(TestProperty, TypeOut.PinCategory, TypeOut.PinSubCategory, SubCategoryObject, bIsWeakPointer);
+		TypeOut.bIsWeakPointer = bIsWeakPointer;
 		TypeOut.PinSubCategoryObject = SubCategoryObject;
 		if (!bResult)
 		{
@@ -3697,7 +3698,7 @@ FText UEdGraphSchema_K2::TypeToText(const FEdGraphPinType& Type)
 {
 	FText PropertyText = TerminalTypeToText(Type.PinCategory, Type.PinSubCategory, Type.PinSubCategoryObject.Get(), Type.bIsWeakPointer);
 
-	if ( Type.bIsMap )
+	if (Type.IsMap())
 	{
 		FFormatNamedArguments Args;
 		Args.Add(TEXT("KeyTitle"), PropertyText);
@@ -3705,13 +3706,13 @@ FText UEdGraphSchema_K2::TypeToText(const FEdGraphPinType& Type)
 		Args.Add(TEXT("ValueTitle"), ValueText);
 		PropertyText = FText::Format(LOCTEXT("MapAsText", "Map of {KeyTitle}s to {ValueTitle}s"), Args);
 	}
-	else if ( Type.bIsSet )
+	else if (Type.IsSet())
 	{
 		FFormatNamedArguments Args;
 		Args.Add(TEXT("PropertyTitle"), PropertyText);
 		PropertyText = FText::Format(LOCTEXT("SetAsText", "Set of {PropertyTitle}s"), Args);
 	}
-	else if (Type.bIsArray)
+	else if (Type.IsArray())
 	{
 		FFormatNamedArguments Args;
 		Args.Add(TEXT("PropertyTitle"), PropertyText);
@@ -3837,15 +3838,15 @@ struct FWildcardArrayPinHelper
 			return true;
 		}
 
-		const UK2Node* OwningNode = InputPin ? Cast<UK2Node>(InputPin->GetOwningNode()) : NULL;
+		const UK2Node* OwningNode = InputPin ? Cast<UK2Node>(InputPin->GetOwningNode()) : nullptr;
 		const bool bInputWildcardPinAcceptsArray = !OwningNode || OwningNode->DoesInputWildcardPinAcceptArray(InputPin);
 		if (bInputWildcardPinAcceptsArray)
 		{
 			return true;
 		}
 
-		const bool bCheckInputPin = (InputPin->PinType.PinCategory == GetDefault<UEdGraphSchema_K2>()->PC_Wildcard) && !InputPin->PinType.bIsArray;
-		const bool bArrayOutputPin = OutputPin && OutputPin->PinType.bIsArray;
+		const bool bCheckInputPin = (InputPin->PinType.PinCategory == GetDefault<UEdGraphSchema_K2>()->PC_Wildcard) && !InputPin->PinType.IsArray();
+		const bool bArrayOutputPin = OutputPin && OutputPin->PinType.IsArray();
 		return !(bCheckInputPin && bArrayOutputPin);
 	}
 };
@@ -3883,7 +3884,7 @@ namespace
 	{
 		bool bResult = false;
 		bool bIsNonNativeClass = false;
-		if(UClass* TargetAsClass = const_cast<UClass*>(Cast<UClass>(InTargetStruct)))
+		if (const UClass* TargetAsClass = Cast<const UClass>(InTargetStruct))
 		{
 			InTargetStruct = TargetAsClass->GetAuthoritativeClass();
 		}
@@ -4005,7 +4006,7 @@ bool UEdGraphSchema_K2::DefaultValueSimpleValidation(const FEdGraphPinType& PinT
 			}
 		}
 	}
-	else if ((PinCategory == PC_Class) || (PinCategory == PC_AssetClass))
+	else if ((PinCategory == PC_Class))
 	{
 		// Should have an object set but no string
 		if (!NewDefaultValue.IsEmpty())
@@ -4068,7 +4069,7 @@ bool UEdGraphSchema_K2::DefaultValueSimpleValidation(const FEdGraphPinType& PinT
 	{
 		// Anything is allowed
 	}
-	else if ((PinCategory == PC_Object) || (PinCategory == PC_Interface) || (PinCategory == PC_Asset))
+	else if ((PinCategory == PC_Object) || (PinCategory == PC_Interface))
 	{
 		if (PinSubCategoryObject == NULL && (PinSubCategory != PSC_Self))
 		{
@@ -4096,10 +4097,25 @@ bool UEdGraphSchema_K2::DefaultValueSimpleValidation(const FEdGraphPinType& PinT
 		{
 			DVSV_RETURN_MSG(FString::Printf(TEXT("%s isn't a %s (specified on pin %s)"), *NewDefaultObject->GetPathName(), *ObjectClass->GetName(), *(PinName)));
 		}
+	}
+	else if ((PinCategory == PC_Asset) || (PinCategory == PC_AssetClass))
+	{
+		// Should not have an object set, should be converted to string before getting here
+		if (NewDefaultObject)
+		{
+			DVSV_RETURN_MSG(FString::Printf(TEXT("NewDefaultObject '%s' specified on object pin '%s'"), *NewDefaultObject->GetPathName(), *(PinName)));
+		}
 
-		if ((PinCategory == PC_Asset) && NewDefaultObject && !NewDefaultObject->IsAsset())
+		if (!NewDefaultValue.IsEmpty())
 		{
-			DVSV_RETURN_MSG(FString::Printf(TEXT("%s is not an asset (specified on pin %s)"), *NewDefaultObject->GetPathName(), *(PinName)));
+			FText PathReason;
+
+			if (!FPackageName::IsValidObjectPath(NewDefaultValue, &PathReason))
+			{
+				DVSV_RETURN_MSG(FString::Printf(TEXT("Asset ID '%s' is invalid format for object pin '%s':"), *NewDefaultValue, *(PinName), *PathReason.ToString()));
+			}
+
+			// Class and IsAsset validation is not foolproof for asset ids, skip
 		}
 	}
 	else if (PinCategory == PC_String)
@@ -4187,7 +4203,7 @@ bool UEdGraphSchema_K2::DefaultValueSimpleValidation(const FEdGraphPinType& PinT
 
 bool UEdGraphSchema_K2::ArePinTypesCompatible(const FEdGraphPinType& Output, const FEdGraphPinType& Input, const UClass* CallingContext, bool bIgnoreArray /*= false*/) const
 {
-	if( !bIgnoreArray && ( Output.bIsMap != Input.bIsMap || Output.bIsSet != Input.bIsSet || Output.bIsArray != Input.bIsArray ) && (Input.PinCategory != PC_Wildcard || Input.IsContainer()) )
+	if( !bIgnoreArray && ( Output.ContainerType != Input.ContainerType ) && (Input.PinCategory != PC_Wildcard || Input.IsContainer()) )
 	{
 		return false;
 	}
@@ -4197,7 +4213,7 @@ bool UEdGraphSchema_K2::ArePinTypesCompatible(const FEdGraphPinType& Output, con
 			&& (Output.PinSubCategoryObject == Input.PinSubCategoryObject)
 			&& (Output.PinSubCategoryMemberReference == Input.PinSubCategoryMemberReference))
 		{
-			if(Input.bIsMap)
+			if(Input.IsMap())
 			{
 				return 
 					Input.PinValueType.TerminalCategory == PC_Wildcard ||
@@ -4227,7 +4243,7 @@ bool UEdGraphSchema_K2::ArePinTypesCompatible(const FEdGraphPinType& Output, con
 		{
 			const UClass* OutputObject = (Output.PinSubCategory == PSC_Self) ? CallingContext : Cast<const UClass>(Output.PinSubCategoryObject.Get());
 			const UClass* InputObject = (Input.PinSubCategory == PSC_Self) ? CallingContext : Cast<const UClass>(Input.PinSubCategoryObject.Get());
-			if ((OutputObject != NULL) && (InputObject != NULL))
+			if (OutputObject && InputObject)
 			{
 				return ExtendedIsChildOf(OutputObject ,InputObject);
 			}
@@ -4239,7 +4255,7 @@ bool UEdGraphSchema_K2::ArePinTypesCompatible(const FEdGraphPinType& Output, con
 			UStruct const* OutputObject = (Output.PinSubCategory == PSC_Self) ? CallingContext : Cast<UStruct>(Output.PinSubCategoryObject.Get());
 			UStruct const* InputObject  = (Input.PinSubCategory == PSC_Self)  ? CallingContext : Cast<UStruct>(Input.PinSubCategoryObject.Get());
 
-			if ((OutputObject != NULL) && (InputObject != NULL))
+			if (OutputObject && InputObject)
 			{
 				if (Output.PinCategory == PC_Struct)
 				{
@@ -4441,57 +4457,90 @@ void UEdGraphSchema_K2::HandleGraphBeingDeleted(UEdGraph& GraphBeingRemoved) con
 	}
 }
 
-void UEdGraphSchema_K2::TrySetDefaultValue(UEdGraphPin& Pin, const FString& NewDefaultValue) const
+void UEdGraphSchema_K2::GetPinDefaultValuesFromString(const FEdGraphPinType& PinType, UObject* OwningObject, const FString& NewDefaultValue, FString& UseDefaultValue, UObject*& UseDefaultObject, FText& UseDefaultText) const
 {
-	FString UseDefaultValue;
-	UObject* UseDefaultObject = nullptr;
-	FText UseDefaultText;
-
-	if ((Pin.PinType.PinCategory == PC_Object) 
-		|| (Pin.PinType.PinCategory == PC_Class) 
-		|| (Pin.PinType.PinCategory == PC_Interface)
-		|| (Pin.PinType.PinCategory == PC_Asset)
-		|| (Pin.PinType.PinCategory == PC_AssetClass))
+	if ((PinType.PinCategory == PC_Object)
+		|| (PinType.PinCategory == PC_Class)
+		|| (PinType.PinCategory == PC_Interface))
 	{
 		FString ObjectPathLocal = NewDefaultValue;
 		ConstructorHelpers::StripObjectClass(ObjectPathLocal);
-		UseDefaultObject = FindObject<UObject>(ANY_PACKAGE, *ObjectPathLocal);
-		UseDefaultValue.Empty();
+
+		// If this is not a full object path it's a relative path so should be saved as a string
+		if (FPackageName::IsValidObjectPath(ObjectPathLocal))
+		{
+			FStringAssetReference AssetRef = ObjectPathLocal;
+			UseDefaultValue.Empty();
+			UseDefaultObject = AssetRef.TryLoad();
+			UseDefaultText = FText::GetEmpty();
+		}
+		else
+		{
+			// "None" should be saved as empty string
+			if (ObjectPathLocal == TEXT("None"))
+			{
+				ObjectPathLocal.Empty();
+			}
+
+			UseDefaultValue = ObjectPathLocal;
+			UseDefaultObject = nullptr;
+			UseDefaultText = FText::GetEmpty();
+		}
 	}
-	else if(Pin.PinType.PinCategory == PC_Text)
+	else if (PinType.PinCategory == PC_Text)
 	{
 		FString PackageNamespace;
 #if USE_STABLE_LOCALIZATION_KEYS
 		if (GIsEditor)
 		{
-			PackageNamespace = TextNamespaceUtil::EnsurePackageNamespace(Pin.GetOwningNodeUnchecked());
+			PackageNamespace = TextNamespaceUtil::EnsurePackageNamespace(OwningObject);
 		}
 #endif // USE_STABLE_LOCALIZATION_KEYS
-		FText NewTextValue;
-		if (!FTextStringHelper::ReadFromString(*NewDefaultValue, NewTextValue, nullptr, *PackageNamespace))
+		if (!FTextStringHelper::ReadFromString(*NewDefaultValue, UseDefaultText, nullptr, *PackageNamespace))
 		{
-			NewTextValue = FText::FromString(NewDefaultValue);
+			UseDefaultText = FText::FromString(NewDefaultValue);
 		}
-		TrySetDefaultText(Pin, NewTextValue);
 		UseDefaultObject = nullptr;
 		UseDefaultValue.Empty();
-		return;
 	}
 	else
 	{
-		UseDefaultObject = nullptr;
 		UseDefaultValue = NewDefaultValue;
+		UseDefaultObject = nullptr;
+		UseDefaultText = FText::GetEmpty();
+
+		if (PinType.PinCategory == PC_Byte && UseDefaultValue.IsEmpty())
+		{
+			UEnum* EnumPtr = Cast<UEnum>(PinType.PinSubCategoryObject.Get());
+			if (EnumPtr)
+			{
+				// Enums are stored as empty string in autogenerated defaults, but should turn into the first value in array 
+				UseDefaultValue = EnumPtr->GetNameStringByIndex(0);
+			}
+		}
+		else if ((PinType.PinCategory == PC_Asset) || (PinType.PinCategory == PC_AssetClass))
+		{
+			ConstructorHelpers::StripObjectClass(UseDefaultValue);
+		}
 	}
+}
+
+void UEdGraphSchema_K2::TrySetDefaultValue(UEdGraphPin& Pin, const FString& NewDefaultValue) const
+{
+	FString UseDefaultValue;
+	UObject* UseDefaultObject = nullptr;
+	FText UseDefaultText;
+
+	GetPinDefaultValuesFromString(Pin.PinType, Pin.GetOwningNodeUnchecked(), NewDefaultValue, UseDefaultValue, UseDefaultObject, UseDefaultText);
 
 	// Check the default value and make it an error if it's bogus
-	if (IsPinDefaultValid(&Pin, UseDefaultValue, UseDefaultObject, UseDefaultText) == TEXT(""))
+	if (IsPinDefaultValid(&Pin, UseDefaultValue, UseDefaultObject, UseDefaultText).IsEmpty())
 	{
 		Pin.DefaultObject = UseDefaultObject;
 		Pin.DefaultValue = UseDefaultValue;
 		Pin.DefaultTextValue = UseDefaultText;
 
 		UEdGraphNode* Node = Pin.GetOwningNode();
-		check(Node);
 		Node->PinDefaultValueChanged(&Pin);
 
 		UBlueprint* Blueprint = FBlueprintEditorUtils::FindBlueprintForNodeChecked(Node);
@@ -4503,16 +4552,21 @@ void UEdGraphSchema_K2::TrySetDefaultObject(UEdGraphPin& Pin, UObject* NewDefaul
 {
 	FText UseDefaultText;
 
+	if ((Pin.PinType.PinCategory == PC_Asset) || (Pin.PinType.PinCategory == PC_AssetClass))
+	{
+		TrySetDefaultValue(Pin, NewDefaultObject ? NewDefaultObject->GetPathName() : FString());
+		return;
+	}
+
 	// Check the default value and make it an error if it's bogus
-	if (IsPinDefaultValid(&Pin, FString(TEXT("")), NewDefaultObject, UseDefaultText) == TEXT(""))
+	if (IsPinDefaultValid(&Pin, FString(), NewDefaultObject, UseDefaultText).IsEmpty())
 	{
 		Pin.DefaultObject = NewDefaultObject;
-		Pin.DefaultValue = NULL;
+		Pin.DefaultValue.Empty();
 		Pin.DefaultTextValue = UseDefaultText;
 	}
 
 	UEdGraphNode* Node = Pin.GetOwningNode();
-	check(Node);
 	Node->PinDefaultValueChanged(&Pin);
 
 	UBlueprint* Blueprint = FBlueprintEditorUtils::FindBlueprintForNodeChecked(Node);
@@ -4525,15 +4579,14 @@ void UEdGraphSchema_K2::TrySetDefaultText(UEdGraphPin& InPin, const FText& InNew
 	if(InPin.PinType.PinCategory == PC_Text)
 	{
 		// Check the default value and make it an error if it's bogus
-		if (IsPinDefaultValid(&InPin, TEXT(""), NULL, InNewDefaultText) == TEXT(""))
+		if (IsPinDefaultValid(&InPin, FString(), nullptr, InNewDefaultText).IsEmpty())
 		{
-			InPin.DefaultObject = NULL;
-			InPin.DefaultValue = NULL;
+			InPin.DefaultObject = nullptr;
+			InPin.DefaultValue.Empty();
 			InPin.DefaultTextValue = InNewDefaultText;
 		}
 
 		UEdGraphNode* Node = InPin.GetOwningNode();
-		check(Node);
 		Node->PinDefaultValueChanged(&InPin);
 
 		UBlueprint* Blueprint = FBlueprintEditorUtils::FindBlueprintForNodeChecked(Node);
@@ -4612,70 +4665,50 @@ bool UEdGraphSchema_K2::ShouldShowAssetPickerForPin(UEdGraphPin* Pin) const
 	return bShow;
 }
 
-void UEdGraphSchema_K2::SetPinDefaultValue(UEdGraphPin* Pin, const UFunction* Function, const UProperty* Param) const
+bool UEdGraphSchema_K2::FindFunctionParameterDefaultValue(const UFunction* Function, const UProperty* Param, FString& OutString)
 {
-	if ((Function != nullptr) && (Param != nullptr))
-	{
-		bool bHasAutomaticValue = false;
+	bool bHasAutomaticValue = false;
 
-		const FString MetadataDefaultValue = Function->GetMetaData(*Param->GetName());
-		if (!MetadataDefaultValue.IsEmpty())
+	const FString MetadataDefaultValue = Function->GetMetaData(*Param->GetName());
+	if (!MetadataDefaultValue.IsEmpty())
+	{
+		// Specified default value in the metadata
+		OutString = MetadataDefaultValue;
+		bHasAutomaticValue = true;
+	}
+	else
+	{
+		const FName MetadataCppDefaultValueKey(*(FString(TEXT("CPP_Default_")) + Param->GetName()));
+		const FString MetadataCppDefaultValue = Function->GetMetaData(MetadataCppDefaultValueKey);
+		if (!MetadataCppDefaultValue.IsEmpty())
 		{
-			// Specified default value in the metadata
-			Pin->AutogeneratedDefaultValue = MetadataDefaultValue;
+			OutString = MetadataCppDefaultValue;
 			bHasAutomaticValue = true;
 		}
-		else
-		{
-			const FName MetadataCppDefaultValueKey( *(FString(TEXT("CPP_Default_")) + Param->GetName()) );
-			const FString MetadataCppDefaultValue = Function->GetMetaData(MetadataCppDefaultValueKey);
-			if (!MetadataCppDefaultValue.IsEmpty())
-			{
-				Pin->AutogeneratedDefaultValue = MetadataCppDefaultValue;
-				bHasAutomaticValue = true;
-			}
-		}
-
-		if (bHasAutomaticValue)
-		{
-			if (Pin->PinType.PinCategory == PC_Text)
-			{
-				FString PackageNamespace;
-#if USE_STABLE_LOCALIZATION_KEYS
-				if (GIsEditor)
-				{
-					PackageNamespace = TextNamespaceUtil::EnsurePackageNamespace(Pin->GetOwningNodeUnchecked());
-				}
-#endif // USE_STABLE_LOCALIZATION_KEYS
-				if (!FTextStringHelper::ReadFromString(*Pin->AutogeneratedDefaultValue, Pin->DefaultTextValue, nullptr, *PackageNamespace))
-				{
-					Pin->DefaultTextValue = FText::FromString(Pin->AutogeneratedDefaultValue);
-				}
-			}
-			else
-			{
-				Pin->DefaultValue = Pin->AutogeneratedDefaultValue;
-			}
-		}
-	}
-	
-	if (Pin->DefaultValue.Len() == 0)
-	{
-		// Set the default value to (T)0
-		SetPinDefaultValueBasedOnType(Pin);
 	}
+
+	return bHasAutomaticValue;
 }
 
-void UEdGraphSchema_K2::SetPinDefaultValueBasedOnType(UEdGraphPin* Pin) const
+void UEdGraphSchema_K2::SetPinAutogeneratedDefaultValue(UEdGraphPin* Pin, const FString& NewValue) const
+{
+	Pin->AutogeneratedDefaultValue = NewValue;
+
+	ResetPinToAutogeneratedDefaultValue(Pin, false);
+}
+
+void UEdGraphSchema_K2::SetPinAutogeneratedDefaultValueBasedOnType(UEdGraphPin* Pin) const
 {
+	FString NewValue;
+
 	// Create a useful default value based on the pin type
 	if(Pin->PinType.IsContainer() )
 	{
-		Pin->AutogeneratedDefaultValue = TEXT("");
+		NewValue = TEXT("");
 	}
 	else if (Pin->PinType.PinCategory == PC_Int)
 	{
-		Pin->AutogeneratedDefaultValue = TEXT("0");
+		NewValue = TEXT("0");
 	}
 	else if(Pin->PinType.PinCategory == PC_Byte)
 	{
@@ -4684,37 +4717,73 @@ void UEdGraphSchema_K2::SetPinDefaultValueBasedOnType(UEdGraphPin* Pin) const
 		{
 			// First element of enum can change. If the enum is { A, B, C } and the default value is A, 
 			// the defult value should not change when enum will be changed into { N, A, B, C }
-			Pin->AutogeneratedDefaultValue = TEXT("");
-			Pin->DefaultValue = EnumPtr->GetNameStringByIndex(0);
-			return;
+			NewValue = TEXT("");
 		}
 		else
 		{
-			Pin->AutogeneratedDefaultValue = TEXT("0");
+			NewValue = TEXT("0");
 		}
 	}
 	else if (Pin->PinType.PinCategory == PC_Float)
 	{
-		Pin->AutogeneratedDefaultValue = TEXT("0.0");
+		// This is a slightly different format than is produced by PropertyValueToString, but changing it has backward compatibility issues
+		NewValue = TEXT("0.0");
 	}
 	else if (Pin->PinType.PinCategory == PC_Boolean)
 	{
-		Pin->AutogeneratedDefaultValue = TEXT("false");
+		NewValue = TEXT("false");
 	}
 	else if (Pin->PinType.PinCategory == PC_Name)
 	{
-		Pin->AutogeneratedDefaultValue = TEXT("None");
+		NewValue = TEXT("None");
 	}
 	else if ((Pin->PinType.PinCategory == PC_Struct) && ((Pin->PinType.PinSubCategoryObject == VectorStruct) || (Pin->PinType.PinSubCategoryObject == RotatorStruct)))
 	{
-		Pin->AutogeneratedDefaultValue = TEXT("0, 0, 0");
+		// This is a slightly different format than is produced by PropertyValueToString, but changing it has backward compatibility issues
+		NewValue = TEXT("0, 0, 0");
+	}
+
+	// PropertyValueToString also has cases for LinerColor and Transform, LinearColor is identical to export text so is fine, the Transform case is specially handled in the vm
+
+	SetPinAutogeneratedDefaultValue(Pin, NewValue);
+}
+
+void UEdGraphSchema_K2::ResetPinToAutogeneratedDefaultValue(UEdGraphPin* Pin, bool bCallModifyCallbacks) const
+{
+	GetPinDefaultValuesFromString(Pin->PinType, Pin->GetOwningNodeUnchecked(), Pin->AutogeneratedDefaultValue, Pin->DefaultValue, Pin->DefaultObject, Pin->DefaultTextValue);
+
+	if (bCallModifyCallbacks)
+	{
+		UEdGraphNode* Node = Pin->GetOwningNode();
+		Node->PinDefaultValueChanged(Pin);
+
+		UBlueprint* Blueprint = FBlueprintEditorUtils::FindBlueprintForNodeChecked(Node);
+		FBlueprintEditorUtils::MarkBlueprintAsModified(Blueprint);
+	}
+}
+
+void UEdGraphSchema_K2::SetPinDefaultValueAtConstruction(UEdGraphPin* Pin, const FString& DefaultValueString) const
+{
+	GetPinDefaultValuesFromString(Pin->PinType, Pin->GetOwningNodeUnchecked(), DefaultValueString, Pin->DefaultValue, Pin->DefaultObject, Pin->DefaultTextValue);
+}
+
+void UEdGraphSchema_K2::SetPinDefaultValue(UEdGraphPin* Pin, const UFunction* Function, const UProperty* Param) const
+{
+	if (Function != nullptr && Param != nullptr)
+	{
+		FString NewValue;
+		FindFunctionParameterDefaultValue(Function, Param, NewValue);
+		SetPinAutogeneratedDefaultValue(Pin, NewValue);
 	}
 	else
 	{
-		Pin->AutogeneratedDefaultValue = TEXT("");
+		SetPinAutogeneratedDefaultValueBasedOnType(Pin);
 	}
+}
 
-	Pin->DefaultValue = Pin->AutogeneratedDefaultValue;
+void UEdGraphSchema_K2::SetPinDefaultValueBasedOnType(UEdGraphPin* Pin) const
+{
+	SetPinAutogeneratedDefaultValueBasedOnType(Pin);
 }
 
 void UEdGraphSchema_K2::ValidateExistingConnections(UEdGraphPin* Pin)
@@ -4863,7 +4932,7 @@ UFunction* UEdGraphSchema_K2::FindSetVariableByNameFunction(const FEdGraphPinTyp
 	UFunction* Function = NULL;
 	if(SetFunctionName != NAME_None)
 	{
-		if(PinType.bIsArray)
+		if(PinType.IsArray())
 		{
 			static FName SetArrayName(GET_FUNCTION_NAME_CHECKED(UKismetArrayLibrary, SetArrayPropertyByName));
 			Function = FindField<UFunction>(UKismetArrayLibrary::StaticClass(), SetArrayName);
@@ -6377,7 +6446,7 @@ void UEdGraphSchema_K2::SplitPin(UEdGraphPin* Pin) const
 		{
 			const FString PinName = FString::Printf(TEXT("%s_%s"), *Pin->PinName, *ProtoPin->PinName);
 			const FEdGraphPinType& ProtoPinType = ProtoPin->PinType;
-			UEdGraphPin* SubPin = GraphNode->CreatePin(Pin->Direction, ProtoPinType.PinCategory, ProtoPinType.PinSubCategory, ProtoPinType.PinSubCategoryObject.Get(), ProtoPinType.bIsArray, false, PinName);
+			UEdGraphPin* SubPin = GraphNode->CreatePin(Pin->Direction, ProtoPinType.PinCategory, ProtoPinType.PinSubCategory, ProtoPinType.PinSubCategoryObject.Get(), PinName, ProtoPinType.ContainerType, false, false, INDEX_NONE, ProtoPinType.PinValueType);
 
 			if (K2Node != nullptr && K2Node->ShouldDrawCompact() && !Pin->ParentPin)
 			{
@@ -6682,6 +6751,15 @@ UEdGraphPin* UEdGraphSchema_K2::DropPinOnNode(UEdGraphNode* InTargetNode, const
 			}
 		}
 
+		HandleParameterDefaultValueChanged(EditablePinNode);
+	}
+	return ResultPin;
+}
+
+void UEdGraphSchema_K2::HandleParameterDefaultValueChanged(UK2Node* InTargetNode) const
+{
+	if (UK2Node_EditablePinBase* EditablePinNode = Cast<UK2Node_EditablePinBase>(InTargetNode))
+	{
 		FParamsChangedHelper ParamsChangedHelper;
 		ParamsChangedHelper.ModifiedBlueprints.Add(FBlueprintEditorUtils::FindBlueprintForNode(InTargetNode));
 		FBlueprintEditorUtils::MarkBlueprintAsStructurallyModified(FBlueprintEditorUtils::FindBlueprintForNode(InTargetNode));
@@ -6697,7 +6775,7 @@ UEdGraphPin* UEdGraphSchema_K2::DropPinOnNode(UEdGraphNode* InTargetNode, const
 		}
 
 		// Now update all the blueprints that got modified
-		for (UBlueprint* Blueprint  : ParamsChangedHelper.ModifiedBlueprints)
+		for (UBlueprint* Blueprint : ParamsChangedHelper.ModifiedBlueprints)
 		{
 			if (Blueprint)
 			{
@@ -6705,7 +6783,6 @@ UEdGraphPin* UEdGraphSchema_K2::DropPinOnNode(UEdGraphNode* InTargetNode, const
 			}
 		}
 	}
-	return ResultPin;
 }
 
 bool UEdGraphSchema_K2::SupportsDropPinOnNode(UEdGraphNode* InTargetNode, const FEdGraphPinType& InSourcePinType, EEdGraphPinDirection InSourcePinDirection, FText& OutErrorMessage) const
