Copying //UE4/Dev-Framework to //UE4/Dev-Main (Source: //UE4/Dev-Framework @ 3431384)

diff --git a/Engine/Source/Editor/BlueprintGraph/Private/K2Node.cpp b/Engine/Source/Editor/BlueprintGraph/Private/K2Node.cpp
--- a/Engine/Source/Editor/BlueprintGraph/Private/K2Node.cpp
+++ b/Engine/Source/Editor/BlueprintGraph/Private/K2Node.cpp
@@ -23,6 +23,7 @@
 #include "PropertyCustomizationHelpers.h"
 
 #include "ObjectEditorUtils.h"
+#include "FrameworkObjectVersion.h"
 
 #define LOCTEXT_NAMESPACE "K2Node"
 
@@ -66,47 +67,63 @@ void UK2Node::PostLoad()
 #endif // WITH_EDITORONLY_DATA
 
 	Super::PostLoad();
-
-	// fix up pin default values
-	FixupPinDefaultValues();
 }
 
 void UK2Node::Serialize(FArchive& Ar)
 {
-	Super::Serialize(Ar);
-
-	if (Ar.IsObjectReferenceCollector() && Ar.IsSaving())
-	{
-		// If looking for references during save, expand any default values on the pins
-		
-		for (const UEdGraphPin* Pin : Pins)
+	if (Ar.IsSaving())
+	{	
+		for (UEdGraphPin* Pin : Pins)
 		{
-			if (!Pin->bDefaultValueIsIgnored && !Pin->DefaultValue.IsEmpty() && Pin->PinType.PinCategory == UEdGraphSchema_K2::PC_Struct && Pin->PinType.PinSubCategoryObject.IsValid())
+			if (!Pin->bDefaultValueIsIgnored && !Pin->DefaultValue.IsEmpty() )
 			{
-				UScriptStruct* Struct = Cast<UScriptStruct>(Pin->PinType.PinSubCategoryObject.Get());
+				// If looking for references during save, expand any default values on the pins
+				if (Ar.IsObjectReferenceCollector() && Pin->PinType.PinCategory == UEdGraphSchema_K2::PC_Struct && Pin->PinType.PinSubCategoryObject.IsValid())
+				{
+					UScriptStruct* Struct = Cast<UScriptStruct>(Pin->PinType.PinSubCategoryObject.Get());
+
+					if (Struct)
+					{
+						int32 StructSize = Struct->GetStructureSize();
+						uint8* StructData = (uint8*)FMemory::Malloc(StructSize);
 
-				if (Struct)
+						// Import the literal text to a dummy struct and then serialize that
+						FOutputDeviceNull NullOutput;
+						Struct->InitializeStruct(StructData);
+						Struct->ImportText(*Pin->DefaultValue, StructData, nullptr, PPF_None, &NullOutput, Pin->PinName);
+						Struct->SerializeItem(Ar, StructData, nullptr);
+						Struct->DestroyStruct(StructData);
+
+						FMemory::Free(StructData);
+					}
+				}
+
+				if (Pin->PinType.PinCategory == UEdGraphSchema_K2::PC_Asset || Pin->PinType.PinCategory == UEdGraphSchema_K2::PC_AssetClass)
 				{
-					int32 StructSize = Struct->GetStructureSize();
-					uint8* StructData = (uint8*)FMemory::Malloc(StructSize);
-					
-					// Import the literal text to a dummy struct and then serialize that
-					FOutputDeviceNull NullOutput;
-					Struct->InitializeStruct(StructData);
-					Struct->ImportText(*Pin->DefaultValue, StructData, nullptr, PPF_None, &NullOutput, Pin->PinName);
-					Struct->SerializeItem(Ar, StructData, nullptr);
-					Struct->DestroyStruct(StructData);
-
-					FMemory::Free(StructData);
+					FStringAssetReference TempRef(Pin->DefaultValue);
+
+					// Serialize the asset reference, this will do the save fixup. It won't actually serialize the string if this is a real archive like linkersave
+					TempRef.SerializePath(Ar, true);
+
+					Pin->DefaultValue = MoveTemp(TempRef.ToString());
 				}
 			}
 		}
 	}
+
+	Super::Serialize(Ar);
+
+	if (Ar.IsLoading())
+	{
+		// Fix up pin default values, must be done before post load
+		FixupPinDefaultValues();
+	}
 }
 
 void UK2Node::FixupPinDefaultValues()
 {
 	const int32 LinkerUE4Version = GetLinkerUE4Version();
+	const int32 LinkerFrameworkVersion = GetLinkerCustomVersion(FFrameworkObjectVersion::GUID);
 	const UEdGraphSchema_K2* K2Schema = GetDefault<UEdGraphSchema_K2>();
 
 	// Swap "new" default error tolerance value with zero on vector/rotator equality nodes, in order to preserve current behavior in existing blueprints.
@@ -140,6 +157,26 @@ void UK2Node::FixupPinDefaultValues()
 			}
 		}
 	}
+
+	// Fix asset ptr pins
+	if (LinkerFrameworkVersion < FFrameworkObjectVersion::ChangeAssetPinsToString)
+	{
+		bool bFoundPin = false;
+		for (int32 i = 0; i < Pins.Num() && !bFoundPin; ++i)
+		{
+			UEdGraphPin* Pin = Pins[i];
+
+			if (Pin->PinType.PinCategory == K2Schema->PC_Asset || Pin->PinType.PinCategory == K2Schema->PC_AssetClass)
+			{
+				if (Pin->DefaultObject && Pin->DefaultValue.IsEmpty())
+				{
+					Pin->DefaultValue = Pin->DefaultObject->GetPathName();
+					Pin->DefaultObject = nullptr;
+				}
+			}
+		}
+	}
+
 }
 
 FText UK2Node::GetToolTipHeading() const
@@ -186,9 +223,9 @@ bool UK2Node::CreatePinsForFunctionEntryExit(const UFunction* Function, bool bFo
 		{
 			const EEdGraphPinDirection Direction = bForFunctionEntry ? EGPD_Output : EGPD_Input;
 
-			UEdGraphPin* Pin = CreatePin(Direction, TEXT(""), TEXT(""), NULL, false, false, Param->GetName());
+			UEdGraphPin* Pin = CreatePin(Direction, FString(), FString(), nullptr, Param->GetName());
 			const bool bPinGood = K2Schema->ConvertPropertyToPinType(Param, /*out*/ Pin->PinType);
-			K2Schema->SetPinDefaultValueBasedOnType(Pin);
+			K2Schema->SetPinAutogeneratedDefaultValueBasedOnType(Pin);
 			
 			UK2Node_CallFunction::GeneratePinTooltipFromFunction(*Pin, Function);
 
@@ -339,10 +376,10 @@ void UK2Node::PinConnectionListChanged(UEdGraphPin* Pin)
 		{
 			const UEdGraphSchema_K2* Schema = Cast<const UEdGraphSchema_K2>(GetSchema());
 
-			Pin->ResetDefaultValue();
-			if (EEdGraphPinDirection::EGPD_Input == Pin->Direction)
+			// We don't want to reset Output pin defaults, that breaks Function entry nodes
+			if (Pin->Direction == EEdGraphPinDirection::EGPD_Input)
 			{
-				Schema->SetPinDefaultValueBasedOnType(Pin);
+				Schema->ResetPinToAutogeneratedDefaultValue(Pin);
 			}
 		}
 	}
@@ -1007,7 +1044,7 @@ void FOptionalPinManager::RebuildProperty(UProperty* TestProperty, FName Categor
 	}
 }
 
-void FOptionalPinManager::CreateVisiblePins(TArray<FOptionalPinFromProperty>& Properties, UStruct* SourceStruct, EEdGraphPinDirection Direction, UK2Node* TargetNode, uint8* StructBasePtr)
+void FOptionalPinManager::CreateVisiblePins(TArray<FOptionalPinFromProperty>& Properties, UStruct* SourceStruct, EEdGraphPinDirection Direction, UK2Node* TargetNode, uint8* StructBasePtr, uint8* DefaultsPtr)
 {
 	const UEdGraphSchema_K2* Schema = GetDefault<UEdGraphSchema_K2>();
 
@@ -1020,7 +1057,7 @@ void FOptionalPinManager::CreateVisiblePins(TArray<FOptionalPinFromProperty>& Pr
 			// Do we treat an array property as one pin, or a pin per entry in the array?
 			// Depends on if we have an instance of the struct to work with.
 			UArrayProperty* ArrayProperty = Cast<UArrayProperty>(OuterProperty);
-			if ((ArrayProperty != NULL) && (StructBasePtr != NULL))
+			if ((ArrayProperty != nullptr) && (StructBasePtr != nullptr))
 			{
 				UProperty* InnerProperty = ArrayProperty->Inner;
 
@@ -1054,13 +1091,24 @@ void FOptionalPinManager::CreateVisiblePins(TArray<FOptionalPinFromProperty>& Pr
 
 						// Let derived classes take a crack at transferring default values
 						uint8* ValuePtr = ArrayHelper.GetRawPtr(Index);
-						if (NewPin != NULL)
+						uint8* DefaultValuePtr = nullptr;
+
+						if (DefaultsPtr)
+						{
+							FScriptArrayHelper_InContainer DefaultsArrayHelper(ArrayProperty, DefaultsPtr);
+
+							if (DefaultsArrayHelper.IsValidIndex(Index))
+							{
+								DefaultValuePtr = DefaultsArrayHelper.GetRawPtr(Index);
+							}
+						}
+						if (NewPin != nullptr)
 						{
-							PostInitNewPin(NewPin, PropertyEntry, Index, ArrayProperty->Inner, ValuePtr);
+							PostInitNewPin(NewPin, PropertyEntry, Index, ArrayProperty->Inner, ValuePtr, DefaultValuePtr);
 						}
 						else
 						{
-							PostRemovedOldPin(PropertyEntry, Index, ArrayProperty->Inner, ValuePtr);
+							PostRemovedOldPin(PropertyEntry, Index, ArrayProperty->Inner, ValuePtr, DefaultValuePtr);
 						}
 					}
 				}
@@ -1073,7 +1121,7 @@ void FOptionalPinManager::CreateVisiblePins(TArray<FOptionalPinFromProperty>& Pr
 				if (Schema->ConvertPropertyToPinType(OuterProperty, /*out*/ PinType))
 				{
 					// Create the pin
-					UEdGraphPin* NewPin = NULL;
+					UEdGraphPin* NewPin = nullptr;
 					if (PropertyEntry.bShowPin)
 					{
 						const FString PinName = PropertyEntry.PropertyName.ToString();
@@ -1088,16 +1136,17 @@ void FOptionalPinManager::CreateVisiblePins(TArray<FOptionalPinFromProperty>& Pr
 					}
 
 					// Let derived classes take a crack at transferring default values
-					if (StructBasePtr != NULL)
+					if (StructBasePtr != nullptr)
 					{
 						uint8* ValuePtr = OuterProperty->ContainerPtrToValuePtr<uint8>(StructBasePtr);
-						if (NewPin != NULL)
+						uint8* DefaultValuePtr = DefaultsPtr ? OuterProperty->ContainerPtrToValuePtr<uint8>(DefaultsPtr) : nullptr;
+						if (NewPin != nullptr)
 						{
-							PostInitNewPin(NewPin, PropertyEntry, INDEX_NONE, OuterProperty, ValuePtr);
+							PostInitNewPin(NewPin, PropertyEntry, INDEX_NONE, OuterProperty, ValuePtr, DefaultValuePtr);
 						}
 						else
 						{
-							PostRemovedOldPin(PropertyEntry, INDEX_NONE, OuterProperty, ValuePtr);
+							PostRemovedOldPin(PropertyEntry, INDEX_NONE, OuterProperty, ValuePtr, DefaultValuePtr);
 						}
 					}
 				}
@@ -1249,8 +1298,8 @@ void UK2Node::GetPinHoverText(const UEdGraphPin& Pin, FString& HoverTextOut) con
 	// grab the debug value of the pin
 	FString WatchText;
 	const FKismetDebugUtilities::EWatchTextResult WatchStatus = FKismetDebugUtilities::GetWatchText(/*inout*/ WatchText, Blueprint, ActiveObject, &Pin);
-	// if this is an array pin, then we possibly have too many lines (too many entries)
-	if (Pin.PinType.bIsArray)
+	// if this is an container pin, then we possibly have too many lines (too many entries)
+	if (Pin.PinType.IsContainer())
 	{
 		int32 LineCounter = 0;
 		int32 OriginalWatchTextLen = WatchText.Len();
@@ -1283,7 +1332,7 @@ void UK2Node::GetPinHoverText(const UEdGraphPin& Pin, FString& HoverTextOut) con
 				break;
 			}
 		}
-	} // if Pin.PinType.bIsArray...
+	} // if Pin.PinType.IsContainer()...
 
 
 	switch (WatchStatus)
diff --git a/Engine/Source/Editor/BlueprintGraph/Private/K2Node_AddComponent.cpp b/Engine/Source/Editor/BlueprintGraph/Private/K2Node_AddComponent.cpp
--- a/Engine/Source/Editor/BlueprintGraph/Private/K2Node_AddComponent.cpp
+++ b/Engine/Source/Editor/BlueprintGraph/Private/K2Node_AddComponent.cpp
@@ -120,7 +120,7 @@ void UK2Node_AddComponent::AllocatePinsForExposedVariables()
 				const bool bIsUnique = (NULL == FindPin(Property->GetName()));
 				if (K2Schema->FindSetVariableByNameFunction(PinType) && bIsUnique)
 				{
-					UEdGraphPin* Pin = CreatePin(EGPD_Input, TEXT(""), TEXT(""), NULL, false, false, Property->GetName());
+					UEdGraphPin* Pin = CreatePin(EGPD_Input, FString(), FString(), nullptr, Property->GetName());
 					Pin->PinType = PinType;
 					bHasExposedVariable = true;
 
@@ -129,7 +129,7 @@ void UK2Node_AddComponent::AllocatePinsForExposedVariables()
 						FString DefaultValueAsString;
 						const bool bDefaultValueSet = FBlueprintEditorUtils::PropertyValueToString(Property, reinterpret_cast<const uint8*>(ClassDefaultObject), DefaultValueAsString);
 						check(bDefaultValueSet);
-						K2Schema->TrySetDefaultValue(*Pin, DefaultValueAsString);
+						K2Schema->SetPinAutogeneratedDefaultValue(Pin, DefaultValueAsString);
 					}
 
 					// Copy tooltip from the property.
diff --git a/Engine/Source/Editor/BlueprintGraph/Private/K2Node_AssignmentStatement.cpp b/Engine/Source/Editor/BlueprintGraph/Private/K2Node_AssignmentStatement.cpp
--- a/Engine/Source/Editor/BlueprintGraph/Private/K2Node_AssignmentStatement.cpp
+++ b/Engine/Source/Editor/BlueprintGraph/Private/K2Node_AssignmentStatement.cpp
@@ -68,11 +68,11 @@ void UK2Node_AssignmentStatement::AllocateDefaultPins()
 {
 	const UEdGraphSchema_K2* Schema = GetDefault<UEdGraphSchema_K2>();
 
-	CreatePin(EGPD_Input, Schema->PC_Exec, TEXT(""), NULL, false, false, Schema->PN_Execute);
-	CreatePin(EGPD_Output, Schema->PC_Exec, TEXT(""), NULL, false, false, Schema->PN_Then);
+	CreatePin(EGPD_Input, Schema->PC_Exec, FString(), nullptr, Schema->PN_Execute);
+	CreatePin(EGPD_Output, Schema->PC_Exec, FString(), nullptr, Schema->PN_Then);
 
-	UEdGraphPin* VariablePin = CreatePin(EGPD_Input, Schema->PC_Wildcard, TEXT(""), NULL, false, false, VariablePinName);
-	UEdGraphPin* ValuePin = CreatePin(EGPD_Input, Schema->PC_Wildcard, TEXT(""), NULL, false, false, ValuePinName);
+	UEdGraphPin* VariablePin = CreatePin(EGPD_Input, Schema->PC_Wildcard, FString(), nullptr, VariablePinName);
+	UEdGraphPin* ValuePin = CreatePin(EGPD_Input, Schema->PC_Wildcard, FString(), nullptr, ValuePinName);
 
 	Super::AllocateDefaultPins();
 }
diff --git a/Engine/Source/Editor/BlueprintGraph/Private/K2Node_BaseAsyncTask.cpp b/Engine/Source/Editor/BlueprintGraph/Private/K2Node_BaseAsyncTask.cpp
--- a/Engine/Source/Editor/BlueprintGraph/Private/K2Node_BaseAsyncTask.cpp
+++ b/Engine/Source/Editor/BlueprintGraph/Private/K2Node_BaseAsyncTask.cpp
@@ -65,8 +65,8 @@ void UK2Node_BaseAsyncTask::AllocateDefaultPins()
 {
 	const UEdGraphSchema_K2* K2Schema = GetDefault<UEdGraphSchema_K2>();
 
-	CreatePin(EGPD_Input, K2Schema->PC_Exec, TEXT(""), NULL, false, false, K2Schema->PN_Execute);
-	CreatePin(EGPD_Output, K2Schema->PC_Exec, TEXT(""), NULL, false, false, K2Schema->PN_Then);
+	CreatePin(EGPD_Input, K2Schema->PC_Exec, FString(), nullptr, K2Schema->PN_Execute);
+	CreatePin(EGPD_Output, K2Schema->PC_Exec, FString(), nullptr, K2Schema->PN_Then);
 
 	bool bExposeProxy = false;
 	for (const UStruct* TestStruct = ProxyClass; TestStruct && !bExposeProxy; TestStruct = TestStruct->GetSuperStruct())
@@ -75,7 +75,7 @@ void UK2Node_BaseAsyncTask::AllocateDefaultPins()
 	}
 	if (bExposeProxy)
 	{
-		CreatePin(EGPD_Output, K2Schema->PC_Object, TEXT(""), ProxyClass, false, false, FBaseAsyncTaskHelper::GetAsyncTaskProxyName());
+		CreatePin(EGPD_Output, K2Schema->PC_Object, FString(), ProxyClass, FBaseAsyncTaskHelper::GetAsyncTaskProxyName());
 	}
 
 	UFunction* DelegateSignatureFunction = NULL;
@@ -83,7 +83,7 @@ void UK2Node_BaseAsyncTask::AllocateDefaultPins()
 	{
 		if (UMulticastDelegateProperty* Property = Cast<UMulticastDelegateProperty>(*PropertyIt))
 		{
-			CreatePin(EGPD_Output, K2Schema->PC_Exec, TEXT(""), NULL, false, false, *Property->GetName());
+			CreatePin(EGPD_Output, K2Schema->PC_Exec, FString(), nullptr, *Property->GetName());
 			if (!DelegateSignatureFunction)
 			{
 				DelegateSignatureFunction = Property->SignatureFunction;
@@ -99,7 +99,7 @@ void UK2Node_BaseAsyncTask::AllocateDefaultPins()
 			const bool bIsFunctionInput = !Param->HasAnyPropertyFlags(CPF_OutParm) || Param->HasAnyPropertyFlags(CPF_ReferenceParm);
 			if (bIsFunctionInput)
 			{
-				UEdGraphPin* Pin = CreatePin(EGPD_Output, TEXT(""), TEXT(""), NULL, false, false, Param->GetName());
+				UEdGraphPin* Pin = CreatePin(EGPD_Output, FString(), FString(), nullptr, Param->GetName());
 				K2Schema->ConvertPropertyToPinType(Param, /*out*/ Pin->PinType);
 			}
 		}
@@ -122,13 +122,13 @@ void UK2Node_BaseAsyncTask::AllocateDefaultPins()
 			}
 
 			const bool bIsRefParam = Param->HasAnyPropertyFlags(CPF_ReferenceParm) && bIsFunctionInput;
-			UEdGraphPin* Pin = CreatePin(EGPD_Input, TEXT(""), TEXT(""), NULL, false, bIsRefParam, Param->GetName());
+			UEdGraphPin* Pin = CreatePin(EGPD_Input, FString(), FString(), nullptr, Param->GetName(), EPinContainerType::None, bIsRefParam);
 			const bool bPinGood = (Pin != NULL) && K2Schema->ConvertPropertyToPinType(Param, /*out*/ Pin->PinType);
 
 			if (bPinGood)
 			{
 				//Flag pin as read only for const reference property
-				Pin->bDefaultValueIsIgnored = Param->HasAllPropertyFlags(CPF_ConstParm | CPF_ReferenceParm) && (!Function->HasMetaData(FBlueprintMetadata::MD_AutoCreateRefTerm) || Pin->PinType.bIsArray);
+				Pin->bDefaultValueIsIgnored = Param->HasAllPropertyFlags(CPF_ConstParm | CPF_ReferenceParm) && (!Function->HasMetaData(FBlueprintMetadata::MD_AutoCreateRefTerm) || Pin->PinType.IsContainer());
 
 				const bool bAdvancedPin = Param->HasAllPropertyFlags(CPF_AdvancedDisplay);
 				Pin->bAdvancedView = bAdvancedPin;
@@ -137,7 +137,15 @@ void UK2Node_BaseAsyncTask::AllocateDefaultPins()
 					AdvancedPinDisplay = ENodeAdvancedPins::Hidden;
 				}
 
-				K2Schema->SetPinDefaultValue(Pin, Function, Param);
+				FString ParamValue;
+				if (K2Schema->FindFunctionParameterDefaultValue(Function, Param, ParamValue))
+				{
+					K2Schema->SetPinAutogeneratedDefaultValue(Pin, ParamValue);
+				}
+				else
+				{
+					K2Schema->SetPinAutogeneratedDefaultValueBasedOnType(Pin);
+				}
 
 				if (PinsToHide.Contains(Pin->PinName))
 				{
@@ -316,7 +324,7 @@ void UK2Node_BaseAsyncTask::ExpandNode(class FKismetCompilerContext& CompilerCon
 		{
 			const FEdGraphPinType& PinType = CurrentPin->PinType;
 			UK2Node_TemporaryVariable* TempVarOutput = CompilerContext.SpawnInternalVariable(
-				this, PinType.PinCategory, PinType.PinSubCategory, PinType.PinSubCategoryObject.Get(), PinType.bIsArray, PinType.bIsSet, PinType.bIsMap, PinType.PinValueType);
+				this, PinType.PinCategory, PinType.PinSubCategory, PinType.PinSubCategoryObject.Get(), PinType.ContainerType, PinType.PinValueType);
 			bIsErrorFree &= TempVarOutput->GetVariablePin() && CompilerContext.MovePinLinksToIntermediate(*CurrentPin, *TempVarOutput->GetVariablePin()).CanSafeConnect();
 			VariableOutputs.Add(FBaseAsyncTaskHelper::FOutputPinAndLocalVariable(CurrentPin, TempVarOutput));
 		}
@@ -452,6 +460,8 @@ UFunction* UK2Node_BaseAsyncTask::GetFactoryFunction() const
 
 void UK2Node_BaseAsyncTask::ValidateNodeDuringCompilation(class FCompilerResultsLog& MessageLog) const
 {
+	Super::ValidateNodeDuringCompilation(MessageLog);
+
 	UEdGraphSchema_K2 const* K2Schema = GetDefault<UEdGraphSchema_K2>();
 
 	if(UObject const* SourceObject = MessageLog.FindSourceObject(this))
diff --git a/Engine/Source/Editor/BlueprintGraph/Private/K2Node_BitmaskLiteral.cpp b/Engine/Source/Editor/BlueprintGraph/Private/K2Node_BitmaskLiteral.cpp
--- a/Engine/Source/Editor/BlueprintGraph/Private/K2Node_BitmaskLiteral.cpp
+++ b/Engine/Source/Editor/BlueprintGraph/Private/K2Node_BitmaskLiteral.cpp
@@ -29,6 +29,7 @@ UK2Node_BitmaskLiteral::UK2Node_BitmaskLiteral(const FObjectInitializer& ObjectI
 
 void UK2Node_BitmaskLiteral::ValidateBitflagsEnumType()
 {
+	const UEdGraphSchema_K2* Schema = GetDefault<UEdGraphSchema_K2>();
 	if (BitflagsEnum != nullptr)
 	{
 		// Reset enum type reference if it no longer has the proper meta data.
@@ -71,7 +72,7 @@ void UK2Node_BitmaskLiteral::ValidateBitflagsEnumType()
 			const int32 NewDefaultValue = OldDefaultValue & ValidBitflagsMask;
 			if (NewDefaultValue != OldDefaultValue)
 			{
-				InputPin->GetSchema()->TrySetDefaultValue(*InputPin, FString::FromInt(NewDefaultValue));
+				Schema->SetPinAutogeneratedDefaultValue(InputPin, FString::FromInt(NewDefaultValue));
 			}
 		}
 	}
@@ -98,10 +99,10 @@ void UK2Node_BitmaskLiteral::AllocateDefaultPins()
 {
 	const UEdGraphSchema_K2* Schema = GetDefault<UEdGraphSchema_K2>();
 
-	UEdGraphPin* InputPin = CreatePin(EGPD_Input, Schema->PC_Int, Schema->PSC_Bitmask, BitflagsEnum, false, false, GetBitmaskInputPinName());
-	Schema->SetPinDefaultValueBasedOnType(InputPin);
+	UEdGraphPin* InputPin = CreatePin(EGPD_Input, Schema->PC_Int, Schema->PSC_Bitmask, BitflagsEnum, GetBitmaskInputPinName());
+	Schema->SetPinAutogeneratedDefaultValueBasedOnType(InputPin);
 
-	CreatePin(EGPD_Output, Schema->PC_Int, Schema->PSC_Bitmask, BitflagsEnum, false, false, Schema->PN_ReturnValue);
+	CreatePin(EGPD_Output, Schema->PC_Int, Schema->PSC_Bitmask, BitflagsEnum, Schema->PN_ReturnValue);
 
 	Super::AllocateDefaultPins();
 }
diff --git a/Engine/Source/Editor/BlueprintGraph/Private/K2Node_BreakStruct.cpp b/Engine/Source/Editor/BlueprintGraph/Private/K2Node_BreakStruct.cpp
--- a/Engine/Source/Editor/BlueprintGraph/Private/K2Node_BreakStruct.cpp
+++ b/Engine/Source/Editor/BlueprintGraph/Private/K2Node_BreakStruct.cpp
@@ -182,7 +182,7 @@ void UK2Node_BreakStruct::AllocateDefaultPins()
 	if(Schema && StructType)
 	{
 		PreloadObject(StructType);
-		CreatePin(EGPD_Input, Schema->PC_Struct, TEXT(""), StructType, false, true, StructType->GetName(), true);
+		CreatePin(EGPD_Input, Schema->PC_Struct, FString(), StructType, StructType->GetName(), EPinContainerType::None, true, true);
 		
 		struct FBreakStructPinManager : public FStructOperationOptionalPinManager
 		{
diff --git a/Engine/Source/Editor/BlueprintGraph/Private/K2Node_CallArrayFunction.cpp b/Engine/Source/Editor/BlueprintGraph/Private/K2Node_CallArrayFunction.cpp
--- a/Engine/Source/Editor/BlueprintGraph/Private/K2Node_CallArrayFunction.cpp
+++ b/Engine/Source/Editor/BlueprintGraph/Private/K2Node_CallArrayFunction.cpp
@@ -21,11 +21,11 @@ void UK2Node_CallArrayFunction::AllocateDefaultPins()
 	UEdGraphPin* TargetArrayPin = GetTargetArrayPin();
 	if (ensure(TargetArrayPin))
 	{
-		TargetArrayPin->PinType.bIsArray = true;
+		TargetArrayPin->PinType.ContainerType = EPinContainerType::Array;
 		TargetArrayPin->PinType.bIsReference = true;
 		TargetArrayPin->PinType.PinCategory = Schema->PC_Wildcard;
-		TargetArrayPin->PinType.PinSubCategory = TEXT("");
-		TargetArrayPin->PinType.PinSubCategoryObject = NULL;
+		TargetArrayPin->PinType.PinSubCategory.Reset();
+		TargetArrayPin->PinType.PinSubCategoryObject = nullptr;
 	}
 
 	TArray< FArrayPropertyPinCombo > ArrayPins;
@@ -88,7 +88,7 @@ void UK2Node_CallArrayFunction::NotifyPinConnectionListChanged(UEdGraphPin* Pin)
 			{
 				UEdGraphPin* LinkedTo = Pin->LinkedTo[0];
 				check(LinkedTo);
-				check(Pin->PinType.bIsArray == LinkedTo->PinType.bIsArray);
+				check(Pin->PinType.ContainerType == LinkedTo->PinType.ContainerType);
 
 				Pin->PinType.PinCategory = LinkedTo->PinType.PinCategory;
 				Pin->PinType.PinSubCategory = LinkedTo->PinType.PinSubCategory;
@@ -121,6 +121,7 @@ void UK2Node_CallArrayFunction::NotifyPinConnectionListChanged(UEdGraphPin* Pin)
 		if (bNeedToPropagate)
 		{
 			PropagateArrayTypeInfo(Pin);
+			GetGraph()->NotifyGraphChanged();
 		}
 	}
 }
@@ -159,9 +160,9 @@ void UK2Node_CallArrayFunction::ConvertDeprecatedNode(UEdGraph* Graph, bool bOnl
 					else if (Pin->PinName == UEdGraphSchema_K2::PN_Self)
 					{
 						// there's no analogous pin, signal that we're expecting this
-						OldToNewPinMap.Add(Pin->PinName, TEXT(""));
+						OldToNewPinMap.Add(Pin->PinName, FString());
 					}
-					else if (Pin->PinType.bIsArray)
+					else if (Pin->PinType.IsArray())
 					{
 						OldToNewPinMap.Add(Pin->PinName, GetItemNode->GetTargetArrayPin()->PinName);
 					}
@@ -312,7 +313,7 @@ void UK2Node_CallArrayFunction::PropagateArrayTypeInfo(const UEdGraphPin* Source
 				// Reset default values
 				if (!Schema->IsPinDefaultValid(CurrentPin, CurrentPin->DefaultValue, CurrentPin->DefaultObject, CurrentPin->DefaultTextValue).IsEmpty())
 				{
-					CurrentPin->ResetDefaultValue();
+					Schema->ResetPinToAutogeneratedDefaultValue(CurrentPin);
 				}
 			}
 		}
diff --git a/Engine/Source/Editor/BlueprintGraph/Private/K2Node_CallFunction.cpp b/Engine/Source/Editor/BlueprintGraph/Private/K2Node_CallFunction.cpp
--- a/Engine/Source/Editor/BlueprintGraph/Private/K2Node_CallFunction.cpp
+++ b/Engine/Source/Editor/BlueprintGraph/Private/K2Node_CallFunction.cpp
@@ -68,7 +68,7 @@ struct FCustomStructureParamHelper
 			{
 				UEdGraphPin* LinkedTo = Pin->LinkedTo[0];
 				check(LinkedTo);
-				ensure(!LinkedTo->PinType.bIsArray);
+				ensure(!LinkedTo->PinType.IsContainer());
 
 				Pin->PinType = LinkedTo->PinType;
 			}
@@ -76,8 +76,8 @@ struct FCustomStructureParamHelper
 			{
 				const UEdGraphSchema_K2* Schema = GetDefault<UEdGraphSchema_K2>();
 				Pin->PinType.PinCategory = Schema->PC_Wildcard;
-				Pin->PinType.PinSubCategory = TEXT("");
-				Pin->PinType.PinSubCategoryObject = NULL;
+				Pin->PinType.PinSubCategory.Reset();
+				Pin->PinType.PinSubCategoryObject = nullptr;
 			}
 		}
 	}
@@ -97,9 +97,9 @@ struct FCustomStructureParamHelper
 			}
 			else
 			{
-				for (auto& Name : Names)
+				for (const FString& Name : Names)
 				{
-					if (auto Pin = Node->FindPin(Name))
+					if (UEdGraphPin* Pin = Node->FindPin(Name))
 					{
 						HandleSinglePin(Pin);
 					}
@@ -549,7 +549,7 @@ void UK2Node_CallFunction::GetPinHoverText(const UEdGraphPin& Pin, FString& Hove
 {
 	if (!bPinTooltipsValid)
 	{
-		for (auto& P : Pins)
+		for (UEdGraphPin* P : Pins)
 		{
 			P->PinToolTip.Empty();
 			GeneratePinTooltip(*P);
@@ -686,7 +686,7 @@ UEdGraphPin* UK2Node_CallFunction::CreateSelfPin(const UFunction* Function)
 
 	// Chase up the function's Super chain, the function can be called on any object that is at least that specific
 	const UFunction* FirstDeclaredFunction = Function;
-	while (FirstDeclaredFunction->GetSuperFunction() != NULL)
+	while (FirstDeclaredFunction->GetSuperFunction() != nullptr)
 	{
 		FirstDeclaredFunction = FirstDeclaredFunction->GetSuperFunction();
 	}
@@ -707,18 +707,18 @@ UEdGraphPin* UK2Node_CallFunction::CreateSelfPin(const UFunction* Function)
 	if (FunctionClass == GetBlueprint()->GeneratedClass)
 	{
 		// This means the function is defined within the blueprint, so the pin should be a true "self" pin
-		SelfPin = CreatePin(EGPD_Input, K2Schema->PC_Object, K2Schema->PSC_Self, NULL, false, false, K2Schema->PN_Self);
+		SelfPin = CreatePin(EGPD_Input, K2Schema->PC_Object, K2Schema->PSC_Self, nullptr, K2Schema->PN_Self);
 	}
 	else if (FunctionClass->IsChildOf(UInterface::StaticClass()))
 	{
-		SelfPin = CreatePin(EGPD_Input, K2Schema->PC_Interface, TEXT(""), FunctionClass, false, false, K2Schema->PN_Self);
+		SelfPin = CreatePin(EGPD_Input, K2Schema->PC_Interface, FString(), FunctionClass, K2Schema->PN_Self);
 	}
 	else
 	{
 		// This means that the function is declared in an external class, and should reference that class
-		SelfPin = CreatePin(EGPD_Input, K2Schema->PC_Object, TEXT(""), FunctionClass, false, false, K2Schema->PN_Self);
+		SelfPin = CreatePin(EGPD_Input, K2Schema->PC_Object, FString(), FunctionClass, K2Schema->PN_Self);
 	}
-	check(SelfPin != NULL);
+	check(SelfPin != nullptr);
 
 	return SelfPin;
 }
@@ -767,7 +767,7 @@ void UK2Node_CallFunction::CreateExecPinsForFunctionCall(const UFunction* Functi
 					if (!bShouldBeHidden)
 					{
 						FString ExecName = Enum->GetNameStringByIndex(ExecIdx);
-						CreatePin(Direction, K2Schema->PC_Exec, TEXT(""), NULL, false, false, ExecName);
+						CreatePin(Direction, K2Schema->PC_Exec, FString(), nullptr, ExecName);
 					}
 				}
 				
@@ -787,12 +787,12 @@ void UK2Node_CallFunction::CreateExecPinsForFunctionCall(const UFunction* Functi
 		if (bCreateSingleExecInputPin)
 		{
 			// Single input exec pin
-			CreatePin(EGPD_Input, K2Schema->PC_Exec, TEXT(""), NULL, false, false, K2Schema->PN_Execute);
+			CreatePin(EGPD_Input, K2Schema->PC_Exec, FString(), nullptr, K2Schema->PN_Execute);
 		}
 
 		if (bCreateThenPin)
 		{
-			UEdGraphPin* OutputExecPin = CreatePin(EGPD_Output, K2Schema->PC_Exec, TEXT(""), NULL, false, false, K2Schema->PN_Then);
+			UEdGraphPin* OutputExecPin = CreatePin(EGPD_Output, K2Schema->PC_Exec, FString(), nullptr, K2Schema->PN_Then);
 			// Use 'completed' name for output pins on latent functions
 			if (Function->HasMetaData(FBlueprintMetadata::MD_Latent))
 			{
@@ -926,7 +926,7 @@ bool UK2Node_CallFunction::CreatePinsForFunctionCall(const UFunction* Function)
 
 		const EEdGraphPinDirection Direction = bIsFunctionInput ? EGPD_Input : EGPD_Output;
 
-		UEdGraphPin* Pin = CreatePin(Direction, TEXT(""), TEXT(""), NULL, false, bIsRefParam, Param->GetName());
+		UEdGraphPin* Pin = CreatePin(Direction, FString(), FString(), nullptr, Param->GetName(), EPinContainerType::None, bIsRefParam);
 		const bool bPinGood = (Pin != NULL) && K2Schema->ConvertPropertyToPinType(Param, /*out*/ Pin->PinType);
 
 		if (bPinGood)
@@ -939,7 +939,7 @@ bool UK2Node_CallFunction::CreatePinsForFunctionCall(const UFunction* Function)
 			}
 
 			//Flag pin as read only for const reference property
-			Pin->bDefaultValueIsIgnored = Param->HasAllPropertyFlags(CPF_ConstParm | CPF_ReferenceParm) && (!Function->HasMetaData(FBlueprintMetadata::MD_AutoCreateRefTerm) || Pin->PinType.bIsArray);
+			Pin->bDefaultValueIsIgnored = Param->HasAllPropertyFlags(CPF_ConstParm | CPF_ReferenceParm) && (!Function->HasMetaData(FBlueprintMetadata::MD_AutoCreateRefTerm) || Pin->PinType.IsContainer());
 
 			const bool bAdvancedPin = Param->HasAllPropertyFlags(CPF_AdvancedDisplay);
 			Pin->bAdvancedView = bAdvancedPin;
@@ -948,7 +948,15 @@ bool UK2Node_CallFunction::CreatePinsForFunctionCall(const UFunction* Function)
 				AdvancedPinDisplay = ENodeAdvancedPins::Hidden;
 			}
 
-			K2Schema->SetPinDefaultValue(Pin, Function, Param);
+			FString ParamValue;
+			if (K2Schema->FindFunctionParameterDefaultValue(Function, Param, ParamValue))
+			{
+				K2Schema->SetPinAutogeneratedDefaultValue(Pin, ParamValue);
+			}
+			else
+			{
+				K2Schema->SetPinAutogeneratedDefaultValueBasedOnType(Pin);
+			}
 			
 			if (PinsToHide.Contains(Pin->PinName))
 			{
@@ -1762,7 +1770,7 @@ void UK2Node_CallFunction::PostPasteNode()
 			if (bPinShouldBeHidden && !Pin->bHidden)
 			{
 				Pin->BreakAllPinLinks();
-				K2Schema->SetPinDefaultValueBasedOnType(Pin);
+				K2Schema->SetPinAutogeneratedDefaultValueBasedOnType(Pin);
 			}
 			Pin->bHidden = bPinShouldBeHidden;
 		}
@@ -1963,7 +1971,7 @@ void UK2Node_CallFunction::ExpandNode(class FKismetCompilerContext& CompilerCont
 		{
 			CompilerContext.MessageLog.Warning(*FString::Printf(*LOCTEXT("WrongEntryPointsNum", "%i entry points found while expanding node @@").ToString(), EntryPoints.Num()), this);
 		}
-		else if (auto BetterSelfPin = EntryPoints[0]->GetAutoWorldContextPin())
+		else if (UEdGraphPin* BetterSelfPin = EntryPoints[0]->GetAutoWorldContextPin())
 		{
 			FString const DefaultToSelfMetaValue = Function->GetMetaData(FBlueprintMetadata::MD_DefaultToSelf);
 			FString const WorldContextMetaValue = Function->GetMetaData(FBlueprintMetadata::MD_WorldContext);
@@ -1972,7 +1980,7 @@ void UK2Node_CallFunction::ExpandNode(class FKismetCompilerContext& CompilerCont
 			{
 				static void Connect(const FString& PinName, UK2Node* Node, UEdGraphPin* BetterSelf, const UEdGraphSchema_K2* InSchema, FCompilerResultsLog& MessageLog)
 				{
-					auto Pin = Node->FindPin(PinName);
+					UEdGraphPin* Pin = Node->FindPin(PinName);
 					if (!PinName.IsEmpty() && Pin && !Pin->LinkedTo.Num())
 					{
 						const bool bConnected = InSchema->TryCreateConnection(Pin, BetterSelf);
@@ -2017,7 +2025,7 @@ void UK2Node_CallFunction::ExpandNode(class FKismetCompilerContext& CompilerCont
 				if (EnumParamPin->Direction == EGPD_Input)
 				{
 					// Create normal exec input
-					UEdGraphPin* ExecutePin = CreatePin(EGPD_Input, Schema->PC_Exec, TEXT(""), NULL, false, false, Schema->PN_Execute);
+					UEdGraphPin* ExecutePin = CreatePin(EGPD_Input, Schema->PC_Exec, FString(), nullptr, Schema->PN_Execute);
 
 					// Create temp enum variable
 					UK2Node_TemporaryVariable* TempEnumVarNode = CompilerContext.SpawnIntermediateNode<UK2Node_TemporaryVariable>(this, SourceGraph);
@@ -2065,7 +2073,7 @@ void UK2Node_CallFunction::ExpandNode(class FKismetCompilerContext& CompilerCont
 				else if (EnumParamPin->Direction == EGPD_Output)
 				{
 					// Create normal exec output
-					UEdGraphPin* ExecutePin = CreatePin(EGPD_Output, Schema->PC_Exec, TEXT(""), NULL, false, false, Schema->PN_Execute);
+					UEdGraphPin* ExecutePin = CreatePin(EGPD_Output, Schema->PC_Exec, FString(), nullptr, Schema->PN_Execute);
 						
 					// Create a SwitchEnum node to switch on the output enum
 					UK2Node_SwitchEnum* SwitchEnumNode = CompilerContext.SpawnIntermediateNode<UK2Node_SwitchEnum>(this, SourceGraph);
@@ -2127,19 +2135,19 @@ void UK2Node_CallFunction::ExpandNode(class FKismetCompilerContext& CompilerCont
 				{
 					const bool bHasDefaultValue = !Pin->DefaultValue.IsEmpty() || Pin->DefaultObject || !Pin->DefaultTextValue.IsEmpty();
 
-					//default values can be reset when the pin is connected
-					const auto DefaultValue = Pin->DefaultValue;
-					const auto DefaultObject = Pin->DefaultObject;
-					const auto DefaultTextValue = Pin->DefaultTextValue;
-					const auto AutogeneratedDefaultValue = Pin->AutogeneratedDefaultValue;
+					// copy defaults as default values can be reset when the pin is connected
+					const FString DefaultValue = Pin->DefaultValue;
+					UObject* DefaultObject = Pin->DefaultObject;
+					const FText DefaultTextValue = Pin->DefaultTextValue;
+					bool bMatchesDefaults = Pin->DoesDefaultValueMatchAutogenerated();
 
 					UEdGraphPin* ValuePin = InnerHandleAutoCreateRef(this, Pin, CompilerContext, SourceGraph, bHasDefaultValue);
 					if ( ValuePin )
 					{
-						if (!DefaultObject && DefaultTextValue.IsEmpty() && DefaultValue.Equals(AutogeneratedDefaultValue, ESearchCase::CaseSensitive))
+						if (bMatchesDefaults)
 						{
 							// Use the latest code to set default value
-							Schema->SetPinDefaultValueBasedOnType(ValuePin);
+							Schema->SetPinAutogeneratedDefaultValueBasedOnType(ValuePin);
 						}
 						else
 						{
@@ -2163,12 +2171,12 @@ void UK2Node_CallFunction::ExpandNode(class FKismetCompilerContext& CompilerCont
 	// Then we go through and expand out array iteration if necessary
 	const bool bAllowMultipleSelfs = AllowMultipleSelfs(true);
 	UEdGraphPin* MultiSelf = Schema->FindSelfPin(*this, EEdGraphPinDirection::EGPD_Input);
-	if(bAllowMultipleSelfs && MultiSelf && !MultiSelf->PinType.bIsArray)
+	if(bAllowMultipleSelfs && MultiSelf && !MultiSelf->PinType.IsArray())
 	{
 		const bool bProperInputToExpandForEach = 
 			(1 == MultiSelf->LinkedTo.Num()) && 
 			(NULL != MultiSelf->LinkedTo[0]) && 
-			(MultiSelf->LinkedTo[0]->PinType.bIsArray);
+			(MultiSelf->LinkedTo[0]->PinType.IsArray());
 		if(bProperInputToExpandForEach)
 		{
 			CallForEachElementInArrayExpansion(this, MultiSelf, CompilerContext, SourceGraph);
@@ -2178,7 +2186,7 @@ void UK2Node_CallFunction::ExpandNode(class FKismetCompilerContext& CompilerCont
 
 UEdGraphPin* UK2Node_CallFunction::InnerHandleAutoCreateRef(UK2Node* Node, UEdGraphPin* Pin, FKismetCompilerContext& CompilerContext, UEdGraph* SourceGraph, bool bForceAssignment)
 {
-	const bool bAddAssigment = !Pin->PinType.bIsArray && bForceAssignment;
+	const bool bAddAssigment = !Pin->PinType.IsContainer() && bForceAssignment;
 
 	// ADD LOCAL VARIABLE
 	UK2Node_TemporaryVariable* LocalVariable = CompilerContext.SpawnIntermediateNode<UK2Node_TemporaryVariable>(Node, SourceGraph);
@@ -2190,7 +2198,7 @@ UEdGraphPin* UK2Node_CallFunction::InnerHandleAutoCreateRef(UK2Node* Node, UEdGr
 		if (!CompilerContext.GetSchema()->TryCreateConnection(LocalVariable->GetVariablePin(), Pin))
 		{
 			CompilerContext.MessageLog.Error(*LOCTEXT("AutoCreateRefTermPin_NotConnected", "AutoCreateRefTerm Expansion: Pin @@ cannot be connected to @@").ToString(), LocalVariable->GetVariablePin(), Pin);
-			return NULL;
+			return nullptr;
 		}
 	}
 	// ADD ASSIGMENT
@@ -2200,18 +2208,18 @@ UEdGraphPin* UK2Node_CallFunction::InnerHandleAutoCreateRef(UK2Node* Node, UEdGr
 		UK2Node_PureAssignmentStatement* AssignDefaultValue = CompilerContext.SpawnIntermediateNode<UK2Node_PureAssignmentStatement>(Node, SourceGraph);
 		AssignDefaultValue->AllocateDefaultPins();
 		const bool bVariableConnected = CompilerContext.GetSchema()->TryCreateConnection(AssignDefaultValue->GetVariablePin(), LocalVariable->GetVariablePin());
-		auto AssignInputPit = AssignDefaultValue->GetValuePin();
+		UEdGraphPin* AssignInputPit = AssignDefaultValue->GetValuePin();
 		const bool bPreviousInputSaved = AssignInputPit && CompilerContext.MovePinLinksToIntermediate(*Pin, *AssignInputPit).CanSafeConnect();
 		const bool bOutputConnected = CompilerContext.GetSchema()->TryCreateConnection(AssignDefaultValue->GetOutputPin(), Pin);
 		if (!bVariableConnected || !bOutputConnected || !bPreviousInputSaved)
 		{
 			CompilerContext.MessageLog.Error(*LOCTEXT("AutoCreateRefTermPin_AssignmentError", "AutoCreateRefTerm Expansion: Assignment Error @@").ToString(), AssignDefaultValue);
-			return NULL;
+			return nullptr;
 		}
-		CompilerContext.GetSchema()->SetPinDefaultValueBasedOnType(AssignDefaultValue->GetValuePin());
+		CompilerContext.GetSchema()->SetPinAutogeneratedDefaultValueBasedOnType(AssignDefaultValue->GetValuePin());
 		return AssignInputPit;
 	}
-	return NULL;
+	return nullptr;
 }
 
 void UK2Node_CallFunction::CallForEachElementInArrayExpansion(UK2Node* Node, UEdGraphPin* MultiSelf, FKismetCompilerContext& CompilerContext, UEdGraph* SourceGraph)
@@ -2221,7 +2229,7 @@ void UK2Node_CallFunction::CallForEachElementInArrayExpansion(UK2Node* Node, UEd
 	const bool bProperInputToExpandForEach = 
 		(1 == MultiSelf->LinkedTo.Num()) && 
 		(NULL != MultiSelf->LinkedTo[0]) && 
-		(MultiSelf->LinkedTo[0]->PinType.bIsArray);
+		(MultiSelf->LinkedTo[0]->PinType.IsArray());
 	ensure(bProperInputToExpandForEach);
 
 	UEdGraphPin* ThenPin = Node->FindPinChecked(Schema->PN_Then);
@@ -2371,7 +2379,7 @@ void UK2Node_CallFunction::ConformContainerPins()
 	{
 		if (Pin && !bOutPropagated)
 		{
-			if (Pin->LinkedTo.Num() != 0 || Pin->DefaultValue != Pin->AutogeneratedDefaultValue)
+			if (Pin->LinkedTo.Num() != 0 || !Pin->DoesDefaultValueMatchAutogenerated())
 			{
 				bOutPropagated = true;
 				if (Pin->LinkedTo.Num() != 0)
@@ -2390,7 +2398,7 @@ void UK2Node_CallFunction::ConformContainerPins()
 	{
 		if (Pin && !bOutPropagated)
 		{
-			if (Pin->LinkedTo.Num() != 0 || Pin->DefaultValue != Pin->AutogeneratedDefaultValue)
+			if (Pin->LinkedTo.Num() != 0 || !Pin->DoesDefaultValueMatchAutogenerated())
 			{
 				bOutPropagated = true;
 				if (Pin->LinkedTo.Num() != 0)
@@ -2597,7 +2605,7 @@ bool UK2Node_CallFunction::HasExternalDependencies(TArray<class UStruct*>* Optio
 	}
 
 	// All structures, that are required for the BP compilation, should be gathered
-	for(auto Pin : Pins)
+	for (UEdGraphPin* Pin : Pins)
 	{
 		UStruct* DepStruct = Pin ? Cast<UStruct>(Pin->PinType.PinSubCategoryObject.Get()) : nullptr;
 
@@ -2743,21 +2751,15 @@ bool UK2Node_CallFunction::IsConnectionDisallowed(const UEdGraphPin* MyPin, cons
 		}
 		else if (UFunction* TargetFunction = GetTargetFunction())
 		{
-			if( 
-				// Strictly speaking this first check is not needed, but by not disabling the connection here we get a better reason later:
-				(	(	FEdGraphUtilities::IsSetParam(TargetFunction, MyPin->PinName) &&
-					(OtherPin->PinType.IsContainer() && !MyPin->PinType.bIsSet) ) ||
-					(	FEdGraphUtilities::IsMapParam(TargetFunction, MyPin->PinName) &&
-					(OtherPin->PinType.IsContainer() && !MyPin->PinType.bIsMap) )  ||
-					(	FEdGraphUtilities::IsArrayDependentParam(TargetFunction, MyPin->PinName) &&
-					(OtherPin->PinType.IsContainer() && !MyPin->PinType.bIsArray) )
+			if (// Strictly speaking this first check is not needed, but by not disabling the connection here we get a better reason later:
+				(	OtherPin->PinType.IsContainer() 
+					// make sure we don't allow connections of mismatched container types (e.g. maps to arrays)
+					&& (OtherPin->PinType.ContainerType != MyPin->PinType.ContainerType)
+					&& (
+						(FEdGraphUtilities::IsSetParam(TargetFunction, MyPin->PinName) && !MyPin->PinType.IsSet()) ||
+						(FEdGraphUtilities::IsMapParam(TargetFunction, MyPin->PinName) && !MyPin->PinType.IsMap()) ||
+						(FEdGraphUtilities::IsArrayDependentParam(TargetFunction, MyPin->PinName) && !MyPin->PinType.IsArray())
 					)
-				&& 
-				// make sure we don't allow connections of mismatched container types (e.g. maps to arrays)
-				( 
-					OtherPin->PinType.bIsMap != MyPin->PinType.bIsMap || 
-					OtherPin->PinType.bIsSet != MyPin->PinType.bIsSet ||
-					OtherPin->PinType.bIsArray != MyPin->PinType.bIsArray 
 				)
 			)
 			{
diff --git a/Engine/Source/Editor/BlueprintGraph/Private/K2Node_CallFunctionOnMember.cpp b/Engine/Source/Editor/BlueprintGraph/Private/K2Node_CallFunctionOnMember.cpp
--- a/Engine/Source/Editor/BlueprintGraph/Private/K2Node_CallFunctionOnMember.cpp
+++ b/Engine/Source/Editor/BlueprintGraph/Private/K2Node_CallFunctionOnMember.cpp
@@ -20,18 +20,18 @@ UEdGraphPin* UK2Node_CallFunctionOnMember::CreateSelfPin(const UFunction* Functi
 {
 	const UEdGraphSchema_K2* K2Schema = GetDefault<UEdGraphSchema_K2>();
 
-	UEdGraphPin* SelfPin = NULL;
+	UEdGraphPin* SelfPin = nullptr;
 	if (MemberVariableToCallOn.IsSelfContext())
 	{
 		// This means the function is defined within the blueprint, so the pin should be a true "self" pin
-		SelfPin = CreatePin(EGPD_Input, K2Schema->PC_Object, K2Schema->PSC_Self, NULL, false, false, K2Schema->PN_Self);
+		SelfPin = CreatePin(EGPD_Input, K2Schema->PC_Object, K2Schema->PSC_Self, nullptr, K2Schema->PN_Self);
 	}
 	else
 	{
 		// This means that the function is declared in an external class, and should reference that class
-		SelfPin = CreatePin(EGPD_Input, K2Schema->PC_Object, TEXT(""), MemberVariableToCallOn.GetMemberParentClass(GetBlueprintClassFromNode()), false, false, K2Schema->PN_Self);
+		SelfPin = CreatePin(EGPD_Input, K2Schema->PC_Object, FString(), MemberVariableToCallOn.GetMemberParentClass(GetBlueprintClassFromNode()), K2Schema->PN_Self);
 	}
-	check(SelfPin != NULL);
+	check(SelfPin);
 
 	return SelfPin;
 }
diff --git a/Engine/Source/Editor/BlueprintGraph/Private/K2Node_CastByteToEnum.cpp b/Engine/Source/Editor/BlueprintGraph/Private/K2Node_CastByteToEnum.cpp
--- a/Engine/Source/Editor/BlueprintGraph/Private/K2Node_CastByteToEnum.cpp
+++ b/Engine/Source/Editor/BlueprintGraph/Private/K2Node_CastByteToEnum.cpp
@@ -32,8 +32,8 @@ void UK2Node_CastByteToEnum::AllocateDefaultPins()
 {
 	const UEdGraphSchema_K2* Schema = GetDefault<UEdGraphSchema_K2>();
 	
-	CreatePin(EGPD_Input, Schema->PC_Byte, TEXT(""), NULL, false, false, ByteInputPinName);
-	CreatePin(EGPD_Output, Schema->PC_Byte, TEXT(""), Enum, false, false, Schema->PN_ReturnValue);
+	CreatePin(EGPD_Input, Schema->PC_Byte, FString(), nullptr, ByteInputPinName);
+	CreatePin(EGPD_Output, Schema->PC_Byte, FString(), Enum, Schema->PN_ReturnValue);
 }
 
 FText UK2Node_CastByteToEnum::GetTooltipText() const
diff --git a/Engine/Source/Editor/BlueprintGraph/Private/K2Node_ClassDynamicCast.cpp b/Engine/Source/Editor/BlueprintGraph/Private/K2Node_ClassDynamicCast.cpp
--- a/Engine/Source/Editor/BlueprintGraph/Private/K2Node_ClassDynamicCast.cpp
+++ b/Engine/Source/Editor/BlueprintGraph/Private/K2Node_ClassDynamicCast.cpp
@@ -42,24 +42,24 @@ void UK2Node_ClassDynamicCast::AllocateDefaultPins()
 	if (!bIsPureCast)
 	{
 		// Input - Execution Pin
-		CreatePin(EGPD_Input, K2Schema->PC_Exec, TEXT(""), NULL, false, false, K2Schema->PN_Execute);
+		CreatePin(EGPD_Input, K2Schema->PC_Exec, FString(), nullptr, K2Schema->PN_Execute);
 
 		// Output - Execution Pins
-		CreatePin(EGPD_Output, K2Schema->PC_Exec, TEXT(""), NULL, false, false, K2Schema->PN_CastSucceeded);
-		CreatePin(EGPD_Output, K2Schema->PC_Exec, TEXT(""), NULL, false, false, K2Schema->PN_CastFailed);
+		CreatePin(EGPD_Output, K2Schema->PC_Exec, FString(), nullptr, K2Schema->PN_CastSucceeded);
+		CreatePin(EGPD_Output, K2Schema->PC_Exec, FString(), nullptr, K2Schema->PN_CastFailed);
 	}
 
 	// Input - Source type Pin
-	CreatePin(EGPD_Input, K2Schema->PC_Class, TEXT(""), UObject::StaticClass(), false, false, FClassDynamicCastHelper::GetClassToCastName());
+	CreatePin(EGPD_Input, K2Schema->PC_Class, FString(), UObject::StaticClass(), FClassDynamicCastHelper::GetClassToCastName());
 
 	// Output - Data Pin
 	if (TargetType != NULL)
 	{
 		const FString CastResultPinName = K2Schema->PN_CastedValuePrefix + TargetType->GetDisplayNameText().ToString();
-		CreatePin(EGPD_Output, K2Schema->PC_Class, TEXT(""), *TargetType, false, false, CastResultPinName);
+		CreatePin(EGPD_Output, K2Schema->PC_Class, FString(), *TargetType, CastResultPinName);
 	}
 
-	UEdGraphPin* BoolSuccessPin = CreatePin(EGPD_Output, K2Schema->PC_Boolean, TEXT(""), nullptr, /*bIsArray =*/false, /*bIsReference =*/false, FClassDynamicCastHelper::CastSuccessPinName);
+	UEdGraphPin* BoolSuccessPin = CreatePin(EGPD_Output, K2Schema->PC_Boolean, FString(), nullptr, FClassDynamicCastHelper::CastSuccessPinName);
 	BoolSuccessPin->bHidden = !bIsPureCast;
 
 	UK2Node::AllocateDefaultPins();
diff --git a/Engine/Source/Editor/BlueprintGraph/Private/K2Node_CommutativeAssociativeBinaryOperator.cpp b/Engine/Source/Editor/BlueprintGraph/Private/K2Node_CommutativeAssociativeBinaryOperator.cpp
--- a/Engine/Source/Editor/BlueprintGraph/Private/K2Node_CommutativeAssociativeBinaryOperator.cpp
+++ b/Engine/Source/Editor/BlueprintGraph/Private/K2Node_CommutativeAssociativeBinaryOperator.cpp
@@ -156,13 +156,11 @@ void UK2Node_CommutativeAssociativeBinaryOperator::AddInputPinInner(int32 Additi
 		InputType.PinCategory, 
 		InputType.PinSubCategory, 
 		InputType.PinSubCategoryObject.Get(), 
-		InputType.bIsArray, 
-		InputType.bIsReference, 
 		*GetNameForPin(AdditionalPinIndex + BinaryOperatorInputsNum),
+		InputType.ContainerType, 
+		InputType.bIsReference, 
 		false,
 		INDEX_NONE,
-		InputType.bIsSet,
-		InputType.bIsMap,
 		InputType.PinValueType
 	);
 }
diff --git a/Engine/Source/Editor/BlueprintGraph/Private/K2Node_Composite.cpp b/Engine/Source/Editor/BlueprintGraph/Private/K2Node_Composite.cpp
--- a/Engine/Source/Editor/BlueprintGraph/Private/K2Node_Composite.cpp
+++ b/Engine/Source/Editor/BlueprintGraph/Private/K2Node_Composite.cpp
@@ -21,6 +21,7 @@ UK2Node_Composite::UK2Node_Composite(const FObjectInitializer& ObjectInitializer
 void UK2Node_Composite::AllocateDefaultPins()
 {
 	UK2Node::AllocateDefaultPins();
+	const UEdGraphSchema_K2* Schema = GetDefault<UEdGraphSchema_K2>();
 
 	if (OutputSourceNode)
 	{
@@ -30,7 +31,7 @@ void UK2Node_Composite::AllocateDefaultPins()
 			if (PortPin->Direction == EGPD_Input)
 			{
 				UEdGraphPin* NewPin = CreatePin(UEdGraphPin::GetComplementaryDirection(PortPin->Direction), PortPin->PinType, PortPin->PinName);
-				NewPin->DefaultValue = NewPin->AutogeneratedDefaultValue = PortPin->DefaultValue;
+				Schema->SetPinAutogeneratedDefaultValue(NewPin, PortPin->GetDefaultAsString());
 			}
 		}
 	}
@@ -43,7 +44,7 @@ void UK2Node_Composite::AllocateDefaultPins()
 			if (PortPin->Direction == EGPD_Output)
 			{
 				UEdGraphPin* NewPin = CreatePin(UEdGraphPin::GetComplementaryDirection(PortPin->Direction), PortPin->PinType, PortPin->PinName);
-				NewPin->DefaultValue = NewPin->AutogeneratedDefaultValue = PortPin->DefaultValue;
+				Schema->SetPinAutogeneratedDefaultValue(NewPin, PortPin->GetDefaultAsString());
 			}
 		}
 	}
diff --git a/Engine/Source/Editor/BlueprintGraph/Private/K2Node_ConstructObjectFromClass.cpp b/Engine/Source/Editor/BlueprintGraph/Private/K2Node_ConstructObjectFromClass.cpp
--- a/Engine/Source/Editor/BlueprintGraph/Private/K2Node_ConstructObjectFromClass.cpp
+++ b/Engine/Source/Editor/BlueprintGraph/Private/K2Node_ConstructObjectFromClass.cpp
@@ -44,24 +44,24 @@ void UK2Node_ConstructObjectFromClass::AllocateDefaultPins()
 	const UEdGraphSchema_K2* K2Schema = GetDefault<UEdGraphSchema_K2>();
 
 	// Add execution pins
-	CreatePin(EGPD_Input, K2Schema->PC_Exec, TEXT(""), NULL, false, false, K2Schema->PN_Execute);
-	CreatePin(EGPD_Output, K2Schema->PC_Exec, TEXT(""), NULL, false, false, K2Schema->PN_Then);
+	CreatePin(EGPD_Input, K2Schema->PC_Exec, FString(), nullptr, K2Schema->PN_Execute);
+	CreatePin(EGPD_Output, K2Schema->PC_Exec, FString(), nullptr, K2Schema->PN_Then);
 
 	// If required add the world context pin
 	if (UseWorldContext())
 	{
-		CreatePin(EGPD_Input, K2Schema->PC_Object, TEXT(""), UObject::StaticClass(), false, false, FK2Node_ConstructObjectFromClassHelper::WorldContextPinName);
+		CreatePin(EGPD_Input, K2Schema->PC_Object, FString(), UObject::StaticClass(), FK2Node_ConstructObjectFromClassHelper::WorldContextPinName);
 	}
 
 	// Add blueprint pin
-	UEdGraphPin* ClassPin = CreatePin(EGPD_Input, K2Schema->PC_Class, TEXT(""), GetClassPinBaseClass(), false, false, FK2Node_ConstructObjectFromClassHelper::ClassPinName);
+	UEdGraphPin* ClassPin = CreatePin(EGPD_Input, K2Schema->PC_Class, FString(), GetClassPinBaseClass(), FK2Node_ConstructObjectFromClassHelper::ClassPinName);
 	
 	// Result pin
-	UEdGraphPin* ResultPin = CreatePin(EGPD_Output, K2Schema->PC_Object, TEXT(""), GetClassPinBaseClass(), false, false, K2Schema->PN_ReturnValue);
+	UEdGraphPin* ResultPin = CreatePin(EGPD_Output, K2Schema->PC_Object, FString(), GetClassPinBaseClass(), K2Schema->PN_ReturnValue);
 	
 	if (UseOuter())
 	{
-		UEdGraphPin* OuterPin = CreatePin(EGPD_Input, K2Schema->PC_Object, TEXT(""), UObject::StaticClass(), false, false, FK2Node_ConstructObjectFromClassHelper::OuterPinName);
+		UEdGraphPin* OuterPin = CreatePin(EGPD_Input, K2Schema->PC_Object, FString(), UObject::StaticClass(), FK2Node_ConstructObjectFromClassHelper::OuterPinName);
 	}
 
 	Super::AllocateDefaultPins();
@@ -111,7 +111,7 @@ void UK2Node_ConstructObjectFromClass::CreatePinsForClass(UClass* InClass, TArra
 			!bIsDelegate &&
 			(NULL == FindPin(Property->GetName()) ) )
 		{
-			UEdGraphPin* Pin = CreatePin(EGPD_Input, TEXT(""), TEXT(""), NULL, false, false, Property->GetName());
+			UEdGraphPin* Pin = CreatePin(EGPD_Input, FString(), FString(), nullptr, Property->GetName());
 			const bool bPinGood = (Pin != NULL) && K2Schema->ConvertPropertyToPinType(Property, /*out*/ Pin->PinType);
 			
 			if (OutClassPins && Pin)
@@ -124,7 +124,7 @@ void UK2Node_ConstructObjectFromClass::CreatePinsForClass(UClass* InClass, TArra
 				FString DefaultValueAsString;
 				const bool bDefaultValueSet = FBlueprintEditorUtils::PropertyValueToString(Property, reinterpret_cast<const uint8*>(ClassDefaultObject), DefaultValueAsString);
 				check( bDefaultValueSet );
-				K2Schema->TrySetDefaultValue(*Pin, DefaultValueAsString);
+				K2Schema->SetPinAutogeneratedDefaultValue(Pin, DefaultValueAsString);
 			}
 
 			// Copy tooltip from the property.
diff --git a/Engine/Source/Editor/BlueprintGraph/Private/K2Node_ConvertAsset.cpp b/Engine/Source/Editor/BlueprintGraph/Private/K2Node_ConvertAsset.cpp
--- a/Engine/Source/Editor/BlueprintGraph/Private/K2Node_ConvertAsset.cpp
+++ b/Engine/Source/Editor/BlueprintGraph/Private/K2Node_ConvertAsset.cpp
@@ -74,18 +74,20 @@ void UK2Node_ConvertAsset::RefreshPinTypes()
 		const FString InputCategory = bIsConnected
 			? (bIsAssetClass ? K2Schema->PC_AssetClass : K2Schema->PC_Asset)
 			: K2Schema->PC_Wildcard;
-		InutPin->PinType = FEdGraphPinType(InputCategory, FString(), TargetType, false, false, false, false, FEdGraphTerminalType() );
+		InutPin->PinType = FEdGraphPinType(InputCategory, FString(), TargetType, EPinContainerType::None, false, FEdGraphTerminalType() );
 
 		const FString OutputCategory = bIsConnected
 			? (bIsAssetClass ? K2Schema->PC_Class : K2Schema->PC_Object)
 			: K2Schema->PC_Wildcard;
-		OutputPin->PinType = FEdGraphPinType(OutputCategory, FString(), TargetType, false, false, false, false, FEdGraphTerminalType() );
+		OutputPin->PinType = FEdGraphPinType(OutputCategory, FString(), TargetType, EPinContainerType::None, false, FEdGraphTerminalType() );
 
 		PinTypeChanged(InutPin);
 		PinTypeChanged(OutputPin);
 
 		if (OutputPin->LinkedTo.Num())
 		{
+			TArray<UEdGraphPin*> PinsToUnlink = OutputPin->LinkedTo;
+
 			UClass const* CallingContext = NULL;
 			if (UBlueprint const* Blueprint = GetBlueprint())
 			{
@@ -96,7 +98,7 @@ void UK2Node_ConvertAsset::RefreshPinTypes()
 				}
 			}
 
-			for (auto TargetPin : OutputPin->LinkedTo)
+			for (auto TargetPin : PinsToUnlink)
 			{
 				if (TargetPin && !K2Schema->ArePinsCompatible(OutputPin, TargetPin, CallingContext))
 				{
@@ -120,14 +122,16 @@ void UK2Node_ConvertAsset::NotifyPinConnectionListChanged(UEdGraphPin* Pin)
 	if (Pin && (UK2Node_ConvertAssetImpl::InputPinName == Pin->PinName))
 	{
 		RefreshPinTypes();
+
+		GetGraph()->NotifyGraphChanged();
 	}
 }
 
 void UK2Node_ConvertAsset::AllocateDefaultPins()
 {
 	const UEdGraphSchema_K2* K2Schema = CastChecked<UEdGraphSchema_K2>(GetSchema());
-	CreatePin(EGPD_Input, K2Schema->PC_Wildcard, TEXT(""), nullptr, false, false, UK2Node_ConvertAssetImpl::InputPinName);
-	CreatePin(EGPD_Output, K2Schema->PC_Wildcard, TEXT(""), nullptr, false, false, UK2Node_ConvertAssetImpl::OutputPinName);
+	CreatePin(EGPD_Input, K2Schema->PC_Wildcard, FString(), nullptr, UK2Node_ConvertAssetImpl::InputPinName);
+	CreatePin(EGPD_Output, K2Schema->PC_Wildcard, FString(), nullptr, UK2Node_ConvertAssetImpl::OutputPinName);
 }
 
 void UK2Node_ConvertAsset::GetMenuActions(FBlueprintActionDatabaseRegistrar& ActionRegistrar) const
diff --git a/Engine/Source/Editor/BlueprintGraph/Private/K2Node_Copy.cpp b/Engine/Source/Editor/BlueprintGraph/Private/K2Node_Copy.cpp
--- a/Engine/Source/Editor/BlueprintGraph/Private/K2Node_Copy.cpp
+++ b/Engine/Source/Editor/BlueprintGraph/Private/K2Node_Copy.cpp
@@ -66,8 +66,8 @@ void UK2Node_Copy::AllocateDefaultPins()
 {
 	const UEdGraphSchema_K2* Schema = GetDefault<UEdGraphSchema_K2>();
 
-	CreatePin(EGPD_Input, Schema->PC_Wildcard, TEXT(""), NULL, false, false, Schema->PN_Item);
-	CreatePin(EGPD_Output, Schema->PC_Wildcard, TEXT(""), NULL, false, false, Schema->PN_ReturnValue);
+	CreatePin(EGPD_Input, Schema->PC_Wildcard, FString(), nullptr, Schema->PN_Item);
+	CreatePin(EGPD_Output, Schema->PC_Wildcard, FString(), nullptr, Schema->PN_ReturnValue);
 
 	Super::AllocateDefaultPins();
 }
@@ -158,7 +158,7 @@ void UK2Node_Copy::PinTypeChanged(UEdGraphPin* Pin)
 				Schema->RecombinePin(ReturnPin->SubPins[0]);
 			}
 			ReturnPin->PinType = Pin->PinType;
-			Schema->SetPinDefaultValueBasedOnType(ReturnPin);
+			Schema->SetPinAutogeneratedDefaultValueBasedOnType(ReturnPin);
 		}
 	}
 
@@ -248,10 +248,10 @@ void UK2Node_Copy::PropagatePinType(FEdGraphPinType& InType)
 	InputPin->PinType = InType;
 	ResultPin->PinType = InType;
 
-	InputPin->PinType.bIsArray = false;
+	InputPin->PinType.ContainerType = EPinContainerType::None;
 	InputPin->PinType.bIsReference = false;
 
-	ResultPin->PinType.bIsArray = false;
+	ResultPin->PinType.ContainerType = EPinContainerType::None;
 	ResultPin->PinType.bIsReference = false;
 
 	const UEdGraphSchema_K2* Schema = GetDefault<UEdGraphSchema_K2>();
diff --git a/Engine/Source/Editor/BlueprintGraph/Private/K2Node_CreateDelegate.cpp b/Engine/Source/Editor/BlueprintGraph/Private/K2Node_CreateDelegate.cpp
--- a/Engine/Source/Editor/BlueprintGraph/Private/K2Node_CreateDelegate.cpp
+++ b/Engine/Source/Editor/BlueprintGraph/Private/K2Node_CreateDelegate.cpp
@@ -32,12 +32,12 @@ void UK2Node_CreateDelegate::AllocateDefaultPins()
 {
 	const UEdGraphSchema_K2* K2Schema = GetDefault<UEdGraphSchema_K2>();
 
-	if (UEdGraphPin* ObjPin = CreatePin(EGPD_Input, K2Schema->PC_Object, TEXT(""), UObject::StaticClass(), false, false, K2Schema->PN_Self))
+	if (UEdGraphPin* ObjPin = CreatePin(EGPD_Input, K2Schema->PC_Object, FString(), UObject::StaticClass(), K2Schema->PN_Self))
 	{
 		ObjPin->PinFriendlyName = NSLOCTEXT("K2Node", "CreateDelegate_ObjectInputName", "Object");
 	}
 
-	if(UEdGraphPin* DelegatePin = CreatePin(EGPD_Output, K2Schema->PC_Delegate, TEXT(""), NULL, false, false, FK2Node_CreateDelegate_Helper::DelegateOutputName))
+	if(UEdGraphPin* DelegatePin = CreatePin(EGPD_Output, K2Schema->PC_Delegate, FString(), nullptr, FK2Node_CreateDelegate_Helper::DelegateOutputName))
 	{
 		DelegatePin->PinFriendlyName = NSLOCTEXT("K2Node", "CreateDelegate_DelegateOutName", "Event");
 	}
diff --git a/Engine/Source/Editor/BlueprintGraph/Private/K2Node_CustomEvent.cpp b/Engine/Source/Editor/BlueprintGraph/Private/K2Node_CustomEvent.cpp
--- a/Engine/Source/Editor/BlueprintGraph/Private/K2Node_CustomEvent.cpp
+++ b/Engine/Source/Editor/BlueprintGraph/Private/K2Node_CustomEvent.cpp
@@ -144,7 +144,7 @@ void UK2Node_CustomEvent::Serialize(FArchive& Ar)
 					{
 						if (PinInfo->PinName == Pin->PinName)
 						{
-							Pin->PinType.bIsConst = PinInfo->PinType.bIsConst = PinInfo->PinType.bIsArray || PinInfo->PinType.bIsReference;
+							Pin->PinType.bIsConst = PinInfo->PinType.bIsConst = (PinInfo->PinType.IsContainer() || PinInfo->PinType.bIsReference);
 							break;
 						}
 					}
@@ -224,10 +224,23 @@ bool UK2Node_CustomEvent::CanCreateUserDefinedPin(const FEdGraphPinType& InPinTy
 UEdGraphPin* UK2Node_CustomEvent::CreatePinFromUserDefinition(const TSharedPtr<FUserPinInfo> NewPinInfo)
 {
 	UEdGraphPin* NewPin = CreatePin(EGPD_Output, NewPinInfo->PinType, NewPinInfo->PinName);
-	NewPin->DefaultValue = NewPin->AutogeneratedDefaultValue = NewPinInfo->PinDefaultValue;
+	const UEdGraphSchema_K2* K2Schema = GetDefault<UEdGraphSchema_K2>();
+	K2Schema->SetPinAutogeneratedDefaultValue(NewPin, NewPinInfo->PinDefaultValue);
 	return NewPin;
 }
 
+bool UK2Node_CustomEvent::ModifyUserDefinedPinDefaultValue(TSharedPtr<FUserPinInfo> PinInfo, const FString& NewDefaultValue)
+{
+	if (Super::ModifyUserDefinedPinDefaultValue(PinInfo, NewDefaultValue))
+	{
+		const UEdGraphSchema_K2* K2Schema = GetDefault<UEdGraphSchema_K2>();
+		K2Schema->HandleParameterDefaultValueChanged(this);
+
+		return true;
+	}
+	return false;
+}
+
 void UK2Node_CustomEvent::RenameCustomEventCloseToName(int32 StartIndex)
 {
 	bool bFoundName = false;
diff --git a/Engine/Source/Editor/BlueprintGraph/Private/K2Node_DelegateSet.cpp b/Engine/Source/Editor/BlueprintGraph/Private/K2Node_DelegateSet.cpp
--- a/Engine/Source/Editor/BlueprintGraph/Private/K2Node_DelegateSet.cpp
+++ b/Engine/Source/Editor/BlueprintGraph/Private/K2Node_DelegateSet.cpp
@@ -141,11 +141,11 @@ void UK2Node_DelegateSet::AllocateDefaultPins()
 	// Cache off the delegate signature, which will update the DelegatePropertyName as well, if it's been redirected
 	UFunction* DelegateSignature = GetDelegateSignature();
 
-	CreatePin(EGPD_Input, K2Schema->PC_Exec, TEXT(""), NULL, false, false, K2Schema->PN_Execute);
-	CreatePin(EGPD_Input, K2Schema->PC_Object, TEXT(""), DelegatePropertyClass, false, false, DelegatePropertyName.ToString());
+	CreatePin(EGPD_Input, K2Schema->PC_Exec, FString(), nullptr, K2Schema->PN_Execute);
+	CreatePin(EGPD_Input, K2Schema->PC_Object, FString(), DelegatePropertyClass, DelegatePropertyName.ToString());
 
-	CreatePin(EGPD_Output, K2Schema->PC_Exec, TEXT(""), NULL, false, false, K2Schema->PN_Then);
-	CreatePin(EGPD_Output, K2Schema->PC_Exec, TEXT(""), NULL, false, false, K2Schema->PN_DelegateEntry);
+	CreatePin(EGPD_Output, K2Schema->PC_Exec, FString(), nullptr, K2Schema->PN_Then);
+	CreatePin(EGPD_Output, K2Schema->PC_Exec, FString(), nullptr, K2Schema->PN_DelegateEntry);
 	
 	CreatePinsForFunctionEntryExit(DelegateSignature, true);
 
diff --git a/Engine/Source/Editor/BlueprintGraph/Private/K2Node_DoOnceMultiInput.cpp b/Engine/Source/Editor/BlueprintGraph/Private/K2Node_DoOnceMultiInput.cpp
--- a/Engine/Source/Editor/BlueprintGraph/Private/K2Node_DoOnceMultiInput.cpp
+++ b/Engine/Source/Editor/BlueprintGraph/Private/K2Node_DoOnceMultiInput.cpp
@@ -185,17 +185,17 @@ void UK2Node_DoOnceMultiInput::AllocateDefaultPins()
 	const UEdGraphSchema_K2* K2Schema = GetDefault<UEdGraphSchema_K2>();
 
 	FText InputPinAName = GetNameForPin(0, true);
-	UEdGraphPin* InputPinA = CreatePin(EGPD_Input, K2Schema->PC_Exec, TEXT(""), NULL, false, false, InputPinAName.BuildSourceString());
+	UEdGraphPin* InputPinA = CreatePin(EGPD_Input, K2Schema->PC_Exec, FString(), nullptr, InputPinAName.BuildSourceString());
 	InputPinA->PinFriendlyName = InputPinAName;
 
 	FText OutputPinAName = GetNameForPin(0, false);
-	UEdGraphPin* OutputPinA = CreatePin(EGPD_Output, K2Schema->PC_Exec, TEXT(""), NULL, false, false, OutputPinAName.BuildSourceString());
+	UEdGraphPin* OutputPinA = CreatePin(EGPD_Output, K2Schema->PC_Exec, FString(), nullptr, OutputPinAName.BuildSourceString());
 	OutputPinA->PinFriendlyName = OutputPinAName;
 
-	UEdGraphPin* DoOnceResetIn = CreatePin(EGPD_Input, K2Schema->PC_Exec, TEXT(""), NULL, false, false, TEXT("Reset In"));
+	UEdGraphPin* DoOnceResetIn = CreatePin(EGPD_Input, K2Schema->PC_Exec, FString(), nullptr, TEXT("Reset In"));
 	DoOnceResetIn->PinFriendlyName = LOCTEXT("DoOnceResetIn", "Reset In");
 
-	UEdGraphPin* DoOnceResetOut = CreatePin(EGPD_Output, K2Schema->PC_Exec, TEXT(""), NULL, false, false, TEXT("Reset Out"));
+	UEdGraphPin* DoOnceResetOut = CreatePin(EGPD_Output, K2Schema->PC_Exec, FString(), nullptr, TEXT("Reset Out"));
 	DoOnceResetOut->PinFriendlyName = LOCTEXT("DoOnceResetOut", "Reset Out");
 
 	for (int32 i = 0; i < NumAdditionalInputs; ++i)
diff --git a/Engine/Source/Editor/BlueprintGraph/Private/K2Node_DynamicCast.cpp b/Engine/Source/Editor/BlueprintGraph/Private/K2Node_DynamicCast.cpp
--- a/Engine/Source/Editor/BlueprintGraph/Private/K2Node_DynamicCast.cpp
+++ b/Engine/Source/Editor/BlueprintGraph/Private/K2Node_DynamicCast.cpp
@@ -47,15 +47,15 @@ void UK2Node_DynamicCast::AllocateDefaultPins()
 	if (!bIsPureCast)
 	{
 		// Input - Execution Pin
-		CreatePin(EGPD_Input, K2Schema->PC_Exec, TEXT(""), NULL, false, false, K2Schema->PN_Execute);
+		CreatePin(EGPD_Input, K2Schema->PC_Exec, FString(), nullptr, K2Schema->PN_Execute);
 
 		// Output - Execution Pins
-		CreatePin(EGPD_Output, K2Schema->PC_Exec, TEXT(""), NULL, false, false, K2Schema->PN_CastSucceeded);
-		CreatePin(EGPD_Output, K2Schema->PC_Exec, TEXT(""), NULL, false, false, K2Schema->PN_CastFailed);
+		CreatePin(EGPD_Output, K2Schema->PC_Exec, FString(), nullptr, K2Schema->PN_CastSucceeded);
+		CreatePin(EGPD_Output, K2Schema->PC_Exec, FString(), nullptr, K2Schema->PN_CastFailed);
 	}
 
 	// Input - Source type Pin
-	CreatePin(EGPD_Input, K2Schema->PC_Wildcard, TEXT(""), UObject::StaticClass(), false, false, K2Schema->PN_ObjectToCast);
+	CreatePin(EGPD_Input, K2Schema->PC_Wildcard, FString(), UObject::StaticClass(), K2Schema->PN_ObjectToCast);
 
 	// Output - Data Pin
 	if (TargetType != NULL)
@@ -63,15 +63,15 @@ void UK2Node_DynamicCast::AllocateDefaultPins()
 		FString CastResultPinName = K2Schema->PN_CastedValuePrefix + TargetType->GetDisplayNameText().ToString();
 		if (TargetType->IsChildOf(UInterface::StaticClass()))
 		{
-			CreatePin(EGPD_Output, K2Schema->PC_Interface, TEXT(""), *TargetType, false, false, CastResultPinName);
+			CreatePin(EGPD_Output, K2Schema->PC_Interface, FString(), *TargetType, CastResultPinName);
 		}
 		else 
 		{
-			CreatePin(EGPD_Output, K2Schema->PC_Object, TEXT(""), *TargetType, false, false, CastResultPinName);
+			CreatePin(EGPD_Output, K2Schema->PC_Object, FString(), *TargetType, CastResultPinName);
 		}
 	}
 
-	UEdGraphPin* BoolSuccessPin = CreatePin(EGPD_Output, K2Schema->PC_Boolean, TEXT(""), nullptr, /*bIsArray =*/false, /*bIsReference =*/false, UK2Node_DynamicCastImpl::CastSuccessPinName);
+	UEdGraphPin* BoolSuccessPin = CreatePin(EGPD_Output, K2Schema->PC_Boolean, FString(), nullptr, UK2Node_DynamicCastImpl::CastSuccessPinName);
 	BoolSuccessPin->bHidden = !bIsPureCast;
 
 	Super::AllocateDefaultPins();
diff --git a/Engine/Source/Editor/BlueprintGraph/Private/K2Node_EaseFunction.cpp b/Engine/Source/Editor/BlueprintGraph/Private/K2Node_EaseFunction.cpp
--- a/Engine/Source/Editor/BlueprintGraph/Private/K2Node_EaseFunction.cpp
+++ b/Engine/Source/Editor/BlueprintGraph/Private/K2Node_EaseFunction.cpp
@@ -86,36 +86,36 @@ void UK2Node_EaseFunction::AllocateDefaultPins()
 	// Add the first pin representing all available easing functions. If EEasingFunc changes its name this will fail a runtime check!
 	UEnum * EaseFuncEnum = FindObject<UEnum>(ANY_PACKAGE, TEXT("EEasingFunc"), true);
 	check(EaseFuncEnum != NULL);
-	CachedEaseFuncPin = CreatePin(EGPD_Input, K2Schema->PC_Byte, TEXT(""), EaseFuncEnum, false, false, FEaseFunctionNodeHelper::GetEaseFuncPinName());
+	CachedEaseFuncPin = CreatePin(EGPD_Input, K2Schema->PC_Byte, FString(), EaseFuncEnum, FEaseFunctionNodeHelper::GetEaseFuncPinName());
 	SetPinToolTip(*CachedEaseFuncPin, LOCTEXT("EaseFunsPinDescription", "Specifies the desired ease function to be applied. If connected no customization is possible."));
 
 	// Make sure that the default value is set correctly if none has been set
-	K2Schema->SetPinDefaultValue(CachedEaseFuncPin);
+	K2Schema->SetPinAutogeneratedDefaultValueBasedOnType(CachedEaseFuncPin);
 
-	UEdGraphPin* AlphaPin = CreatePin(EGPD_Input, K2Schema->PC_Float, TEXT(""), NULL, false, false, FEaseFunctionNodeHelper::GetAlphaPinName());
+	UEdGraphPin* AlphaPin = CreatePin(EGPD_Input, K2Schema->PC_Float, FString(), nullptr, FEaseFunctionNodeHelper::GetAlphaPinName());
 	SetPinToolTip(*AlphaPin, LOCTEXT("AlphaPinTooltip", "Alpha value used to specify the easing in time."));
 
 	// Add wildcard pins for A, B and the return Pin
-	UEdGraphPin* APin = CreatePin(EGPD_Input, K2Schema->PC_Wildcard, TEXT(""), NULL, false, false, FEaseFunctionNodeHelper::GetAPinName());
+	UEdGraphPin* APin = CreatePin(EGPD_Input, K2Schema->PC_Wildcard, FString(), nullptr, FEaseFunctionNodeHelper::GetAPinName());
 	SetPinToolTip(*APin, LOCTEXT("APinDescription", "Easing start value"));
 
-	UEdGraphPin* BPin = CreatePin(EGPD_Input, K2Schema->PC_Wildcard, TEXT(""), NULL, false, false, FEaseFunctionNodeHelper::GetBPinName());
+	UEdGraphPin* BPin = CreatePin(EGPD_Input, K2Schema->PC_Wildcard, FString(), nullptr, FEaseFunctionNodeHelper::GetBPinName());
 	SetPinToolTip(*BPin, LOCTEXT("BPinDescription", "Easing end value"));
 
-	UEdGraphPin* ResultPin = CreatePin(EGPD_Output, K2Schema->PC_Wildcard, TEXT(""), NULL, false, false, FEaseFunctionNodeHelper::GetResultPinName());
+	UEdGraphPin* ResultPin = CreatePin(EGPD_Output, K2Schema->PC_Wildcard, FString(), nullptr, FEaseFunctionNodeHelper::GetResultPinName());
 	SetPinToolTip(*ResultPin, LOCTEXT("ResultPinDescription", "Easing result value"));
 
-	UEdGraphPin* ShortestPathPin = CreatePin(EGPD_Input, K2Schema->PC_Boolean, TEXT(""), NULL, false, false, FEaseFunctionNodeHelper::GetShortestPathPinName());
+	UEdGraphPin* ShortestPathPin = CreatePin(EGPD_Input, K2Schema->PC_Boolean, FString(), nullptr, FEaseFunctionNodeHelper::GetShortestPathPinName());
 	SetPinToolTip(*ShortestPathPin, LOCTEXT("ShortestPathPinTooltip", "When interpolating the shortest path should be taken."));
-	ShortestPathPin->DefaultValue = ShortestPathPin->AutogeneratedDefaultValue = TEXT("true");
+	K2Schema->SetPinAutogeneratedDefaultValue(ShortestPathPin, TEXT("true"));
 
-	UEdGraphPin* BlendExpPin = CreatePin(EGPD_Input, K2Schema->PC_Float, TEXT(""), NULL, false, false, FEaseFunctionNodeHelper::GetBlendExpPinName());
+	UEdGraphPin* BlendExpPin = CreatePin(EGPD_Input, K2Schema->PC_Float, FString(), nullptr, FEaseFunctionNodeHelper::GetBlendExpPinName());
 	SetPinToolTip(*BlendExpPin, LOCTEXT("BlendExpPinDescription", "Blend Exponent for basic ease functions"));
-	BlendExpPin->DefaultValue = BlendExpPin->AutogeneratedDefaultValue = TEXT("2.0");
+	K2Schema->SetPinAutogeneratedDefaultValue(BlendExpPin, TEXT("2.0"));
 
-	UEdGraphPin* StepsPin = CreatePin(EGPD_Input, K2Schema->PC_Int, TEXT(""), NULL, false, false, FEaseFunctionNodeHelper::GetStepsPinName());
+	UEdGraphPin* StepsPin = CreatePin(EGPD_Input, K2Schema->PC_Int, FString(), nullptr, FEaseFunctionNodeHelper::GetStepsPinName());
 	SetPinToolTip(*StepsPin, LOCTEXT("StepsPinDescription", "Number of steps required to go from A to B"));
-	StepsPin->DefaultValue = StepsPin->AutogeneratedDefaultValue = TEXT("2");
+	K2Schema->SetPinAutogeneratedDefaultValue(StepsPin, TEXT("2"));
 
 	RefreshPinVisibility();
 }
@@ -299,7 +299,7 @@ void UK2Node_EaseFunction::PinTypeChanged(UEdGraphPin* Pin)
 			if (bChanged)
 			{
 				// Just in case we switch to an invalid function clean it first
-				EaseFunctionName = TEXT("");
+				EaseFunctionName.Reset();
 
 				// Generate the right function name
 				if (InstigatorPin->PinType.PinCategory == Schema->PC_Float)
diff --git a/Engine/Source/Editor/BlueprintGraph/Private/K2Node_EditablePinBase.cpp b/Engine/Source/Editor/BlueprintGraph/Private/K2Node_EditablePinBase.cpp
--- a/Engine/Source/Editor/BlueprintGraph/Private/K2Node_EditablePinBase.cpp
+++ b/Engine/Source/Editor/BlueprintGraph/Private/K2Node_EditablePinBase.cpp
@@ -6,6 +6,41 @@
 #include "EdGraphSchema_K2.h"
 #include "Kismet2/KismetDebugUtilities.h"
 
+FArchive& operator<<(FArchive& Ar, FUserPinInfo& Info)
+{
+	Ar << Info.PinName;
+
+	if (Ar.UE4Ver() >= VER_UE4_SERIALIZE_PINTYPE_CONST)
+	{
+		Info.PinType.Serialize(Ar);
+		Ar << Info.DesiredPinDirection;
+	}
+	else
+	{
+		bool bIsArray = (Info.PinType.ContainerType == EPinContainerType::Array);
+		Ar << bIsArray;
+
+		bool bIsReference = Info.PinType.bIsReference;
+		Ar << bIsReference;
+
+		if (Ar.IsLoading())
+		{
+			Info.PinType.ContainerType = (bIsArray ? EPinContainerType::Array : EPinContainerType::None);
+			Info.PinType.bIsReference = bIsReference;
+		}
+
+		Ar << Info.PinType.PinCategory;
+		Ar << Info.PinType.PinSubCategory;
+
+		Ar << Info.PinType.PinSubCategoryObject;
+	}
+
+	Ar << Info.PinDefaultValue;
+
+	return Ar;
+}
+
+
 UK2Node_EditablePinBase::UK2Node_EditablePinBase(const FObjectInitializer& ObjectInitializer)
 	: Super(ObjectInitializer)
 {
@@ -105,12 +140,20 @@ void UK2Node_EditablePinBase::ExportCustomProperties(FOutputDevice& Out, uint32
 
 		Out.Logf( TEXT("%sCustomProperties UserDefinedPin "), FCString::Spc(Indent));
 		Out.Logf( TEXT("Name=\"%s\" "), *PinInfo.PinName);
-		Out.Logf( TEXT("IsArray=%s "), (PinInfo.PinType.bIsArray ? TEXT("1") : TEXT("0")));
 		Out.Logf( TEXT("IsReference=%s "), (PinInfo.PinType.bIsReference ? TEXT("1") : TEXT("0")));
 
+		if (UEnum* PinContainerType = FindObject<UEnum>(ANY_PACKAGE, TEXT("EPinContainerType")))
+		{ 
+			const FString ValueName = PinContainerType->GetNameStringByValue((uint8)PinInfo.PinType.ContainerType);
+			if (!ValueName.IsEmpty())
+			{
+				Out.Logf(TEXT("PinContainerType=\"%s\" "), *ValueName);
+			}
+		}
+
 		if (UEnum* PinDirEnum = FindObject<UEnum>(ANY_PACKAGE, TEXT("EEdGraphPinDirection")))
 		{ 
-			FString ValueName = PinDirEnum->GetNameStringByValue(PinInfo.DesiredPinDirection);
+			const FString ValueName = PinDirEnum->GetNameStringByValue(PinInfo.DesiredPinDirection);
 			if (!ValueName.IsEmpty())
 			{
 				Out.Logf(TEXT("PinDir=\"%s\" "), *ValueName);
@@ -156,7 +199,20 @@ void UK2Node_EditablePinBase::ImportCustomProperties(const TCHAR* SourceText, FF
 		int32 BoolAsInt = 0;
 		if (FParse::Value(SourceText, TEXT("IsArray="), BoolAsInt))
 		{
-			PinInfo->PinType.bIsArray = (BoolAsInt != 0);
+			PinInfo->PinType.ContainerType = (BoolAsInt != 0 ? EPinContainerType::Array : EPinContainerType::None);
+		}
+
+		if (UEnum* PinContainerType = FindObject<UEnum>(ANY_PACKAGE, TEXT("EPinContainerType")))
+		{
+			FString DesiredContainerType;
+			if (FParse::Value(SourceText, TEXT("PinContainerType="), DesiredContainerType))
+			{
+				const int32 DesiredContainerTypeVal = PinContainerType->GetValueByName(*DesiredContainerType);
+				if (DesiredContainerTypeVal != INDEX_NONE)
+				{
+					PinInfo->PinType.ContainerType = (EPinContainerType)DesiredContainerTypeVal;
+				}
+			}
 		}
 
 		if (FParse::Value(SourceText, TEXT("IsReference="), BoolAsInt))
@@ -169,7 +225,7 @@ void UK2Node_EditablePinBase::ImportCustomProperties(const TCHAR* SourceText, FF
 			FString DesiredDirection;
 			if (FParse::Value(SourceText, TEXT("PinDir="), DesiredDirection))
 			{
-				int32 DesiredDirectionVal = PinDirEnum->GetValueByName(*DesiredDirection);
+				const int32 DesiredDirectionVal = PinDirEnum->GetValueByName(*DesiredDirection);
 				if (DesiredDirectionVal != INDEX_NONE)
 				{
 					PinInfo->DesiredPinDirection = (EEdGraphPinDirection)DesiredDirectionVal;
@@ -268,24 +324,53 @@ void UK2Node_EditablePinBase::AddReferencedObjects(UObject* InThis, FReferenceCo
 	Super::AddReferencedObjects( This, Collector );
 }
 
+void UK2Node_EditablePinBase::PinDefaultValueChanged(UEdGraphPin* Pin)
+{
+	static bool bRecursivelyChangingDefaultValue = false;
+
+	// Only do this if we're editable and not already calling this code
+	if (!bIsEditable || bRecursivelyChangingDefaultValue)
+	{
+		return;
+	}
+
+	// See if this is a user defined pin
+	for (int32 Index = 0; Index < UserDefinedPins.Num(); ++Index)
+	{
+		TSharedPtr<FUserPinInfo> PinInfo = UserDefinedPins[Index];
+		if (Pin->PinName == PinInfo->PinName && Pin->Direction == PinInfo->DesiredPinDirection)
+		{
+			FString DefaultsString = Pin->GetDefaultAsString();
+
+			if (DefaultsString != PinInfo->PinDefaultValue)
+			{
+				// Make sure this doesn't get called recursively
+				TGuardValue<bool> CircularGuard(bRecursivelyChangingDefaultValue, true);
+				ModifyUserDefinedPinDefaultValue(PinInfo, Pin->GetDefaultAsString());
+			}
+		}
+	}
+}
+
 bool UK2Node_EditablePinBase::ModifyUserDefinedPinDefaultValue(TSharedPtr<FUserPinInfo> PinInfo, const FString& InDefaultValue)
 {
+	const UEdGraphSchema_K2* K2Schema = GetDefault<UEdGraphSchema_K2>();
 	FString NewDefaultValue = InDefaultValue;
 
 	// Find and modify the current pin
 	if (UEdGraphPin* OldPin = FindPin(PinInfo->PinName))
 	{
 		FString SavedDefaultValue = OldPin->DefaultValue;
-		OldPin->DefaultValue = OldPin->AutogeneratedDefaultValue = NewDefaultValue;
+		
+		K2Schema->SetPinAutogeneratedDefaultValue(OldPin, NewDefaultValue);
 
 		// Validate the new default value
-		const UEdGraphSchema* Schema = GetSchema();
-		FString ErrorString = Schema->IsCurrentPinDefaultValid(OldPin);
+		FString ErrorString = K2Schema->IsCurrentPinDefaultValid(OldPin);
 
 		if (!ErrorString.IsEmpty())
 		{
 			NewDefaultValue = SavedDefaultValue;
-			OldPin->DefaultValue = OldPin->AutogeneratedDefaultValue = SavedDefaultValue;
+			K2Schema->SetPinAutogeneratedDefaultValue(OldPin, SavedDefaultValue);
 
 			return false;
 		}
diff --git a/Engine/Source/Editor/BlueprintGraph/Private/K2Node_EnumEquality.cpp b/Engine/Source/Editor/BlueprintGraph/Private/K2Node_EnumEquality.cpp
--- a/Engine/Source/Editor/BlueprintGraph/Private/K2Node_EnumEquality.cpp
+++ b/Engine/Source/Editor/BlueprintGraph/Private/K2Node_EnumEquality.cpp
@@ -23,11 +23,11 @@ void UK2Node_EnumEquality::AllocateDefaultPins()
 	const UEdGraphSchema_K2* Schema = GetDefault<UEdGraphSchema_K2>();
 
 	// Create the return value pin
-	CreatePin(EGPD_Output, Schema->PC_Boolean, TEXT(""), NULL, false, false, Schema->PN_ReturnValue);
+	CreatePin(EGPD_Output, Schema->PC_Boolean, FString(), nullptr, Schema->PN_ReturnValue);
 
 	// Create the input pins
-	CreatePin(EGPD_Input, Schema->PC_Wildcard, TEXT(""), NULL, false, false, "A");
-	CreatePin(EGPD_Input, Schema->PC_Wildcard, TEXT(""), NULL, false, false, "B");
+	CreatePin(EGPD_Input, Schema->PC_Wildcard, FString(), nullptr, TEXT("A"));
+	CreatePin(EGPD_Input, Schema->PC_Wildcard, FString(), nullptr, TEXT("B"));
 
 	Super::AllocateDefaultPins();
 }
@@ -58,11 +58,11 @@ void UK2Node_EnumEquality::PostReconstructNode()
 // 	{
 // 		// Restore the wildcard status
 // 		Input1Pin->PinType.PinCategory = Schema->PC_Wildcard;
-// 		Input1Pin->PinType.PinSubCategory = TEXT("");
-// 		Input1Pin->PinType.PinSubCategoryObject = NULL;
+// 		Input1Pin->PinType.PinSubCategory.Reset();
+// 		Input1Pin->PinType.PinSubCategoryObject = nullptr;
 // 		Input2Pin->PinType.PinCategory = Schema->PC_Wildcard;
-// 		Input2Pin->PinType.PinSubCategory = TEXT("");
-// 		Input2Pin->PinType.PinSubCategoryObject = NULL;
+// 		Input2Pin->PinType.PinSubCategory.Reset();
+// 		Input2Pin->PinType.PinSubCategoryObject = nullptr;
 // 	}
 }
 
@@ -82,13 +82,13 @@ void UK2Node_EnumEquality::NotifyPinConnectionListChanged(UEdGraphPin* Pin)
 		{
 			// Restore the wildcard status
 			Input1Pin->PinType.PinCategory = Schema->PC_Wildcard;
-			Input1Pin->PinType.PinSubCategory = TEXT("");
-			Input1Pin->PinType.PinSubCategoryObject = NULL;
+			Input1Pin->PinType.PinSubCategory.Reset();
+			Input1Pin->PinType.PinSubCategoryObject = nullptr;
 			Input2Pin->PinType.PinCategory = Schema->PC_Wildcard;
-			Input2Pin->PinType.PinSubCategory = TEXT("");
-			Input2Pin->PinType.PinSubCategoryObject = NULL;
-			Schema->SetPinDefaultValueBasedOnType(Input1Pin);
-			Schema->SetPinDefaultValueBasedOnType(Input2Pin);
+			Input2Pin->PinType.PinSubCategory.Reset();
+			Input2Pin->PinType.PinSubCategoryObject = nullptr;
+			Schema->SetPinAutogeneratedDefaultValueBasedOnType(Input1Pin);
+			Schema->SetPinAutogeneratedDefaultValueBasedOnType(Input2Pin);
 			// We have to refresh the graph to get the enum dropdowns to go away
 			GetGraph()->NotifyGraphChanged();
 		}
@@ -109,7 +109,7 @@ void UK2Node_EnumEquality::NotifyPinConnectionListChanged(UEdGraphPin* Pin)
 				// If we copied the pin type to a wildcard we have to refresh the graph to get the enum dropdown
 				if (OtherPin->LinkedTo.Num() <= 0 && OtherPin->DefaultValue.IsEmpty())
 				{
-					Schema->SetPinDefaultValueBasedOnType(OtherPin);
+					Schema->SetPinAutogeneratedDefaultValueBasedOnType(OtherPin);
 					GetGraph()->NotifyGraphChanged();
 				}
 			}
@@ -121,7 +121,7 @@ void UK2Node_EnumEquality::NotifyPinConnectionListChanged(UEdGraphPin* Pin)
 		}
 		else if(Pin->DefaultValue.IsEmpty())
 		{
-			Schema->SetPinDefaultValueBasedOnType(Pin);
+			Schema->SetPinAutogeneratedDefaultValueBasedOnType(Pin);
 		}
 	}
 }
diff --git a/Engine/Source/Editor/BlueprintGraph/Private/K2Node_EnumLiteral.cpp b/Engine/Source/Editor/BlueprintGraph/Private/K2Node_EnumLiteral.cpp
--- a/Engine/Source/Editor/BlueprintGraph/Private/K2Node_EnumLiteral.cpp
+++ b/Engine/Source/Editor/BlueprintGraph/Private/K2Node_EnumLiteral.cpp
@@ -34,10 +34,10 @@ void UK2Node_EnumLiteral::AllocateDefaultPins()
 {
 	const UEdGraphSchema_K2* Schema = GetDefault<UEdGraphSchema_K2>();
 
-	auto InputPin = CreatePin(EGPD_Input, Schema->PC_Byte, TEXT(""), Enum, false, false, GetEnumInputPinName());
-	Schema->SetPinDefaultValueBasedOnType(InputPin);
+	auto InputPin = CreatePin(EGPD_Input, Schema->PC_Byte, FString(), Enum, GetEnumInputPinName());
+	Schema->SetPinAutogeneratedDefaultValueBasedOnType(InputPin);
 
-	CreatePin(EGPD_Output, Schema->PC_Byte, TEXT(""), Enum, false, false, Schema->PN_ReturnValue);
+	CreatePin(EGPD_Output, Schema->PC_Byte, FString(), Enum, Schema->PN_ReturnValue);
 
 	Super::AllocateDefaultPins();
 }
diff --git a/Engine/Source/Editor/BlueprintGraph/Private/K2Node_Event.cpp b/Engine/Source/Editor/BlueprintGraph/Private/K2Node_Event.cpp
--- a/Engine/Source/Editor/BlueprintGraph/Private/K2Node_Event.cpp
+++ b/Engine/Source/Editor/BlueprintGraph/Private/K2Node_Event.cpp
@@ -76,7 +76,7 @@ void UK2Node_Event::PostLoad()
 	if (!FindPin(DelegateOutputName))
 	{
 		const UEdGraphSchema_K2* K2Schema = GetDefault<UEdGraphSchema_K2>();
-		CreatePin(EGPD_Output, K2Schema->PC_Delegate, TEXT(""), NULL, false, false, DelegateOutputName);
+		CreatePin(EGPD_Output, K2Schema->PC_Delegate, FString(), nullptr, DelegateOutputName);
 	}
 }
 
@@ -191,7 +191,7 @@ FString UK2Node_Event::GetDocumentationLink() const
 		return FString::Printf(TEXT("Shared/Types/%s%s"), EventSignatureClass->GetPrefixCPP(), *EventSignatureClass->GetName());
 	}
 
-	return TEXT("");
+	return FString();
 }
 
 FString UK2Node_Event::GetDocumentationExcerptName() const
@@ -308,11 +308,11 @@ void UK2Node_Event::AllocateDefaultPins()
 {
 	const UEdGraphSchema_K2* K2Schema = GetDefault<UEdGraphSchema_K2>();
 
-	CreatePin(EGPD_Output, K2Schema->PC_Delegate, TEXT(""), NULL, false, false, DelegateOutputName);
-	CreatePin(EGPD_Output, K2Schema->PC_Exec, TEXT(""), NULL, false, false, K2Schema->PN_Then);
+	CreatePin(EGPD_Output, K2Schema->PC_Delegate, FString(), nullptr, DelegateOutputName);
+	CreatePin(EGPD_Output, K2Schema->PC_Exec, FString(), nullptr, K2Schema->PN_Then);
 
 	const UFunction* Function = FindEventSignatureFunction();
-	if (Function != NULL)
+	if (Function)
 	{
 		CreatePinsForFunctionEntryExit(Function, /*bIsFunctionEntry=*/ true);
 	}
diff --git a/Engine/Source/Editor/BlueprintGraph/Private/K2Node_ExecutionSequence.cpp b/Engine/Source/Editor/BlueprintGraph/Private/K2Node_ExecutionSequence.cpp
--- a/Engine/Source/Editor/BlueprintGraph/Private/K2Node_ExecutionSequence.cpp
+++ b/Engine/Source/Editor/BlueprintGraph/Private/K2Node_ExecutionSequence.cpp
@@ -163,13 +163,11 @@ void UK2Node_ExecutionSequence::AllocateDefaultPins()
 {
 	const UEdGraphSchema_K2* K2Schema = GetDefault<UEdGraphSchema_K2>();
 
-	CreatePin(EGPD_Input, K2Schema->PC_Exec, TEXT(""), NULL, false, false, K2Schema->PN_Execute);
+	CreatePin(EGPD_Input, K2Schema->PC_Exec, FString(), nullptr, K2Schema->PN_Execute);
 
 	// Add two default pins
-	for (int32 i = 0; i < 2; ++i)
-	{
-		CreatePin(EGPD_Output, K2Schema->PC_Exec, TEXT(""), NULL, false, false, GetPinNameGivenIndex(i));
-	}
+	CreatePin(EGPD_Output, K2Schema->PC_Exec, FString(), nullptr, GetPinNameGivenIndex(0));
+	CreatePin(EGPD_Output, K2Schema->PC_Exec, FString(), nullptr, GetPinNameGivenIndex(1));
 
 	Super::AllocateDefaultPins();
 }
@@ -216,7 +214,7 @@ FString UK2Node_ExecutionSequence::GetUniquePinName()
 void UK2Node_ExecutionSequence::AddPinToExecutionNode()
 {
 	const UEdGraphSchema_K2* K2Schema = GetDefault<UEdGraphSchema_K2>();
-	CreatePin(EGPD_Output, K2Schema->PC_Exec, TEXT(""), NULL, false, false, GetUniquePinName());
+	CreatePin(EGPD_Output, K2Schema->PC_Exec, FString(), nullptr, GetUniquePinName());
 }
 
 void UK2Node_ExecutionSequence::RemovePinFromExecutionNode(UEdGraphPin* TargetPin) 
@@ -272,7 +270,7 @@ void UK2Node_ExecutionSequence::ReallocatePinsDuringReconstruction(TArray<UEdGra
 
 	// Create the execution input pin
 	const UEdGraphSchema_K2* K2Schema = GetDefault<UEdGraphSchema_K2>();
-	CreatePin(EGPD_Input, K2Schema->PC_Exec, TEXT(""), NULL, false, false, K2Schema->PN_Execute);
+	CreatePin(EGPD_Input, K2Schema->PC_Exec, FString(), nullptr, K2Schema->PN_Execute);
 
 	// Create a new pin for each old execution output pin, and coerce the names to match on both sides
 	int32 ExecOutPinCount = 0;
@@ -288,7 +286,7 @@ void UK2Node_ExecutionSequence::ReallocatePinsDuringReconstruction(TArray<UEdGra
 			TestPin->PinName = NewPinName;
 
 			// Create the new output pin to match
-			CreatePin(EGPD_Output, K2Schema->PC_Exec, TEXT(""), NULL, false, false, NewPinName);
+			CreatePin(EGPD_Output, K2Schema->PC_Exec, FString(), nullptr, NewPinName);
 		}
 	}
 }
diff --git a/Engine/Source/Editor/BlueprintGraph/Private/K2Node_ForEachElementInEnum.cpp b/Engine/Source/Editor/BlueprintGraph/Private/K2Node_ForEachElementInEnum.cpp
--- a/Engine/Source/Editor/BlueprintGraph/Private/K2Node_ForEachElementInEnum.cpp
+++ b/Engine/Source/Editor/BlueprintGraph/Private/K2Node_ForEachElementInEnum.cpp
@@ -144,15 +144,15 @@ void UK2Node_ForEachElementInEnum::AllocateDefaultPins()
 	const UEdGraphSchema_K2* K2Schema = GetDefault<UEdGraphSchema_K2>();
 	check(K2Schema);
 
-	CreatePin(EGPD_Input, K2Schema->PC_Exec, TEXT(""), NULL, false, false, K2Schema->PN_Execute);
+	CreatePin(EGPD_Input, K2Schema->PC_Exec, FString(), nullptr, K2Schema->PN_Execute);
 
 	if (Enum)
 	{
-		CreatePin(EGPD_Output, K2Schema->PC_Exec, TEXT(""), NULL, false, false, InsideLoopPinName);
-		CreatePin(EGPD_Output, K2Schema->PC_Byte, TEXT(""), Enum, false, false, EnumOuputPinName);
+		CreatePin(EGPD_Output, K2Schema->PC_Exec, FString(), nullptr, InsideLoopPinName);
+		CreatePin(EGPD_Output, K2Schema->PC_Byte, FString(), Enum, EnumOuputPinName);
 	}
 
-	if (auto CompletedPin = CreatePin(EGPD_Output, K2Schema->PC_Exec, TEXT(""), NULL, false, false, K2Schema->PN_Then))
+	if (UEdGraphPin* CompletedPin = CreatePin(EGPD_Output, K2Schema->PC_Exec, FString(), nullptr, K2Schema->PN_Then))
 	{
 		CompletedPin->PinFriendlyName = LOCTEXT("Completed", "Completed");
 	}
@@ -160,6 +160,8 @@ void UK2Node_ForEachElementInEnum::AllocateDefaultPins()
 
 void UK2Node_ForEachElementInEnum::ValidateNodeDuringCompilation(FCompilerResultsLog& MessageLog) const
 {
+	Super::ValidateNodeDuringCompilation(MessageLog);
+
 	if (!Enum)
 	{
 		MessageLog.Error(*NSLOCTEXT("K2Node", "ForEachElementInEnum_NoEnumError", "No Enum in @@").ToString(), this);
diff --git a/Engine/Source/Editor/BlueprintGraph/Private/K2Node_FormatText.cpp b/Engine/Source/Editor/BlueprintGraph/Private/K2Node_FormatText.cpp
--- a/Engine/Source/Editor/BlueprintGraph/Private/K2Node_FormatText.cpp
+++ b/Engine/Source/Editor/BlueprintGraph/Private/K2Node_FormatText.cpp
@@ -44,12 +44,12 @@ void UK2Node_FormatText::AllocateDefaultPins()
 	Super::AllocateDefaultPins();
 
 	const UEdGraphSchema_K2* K2Schema = GetDefault<UEdGraphSchema_K2>();
-	CachedFormatPin = CreatePin(EGPD_Input, K2Schema->PC_Text, TEXT(""), NULL, false, false, FFormatTextNodeHelper::GetFormatPinName());
-	CreatePin(EGPD_Output, K2Schema->PC_Text, TEXT(""), NULL, false, false, TEXT("Result"));
+	CachedFormatPin = CreatePin(EGPD_Input, K2Schema->PC_Text, FString(), nullptr, FFormatTextNodeHelper::GetFormatPinName());
+	CreatePin(EGPD_Output, K2Schema->PC_Text, FString(), nullptr, TEXT("Result"));
 
-	for(auto It = PinNames.CreateConstIterator(); It; ++It)
+	for (const FString& PinName : PinNames)
 	{
-		CreatePin(EGPD_Input, K2Schema->PC_Wildcard, TEXT(""), NULL, false, false, *It);
+		CreatePin(EGPD_Input, K2Schema->PC_Wildcard, FString(), nullptr, PinName);
 	}
 }
 
@@ -63,7 +63,7 @@ void UK2Node_FormatText::SynchronizeArgumentPinType(UEdGraphPin* Pin)
 		bool bPinTypeChanged = false;
 		if (Pin->LinkedTo.Num() == 0)
 		{
-			static const FEdGraphPinType WildcardPinType = FEdGraphPinType(K2Schema->PC_Wildcard, TEXT(""), nullptr, false, false, false, false, FEdGraphTerminalType());
+			static const FEdGraphPinType WildcardPinType = FEdGraphPinType(K2Schema->PC_Wildcard, FString(), nullptr, EPinContainerType::None, false, FEdGraphTerminalType());
 
 			// Ensure wildcard
 			if (Pin->PinType != WildcardPinType)
@@ -137,7 +137,7 @@ void UK2Node_FormatText::PostEditChangeProperty(struct FPropertyChangedEvent& Pr
 
 void UK2Node_FormatText::PinConnectionListChanged(UEdGraphPin* Pin)
 {
-	const auto FormatPin = GetFormatPin();
+	UEdGraphPin* FormatPin = GetFormatPin();
 
 	Modify();
 
@@ -168,7 +168,7 @@ void UK2Node_FormatText::PinConnectionListChanged(UEdGraphPin* Pin)
 
 void UK2Node_FormatText::PinDefaultValueChanged(UEdGraphPin* Pin)
 {
-	const auto FormatPin = GetFormatPin();
+	const UEdGraphPin* FormatPin = GetFormatPin();
 	if(Pin == FormatPin && FormatPin->LinkedTo.Num() == 0)
 	{
 		const UEdGraphSchema_K2* K2Schema = GetDefault<UEdGraphSchema_K2>();
@@ -178,13 +178,13 @@ void UK2Node_FormatText::PinDefaultValueChanged(UEdGraphPin* Pin)
 
 		PinNames.Empty();
 
-		for(auto It = ArgumentParams.CreateConstIterator(); It; ++It)
+		for (const FString& Param : ArgumentParams)
 		{
-			if(!FindArgumentPin(*It))
+			if(!FindArgumentPin(Param))
 			{
-				CreatePin(EGPD_Input, K2Schema->PC_Wildcard, TEXT(""), NULL, false, false, *It);
+				CreatePin(EGPD_Input, K2Schema->PC_Wildcard, FString(), nullptr, Param);
 			}
-			PinNames.Add(*It);
+			PinNames.Add(Param);
 		}
 
 		for(auto It = Pins.CreateConstIterator(); It; ++It)
@@ -421,7 +421,7 @@ void UK2Node_FormatText::ExpandNode(class FKismetCompilerContext& CompilerContex
 
 UEdGraphPin* UK2Node_FormatText::FindArgumentPin(const FString& InPinName) const
 {
-	const auto FormatPin = GetFormatPin();
+	const UEdGraphPin* FormatPin = GetFormatPin();
 	for(int32 PinIdx=0; PinIdx<Pins.Num(); PinIdx++)
 	{
 		if( Pins[PinIdx] != FormatPin && Pins[PinIdx]->Direction != EGPD_Output && Pins[PinIdx]->PinName.Equals(InPinName) )
@@ -527,7 +527,7 @@ void UK2Node_FormatText::AddArgumentPin()
 
 	const UEdGraphSchema_K2* K2Schema = Cast<const UEdGraphSchema_K2>(GetSchema());
 	FString PinName = GetUniquePinName();
-	CreatePin(EGPD_Input, K2Schema->PC_Wildcard, TEXT(""), NULL, false, false, PinName);
+	CreatePin(EGPD_Input, K2Schema->PC_Wildcard, FString(), nullptr, PinName);
 	PinNames.Add(PinName);
 
 	FBlueprintEditorUtils::MarkBlueprintAsStructurallyModified(GetBlueprint());
diff --git a/Engine/Source/Editor/BlueprintGraph/Private/K2Node_FunctionEntry.cpp b/Engine/Source/Editor/BlueprintGraph/Private/K2Node_FunctionEntry.cpp
--- a/Engine/Source/Editor/BlueprintGraph/Private/K2Node_FunctionEntry.cpp
+++ b/Engine/Source/Editor/BlueprintGraph/Private/K2Node_FunctionEntry.cpp
@@ -3,8 +3,10 @@
 
 #include "K2Node_FunctionEntry.h"
 #include "Engine/Blueprint.h"
+#include "Animation/AnimBlueprint.h"
 #include "UObject/UnrealType.h"
 #include "UObject/BlueprintsObjectVersion.h"
+#include "UObject/FrameworkObjectVersion.h"
 #include "UObject/StructOnScope.h"
 #include "Engine/UserDefinedStruct.h"
 #include "EdGraph/EdGraphSchema.h"
@@ -195,10 +197,12 @@ void UK2Node_FunctionEntry::Serialize(FArchive& Ar)
 	Super::Serialize(Ar);
 
 	Ar.UsingCustomVersion(FBlueprintsObjectVersion::GUID);
+	Ar.UsingCustomVersion(FFrameworkObjectVersion::GUID);
 
 	if (Ar.IsLoading())
 	{
-		if (Ar.UE4Ver() < VER_UE4_BLUEPRINT_ENFORCE_CONST_IN_FUNCTION_OVERRIDES)
+		if (Ar.UE4Ver() < VER_UE4_BLUEPRINT_ENFORCE_CONST_IN_FUNCTION_OVERRIDES
+			|| ((Ar.CustomVer(FFrameworkObjectVersion::GUID) < FFrameworkObjectVersion::EnforceConstInAnimBlueprintFunctionGraphs) && GetBlueprint()->IsA<UAnimBlueprint>()))
 		{
 			// Allow legacy implementations to violate const-correctness
 			bEnforceConstCorrectness = false;
@@ -219,6 +223,33 @@ void UK2Node_FunctionEntry::Serialize(FArchive& Ar)
 				FUNC_Delegate | FUNC_HasOutParms | FUNC_HasDefaults | FUNC_DLLImport | FUNC_NetValidate;
 			ExtraFlags &= ~InvalidExtraFlagsMask;
 		}
+
+		if (Ar.CustomVer(FFrameworkObjectVersion::GUID) < FFrameworkObjectVersion::ChangeAssetPinsToString)
+		{
+			const UEdGraphSchema_K2* K2Schema = GetDefault<UEdGraphSchema_K2>();
+
+			// Prior to this version, changing the type of a local variable would lead to corrupt default value strings
+			for (FBPVariableDescription& LocalVar : LocalVariables)
+			{
+				FString UseDefaultValue;
+				UObject* UseDefaultObject = nullptr;
+				FText UseDefaultText;
+
+				if (!LocalVar.DefaultValue.IsEmpty())
+				{
+					K2Schema->GetPinDefaultValuesFromString(LocalVar.VarType, this, LocalVar.DefaultValue, UseDefaultValue, UseDefaultObject, UseDefaultText);
+					FString ErrorMessage;
+
+					if (!K2Schema->DefaultValueSimpleValidation(LocalVar.VarType, LocalVar.VarName.ToString(), UseDefaultValue, UseDefaultObject, UseDefaultText, &ErrorMessage))
+					{
+						const UBlueprint* Blueprint = GetBlueprint();
+						UE_LOG(LogBlueprint, Log, TEXT("Clearing invalid default value for local variable %s on blueprint %s: %s"), *LocalVar.VarName.ToString(), Blueprint ? *Blueprint->GetName() : TEXT("Unknown"), *ErrorMessage);
+
+						LocalVar.DefaultValue.Reset();
+					}
+				}
+			}
+		}
 	}
 }
 
@@ -234,7 +265,7 @@ FText UK2Node_FunctionEntry::GetNodeTitle(ENodeTitleType::Type TitleType) const
 void UK2Node_FunctionEntry::AllocateDefaultPins()
 {
 	const UEdGraphSchema_K2* K2Schema = GetDefault<UEdGraphSchema_K2>();
-	CreatePin(EGPD_Output, K2Schema->PC_Exec, TEXT(""), NULL, false, false, K2Schema->PN_Then);
+	CreatePin(EGPD_Output, K2Schema->PC_Exec, FString(), nullptr, K2Schema->PN_Then);
 
 	UFunction* Function = FindField<UFunction>(SignatureClass, SignatureName);
 
@@ -258,8 +289,6 @@ void UK2Node_FunctionEntry::AllocateDefaultPins()
 			K2Schema->PC_Object,
 			FString(),
 			UObject::StaticClass(),
-			false,
-			false,
 			FFunctionEntryHelper::GetWorldContextPinName());
 		WorldContextPin->bHidden = true;
 	}
@@ -304,7 +333,7 @@ UEdGraphPin* UK2Node_FunctionEntry::CreatePinFromUserDefinition(const TSharedPtr
 	}
 
 	UEdGraphPin* NewPin = CreatePin(EGPD_Output, NewPinInfo->PinType, NewPinInfo->PinName);
-	NewPin->DefaultValue = NewPin->AutogeneratedDefaultValue = NewPinInfo->PinDefaultValue;
+	Schema->SetPinAutogeneratedDefaultValue(NewPin, NewPinInfo->PinDefaultValue);
 	return NewPin;
 }
 
@@ -438,7 +467,7 @@ void UK2Node_FunctionEntry::ExpandNode(class FKismetCompilerContext& CompilerCon
 
 						if(UEdGraphPin* SetPin = VariableSetNode->FindPin(Property->GetName()))
 						{
-							if(LocalVar.VarType.bIsArray)
+							if(LocalVar.VarType.IsArray())
 							{
 								TSharedPtr<FStructOnScope> StructData = MakeShareable(new FStructOnScope(Function));
 								FBlueprintEditorUtils::PropertyValueFromString(Property, LocalVar.DefaultValue, StructData->GetStructMemory());
@@ -473,7 +502,7 @@ void UK2Node_FunctionEntry::ExpandNode(class FKismetCompilerContext& CompilerCon
 									Schema->TrySetDefaultValue(*MakeArray->Pins[ArrayIndex + 1], DefaultValue);
 								}
 							}
-							else if(LocalVar.VarType.bIsSet || LocalVar.VarType.bIsMap)
+							else if(LocalVar.VarType.IsSet() || LocalVar.VarType.IsMap())
 							{
 								UK2Node_MakeVariable* MakeVariableNode = CompilerContext.SpawnIntermediateNode<UK2Node_MakeVariable>(this, SourceGraph);
 								MakeVariableNode->SetupVariable(LocalVar, SetPin, CompilerContext, Function, Property);
@@ -546,5 +575,17 @@ void UK2Node_FunctionEntry::PostReconstructNode()
 	}
 }
 
+bool UK2Node_FunctionEntry::ModifyUserDefinedPinDefaultValue(TSharedPtr<FUserPinInfo> PinInfo, const FString& NewDefaultValue)
+{
+	if (Super::ModifyUserDefinedPinDefaultValue(PinInfo, NewDefaultValue))
+	{
+		const UEdGraphSchema_K2* K2Schema = GetDefault<UEdGraphSchema_K2>();
+		K2Schema->HandleParameterDefaultValueChanged(this);
+
+		return true;
+	}
+	return false;
+}
+
 
 #undef LOCTEXT_NAMESPACE
diff --git a/Engine/Source/Editor/BlueprintGraph/Private/K2Node_FunctionResult.cpp b/Engine/Source/Editor/BlueprintGraph/Private/K2Node_FunctionResult.cpp
--- a/Engine/Source/Editor/BlueprintGraph/Private/K2Node_FunctionResult.cpp
+++ b/Engine/Source/Editor/BlueprintGraph/Private/K2Node_FunctionResult.cpp
@@ -23,7 +23,7 @@ private:
 			const bool bNeedToResetDefaultValue = (Pin->DefaultValue.IsEmpty() && Pin->DefaultObject == nullptr && Pin->DefaultTextValue.IsEmpty()) || !(K2Schema->IsPinDefaultValid(Pin, Pin->DefaultValue, Pin->DefaultObject, Pin->DefaultTextValue).IsEmpty());
 			if (bValuePin && bNotConnected && bNeedToResetDefaultValue)
 			{
-				K2Schema->SetPinDefaultValueBasedOnType(Pin);
+				K2Schema->SetPinAutogeneratedDefaultValueBasedOnType(Pin);
 			}
 		}
 	}
@@ -133,7 +133,7 @@ FText UK2Node_FunctionResult::GetNodeTitle(ENodeTitleType::Type TitleType) const
 void UK2Node_FunctionResult::AllocateDefaultPins()
 {
 	const UEdGraphSchema_K2* K2Schema = GetDefault<UEdGraphSchema_K2>();
-	CreatePin(EGPD_Input, K2Schema->PC_Exec, TEXT(""), NULL, false, false, K2Schema->PN_Execute);
+	CreatePin(EGPD_Input, K2Schema->PC_Exec, FString(), nullptr, K2Schema->PN_Execute);
 
 	UFunction* Function = FindField<UFunction>(SignatureClass, SignatureName);
 	if (Function != NULL)
diff --git a/Engine/Source/Editor/BlueprintGraph/Private/K2Node_GenericCreateObject.cpp b/Engine/Source/Editor/BlueprintGraph/Private/K2Node_GenericCreateObject.cpp
--- a/Engine/Source/Editor/BlueprintGraph/Private/K2Node_GenericCreateObject.cpp
+++ b/Engine/Source/Editor/BlueprintGraph/Private/K2Node_GenericCreateObject.cpp
@@ -17,29 +17,29 @@ void UK2Node_GenericCreateObject::ExpandNode(class FKismetCompilerContext& Compi
 	bool bSucceeded = true;
 	//connect exe
 	{
-		auto SpawnExecPin = GetExecPin();
-		auto CallExecPin = CallCreateNode->GetExecPin();
+		UEdGraphPin* SpawnExecPin = GetExecPin();
+		UEdGraphPin* CallExecPin = CallCreateNode->GetExecPin();
 		bSucceeded &= SpawnExecPin && CallExecPin && CompilerContext.MovePinLinksToIntermediate(*SpawnExecPin, *CallExecPin).CanSafeConnect();
 	}
 
 	//connect class
 	{
-		auto SpawnClassPin = GetClassPin();
-		auto CallClassPin = CallCreateNode->FindPin(TEXT("ObjectClass"));
+		UEdGraphPin* SpawnClassPin = GetClassPin();
+		UEdGraphPin* CallClassPin = CallCreateNode->FindPin(TEXT("ObjectClass"));
 		bSucceeded &= SpawnClassPin && CallClassPin && CompilerContext.MovePinLinksToIntermediate(*SpawnClassPin, *CallClassPin).CanSafeConnect();
 	}
 		
 	//connect outer
 	{
-		auto SpawnOuterPin = GetOuterPin();
-		auto CallOuterPin = CallCreateNode->FindPin(TEXT("Outer"));
+		UEdGraphPin* SpawnOuterPin = GetOuterPin();
+		UEdGraphPin* CallOuterPin = CallCreateNode->FindPin(TEXT("Outer"));
 		bSucceeded &= SpawnOuterPin && CallOuterPin && CompilerContext.MovePinLinksToIntermediate(*SpawnOuterPin, *CallOuterPin).CanSafeConnect();
 	}
 
 	UEdGraphPin* CallResultPin = nullptr;
 	//connect result
 	{
-		auto SpawnResultPin = GetResultPin();
+		UEdGraphPin* SpawnResultPin = GetResultPin();
 		CallResultPin = CallCreateNode->GetReturnValuePin();
 
 		// cast HACK. It should be safe. The only problem is native code generation.
@@ -52,8 +52,8 @@ void UK2Node_GenericCreateObject::ExpandNode(class FKismetCompilerContext& Compi
 
 	//assign exposed values and connect then
 	{
-		auto LastThen = FKismetCompilerUtilities::GenerateAssignmentNodes(CompilerContext, SourceGraph, CallCreateNode, this, CallResultPin, GetClassToSpawn());
-		auto SpawnNodeThen = GetThenPin();
+		UEdGraphPin* LastThen = FKismetCompilerUtilities::GenerateAssignmentNodes(CompilerContext, SourceGraph, CallCreateNode, this, CallResultPin, GetClassToSpawn());
+		UEdGraphPin* SpawnNodeThen = GetThenPin();
 		bSucceeded &= SpawnNodeThen && LastThen && CompilerContext.MovePinLinksToIntermediate(*SpawnNodeThen, *LastThen).CanSafeConnect();
 	}
 
@@ -65,23 +65,18 @@ void UK2Node_GenericCreateObject::ExpandNode(class FKismetCompilerContext& Compi
 	}
 }
 
-void UK2Node_GenericCreateObject::ValidateNodeDuringCompilation(class FCompilerResultsLog& MessageLog) const
-{
-	Super::ValidateNodeDuringCompilation(MessageLog);
-}
-
 void UK2Node_GenericCreateObject::EarlyValidation(class FCompilerResultsLog& MessageLog) const
 {
 	Super::EarlyValidation(MessageLog);
-	auto ClassPin = GetClassPin(&Pins);
+	UEdGraphPin* ClassPin = GetClassPin(&Pins);
 	const bool bAllowAbstract = ClassPin && ClassPin->LinkedTo.Num();
-	auto ClassToSpawn = GetClassToSpawn();
+	UClass* ClassToSpawn = GetClassToSpawn();
 	if (!UGameplayStatics::CanSpawnObjectOfClass(ClassToSpawn, bAllowAbstract))
 	{
 		MessageLog.Error(*FString::Printf(*LOCTEXT("GenericCreateObject_WrongClass", "Wrong class to spawn '%s' in @@").ToString(), *GetPathNameSafe(ClassToSpawn)), this);
 	}
 
-	auto OuterPin = GetOuterPin();
+	UEdGraphPin* OuterPin = GetOuterPin();
 	if (!OuterPin || (!OuterPin->DefaultObject && !OuterPin->LinkedTo.Num()))
 	{
 		MessageLog.Error(*LOCTEXT("GenericCreateObject_NoOuter", "Outer object is required in @@").ToString(), this);
