Copying //UE4/Dev-Framework to //UE4/Dev-Main (Source: //UE4/Dev-Framework @ 3431384)

diff --git a/Engine/Source/Editor/BlueprintGraph/Private/K2Node_GetArrayItem.cpp b/Engine/Source/Editor/BlueprintGraph/Private/K2Node_GetArrayItem.cpp
--- a/Engine/Source/Editor/BlueprintGraph/Private/K2Node_GetArrayItem.cpp
+++ b/Engine/Source/Editor/BlueprintGraph/Private/K2Node_GetArrayItem.cpp
@@ -119,12 +119,12 @@ UK2Node_GetArrayItem::UK2Node_GetArrayItem(const FObjectInitializer& ObjectIniti
 void UK2Node_GetArrayItem::AllocateDefaultPins()
 {
 	// Create the output pin
-	CreatePin(EGPD_Input, UEdGraphSchema_K2::PC_Wildcard, TEXT(""), NULL, true, false, TEXT("Array"));
-	UEdGraphPin* IndexPin = CreatePin(EGPD_Input, UEdGraphSchema_K2::PC_Int, TEXT(""), NULL, false, false, TEXT("Dimension 1"));
-	GetDefault<UEdGraphSchema_K2>()->SetPinDefaultValueBasedOnType(IndexPin);
+	CreatePin(EGPD_Input, UEdGraphSchema_K2::PC_Wildcard, FString(), nullptr, TEXT("Array"), EPinContainerType::Array);
+	UEdGraphPin* IndexPin = CreatePin(EGPD_Input, UEdGraphSchema_K2::PC_Int, FString(), nullptr, TEXT("Dimension 1"));
+	GetDefault<UEdGraphSchema_K2>()->SetPinAutogeneratedDefaultValueBasedOnType(IndexPin);
 
 	// Create the input pins to create the arrays from
-	CreatePin(EGPD_Output, UEdGraphSchema_K2::PC_Wildcard, TEXT(""), NULL, false, bReturnByRefDesired, TEXT("Output"));
+	CreatePin(EGPD_Output, UEdGraphSchema_K2::PC_Wildcard, FString(), nullptr, TEXT("Output"), EPinContainerType::None, bReturnByRefDesired);
 }
 
 void UK2Node_GetArrayItem::PostReconstructNode()
@@ -267,6 +267,7 @@ void UK2Node_GetArrayItem::NotifyPinConnectionListChanged(UEdGraphPin* Pin)
 				if (NewType.PinCategory != UEdGraphSchema_K2::PC_Wildcard)
 				{
 					PropagatePinType(NewType);
+					GetGraph()->NotifyGraphChanged();
 				}
 			}
 		}
@@ -410,7 +411,7 @@ void UK2Node_GetArrayItem::PropagatePinType(FEdGraphPinType& InType)
 	UEdGraphPin* ResultPin = Pins[2];
 
 	ArrayPin->PinType = InType;
-	ArrayPin->PinType.bIsArray = true;
+	ArrayPin->PinType.ContainerType = EPinContainerType::Array;
 	ArrayPin->PinType.bIsReference = false;
 
 	// IsSetToReturnRef() has to be called after the ArrayPin's type is set, since it uses that to determine the result
@@ -426,7 +427,7 @@ void UK2Node_GetArrayItem::PropagatePinType(FEdGraphPinType& InType)
 	}
 
 	ResultPin->PinType = InType;
-	ResultPin->PinType.bIsArray = false;
+	ResultPin->PinType.ContainerType = EPinContainerType::None;
 	ResultPin->PinType.bIsReference = bMakeOutputRef;
 	ResultPin->PinType.bIsConst = false;
 
diff --git a/Engine/Source/Editor/BlueprintGraph/Private/K2Node_GetClassDefaults.cpp b/Engine/Source/Editor/BlueprintGraph/Private/K2Node_GetClassDefaults.cpp
--- a/Engine/Source/Editor/BlueprintGraph/Private/K2Node_GetClassDefaults.cpp
+++ b/Engine/Source/Editor/BlueprintGraph/Private/K2Node_GetClassDefaults.cpp
@@ -212,7 +212,7 @@ void UK2Node_GetClassDefaults::AllocateDefaultPins()
 	const UEdGraphSchema_K2* K2Schema = GetDefault<UEdGraphSchema_K2>();
 
 	// Create the class input type selector pin
-	UEdGraphPin* ClassPin = CreatePin(EGPD_Input, K2Schema->PC_Class, TEXT(""), UObject::StaticClass(), false, false, ClassPinName);
+	UEdGraphPin* ClassPin = CreatePin(EGPD_Input, K2Schema->PC_Class, FString(), UObject::StaticClass(), ClassPinName);
 	K2Schema->ConstructBasicPinTooltip(*ClassPin, LOCTEXT("ClassPinDescription", "The class from which to access one or more default values."), ClassPin->PinToolTip);
 }
 
diff --git a/Engine/Source/Editor/BlueprintGraph/Private/K2Node_GetDataTableRow.cpp b/Engine/Source/Editor/BlueprintGraph/Private/K2Node_GetDataTableRow.cpp
--- a/Engine/Source/Editor/BlueprintGraph/Private/K2Node_GetDataTableRow.cpp
+++ b/Engine/Source/Editor/BlueprintGraph/Private/K2Node_GetDataTableRow.cpp
@@ -37,21 +37,21 @@ void UK2Node_GetDataTableRow::AllocateDefaultPins()
 	const UEdGraphSchema_K2* K2Schema = GetDefault<UEdGraphSchema_K2>();
 
 	// Add execution pins
-	CreatePin(EGPD_Input, K2Schema->PC_Exec, TEXT(""), NULL, false, false, K2Schema->PN_Execute);
-	UEdGraphPin* RowFoundPin = CreatePin(EGPD_Output, K2Schema->PC_Exec, TEXT(""), NULL, false, false, K2Schema->PN_Then);
+	CreatePin(EGPD_Input, K2Schema->PC_Exec, FString(), nullptr, K2Schema->PN_Execute);
+	UEdGraphPin* RowFoundPin = CreatePin(EGPD_Output, K2Schema->PC_Exec, FString(), nullptr, K2Schema->PN_Then);
 	RowFoundPin->PinFriendlyName = LOCTEXT("GetDataTableRow Row Found Exec pin", "Row Found");
-	CreatePin(EGPD_Output, K2Schema->PC_Exec, TEXT(""), NULL, false, false, UK2Node_GetDataTableRowHelper::RowNotFoundPinName);
+	CreatePin(EGPD_Output, K2Schema->PC_Exec, FString(), nullptr, UK2Node_GetDataTableRowHelper::RowNotFoundPinName);
 
 	// Add DataTable pin
-	UEdGraphPin* DataTablePin = CreatePin(EGPD_Input, K2Schema->PC_Object, TEXT(""), UDataTable::StaticClass(), false, false, UK2Node_GetDataTableRowHelper::DataTablePinName);
+	UEdGraphPin* DataTablePin = CreatePin(EGPD_Input, K2Schema->PC_Object, FString(), UDataTable::StaticClass(), UK2Node_GetDataTableRowHelper::DataTablePinName);
 	SetPinToolTip(*DataTablePin, LOCTEXT("DataTablePinDescription", "The DataTable you want to retreive a row from"));
 
 	// Row Name pin
-	UEdGraphPin* RowNamePin = CreatePin(EGPD_Input, K2Schema->PC_Name, TEXT(""), NULL, false, false, UK2Node_GetDataTableRowHelper::RowNamePinName);
+	UEdGraphPin* RowNamePin = CreatePin(EGPD_Input, K2Schema->PC_Name, FString(), nullptr, UK2Node_GetDataTableRowHelper::RowNamePinName);
 	SetPinToolTip(*RowNamePin, LOCTEXT("RowNamePinDescription", "The name of the row to retrieve from the DataTable"));
 
 	// Result pin
-	UEdGraphPin* ResultPin = CreatePin(EGPD_Output, K2Schema->PC_Wildcard, TEXT(""), nullptr, false, false, K2Schema->PN_ReturnValue);
+	UEdGraphPin* ResultPin = CreatePin(EGPD_Output, K2Schema->PC_Wildcard, FString(), nullptr, K2Schema->PN_ReturnValue);
 	ResultPin->PinFriendlyName = LOCTEXT("GetDataTableRow Output Row", "Out Row");
 	SetPinToolTip(*ResultPin, LOCTEXT("ResultPinDescription", "The returned TableRow, if found"));
 
@@ -149,12 +149,12 @@ void UK2Node_GetDataTableRow::OnDataTableRowListChanged(const UDataTable* DataTa
 	UEdGraphPin* DataTablePin = GetDataTablePin();
 	if (DataTable && DataTablePin && DataTable == DataTablePin->DefaultObject)
 	{
-		auto RowNamePin = GetRowNamePin();
+		UEdGraphPin* RowNamePin = GetRowNamePin();
 		const bool TryRefresh = RowNamePin && !RowNamePin->LinkedTo.Num();
 		const FName CurrentName = RowNamePin ? FName(*RowNamePin->GetDefaultAsString()) : NAME_None;
 		if (TryRefresh && RowNamePin && !DataTable->GetRowNames().Contains(CurrentName))
 		{
-			if (auto BP = GetBlueprint())
+			if (UBlueprint* BP = GetBlueprint())
 			{
 				FBlueprintEditorUtils::MarkBlueprintAsModified(BP);
 			}
@@ -269,17 +269,16 @@ UEdGraphPin* UK2Node_GetDataTableRow::GetDataTablePin(const TArray<UEdGraphPin*>
 {
 	const TArray<UEdGraphPin*>* PinsToSearch = InPinsToSearch ? InPinsToSearch : &Pins;
     
-	UEdGraphPin* Pin = NULL;
-	for( auto PinIt = PinsToSearch->CreateConstIterator(); PinIt; ++PinIt )
+	UEdGraphPin* Pin = nullptr;
+	for (UEdGraphPin* TestPin : *PinsToSearch)
 	{
-		UEdGraphPin* TestPin = *PinIt;
-		if( TestPin && TestPin->PinName == UK2Node_GetDataTableRowHelper::DataTablePinName)
+		if (TestPin && TestPin->PinName == UK2Node_GetDataTableRowHelper::DataTablePinName)
 		{
 			Pin = TestPin;
 			break;
 		}
 	}
-	check(Pin == NULL || Pin->Direction == EGPD_Input);
+	check(Pin == nullptr || Pin->Direction == EGPD_Input);
 	return Pin;
 }
 
@@ -418,8 +417,10 @@ void UK2Node_GetDataTableRow::PostReconstructNode()
 
 void UK2Node_GetDataTableRow::EarlyValidation(class FCompilerResultsLog& MessageLog) const
 {
-	const auto DataTablePin = GetDataTablePin();
-	const auto RowNamePin = GetRowNamePin();
+	Super::EarlyValidation(MessageLog);
+
+	const UEdGraphPin* DataTablePin = GetDataTablePin();
+	const UEdGraphPin* RowNamePin = GetRowNamePin();
 	if (!DataTablePin || !RowNamePin)
 	{
 		MessageLog.Error(*LOCTEXT("MissingPins", "Missing pins in @@").ToString(), this);
@@ -428,7 +429,7 @@ void UK2Node_GetDataTableRow::EarlyValidation(class FCompilerResultsLog& Message
 
 	if (DataTablePin->LinkedTo.Num() == 0)
 	{
-		const auto DataTable = Cast<UDataTable>(DataTablePin->DefaultObject);
+		const UDataTable* DataTable = Cast<UDataTable>(DataTablePin->DefaultObject);
 		if (!DataTable)
 		{
 			MessageLog.Error(*LOCTEXT("NoDataTable", "No DataTable in @@").ToString(), this);
diff --git a/Engine/Source/Editor/BlueprintGraph/Private/K2Node_GetEnumeratorName.cpp b/Engine/Source/Editor/BlueprintGraph/Private/K2Node_GetEnumeratorName.cpp
--- a/Engine/Source/Editor/BlueprintGraph/Private/K2Node_GetEnumeratorName.cpp
+++ b/Engine/Source/Editor/BlueprintGraph/Private/K2Node_GetEnumeratorName.cpp
@@ -22,8 +22,8 @@ void UK2Node_GetEnumeratorName::AllocateDefaultPins()
 {
 	const UEdGraphSchema_K2* Schema = GetDefault<UEdGraphSchema_K2>();
 
-	CreatePin(EGPD_Input, Schema->PC_Byte, TEXT(""), NULL, false, false, EnumeratorPinName);
-	CreatePin(EGPD_Output, Schema->PC_Name, TEXT(""), NULL, false, false, Schema->PN_ReturnValue);
+	CreatePin(EGPD_Input, Schema->PC_Byte, FString(), nullptr, EnumeratorPinName);
+	CreatePin(EGPD_Output, Schema->PC_Name, FString(), nullptr, Schema->PN_ReturnValue);
 }
 
 FText UK2Node_GetEnumeratorName::GetTooltipText() const
@@ -45,11 +45,13 @@ UEnum* UK2Node_GetEnumeratorName::GetEnum() const
 {
 	const UEdGraphPin* InputPin = FindPinChecked(EnumeratorPinName);
 	const UEdGraphPin* EnumPin = InputPin->LinkedTo.Num() ? InputPin->LinkedTo[0] : InputPin;
-	return EnumPin ? Cast<UEnum>(EnumPin->PinType.PinSubCategoryObject.Get()) : NULL;
+	return EnumPin ? Cast<UEnum>(EnumPin->PinType.PinSubCategoryObject.Get()) : nullptr;
 }
 
 void UK2Node_GetEnumeratorName::ValidateNodeDuringCompilation(class FCompilerResultsLog& MessageLog) const
 {
+	Super::ValidateNodeDuringCompilation(MessageLog);
+
 	const UEdGraphSchema_K2* Schema = GetDefault<UEdGraphSchema_K2>();
 	const UEdGraphPin* OutputPin = FindPinChecked(Schema->PN_ReturnValue); 
 	/*Don't validate isolated nodes */
@@ -67,8 +69,10 @@ FSlateIcon UK2Node_GetEnumeratorName::GetIconAndTint(FLinearColor& OutColor) con
 
 void UK2Node_GetEnumeratorName::EarlyValidation(class FCompilerResultsLog& MessageLog) const
 {
+	Super::EarlyValidation(MessageLog);
+
 	const UEnum* Enum = GetEnum();
-	if (NULL == Enum)
+	if (nullptr == Enum)
 	{
 		MessageLog.Error(*NSLOCTEXT("K2Node", "GetNumEnumEntries_NoIntput_Error", "@@ Must have non-default Enum input").ToString(), this);
 	}
diff --git a/Engine/Source/Editor/BlueprintGraph/Private/K2Node_GetEnumeratorNameAsString.cpp b/Engine/Source/Editor/BlueprintGraph/Private/K2Node_GetEnumeratorNameAsString.cpp
--- a/Engine/Source/Editor/BlueprintGraph/Private/K2Node_GetEnumeratorNameAsString.cpp
+++ b/Engine/Source/Editor/BlueprintGraph/Private/K2Node_GetEnumeratorNameAsString.cpp
@@ -17,8 +17,8 @@ void UK2Node_GetEnumeratorNameAsString::AllocateDefaultPins()
 {
 	const UEdGraphSchema_K2* Schema = GetDefault<UEdGraphSchema_K2>();
 
-	CreatePin(EGPD_Input, Schema->PC_Byte, TEXT(""), NULL, false, false, EnumeratorPinName);
-	CreatePin(EGPD_Output, Schema->PC_String, TEXT(""), NULL, false, false, Schema->PN_ReturnValue);
+	CreatePin(EGPD_Input, Schema->PC_Byte, FString(), nullptr, EnumeratorPinName);
+	CreatePin(EGPD_Output, Schema->PC_String, FString(), nullptr, Schema->PN_ReturnValue);
 }
 
 FText UK2Node_GetEnumeratorNameAsString::GetTooltipText() const
diff --git a/Engine/Source/Editor/BlueprintGraph/Private/K2Node_GetNumEnumEntries.cpp b/Engine/Source/Editor/BlueprintGraph/Private/K2Node_GetNumEnumEntries.cpp
--- a/Engine/Source/Editor/BlueprintGraph/Private/K2Node_GetNumEnumEntries.cpp
+++ b/Engine/Source/Editor/BlueprintGraph/Private/K2Node_GetNumEnumEntries.cpp
@@ -20,7 +20,7 @@ void UK2Node_GetNumEnumEntries::AllocateDefaultPins()
 	const UEdGraphSchema_K2* Schema = GetDefault<UEdGraphSchema_K2>();
 
 	// Create the return value pin
-	CreatePin(EGPD_Output, Schema->PC_Int, TEXT(""), NULL, false, false, Schema->PN_ReturnValue);
+	CreatePin(EGPD_Output, Schema->PC_Int, FString(), nullptr, Schema->PN_ReturnValue);
 
 	Super::AllocateDefaultPins();
 }
diff --git a/Engine/Source/Editor/BlueprintGraph/Private/K2Node_IfThenElse.cpp b/Engine/Source/Editor/BlueprintGraph/Private/K2Node_IfThenElse.cpp
--- a/Engine/Source/Editor/BlueprintGraph/Private/K2Node_IfThenElse.cpp
+++ b/Engine/Source/Editor/BlueprintGraph/Private/K2Node_IfThenElse.cpp
@@ -93,14 +93,14 @@ void UK2Node_IfThenElse::AllocateDefaultPins()
 {
 	const UEdGraphSchema_K2* K2Schema = GetDefault<UEdGraphSchema_K2>();
 
-	CreatePin(EGPD_Input, K2Schema->PC_Exec, TEXT(""), NULL, false, false, K2Schema->PN_Execute);
-	UEdGraphPin* ConditionPin = CreatePin(EGPD_Input, K2Schema->PC_Boolean, TEXT(""), NULL, false, false, K2Schema->PN_Condition);
-	ConditionPin->DefaultValue = ConditionPin->AutogeneratedDefaultValue = TEXT("true");
+	CreatePin(EGPD_Input, K2Schema->PC_Exec, FString(), nullptr, K2Schema->PN_Execute);
+	UEdGraphPin* ConditionPin = CreatePin(EGPD_Input, K2Schema->PC_Boolean, FString(), nullptr, K2Schema->PN_Condition);
+	K2Schema->SetPinAutogeneratedDefaultValue(ConditionPin, TEXT("true"));
 
-	UEdGraphPin* TruePin = CreatePin(EGPD_Output, K2Schema->PC_Exec, TEXT(""), NULL, false, false, K2Schema->PN_Then);
+	UEdGraphPin* TruePin = CreatePin(EGPD_Output, K2Schema->PC_Exec, FString(), nullptr, K2Schema->PN_Then);
 	TruePin->PinFriendlyName =LOCTEXT("true", "true");
 
-	UEdGraphPin* FalsePin = CreatePin(EGPD_Output, K2Schema->PC_Exec, TEXT(""), NULL, false, false, K2Schema->PN_Else);
+	UEdGraphPin* FalsePin = CreatePin(EGPD_Output, K2Schema->PC_Exec, FString(), nullptr, K2Schema->PN_Else);
 	FalsePin->PinFriendlyName = LOCTEXT("false", "false");
 
 	Super::AllocateDefaultPins();
diff --git a/Engine/Source/Editor/BlueprintGraph/Private/K2Node_InputAction.cpp b/Engine/Source/Editor/BlueprintGraph/Private/K2Node_InputAction.cpp
--- a/Engine/Source/Editor/BlueprintGraph/Private/K2Node_InputAction.cpp
+++ b/Engine/Source/Editor/BlueprintGraph/Private/K2Node_InputAction.cpp
@@ -40,9 +40,9 @@ void UK2Node_InputAction::AllocateDefaultPins()
 {
 	const UEdGraphSchema_K2* K2Schema = GetDefault<UEdGraphSchema_K2>();
 
-	CreatePin(EGPD_Output, K2Schema->PC_Exec, TEXT(""), NULL, false, false, TEXT("Pressed"));
-	CreatePin(EGPD_Output, K2Schema->PC_Exec, TEXT(""), NULL, false, false, TEXT("Released"));
-	CreatePin(EGPD_Output, K2Schema->PC_Struct, TEXT(""), FKey::StaticStruct(), false, false, TEXT("Key"));
+	CreatePin(EGPD_Output, K2Schema->PC_Exec, FString(), nullptr, TEXT("Pressed"));
+	CreatePin(EGPD_Output, K2Schema->PC_Exec, FString(), nullptr, TEXT("Released"));
+	CreatePin(EGPD_Output, K2Schema->PC_Struct, FString(), FKey::StaticStruct(), TEXT("Key"));
 
 	Super::AllocateDefaultPins();
 }
diff --git a/Engine/Source/Editor/BlueprintGraph/Private/K2Node_InputKey.cpp b/Engine/Source/Editor/BlueprintGraph/Private/K2Node_InputKey.cpp
--- a/Engine/Source/Editor/BlueprintGraph/Private/K2Node_InputKey.cpp
+++ b/Engine/Source/Editor/BlueprintGraph/Private/K2Node_InputKey.cpp
@@ -55,9 +55,9 @@ void UK2Node_InputKey::AllocateDefaultPins()
 {
 	const UEdGraphSchema_K2* K2Schema = GetDefault<UEdGraphSchema_K2>();
 
-	CreatePin(EGPD_Output, K2Schema->PC_Exec, TEXT(""), NULL, false, false, TEXT("Pressed"));
-	CreatePin(EGPD_Output, K2Schema->PC_Exec, TEXT(""), NULL, false, false, TEXT("Released"));
-	CreatePin(EGPD_Output, K2Schema->PC_Struct, TEXT(""), FKey::StaticStruct(), false, false, TEXT("Key"));
+	CreatePin(EGPD_Output, K2Schema->PC_Exec, FString(), nullptr, TEXT("Pressed"));
+	CreatePin(EGPD_Output, K2Schema->PC_Exec, FString(), nullptr, TEXT("Released"));
+	CreatePin(EGPD_Output, K2Schema->PC_Struct, FString(), FKey::StaticStruct(), TEXT("Key"));
 
 	Super::AllocateDefaultPins();
 }
diff --git a/Engine/Source/Editor/BlueprintGraph/Private/K2Node_InputTouch.cpp b/Engine/Source/Editor/BlueprintGraph/Private/K2Node_InputTouch.cpp
--- a/Engine/Source/Editor/BlueprintGraph/Private/K2Node_InputTouch.cpp
+++ b/Engine/Source/Editor/BlueprintGraph/Private/K2Node_InputTouch.cpp
@@ -46,14 +46,14 @@ void UK2Node_InputTouch::AllocateDefaultPins()
 {
 	const UEdGraphSchema_K2* K2Schema = GetDefault<UEdGraphSchema_K2>();
 
-	CreatePin(EGPD_Output, K2Schema->PC_Exec, TEXT(""), NULL, false, false, TEXT("Pressed"));
-	CreatePin(EGPD_Output, K2Schema->PC_Exec, TEXT(""), NULL, false, false, TEXT("Released"));
-	CreatePin(EGPD_Output, K2Schema->PC_Exec, TEXT(""), NULL, false, false, TEXT("Moved"));
+	CreatePin(EGPD_Output, K2Schema->PC_Exec, FString(), nullptr, TEXT("Pressed"));
+	CreatePin(EGPD_Output, K2Schema->PC_Exec, FString(), nullptr, TEXT("Released"));
+	CreatePin(EGPD_Output, K2Schema->PC_Exec, FString(), nullptr, TEXT("Moved"));
 	
 	UScriptStruct* VectorStruct = TBaseStructure<FVector>::Get();
-	CreatePin(EGPD_Output, K2Schema->PC_Struct, TEXT(""), VectorStruct, false, false, TEXT("Location"));
+	CreatePin(EGPD_Output, K2Schema->PC_Struct, FString(), VectorStruct, TEXT("Location"));
 
-	CreatePin(EGPD_Output, K2Schema->PC_Byte, TEXT(""), GetTouchIndexEnum(), false, false, TEXT("FingerIndex"));
+	CreatePin(EGPD_Output, K2Schema->PC_Byte, FString(), GetTouchIndexEnum(), TEXT("FingerIndex"));
 
 	Super::AllocateDefaultPins();
 }
diff --git a/Engine/Source/Editor/BlueprintGraph/Private/K2Node_Knot.cpp b/Engine/Source/Editor/BlueprintGraph/Private/K2Node_Knot.cpp
--- a/Engine/Source/Editor/BlueprintGraph/Private/K2Node_Knot.cpp
+++ b/Engine/Source/Editor/BlueprintGraph/Private/K2Node_Knot.cpp
@@ -24,10 +24,10 @@ void UK2Node_Knot::AllocateDefaultPins()
 	const FString InputPinName(TEXT("InputPin"));
 	const FString OutputPinName(TEXT("OutputPin"));
 
-	UEdGraphPin* MyInputPin = CreatePin(EGPD_Input, Schema->PC_Wildcard, FString(), nullptr, /*bIsArray=*/ false, /*bIsReference=*/ false, InputPinName);
+	UEdGraphPin* MyInputPin = CreatePin(EGPD_Input, Schema->PC_Wildcard, FString(), nullptr, InputPinName);
 	MyInputPin->bDefaultValueIsIgnored = true;
 
-	UEdGraphPin* MyOutputPin = CreatePin(EGPD_Output, Schema->PC_Wildcard, FString(), nullptr, /*bIsArray=*/ false, /*bIsReference=*/ false, OutputPinName);
+	UEdGraphPin* MyOutputPin = CreatePin(EGPD_Output, Schema->PC_Wildcard, FString(), nullptr, OutputPinName);
 }
 
 FText UK2Node_Knot::GetTooltipText() const
diff --git a/Engine/Source/Editor/BlueprintGraph/Private/K2Node_Literal.cpp b/Engine/Source/Editor/BlueprintGraph/Private/K2Node_Literal.cpp
--- a/Engine/Source/Editor/BlueprintGraph/Private/K2Node_Literal.cpp
+++ b/Engine/Source/Editor/BlueprintGraph/Private/K2Node_Literal.cpp
@@ -77,7 +77,7 @@ void UK2Node_Literal::AllocateDefaultPins()
 {
 	// The literal node only has one pin:  an output of the desired value, on a wildcard pin type
 	const UEdGraphSchema_K2* Schema = GetDefault<UEdGraphSchema_K2>();
-	CreatePin(EGPD_Output, Schema->PC_Object, TEXT(""), NULL, false, false, *ValuePinName);
+	CreatePin(EGPD_Output, Schema->PC_Object, FString(), nullptr, *ValuePinName);
 
 	// After allocating the pins, try to coerce pin type
 	SetObjectRef( ObjectRef );
@@ -268,7 +268,7 @@ void UK2Node_Literal::SetObjectRef(UObject* NewValue)
 		{
 			ValuePin->Modify();
 			ValuePin->PinType.PinCategory = Schema->PC_Object;
-			ValuePin->PinType.PinSubCategory = TEXT("");
+			ValuePin->PinType.PinSubCategory.Reset();
 			ValuePin->PinType.PinSubCategoryObject = ObjectRef->GetClass();
 		}
 	}
diff --git a/Engine/Source/Editor/BlueprintGraph/Private/K2Node_LoadAsset.cpp b/Engine/Source/Editor/BlueprintGraph/Private/K2Node_LoadAsset.cpp
--- a/Engine/Source/Editor/BlueprintGraph/Private/K2Node_LoadAsset.cpp
+++ b/Engine/Source/Editor/BlueprintGraph/Private/K2Node_LoadAsset.cpp
@@ -22,15 +22,15 @@ void UK2Node_LoadAsset::AllocateDefaultPins()
 
 	if (K2Schema)
 	{
-		CreatePin(EGPD_Input, K2Schema->PC_Exec, TEXT(""), nullptr, false, false, K2Schema->PN_Execute);
-		UEdGraphPin* OutputExecPin = CreatePin(EGPD_Output, K2Schema->PC_Exec, TEXT(""), nullptr, false, false, K2Schema->PN_Then);
+		CreatePin(EGPD_Input, K2Schema->PC_Exec, FString(), nullptr, K2Schema->PN_Execute);
+		UEdGraphPin* OutputExecPin = CreatePin(EGPD_Output, K2Schema->PC_Exec, FString(), nullptr, K2Schema->PN_Then);
 		if (OutputExecPin)
 		{
 			OutputExecPin->PinFriendlyName = FText::FromString(K2Schema->PN_Completed);
 		}
 
-		CreatePin(EGPD_Input, GetInputCategory(), TEXT(""), UObject::StaticClass(), false, false, GetInputPinName());
-		CreatePin(EGPD_Output, GetOutputCategory(), TEXT(""), UObject::StaticClass(), false, false, GetOutputPinName());
+		CreatePin(EGPD_Input, GetInputCategory(), FString(), UObject::StaticClass(), GetInputPinName());
+		CreatePin(EGPD_Output, GetOutputCategory(), FString(), UObject::StaticClass(), GetOutputPinName());
 	}
 }
 
@@ -54,7 +54,7 @@ void UK2Node_LoadAsset::ExpandNode(class FKismetCompilerContext& CompilerContext
 	}
 
 	// Create Local Variable
-	UK2Node_TemporaryVariable* TempVarOutput = CompilerContext.SpawnInternalVariable(this, GetOutputCategory(), FString(), UObject::StaticClass(), false);
+	UK2Node_TemporaryVariable* TempVarOutput = CompilerContext.SpawnInternalVariable(this, GetOutputCategory(), FString(), UObject::StaticClass());
 
 	// Create assign node
 	auto AssignNode = CompilerContext.SpawnIntermediateNode<UK2Node_AssignmentStatement>(this, SourceGraph);
diff --git a/Engine/Source/Editor/BlueprintGraph/Private/K2Node_LocalVariable.cpp b/Engine/Source/Editor/BlueprintGraph/Private/K2Node_LocalVariable.cpp
--- a/Engine/Source/Editor/BlueprintGraph/Private/K2Node_LocalVariable.cpp
+++ b/Engine/Source/Editor/BlueprintGraph/Private/K2Node_LocalVariable.cpp
@@ -142,9 +142,9 @@ void UDEPRECATED_K2Node_LocalVariable::ReconstructNode()
 					SetNode->PostPlacedNewNode();
 
 					// Locally re-construct the pins, can't allow the node to do it because the property does not yet exist
-					UEdGraphPin* ExecPin = SetNode->CreatePin(EGPD_Input, K2Schema->PC_Exec, TEXT(""), NULL, false, false, K2Schema->PN_Execute);
-					UEdGraphPin* ThenPin = SetNode->CreatePin(EGPD_Output, K2Schema->PC_Exec, TEXT(""), NULL, false, false, K2Schema->PN_Then);
-					UEdGraphPin* ValuePin = SetNode->CreatePin(EGPD_Input, TEXT(""), TEXT(""), NULL, false, false, NewVar.VarName.ToString());
+					UEdGraphPin* ExecPin = SetNode->CreatePin(EGPD_Input, K2Schema->PC_Exec, FString(), nullptr, K2Schema->PN_Execute);
+					UEdGraphPin* ThenPin = SetNode->CreatePin(EGPD_Output, K2Schema->PC_Exec, FString(), nullptr, K2Schema->PN_Then);
+					UEdGraphPin* ValuePin = SetNode->CreatePin(EGPD_Input, FString(), FString(), nullptr, NewVar.VarName.ToString());
 					ValuePin->PinType = NewVar.VarType;
 
 					// Move the pin links over to the new node
@@ -173,9 +173,9 @@ void UDEPRECATED_K2Node_LocalVariable::ReconstructNode()
 			GetNode->PostPlacedNewNode();
 
 			// Locally re-construct the pins, can't allow the node to do it because the property does not yet exist
-			UEdGraphPin* VariablePin = GetNode->CreatePin(EGPD_Output, TEXT(""), TEXT(""), NULL, false, false, NewVar.VarName.ToString());
+			UEdGraphPin* VariablePin = GetNode->CreatePin(EGPD_Output, FString(), FString(), nullptr, NewVar.VarName.ToString());
 			VariablePin->PinType = NewVar.VarType;
-			K2Schema->SetPinDefaultValueBasedOnType(VariablePin);
+			K2Schema->SetPinAutogeneratedDefaultValueBasedOnType(VariablePin);
 
 			// Position the new node
 			GetNode->NodePosX = NodePosX;
diff --git a/Engine/Source/Editor/BlueprintGraph/Private/K2Node_MCDelegate.cpp b/Engine/Source/Editor/BlueprintGraph/Private/K2Node_MCDelegate.cpp
--- a/Engine/Source/Editor/BlueprintGraph/Private/K2Node_MCDelegate.cpp
+++ b/Engine/Source/Editor/BlueprintGraph/Private/K2Node_MCDelegate.cpp
@@ -77,8 +77,8 @@ void UK2Node_BaseMCDelegate::AllocateDefaultPins()
 
 	const UEdGraphSchema_K2* K2Schema = GetDefault<UEdGraphSchema_K2>();
 
-	CreatePin(EGPD_Input, K2Schema->PC_Exec, TEXT(""), NULL, false, false, K2Schema->PN_Execute);
-	CreatePin(EGPD_Output, K2Schema->PC_Exec, TEXT(""), NULL, false, false, K2Schema->PN_Then);
+	CreatePin(EGPD_Input, K2Schema->PC_Exec, FString(), nullptr, K2Schema->PN_Execute);
+	CreatePin(EGPD_Output, K2Schema->PC_Exec, FString(), nullptr, K2Schema->PN_Then);
 
 	UClass* PropertyOwnerClass = DelegateReference.GetMemberParentClass(GetBlueprintClassFromNode());
 	if (PropertyOwnerClass != nullptr)
@@ -92,11 +92,11 @@ void UK2Node_BaseMCDelegate::AllocateDefaultPins()
 	UEdGraphPin* SelfPin = NULL;
 	if (bUseSelf)
 	{
-		SelfPin = CreatePin(EGPD_Input, K2Schema->PC_Object, K2Schema->PSC_Self, NULL, false, false, K2Schema->PN_Self);
+		SelfPin = CreatePin(EGPD_Input, K2Schema->PC_Object, K2Schema->PSC_Self, nullptr, K2Schema->PN_Self);
 	}
 	else
 	{
-		SelfPin = CreatePin(EGPD_Input, K2Schema->PC_Object, TEXT(""), PropertyOwnerClass, false, false, K2Schema->PN_Self);
+		SelfPin = CreatePin(EGPD_Input, K2Schema->PC_Object, FString(), PropertyOwnerClass, K2Schema->PN_Self);
 	}
 
 	if(SelfPin)
@@ -136,13 +136,13 @@ UEdGraphPin* UK2Node_BaseMCDelegate::GetDelegatePin() const
 
 FString UK2Node_BaseMCDelegate::GetDocumentationLink() const
 {
-	UClass* ParentClass = NULL;
+	UClass* ParentClass = nullptr;
 	if (DelegateReference.IsSelfContext())
 	{
 		if (HasValidBlueprint())
 		{
 			UField* Delegate = FindField<UField>(GetBlueprint()->GeneratedClass, DelegateReference.GetMemberName());
-			if (Delegate != NULL)
+			if (Delegate)
 			{
 				ParentClass = Delegate->GetOwnerClass();
 			}
@@ -153,12 +153,12 @@ FString UK2Node_BaseMCDelegate::GetDocumentationLink() const
 		ParentClass = DelegateReference.GetMemberParentClass(GetBlueprintClassFromNode());
 	}
 
-	if ( ParentClass != NULL )
+	if (ParentClass)
 	{
 		return FString( TEXT("Shared/") ) + ParentClass->GetName();
 	}
 
-	return TEXT("");
+	return FString();
 }
 
 FString UK2Node_BaseMCDelegate::GetDocumentationExcerptName() const
@@ -181,7 +181,7 @@ void UK2Node_BaseMCDelegate::ExpandNode(class FKismetCompilerContext& CompilerCo
 		const bool bProperInputToExpandForEach = 
 			(MultiSelf->LinkedTo.Num()) && 
 			(NULL != MultiSelf->LinkedTo[0]) && 
-			(MultiSelf->LinkedTo[0]->PinType.bIsArray);
+			(MultiSelf->LinkedTo[0]->PinType.IsArray());
 		if(bProperInputToExpandForEach)
 		{
 			if(MultiSelf->LinkedTo.Num() > 1)
@@ -296,7 +296,7 @@ void UK2Node_AddDelegate::AllocateDefaultPins()
 
 	const UEdGraphSchema_K2* K2Schema = GetDefault<UEdGraphSchema_K2>();
 
-	if(UEdGraphPin* DelegatePin = CreatePin(EGPD_Input, K2Schema->PC_Delegate, TEXT(""), NULL, false, true, FK2Node_BaseMCDelegateHelper::DelegatePinName, true))
+	if(UEdGraphPin* DelegatePin = CreatePin(EGPD_Input, K2Schema->PC_Delegate, FString(), nullptr, FK2Node_BaseMCDelegateHelper::DelegatePinName, EPinContainerType::None, true, true))
 	{
 		FMemberReference::FillSimpleMemberReference<UFunction>(GetDelegateSignature(), DelegatePin->PinType.PinSubCategoryMemberReference);
 		DelegatePin->PinFriendlyName = NSLOCTEXT("K2Node", "PinFriendlyDelegatetName", "Event");
@@ -364,7 +364,7 @@ void UK2Node_RemoveDelegate::AllocateDefaultPins()
 
 	const UEdGraphSchema_K2* K2Schema = GetDefault<UEdGraphSchema_K2>();
 
-	if(UEdGraphPin* DelegatePin = CreatePin(EGPD_Input, K2Schema->PC_Delegate, TEXT(""), NULL, false, true, FK2Node_BaseMCDelegateHelper::DelegatePinName, true))
+	if(UEdGraphPin* DelegatePin = CreatePin(EGPD_Input, K2Schema->PC_Delegate, FString(), nullptr, FK2Node_BaseMCDelegateHelper::DelegatePinName, EPinContainerType::None, true, true))
 	{
 		FMemberReference::FillSimpleMemberReference<UFunction>(GetDelegateSignature(), DelegatePin->PinType.PinSubCategoryMemberReference);
 		DelegatePin->PinFriendlyName = NSLOCTEXT("K2Node", "PinFriendlyDelegatetName", "Event");
@@ -407,7 +407,7 @@ bool UK2Node_CallDelegate::CreatePinsForFunctionInputs(const UFunction* Function
 		const bool bIsFunctionInput = !Param->HasAnyPropertyFlags(CPF_OutParm) || Param->HasAnyPropertyFlags(CPF_ReferenceParm);
 		if (bIsFunctionInput)
 		{
-			UEdGraphPin* Pin = CreatePin(EGPD_Input, TEXT(""), TEXT(""), NULL, false, false, Param->GetName());
+			UEdGraphPin* Pin = CreatePin(EGPD_Input, FString(), FString(), nullptr, Param->GetName());
 			const bool bPinGood = K2Schema->ConvertPropertyToPinType(Param, /*out*/ Pin->PinType);
 
 			bAllPinsGood = bAllPinsGood && bPinGood;
diff --git a/Engine/Source/Editor/BlueprintGraph/Private/K2Node_MacroInstance.cpp b/Engine/Source/Editor/BlueprintGraph/Private/K2Node_MacroInstance.cpp
--- a/Engine/Source/Editor/BlueprintGraph/Private/K2Node_MacroInstance.cpp
+++ b/Engine/Source/Editor/BlueprintGraph/Private/K2Node_MacroInstance.cpp
@@ -79,6 +79,7 @@ void UK2Node_MacroInstance::AllocateDefaultPins()
 {
 	UK2Node::AllocateDefaultPins();
 
+	const UEdGraphSchema_K2* Schema = GetDefault<UEdGraphSchema_K2>();
 	PreloadObject(MacroGraphReference.GetBlueprint());
 	UEdGraph* MacroGraph = MacroGraphReference.GetGraph();
 	
@@ -108,7 +109,7 @@ void UK2Node_MacroInstance::AllocateDefaultPins()
 						if (PortPin->ParentPin == NULL)
 						{
 							UEdGraphPin* NewLocalPin = CreatePin(UEdGraphPin::GetComplementaryDirection(PortPin->Direction), PortPin->PinType, PortPin->PinName);
-							NewLocalPin->DefaultValue = NewLocalPin->AutogeneratedDefaultValue = PortPin->DefaultValue;
+							Schema->SetPinAutogeneratedDefaultValue(NewLocalPin, PortPin->GetDefaultAsString());
 						}
 					}
 				}
diff --git a/Engine/Source/Editor/BlueprintGraph/Private/K2Node_MakeArray.cpp b/Engine/Source/Editor/BlueprintGraph/Private/K2Node_MakeArray.cpp
--- a/Engine/Source/Editor/BlueprintGraph/Private/K2Node_MakeArray.cpp
+++ b/Engine/Source/Editor/BlueprintGraph/Private/K2Node_MakeArray.cpp
@@ -59,9 +59,8 @@ public:
 		CreateArrayStatement.Type = KCST_CreateArray;
 		CreateArrayStatement.LHS = *ArrayTerm;
 
-		for(auto PinIt = Node->Pins.CreateIterator(); PinIt; ++PinIt)
+		for (UEdGraphPin* Pin : Node->Pins)
 		{
-			UEdGraphPin* Pin = *PinIt;
 			if(Pin && Pin->Direction == EGPD_Input)
 			{
 				FBPTerminal** InputTerm = Context.NetMap.Find(FEdGraphUtilities::GetNetFromPin(Pin));
@@ -128,12 +127,12 @@ void UK2Node_MakeArray::ReallocatePinsDuringReconstruction(TArray<UEdGraphPin*>&
 void UK2Node_MakeArray::AllocateDefaultPins()
 {
 	// Create the output pin
-	CreatePin(EGPD_Output, UEdGraphSchema_K2::PC_Wildcard, TEXT(""), NULL, true, false, *OutputPinName);
+	CreatePin(EGPD_Output, UEdGraphSchema_K2::PC_Wildcard, FString(), nullptr, *OutputPinName, EPinContainerType::Array);
 
 	// Create the input pins to create the arrays from
 	for (int32 i = 0; i < NumInputs; ++i)
 	{
-		CreatePin(EGPD_Input, UEdGraphSchema_K2::PC_Wildcard, TEXT(""), NULL, false, false, *FString::Printf(TEXT("[%d]"), i));
+		CreatePin(EGPD_Input, UEdGraphSchema_K2::PC_Wildcard, FString(), nullptr, *FString::Printf(TEXT("[%d]"), i));
 	}
 }
 
@@ -161,8 +160,8 @@ void UK2Node_MakeArray::ClearPinTypeToWildcard()
 	{
 		UEdGraphPin* OutputPin = GetOutputPin();
 		OutputPin->PinType.PinCategory = UEdGraphSchema_K2::PC_Wildcard;
-		OutputPin->PinType.PinSubCategory = TEXT("");
-		OutputPin->PinType.PinSubCategoryObject = NULL;
+		OutputPin->PinType.PinSubCategory.Reset();
+		OutputPin->PinType.PinSubCategoryObject = nullptr;
 
 		PropagatePinType();
 	}
@@ -194,8 +193,9 @@ void UK2Node_MakeArray::NotifyPinConnectionListChanged(UEdGraphPin* Pin)
 		{
 			// Update the types on all the pins
 			OutputPin->PinType = Pin->LinkedTo[0]->PinType;
-			OutputPin->PinType.bIsArray = true;
+			OutputPin->PinType.ContainerType = EPinContainerType::Array;
 			PropagatePinType();
+			GetGraph()->NotifyGraphChanged();
 		}
 	}
 	else
@@ -208,7 +208,7 @@ void UK2Node_MakeArray::NotifyPinConnectionListChanged(UEdGraphPin* Pin)
 			for (int32 PinIndex = 0; PinIndex < InputPins.Num(); ++PinIndex)
 			{
 				UEdGraphPin*& CurrentPin = InputPins[PinIndex];
-				if (CurrentPin->GetDefaultAsString() != CurrentPin->AutogeneratedDefaultValue)
+				if (!CurrentPin->DoesDefaultValueMatchAutogenerated())
 				{
 					bResetOutputPin = false;
 				}
@@ -217,10 +217,11 @@ void UK2Node_MakeArray::NotifyPinConnectionListChanged(UEdGraphPin* Pin)
 			if( bResetOutputPin == true )
 			{
 				OutputPin->PinType.PinCategory = UEdGraphSchema_K2::PC_Wildcard;
-				OutputPin->PinType.PinSubCategory = TEXT("");
-				OutputPin->PinType.PinSubCategoryObject = NULL;
+				OutputPin->PinType.PinSubCategory.Reset();
+				OutputPin->PinType.PinSubCategoryObject = nullptr;
 
 				PropagatePinType();
+				GetGraph()->NotifyGraphChanged();
 			}
 		}
 	}
@@ -275,14 +276,10 @@ void UK2Node_MakeArray::PropagatePinType()
 					CurrentPin->PinType.PinSubCategoryObject = OutputPin->PinType.PinSubCategoryObject;
 				}
 
-				if (CurrentPin->PinType.PinCategory == UEdGraphSchema_K2::PC_Wildcard)
+				if (CurrentPin->PinType.PinCategory == UEdGraphSchema_K2::PC_Wildcard || CurrentPin->GetDefaultAsString().IsEmpty() == true)
 				{
-					CurrentPin->ResetDefaultValue();
-				}
-				else if (CurrentPin->GetDefaultAsString().IsEmpty() == true)
-				{
-					// Only reset default value if there isn't one set. Otherwise this deletes data!
-					Schema->SetPinDefaultValueBasedOnType(CurrentPin);
+					// Only reset default value if there isn't one set or it is a wildcard. Otherwise this deletes data!
+					Schema->SetPinAutogeneratedDefaultValueBasedOnType(CurrentPin);
 				}
 
 				// Verify that all previous connections to this pin are still valid with the new type
@@ -341,12 +338,12 @@ void UK2Node_MakeArray::PostReconstructNode()
 			{
 				OutputPin->PinType = Pins[PinIndex]->LinkedTo[0]->PinType;
 			}
-			OutputPin->PinType.bIsArray = true;
+			OutputPin->PinType.ContainerType = EPinContainerType::Array;
 			PropagatePinType();
 			bResetToWildcard = false;
 			break;
 		}
-		else if(Pins[PinIndex]->GetDefaultAsString() != Pins[PinIndex]->AutogeneratedDefaultValue)
+		else if(!Pins[PinIndex]->DoesDefaultValueMatchAutogenerated())
 		{
 			// The pin has user data in it, continue to use its type as the type for all pins.
 
@@ -361,7 +358,7 @@ void UK2Node_MakeArray::PostReconstructNode()
 			{
 				OutputPin->PinType = Pins[PinIndex]->PinType;
 			}
-			OutputPin->PinType.bIsArray = true;
+			OutputPin->PinType.ContainerType = EPinContainerType::Array;
 			PropagatePinType();
 			bResetToWildcard = false;
 			break;
@@ -394,8 +391,8 @@ void UK2Node_MakeArray::AddInputPin()
 
 	++NumInputs;
 	FEdGraphPinType OutputPinType = GetOutputPin()->PinType;
-	UEdGraphPin* Pin = CreatePin(EGPD_Input, OutputPinType.PinCategory, OutputPinType.PinSubCategory, OutputPinType.PinSubCategoryObject.Get(), false, false, *FString::Printf(TEXT("[%d]"), (NumInputs-1)));
-	GetDefault<UEdGraphSchema_K2>()->SetPinDefaultValueBasedOnType(Pin);
+	UEdGraphPin* Pin = CreatePin(EGPD_Input, OutputPinType.PinCategory, OutputPinType.PinSubCategory, OutputPinType.PinSubCategoryObject.Get(), *FString::Printf(TEXT("[%d]"), (NumInputs-1)));
+	GetDefault<UEdGraphSchema_K2>()->SetPinAutogeneratedDefaultValueBasedOnType(Pin);
 
 	const bool bIsCompiling = GetBlueprint()->bBeingCompiled;
 	if( !bIsCompiling )
@@ -533,7 +530,7 @@ bool UK2Node_MakeArray::IsConnectionDisallowed(const UEdGraphPin* MyPin, const U
 		return true;
 	}
 
-	auto Schema = Cast<const UEdGraphSchema_K2>(GetSchema());
+	const UEdGraphSchema_K2* Schema = Cast<const UEdGraphSchema_K2>(GetSchema());
 	if (!ensure(Schema) || (ensure(OtherPin) && Schema->IsExecPin(*OtherPin)))
 	{
 		OutReason = NSLOCTEXT("K2Node", "MakeArray_InputIsExec", "Cannot make an array with an execution input!").ToString();
@@ -545,8 +542,10 @@ bool UK2Node_MakeArray::IsConnectionDisallowed(const UEdGraphPin* MyPin, const U
 
 void UK2Node_MakeArray::ValidateNodeDuringCompilation(class FCompilerResultsLog& MessageLog) const
 {
-	auto Schema = Cast<const UEdGraphSchema_K2>(GetSchema());
-	auto OutputPin = GetOutputPin();
+	Super::ValidateNodeDuringCompilation(MessageLog);
+
+	const UEdGraphSchema_K2* Schema = Cast<const UEdGraphSchema_K2>(GetSchema());
+	UEdGraphPin* OutputPin = GetOutputPin();
 	if (!ensure(Schema) || !ensure(OutputPin) || Schema->IsExecPin(*OutputPin))
 	{
 		MessageLog.Error(*NSLOCTEXT("K2Node", "MakeArray_OutputIsExec", "Uaccepted array type in @@").ToString(), this);
diff --git a/Engine/Source/Editor/BlueprintGraph/Private/K2Node_MakeStruct.cpp b/Engine/Source/Editor/BlueprintGraph/Private/K2Node_MakeStruct.cpp
--- a/Engine/Source/Editor/BlueprintGraph/Private/K2Node_MakeStruct.cpp
+++ b/Engine/Source/Editor/BlueprintGraph/Private/K2Node_MakeStruct.cpp
@@ -39,36 +39,6 @@ void UK2Node_MakeStruct::FMakeStructPinManager::GetRecordDefaults(UProperty* Tes
 
 void UK2Node_MakeStruct::FMakeStructPinManager::CustomizePinData(UEdGraphPin* Pin, FName SourcePropertyName, int32 ArrayIndex, UProperty* Property) const
 {
-	struct FPinDefaultValueHelper
-	{
-		static void Set(UEdGraphPin& InPin, const FString& Value, bool bIsText, bool bIsObject)
-		{
-			InPin.AutogeneratedDefaultValue = Value;
-			if (bIsText)
-			{
-				FString PackageNamespace;
-#if USE_STABLE_LOCALIZATION_KEYS
-				if (GIsEditor)
-				{
-					PackageNamespace = TextNamespaceUtil::EnsurePackageNamespace(InPin.GetOwningNodeUnchecked());
-				}
-#endif // USE_STABLE_LOCALIZATION_KEYS
-				if (!FTextStringHelper::ReadFromString(*Value, InPin.DefaultTextValue, nullptr, *PackageNamespace))
-				{
-					InPin.DefaultTextValue = FText::FromString(Value);
-				}
-			}
-			else if (bIsObject)
-			{
-				InPin.DefaultObject = StaticFindObject(UObject::StaticClass(), nullptr, *Value);
-			}
-			else
-			{
-				InPin.DefaultValue = Value;
-			}
-		}
-	};
-
 	UK2Node_StructOperation::FStructOperationOptionalPinManager::CustomizePinData(Pin, SourcePropertyName, ArrayIndex, Property);
 	if (Pin && Property)
 	{
@@ -77,11 +47,11 @@ void UK2Node_MakeStruct::FMakeStructPinManager::CustomizePinData(UEdGraphPin* Pi
 
 		// Should pin default value be filled as FText?
 		const bool bIsText = Property->IsA<UTextProperty>();
-		checkSlow(bIsText == ((Schema->PC_Text == Pin->PinType.PinCategory) && !Pin->PinType.bIsArray));
+		checkSlow(bIsText == ((Schema->PC_Text == Pin->PinType.PinCategory) && !Pin->PinType.IsContainer()));
 
 		const bool bIsObject = Property->IsA<UObjectPropertyBase>();
 		checkSlow(bIsObject == ((Schema->PC_Object == Pin->PinType.PinCategory || Schema->PC_Class == Pin->PinType.PinCategory || 
-			Schema->PC_Asset == Pin->PinType.PinCategory || Schema->PC_AssetClass == Pin->PinType.PinCategory) && !Pin->PinType.bIsArray));
+			Schema->PC_Asset == Pin->PinType.PinCategory || Schema->PC_AssetClass == Pin->PinType.PinCategory) && !Pin->PinType.IsContainer()));
 
 		if (Property->HasAnyPropertyFlags(CPF_AdvancedDisplay))
 		{
@@ -92,24 +62,24 @@ void UK2Node_MakeStruct::FMakeStructPinManager::CustomizePinData(UEdGraphPin* Pi
 		const FString& MetadataDefaultValue = Property->GetMetaData(TEXT("MakeStructureDefaultValue"));
 		if (!MetadataDefaultValue.IsEmpty())
 		{
-			FPinDefaultValueHelper::Set(*Pin, MetadataDefaultValue, bIsText, bIsObject);
+			Schema->SetPinAutogeneratedDefaultValue(Pin, MetadataDefaultValue);
 			return;
 		}
 
-		if (NULL != SampleStructMemory)
+		if (nullptr != SampleStructMemory)
 		{
 			FString NewDefaultValue;
-			if (Property->ExportText_InContainer(0, NewDefaultValue, SampleStructMemory, SampleStructMemory, NULL, PPF_None))
+			if (Property->ExportText_InContainer(0, NewDefaultValue, SampleStructMemory, SampleStructMemory, nullptr, PPF_None))
 			{
-				if (Schema->IsPinDefaultValid(Pin, NewDefaultValue, NULL, FText::GetEmpty()).IsEmpty())
+				if (Schema->IsPinDefaultValid(Pin, NewDefaultValue, nullptr, FText::GetEmpty()).IsEmpty())
 				{
-					FPinDefaultValueHelper::Set(*Pin, NewDefaultValue, bIsText, bIsObject);
+					Schema->SetPinAutogeneratedDefaultValue(Pin, NewDefaultValue);
 					return;
 				}
 			}
 		}
 
-		Schema->SetPinDefaultValueBasedOnType(Pin);
+		Schema->SetPinAutogeneratedDefaultValueBasedOnType(Pin);
 	}
 }
 
@@ -151,7 +121,7 @@ void UK2Node_MakeStruct::AllocateDefaultPins()
 	if(Schema && StructType)
 	{
 		PreloadObject(StructType);
-		CreatePin(EGPD_Output, Schema->PC_Struct, TEXT(""), StructType, false, false, StructType->GetName());
+		CreatePin(EGPD_Output, Schema->PC_Struct, FString(), StructType, StructType->GetName());
 		
 		bool bHasAdvancedPins = false;
 		{
diff --git a/Engine/Source/Editor/BlueprintGraph/Private/K2Node_MakeVariable.cpp b/Engine/Source/Editor/BlueprintGraph/Private/K2Node_MakeVariable.cpp
--- a/Engine/Source/Editor/BlueprintGraph/Private/K2Node_MakeVariable.cpp
+++ b/Engine/Source/Editor/BlueprintGraph/Private/K2Node_MakeVariable.cpp
@@ -52,15 +52,15 @@ public:
 
 		// This node only supports containers at the moment, it could be extended to support any type:
 		const FBPVariableDescription& VariableType = MakeVariableNode->GetVariableType();
-		if(VariableType.VarType.bIsArray)
+		if(VariableType.VarType.IsArray())
 		{
 			CreateVariableStatement.Type = KCST_CreateArray;
 		}
-		else if(VariableType.VarType.bIsSet)
+		else if(VariableType.VarType.IsSet())
 		{
 			CreateVariableStatement.Type = KCST_CreateSet;
 		}
-		else if( VariableType.VarType.bIsMap)
+		else if( VariableType.VarType.IsMap())
 		{
 			CreateVariableStatement.Type = KCST_CreateMap;
 		}
@@ -102,7 +102,7 @@ void UK2Node_MakeVariable::SetupVariable(const FBPVariableDescription& InVariabl
 	TSharedPtr<FStructOnScope> StructData = MakeShareable(new FStructOnScope(Scope));
 	FBlueprintEditorUtils::PropertyValueFromString(Property, VariableType.DefaultValue, StructData->GetStructMemory());
 
-	if(VariableType.VarType.bIsArray)
+	if(VariableType.VarType.IsArray())
 	{
 		const UArrayProperty* ArrayProperty = CastChecked<UArrayProperty>(Property);
 		FScriptArrayHelper ArrayHelper(ArrayProperty, ArrayProperty->ContainerPtrToValuePtr<void>(StructData->GetStructMemory()));
@@ -118,10 +118,10 @@ void UK2Node_MakeVariable::SetupVariable(const FBPVariableDescription& InVariabl
 			
 			UEdGraphPin* VariableInputPin = CreatePin(EGPD_Input, ContainedType, *FString::FromInt(ArrayIndex));
 			// Add one to the index for the pin to set the default on to skip the output pin
-			Schema->TrySetDefaultValue(*VariableInputPin , DefaultValue);
+			Schema->SetPinAutogeneratedDefaultValue(VariableInputPin, DefaultValue);
 		}
 	}
-	else if( VariableType.VarType.bIsSet)
+	else if( VariableType.VarType.IsSet())
 	{
 		const USetProperty* SetProperty = CastChecked<USetProperty>(Property);
 		FScriptSetHelper SetHelper(SetProperty, SetProperty->ContainerPtrToValuePtr<void>(StructData->GetStructMemory()));
@@ -140,11 +140,11 @@ void UK2Node_MakeVariable::SetupVariable(const FBPVariableDescription& InVariabl
 				FBlueprintEditorUtils::PropertyValueToString_Direct(SetProperty->ElementProp, PropData, DefaultValue);
 				UEdGraphPin* VariableInputPin = CreatePin(EGPD_Input, ContainedType, *FString::FromInt(I));
 				// Add one to the index for the pin to set the default on to skip the output pin
-				Schema->TrySetDefaultValue(*VariableInputPin , DefaultValue);
+				Schema->SetPinAutogeneratedDefaultValue(VariableInputPin, DefaultValue);
 			}
 		}
 	}
-	else if( VariableType.VarType.bIsMap)
+	else if( VariableType.VarType.IsMap())
 	{
 		const UMapProperty* MapProperty = CastChecked<UMapProperty>(Property);
 		FScriptMapHelper MapHelper(MapProperty, MapProperty->ContainerPtrToValuePtr<void>(StructData->GetStructMemory()));
@@ -165,7 +165,7 @@ void UK2Node_MakeVariable::SetupVariable(const FBPVariableDescription& InVariabl
 				FBlueprintEditorUtils::PropertyValueToString_Direct(MapProperty->KeyProp, KeyData, KeyDefaultValue);
 				UEdGraphPin* VariableInputPin = CreatePin(EGPD_Input, ContainedType, *FString::FromInt(I).Append(TEXT("_Key")));
 				// Add one to the index for the pin to set the default on to skip the output pin
-				Schema->TrySetDefaultValue(*VariableInputPin , KeyDefaultValue);
+				Schema->SetPinAutogeneratedDefaultValue(VariableInputPin, KeyDefaultValue);
 
 				const uint8* ValueData = MapHelper.GetValuePtr(I);
 
@@ -173,7 +173,7 @@ void UK2Node_MakeVariable::SetupVariable(const FBPVariableDescription& InVariabl
 				FBlueprintEditorUtils::PropertyValueToString_Direct(MapProperty->ValueProp, ValueData, ValueDefaultValue);
 				VariableInputPin = CreatePin(EGPD_Input, ValueType, *FString::FromInt(I).Append(TEXT("_Value")));
 				// Add one to the index for the pin to set the default on to skip the output pin
-				Schema->TrySetDefaultValue(*VariableInputPin , ValueDefaultValue);
+				Schema->SetPinAutogeneratedDefaultValue(VariableInputPin, ValueDefaultValue);
 			}
 		}
 	}
diff --git a/Engine/Source/Editor/BlueprintGraph/Private/K2Node_MathExpression.cpp b/Engine/Source/Editor/BlueprintGraph/Private/K2Node_MathExpression.cpp
--- a/Engine/Source/Editor/BlueprintGraph/Private/K2Node_MathExpression.cpp
+++ b/Engine/Source/Editor/BlueprintGraph/Private/K2Node_MathExpression.cpp
@@ -2696,6 +2696,8 @@ void UK2Node_MathExpression::ClearExpression()
 //------------------------------------------------------------------------------
 void UK2Node_MathExpression::ValidateNodeDuringCompilation(FCompilerResultsLog& MessageLog) const
 {
+	Super::ValidateNodeDuringCompilation(MessageLog);
+
 	if (CachedMessageLog.IsValid())
 	{
 		MessageLog.Append(*CachedMessageLog);
diff --git a/Engine/Source/Editor/BlueprintGraph/Private/K2Node_MatineeController.cpp b/Engine/Source/Editor/BlueprintGraph/Private/K2Node_MatineeController.cpp
--- a/Engine/Source/Editor/BlueprintGraph/Private/K2Node_MatineeController.cpp
+++ b/Engine/Source/Editor/BlueprintGraph/Private/K2Node_MatineeController.cpp
@@ -39,17 +39,17 @@ void UK2Node_MatineeController::AllocateDefaultPins()
 	const UEdGraphSchema_K2* K2Schema = GetDefault<UEdGraphSchema_K2>();
 
 	// Preload the matinee data, if needed, so that we can have all the event tracks we need
-	if (MatineeActor != NULL)
+	if (MatineeActor)
 	{
 		PreloadObject(MatineeActor);
 		PreloadObject(MatineeActor->MatineeData);
 	}
 
 	// Create the "finished" playing pin
-	CreatePin(EGPD_Output, K2Schema->PC_Exec, TEXT(""), NULL, false, false, K2Schema->PN_MatineeFinished);
+	CreatePin(EGPD_Output, K2Schema->PC_Exec, FString(), nullptr, K2Schema->PN_MatineeFinished);
 
 	// Create pins for each event
-	if(MatineeActor != NULL && MatineeActor->MatineeData != NULL)
+	if (MatineeActor && MatineeActor->MatineeData)
 	{
 		TArray<FName> EventNames;
 		MatineeActor->MatineeData->GetAllEventNames(EventNames);
@@ -57,7 +57,7 @@ void UK2Node_MatineeController::AllocateDefaultPins()
 		for(int32 i=0; i<EventNames.Num(); i++)
 		{
 			FName EventName = EventNames[i];
-			CreatePin(EGPD_Output, K2Schema->PC_Exec, TEXT(""), NULL, false, false, EventName.ToString());			
+			CreatePin(EGPD_Output, K2Schema->PC_Exec, FString(), nullptr, EventName.ToString());			
 		}
 	}
 
@@ -213,7 +213,7 @@ void UK2Node_MatineeController::OnEventKeyframeAdded(const AMatineeActor* InMati
 			const UEdGraphSchema_K2* K2Schema = GetDefault<UEdGraphSchema_K2>();
 
 			// Add one to the index as we insert "finished" at index 0
-			CreatePin(EGPD_Output, K2Schema->PC_Exec, TEXT(""), NULL, false, false, InPinName.ToString(), false, InIndex + 1);
+			CreatePin(EGPD_Output, K2Schema->PC_Exec, FString(), nullptr, InPinName.ToString(), EPinContainerType::None, false, false, InIndex + 1);
 
 			// Update and refresh the blueprint
 			FBlueprintEditorUtils::MarkBlueprintAsStructurallyModified(GetBlueprint());
diff --git a/Engine/Source/Editor/BlueprintGraph/Private/K2Node_Message.cpp b/Engine/Source/Editor/BlueprintGraph/Private/K2Node_Message.cpp
--- a/Engine/Source/Editor/BlueprintGraph/Private/K2Node_Message.cpp
+++ b/Engine/Source/Editor/BlueprintGraph/Private/K2Node_Message.cpp
@@ -62,9 +62,9 @@ void UK2Node_Message::AllocateDefaultPins()
 	if (MessageNodeFunction && MessageNodeFunction->HasAnyFunctionFlags(FUNC_BlueprintPure))
 	{
 		// Input - Execution Pin
-		CreatePin(EGPD_Input,  UEdGraphSchema_K2::PC_Exec, TEXT(""), NULL, false, false, UEdGraphSchema_K2::PN_Execute);
+		CreatePin(EGPD_Input,  UEdGraphSchema_K2::PC_Exec, FString(), nullptr, UEdGraphSchema_K2::PN_Execute);
 		// Output - Execution Pin
-		CreatePin(EGPD_Output, UEdGraphSchema_K2::PC_Exec, TEXT(""), NULL, false, false, UEdGraphSchema_K2::PN_Then);
+		CreatePin(EGPD_Output, UEdGraphSchema_K2::PC_Exec, FString(), nullptr, UEdGraphSchema_K2::PN_Then);
 	}
 
 	Super::AllocateDefaultPins();
@@ -73,7 +73,7 @@ void UK2Node_Message::AllocateDefaultPins()
 UEdGraphPin* UK2Node_Message::CreateSelfPin(const UFunction* Function)
 {
 	const UEdGraphSchema_K2* K2Schema = GetDefault<UEdGraphSchema_K2>();
-	UEdGraphPin* SelfPin = CreatePin(EGPD_Input, K2Schema->PC_Object, TEXT(""), UObject::StaticClass(), false, false, K2Schema->PN_Self);
+	UEdGraphPin* SelfPin = CreatePin(EGPD_Input, K2Schema->PC_Object, FString(), UObject::StaticClass(), K2Schema->PN_Self);
 	SelfPin->bDefaultValueIsIgnored = true;
 	return SelfPin;
 }
@@ -188,7 +188,7 @@ void UK2Node_Message::ExpandNode(class FKismetCompilerContext& CompilerContext,
 
 		// First, create an internal cast-to-interface node
 		UK2Node_DynamicCast* CastToInterfaceNode = CompilerContext.SpawnIntermediateNode<UK2Node_DynamicCast>(this, SourceGraph);
-		CastToInterfaceNode->TargetType = MessageNodeFunction->GetOuterUClass();
+		CastToInterfaceNode->TargetType = MessageNodeFunction->GetOuterUClass()->GetAuthoritativeClass();
 		CastToInterfaceNode->SetPurity(false);
 		CastToInterfaceNode->AllocateDefaultPins();
 
@@ -308,9 +308,9 @@ void UK2Node_Message::ExpandNode(class FKismetCompilerContext& CompilerContext,
 							LastOutCastSuccessPin = AssignTempVar->GetThenPin();
 						}
 
-						UK2Node* DefaultValueNode = NULL;
-						UEdGraphPin* DefaultValueThenPin = NULL;
-						if (CurrentPin->PinType.bIsArray)
+						UK2Node* DefaultValueNode = nullptr;
+						UEdGraphPin* DefaultValueThenPin = nullptr;
+						if (CurrentPin->PinType.IsArray())
 						{
 							UK2Node_CallArrayFunction* ClearArray = CompilerContext.SpawnIntermediateNode<UK2Node_CallArrayFunction>(this, SourceGraph);
 							DefaultValueNode = ClearArray;
@@ -332,7 +332,7 @@ void UK2Node_Message::ExpandNode(class FKismetCompilerContext& CompilerContext,
 
 							Schema->TryCreateConnection(AssignDefaultValue->GetVariablePin(), VarOutPin);
 							AssignDefaultValue->PinConnectionListChanged(AssignDefaultValue->GetVariablePin());
-							Schema->SetPinDefaultValueBasedOnType(AssignDefaultValue->GetValuePin());
+							Schema->SetPinAutogeneratedDefaultValueBasedOnType(AssignDefaultValue->GetValuePin());
 
 							DefaultValueThenPin = AssignDefaultValue->GetThenPin();
 						}
diff --git a/Engine/Source/Editor/BlueprintGraph/Private/K2Node_MultiGate.cpp b/Engine/Source/Editor/BlueprintGraph/Private/K2Node_MultiGate.cpp
--- a/Engine/Source/Editor/BlueprintGraph/Private/K2Node_MultiGate.cpp
+++ b/Engine/Source/Editor/BlueprintGraph/Private/K2Node_MultiGate.cpp
@@ -498,12 +498,11 @@ void UK2Node_MultiGate::AllocateDefaultPins()
 	Super::AllocateDefaultPins();
 
 	const UEdGraphSchema_K2* K2Schema = GetDefault<UEdGraphSchema_K2>();
-	CreatePin(EGPD_Input, K2Schema->PC_Exec, TEXT(""), NULL, false, false, TEXT("Reset"));
-	CreatePin(EGPD_Input, K2Schema->PC_Boolean, TEXT(""), NULL, false, false, TEXT("IsRandom"));
-	CreatePin(EGPD_Input, K2Schema->PC_Boolean, TEXT(""), NULL, false, false, TEXT("Loop"));
-	UEdGraphPin* IndexPin = CreatePin(EGPD_Input, K2Schema->PC_Int, TEXT(""), NULL, false, false, TEXT("StartIndex"));
-	IndexPin->DefaultValue = TEXT("-1");
-	IndexPin->AutogeneratedDefaultValue = TEXT("-1");
+	CreatePin(EGPD_Input, K2Schema->PC_Exec, FString(), nullptr, TEXT("Reset"));
+	CreatePin(EGPD_Input, K2Schema->PC_Boolean, FString(), nullptr, TEXT("IsRandom"));
+	CreatePin(EGPD_Input, K2Schema->PC_Boolean, FString(), nullptr, TEXT("Loop"));
+	UEdGraphPin* IndexPin = CreatePin(EGPD_Input, K2Schema->PC_Int, FString(), nullptr, TEXT("StartIndex"));
+	K2Schema->SetPinAutogeneratedDefaultValue(IndexPin, TEXT("-1"));
 }
 
 void UK2Node_MultiGate::ReallocatePinsDuringReconstruction(TArray<UEdGraphPin*>& OldPins)
@@ -511,12 +510,11 @@ void UK2Node_MultiGate::ReallocatePinsDuringReconstruction(TArray<UEdGraphPin*>&
 	Super::ReallocatePinsDuringReconstruction(OldPins);
 
 	const UEdGraphSchema_K2* K2Schema = GetDefault<UEdGraphSchema_K2>();
-	CreatePin(EGPD_Input, K2Schema->PC_Exec, TEXT(""), NULL, false, false, TEXT("Reset"));
-	CreatePin(EGPD_Input, K2Schema->PC_Boolean, TEXT(""), NULL, false, false, TEXT("IsRandom"));
-	CreatePin(EGPD_Input, K2Schema->PC_Boolean, TEXT(""), NULL, false, false, TEXT("Loop"));
-	UEdGraphPin* IndexPin = CreatePin(EGPD_Input, K2Schema->PC_Int, TEXT(""), NULL, false, false, TEXT("StartIndex"));
-	IndexPin->DefaultValue = TEXT("-1");
-	IndexPin->AutogeneratedDefaultValue = TEXT("-1");
+	CreatePin(EGPD_Input, K2Schema->PC_Exec, FString(), nullptr, TEXT("Reset"));
+	CreatePin(EGPD_Input, K2Schema->PC_Boolean, FString(), nullptr, TEXT("IsRandom"));
+	CreatePin(EGPD_Input, K2Schema->PC_Boolean, FString(), nullptr, TEXT("Loop"));
+	UEdGraphPin* IndexPin = CreatePin(EGPD_Input, K2Schema->PC_Int, FString(), nullptr, TEXT("StartIndex"));
+	K2Schema->SetPinAutogeneratedDefaultValue(IndexPin, TEXT("-1"));
 }
 
 UEdGraphPin* UK2Node_MultiGate::GetResetPin() const
diff --git a/Engine/Source/Editor/BlueprintGraph/Private/K2Node_PureAssignmentStatement.cpp b/Engine/Source/Editor/BlueprintGraph/Private/K2Node_PureAssignmentStatement.cpp
--- a/Engine/Source/Editor/BlueprintGraph/Private/K2Node_PureAssignmentStatement.cpp
+++ b/Engine/Source/Editor/BlueprintGraph/Private/K2Node_PureAssignmentStatement.cpp
@@ -79,9 +79,9 @@ void UK2Node_PureAssignmentStatement::AllocateDefaultPins()
 {
 	const UEdGraphSchema_K2* Schema = GetDefault<UEdGraphSchema_K2>();
 
-	CreatePin(EGPD_Input, Schema->PC_Wildcard, TEXT(""), NULL, false, false, VariablePinName);
-	CreatePin(EGPD_Input, Schema->PC_Wildcard, TEXT(""), NULL, false, false, ValuePinName);
-	CreatePin(EGPD_Output, Schema->PC_Wildcard, TEXT(""), NULL, false, false, OutputPinName);
+	CreatePin(EGPD_Input, Schema->PC_Wildcard, FString(), nullptr, VariablePinName);
+	CreatePin(EGPD_Input, Schema->PC_Wildcard, FString(), nullptr, ValuePinName);
+	CreatePin(EGPD_Output, Schema->PC_Wildcard, FString(), nullptr, OutputPinName);
 
 	Super::AllocateDefaultPins();
 }
diff --git a/Engine/Source/Editor/BlueprintGraph/Private/K2Node_Select.cpp b/Engine/Source/Editor/BlueprintGraph/Private/K2Node_Select.cpp
--- a/Engine/Source/Editor/BlueprintGraph/Private/K2Node_Select.cpp
+++ b/Engine/Source/Editor/BlueprintGraph/Private/K2Node_Select.cpp
@@ -320,9 +320,9 @@ public:
 			NopStatement.Type = KCST_Nop;
 			NopStatement.bIsJumpTarget = true;
 			// Loop through the unconditional goto statements and fix their jump targets
-			for (auto It = GotoStatementList.CreateConstIterator(); It; It++)
+			for (FBlueprintCompiledStatement* GotoStatement : GotoStatementList)
 			{
-				(*It)->TargetLabel = &NopStatement;
+				GotoStatement->TargetLabel = &NopStatement;
 			}
 		}
 	}
@@ -360,7 +360,7 @@ void UK2Node_Select::AllocateDefaultPins()
 	// Create the option pins
 	for (int32 Idx = 0; Idx < NumOptionPins; Idx++)
 	{
-		UEdGraphPin* NewPin = NULL;
+		UEdGraphPin* NewPin = nullptr;
 
 		if (Enum)
 		{
@@ -368,13 +368,13 @@ void UK2Node_Select::AllocateDefaultPins()
 			UEdGraphPin* TempPin = FindPin(PinName);
 			if (!TempPin)
 			{
-				NewPin = CreatePin(EGPD_Input, Schema->PC_Wildcard, TEXT(""), NULL, false, false, PinName);
+				NewPin = CreatePin(EGPD_Input, Schema->PC_Wildcard, FString(), nullptr, PinName);
 			}
 		}
 		else
 		{
 			const FString PinName = FString::Printf(TEXT("Option %d"), Idx);
-			NewPin = CreatePin(EGPD_Input, Schema->PC_Wildcard, TEXT(""), NULL, false, false, PinName);
+			NewPin = CreatePin(EGPD_Input, Schema->PC_Wildcard, FString(), nullptr, PinName);
 		}
 
 		if (NewPin)
@@ -392,10 +392,10 @@ void UK2Node_Select::AllocateDefaultPins()
 	}
 
 	// Create the index wildcard pin
-	CreatePin(EGPD_Input, IndexPinType.PinCategory, IndexPinType.PinSubCategory, IndexPinType.PinSubCategoryObject.Get(), false, false, "Index");
+	CreatePin(EGPD_Input, IndexPinType.PinCategory, IndexPinType.PinSubCategory, IndexPinType.PinSubCategoryObject.Get(), TEXT("Index"));
 
 	// Create the return value
-	auto ReturnPin = CreatePin(EGPD_Output, Schema->PC_Wildcard, TEXT(""), NULL, false, false, Schema->PN_ReturnValue);
+	UEdGraphPin* ReturnPin = CreatePin(EGPD_Output, Schema->PC_Wildcard, FString(), nullptr, Schema->PN_ReturnValue);
 	ReturnPin->bDisplayAsMutableRef = UseSelectRef;
 
 	Super::AllocateDefaultPins();
@@ -465,68 +465,44 @@ void UK2Node_Select::ReallocatePinsDuringReconstruction(TArray<UEdGraphPin*>& Ol
 	const UEdGraphSchema_K2* Schema = GetDefault<UEdGraphSchema_K2>();
 
 	// See if this node was saved in the old version with a boolean as the condition
-	UEdGraphPin* OldConditionPin = NULL;
-	UEdGraphPin* OldIndexPin = NULL;
-	UEdGraphPin* OldReturnPin = NULL;
-	for (auto It = OldPins.CreateConstIterator(); It; It++)
+	UEdGraphPin* OldConditionPin = nullptr;
+	UEdGraphPin* OldIndexPin = nullptr;
+	UEdGraphPin* OldReturnPin = nullptr;
+	for (UEdGraphPin* OldPin : OldPins)
 	{
-		if ((*It)->PinName == TEXT("bPickOption0"))
+		if (OldPin->PinName == TEXT("bPickOption0"))
 		{
-			OldConditionPin = (*It);
+			OldConditionPin = OldPin;
 		}
-		else if ((*It)->PinName == TEXT("Index"))
+		else if (OldPin->PinName == TEXT("Index"))
 		{
-			OldIndexPin = (*It);
+			OldIndexPin = OldPin;
 		}
-		else if ((*It)->PinName == Schema->PN_ReturnValue)
+		else if (OldPin->PinName == Schema->PN_ReturnValue)
 		{
-			OldReturnPin = (*It);
+			OldReturnPin = OldPin;
 		}
 	}
 
-	bool bIsAnyOptionOrReturnConnected = false;
-	for (auto OldPin : OldPins)
-	{
-		if ((OldPin != OldConditionPin) && (OldPin != OldIndexPin) && OldPin->LinkedTo.Num())
-		{
-			bIsAnyOptionOrReturnConnected = true;
-			break;
-		}
-	}
-	auto NewReturn = GetReturnValuePin();
-	if (!bIsAnyOptionOrReturnConnected && OldReturnPin && NewReturn 
-		&& (NewReturn->PinType.PinCategory == Schema->PC_Wildcard))
+	UEdGraphPin* ReturnPin = GetReturnValuePin();
+	check(ReturnPin);
+
+	if (OldReturnPin && (ReturnPin->PinType.PinCategory == Schema->PC_Wildcard))
 	{
-		NewReturn->PinType = OldReturnPin->PinType;
-	}
-	else if (OldReturnPin && NewReturn && OldReturnPin->LinkedTo.Num())
-	{
-		auto BP = GetBlueprint();
-		UClass* SelfClass = BP ? BP->GeneratedClass : nullptr;
-		bool OldTypeIsValid = true;
-		for (auto OutPin : OldReturnPin->LinkedTo)
-		{
-			if (OutPin && !Schema->ArePinTypesCompatible(OldReturnPin->PinType, OutPin->PinType, SelfClass))
-			{
-				OldTypeIsValid = false;
-				break;
-			}
-		}
-		if (OldTypeIsValid)
-		{
-			NewReturn->PinType = OldReturnPin->PinType;
-		}
+		// Always copy type from node prior, if pins have changed those will error at compilation time
+		ReturnPin->PinType = OldReturnPin->PinType;
 	}
 
 	UEdGraphPin* IndexPin = GetIndexPin();
+	check(IndexPin);
 
 	// If we are fixing up an old bool node (swap the options and copy the condition links)
 	if (OldConditionPin)
 	{
 		// Set the index pin type
 		IndexPinType.PinCategory = Schema->PC_Boolean;
-		IndexPinType.PinSubCategory = TEXT("");
-		IndexPinType.PinSubCategoryObject = NULL;
+		IndexPinType.PinSubCategory.Reset();
+		IndexPinType.PinSubCategoryObject = nullptr;
 
 		// Set the pin type and Copy the pin
 		IndexPin->PinType = IndexPinType;
@@ -540,9 +516,8 @@ void UK2Node_Select::ReallocatePinsDuringReconstruction(TArray<UEdGraphPin*>& Ol
 		UEdGraphPin* OptionPin0 = FindPin("Option 0");
 		UEdGraphPin* OptionPin1 = FindPin("Option 1");
 
-		for (auto It = OldPins.CreateConstIterator(); It; It++)
+		for (UEdGraphPin* OldPin : OldPins)
 		{
-			UEdGraphPin* OldPin = (*It);
 			if (OldPin->PinName == OptionPin0->PinName)
 			{
 				Schema->MovePinLinks(*OldPin, *OptionPin1);
@@ -557,44 +532,65 @@ void UK2Node_Select::ReallocatePinsDuringReconstruction(TArray<UEdGraphPin*>& Ol
 	// If the index pin has links or a default value but is a wildcard, this is an old int pin so convert it
 	if (OldIndexPin &&
 		IndexPinType.PinCategory == Schema->PC_Wildcard &&
-		(OldIndexPin->LinkedTo.Num() > 0 || OldIndexPin->DefaultValue != TEXT("")))
+		(OldIndexPin->LinkedTo.Num() > 0 || !OldIndexPin->DefaultValue.IsEmpty()))
 	{
 		IndexPinType.PinCategory = Schema->PC_Int;
-		IndexPinType.PinSubCategory = TEXT("");
-		IndexPinType.PinSubCategoryObject = NULL;
+		IndexPinType.PinSubCategory.Reset();
+		IndexPinType.PinSubCategoryObject = nullptr;
 		IndexPin->PinType = IndexPinType;
 	}
+
+	// Set up default values for index and option pins now that the information is available
+	Schema->SetPinAutogeneratedDefaultValueBasedOnType(IndexPin);
+
+	const bool bFillTypeFromReturn = ReturnPin->PinType.PinCategory != Schema->PC_Wildcard;
+	TArray<UEdGraphPin*> OptionPins;
+	GetOptionPins(OptionPins);
+	for (UEdGraphPin* Pin : OptionPins)
+	{
+		const bool bTypeShouldBeFilled = Pin && (Pin->PinType.PinCategory == Schema->PC_Wildcard);
+		if (bTypeShouldBeFilled && bFillTypeFromReturn)
+		{
+			Pin->PinType = ReturnPin->PinType;			
+		}
+		Schema->SetPinAutogeneratedDefaultValueBasedOnType(Pin);
+	}
 }
 
 void UK2Node_Select::PostReconstructNode()
 {
+	// After ReconstructNode we must be sure that no additional reconstruction is required
 	bReconstructNode = false;
 
-	const UEdGraphSchema_K2* Schema = Cast<UEdGraphSchema_K2>(GetSchema());
-
 	UEdGraphPin* ReturnPin = GetReturnValuePin();
-	PinConnectionListChanged(ReturnPin);
-	const bool bFillTypeFromReturn = Schema && ReturnPin && (ReturnPin->PinType.PinCategory != Schema->PC_Wildcard);
+	const bool bFillTypeFromConnected = ReturnPin && (ReturnPin->PinType.PinCategory == UEdGraphSchema_K2::PC_Wildcard);
 
-	TArray<UEdGraphPin*> OptionPins;
-	GetOptionPins(OptionPins);
-	for (auto It = OptionPins.CreateConstIterator(); It; It++)
+	if (bFillTypeFromConnected)
 	{
-		UEdGraphPin* Pin = *It;
-		const bool bTypeShouldBeFilled = Schema && Pin && (Pin->PinType.PinCategory == Schema->PC_Wildcard);
-		if (bTypeShouldBeFilled && bFillTypeFromReturn)
+		FEdGraphPinType PinType = ReturnPin->PinType;
+
+		if (ReturnPin->LinkedTo.Num() > 0)
 		{
-			Pin->Modify();
-			Pin->PinType = ReturnPin->PinType;
-			UEdGraphSchema_K2::ValidateExistingConnections(Pin);
+			PinType = ReturnPin->LinkedTo[0]->PinType;
+		}
+		else
+		{
+			TArray<UEdGraphPin*> OptionPins;
+			GetOptionPins(OptionPins);
+			for (UEdGraphPin* Pin : OptionPins)
+			{
+				if (Pin && Pin->LinkedTo.Num() > 0)
+				{
+					PinType = Pin->LinkedTo[0]->PinType;
+					break;
+				}
+			}
 		}
 
-		PinConnectionListChanged(*It);
+		ReturnPin->PinType = PinType;
+		PinTypeChanged(ReturnPin);
 	}
 
-	//After ReconstructNode we must be sure, that no additional reconstruction is required
-	bReconstructNode = false;
-
 	Super::PostReconstructNode();
 }
 
@@ -609,40 +605,15 @@ void UK2Node_Select::NotifyPinConnectionListChanged(UEdGraphPin* Pin)
 	if (Pin == GetIndexPin())
 	{
 		// If the index pin was just linked to another pin
-		if (Pin->LinkedTo.Num() > 0)
+		if (Pin->LinkedTo.Num() > 0 && Pin->PinType.PinCategory == Schema->PC_Wildcard)
 		{
 			UEdGraphPin* LinkPin = Pin->LinkedTo[0];
-			IndexPinType = LinkPin->PinType;
-			Pin->PinType = IndexPinType;
-
-			// See if it was an enum pin
-			if (LinkPin->PinType.PinCategory == Schema->PC_Byte &&
-				LinkPin->PinType.PinSubCategoryObject != NULL &&
-				LinkPin->PinType.PinSubCategoryObject->IsA(UEnum::StaticClass()))
-			{
-				UEnum* EnumPtr = Cast<UEnum>(LinkPin->PinType.PinSubCategoryObject.Get());
-				SetEnum(EnumPtr);
-			}
-			else
-			{
-				SetEnum(NULL);
-			}
-
-			Schema->SetPinDefaultValueBasedOnType(Pin);
 
-			GetGraph()->NotifyGraphChanged();
-			UBlueprint* Blueprint = GetBlueprint();
-			if(!Blueprint->bBeingCompiled)
+			if (Pin->PinType != LinkPin->PinType)
 			{
-				FBlueprintEditorUtils::MarkBlueprintAsModified(Blueprint);
-				Blueprint->BroadcastChanged();
-			}
+				Pin->PinType = LinkPin->PinType;
 
-			// If the index pin is a boolean, we need to remove all but 2 options
-			if (IndexPinType.PinCategory == Schema->PC_Boolean)
-			{
-				NumOptionPins = 2;
-				bReconstructNode = true;
+				PinTypeChanged(Pin);
 			}
 		}
 	}
@@ -662,32 +633,14 @@ void UK2Node_Select::NotifyPinConnectionListChanged(UEdGraphPin* Pin)
 			// If the pin is linked, make sure the other wildcard pins match
 			if (Pin->LinkedTo.Num() > 0)
 			{
-				// Set pin type on the pin
-				Pin->PinType = Pin->LinkedTo[0]->PinType;
+				UEdGraphPin* LinkPin = Pin->LinkedTo[0];
 
-				// Make sure the return pin is the same pin type
-				if (ReturnPin != Pin)
+				if (Pin->PinType != LinkPin->PinType)
 				{
-					ReturnPin->Modify();
+					Pin->PinType = LinkPin->PinType;
 
-					ReturnPin->PinType = Pin->PinType;
-					UEdGraphSchema_K2::ValidateExistingConnections(ReturnPin);
-				}
-
-				// Make sure all options are of the same pin type
-				for (auto It = OptionPins.CreateConstIterator(); It; It++)
-				{
-					UEdGraphPin* OptionPin = (*It);
-					if (*It && *It != Pin)
-					{
-						(*It)->Modify();
-
-						(*It)->PinType = Pin->PinType;
-						UEdGraphSchema_K2::ValidateExistingConnections(*It);
-					}
+					PinTypeChanged(Pin);
 				}
-
-				bReconstructNode = true;
 			}
 		}
 	}
@@ -722,13 +675,12 @@ void UK2Node_Select::GetOptionPins(TArray<UEdGraphPin*>& OptionPins) const
 
 	// If the select node is currently dealing with an enum
 	if (IndexPinType.PinCategory == K2Schema->PC_Byte &&
-		IndexPinType.PinSubCategory == TEXT("") &&
-		IndexPinType.PinSubCategoryObject != NULL &&
+		IndexPinType.PinSubCategory.IsEmpty() &&
+		IndexPinType.PinSubCategoryObject != nullptr &&
 		IndexPinType.PinSubCategoryObject->IsA(UEnum::StaticClass()))
 	{
-		for (auto It = Pins.CreateConstIterator(); It; It++)
+		for (UEdGraphPin* Pin : Pins)
 		{
-			UEdGraphPin* Pin = (*It);
 			if (EnumEntries.Contains(FName(*Pin->PinName)))
 			{
 				OptionPins.Add(Pin);
@@ -737,9 +689,8 @@ void UK2Node_Select::GetOptionPins(TArray<UEdGraphPin*>& OptionPins) const
 	}
 	else
 	{
-		for (auto It = Pins.CreateConstIterator(); It; It++)
+		for (UEdGraphPin* Pin : Pins)
 		{
-			UEdGraphPin* Pin = (*It);
 			if (Pin->PinName.Left(6) == "Option")
 			{
 				OptionPins.Add(Pin);
@@ -887,6 +838,18 @@ bool UK2Node_Select::CanRemoveOptionPinToNode() const
 void UK2Node_Select::ChangePinType(UEdGraphPin* Pin)
 {
 	PinTypeChanged(Pin);
+
+	if (bReconstructNode)
+	{
+		ReconstructNode();
+	}
+
+	UBlueprint* Blueprint = GetBlueprint();
+	if (!Blueprint->bBeingCompiled)
+	{
+		FBlueprintEditorUtils::MarkBlueprintAsModified(Blueprint);
+		Blueprint->BroadcastChanged();
+	}
 }
 
 bool UK2Node_Select::CanChangePinType(UEdGraphPin* Pin) const
@@ -910,9 +873,8 @@ bool UK2Node_Select::CanChangePinType(UEdGraphPin* Pin) const
 		{
 			TArray<UEdGraphPin*> OptionPins;
 			GetOptionPins(OptionPins);
-			for (auto It = OptionPins.CreateConstIterator(); It; It++)
+			for (UEdGraphPin* OptionPin : OptionPins)
 			{
-				UEdGraphPin* OptionPin = (*It);
 				if (OptionPin && OptionPin->LinkedTo.Num() > 0)
 				{
 					return false;
@@ -946,27 +908,27 @@ void UK2Node_Select::PinTypeChanged(UEdGraphPin* Pin)
 			if (IndexPinType.PinCategory == Schema->PC_Boolean)
 			{
 				NumOptionPins = 2;
-				bReconstructNode = true;
 			}
 
-			// Reset the default value
-			Schema->SetPinDefaultValueBasedOnType(Pin);
+			if (!Schema->IsPinDefaultValid(Pin, Pin->DefaultValue, Pin->DefaultObject, Pin->DefaultTextValue).IsEmpty())
+			{
+				Schema->ResetPinToAutogeneratedDefaultValue(Pin);
+			}
+
+			bReconstructNode = true;
 		}
 	}
 	else
 	{
 		// Set the return value
 		UEdGraphPin* ReturnPin = GetReturnValuePin();
-		if (ReturnPin->PinType != Pin->PinType)
+	
+		// Recombine the sub pins back into the ReturnPin
+		if (ReturnPin->SubPins.Num() > 0)
 		{
-			// Recombine the sub pins back into the ReturnPin
-			if (ReturnPin->SubPins.Num() > 0)
-			{
-				Schema->RecombinePin(ReturnPin->SubPins[0]);
-			}
-			ReturnPin->PinType = Pin->PinType;
-			Schema->SetPinDefaultValueBasedOnType(ReturnPin);
+			Schema->RecombinePin(ReturnPin->SubPins[0]);
 		}
+		ReturnPin->PinType = Pin->PinType;
 
 		// Recombine all option pins back into their root
 		TArray<UEdGraphPin*> OptionPins;
@@ -982,33 +944,21 @@ void UK2Node_Select::PinTypeChanged(UEdGraphPin* Pin)
 
 		// Get the options again and set them
 		GetOptionPins(OptionPins);
-		for (auto It = OptionPins.CreateConstIterator(); It; It++)
+		for (UEdGraphPin* OptionPin : OptionPins)
 		{
-			UEdGraphPin* OptionPin = (*It);
 			if (OptionPin->PinType != Pin->PinType ||
 				OptionPin == Pin)
 			{
 				OptionPin->PinType = Pin->PinType;
-				Schema->SetPinDefaultValueBasedOnType(OptionPin);
 			}
-		}
-	}
 
+			if (!Schema->IsPinDefaultValid(OptionPin, OptionPin->DefaultValue, OptionPin->DefaultObject, OptionPin->DefaultTextValue).IsEmpty())
+			{
+				Schema->ResetPinToAutogeneratedDefaultValue(OptionPin);
+			}
+		}
 
-	// Reconstruct the node since the options could change
-	if (bReconstructNode)
-	{
-		ReconstructNode();
-	}
-
-	// Let the graph know to refresh
-	GetGraph()->NotifyGraphChanged();
-
-	UBlueprint* Blueprint = GetBlueprint();
-	if(!Blueprint->bBeingCompiled)
-	{
-		FBlueprintEditorUtils::MarkBlueprintAsModified(Blueprint);
-		Blueprint->BroadcastChanged();
+		bReconstructNode = true;
 	}
 }
 
@@ -1091,8 +1041,8 @@ void UK2Node_Select::ExpandNode(class FKismetCompilerContext& CompilerContext, U
 	}
 
 	bool bSuccess = false;
-	auto Schema = CompilerContext.GetSchema();
-	for (auto Pin : Pins)
+	const UEdGraphSchema_K2* Schema = CompilerContext.GetSchema();
+	for (UEdGraphPin* Pin : Pins)
 	{
 		const bool bValidAutoRefPin = Pin && !Schema->IsMetaPin(*Pin) && (Pin->Direction == EGPD_Input) && (!Pin->LinkedTo.Num() || (GetIndexPin() == Pin));
 		if (!bValidAutoRefPin)
@@ -1100,19 +1050,19 @@ void UK2Node_Select::ExpandNode(class FKismetCompilerContext& CompilerContext, U
 			continue;
 		}
 
-		//default values can be reset when the pin is connected
-		const auto DefaultValue = Pin->DefaultValue;
-		const auto DefaultObject = Pin->DefaultObject;
-		const auto DefaultTextValue = Pin->DefaultTextValue;
-		const auto AutogeneratedDefaultValue = Pin->AutogeneratedDefaultValue;
+		// copy defaults as default values can be reset when the pin is connected
+		const FString DefaultValue = Pin->DefaultValue;
+		UObject* DefaultObject = Pin->DefaultObject;
+		const FText DefaultTextValue = Pin->DefaultTextValue;
+		bool bMatchesDefaults = Pin->DoesDefaultValueMatchAutogenerated();
 
-		auto ValuePin = UK2Node_CallFunction::InnerHandleAutoCreateRef(this, Pin, CompilerContext, SourceGraph, true);
+		UEdGraphPin* ValuePin = UK2Node_CallFunction::InnerHandleAutoCreateRef(this, Pin, CompilerContext, SourceGraph, true);
 		if (ValuePin)
 		{
-			if (!DefaultObject && DefaultTextValue.IsEmpty() && DefaultValue.Equals(AutogeneratedDefaultValue, ESearchCase::CaseSensitive))
+			if (bMatchesDefaults)
 			{
 				// Use the latest code to set default value
-				Schema->SetPinDefaultValueBasedOnType(ValuePin);
+				Schema->SetPinAutogeneratedDefaultValueBasedOnType(ValuePin);
 			}
 			else
 			{
diff --git a/Engine/Source/Editor/BlueprintGraph/Private/K2Node_Self.cpp b/Engine/Source/Editor/BlueprintGraph/Private/K2Node_Self.cpp
--- a/Engine/Source/Editor/BlueprintGraph/Private/K2Node_Self.cpp
+++ b/Engine/Source/Editor/BlueprintGraph/Private/K2Node_Self.cpp
@@ -43,7 +43,7 @@ UK2Node_Self::UK2Node_Self(const FObjectInitializer& ObjectInitializer)
 void UK2Node_Self::AllocateDefaultPins()
 {
 	const UEdGraphSchema_K2* K2Schema = GetDefault<UEdGraphSchema_K2>();
-	CreatePin(EGPD_Output, K2Schema->PC_Object, K2Schema->PSC_Self, NULL, false, false, K2Schema->PN_Self);
+	CreatePin(EGPD_Output, K2Schema->PC_Object, K2Schema->PSC_Self, nullptr, K2Schema->PN_Self);
 
 	Super::AllocateDefaultPins();
 }
diff --git a/Engine/Source/Editor/BlueprintGraph/Private/K2Node_SetFieldsInStruct.cpp b/Engine/Source/Editor/BlueprintGraph/Private/K2Node_SetFieldsInStruct.cpp
--- a/Engine/Source/Editor/BlueprintGraph/Private/K2Node_SetFieldsInStruct.cpp
+++ b/Engine/Source/Editor/BlueprintGraph/Private/K2Node_SetFieldsInStruct.cpp
@@ -126,12 +126,12 @@ void UK2Node_SetFieldsInStruct::AllocateDefaultPins()
 	const UEdGraphSchema_K2* Schema = GetDefault<UEdGraphSchema_K2>();
 	if (Schema && StructType)
 	{
-		CreatePin(EGPD_Input, Schema->PC_Exec, TEXT(""), NULL, false, false, Schema->PN_Execute);
-		CreatePin(EGPD_Output, Schema->PC_Exec, TEXT(""), NULL, false, false, Schema->PN_Then);
+		CreatePin(EGPD_Input, Schema->PC_Exec, FString(), nullptr, Schema->PN_Execute);
+		CreatePin(EGPD_Output, Schema->PC_Exec, FString(), nullptr, Schema->PN_Then);
 
-		UEdGraphPin* InPin = CreatePin(EGPD_Input, Schema->PC_Struct, TEXT(""), StructType, false, true, SetFieldsInStructHelper::StructRefPinName());
+		UEdGraphPin* InPin = CreatePin(EGPD_Input, Schema->PC_Struct, FString(), StructType, SetFieldsInStructHelper::StructRefPinName(), EPinContainerType::None, true);
 
-		UEdGraphPin* OutPin = CreatePin(EGPD_Output, Schema->PC_Struct, TEXT(""), StructType, false, true, SetFieldsInStructHelper::StructOutPinName());
+		UEdGraphPin* OutPin = CreatePin(EGPD_Output, Schema->PC_Struct, FString(), StructType, SetFieldsInStructHelper::StructOutPinName(), EPinContainerType::None, true);
 
 		// Input pin will forward the ref to the output, if the input value is not a reference connection, a copy is made and modified instead and provided as a reference until the function is called again.
 		InPin->AssignByRefPassThroughConnection(OutPin);
diff --git a/Engine/Source/Editor/BlueprintGraph/Private/K2Node_SetVariableOnPersistentFrame.cpp b/Engine/Source/Editor/BlueprintGraph/Private/K2Node_SetVariableOnPersistentFrame.cpp
--- a/Engine/Source/Editor/BlueprintGraph/Private/K2Node_SetVariableOnPersistentFrame.cpp
+++ b/Engine/Source/Editor/BlueprintGraph/Private/K2Node_SetVariableOnPersistentFrame.cpp
@@ -84,8 +84,8 @@ void UK2Node_SetVariableOnPersistentFrame::AllocateDefaultPins()
 {
 	const UEdGraphSchema_K2* K2Schema = GetDefault<UEdGraphSchema_K2>();
 
-	CreatePin(EGPD_Input, K2Schema->PC_Exec, TEXT(""), NULL, false, false, K2Schema->PN_Execute);
-	CreatePin(EGPD_Output, K2Schema->PC_Exec, TEXT(""), NULL, false, false, K2Schema->PN_Then);
+	CreatePin(EGPD_Input, K2Schema->PC_Exec, FString(), nullptr, K2Schema->PN_Execute);
+	CreatePin(EGPD_Output, K2Schema->PC_Exec, FString(), nullptr, K2Schema->PN_Then);
 
 	Super::AllocateDefaultPins();
 }
diff --git a/Engine/Source/Editor/BlueprintGraph/Private/K2Node_SpawnActor.cpp b/Engine/Source/Editor/BlueprintGraph/Private/K2Node_SpawnActor.cpp
--- a/Engine/Source/Editor/BlueprintGraph/Private/K2Node_SpawnActor.cpp
+++ b/Engine/Source/Editor/BlueprintGraph/Private/K2Node_SpawnActor.cpp
@@ -31,30 +31,30 @@ void UK2Node_SpawnActor::AllocateDefaultPins()
 	UEdGraphSchema_K2 const* K2Schema = GetDefault<UEdGraphSchema_K2>();
 
 	// Add execution pins
-	CreatePin(EGPD_Input, K2Schema->PC_Exec, TEXT(""), NULL, false, false, K2Schema->PN_Execute);
-	CreatePin(EGPD_Output, K2Schema->PC_Exec, TEXT(""), NULL, false, false, K2Schema->PN_Then);
+	CreatePin(EGPD_Input, K2Schema->PC_Exec, FString(), nullptr, K2Schema->PN_Execute);
+	CreatePin(EGPD_Output, K2Schema->PC_Exec, FString(), nullptr, K2Schema->PN_Then);
 
 	// If required add the world context pin
 	if (GetBlueprint()->ParentClass->HasMetaDataHierarchical(FBlueprintMetadata::MD_ShowWorldContextPin))
 	{
-		CreatePin(EGPD_Input, K2Schema->PC_Object, TEXT(""), UObject::StaticClass(), false, false, WorldContextPinName);
+		CreatePin(EGPD_Input, K2Schema->PC_Object, FString(), UObject::StaticClass(), WorldContextPinName);
 	}
 
 	// Add blueprint pin
-	UEdGraphPin* BlueprintPin = CreatePin(EGPD_Input, K2Schema->PC_Object, TEXT(""), UBlueprint::StaticClass(), false, false, BlueprintPinName);
+	UEdGraphPin* BlueprintPin = CreatePin(EGPD_Input, K2Schema->PC_Object, FString(), UBlueprint::StaticClass(), BlueprintPinName);
 	K2Schema->ConstructBasicPinTooltip(*BlueprintPin, LOCTEXT("BlueprintPinDescription", "The blueprint Actor you want to spawn"), BlueprintPin->PinToolTip);
 
 	// Transform pin
 	UScriptStruct* TransformStruct = TBaseStructure<FTransform>::Get();
-	UEdGraphPin* TransformPin = CreatePin(EGPD_Input, K2Schema->PC_Struct, TEXT(""), TransformStruct, false, false, SpawnTransformPinName);
+	UEdGraphPin* TransformPin = CreatePin(EGPD_Input, K2Schema->PC_Struct, FString(), TransformStruct, SpawnTransformPinName);
 	K2Schema->ConstructBasicPinTooltip(*TransformPin, LOCTEXT("TransformPinDescription", "The transform to spawn the Actor with"), TransformPin->PinToolTip);
 
 	// bNoCollisionFail pin
-	UEdGraphPin* NoCollisionFailPin = CreatePin(EGPD_Input, K2Schema->PC_Boolean, TEXT(""), NULL, false, false, NoCollisionFailPinName);
+	UEdGraphPin* NoCollisionFailPin = CreatePin(EGPD_Input, K2Schema->PC_Boolean, FString(), nullptr, NoCollisionFailPinName);
 	K2Schema->ConstructBasicPinTooltip(*NoCollisionFailPin, LOCTEXT("NoCollisionFailPinDescription", "Determines if the Actor should be spawned when the location is blocked by a collision"), NoCollisionFailPin->PinToolTip);
 
 	// Result pin
-	UEdGraphPin* ResultPin = CreatePin(EGPD_Output, K2Schema->PC_Object, TEXT(""), AActor::StaticClass(), false, false, K2Schema->PN_ReturnValue);
+	UEdGraphPin* ResultPin = CreatePin(EGPD_Output, K2Schema->PC_Object, FString(), AActor::StaticClass(), K2Schema->PN_ReturnValue);
 	K2Schema->ConstructBasicPinTooltip(*ResultPin, LOCTEXT("ResultPinDescription", "The spawned Actor"), ResultPin->PinToolTip);
 
 	Super::AllocateDefaultPins();
@@ -80,8 +80,8 @@ void UK2Node_SpawnActor::CreatePinsForClass(UClass* InClass)
 			Property->HasAllPropertyFlags(CPF_BlueprintVisible) &&
 			!bIsDelegate )
 		{
-			UEdGraphPin* Pin = CreatePin(EGPD_Input, TEXT(""), TEXT(""), NULL, false, false, Property->GetName());
-			const bool bPinGood = (Pin != NULL) && K2Schema->ConvertPropertyToPinType(Property, /*out*/ Pin->PinType);	
+			UEdGraphPin* Pin = CreatePin(EGPD_Input, FString(), FString(), nullptr, Property->GetName());
+			const bool bPinGood = (Pin != nullptr) && K2Schema->ConvertPropertyToPinType(Property, /*out*/ Pin->PinType);	
 
 			// Copy tooltip from the property.
 			if (Pin != nullptr)
@@ -425,7 +425,7 @@ void UK2Node_SpawnActor::ExpandNode(class FKismetCompilerContext& CompilerContex
 			if(SetByNameFunction)
 			{
 				UK2Node_CallFunction* SetVarNode = NULL;
-				if(SpawnVarPin->PinType.bIsArray)
+				if(SpawnVarPin->PinType.IsArray())
 				{
 					SetVarNode = CompilerContext.SpawnIntermediateNode<UK2Node_CallArrayFunction>(this, SourceGraph);
 				}
@@ -452,7 +452,7 @@ void UK2Node_SpawnActor::ExpandNode(class FKismetCompilerContext& CompilerContex
 				// Move connection from the variable pin on the spawn node to the 'value' pin
 				UEdGraphPin* ValuePin = SetVarNode->FindPinChecked(ValueParamName);
 				CompilerContext.MovePinLinksToIntermediate(*SpawnVarPin, *ValuePin);
-				if(SpawnVarPin->PinType.bIsArray)
+				if(SpawnVarPin->PinType.IsArray())
 				{
 					SetVarNode->PinConnectionListChanged(ValuePin);
 				}
diff --git a/Engine/Source/Editor/BlueprintGraph/Private/K2Node_SpawnActorFromClass.cpp b/Engine/Source/Editor/BlueprintGraph/Private/K2Node_SpawnActorFromClass.cpp
--- a/Engine/Source/Editor/BlueprintGraph/Private/K2Node_SpawnActorFromClass.cpp
+++ b/Engine/Source/Editor/BlueprintGraph/Private/K2Node_SpawnActorFromClass.cpp
@@ -49,28 +49,28 @@ void UK2Node_SpawnActorFromClass::AllocateDefaultPins()
 	const UEdGraphSchema_K2* K2Schema = GetDefault<UEdGraphSchema_K2>();
 
 	// Add execution pins
-	CreatePin(EGPD_Input, K2Schema->PC_Exec, TEXT(""), NULL, false, false, K2Schema->PN_Execute);
-	CreatePin(EGPD_Output, K2Schema->PC_Exec, TEXT(""), NULL, false, false, K2Schema->PN_Then);
+	CreatePin(EGPD_Input, K2Schema->PC_Exec, FString(), nullptr, K2Schema->PN_Execute);
+	CreatePin(EGPD_Output, K2Schema->PC_Exec, FString(), nullptr, K2Schema->PN_Then);
 
 	// If required add the world context pin
 	if (GetBlueprint()->ParentClass->HasMetaDataHierarchical(FBlueprintMetadata::MD_ShowWorldContextPin))
 	{
-		CreatePin(EGPD_Input, K2Schema->PC_Object, TEXT(""), UObject::StaticClass(), false, false, FK2Node_SpawnActorFromClassHelper::WorldContextPinName);
+		CreatePin(EGPD_Input, K2Schema->PC_Object, FString(), UObject::StaticClass(), FK2Node_SpawnActorFromClassHelper::WorldContextPinName);
 	}
 
 	// Add blueprint pin
-	UEdGraphPin* ClassPin = CreatePin(EGPD_Input, K2Schema->PC_Class, TEXT(""), AActor::StaticClass(), false, false, FK2Node_SpawnActorFromClassHelper::ClassPinName);
+	UEdGraphPin* ClassPin = CreatePin(EGPD_Input, K2Schema->PC_Class, FString(), AActor::StaticClass(), FK2Node_SpawnActorFromClassHelper::ClassPinName);
 
 	// Transform pin
 	UScriptStruct* TransformStruct = TBaseStructure<FTransform>::Get();
-	UEdGraphPin* TransformPin = CreatePin(EGPD_Input, K2Schema->PC_Struct, TEXT(""), TransformStruct, false, false, FK2Node_SpawnActorFromClassHelper::SpawnTransformPinName);
+	UEdGraphPin* TransformPin = CreatePin(EGPD_Input, K2Schema->PC_Struct, FString(), TransformStruct, FK2Node_SpawnActorFromClassHelper::SpawnTransformPinName);
 
 	// Collision handling method pin
 	UEnum* const MethodEnum = FindObjectChecked<UEnum>(ANY_PACKAGE, TEXT("ESpawnActorCollisionHandlingMethod"), true);
-	UEdGraphPin* const CollisionHandlingOverridePin = CreatePin(EGPD_Input, K2Schema->PC_Byte, TEXT(""), MethodEnum, false, false, FK2Node_SpawnActorFromClassHelper::CollisionHandlingOverridePinName);
+	UEdGraphPin* const CollisionHandlingOverridePin = CreatePin(EGPD_Input, K2Schema->PC_Byte, FString(), MethodEnum, FK2Node_SpawnActorFromClassHelper::CollisionHandlingOverridePinName);
 	CollisionHandlingOverridePin->DefaultValue = MethodEnum->GetNameStringByValue(static_cast<int>(ESpawnActorCollisionHandlingMethod::Undefined));
 
-	UEdGraphPin* OwnerPin = CreatePin(EGPD_Input, K2Schema->PC_Object, TEXT(""), AActor::StaticClass(),/*bIsArray =*/false, /*bIsReference =*/false, FK2Node_SpawnActorFromClassHelper::OwnerPinName);
+	UEdGraphPin* OwnerPin = CreatePin(EGPD_Input, K2Schema->PC_Object, FString(), AActor::StaticClass(), FK2Node_SpawnActorFromClassHelper::OwnerPinName);
 	OwnerPin->bAdvancedView = true;
 	if (ENodeAdvancedPins::NoPins == AdvancedPinDisplay)
 	{
@@ -78,7 +78,7 @@ void UK2Node_SpawnActorFromClass::AllocateDefaultPins()
 	}
 
 	// Result pin
-	UEdGraphPin* ResultPin = CreatePin(EGPD_Output, K2Schema->PC_Object, TEXT(""), AActor::StaticClass(), false, false, K2Schema->PN_ReturnValue);
+	UEdGraphPin* ResultPin = CreatePin(EGPD_Output, K2Schema->PC_Object, FString(), AActor::StaticClass(), K2Schema->PN_ReturnValue);
 
 	Super::AllocateDefaultPins();
 }
@@ -106,16 +106,16 @@ void UK2Node_SpawnActorFromClass::CreatePinsForClass(UClass* InClass, TArray<UEd
 			!bIsDelegate &&
 			(NULL == FindPin(Property->GetName()) ) )
 		{
-			UEdGraphPin* Pin = CreatePin(EGPD_Input, TEXT(""), TEXT(""), NULL, false, false, Property->GetName());
-			const bool bPinGood = (Pin != NULL) && K2Schema->ConvertPropertyToPinType(Property, /*out*/ Pin->PinType);
+			UEdGraphPin* Pin = CreatePin(EGPD_Input, FString(), FString(), nullptr, Property->GetName());
+			const bool bPinGood = (Pin != nullptr) && K2Schema->ConvertPropertyToPinType(Property, /*out*/ Pin->PinType);
 			OutClassPins.Add(Pin);
 
-			if (ClassDefaultObject && Pin != NULL && K2Schema->PinDefaultValueIsEditable(*Pin))
+			if (ClassDefaultObject && Pin != nullptr && K2Schema->PinDefaultValueIsEditable(*Pin))
 			{
 				FString DefaultValueAsString;
 				const bool bDefaultValueSet = FBlueprintEditorUtils::PropertyValueToString(Property, reinterpret_cast<const uint8*>(ClassDefaultObject), DefaultValueAsString);
 				check( bDefaultValueSet );
-				K2Schema->TrySetDefaultValue(*Pin, DefaultValueAsString);
+				K2Schema->SetPinAutogeneratedDefaultValue(Pin, DefaultValueAsString);
 			}
 
 			// Copy tooltip from the property.
diff --git a/Engine/Source/Editor/BlueprintGraph/Private/K2Node_StructMemberSet.cpp b/Engine/Source/Editor/BlueprintGraph/Private/K2Node_StructMemberSet.cpp
--- a/Engine/Source/Editor/BlueprintGraph/Private/K2Node_StructMemberSet.cpp
+++ b/Engine/Source/Editor/BlueprintGraph/Private/K2Node_StructMemberSet.cpp
@@ -31,8 +31,8 @@ void UK2Node_StructMemberSet::AllocateDefaultPins()
 	const UEdGraphSchema_K2* Schema = GetDefault<UEdGraphSchema_K2>();
 
 	// Add the execution sequencing pin
-	CreatePin(EGPD_Input, Schema->PC_Exec, TEXT(""), NULL, false, false, Schema->PN_Execute);
-	CreatePin(EGPD_Output, Schema->PC_Exec, TEXT(""), NULL, false, false, Schema->PN_Then);
+	CreatePin(EGPD_Input, Schema->PC_Exec, FString(), nullptr, Schema->PN_Execute);
+	CreatePin(EGPD_Output, Schema->PC_Exec, FString(), nullptr, Schema->PN_Then);
 
 	// Display any currently visible optional pins
 	{
diff --git a/Engine/Source/Editor/BlueprintGraph/Private/K2Node_StructOperation.cpp b/Engine/Source/Editor/BlueprintGraph/Private/K2Node_StructOperation.cpp
--- a/Engine/Source/Editor/BlueprintGraph/Private/K2Node_StructOperation.cpp
+++ b/Engine/Source/Editor/BlueprintGraph/Private/K2Node_StructOperation.cpp
@@ -14,6 +14,13 @@ UK2Node_StructOperation::UK2Node_StructOperation(const FObjectInitializer& Objec
 	: Super(ObjectInitializer)
 {
 }
+
+void UK2Node_StructOperation::ValidateNodeDuringCompilation(class FCompilerResultsLog& MessageLog) const
+{
+	// Skip UK2Node_Variable's validation because it doesn't need a property (see CL# 1756451)
+	Super::Super::ValidateNodeDuringCompilation(MessageLog);
+}
+
 bool UK2Node_StructOperation::HasExternalDependencies(TArray<class UStruct*>* OptionalOutput) const
 {
 	const bool bResult = nullptr != StructType;
@@ -32,10 +39,10 @@ void UK2Node_StructOperation::FStructOperationOptionalPinManager::CustomizePinDa
 
 	if (Pin && Property)
 	{
-		const auto UDStructure = Cast<const UUserDefinedStruct>(Property->GetOwnerStruct());
+		const UUserDefinedStruct* UDStructure = Cast<const UUserDefinedStruct>(Property->GetOwnerStruct());
 		if (UDStructure)
 		{
-			const auto VarDesc = FStructureEditorUtils::GetVarDesc(UDStructure).FindByPredicate(
+			const FStructVariableDescription* VarDesc = FStructureEditorUtils::GetVarDesc(UDStructure).FindByPredicate(
 				FStructureEditorUtils::FFindByNameHelper<FStructVariableDescription>(Property->GetFName()));
 			if (VarDesc)
 			{
diff --git a/Engine/Source/Editor/BlueprintGraph/Private/K2Node_Switch.cpp b/Engine/Source/Editor/BlueprintGraph/Private/K2Node_Switch.cpp
--- a/Engine/Source/Editor/BlueprintGraph/Private/K2Node_Switch.cpp
+++ b/Engine/Source/Editor/BlueprintGraph/Private/K2Node_Switch.cpp
@@ -184,11 +184,11 @@ void UK2Node_Switch::AllocateDefaultPins()
 	// Add default pin
 	if (bHasDefaultPin)
 	{
-		CreatePin(EGPD_Output, K2Schema->PC_Exec,TEXT(""), NULL, false, false, DefaultPinName);
+		CreatePin(EGPD_Output, K2Schema->PC_Exec, FString(), nullptr, DefaultPinName);
 	}
 
 	// Add exec input pin
-	CreatePin(EGPD_Input, K2Schema->PC_Exec, TEXT(""), NULL, false, false, K2Schema->PN_Execute);
+	CreatePin(EGPD_Input, K2Schema->PC_Exec, FString(), nullptr, K2Schema->PN_Execute);
 
 	// Create selection pin based on type
 	CreateSelectionPin();
@@ -237,7 +237,7 @@ void UK2Node_Switch::AddPinToSwitchNode()
 	FString NewPinName = GetUniquePinName();
 	if (NewPinName.Len() > 0)
 	{
-		CreatePin(EGPD_Output, K2Schema->PC_Exec, TEXT(""), NULL, false, false, NewPinName);
+		CreatePin(EGPD_Output, K2Schema->PC_Exec, FString(), nullptr, NewPinName);
 	}
 }
 
@@ -247,7 +247,7 @@ void UK2Node_Switch::RemovePinFromSwitchNode(UEdGraphPin* TargetPin)
 	if(bHasDefaultPin && TargetPin == GetDefaultPin())
 	{
 		UProperty* HasDefaultPinProperty = FindField<UProperty>(GetClass(), "bHasDefaultPin");
-		if(HasDefaultPinProperty != NULL)
+		if(HasDefaultPinProperty)
 		{
 			PreEditChange(HasDefaultPinProperty);
 
@@ -266,6 +266,23 @@ void UK2Node_Switch::RemovePinFromSwitchNode(UEdGraphPin* TargetPin)
 	}
 }
 
+bool UK2Node_Switch::CanRemoveExecutionPin(UEdGraphPin* TargetPin) const
+{
+	const UEdGraphSchema_K2* K2Schema = GetDefault<UEdGraphSchema_K2>();
+	// Don't allow removing last pin
+	int32 NumExecPins = 0;
+	for (int32 i = 0; i < Pins.Num(); ++i)
+	{
+		UEdGraphPin* PotentialPin = Pins[i];
+		if (K2Schema->IsExecPin(*PotentialPin) && (PotentialPin->Direction == EGPD_Output))
+		{
+			NumExecPins++;
+		}
+	}
+
+	return NumExecPins > 1;
+}
+
 // Returns the exec output pin name for a given 0-based index
 FString UK2Node_Switch::GetPinNameGivenIndex(int32 Index)
 {
@@ -277,7 +294,7 @@ void UK2Node_Switch::CreateFunctionPin()
 {
 	// Set properties on the function pin
 	const UEdGraphSchema_K2* K2Schema = GetDefault<UEdGraphSchema_K2>();
-	UEdGraphPin* FunctionPin = CreatePin(EGPD_Input, K2Schema->PC_Object, TEXT(""), FunctionClass, false, false, FunctionName.ToString());
+	UEdGraphPin* FunctionPin = CreatePin(EGPD_Input, K2Schema->PC_Object, FString(), FunctionClass, FunctionName.ToString());
 	FunctionPin->bDefaultValueIsReadOnly = true;
 	FunctionPin->bNotConnectable = true;
 	FunctionPin->bHidden = true;
diff --git a/Engine/Source/Editor/BlueprintGraph/Private/K2Node_SwitchEnum.cpp b/Engine/Source/Editor/BlueprintGraph/Private/K2Node_SwitchEnum.cpp
--- a/Engine/Source/Editor/BlueprintGraph/Private/K2Node_SwitchEnum.cpp
+++ b/Engine/Source/Editor/BlueprintGraph/Private/K2Node_SwitchEnum.cpp
@@ -106,8 +106,8 @@ void UK2Node_SwitchEnum::GetMenuActions(FBlueprintActionDatabaseRegistrar& Actio
 void UK2Node_SwitchEnum::CreateSelectionPin()
 {
 	const UEdGraphSchema_K2* K2Schema = GetDefault<UEdGraphSchema_K2>();
-	UEdGraphPin* Pin = CreatePin(EGPD_Input, K2Schema->PC_Byte, TEXT(""), Enum, false, false, TEXT("Selection"));
-	K2Schema->SetPinDefaultValueBasedOnType(Pin);
+	UEdGraphPin* Pin = CreatePin(EGPD_Input, K2Schema->PC_Byte, FString(), Enum, TEXT("Selection"));
+	K2Schema->SetPinAutogeneratedDefaultValueBasedOnType(Pin);
 }
 
 FEdGraphPinType UK2Node_SwitchEnum::GetPinType() const
@@ -149,7 +149,7 @@ void UK2Node_SwitchEnum::CreateCasePins()
 	for (auto EnumIt = EnumEntries.CreateConstIterator(); EnumIt; ++EnumIt)
 	{
 		FName EnumEntry = *EnumIt;
-		UEdGraphPin * NewPin = CreatePin(EGPD_Output, K2Schema->PC_Exec, TEXT(""), NULL, false, false, EnumEntry.ToString());
+		UEdGraphPin * NewPin = CreatePin(EGPD_Output, K2Schema->PC_Exec, FString(), nullptr, EnumEntry.ToString());
 		int32 Index = EnumIt.GetIndex();
 		if (EnumFriendlyNames.IsValidIndex(Index))
 		{
diff --git a/Engine/Source/Editor/BlueprintGraph/Private/K2Node_SwitchInteger.cpp b/Engine/Source/Editor/BlueprintGraph/Private/K2Node_SwitchInteger.cpp
--- a/Engine/Source/Editor/BlueprintGraph/Private/K2Node_SwitchInteger.cpp
+++ b/Engine/Source/Editor/BlueprintGraph/Private/K2Node_SwitchInteger.cpp
@@ -48,8 +48,8 @@ void UK2Node_SwitchInteger::GetMenuActions(FBlueprintActionDatabaseRegistrar& Ac
 void UK2Node_SwitchInteger::CreateSelectionPin()
 {
 	const UEdGraphSchema_K2* K2Schema = GetDefault<UEdGraphSchema_K2>();
-	UEdGraphPin* Pin = CreatePin(EGPD_Input, K2Schema->PC_Int, TEXT(""), NULL, false, false, TEXT("Selection"));
-	K2Schema->SetPinDefaultValueBasedOnType(Pin);
+	UEdGraphPin* Pin = CreatePin(EGPD_Input, K2Schema->PC_Int, FString(), nullptr, TEXT("Selection"));
+	K2Schema->SetPinAutogeneratedDefaultValueBasedOnType(Pin);
 }
 
 FEdGraphPinType UK2Node_SwitchInteger::GetPinType() const 
@@ -110,12 +110,52 @@ void UK2Node_SwitchInteger::ReallocatePinsDuringReconstruction(TArray<UEdGraphPi
 				TestPin->PinName = NewPinName;
 
 				// Create the new output pin to match
-				CreatePin(EGPD_Output, K2Schema->PC_Exec, TEXT(""), NULL, false, false, NewPinName);
+				CreatePin(EGPD_Output, K2Schema->PC_Exec, FString(), nullptr, NewPinName);
 			}
 		}
 	}
 }
 
+bool UK2Node_SwitchInteger::CanRemoveExecutionPin(UEdGraphPin* TargetPin) const
+{
+	if (!Super::CanRemoveExecutionPin(TargetPin))
+	{
+		return false;
+	}
+
+	const UEdGraphSchema_K2* K2Schema = GetDefault<UEdGraphSchema_K2>();
+	int32 PinIndex = (StartIndex >= 0) ? StartIndex : 0;
+	int32 FoundIndex = INDEX_NONE;
+
+	UEdGraphPin* DefaultPin = GetDefaultPin();
+
+	if (TargetPin == DefaultPin)
+	{
+		return true;
+	}
+
+	for (int32 i = 0; i < Pins.Num(); ++i)
+	{
+		UEdGraphPin* PotentialPin = Pins[i];
+		if (K2Schema->IsExecPin(*PotentialPin) && (PotentialPin->Direction == EGPD_Output) && (PotentialPin != DefaultPin))
+		{
+			if (PotentialPin == TargetPin)
+			{
+				FoundIndex = PinIndex;
+			}
+
+			++PinIndex;
+		}
+	}
+
+	// Only allow removing last pin
+	if (FoundIndex != PinIndex - 1)
+	{
+		return false;
+	}
+	return true;
+}
+
 void UK2Node_SwitchInteger::RemovePin(UEdGraphPin* TargetPin)
 {
 	const UEdGraphSchema_K2* K2Schema = GetDefault<UEdGraphSchema_K2>();
diff --git a/Engine/Source/Editor/BlueprintGraph/Private/K2Node_SwitchName.cpp b/Engine/Source/Editor/BlueprintGraph/Private/K2Node_SwitchName.cpp
--- a/Engine/Source/Editor/BlueprintGraph/Private/K2Node_SwitchName.cpp
+++ b/Engine/Source/Editor/BlueprintGraph/Private/K2Node_SwitchName.cpp
@@ -63,8 +63,8 @@ void UK2Node_SwitchName::GetMenuActions(FBlueprintActionDatabaseRegistrar& Actio
 void UK2Node_SwitchName::CreateSelectionPin()
 {
 	const UEdGraphSchema_K2* K2Schema = GetDefault<UEdGraphSchema_K2>();
-	UEdGraphPin* Pin = CreatePin(EGPD_Input, K2Schema->PC_Name, TEXT(""), NULL, false, false, TEXT("Selection"));
-	K2Schema->SetPinDefaultValueBasedOnType(Pin);
+	UEdGraphPin* Pin = CreatePin(EGPD_Input, K2Schema->PC_Name, FString(), nullptr, TEXT("Selection"));
+	K2Schema->SetPinAutogeneratedDefaultValueBasedOnType(Pin);
 }
 
 FEdGraphPinType UK2Node_SwitchName::GetPinType() const 
@@ -83,9 +83,9 @@ FString UK2Node_SwitchName::GetPinNameGivenIndex(int32 Index)
 void UK2Node_SwitchName::CreateCasePins()
 {
 	const UEdGraphSchema_K2* K2Schema = GetDefault<UEdGraphSchema_K2>();
-	for( TArray<FName>::TIterator it(PinNames); it; ++it )
+	for ( FName PinName : PinNames)
 	{
-		CreatePin(EGPD_Output, K2Schema->PC_Exec, TEXT(""), NULL, false, false, (*it).ToString());
+		CreatePin(EGPD_Output, K2Schema->PC_Exec, FString(), nullptr, PinName.ToString());
 	}
 }
 
@@ -110,7 +110,7 @@ void UK2Node_SwitchName::AddPinToSwitchNode()
 	PinNames.Add(FName(*PinName));
 
 	const UEdGraphSchema_K2* K2Schema = GetDefault<UEdGraphSchema_K2>();
-	CreatePin(EGPD_Output, K2Schema->PC_Exec, TEXT(""), NULL, false, false, PinName);
+	CreatePin(EGPD_Output, K2Schema->PC_Exec, FString(), nullptr, PinName);
 }
 
 void UK2Node_SwitchName::RemovePin(UEdGraphPin* TargetPin)
diff --git a/Engine/Source/Editor/BlueprintGraph/Private/K2Node_SwitchString.cpp b/Engine/Source/Editor/BlueprintGraph/Private/K2Node_SwitchString.cpp
--- a/Engine/Source/Editor/BlueprintGraph/Private/K2Node_SwitchString.cpp
+++ b/Engine/Source/Editor/BlueprintGraph/Private/K2Node_SwitchString.cpp
@@ -5,6 +5,7 @@
 #include "Kismet/KismetStringLibrary.h"
 #include "BlueprintNodeSpawner.h"
 #include "BlueprintActionDatabaseRegistrar.h"
+#include "Kismet2/CompilerResultsLog.h"
 
 UK2Node_SwitchString::UK2Node_SwitchString(const FObjectInitializer& ObjectInitializer)
 	: Super(ObjectInitializer)
@@ -79,8 +80,8 @@ void UK2Node_SwitchString::GetMenuActions(FBlueprintActionDatabaseRegistrar& Act
 void UK2Node_SwitchString::CreateSelectionPin()
 {
 	const UEdGraphSchema_K2* K2Schema = GetDefault<UEdGraphSchema_K2>();
-	UEdGraphPin* Pin = CreatePin(EGPD_Input, K2Schema->PC_String, TEXT(""), NULL, false, false, TEXT("Selection"));
-	K2Schema->SetPinDefaultValueBasedOnType(Pin);
+	UEdGraphPin* Pin = CreatePin(EGPD_Input, K2Schema->PC_String, FString(), nullptr, TEXT("Selection"));
+	K2Schema->SetPinAutogeneratedDefaultValueBasedOnType(Pin);
 }
 
 FEdGraphPinType UK2Node_SwitchString::GetPinType() const 
@@ -100,9 +101,9 @@ FString UK2Node_SwitchString::GetPinNameGivenIndex(int32 Index)
 void UK2Node_SwitchString::CreateCasePins()
 {
 	const UEdGraphSchema_K2* K2Schema = GetDefault<UEdGraphSchema_K2>();
-	for( TArray<FString>::TIterator it(PinNames); it; ++it )
+	for( const FString& PinName : PinNames)
 	{
-		CreatePin(EGPD_Output, K2Schema->PC_Exec, TEXT(""), NULL, false, false, *it);
+		CreatePin(EGPD_Output, K2Schema->PC_Exec, FString(), nullptr, PinName);
 	}
 }
 
@@ -127,7 +128,7 @@ void UK2Node_SwitchString::AddPinToSwitchNode()
 	PinNames.Add(PinName);
 
 	const UEdGraphSchema_K2* K2Schema = GetDefault<UEdGraphSchema_K2>();
-	CreatePin(EGPD_Output, K2Schema->PC_Exec, TEXT(""), NULL, false, false, PinName);
+	CreatePin(EGPD_Output, K2Schema->PC_Exec, FString(), nullptr, PinName);
 }
 
 void UK2Node_SwitchString::RemovePin(UEdGraphPin* TargetPin)
@@ -137,3 +138,15 @@ void UK2Node_SwitchString::RemovePin(UEdGraphPin* TargetPin)
 	// Clean-up pin name array
 	PinNames.Remove(TargetPin->PinName);
 }
+
+void UK2Node_SwitchString::ValidateNodeDuringCompilation(FCompilerResultsLog& MessageLog) const
+{
+	Super::ValidateNodeDuringCompilation(MessageLog);
+
+	TSet<FString> UniquePinNames(PinNames);
+
+	if (UniquePinNames.Num() != PinNames.Num())
+	{
+		MessageLog.Error(*NSLOCTEXT("K2Node", "SwitchString_DuplicateCases", "@@ contains duplicate cases.").ToString(), this);
+	}
+}
\ No newline at end of file
