Copying //UE4/Dev-Framework to //UE4/Dev-Main (Source: //UE4/Dev-Framework @ 3431384)

diff --git a/Engine/Source/Editor/BlueprintGraph/Private/K2Node_TemporaryVariable.cpp b/Engine/Source/Editor/BlueprintGraph/Private/K2Node_TemporaryVariable.cpp
--- a/Engine/Source/Editor/BlueprintGraph/Private/K2Node_TemporaryVariable.cpp
+++ b/Engine/Source/Editor/BlueprintGraph/Private/K2Node_TemporaryVariable.cpp
@@ -43,7 +43,7 @@ void UK2Node_TemporaryVariable::AllocateDefaultPins()
 {
 	const UEdGraphSchema_K2* K2Schema = GetDefault<UEdGraphSchema_K2>();
 
-	UEdGraphPin* VariablePin = CreatePin(EGPD_Output, TEXT(""), TEXT(""), NULL, false, false, TEXT("Variable"));
+	UEdGraphPin* VariablePin = CreatePin(EGPD_Output, FString(), FString(), nullptr, TEXT("Variable"));
 	VariablePin->PinType = VariableType;
 
 	Super::AllocateDefaultPins();
@@ -171,39 +171,39 @@ void UK2Node_TemporaryVariable::GetMenuActions(FBlueprintActionDatabaseRegistrar
 
 	UEdGraphSchema_K2 const* K2Schema = GetDefault<UEdGraphSchema_K2>();
 
-	ActionRegistrar.AddBlueprintAction(ActionKey, MakeTempVarNodeSpawner(FEdGraphPinType(K2Schema->PC_Int, TEXT(""), nullptr, /*bIsArray =*/false, /*bIsReference =*/false, /*bIsSet=*/false, /*bIsMap=*/false, /*InValueTerminalType=*/FEdGraphTerminalType()), /*bIsPersistent =*/false));
-	ActionRegistrar.AddBlueprintAction(ActionKey, MakeTempVarNodeSpawner(FEdGraphPinType(K2Schema->PC_Int, TEXT(""), nullptr, /*bIsArray =*/ true, /*bIsReference =*/false, /*bIsSet=*/false, /*bIsMap=*/false, /*InValueTerminalType=*/FEdGraphTerminalType()), /*bIsPersistent =*/false));
-	ActionRegistrar.AddBlueprintAction(ActionKey, MakeTempVarNodeSpawner(FEdGraphPinType(K2Schema->PC_Float, TEXT(""), nullptr, /*bIsArray =*/false, /*bIsReference =*/false, /*bIsSet=*/false, /*bIsMap=*/false, /*InValueTerminalType=*/FEdGraphTerminalType()), /*bIsPersistent =*/false));
-	ActionRegistrar.AddBlueprintAction(ActionKey, MakeTempVarNodeSpawner(FEdGraphPinType(K2Schema->PC_Float, TEXT(""), nullptr, /*bIsArray =*/ true, /*bIsReference =*/false, /*bIsSet=*/false, /*bIsMap=*/false, /*InValueTerminalType=*/FEdGraphTerminalType()), /*bIsPersistent =*/false));
-	ActionRegistrar.AddBlueprintAction(ActionKey, MakeTempVarNodeSpawner(FEdGraphPinType(K2Schema->PC_Boolean, TEXT(""), nullptr, /*bIsArray =*/false, /*bIsReference =*/false, /*bIsSet=*/false, /*bIsMap=*/false, /*InValueTerminalType=*/FEdGraphTerminalType()), /*bIsPersistent =*/false));
-	ActionRegistrar.AddBlueprintAction(ActionKey, MakeTempVarNodeSpawner(FEdGraphPinType(K2Schema->PC_Boolean, TEXT(""), nullptr, /*bIsArray =*/ true, /*bIsReference =*/false, /*bIsSet=*/false, /*bIsMap=*/false, /*InValueTerminalType=*/FEdGraphTerminalType()), /*bIsPersistent =*/false));
-	ActionRegistrar.AddBlueprintAction(ActionKey, MakeTempVarNodeSpawner(FEdGraphPinType(K2Schema->PC_String, TEXT(""), nullptr, /*bIsArray =*/false, /*bIsReference =*/false, /*bIsSet=*/false, /*bIsMap=*/false, /*InValueTerminalType=*/FEdGraphTerminalType()), /*bIsPersistent =*/false));
-	ActionRegistrar.AddBlueprintAction(ActionKey, MakeTempVarNodeSpawner(FEdGraphPinType(K2Schema->PC_String, TEXT(""), nullptr, /*bIsArray =*/ true, /*bIsReference =*/false, /*bIsSet=*/false, /*bIsMap=*/false, /*InValueTerminalType=*/FEdGraphTerminalType()), /*bIsPersistent =*/false));
-	ActionRegistrar.AddBlueprintAction(ActionKey, MakeTempVarNodeSpawner(FEdGraphPinType(K2Schema->PC_Text, TEXT(""), nullptr, /*bIsArray =*/false, /*bIsReference =*/false, /*bIsSet=*/false, /*bIsMap=*/false, /*InValueTerminalType=*/FEdGraphTerminalType()), /*bIsPersistent =*/false));
-	ActionRegistrar.AddBlueprintAction(ActionKey, MakeTempVarNodeSpawner(FEdGraphPinType(K2Schema->PC_Text, TEXT(""), nullptr, /*bIsArray =*/ true, /*bIsReference =*/false, /*bIsSet=*/false, /*bIsMap=*/false, /*InValueTerminalType=*/FEdGraphTerminalType()), /*bIsPersistent =*/false));
+	ActionRegistrar.AddBlueprintAction(ActionKey, MakeTempVarNodeSpawner(FEdGraphPinType(K2Schema->PC_Int, FString(), nullptr, EPinContainerType::None, /*bIsReference =*/ false, /*InValueTerminalType=*/FEdGraphTerminalType()), /*bIsPersistent =*/false));
+	ActionRegistrar.AddBlueprintAction(ActionKey, MakeTempVarNodeSpawner(FEdGraphPinType(K2Schema->PC_Int, FString(), nullptr, EPinContainerType::Array, /*bIsReference =*/ false, /*InValueTerminalType=*/FEdGraphTerminalType()), /*bIsPersistent =*/false));
+	ActionRegistrar.AddBlueprintAction(ActionKey, MakeTempVarNodeSpawner(FEdGraphPinType(K2Schema->PC_Float, FString(), nullptr, EPinContainerType::None, /*bIsReference =*/ false, /*InValueTerminalType=*/FEdGraphTerminalType()), /*bIsPersistent =*/false));
+	ActionRegistrar.AddBlueprintAction(ActionKey, MakeTempVarNodeSpawner(FEdGraphPinType(K2Schema->PC_Float, FString(), nullptr, EPinContainerType::Array, /*bIsReference =*/ false, /*InValueTerminalType=*/FEdGraphTerminalType()), /*bIsPersistent =*/false));
+	ActionRegistrar.AddBlueprintAction(ActionKey, MakeTempVarNodeSpawner(FEdGraphPinType(K2Schema->PC_Boolean, FString(), nullptr, EPinContainerType::None, /*bIsReference =*/ false, /*InValueTerminalType=*/FEdGraphTerminalType()), /*bIsPersistent =*/false));
+	ActionRegistrar.AddBlueprintAction(ActionKey, MakeTempVarNodeSpawner(FEdGraphPinType(K2Schema->PC_Boolean, FString(), nullptr, EPinContainerType::Array, /*bIsReference =*/ false, /*InValueTerminalType=*/FEdGraphTerminalType()), /*bIsPersistent =*/false));
+	ActionRegistrar.AddBlueprintAction(ActionKey, MakeTempVarNodeSpawner(FEdGraphPinType(K2Schema->PC_String, FString(), nullptr, EPinContainerType::None, /*bIsReference =*/ false, /*InValueTerminalType=*/FEdGraphTerminalType()), /*bIsPersistent =*/false));
+	ActionRegistrar.AddBlueprintAction(ActionKey, MakeTempVarNodeSpawner(FEdGraphPinType(K2Schema->PC_String, FString(), nullptr, EPinContainerType::Array, /*bIsReference =*/ false, /*InValueTerminalType=*/FEdGraphTerminalType()), /*bIsPersistent =*/false));
+	ActionRegistrar.AddBlueprintAction(ActionKey, MakeTempVarNodeSpawner(FEdGraphPinType(K2Schema->PC_Text, FString(), nullptr, EPinContainerType::None, /*bIsReference =*/ false, /*InValueTerminalType=*/FEdGraphTerminalType()), /*bIsPersistent =*/false));
+	ActionRegistrar.AddBlueprintAction(ActionKey, MakeTempVarNodeSpawner(FEdGraphPinType(K2Schema->PC_Text, FString(), nullptr, EPinContainerType::Array, /*bIsReference =*/ false, /*InValueTerminalType=*/FEdGraphTerminalType()), /*bIsPersistent =*/false));
 	
-	ActionRegistrar.AddBlueprintAction(ActionKey, MakeTempVarNodeSpawner(FEdGraphPinType(K2Schema->PC_Wildcard, TEXT(""), nullptr, /*bIsArray =*/false, /*bIsReference =*/false, /*bIsSet=*/false, /*bIsMap=*/false, /*InValueTerminalType=*/FEdGraphTerminalType()), /*bIsPersistent =*/false));
-	ActionRegistrar.AddBlueprintAction(ActionKey, MakeTempVarNodeSpawner(FEdGraphPinType(K2Schema->PC_Wildcard, TEXT(""), nullptr, /*bIsArray =*/ true, /*bIsReference =*/false, /*bIsSet=*/false, /*bIsMap=*/false, /*InValueTerminalType=*/FEdGraphTerminalType()), /*bIsPersistent =*/false));
+	ActionRegistrar.AddBlueprintAction(ActionKey, MakeTempVarNodeSpawner(FEdGraphPinType(K2Schema->PC_Wildcard, FString(), nullptr, EPinContainerType::None, /*bIsReference =*/ false, /*InValueTerminalType=*/FEdGraphTerminalType()), /*bIsPersistent =*/false));
+	ActionRegistrar.AddBlueprintAction(ActionKey, MakeTempVarNodeSpawner(FEdGraphPinType(K2Schema->PC_Wildcard, FString(), nullptr, EPinContainerType::Array, /*bIsReference =*/ false, /*InValueTerminalType=*/FEdGraphTerminalType()), /*bIsPersistent =*/false));
 
 	UScriptStruct* VectorStruct  = TBaseStructure<FVector>::Get();
-	ActionRegistrar.AddBlueprintAction(ActionKey, MakeTempVarNodeSpawner(FEdGraphPinType(K2Schema->PC_Struct, TEXT("Vector"), VectorStruct, /*bIsArray =*/false, /*bIsReference =*/false, /*bIsSet=*/false, /*bIsMap=*/false, /*InValueTerminalType=*/FEdGraphTerminalType()), /*bIsPersistent =*/false));
-	ActionRegistrar.AddBlueprintAction(ActionKey, MakeTempVarNodeSpawner(FEdGraphPinType(K2Schema->PC_Struct, TEXT("Vector"), VectorStruct, /*bIsArray =*/ true, /*bIsReference =*/false, /*bIsSet=*/false, /*bIsMap=*/false, /*InValueTerminalType=*/FEdGraphTerminalType()), /*bIsPersistent =*/false));
+	ActionRegistrar.AddBlueprintAction(ActionKey, MakeTempVarNodeSpawner(FEdGraphPinType(K2Schema->PC_Struct, TEXT("Vector"), VectorStruct, EPinContainerType::None, /*bIsReference =*/ false, /*InValueTerminalType=*/FEdGraphTerminalType()), /*bIsPersistent =*/false));
+	ActionRegistrar.AddBlueprintAction(ActionKey, MakeTempVarNodeSpawner(FEdGraphPinType(K2Schema->PC_Struct, TEXT("Vector"), VectorStruct, EPinContainerType::Array, /*bIsReference =*/ false, /*InValueTerminalType=*/FEdGraphTerminalType()), /*bIsPersistent =*/false));
 	
 	UScriptStruct* RotatorStruct = TBaseStructure<FRotator>::Get();
-	ActionRegistrar.AddBlueprintAction(ActionKey, MakeTempVarNodeSpawner(FEdGraphPinType(K2Schema->PC_Struct, TEXT("Rotator"), RotatorStruct, /*bIsArray =*/false, /*bIsReference =*/false, /*bIsSet=*/false, /*bIsMap=*/false, /*InValueTerminalType=*/FEdGraphTerminalType()), /*bIsPersistent =*/false));
-	ActionRegistrar.AddBlueprintAction(ActionKey, MakeTempVarNodeSpawner(FEdGraphPinType(K2Schema->PC_Struct, TEXT("Rotator"), RotatorStruct, /*bIsArray =*/ true, /*bIsReference =*/false, /*bIsSet=*/false, /*bIsMap=*/false, /*InValueTerminalType=*/FEdGraphTerminalType()), /*bIsPersistent =*/false));
+	ActionRegistrar.AddBlueprintAction(ActionKey, MakeTempVarNodeSpawner(FEdGraphPinType(K2Schema->PC_Struct, TEXT("Rotator"), RotatorStruct, EPinContainerType::None, /*bIsReference =*/ false, /*InValueTerminalType=*/FEdGraphTerminalType()), /*bIsPersistent =*/false));
+	ActionRegistrar.AddBlueprintAction(ActionKey, MakeTempVarNodeSpawner(FEdGraphPinType(K2Schema->PC_Struct, TEXT("Rotator"), RotatorStruct, EPinContainerType::Array, /*bIsReference =*/ false, /*InValueTerminalType=*/FEdGraphTerminalType()), /*bIsPersistent =*/false));
 	
 	UScriptStruct* TransformStruct = TBaseStructure<FTransform>::Get();
-	ActionRegistrar.AddBlueprintAction(ActionKey, MakeTempVarNodeSpawner(FEdGraphPinType(K2Schema->PC_Struct, TEXT("Transform"), TransformStruct, /*bIsArray =*/false, /*bIsReference =*/false, /*bIsSet=*/false, /*bIsMap=*/false, /*InValueTerminalType=*/FEdGraphTerminalType()), /*bIsPersistent =*/false));
-	ActionRegistrar.AddBlueprintAction(ActionKey, MakeTempVarNodeSpawner(FEdGraphPinType(K2Schema->PC_Struct, TEXT("Transform"), TransformStruct, /*bIsArray =*/ true, /*bIsReference =*/false, /*bIsSet=*/false, /*bIsMap=*/false, /*InValueTerminalType=*/FEdGraphTerminalType()), /*bIsPersistent =*/false));
+	ActionRegistrar.AddBlueprintAction(ActionKey, MakeTempVarNodeSpawner(FEdGraphPinType(K2Schema->PC_Struct, TEXT("Transform"), TransformStruct, EPinContainerType::None, /*bIsReference =*/ false, /*InValueTerminalType=*/FEdGraphTerminalType()), /*bIsPersistent =*/false));
+	ActionRegistrar.AddBlueprintAction(ActionKey, MakeTempVarNodeSpawner(FEdGraphPinType(K2Schema->PC_Struct, TEXT("Transform"), TransformStruct, EPinContainerType::Array, /*bIsReference =*/ false, /*InValueTerminalType=*/FEdGraphTerminalType()), /*bIsPersistent =*/false));
 
 	UScriptStruct* BlendSampleStruct = FindObjectChecked<UScriptStruct>(ANY_PACKAGE, TEXT("BlendSampleData"));
-	ActionRegistrar.AddBlueprintAction(ActionKey, MakeTempVarNodeSpawner(FEdGraphPinType(K2Schema->PC_Struct, TEXT("BlendSampleData"), BlendSampleStruct, /*bIsArray =*/false, /*bIsReference =*/false, /*bIsSet=*/false, /*bIsMap=*/false, /*InValueTerminalType=*/FEdGraphTerminalType()), /*bIsPersistent =*/false));
-	ActionRegistrar.AddBlueprintAction(ActionKey, MakeTempVarNodeSpawner(FEdGraphPinType(K2Schema->PC_Struct, TEXT("BlendSampleData"), BlendSampleStruct, /*bIsArray =*/ true, /*bIsReference =*/false, /*bIsSet=*/false, /*bIsMap=*/false, /*InValueTerminalType=*/FEdGraphTerminalType()), /*bIsPersistent =*/false));
+	ActionRegistrar.AddBlueprintAction(ActionKey, MakeTempVarNodeSpawner(FEdGraphPinType(K2Schema->PC_Struct, TEXT("BlendSampleData"), BlendSampleStruct, EPinContainerType::None, /*bIsReference =*/ false, /*InValueTerminalType=*/FEdGraphTerminalType()), /*bIsPersistent =*/false));
+	ActionRegistrar.AddBlueprintAction(ActionKey, MakeTempVarNodeSpawner(FEdGraphPinType(K2Schema->PC_Struct, TEXT("BlendSampleData"), BlendSampleStruct, EPinContainerType::Array, /*bIsReference =*/ false, /*InValueTerminalType=*/FEdGraphTerminalType()), /*bIsPersistent =*/false));
 
 	// add persistent bool and int types (for macro graphs)
-	ActionRegistrar.AddBlueprintAction(ActionKey, MakeTempVarNodeSpawner(FEdGraphPinType(K2Schema->PC_Int, TEXT(""), nullptr, /*bIsArray =*/false, /*bIsReference =*/false, /*bIsSet=*/false, /*bIsMap=*/false, /*InValueTerminalType=*/FEdGraphTerminalType()), /*bIsPersistent =*/true));
-	ActionRegistrar.AddBlueprintAction(ActionKey, MakeTempVarNodeSpawner(FEdGraphPinType(K2Schema->PC_Boolean, TEXT(""), nullptr, /*bIsArray =*/false, /*bIsReference =*/false, /*bIsSet=*/false, /*bIsMap=*/false, /*InValueTerminalType=*/FEdGraphTerminalType()), /*bIsPersistent =*/true));
+	ActionRegistrar.AddBlueprintAction(ActionKey, MakeTempVarNodeSpawner(FEdGraphPinType(K2Schema->PC_Int, FString(), nullptr, EPinContainerType::None, /*bIsReference =*/ false, /*InValueTerminalType=*/FEdGraphTerminalType()), /*bIsPersistent =*/true));
+	ActionRegistrar.AddBlueprintAction(ActionKey, MakeTempVarNodeSpawner(FEdGraphPinType(K2Schema->PC_Boolean, FString(), nullptr, EPinContainerType::None, /*bIsReference =*/ false, /*InValueTerminalType=*/FEdGraphTerminalType()), /*bIsPersistent =*/true));
 }
 
 FText UK2Node_TemporaryVariable::GetMenuCategory() const
diff --git a/Engine/Source/Editor/BlueprintGraph/Private/K2Node_Timeline.cpp b/Engine/Source/Editor/BlueprintGraph/Private/K2Node_Timeline.cpp
--- a/Engine/Source/Editor/BlueprintGraph/Private/K2Node_Timeline.cpp
+++ b/Engine/Source/Editor/BlueprintGraph/Private/K2Node_Timeline.cpp
@@ -64,21 +64,21 @@ void UK2Node_Timeline::AllocateDefaultPins()
 	const UEdGraphSchema_K2* K2Schema = GetDefault<UEdGraphSchema_K2>();
 	bCanRenameNode = 1;
 
-	CreatePin(EGPD_Input, K2Schema->PC_Exec, TEXT(""), NULL, false, false, PlayPinName);
-	CreatePin(EGPD_Input, K2Schema->PC_Exec, TEXT(""), NULL, false, false, PlayFromStartPinName);
-	CreatePin(EGPD_Input, K2Schema->PC_Exec, TEXT(""), NULL, false, false, StopPinName);
-	CreatePin(EGPD_Input,  K2Schema->PC_Exec, TEXT(""), NULL, false, false, ReversePinName);
-	CreatePin(EGPD_Input,  K2Schema->PC_Exec, TEXT(""), NULL, false, false, ReverseFromEndPinName);
+	CreatePin(EGPD_Input, K2Schema->PC_Exec, FString(), nullptr, PlayPinName);
+	CreatePin(EGPD_Input, K2Schema->PC_Exec, FString(), nullptr, PlayFromStartPinName);
+	CreatePin(EGPD_Input, K2Schema->PC_Exec, FString(), nullptr, StopPinName);
+	CreatePin(EGPD_Input,  K2Schema->PC_Exec, FString(), nullptr, ReversePinName);
+	CreatePin(EGPD_Input,  K2Schema->PC_Exec, FString(), nullptr, ReverseFromEndPinName);
 
-	CreatePin(EGPD_Output, K2Schema->PC_Exec, TEXT(""), NULL, false, false, UpdatePinName);
-	CreatePin(EGPD_Output, K2Schema->PC_Exec, TEXT(""), NULL, false, false, FinishedPinName);
+	CreatePin(EGPD_Output, K2Schema->PC_Exec, FString(), nullptr, UpdatePinName);
+	CreatePin(EGPD_Output, K2Schema->PC_Exec, FString(), nullptr, FinishedPinName);
 
-	CreatePin(EGPD_Input, K2Schema->PC_Exec, TEXT(""), NULL, false, false, SetNewTimePinName);
+	CreatePin(EGPD_Input, K2Schema->PC_Exec, FString(), nullptr, SetNewTimePinName);
 
-	UEdGraphPin* NewPositionPin = CreatePin(EGPD_Input, K2Schema->PC_Float, TEXT(""), NULL, false, false, NewTimePinName);
-	NewPositionPin->DefaultValue = NewPositionPin->AutogeneratedDefaultValue = TEXT("0.0");
+	UEdGraphPin* NewPositionPin = CreatePin(EGPD_Input, K2Schema->PC_Float, FString(), nullptr, NewTimePinName);
+	K2Schema->SetPinAutogeneratedDefaultValue(NewPositionPin, TEXT("0.0"));
 
-	CreatePin(EGPD_Output, K2Schema->PC_Byte, TEXT(""), FTimeline::GetTimelineDirectionEnum(), false, false, DirectionPinName);
+	CreatePin(EGPD_Output, K2Schema->PC_Byte, FString(), FTimeline::GetTimelineDirectionEnum(), DirectionPinName);
 
 	UBlueprint* Blueprint = GetBlueprint();
 	check(Blueprint);
@@ -92,27 +92,27 @@ void UK2Node_Timeline::AllocateDefaultPins()
 		for(int32 i=0; i<Timeline->FloatTracks.Num(); i++)
 		{
 			FTTFloatTrack& FloatTrack = Timeline->FloatTracks[i];
-			CreatePin(EGPD_Output, K2Schema->PC_Float, TEXT(""), NULL, false, false, FloatTrack.TrackName.ToString());			
+			CreatePin(EGPD_Output, K2Schema->PC_Float, FString(), nullptr, FloatTrack.TrackName.ToString());			
 		}
 
 		for(int32 i=0; i<Timeline->VectorTracks.Num(); i++)
 		{
 			FTTVectorTrack& VectorTrack = Timeline->VectorTracks[i];
 			UScriptStruct* VectorStruct = TBaseStructure<FVector>::Get();
-			CreatePin(EGPD_Output, K2Schema->PC_Struct, TEXT(""), VectorStruct, false, false, VectorTrack.TrackName.ToString());			
+			CreatePin(EGPD_Output, K2Schema->PC_Struct, FString(), VectorStruct, VectorTrack.TrackName.ToString());			
 		}
 
 		for(int32 i=0; i<Timeline->LinearColorTracks.Num(); i++)
 		{
 			FTTLinearColorTrack& LinearColorTrack = Timeline->LinearColorTracks[i];
 			UScriptStruct* LinearColorStruct = TBaseStructure<FLinearColor>::Get();
-			CreatePin(EGPD_Output, K2Schema->PC_Struct, TEXT(""), LinearColorStruct, false, false, LinearColorTrack.TrackName.ToString());			
+			CreatePin(EGPD_Output, K2Schema->PC_Struct, FString(), LinearColorStruct, LinearColorTrack.TrackName.ToString());			
 		}
 
 		for(int32 i=0; i<Timeline->EventTracks.Num(); i++)
 		{
 			FTTEventTrack& EventTrack = Timeline->EventTracks[i];
-			CreatePin(EGPD_Output, K2Schema->PC_Exec, TEXT(""), NULL, false, false, EventTrack.TrackName.ToString());			
+			CreatePin(EGPD_Output, K2Schema->PC_Exec, FString(), nullptr, EventTrack.TrackName.ToString());			
 		}
 
 		// cache play status
diff --git a/Engine/Source/Editor/BlueprintGraph/Private/K2Node_Tunnel.cpp b/Engine/Source/Editor/BlueprintGraph/Private/K2Node_Tunnel.cpp
--- a/Engine/Source/Editor/BlueprintGraph/Private/K2Node_Tunnel.cpp
+++ b/Engine/Source/Editor/BlueprintGraph/Private/K2Node_Tunnel.cpp
@@ -189,8 +189,9 @@ UEdGraphPin* UK2Node_Tunnel::CreatePinFromUserDefinition(const TSharedPtr<FUserP
 		Direction = NewPinInfo->DesiredPinDirection;
 	}
 
+	const UEdGraphSchema_K2* const Schema = GetDefault<UEdGraphSchema_K2>();
 	UEdGraphPin* Result = CreatePin(Direction, NewPinInfo->PinType, NewPinInfo->PinName);
-	Result->DefaultValue = Result->AutogeneratedDefaultValue = NewPinInfo->PinDefaultValue;
+	Schema->SetPinAutogeneratedDefaultValue(Result, NewPinInfo->PinDefaultValue);
 
 	// Make sure it mirrors onto the associated node
 	UEdGraphNode* TargetNode = ((InputSinkNode != NULL) ? InputSinkNode : OutputSourceNode);
@@ -212,6 +213,18 @@ UEdGraphPin* UK2Node_Tunnel::CreatePinFromUserDefinition(const TSharedPtr<FUserP
 	return Result;
 }
 
+bool UK2Node_Tunnel::ModifyUserDefinedPinDefaultValue(TSharedPtr<FUserPinInfo> PinInfo, const FString& NewDefaultValue)
+{
+	if (Super::ModifyUserDefinedPinDefaultValue(PinInfo, NewDefaultValue))
+	{
+		const UEdGraphSchema_K2* K2Schema = GetDefault<UEdGraphSchema_K2>();
+		K2Schema->HandleParameterDefaultValueChanged(this);
+
+		return true;
+	}
+	return false;
+}
+
 bool UK2Node_Tunnel::CanModifyExecutionWires()
 {
 	return true;
diff --git a/Engine/Source/Editor/BlueprintGraph/Private/K2Node_TunnelBoundary.cpp b/Engine/Source/Editor/BlueprintGraph/Private/K2Node_TunnelBoundary.cpp
--- a/Engine/Source/Editor/BlueprintGraph/Private/K2Node_TunnelBoundary.cpp
+++ b/Engine/Source/Editor/BlueprintGraph/Private/K2Node_TunnelBoundary.cpp
@@ -109,7 +109,7 @@ void UK2Node_TunnelBoundary::CreateBoundaryNodesForTunnelInstance(UK2Node_Tunnel
 	UEdGraphNode* TunnelExitNode = nullptr;
 	TArray<UK2Node_Tunnel*> Tunnels;
 	TunnelGraph->GetNodesOfClass<UK2Node_Tunnel>(Tunnels);
-	for (auto TunnelNode : Tunnels)
+	for (UK2Node_Tunnel* TunnelNode : Tunnels)
 	{
 		if (IsPureTunnel(TunnelNode))
 		{
@@ -144,7 +144,7 @@ void UK2Node_TunnelBoundary::CreateBoundaryNodesForTunnelInstance(UK2Node_Tunnel
 		TunnelBoundaryNode->CreateNewGuid();
 		MessageLog.RegisterIntermediateTunnelNode(TunnelBoundaryNode, RegistrantTunnelInstance);
 		TunnelBoundaryNode->WireUpTunnelEntry(TunnelInstance, EntryPin, MessageLog);
-		for (auto LinkedPin : EntryPin->LinkedTo)
+		for (UEdGraphPin* LinkedPin : EntryPin->LinkedTo)
 		{
 			FindTunnelExitSiteInstances(LinkedPin, ExecutionEndpointPins, VisitedPins, TunnelExitNode);
 		}
@@ -158,7 +158,7 @@ void UK2Node_TunnelBoundary::CreateBoundaryNodesForTunnelInstance(UK2Node_Tunnel
 		TunnelBoundaryNode->TunnelBoundaryType = TBT_EndOfThread;
 		MessageLog.RegisterIntermediateTunnelNode(TunnelBoundaryNode, RegistrantTunnelInstance);
 		// Create exit point pin
-		UEdGraphPin* BoundaryTerminalPin = TunnelBoundaryNode->CreatePin(EGPD_Input, UEdGraphSchema_K2::PC_Exec, TEXT(""), NULL, false, false, TEXT("TunnelEndOfThread"));
+		UEdGraphPin* BoundaryTerminalPin = TunnelBoundaryNode->CreatePin(EGPD_Input, UEdGraphSchema_K2::PC_Exec, FString(), nullptr, TEXT("TunnelEndOfThread"));
 		// Wire up the endpoint unwired links to the boundary node.
 		for (UEdGraphPin* TerminationPin : ExecutionEndpointPins)
 		{
@@ -184,7 +184,7 @@ void UK2Node_TunnelBoundary::CreateBoundariesForExpansionNodes(UEdGraphNode* Sou
 	UEdGraph* TargetGraph = nullptr;
 	TMap<UEdGraphPin*, TArray<UEdGraphPin*>> EntryPins;
 	TMap<UEdGraphPin*, TArray<UEdGraphPin*>> ExitPins;
-	for (auto ExpansionNode : ExpansionNodes)
+	for (UEdGraphNode* ExpansionNode : ExpansionNodes)
 	{
 		if (!TargetGraph)
 		{
@@ -220,7 +220,7 @@ void UK2Node_TunnelBoundary::CreateBoundariesForExpansionNodes(UEdGraphNode* Sou
 	{
 		if (EntryPins.Num() > 0)
 		{
-			for (auto PinEntry : EntryPins)
+			for (const auto& PinEntry : EntryPins)
 			{
 				// Create entry node and pins
 				UK2Node_TunnelBoundary* EntryBoundaryNode = TargetGraph->CreateBlankNode<UK2Node_TunnelBoundary>();
@@ -231,7 +231,7 @@ void UK2Node_TunnelBoundary::CreateBoundariesForExpansionNodes(UEdGraphNode* Sou
 		}
 		if (ExitPins.Num() > 0)
 		{
-			for (auto PinExit : ExitPins)
+			for (const auto& PinExit : ExitPins)
 			{
 				// Create exit node and pins
 				UK2Node_TunnelBoundary* ExitBoundaryNode = TargetGraph->CreateBlankNode<UK2Node_TunnelBoundary>();
@@ -259,12 +259,12 @@ void UK2Node_TunnelBoundary::WireUpTunnelEntry(UK2Node_Tunnel* TunnelInstance, U
 		// Find the true source pin
 		UEdGraphPin* SourcePin = TunnelInstance->FindPin(TunnelPin->PinName);
 		// Wire in the node
-		UEdGraphPin* OutputPin = CreatePin(EGPD_Output, UEdGraphSchema_K2::PC_Exec, TEXT(""), NULL, false, false, TEXT("TunnelEntryExec"));
-		for (auto LinkedPin : TunnelPin->LinkedTo)
+		UEdGraphPin* OutputPin = CreatePin(EGPD_Output, UEdGraphSchema_K2::PC_Exec, FString(), nullptr, TEXT("TunnelEntryExec"));
+		for (UEdGraphPin* LinkedPin : TunnelPin->LinkedTo)
 		{
 			LinkedPin->MakeLinkTo(OutputPin);
 		}
-		UEdGraphPin* InputPin = CreatePin(EGPD_Input, UEdGraphSchema_K2::PC_Exec, TEXT(""), NULL, false, false, SourcePin->PinName);
+		UEdGraphPin* InputPin = CreatePin(EGPD_Input, UEdGraphSchema_K2::PC_Exec, FString(), nullptr, SourcePin->PinName);
 		MessageLog.NotifyIntermediatePinCreation(InputPin, SourcePin);
 		TunnelPin->BreakAllPinLinks();
 		TunnelPin->MakeLinkTo(InputPin);
@@ -282,19 +282,19 @@ void UK2Node_TunnelBoundary::WireUpTunnelExit(UK2Node_Tunnel* TunnelInstance, UE
 		// Find the true source pin
 		UEdGraphPin* SourcePin = TunnelInstance->FindPin(TunnelPin->PinName);
 		// Wire in the node
-		UEdGraphPin* InputPin = CreatePin(EGPD_Input, UEdGraphSchema_K2::PC_Exec, TEXT(""), NULL, false, false, SourcePin->PinName);
-		for (auto LinkedPin : TunnelPin->LinkedTo)
+		UEdGraphPin* InputPin = CreatePin(EGPD_Input, UEdGraphSchema_K2::PC_Exec, FString(), nullptr, SourcePin->PinName);
+		for (UEdGraphPin* LinkedPin : TunnelPin->LinkedTo)
 		{
 			LinkedPin->MakeLinkTo(InputPin);
 		}
-		UEdGraphPin* OutputPin = CreatePin(EGPD_Output, UEdGraphSchema_K2::PC_Exec, TEXT(""), NULL, false, false, TEXT("TunnelExitExec"));
+		UEdGraphPin* OutputPin = CreatePin(EGPD_Output, UEdGraphSchema_K2::PC_Exec, FString(), nullptr, TEXT("TunnelExitExec"));
 		MessageLog.NotifyIntermediatePinCreation(OutputPin, SourcePin);
 		TunnelPin->BreakAllPinLinks();
 		TunnelPin->MakeLinkTo(OutputPin);
 	}
 }
 
-void UK2Node_TunnelBoundary::WireUpEntry(UEdGraphNode* SourceNode, UEdGraphPin* SourcePin, TArray<UEdGraphPin*>& EntryPins, FCompilerResultsLog& MessageLog)
+void UK2Node_TunnelBoundary::WireUpEntry(UEdGraphNode* SourceNode, UEdGraphPin* SourcePin, const TArray<UEdGraphPin*>& EntryPins, FCompilerResultsLog& MessageLog)
 {
 	if (SourceNode && SourcePin && EntryPins.Num())
 	{
@@ -303,13 +303,13 @@ void UK2Node_TunnelBoundary::WireUpEntry(UEdGraphNode* SourceNode, UEdGraphPin*
 		// Create the pin base name
 		CreateBaseNodeName(SourceNode);
 		// Wire in the node
-		UEdGraphPin* InputPin = CreatePin(EGPD_Input, SourcePin->PinType.PinCategory, SourcePin->PinType.PinSubCategory, SourcePin->PinType.PinSubCategoryObject.Get(), SourcePin->PinType.bIsArray, SourcePin->PinType.bIsReference, TEXT("EntryBoundary"), SourcePin->PinType.bIsConst, INDEX_NONE, SourcePin->PinType.bIsSet, SourcePin->PinType.bIsMap, SourcePin->PinType.PinValueType);
+		UEdGraphPin* InputPin = CreatePin(EGPD_Input, SourcePin->PinType.PinCategory, SourcePin->PinType.PinSubCategory, SourcePin->PinType.PinSubCategoryObject.Get(), TEXT("EntryBoundary"), SourcePin->PinType.ContainerType, SourcePin->PinType.bIsReference, SourcePin->PinType.bIsConst, INDEX_NONE, SourcePin->PinType.PinValueType);
 		MessageLog.NotifyIntermediatePinCreation(InputPin, SourcePin);
-		UEdGraphPin* OutputPin = CreatePin(EGPD_Output, SourcePin->PinType.PinCategory, SourcePin->PinType.PinSubCategory, SourcePin->PinType.PinSubCategoryObject.Get(), SourcePin->PinType.bIsArray, SourcePin->PinType.bIsReference, SourcePin->PinName, SourcePin->PinType.bIsConst, INDEX_NONE, SourcePin->PinType.bIsSet, SourcePin->PinType.bIsMap, SourcePin->PinType.PinValueType);
+		UEdGraphPin* OutputPin = CreatePin(EGPD_Output, SourcePin->PinType.PinCategory, SourcePin->PinType.PinSubCategory, SourcePin->PinType.PinSubCategoryObject.Get(), SourcePin->PinName, SourcePin->PinType.ContainerType, SourcePin->PinType.bIsReference, SourcePin->PinType.bIsConst, INDEX_NONE, SourcePin->PinType.PinValueType);
 		MessageLog.AddPinTraceFilter(OutputPin);
-		for (auto EntryPin : EntryPins)
+		for (UEdGraphPin* EntryPin : EntryPins)
 		{
-			for (auto LinkedPin : EntryPin->LinkedTo)
+			for (UEdGraphPin* LinkedPin : EntryPin->LinkedTo)
 			{
 				check (LinkedPin->Direction != InputPin->Direction);
 				LinkedPin->MakeLinkTo(InputPin);
@@ -321,7 +321,7 @@ void UK2Node_TunnelBoundary::WireUpEntry(UEdGraphNode* SourceNode, UEdGraphPin*
 	}
 }
 
-void UK2Node_TunnelBoundary::WireUpExit(UEdGraphNode* SourceNode, UEdGraphPin* SourcePin, TArray<UEdGraphPin*>& ExitPins, FCompilerResultsLog& MessageLog)
+void UK2Node_TunnelBoundary::WireUpExit(UEdGraphNode* SourceNode, UEdGraphPin* SourcePin, const TArray<UEdGraphPin*>& ExitPins, FCompilerResultsLog& MessageLog)
 {
 	if (SourceNode && SourcePin && ExitPins.Num())
 	{
@@ -330,13 +330,13 @@ void UK2Node_TunnelBoundary::WireUpExit(UEdGraphNode* SourceNode, UEdGraphPin* S
 		// Create the pin base name
 		CreateBaseNodeName(SourceNode);
 		// Wire in the node
-		UEdGraphPin* OutputPin = CreatePin(EGPD_Output, SourcePin->PinType.PinCategory, SourcePin->PinType.PinSubCategory, SourcePin->PinType.PinSubCategoryObject.Get(), SourcePin->PinType.bIsArray, SourcePin->PinType.bIsReference, SourcePin->PinName, SourcePin->PinType.bIsConst, INDEX_NONE, SourcePin->PinType.bIsSet, SourcePin->PinType.bIsMap, SourcePin->PinType.PinValueType);
+		UEdGraphPin* OutputPin = CreatePin(EGPD_Output, SourcePin->PinType.PinCategory, SourcePin->PinType.PinSubCategory, SourcePin->PinType.PinSubCategoryObject.Get(), SourcePin->PinName, SourcePin->PinType.ContainerType, SourcePin->PinType.bIsReference, SourcePin->PinType.bIsConst, INDEX_NONE, SourcePin->PinType.PinValueType);
 		MessageLog.NotifyIntermediatePinCreation(OutputPin, SourcePin);
-		UEdGraphPin* InputPin = CreatePin(EGPD_Input, SourcePin->PinType.PinCategory, SourcePin->PinType.PinSubCategory, SourcePin->PinType.PinSubCategoryObject.Get(), SourcePin->PinType.bIsArray, SourcePin->PinType.bIsReference, TEXT("ExitBoundary"), SourcePin->PinType.bIsConst, INDEX_NONE, SourcePin->PinType.bIsSet, SourcePin->PinType.bIsMap, SourcePin->PinType.PinValueType);
+		UEdGraphPin* InputPin = CreatePin(EGPD_Input, SourcePin->PinType.PinCategory, SourcePin->PinType.PinSubCategory, SourcePin->PinType.PinSubCategoryObject.Get(), TEXT("ExitBoundary"), SourcePin->PinType.ContainerType, SourcePin->PinType.bIsReference, SourcePin->PinType.bIsConst, INDEX_NONE, SourcePin->PinType.PinValueType);
 		MessageLog.AddPinTraceFilter(InputPin);
-		for (auto ExitPin : ExitPins)
+		for (UEdGraphPin* ExitPin : ExitPins)
 		{
-			for (auto LinkedPin : ExitPin->LinkedTo)
+			for (UEdGraphPin* LinkedPin : ExitPin->LinkedTo)
 			{
 				check (LinkedPin->Direction != OutputPin->Direction);
 				LinkedPin->MakeLinkTo(OutputPin);
@@ -381,7 +381,7 @@ void UK2Node_TunnelBoundary::BuildSourceNodeMap(UEdGraphNode* Tunnel, TMap<FGuid
 	}
 	if (TunnelGraph)
 	{
-		for (auto GraphNode : TunnelGraph->Nodes)
+		for (UEdGraphNode* GraphNode : TunnelGraph->Nodes)
 		{
 			SourceNodeMap.Add(GraphNode->NodeGuid) = GraphNode;
 		}
@@ -410,7 +410,7 @@ UEdGraphNode* UK2Node_TunnelBoundary::FindTrueSourceTunnelInstance(UEdGraphNode*
 			}
 			if (TunnelGraph)
 			{
-				for (auto GraphNode : TunnelGraph->Nodes)
+				for (UEdGraphNode* GraphNode : TunnelGraph->Nodes)
 				{
 					if (GraphNode->NodeGuid == Tunnel->NodeGuid)
 					{
@@ -445,7 +445,7 @@ void UK2Node_TunnelBoundary::FindTunnelExitSiteInstances(UEdGraphPin* NodeEntryP
 		// Find all exec out pins
 		TArray<UEdGraphPin*> ExecPins;
 		TArray<UEdGraphPin*> ConnectedExecPins;
-		for (auto Pin : PinNode->Pins)
+		for (UEdGraphPin* Pin : PinNode->Pins)
 		{
 			if (Pin->Direction == EGPD_Output && Pin->PinType.PinCategory == UEdGraphSchema_K2::PC_Exec)
 			{
@@ -458,11 +458,10 @@ void UK2Node_TunnelBoundary::FindTunnelExitSiteInstances(UEdGraphPin* NodeEntryP
 		}
 		if (ConnectedExecPins.Num() > 1)
 		{
-//			if (FEdGraphUtilities::IsExecutionSequence(PinNode))
 			if (PinNode->IsA<UK2Node_ExecutionSequence>())
 			{
 				// Execution sequence style nodes, follow the last connected pin
-				for (auto LinkedPin : ConnectedExecPins.Last()->LinkedTo)
+				for (UEdGraphPin* LinkedPin : ConnectedExecPins.Last()->LinkedTo)
 				{
 					FindTunnelExitSiteInstances(LinkedPin, ExitPins, VisitedPins);
 				}
@@ -470,9 +469,9 @@ void UK2Node_TunnelBoundary::FindTunnelExitSiteInstances(UEdGraphPin* NodeEntryP
 			else
 			{
 				// Branch style nodes, more tricky
-				for (auto ExecPin : ConnectedExecPins)
+				for (UEdGraphPin* ExecPin : ConnectedExecPins)
 				{
-					for (auto LinkedPin : ExecPin->LinkedTo)
+					for (UEdGraphPin* LinkedPin : ExecPin->LinkedTo)
 					{
 						FindTunnelExitSiteInstances(LinkedPin, ExitPins, VisitedPins);
 					}
@@ -481,7 +480,7 @@ void UK2Node_TunnelBoundary::FindTunnelExitSiteInstances(UEdGraphPin* NodeEntryP
 		}
 		else if (ConnectedExecPins.Num() == 1)
 		{
-			for (auto LinkedPin : ConnectedExecPins[0]->LinkedTo)
+			for (UEdGraphPin* LinkedPin : ConnectedExecPins[0]->LinkedTo)
 			{
 				FindTunnelExitSiteInstances(LinkedPin, ExitPins, VisitedPins);
 			}
diff --git a/Engine/Source/Editor/BlueprintGraph/Private/K2Node_Variable.cpp b/Engine/Source/Editor/BlueprintGraph/Private/K2Node_Variable.cpp
--- a/Engine/Source/Editor/BlueprintGraph/Private/K2Node_Variable.cpp
+++ b/Engine/Source/Editor/BlueprintGraph/Private/K2Node_Variable.cpp
@@ -85,13 +85,13 @@ bool UK2Node_Variable::CreatePinForVariable(EEdGraphPinDirection Direction, FStr
 		}
 	}
 
-	if (VariableProperty != NULL)
+	if (VariableProperty != nullptr)
 	{
 		const FString PinName = InPinName.IsEmpty()? GetVarNameString() : InPinName;
 		// Create the pin
-		UEdGraphPin* VariablePin = CreatePin(Direction, TEXT(""), TEXT(""), NULL, false, false, PinName);
+		UEdGraphPin* VariablePin = CreatePin(Direction, FString(), FString(), nullptr, PinName);
 		K2Schema->ConvertPropertyToPinType(VariableProperty, /*out*/ VariablePin->PinType);
-		K2Schema->SetPinDefaultValueBasedOnType(VariablePin);
+		K2Schema->SetPinAutogeneratedDefaultValueBasedOnType(VariablePin);
 	}
 	else
 	{
@@ -159,7 +159,7 @@ void UK2Node_Variable::CreatePinForSelf()
 				TargetClass = TargetClass->GetAuthoritativeClass();
 			}
 
-			UEdGraphPin* TargetPin = CreatePin(EGPD_Input, K2Schema->PC_Object, TEXT(""), TargetClass, false, false, K2Schema->PN_Self);
+			UEdGraphPin* TargetPin = CreatePin(EGPD_Input, K2Schema->PC_Object, FString(), TargetClass, K2Schema->PN_Self);
 			TargetPin->PinFriendlyName =  LOCTEXT("Target", "Target");
 
 			if (bSelfTarget)
@@ -194,14 +194,14 @@ bool UK2Node_Variable::RecreatePinForVariable(EEdGraphPinDirection Direction, TA
 		}
 	}
 
-	if(NULL != OldVariablePin)
+	if(nullptr != OldVariablePin)
 	{
 		// create new pin from old one
-		UEdGraphPin* VariablePin = CreatePin(Direction, TEXT(""), TEXT(""), NULL, false, false, PinName);
+		UEdGraphPin* VariablePin = CreatePin(Direction, FString(), FString(), nullptr, PinName);
 		VariablePin->PinType = OldVariablePin->PinType;
 		
 		const UEdGraphSchema_K2* K2Schema = GetDefault<UEdGraphSchema_K2>();
-		K2Schema->SetPinDefaultValueBasedOnType(VariablePin);
+		K2Schema->SetPinAutogeneratedDefaultValueBasedOnType(VariablePin);
 
 		Message_Note(*FString::Printf(TEXT("Pin for variable '%s' recreated, but the variable is missing."), *PinName));
 		return true;
@@ -515,11 +515,20 @@ FSlateIcon UK2Node_Variable::GetVarIconFromPinType(const FEdGraphPinType& InPinT
 	const UEdGraphSchema_K2* K2Schema = GetDefault<UEdGraphSchema_K2>();
 	IconColorOut = K2Schema->GetPinTypeColor(InPinType);
 
-	if(InPinType.bIsArray)
+	if (InPinType.IsArray())
 	{
 		return FSlateIcon("EditorStyle", "Kismet.AllClasses.ArrayVariableIcon");
 	}
-	else if(InPinType.PinSubCategoryObject.IsValid())
+	else if (InPinType.IsSet())
+	{
+		return FSlateIcon("EditorStyle", "Kismet.AllClasses.SetVariableIcon");
+	}
+	else if (InPinType.IsMap())
+	{
+		// TODO: Need to properly deal with Key/Value stuff
+		return FSlateIcon("EditorStyle", "Kismet.AllClasses.MapVariableKeyIcon");
+	}
+	else if (InPinType.PinSubCategoryObject.IsValid())
 	{
 		if(UClass* Class = Cast<UClass>(InPinType.PinSubCategoryObject.Get()))
 		{
diff --git a/Engine/Source/Editor/BlueprintGraph/Private/K2Node_VariableGet.cpp b/Engine/Source/Editor/BlueprintGraph/Private/K2Node_VariableGet.cpp
--- a/Engine/Source/Editor/BlueprintGraph/Private/K2Node_VariableGet.cpp
+++ b/Engine/Source/Editor/BlueprintGraph/Private/K2Node_VariableGet.cpp
@@ -132,13 +132,13 @@ void UK2Node_VariableGet::CreateNonPurePins(TArray<UEdGraphPin*>* InOldPinsPtr)
 		if (IsValidTypeForNonPure(PinType))
 		{
 			// Input - Execution Pin
-			CreatePin(EGPD_Input, K2Schema->PC_Exec, TEXT(""), NULL, false, false, K2Schema->PN_Execute);
+			CreatePin(EGPD_Input, K2Schema->PC_Exec, FString(), nullptr, K2Schema->PN_Execute);
 
 			// Output - Execution Pins
-			UEdGraphPin* ValidPin = CreatePin(EGPD_Output, K2Schema->PC_Exec, TEXT(""), NULL, false, false, K2Schema->PN_Then);
+			UEdGraphPin* ValidPin = CreatePin(EGPD_Output, K2Schema->PC_Exec, FString(), nullptr, K2Schema->PN_Then);
 			ValidPin->PinFriendlyName = LOCTEXT("Valid", "Is Valid");
 
-			UEdGraphPin* InvalidPin = CreatePin(EGPD_Output, K2Schema->PC_Exec, TEXT(""), NULL, false, false, K2Schema->PN_Else);
+			UEdGraphPin* InvalidPin = CreatePin(EGPD_Output, K2Schema->PC_Exec, FString(), nullptr, K2Schema->PN_Else);
 			InvalidPin->PinFriendlyName = LOCTEXT("Invalid", "Is Not Valid");
 		}
 		else
@@ -320,7 +320,7 @@ FNodeHandlingFunctor* UK2Node_VariableGet::CreateNodeHandler(FKismetCompilerCont
 
 bool UK2Node_VariableGet::IsValidTypeForNonPure(const FEdGraphPinType& InPinType)
 {
-	return InPinType.bIsArray == false && (InPinType.PinCategory == UObject::StaticClass()->GetName() ||InPinType.PinCategory == UClass::StaticClass()->GetName());
+	return !InPinType.IsContainer() && (InPinType.PinCategory == UObject::StaticClass()->GetName() ||InPinType.PinCategory == UClass::StaticClass()->GetName());
 }
 
 void UK2Node_VariableGet::GetContextMenuActions(const FGraphNodeContextMenuBuilder& Context) const
diff --git a/Engine/Source/Editor/BlueprintGraph/Private/K2Node_VariableSet.cpp b/Engine/Source/Editor/BlueprintGraph/Private/K2Node_VariableSet.cpp
--- a/Engine/Source/Editor/BlueprintGraph/Private/K2Node_VariableSet.cpp
+++ b/Engine/Source/Editor/BlueprintGraph/Private/K2Node_VariableSet.cpp
@@ -82,8 +82,8 @@ void UK2Node_VariableSet::AllocateDefaultPins()
 {
 	const UEdGraphSchema_K2* K2Schema = GetDefault<UEdGraphSchema_K2>();
 
-	CreatePin(EGPD_Input, K2Schema->PC_Exec, TEXT(""), NULL, false, false, K2Schema->PN_Execute);
-	CreatePin(EGPD_Output, K2Schema->PC_Exec, TEXT(""), NULL, false, false, K2Schema->PN_Then);
+	CreatePin(EGPD_Input, K2Schema->PC_Exec, FString(), nullptr, K2Schema->PN_Execute);
+	CreatePin(EGPD_Output, K2Schema->PC_Exec, FString(), nullptr, K2Schema->PN_Then);
 
 	if (GetVarName() != NAME_None)
 	{
@@ -105,8 +105,8 @@ void UK2Node_VariableSet::ReallocatePinsDuringReconstruction(TArray<UEdGraphPin*
 {
 	const UEdGraphSchema_K2* K2Schema = GetDefault<UEdGraphSchema_K2>();
 
-	CreatePin(EGPD_Input, K2Schema->PC_Exec, TEXT(""), NULL, false, false, K2Schema->PN_Execute);
-	CreatePin(EGPD_Output, K2Schema->PC_Exec, TEXT(""), NULL, false, false, K2Schema->PN_Then);
+	CreatePin(EGPD_Input, K2Schema->PC_Exec, FString(), nullptr, K2Schema->PN_Execute);
+	CreatePin(EGPD_Output, K2Schema->PC_Exec, FString(), nullptr, K2Schema->PN_Then);
 
 	if (GetVarName() != NAME_None)
 	{
diff --git a/Engine/Source/Editor/BlueprintGraph/Private/K2Node_VariableSetRef.cpp b/Engine/Source/Editor/BlueprintGraph/Private/K2Node_VariableSetRef.cpp
--- a/Engine/Source/Editor/BlueprintGraph/Private/K2Node_VariableSetRef.cpp
+++ b/Engine/Source/Editor/BlueprintGraph/Private/K2Node_VariableSetRef.cpp
@@ -93,11 +93,11 @@ void UK2Node_VariableSetRef::AllocateDefaultPins()
 {
 	const UEdGraphSchema_K2* K2Schema = GetDefault<UEdGraphSchema_K2>();
 
-	CreatePin(EGPD_Input, K2Schema->PC_Exec, TEXT(""), NULL, false, false, K2Schema->PN_Execute);
-	CreatePin(EGPD_Output, K2Schema->PC_Exec, TEXT(""), NULL, false, false, K2Schema->PN_Then);
+	CreatePin(EGPD_Input, K2Schema->PC_Exec, FString(), nullptr, K2Schema->PN_Execute);
+	CreatePin(EGPD_Output, K2Schema->PC_Exec, FString(), nullptr, K2Schema->PN_Then);
 
-	CreatePin(EGPD_Input, K2Schema->PC_Wildcard, TEXT(""), NULL, false, true, TargetVarPinName);
-	UEdGraphPin* ValuePin = CreatePin(EGPD_Input, K2Schema->PC_Wildcard, TEXT(""), NULL, false, false, VarValuePinName);
+	CreatePin(EGPD_Input, K2Schema->PC_Wildcard, FString(), nullptr, TargetVarPinName, EPinContainerType::None, true);
+	UEdGraphPin* ValuePin = CreatePin(EGPD_Input, K2Schema->PC_Wildcard, FString(), nullptr, VarValuePinName);
 }
 
 void UK2Node_VariableSetRef::ReallocatePinsDuringReconstruction(TArray<UEdGraphPin*>& OldPins)
@@ -192,7 +192,7 @@ void UK2Node_VariableSetRef::CoerceTypeFromPin(const UEdGraphPin* Pin)
 
 	if( Pin )
 	{
-		check((Pin != TargetPin) || (Pin->PinType.bIsReference && !Pin->PinType.bIsArray));
+		check((Pin != TargetPin) || (Pin->PinType.bIsReference && !Pin->PinType.IsContainer()));
 
 		TargetPin->PinType = Pin->PinType;
 		TargetPin->PinType.bIsReference = true;
@@ -204,13 +204,13 @@ void UK2Node_VariableSetRef::CoerceTypeFromPin(const UEdGraphPin* Pin)
 	{
 		// Pin disconnected...revert to wildcard
 		TargetPin->PinType.PinCategory = K2Schema->PC_Wildcard;
-		TargetPin->PinType.PinSubCategory = TEXT("");
-		TargetPin->PinType.PinSubCategoryObject = NULL;
+		TargetPin->PinType.PinSubCategory.Reset();
+		TargetPin->PinType.PinSubCategoryObject = nullptr;
 		TargetPin->BreakAllPinLinks();
 
 		ValuePin->PinType.PinCategory = K2Schema->PC_Wildcard;
-		ValuePin->PinType.PinSubCategory = TEXT("");
-		ValuePin->PinType.PinSubCategoryObject = NULL;
+		ValuePin->PinType.PinSubCategory.Reset();
+		ValuePin->PinType.PinSubCategoryObject = nullptr;
 		ValuePin->BreakAllPinLinks();
 
 		CachedNodeTitle.MarkDirty();
diff --git a/Engine/Source/Editor/BlueprintGraph/Private/MathExpressionHandler.cpp b/Engine/Source/Editor/BlueprintGraph/Private/MathExpressionHandler.cpp
--- a/Engine/Source/Editor/BlueprintGraph/Private/MathExpressionHandler.cpp
+++ b/Engine/Source/Editor/BlueprintGraph/Private/MathExpressionHandler.cpp
@@ -14,7 +14,7 @@
 
 struct KCHandler_MathExpressionHelper
 {
-	static bool CreateMap(UEdGraphNode& KeyNode, UEdGraphNode& ValueNode, EEdGraphPinDirection KeyPinDirection, UEdGraphSchema_K2* Schema, TMap<UEdGraphPin*, UEdGraphPin*>& OutMap)
+	static bool CreateMap(UEdGraphNode& KeyNode, UEdGraphNode& ValueNode, EEdGraphPinDirection KeyPinDirection, const UEdGraphSchema_K2* Schema, TMap<UEdGraphPin*, UEdGraphPin*>& OutMap)
 	{
 		check(Schema);
 		for (auto InnerInputPin : KeyNode.Pins)
diff --git a/Engine/Source/Editor/BlueprintGraph/Public/BlueprintActionFilter.h b/Engine/Source/Editor/BlueprintGraph/Public/BlueprintActionFilter.h
--- a/Engine/Source/Editor/BlueprintGraph/Public/BlueprintActionFilter.h
+++ b/Engine/Source/Editor/BlueprintGraph/Public/BlueprintActionFilter.h
@@ -134,7 +134,7 @@ struct BLUEPRINTGRAPH_API FBlueprintActionInfo
 
 private:
 	/** The class or asset-object that the NodeSpawner action is keyed to (in the action database)*/
-	UObject const* ActionOwner;
+	const TWeakObjectPtr<UObject> ActionOwner;
 
 	/** Keeps track of the fields we've cached (needed in case one turns out to be null) */
 	uint32 CacheFlags;
diff --git a/Engine/Source/Editor/Cascade/Private/Cascade.cpp b/Engine/Source/Editor/Cascade/Private/Cascade.cpp
--- a/Engine/Source/Editor/Cascade/Private/Cascade.cpp
+++ b/Engine/Source/Editor/Cascade/Private/Cascade.cpp
@@ -1588,7 +1588,7 @@ void FCascade::Tick(float DeltaTime)
 			Position.X = MotionModeRadius * FMath::Sin(AccumulatedMotionTime);
 			Position.Y = MotionModeRadius * FMath::Cos(AccumulatedMotionTime);
 			Position.Z = 0.0f;
-			ParticleSystemComponent->ComponentToWorld = FTransform(Position);
+			ParticleSystemComponent->SetComponentToWorld(FTransform(Position));
 		}
 
 		if (ParticleSystemComponent->IsComponentTickEnabled())
diff --git a/Engine/Source/Editor/ClothPainter/Private/ClothPainter.cpp b/Engine/Source/Editor/ClothPainter/Private/ClothPainter.cpp
--- a/Engine/Source/Editor/ClothPainter/Private/ClothPainter.cpp
+++ b/Engine/Source/Editor/ClothPainter/Private/ClothPainter.cpp
@@ -78,7 +78,7 @@ bool FClothPainter::PaintInternal(const FVector& InCameraOrigin, const FVector&
 			}
 			else if (PaintSettings->PaintTool == EClothPaintTool::Gradient)
 			{
-				const FMatrix ComponentToWorldMatrix = HitResult.Component->ComponentToWorld.ToMatrixWithScale();
+				const FMatrix ComponentToWorldMatrix = HitResult.Component->GetComponentTransform().ToMatrixWithScale();
 				const FVector ComponentSpaceCameraPosition(ComponentToWorldMatrix.InverseTransformPosition(InCameraOrigin));
 				const FVector ComponentSpaceBrushPosition(ComponentToWorldMatrix.InverseTransformPosition(HitResult.Location));
 				const float ComponentSpaceBrushRadius = ComponentToWorldMatrix.InverseTransformVector(FVector(PaintSettings->bUseRegularBrushForGradient ? BrushSettings->GetBrushRadius() : 2.0f, 0.0f, 0.0f)).Size();
@@ -332,7 +332,7 @@ void FClothPainter::Render(const FSceneView* View, FViewport* Viewport, FPrimiti
 	// Render simulation mesh vertices if not simulating
 	if (SkeletalMeshComponent && !bShouldSimulate)
 	{
-		const FMatrix ComponentToWorldMatrix = SkeletalMeshComponent->ComponentToWorld.ToMatrixWithScale();
+		const FMatrix ComponentToWorldMatrix = SkeletalMeshComponent->GetComponentTransform().ToMatrixWithScale();
 		const TArray<FVector>& AllVertices = Adapter->GetMeshVertices();
 		for (const FVector& Vertex : AllVertices)
 		{
@@ -346,7 +346,7 @@ void FClothPainter::Render(const FSceneView* View, FViewport* Viewport, FPrimiti
 		TArray<MeshPaintHelpers::FPaintRay> PaintRays;
 		MeshPaintHelpers::RetrieveViewportPaintRays(View, Viewport, PDI, PaintRays);
 
-		const FMatrix ComponentToWorldMatrix = SkeletalMeshComponent->ComponentToWorld.ToMatrixWithScale();
+		const FMatrix ComponentToWorldMatrix = SkeletalMeshComponent->GetComponentTransform().ToMatrixWithScale();
 
 		for (const MeshPaintHelpers::FPaintRay& PaintRay : PaintRays)
 		{
diff --git a/Engine/Source/Editor/ComponentVisualizers/Private/ConstraintComponentVisualizer.cpp b/Engine/Source/Editor/ComponentVisualizers/Private/ConstraintComponentVisualizer.cpp
--- a/Engine/Source/Editor/ComponentVisualizers/Private/ConstraintComponentVisualizer.cpp
+++ b/Engine/Source/Editor/ComponentVisualizers/Private/ConstraintComponentVisualizer.cpp
@@ -49,9 +49,9 @@ void FConstraintComponentVisualizer::DrawVisualization( const UActorComponent* C
 		// Otherwise use the component frame
 		else
 		{
-			Con1Frame = ConstraintComp->ComponentToWorld;
+			Con1Frame = ConstraintComp->GetComponentTransform();
 			Con1Frame.RemoveScaling();
-			Con2Frame = ConstraintComp->ComponentToWorld;
+			Con2Frame = ConstraintComp->GetComponentTransform();
 			Con2Frame.SetRotation(Con2Frame.GetRotation() * ConstraintComp->ConstraintInstance.AngularRotationOffset.Quaternion());
 			Con1Frame.RemoveScaling();
 		}
diff --git a/Engine/Source/Editor/ComponentVisualizers/Private/DecalComponentVisualizer.cpp b/Engine/Source/Editor/ComponentVisualizers/Private/DecalComponentVisualizer.cpp
--- a/Engine/Source/Editor/ComponentVisualizers/Private/DecalComponentVisualizer.cpp
+++ b/Engine/Source/Editor/ComponentVisualizers/Private/DecalComponentVisualizer.cpp
@@ -11,7 +11,7 @@ void FDecalComponentVisualizer::DrawVisualization( const UActorComponent* Compon
 	const UDecalComponent* DecalComponent = Cast<const UDecalComponent>(Component);
 	if(DecalComponent)
 	{
-		const FMatrix LocalToWorld = DecalComponent->ComponentToWorld.ToMatrixWithScale();
+		const FMatrix LocalToWorld = DecalComponent->GetComponentTransform().ToMatrixWithScale();
 		
 		const FLinearColor DrawColor = FColor(0, 157, 0, 255);
 
diff --git a/Engine/Source/Editor/ComponentVisualizers/Private/PointLightComponentVisualizer.cpp b/Engine/Source/Editor/ComponentVisualizers/Private/PointLightComponentVisualizer.cpp
--- a/Engine/Source/Editor/ComponentVisualizers/Private/PointLightComponentVisualizer.cpp
+++ b/Engine/Source/Editor/ComponentVisualizers/Private/PointLightComponentVisualizer.cpp
@@ -13,7 +13,7 @@ void FPointLightComponentVisualizer::DrawVisualization( const UActorComponent* C
 		const UPointLightComponent* PointLightComp = Cast<const UPointLightComponent>(Component);
 		if(PointLightComp != NULL)
 		{
-			FTransform LightTM = PointLightComp->ComponentToWorld;
+			FTransform LightTM = PointLightComp->GetComponentTransform();
 			LightTM.RemoveScaling();
 
 			// Draw light radius
diff --git a/Engine/Source/Editor/ComponentVisualizers/Private/RadialForceComponentVisualizer.cpp b/Engine/Source/Editor/ComponentVisualizers/Private/RadialForceComponentVisualizer.cpp
--- a/Engine/Source/Editor/ComponentVisualizers/Private/RadialForceComponentVisualizer.cpp
+++ b/Engine/Source/Editor/ComponentVisualizers/Private/RadialForceComponentVisualizer.cpp
@@ -11,7 +11,7 @@ void FRadialForceComponentVisualizer::DrawVisualization( const UActorComponent*
 	const URadialForceComponent* ForceComp = Cast<const URadialForceComponent>(Component);
 	if(ForceComp != NULL)
 	{
-		FTransform TM = ForceComp->ComponentToWorld;
+		FTransform TM = ForceComp->GetComponentTransform();
 		TM.RemoveScaling();
 
 		// Draw light radius
diff --git a/Engine/Source/Editor/ComponentVisualizers/Private/SplineComponentVisualizer.cpp b/Engine/Source/Editor/ComponentVisualizers/Private/SplineComponentVisualizer.cpp
--- a/Engine/Source/Editor/ComponentVisualizers/Private/SplineComponentVisualizer.cpp
+++ b/Engine/Source/Editor/ComponentVisualizers/Private/SplineComponentVisualizer.cpp
@@ -550,11 +550,11 @@ bool FSplineComponentVisualizer::GetWidgetLocation(const FEditorViewportClient*
 			check(SelectedTangentHandleType != ESelectedTangentHandle::None);
 			if (SelectedTangentHandleType == ESelectedTangentHandle::Leave)
 			{
-				OutLocation = SplineComp->ComponentToWorld.TransformPosition(Point.OutVal + Point.LeaveTangent);
+				OutLocation = SplineComp->GetComponentTransform().TransformPosition(Point.OutVal + Point.LeaveTangent);
 			}
 			else if (SelectedTangentHandleType == ESelectedTangentHandle::Arrive)
 			{
-				OutLocation = SplineComp->ComponentToWorld.TransformPosition(Point.OutVal - Point.ArriveTangent);
+				OutLocation = SplineComp->GetComponentTransform().TransformPosition(Point.OutVal - Point.ArriveTangent);
 			}
 
 			return true;
@@ -565,7 +565,7 @@ bool FSplineComponentVisualizer::GetWidgetLocation(const FEditorViewportClient*
 			check(LastKeyIndexSelected < Position.Points.Num());
 			check(SelectedKeys.Contains(LastKeyIndexSelected));
 			const auto& Point = Position.Points[LastKeyIndexSelected];
-			OutLocation = SplineComp->ComponentToWorld.TransformPosition(Point.OutVal);
+			OutLocation = SplineComp->GetComponentTransform().TransformPosition(Point.OutVal);
 			return true;
 		}
 	}
@@ -625,17 +625,17 @@ bool FSplineComponentVisualizer::HandleInputDelta(FEditorViewportClient* Viewpor
 				{
 					if (SelectedTangentHandleType == ESelectedTangentHandle::Leave)
 					{
-						EditedPoint.LeaveTangent += SplineComp->ComponentToWorld.InverseTransformVector(DeltaTranslate);
+						EditedPoint.LeaveTangent += SplineComp->GetComponentTransform().InverseTransformVector(DeltaTranslate);
 					}
 					else
 					{
-						EditedPoint.ArriveTangent += SplineComp->ComponentToWorld.InverseTransformVector(-DeltaTranslate);
+						EditedPoint.ArriveTangent += SplineComp->GetComponentTransform().InverseTransformVector(-DeltaTranslate);
 					}
 				}
 				else
 				{
 					const FVector Delta = (SelectedTangentHandleType == ESelectedTangentHandle::Leave) ? DeltaTranslate : -DeltaTranslate;
-					const FVector Tangent = EditedPoint.LeaveTangent + SplineComp->ComponentToWorld.InverseTransformVector(Delta);
+					const FVector Tangent = EditedPoint.LeaveTangent + SplineComp->GetComponentTransform().InverseTransformVector(Delta);
 
 					EditedPoint.LeaveTangent = Tangent;
 					EditedPoint.ArriveTangent = Tangent;
@@ -671,11 +671,11 @@ bool FSplineComponentVisualizer::HandleInputDelta(FEditorViewportClient* Viewpor
 				if (!DeltaTranslate.IsZero())
 				{
 					// Find key position in world space
-					const FVector CurrentWorldPos = SplineComp->ComponentToWorld.TransformPosition(EditedPoint.OutVal);
+					const FVector CurrentWorldPos = SplineComp->GetComponentTransform().TransformPosition(EditedPoint.OutVal);
 					// Move in world space
 					const FVector NewWorldPos = CurrentWorldPos + DeltaTranslate;
 					// Convert back to local space
-					EditedPoint.OutVal = SplineComp->ComponentToWorld.InverseTransformPosition(NewWorldPos);
+					EditedPoint.OutVal = SplineComp->GetComponentTransform().InverseTransformPosition(NewWorldPos);
 				}
 
 				if (!DeltaRotate.IsZero())
@@ -684,16 +684,16 @@ bool FSplineComponentVisualizer::HandleInputDelta(FEditorViewportClient* Viewpor
 					EditedPoint.InterpMode = CIM_CurveUser;
 
 					// Rotate tangent according to delta rotation
-					FVector NewTangent = SplineComp->ComponentToWorld.GetRotation().RotateVector(EditedPoint.LeaveTangent); // convert local-space tangent vector to world-space
+					FVector NewTangent = SplineComp->GetComponentTransform().GetRotation().RotateVector(EditedPoint.LeaveTangent); // convert local-space tangent vector to world-space
 					NewTangent = DeltaRotate.RotateVector(NewTangent); // apply world-space delta rotation to world-space tangent
-					NewTangent = SplineComp->ComponentToWorld.GetRotation().Inverse().RotateVector(NewTangent); // convert world-space tangent vector back into local-space
+					NewTangent = SplineComp->GetComponentTransform().GetRotation().Inverse().RotateVector(NewTangent); // convert world-space tangent vector back into local-space
 					EditedPoint.LeaveTangent = NewTangent;
 					EditedPoint.ArriveTangent = NewTangent;
 
 					// Rotate spline rotation according to delta rotation
-					FQuat NewRot = SplineComp->ComponentToWorld.GetRotation() * EditedRotPoint.OutVal; // convert local-space rotation to world-space
+					FQuat NewRot = SplineComp->GetComponentTransform().GetRotation() * EditedRotPoint.OutVal; // convert local-space rotation to world-space
 					NewRot = DeltaRotate.Quaternion() * NewRot; // apply world-space rotation
-					NewRot = SplineComp->ComponentToWorld.GetRotation().Inverse() * NewRot; // convert world-space rotation to local-space
+					NewRot = SplineComp->GetComponentTransform().GetRotation().Inverse() * NewRot; // convert world-space rotation to local-space
 					EditedRotPoint.OutVal = NewRot;
 				}
 
@@ -894,7 +894,7 @@ void FSplineComponentVisualizer::OnAddKey()
 
 	FInterpCurvePoint<FVector> NewPoint(
 		SelectedSegmentIndex,
-		SplineComp->ComponentToWorld.InverseTransformPosition(SelectedSplinePosition),
+		SplineComp->GetComponentTransform().InverseTransformPosition(SelectedSplinePosition),
 		FVector::ZeroVector,
 		FVector::ZeroVector,
 		CIM_CurveAuto);
diff --git a/Engine/Source/Editor/ComponentVisualizers/Private/SplineMeshComponentVisualizer.cpp b/Engine/Source/Editor/ComponentVisualizers/Private/SplineMeshComponentVisualizer.cpp
--- a/Engine/Source/Editor/ComponentVisualizers/Private/SplineMeshComponentVisualizer.cpp
+++ b/Engine/Source/Editor/ComponentVisualizers/Private/SplineMeshComponentVisualizer.cpp
@@ -52,8 +52,8 @@ void FSplineMeshComponentVisualizer::DrawVisualization(const UActorComponent* Co
 		// Draw the tangent handles before anything else so they will not overdraw the rest of the spline
 		for (int32 PointIndex = 0; PointIndex < 2; PointIndex++)
 		{
-			const FVector KeyPos = SplineMeshComp->ComponentToWorld.TransformPosition(Spline.Points[PointIndex].OutVal);
-			const FVector TangentWorldDirection = SplineMeshComp->ComponentToWorld.TransformVector(Spline.Points[PointIndex].LeaveTangent);
+			const FVector KeyPos = SplineMeshComp->GetComponentTransform().TransformPosition(Spline.Points[PointIndex].OutVal);
+			const FVector TangentWorldDirection = SplineMeshComp->GetComponentTransform().TransformVector(Spline.Points[PointIndex].LeaveTangent);
 
 			PDI->SetHitProxy(NULL);
 			DrawDashedLine(PDI, KeyPos, KeyPos + TangentWorldDirection, Color, 5, SDPG_Foreground);
@@ -68,17 +68,17 @@ void FSplineMeshComponentVisualizer::DrawVisualization(const UActorComponent* Co
 		// Draw the keypoints
 		for (int32 PointIndex = 0; PointIndex < 2; PointIndex++)
 		{
-			const FVector NewKeyPos = SplineMeshComp->ComponentToWorld.TransformPosition(Spline.Points[PointIndex].OutVal);
+			const FVector NewKeyPos = SplineMeshComp->GetComponentTransform().TransformPosition(Spline.Points[PointIndex].OutVal);
 			PDI->SetHitProxy(new HSplineMeshKeyProxy(Component, PointIndex));
 			PDI->DrawPoint(NewKeyPos, Color, GrabHandleSize, SDPG_Foreground);
 			PDI->SetHitProxy(NULL);
 		}
 
 		// Draw the spline
-		FVector StartPos = SplineMeshComp->ComponentToWorld.TransformPosition(Spline.Points[0].OutVal);
+		FVector StartPos = SplineMeshComp->GetComponentTransform().TransformPosition(Spline.Points[0].OutVal);
 		for (int32 Step = 1; Step < 32; Step++)
 		{
-			const FVector EndPos = SplineMeshComp->ComponentToWorld.TransformPosition(Spline.Eval(Step / 32.0f, FVector::ZeroVector));
+			const FVector EndPos = SplineMeshComp->GetComponentTransform().TransformPosition(Spline.Eval(Step / 32.0f, FVector::ZeroVector));
 			PDI->DrawLine(StartPos, EndPos, Color, SDPG_Foreground);
 			StartPos = EndPos;
 		}
@@ -170,11 +170,11 @@ bool FSplineMeshComponentVisualizer::GetWidgetLocation(const FEditorViewportClie
 			check(SelectedTangentHandleType != ESelectedTangentHandle::None);
 			if (SelectedTangentHandleType == ESelectedTangentHandle::Leave)
 			{
-				OutLocation = SplineMeshComp->ComponentToWorld.TransformPosition(Point.OutVal + Point.LeaveTangent);
+				OutLocation = SplineMeshComp->GetComponentTransform().TransformPosition(Point.OutVal + Point.LeaveTangent);
 			}
 			else if (SelectedTangentHandleType == ESelectedTangentHandle::Arrive)
 			{
-				OutLocation = SplineMeshComp->ComponentToWorld.TransformPosition(Point.OutVal - Point.ArriveTangent);
+				OutLocation = SplineMeshComp->GetComponentTransform().TransformPosition(Point.OutVal - Point.ArriveTangent);
 			}
 
 			return true;
@@ -184,7 +184,7 @@ bool FSplineMeshComponentVisualizer::GetWidgetLocation(const FEditorViewportClie
 			// Otherwise use the last key index set
 			check(SelectedKey < 2);
 			const auto& Point = Spline.Points[SelectedKey];
-			OutLocation = SplineMeshComp->ComponentToWorld.TransformPosition(Point.OutVal);
+			OutLocation = SplineMeshComp->GetComponentTransform().TransformPosition(Point.OutVal);
 			return true;
 		}
 	}
@@ -220,7 +220,7 @@ bool FSplineMeshComponentVisualizer::GetCustomInputCoordinateSystem(const FEdito
 				const FVector Bitangent = (Tangent.Z == 1.0f) ? FVector(1.0f, 0.0f, 0.0f) : FVector(-Tangent.Y, Tangent.X, 0.0f).GetSafeNormal();
 				const FVector Normal = FVector::CrossProduct(Tangent, Bitangent);
 
-				OutMatrix = FMatrix(Tangent, Bitangent, Normal, FVector::ZeroVector) * FQuatRotationTranslationMatrix(SplineMeshComp->ComponentToWorld.GetRotation(), FVector::ZeroVector);
+				OutMatrix = FMatrix(Tangent, Bitangent, Normal, FVector::ZeroVector) * FQuatRotationTranslationMatrix(SplineMeshComp->GetComponentTransform().GetRotation(), FVector::ZeroVector);
 				return true;
 			}
 		}
@@ -247,7 +247,7 @@ bool FSplineMeshComponentVisualizer::HandleInputDelta(FEditorViewportClient* Vie
 				check(SelectedTangentHandleType != ESelectedTangentHandle::None);
 
 				const FVector Delta = (SelectedTangentHandleType == ESelectedTangentHandle::Leave) ? DeltaTranslate : -DeltaTranslate;
-				const FVector NewTangent = OldTangent + SplineMeshComp->ComponentToWorld.InverseTransformVector(Delta);
+				const FVector NewTangent = OldTangent + SplineMeshComp->GetComponentTransform().InverseTransformVector(Delta);
 
 				SplineMeshComp->Modify();
 
@@ -279,11 +279,11 @@ bool FSplineMeshComponentVisualizer::HandleInputDelta(FEditorViewportClient* Vie
 			if (!DeltaTranslate.IsZero())
 			{
 				// Find key position in world space
-				const FVector CurrentWorldPos = SplineMeshComp->ComponentToWorld.TransformPosition(KeyPosition);
+				const FVector CurrentWorldPos = SplineMeshComp->GetComponentTransform().TransformPosition(KeyPosition);
 				// Move in world space
 				const FVector NewWorldPos = CurrentWorldPos + DeltaTranslate;
 				// Convert back to local space
-				KeyPosition = SplineMeshComp->ComponentToWorld.InverseTransformPosition(NewWorldPos);
+				KeyPosition = SplineMeshComp->GetComponentTransform().InverseTransformPosition(NewWorldPos);
 
 				bModifiedPosition = true;
 			}
diff --git a/Engine/Source/Editor/ComponentVisualizers/Private/SpotLightComponentVisualizer.cpp b/Engine/Source/Editor/ComponentVisualizers/Private/SpotLightComponentVisualizer.cpp
--- a/Engine/Source/Editor/ComponentVisualizers/Private/SpotLightComponentVisualizer.cpp
+++ b/Engine/Source/Editor/ComponentVisualizers/Private/SpotLightComponentVisualizer.cpp
@@ -13,7 +13,7 @@ void FSpotLightComponentVisualizer::DrawVisualization( const UActorComponent* Co
 		const USpotLightComponent* SpotLightComp = Cast<const USpotLightComponent>(Component);
 		if(SpotLightComp != NULL)
 		{
-			FTransform TransformNoScale = SpotLightComp->ComponentToWorld;
+			FTransform TransformNoScale = SpotLightComp->GetComponentTransform();
 			TransformNoScale.RemoveScaling();
 
 			// Draw point light source shape
diff --git a/Engine/Source/Editor/ComponentVisualizers/Private/StereoLayerComponentVisualizer.cpp b/Engine/Source/Editor/ComponentVisualizers/Private/StereoLayerComponentVisualizer.cpp
--- a/Engine/Source/Editor/ComponentVisualizers/Private/StereoLayerComponentVisualizer.cpp
+++ b/Engine/Source/Editor/ComponentVisualizers/Private/StereoLayerComponentVisualizer.cpp
@@ -17,15 +17,15 @@ void FStereoLayerComponentVisualizer::DrawVisualization( const UActorComponent*
             const FVector2D QuadSize = StereoLayerComp->GetQuadSize() / 2.0f;
             const FBox QuadBox(FVector(0.0f, -QuadSize.X, -QuadSize.Y), FVector(0.0f, QuadSize.X, QuadSize.Y));
 
-            DrawWireBox(PDI, StereoLayerComp->ComponentToWorld.ToMatrixWithScale(), QuadBox, YellowColor, 0);
+            DrawWireBox(PDI, StereoLayerComp->GetComponentTransform().ToMatrixWithScale(), QuadBox, YellowColor, 0);
         }
         else if(StereoLayerComp->StereoLayerShape == EStereoLayerShape::SLSH_CylinderLayer)
         {
             float ArcAngle = StereoLayerComp->CylinderOverlayArc * 180 / (StereoLayerComp->CylinderRadius * PI);
             
-            FVector X = StereoLayerComp->ComponentToWorld.GetUnitAxis(EAxis::Type::X);
-            FVector Y = StereoLayerComp->ComponentToWorld.GetUnitAxis(EAxis::Type::Y);
-            FVector Base = StereoLayerComp->ComponentToWorld.GetLocation();
+            FVector X = StereoLayerComp->GetComponentTransform().GetUnitAxis(EAxis::Type::X);
+            FVector Y = StereoLayerComp->GetComponentTransform().GetUnitAxis(EAxis::Type::Y);
+            FVector Base = StereoLayerComp->GetComponentTransform().GetLocation();
             FVector HalfHeight = FVector(0, 0, StereoLayerComp->CylinderHeight/2);
             
             FVector LeftVertex = Base + StereoLayerComp->CylinderRadius * ( FMath::Cos(ArcAngle/2 * (PI/180.0f)) * X + FMath::Sin(ArcAngle/2 * (PI/180.0f)) * Y );
diff --git a/Engine/Source/Editor/ContentBrowser/Private/AssetViewWidgets.cpp b/Engine/Source/Editor/ContentBrowser/Private/AssetViewWidgets.cpp
--- a/Engine/Source/Editor/ContentBrowser/Private/AssetViewWidgets.cpp
+++ b/Engine/Source/Editor/ContentBrowser/Private/AssetViewWidgets.cpp
@@ -704,7 +704,7 @@ TSharedRef<SWidget> SAssetViewItem::CreateToolTipWidget() const
 			}
 
 			// Add asset source files
-			TOptional<FAssetImportInfo> ImportInfo = FAssetSourceFilenameCache::ExtractAssetImportInfo(AssetData.TagsAndValues);
+			TOptional<FAssetImportInfo> ImportInfo = FAssetSourceFilenameCache::ExtractAssetImportInfo(AssetData);
 			if (ImportInfo.IsSet())
 			{
 				for (const auto& File : ImportInfo->SourceFiles)
diff --git a/Engine/Source/Editor/ContentBrowser/Private/SAssetView.cpp b/Engine/Source/Editor/ContentBrowser/Private/SAssetView.cpp
--- a/Engine/Source/Editor/ContentBrowser/Private/SAssetView.cpp
+++ b/Engine/Source/Editor/ContentBrowser/Private/SAssetView.cpp
@@ -551,11 +551,8 @@ void SAssetView::DeferredCreateNewAsset()
 		FName PackagePathFName = FName(*DeferredAssetToCreate->PackagePath);
 		FName AssetName = FName(*DeferredAssetToCreate->DefaultAssetName);
 		FName AssetClassName = DeferredAssetToCreate->AssetClass->GetFName();
-		TMap<FName, FString> EmptyTags;
-		TArray<int32> EmptyChunkIDs;
-		const uint32 EmptyPackageFlags = 0;
 
-		FAssetData NewAssetData(PackageName, PackagePathFName, NAME_None, AssetName, AssetClassName, EmptyTags, EmptyChunkIDs, EmptyPackageFlags);
+		FAssetData NewAssetData(PackageName, PackagePathFName, AssetName, AssetClassName);
 		TSharedPtr<FAssetViewItem> NewItem = MakeShareable(new FAssetViewCreation(NewAssetData, DeferredAssetToCreate->AssetClass, DeferredAssetToCreate->Factory));
 
 		NewItem->bRenameWhenScrolledIntoview = true;
@@ -590,11 +587,8 @@ void SAssetView::DuplicateAsset(const FString& PackagePath, const TWeakObjectPtr
 	FName PackagePathFName = FName(*PackagePath);
 	FName AssetName = FName(*AssetNameStr);
 	FName AssetClass = OriginalObject->GetClass()->GetFName();
-	TMap<FName, FString> EmptyTags;
-	TArray<int32> EmptyChunkIDs;
-	const uint32 EmptyPackageFlags = 0;
-
-	FAssetData NewAssetData(PackageName, PackagePathFName, NAME_None, AssetName, AssetClass, EmptyTags, EmptyChunkIDs, EmptyPackageFlags);
+	
+	FAssetData NewAssetData(PackageName, PackagePathFName, AssetName, AssetClass);
 	TSharedPtr<FAssetViewItem> NewItem = MakeShareable(new FAssetViewDuplication(NewAssetData, OriginalObject));
 	NewItem->bRenameWhenScrolledIntoview = true;
 
diff --git a/Engine/Source/Editor/ContentBrowser/Private/SFilterList.cpp b/Engine/Source/Editor/ContentBrowser/Private/SFilterList.cpp
--- a/Engine/Source/Editor/ContentBrowser/Private/SFilterList.cpp
+++ b/Engine/Source/Editor/ContentBrowser/Private/SFilterList.cpp
@@ -415,7 +415,7 @@ void SFilterList::Construct( const FArguments& InArgs )
 	{
 		if (UContentBrowserFrontEndFilterExtension* PotentialExtension = *ExtensionIt)
 		{
-			if (PotentialExtension->HasAnyFlags(RF_ClassDefaultObject) && !PotentialExtension->GetClass()->HasAnyCastFlag(CLASS_Deprecated | CLASS_Abstract))
+			if (PotentialExtension->HasAnyFlags(RF_ClassDefaultObject) && !PotentialExtension->GetClass()->HasAnyClassFlags(CLASS_Deprecated | CLASS_Abstract))
 			{
 				// Grab the filters
 				TArray< TSharedRef<FFrontendFilter> > ExtendedFrontendFilters;
diff --git a/Engine/Source/Editor/DetailCustomizations/Private/ComponentTransformDetails.cpp b/Engine/Source/Editor/DetailCustomizations/Private/ComponentTransformDetails.cpp
--- a/Engine/Source/Editor/DetailCustomizations/Private/ComponentTransformDetails.cpp
+++ b/Engine/Source/Editor/DetailCustomizations/Private/ComponentTransformDetails.cpp
@@ -693,12 +693,12 @@ void FComponentTransformDetails::OnToggleAbsoluteLocation( bool bEnable )
 				{
 					if (SceneComponent->bAbsoluteLocation)
 					{
-						SceneComponent->RelativeLocation = SceneComponent->ComponentToWorld.GetTranslation();
+						SceneComponent->RelativeLocation = SceneComponent->GetComponentTransform().GetTranslation();
 					}
 					else
 					{
 						FTransform ParentToWorld = SceneComponent->GetAttachParent()->GetSocketTransform(SceneComponent->GetAttachSocketName());
-						FTransform RelativeTM = SceneComponent->ComponentToWorld.GetRelativeTransform(ParentToWorld);
+						FTransform RelativeTM = SceneComponent->GetComponentTransform().GetRelativeTransform(ParentToWorld);
 						SceneComponent->RelativeLocation = RelativeTM.GetTranslation();
 					}
 				}
diff --git a/Engine/Source/Editor/EditorStyle/Private/SlateEditorStyle.cpp b/Engine/Source/Editor/EditorStyle/Private/SlateEditorStyle.cpp
--- a/Engine/Source/Editor/EditorStyle/Private/SlateEditorStyle.cpp
+++ b/Engine/Source/Editor/EditorStyle/Private/SlateEditorStyle.cpp
@@ -5475,6 +5475,9 @@ void FSlateEditorStyle::FStyle::SetupPersonaStyle()
 
 		Set( "Kismet.AllClasses.VariableIcon", new IMAGE_BRUSH( "/Icons/pill_16x", Icon16x16 ) );
 		Set( "Kismet.AllClasses.ArrayVariableIcon", new IMAGE_BRUSH( "/Icons/pillarray_16x", Icon16x16 ) );
+		Set( "Kismet.AllClasses.SetVariableIcon", new IMAGE_BRUSH( "/Icons/pillset_16x", Icon16x16 ) );
+		Set( "Kismet.AllClasses.MapValueVariableIcon", new IMAGE_BRUSH( "/Icons/pillmapvalue_16x", Icon16x16 ) );
+		Set( "Kismet.AllClasses.MapKeyVariableIcon", new IMAGE_BRUSH( "/Icons/pillmapkey_16x", Icon16x16 ) );
 		Set( "Kismet.AllClasses.FunctionIcon", new IMAGE_BRUSH( "/Icons/icon_BluePrintEditor_Function_16px", Icon16x16 ) );
 
 		Set( "BlueprintEditor.ResetCamera", new IMAGE_BRUSH( "Icons/icon_Camera_Reset_40px", Icon16x16));
diff --git a/Engine/Source/Editor/EnvironmentQueryEditor/Private/EnvironmentQueryGraphNode_Option.cpp b/Engine/Source/Editor/EnvironmentQueryEditor/Private/EnvironmentQueryGraphNode_Option.cpp
--- a/Engine/Source/Editor/EnvironmentQueryEditor/Private/EnvironmentQueryGraphNode_Option.cpp
+++ b/Engine/Source/Editor/EnvironmentQueryEditor/Private/EnvironmentQueryGraphNode_Option.cpp
@@ -21,7 +21,7 @@ UEnvironmentQueryGraphNode_Option::UEnvironmentQueryGraphNode_Option(const FObje
 
 void UEnvironmentQueryGraphNode_Option::AllocateDefaultPins()
 {
-	UEdGraphPin* Inputs = CreatePin(EGPD_Input, TEXT("Transition"), TEXT(""), NULL, false, false, TEXT("Out"));
+	UEdGraphPin* Inputs = CreatePin(EGPD_Input, TEXT("Transition"), FString(), nullptr, TEXT("Out"));
 }
 
 void UEnvironmentQueryGraphNode_Option::PostPlacedNewNode()
@@ -161,7 +161,7 @@ void UEnvironmentQueryGraphNode_Option::UpdateNodeData()
 	{
 		CompositeGenerator->VerifyItemTypes();
 
-		ErrorMessage = CompositeGenerator->bHasMatchingItemType ? TEXT("") : LOCTEXT("NestedGeneratorMismatch", "Nested generators must work on exactly the same item types!").ToString();
+		ErrorMessage = CompositeGenerator->bHasMatchingItemType ? FString() : LOCTEXT("NestedGeneratorMismatch", "Nested generators must work on exactly the same item types!").ToString();
 	}
 }
 
diff --git a/Engine/Source/Editor/EnvironmentQueryEditor/Private/EnvironmentQueryGraphNode_Root.cpp b/Engine/Source/Editor/EnvironmentQueryEditor/Private/EnvironmentQueryGraphNode_Root.cpp
--- a/Engine/Source/Editor/EnvironmentQueryEditor/Private/EnvironmentQueryGraphNode_Root.cpp
+++ b/Engine/Source/Editor/EnvironmentQueryEditor/Private/EnvironmentQueryGraphNode_Root.cpp
@@ -9,7 +9,7 @@ UEnvironmentQueryGraphNode_Root::UEnvironmentQueryGraphNode_Root(const FObjectIn
 
 void UEnvironmentQueryGraphNode_Root::AllocateDefaultPins()
 {
-	UEdGraphPin* Outputs = CreatePin(EGPD_Output, TEXT("Transition"), TEXT(""), NULL, false, false, TEXT("In"));
+	UEdGraphPin* Outputs = CreatePin(EGPD_Output, TEXT("Transition"), FString(), nullptr, TEXT("In"));
 }
 
 FText UEnvironmentQueryGraphNode_Root::GetNodeTitle(ENodeTitleType::Type TitleType) const
diff --git a/Engine/Source/Editor/FoliageEdit/Private/FoliageEdMode.cpp b/Engine/Source/Editor/FoliageEdit/Private/FoliageEdMode.cpp
--- a/Engine/Source/Editor/FoliageEdit/Private/FoliageEdMode.cpp
+++ b/Engine/Source/Editor/FoliageEdit/Private/FoliageEdMode.cpp
@@ -2459,7 +2459,7 @@ void FEdModeFoliage::ApplyPaintBucket_Add(AActor* Actor)
 			const bool bHasColorData = bHasInstancedColorData || LODModel.ColorVertexBuffer.GetNumVertices();
 
 			// Get the raw triangle data for this static mesh
-			FTransform LocalToWorld = StaticMeshComponent->ComponentToWorld;
+			FTransform LocalToWorld = StaticMeshComponent->GetComponentTransform();
 			FIndexArrayView Indices = LODModel.IndexBuffer.GetArrayView();
 			const FPositionVertexBuffer& PositionVertexBuffer = LODModel.PositionVertexBuffer;
 			const FColorVertexBuffer& ColorVertexBuffer = LODModel.ColorVertexBuffer;
@@ -2630,9 +2630,9 @@ bool FEdModeFoliage::GetStaticMeshVertexColorForHit(const UStaticMeshComponent*
 				int32 Index2 = Indices[IndexBufferIdx + 2];
 
 				// Lookup the triangle world positions and colors.
-				FVector WorldVert0 = InStaticMeshComponent->ComponentToWorld.TransformPosition(PositionVertexBuffer.VertexPosition(Index0));
-				FVector WorldVert1 = InStaticMeshComponent->ComponentToWorld.TransformPosition(PositionVertexBuffer.VertexPosition(Index1));
-				FVector WorldVert2 = InStaticMeshComponent->ComponentToWorld.TransformPosition(PositionVertexBuffer.VertexPosition(Index2));
+				FVector WorldVert0 = InStaticMeshComponent->GetComponentTransform().TransformPosition(PositionVertexBuffer.VertexPosition(Index0));
+				FVector WorldVert1 = InStaticMeshComponent->GetComponentTransform().TransformPosition(PositionVertexBuffer.VertexPosition(Index1));
+				FVector WorldVert2 = InStaticMeshComponent->GetComponentTransform().TransformPosition(PositionVertexBuffer.VertexPosition(Index2));
 
 				FLinearColor Color0;
 				FLinearColor Color1;
diff --git a/Engine/Source/Editor/GameProjectGeneration/Private/GameProjectUtils.cpp b/Engine/Source/Editor/GameProjectGeneration/Private/GameProjectUtils.cpp
--- a/Engine/Source/Editor/GameProjectGeneration/Private/GameProjectUtils.cpp
+++ b/Engine/Source/Editor/GameProjectGeneration/Private/GameProjectUtils.cpp
@@ -2923,7 +2923,7 @@ bool GameProjectUtils::GenerateGameModuleBuildFile(const FString& NewBuildFileNa
 	return WriteOutputFile(NewBuildFileName, FinalOutput, OutFailReason);
 }
 
-bool GameProjectUtils::GeneratePluginModuleBuildFile(const FString& NewBuildFileName, const FString& ModuleName, const TArray<FString>& PublicDependencyModuleNames, const TArray<FString>& PrivateDependencyModuleNames, const FString& PchIncPath, FText& OutFailReason)
+bool GameProjectUtils::GeneratePluginModuleBuildFile(const FString& NewBuildFileName, const FString& ModuleName, const TArray<FString>& PublicDependencyModuleNames, const TArray<FString>& PrivateDependencyModuleNames, FText& OutFailReason, bool bUseExplicitOrSharedPCHs/* = true*/)
 {
 	FString Template;
 	if ( !ReadTemplateFile(TEXT("PluginModule.Build.cs.template"), Template, OutFailReason) )
@@ -2936,12 +2936,8 @@ bool GameProjectUtils::GeneratePluginModuleBuildFile(const FString& NewBuildFile
 	FinalOutput = FinalOutput.Replace(TEXT("%PRIVATE_DEPENDENCY_MODULE_NAMES%"), *MakeCommaDelimitedList(PrivateDependencyModuleNames), ESearchCase::CaseSensitive);
 	FinalOutput = FinalOutput.Replace(TEXT("%MODULE_NAME%"), *ModuleName, ESearchCase::CaseSensitive);
 	
-	FString AddedBody;
-	if (!PchIncPath.IsEmpty())
-	{
-		AddedBody = FString::Printf(TEXT("PrivatePCHHeaderFile = \"%s\";"), *PchIncPath);
-	}
-	FinalOutput = FinalOutput.Replace(TEXT("%ADDED_BODY%"), *AddedBody, ESearchCase::CaseSensitive);
+	const FString PCHUsage = bUseExplicitOrSharedPCHs ? TEXT("UseExplicitOrSharedPCHs") : TEXT("UseSharedPCHs");
+	FinalOutput = FinalOutput.Replace(TEXT("%PCH_USAGE%"), *PCHUsage, ESearchCase::CaseSensitive);
 
 	return WriteOutputFile(NewBuildFileName, FinalOutput, OutFailReason);
 }
diff --git a/Engine/Source/Editor/GameProjectGeneration/Public/GameProjectUtils.h b/Engine/Source/Editor/GameProjectGeneration/Public/GameProjectUtils.h
--- a/Engine/Source/Editor/GameProjectGeneration/Public/GameProjectUtils.h
+++ b/Engine/Source/Editor/GameProjectGeneration/Public/GameProjectUtils.h
@@ -226,8 +226,8 @@ public:
 	/** Generates a Build.cs file for a game module */
 	static bool GenerateGameModuleBuildFile(const FString& NewBuildFileName, const FString& ModuleName, const TArray<FString>& PublicDependencyModuleNames, const TArray<FString>& PrivateDependencyModuleNames, FText& OutFailReason);
 
-	/** Generates a Build.cs file for a plugin module */
-	static bool GeneratePluginModuleBuildFile(const FString& NewBuildFileName, const FString& ModuleName, const TArray<FString>& PublicDependencyModuleNames, const TArray<FString>& PrivateDependencyModuleNames, const FString& PCHIncPath, FText& OutFailReason);
+	/** Generates a Build.cs file for a plugin module. Set 'bUseExplicitOrSharedPCHs' to false to disable IWYU conventions. */
+	static bool GeneratePluginModuleBuildFile(const FString& NewBuildFileName, const FString& ModuleName, const TArray<FString>& PublicDependencyModuleNames, const TArray<FString>& PrivateDependencyModuleNames, FText& OutFailReason, bool bUseExplicitOrSharedPCHs = true);
 
 	/** Generates a module .cpp file, intended for plugin use */
 	static bool GeneratePluginModuleCPPFile(const FString& CPPFileName, const FString& ModuleName, const FString& StartupSourceCode, FText& OutFailReason);
diff --git a/Engine/Source/Editor/GameplayTasksEditor/Private/K2Node_LatentGameplayTaskCall.cpp b/Engine/Source/Editor/GameplayTasksEditor/Private/K2Node_LatentGameplayTaskCall.cpp
--- a/Engine/Source/Editor/GameplayTasksEditor/Private/K2Node_LatentGameplayTaskCall.cpp
+++ b/Engine/Source/Editor/GameplayTasksEditor/Private/K2Node_LatentGameplayTaskCall.cpp
@@ -211,7 +211,7 @@ void UK2Node_LatentGameplayTaskCall::CreatePinsForClass(UClass* InClass)
 		{
 
 
-			UEdGraphPin* Pin = CreatePin(EGPD_Input, TEXT(""), TEXT(""), NULL, false, false, Property->GetName());
+			UEdGraphPin* Pin = CreatePin(EGPD_Input, FString(), FString(), nullptr, Property->GetName());
 			check(Pin);
 			const bool bPinGood = K2Schema->ConvertPropertyToPinType(Property, /*out*/ Pin->PinType);
 			SpawnParamPins.Add(Pin->PinName);
@@ -221,7 +221,7 @@ void UK2Node_LatentGameplayTaskCall::CreatePinsForClass(UClass* InClass)
 				FString DefaultValueAsString;
 				const bool bDefaultValueSet = FBlueprintEditorUtils::PropertyValueToString(Property, reinterpret_cast<const uint8*>(ClassDefaultObject), DefaultValueAsString);
 				check(bDefaultValueSet);
-				K2Schema->TrySetDefaultValue(*Pin, DefaultValueAsString);
+				K2Schema->SetPinAutogeneratedDefaultValue(Pin, DefaultValueAsString);
 			}
 
 			// Copy tooltip from the property.
@@ -450,8 +450,8 @@ bool UK2Node_LatentGameplayTaskCall::ConnectSpawnProperties(UClass* ClassToSpawn
 			UFunction* SetByNameFunction = Schema->FindSetVariableByNameFunction(SpawnVarPin->PinType);
 			if (SetByNameFunction)
 			{
-				UK2Node_CallFunction* SetVarNode = NULL;
-				if (SpawnVarPin->PinType.bIsArray)
+				UK2Node_CallFunction* SetVarNode = nullptr;
+				if (SpawnVarPin->PinType.IsArray())
 				{
 					SetVarNode = CompilerContext.SpawnIntermediateNode<UK2Node_CallArrayFunction>(this, SourceGraph);
 				}
@@ -587,7 +587,7 @@ void UK2Node_LatentGameplayTaskCall::ExpandNode(class FKismetCompilerContext& Co
 		{
 			const FEdGraphPinType& PinType = CurrentPin->PinType;
 			UK2Node_TemporaryVariable* TempVarOutput = CompilerContext.SpawnInternalVariable(
-				this, PinType.PinCategory, PinType.PinSubCategory, PinType.PinSubCategoryObject.Get(), PinType.bIsArray, PinType.bIsSet, PinType.bIsMap, PinType.PinValueType);
+				this, PinType.PinCategory, PinType.PinSubCategory, PinType.PinSubCategoryObject.Get(), PinType.ContainerType, PinType.PinValueType);
 			bIsErrorFree &= TempVarOutput->GetVariablePin() && CompilerContext.MovePinLinksToIntermediate(*CurrentPin, *TempVarOutput->GetVariablePin()).CanSafeConnect();
 			VariableOutputs.Add(FBaseAsyncTaskHelper::FOutputPinAndLocalVariable(CurrentPin, TempVarOutput));
 		}
diff --git a/Engine/Source/Editor/GraphEditor/Private/BlueprintConnectionDrawingPolicy.cpp b/Engine/Source/Editor/GraphEditor/Private/BlueprintConnectionDrawingPolicy.cpp
--- a/Engine/Source/Editor/GraphEditor/Private/BlueprintConnectionDrawingPolicy.cpp
+++ b/Engine/Source/Editor/GraphEditor/Private/BlueprintConnectionDrawingPolicy.cpp
@@ -442,8 +442,7 @@ void FKismetConnectionDrawingPolicy::DetermineWiringStyle(UEdGraphPin* OutputPin
 	const bool bDeemphasizeUnhoveredPins = HoveredPins.Num() > 0;
 
 	// If this is a K2 graph, try to be a little more specific
-	const UEdGraphSchema_K2* K2Schema = Cast<const UEdGraphSchema_K2>(Schema);
-	if (K2Schema != NULL)
+	if (const UEdGraphSchema_K2* K2Schema = Cast<const UEdGraphSchema_K2>(Schema))
 	{
 		// If the output or input connect to a knot that is going backwards, we will flip the direction on values going into them
 		{
@@ -512,8 +511,8 @@ void FKismetConnectionDrawingPolicy::DetermineWiringStyle(UEdGraphPin* OutputPin
 		}
 		else
 		{
-			// Array types should draw thicker
-			if ((InputPin && InputPin->PinType.bIsArray) || (OutputPin && OutputPin->PinType.bIsArray))
+			// Container types should draw thicker
+			if ((InputPin && InputPin->PinType.IsContainer()) || (OutputPin && OutputPin->PinType.IsContainer()))
 			{
 				Params.WireThickness = 3.0f;
 			}
diff --git a/Engine/Source/Editor/GraphEditor/Private/GraphDiffControl.cpp b/Engine/Source/Editor/GraphEditor/Private/GraphDiffControl.cpp
--- a/Engine/Source/Editor/GraphEditor/Private/GraphDiffControl.cpp
+++ b/Engine/Source/Editor/GraphEditor/Private/GraphDiffControl.cpp
@@ -153,15 +153,16 @@ static void DiffR_PinTypeChanged(FDiffResults& Results, class UEdGraphPin* Pin2,
 			Diff.DisplayString = FText::Format(LOCTEXT("DIF_PinSubCategoryObjFmt", "Pin SubCategoryObject '{0}' ['{1}' -> '{2}']"), FText::FromString(Pin2->PinName), FText::FromString(Obj1), FText::FromString(Obj2));
 		}
 	}
-	else if(Type1.bIsArray != Type2.bIsArray)
+	else if(Type1.ContainerType != Type2.ContainerType)
 	{
+		// TODO: Make the messaging correct about the nature of the diff
 		Diff.Diff = EDiffType::PIN_TYPE_IS_ARRAY;
 
 		// Only bother setting up the display data if we're storing the result
 		if(Results.CanStoreResults())
 		{
-			FText IsArray1 = Pin1->PinType.bIsArray ? LOCTEXT("true", "true") : LOCTEXT("false", "false");
-			FText IsArray2 = Pin2->PinType.bIsArray ? LOCTEXT("true", "true") : LOCTEXT("false", "false");
+			FText IsArray1 = Pin1->PinType.IsArray() ? LOCTEXT("true", "true") : LOCTEXT("false", "false");
+			FText IsArray2 = Pin2->PinType.IsArray() ? LOCTEXT("true", "true") : LOCTEXT("false", "false");
 
 			Diff.ToolTip = FText::Format(LOCTEXT("DIF_PinIsArrayToolTipFmt", "PinType IsArray for '{0}' modified. Was '{1}', but is now '{2}"), FText::FromString(Pin2->PinName), IsArray1, IsArray2);
 			Diff.DisplayString = FText::Format(LOCTEXT("DIF_PinIsArrayFmt", "Pin IsArray '{0}' ['{1}' -> '{2}']"), FText::FromString(Pin2->PinName), IsArray1, IsArray2);
@@ -410,7 +411,7 @@ static bool IsPinTypeDifferent(const FEdGraphPinType& T1, const FEdGraphPinType&
 {
 	bool bIsDifferent =  (T1.PinCategory != T2.PinCategory) 
 		|| (T1.PinSubCategory != T2.PinSubCategory) 
-		|| (T1.bIsArray != T2.bIsArray) 
+		|| (T1.ContainerType != T2.ContainerType) 
 		|| (T1.bIsReference != T2.bIsReference);
 
 	const UObject* T1Obj = T1.PinSubCategoryObject.Get();
diff --git a/Engine/Source/Editor/GraphEditor/Private/GraphEditorActions.cpp b/Engine/Source/Editor/GraphEditor/Private/GraphEditorActions.cpp
--- a/Engine/Source/Editor/GraphEditor/Private/GraphEditorActions.cpp
+++ b/Engine/Source/Editor/GraphEditor/Private/GraphEditorActions.cpp
@@ -75,6 +75,7 @@ void FGraphEditorCommandsImpl::RegisterCommands()
 	UI_COMMAND( RecombineStructPin, "Recombine Struct Pin", "Takes struct pins that have been broken in to composite elements and combines them back to a single struct pin", EUserInterfaceActionType::Button, FInputChord() )
 	UI_COMMAND( StartWatchingPin, "Watch this value", "Adds this pin or variable to the watch list", EUserInterfaceActionType::Button, FInputChord() )
 	UI_COMMAND( StopWatchingPin, "Stop watching this value", "Removes this pin or variable from the watch list ", EUserInterfaceActionType::Button, FInputChord() )
+	UI_COMMAND( ResetPinToDefaultValue, "Reset to Default Value", "Reset value of this pin to the default", EUserInterfaceActionType::Button, FInputChord() )
 
 	UI_COMMAND( SelectBone, "Select Bone", "Assign or change the bone for SkeletalControls", EUserInterfaceActionType::Button, FInputChord() )
 	UI_COMMAND( AddBlendListPin, "Add Blend Pin", "Add Blend Pin to BlendList", EUserInterfaceActionType::Button, FInputChord() )
diff --git a/Engine/Source/Editor/GraphEditor/Private/KismetNodes/SGraphNodeSwitchStatement.cpp b/Engine/Source/Editor/GraphEditor/Private/KismetNodes/SGraphNodeSwitchStatement.cpp
--- a/Engine/Source/Editor/GraphEditor/Private/KismetNodes/SGraphNodeSwitchStatement.cpp
+++ b/Engine/Source/Editor/GraphEditor/Private/KismetNodes/SGraphNodeSwitchStatement.cpp
@@ -26,7 +26,7 @@ public:
 	{
 		SGraphPin::Construct(SGraphPin::FArguments().PinLabelStyle(FName("Graph.Node.DefaultPinName")), InPin);
 
-		bWasEventPin = false;
+		CachePinIcons();
 	}
 };
 
diff --git a/Engine/Source/Editor/GraphEditor/Private/KismetPins/SGraphPinClass.cpp b/Engine/Source/Editor/GraphEditor/Private/KismetPins/SGraphPinClass.cpp
--- a/Engine/Source/Editor/GraphEditor/Private/KismetPins/SGraphPinClass.cpp
+++ b/Engine/Source/Editor/GraphEditor/Private/KismetPins/SGraphPinClass.cpp
@@ -9,6 +9,7 @@
 #include "ClassViewerModule.h"
 #include "ClassViewerFilter.h"
 #include "ScopedTransaction.h"
+#include "AssetRegistryModule.h"
 
 #define LOCTEXT_NAMESPACE "SGraphPinClass"
 
@@ -136,4 +137,54 @@ FText SGraphPinClass::GetDefaultComboText() const
 	return LOCTEXT( "DefaultComboText", "Select Class" );
 }
 
+const FAssetData& SGraphPinClass::GetAssetData(bool bRuntimePath) const
+{
+	if (bRuntimePath)
+	{
+		// For runtime use the default path
+		return SGraphPinObject::GetAssetData(bRuntimePath);
+	}
+
+	FString CachedRuntimePath = CachedEditorAssetData.ObjectPath.ToString() + TEXT("_C");
+
+	if (GraphPinObj->DefaultObject)
+	{
+		if (!GraphPinObj->DefaultObject->GetPathName().Equals(CachedRuntimePath, ESearchCase::CaseSensitive))
+		{
+			// This will cause it to use the UBlueprint
+			CachedEditorAssetData = FAssetData(GraphPinObj->DefaultObject, false);
+		}
+	}
+	else if (!GraphPinObj->DefaultValue.IsEmpty())
+	{
+		if (!GraphPinObj->DefaultValue.Equals(CachedRuntimePath, ESearchCase::CaseSensitive))
+		{
+			FString EditorPath = GraphPinObj->DefaultValue;
+			EditorPath.RemoveFromEnd(TEXT("_C"));
+			const FAssetRegistryModule& AssetRegistryModule = FModuleManager::LoadModuleChecked<FAssetRegistryModule>(TEXT("AssetRegistry"));
+
+			CachedEditorAssetData = AssetRegistryModule.Get().GetAssetByObjectPath(FName(*EditorPath));
+
+			if (!CachedEditorAssetData.IsValid())
+			{
+				FString PackageName = FPackageName::ObjectPathToPackageName(EditorPath);
+				FString PackagePath = FPackageName::GetLongPackagePath(PackageName);
+				FString ObjectName = FPackageName::ObjectPathToObjectName(EditorPath);
+
+				// Fake one
+				CachedEditorAssetData = FAssetData(FName(*PackageName), FName(*PackagePath), FName(*ObjectName), UObject::StaticClass()->GetFName());
+			}
+		}
+	}
+	else
+	{
+		if (CachedEditorAssetData.IsValid())
+		{
+			CachedEditorAssetData = FAssetData();
+		}
+	}
+
+	return CachedEditorAssetData;
+}
+
 #undef LOCTEXT_NAMESPACE
diff --git a/Engine/Source/Editor/GraphEditor/Private/KismetPins/SGraphPinExec.cpp b/Engine/Source/Editor/GraphEditor/Private/KismetPins/SGraphPinExec.cpp
--- a/Engine/Source/Editor/GraphEditor/Private/KismetPins/SGraphPinExec.cpp
+++ b/Engine/Source/Editor/GraphEditor/Private/KismetPins/SGraphPinExec.cpp
@@ -8,8 +8,12 @@ void SGraphPinExec::Construct(const FArguments& InArgs, UEdGraphPin* InPin)
 {
 	SGraphPin::Construct(SGraphPin::FArguments(), InPin);
 
-	bWasEventPin = false;
+	// Call utility function so inheritors can also call it since arguments can be passed through
+	CachePinIcons();
+}
 
+void SGraphPinExec::CachePinIcons()
+{
 	CachedImg_Pin_ConnectedHovered = FEditorStyle::GetBrush(TEXT("Graph.ExecPin.ConnectedHovered"));
 	CachedImg_Pin_Connected = FEditorStyle::GetBrush(TEXT("Graph.ExecPin.Connected"));
 	CachedImg_Pin_DisconnectedHovered = FEditorStyle::GetBrush(TEXT("Graph.ExecPin.DisconnectedHovered"));
diff --git a/Engine/Source/Editor/GraphEditor/Private/KismetPins/SGraphPinObject.cpp b/Engine/Source/Editor/GraphEditor/Private/KismetPins/SGraphPinObject.cpp
--- a/Engine/Source/Editor/GraphEditor/Private/KismetPins/SGraphPinObject.cpp
+++ b/Engine/Source/Editor/GraphEditor/Private/KismetPins/SGraphPinObject.cpp
@@ -13,6 +13,7 @@
 #include "ContentBrowserModule.h"
 #include "ScopedTransaction.h"
 #include "Engine/Selection.h"
+#include "AssetRegistryModule.h"
 
 #define LOCTEXT_NAMESPACE "SGraphPinObject"
 
@@ -41,8 +42,6 @@ TSharedRef<SWidget>	SGraphPinObject::GetDefaultValueWidget()
 {
 	if( AllowSelfPinWidget() )
 	{
-		UObject* DefaultObject = GraphPinObj->DefaultObject;
-
 		if(GraphPinObj->GetSchema()->IsSelfPin(*GraphPinObj))
 		{
 			return SNew(SEditableTextBox)
@@ -162,10 +161,11 @@ FReply SGraphPinObject::OnClickUse()
 
 FReply SGraphPinObject::OnClickBrowse()
 {
-	if (GraphPinObj->DefaultObject != NULL)
+	const FAssetData& AssetData = GetAssetData(false);
+	if (AssetData.IsValid())
 	{
-		TArray<UObject*> Objects;
-		Objects.Add(GraphPinObj->DefaultObject);
+		TArray<FAssetData> Objects;
+		Objects.Add(AssetData);
 
 		GEditor->SyncBrowserToObjects(Objects);
 	}
@@ -223,8 +223,8 @@ TSharedRef<SWidget> SGraphPinObject::GenerateAssetPicker()
 
 void SGraphPinObject::OnAssetSelectedFromPicker(const class FAssetData& AssetData)
 {
-	UObject* AssetObject = AssetData.GetAsset();
-	if(GraphPinObj->DefaultObject != AssetObject)
+	const FAssetData& CurrentAssetData = GetAssetData(true);
+	if(CurrentAssetData != AssetData)
 	{
 		const FScopedTransaction Transaction( NSLOCTEXT("GraphEditor", "ChangeObjectPinValue", "Change Object Pin Value" ) );
 		GraphPinObj->Modify();
@@ -233,18 +233,18 @@ void SGraphPinObject::OnAssetSelectedFromPicker(const class FAssetData& AssetDat
 		AssetPickerAnchor->SetIsOpen(false);
 
 		// Set the object found from the asset picker
-		GraphPinObj->GetSchema()->TrySetDefaultObject(*GraphPinObj, AssetObject);
+		GraphPinObj->GetSchema()->TrySetDefaultValue(*GraphPinObj, AssetData.ObjectPath.ToString());
 	}
 }
 
 
 FText SGraphPinObject::GetValue() const
 {
-	UObject* DefaultObject = GraphPinObj->DefaultObject;
+	const FAssetData& CurrentAssetData = GetAssetData(true);
 	FText Value;
-	if (DefaultObject != NULL)
+	if (CurrentAssetData.IsValid())
 	{
-		Value =  FText::FromString(DefaultObject->GetFullName());
+		Value = FText::FromString(CurrentAssetData.GetFullName());
 	}
 	else
 	{
@@ -266,10 +266,10 @@ FText SGraphPinObject::GetObjectName() const
 	
 	if (GraphPinObj != NULL)
 	{
-		UObject* DefaultObject = GraphPinObj->DefaultObject;
-		if (DefaultObject != NULL)
+		const FAssetData& CurrentAssetData = GetAssetData(true);
+		if (CurrentAssetData.IsValid())
 		{
-			Value = FText::FromString(DefaultObject->GetName());
+			Value = FText::FromString(CurrentAssetData.AssetName.ToString());
 			int32 StringLen = Value.ToString().Len();
 
 			//If string is too long, then truncate (eg. "abcdefgijklmnopq" is converted as "abcd...nopq")
@@ -298,16 +298,18 @@ FText SGraphPinObject::OnGetComboTextValue() const
 {
 	FText Value = GetDefaultComboText();
 	
-	if (GraphPinObj != NULL)
+	if (GraphPinObj != nullptr)
 	{
+		const FAssetData& CurrentAssetData = GetAssetData(true);
+
 		UObject* DefaultObject = GraphPinObj->DefaultObject;
 		if (UField* Field = Cast<UField>(DefaultObject))
 		{
 			Value = Field->GetDisplayNameText();
 		}
-		else if (DefaultObject != NULL)
+		else if (CurrentAssetData.IsValid())
 		{
-			Value = FText::FromString(DefaultObject->GetName());
+			Value = FText::FromString(CurrentAssetData.AssetName.ToString());
 		}
 	}
 	return Value;
@@ -315,20 +317,62 @@ FText SGraphPinObject::OnGetComboTextValue() const
 
 FSlateColor SGraphPinObject::OnGetComboForeground() const
 {
-	float Alpha = IsHovered() ? GraphPinObjectDefs::ActiveComboAlpha : GraphPinObjectDefs::InActiveComboAlpha;
+	float Alpha = (IsHovered() || bOnlyShowDefaultValue) ? GraphPinObjectDefs::ActiveComboAlpha : GraphPinObjectDefs::InActiveComboAlpha;
 	return FSlateColor( FLinearColor( 1.f, 1.f, 1.f, Alpha ) );
 }
 
 FSlateColor SGraphPinObject::OnGetWidgetForeground() const
 {
-	float Alpha = IsHovered() ? GraphPinObjectDefs::ActivePinForegroundAlpha : GraphPinObjectDefs::InactivePinForegroundAlpha;
+	float Alpha = (IsHovered() || bOnlyShowDefaultValue) ? GraphPinObjectDefs::ActivePinForegroundAlpha : GraphPinObjectDefs::InactivePinForegroundAlpha;
 	return FSlateColor( FLinearColor( 1.f, 1.f, 1.f, Alpha ) );
 }
 
 FSlateColor SGraphPinObject::OnGetWidgetBackground() const
 {
-	float Alpha = IsHovered() ? GraphPinObjectDefs::ActivePinBackgroundAlpha : GraphPinObjectDefs::InactivePinBackgroundAlpha;
+	float Alpha = (IsHovered() || bOnlyShowDefaultValue) ? GraphPinObjectDefs::ActivePinBackgroundAlpha : GraphPinObjectDefs::InactivePinBackgroundAlpha;
 	return FSlateColor( FLinearColor( 1.f, 1.f, 1.f, Alpha ) );
 }
 
+const FAssetData& SGraphPinObject::GetAssetData(bool bRuntimePath) const
+{
+	// For normal assets, the editor and runtime path are the same
+	if (GraphPinObj->DefaultObject)
+	{
+		if (!GraphPinObj->DefaultObject->GetPathName().Equals(CachedAssetData.ObjectPath.ToString(), ESearchCase::CaseSensitive))
+		{
+			// This always uses the exact object pointed at
+			CachedAssetData = FAssetData(GraphPinObj->DefaultObject, true);
+		}
+	}
+	else if (!GraphPinObj->DefaultValue.IsEmpty())
+	{
+		FName ObjectPath = FName(*GraphPinObj->DefaultValue);
+		if (ObjectPath != CachedAssetData.ObjectPath)
+		{
+			const FAssetRegistryModule& AssetRegistryModule = FModuleManager::LoadModuleChecked<FAssetRegistryModule>(TEXT("AssetRegistry"));
+
+			CachedAssetData = AssetRegistryModule.Get().GetAssetByObjectPath(ObjectPath);
+
+			if (!CachedAssetData.IsValid())
+			{
+				FString PackageName = FPackageName::ObjectPathToPackageName(GraphPinObj->DefaultValue);
+				FString PackagePath = FPackageName::GetLongPackagePath(PackageName);
+				FString ObjectName = FPackageName::ObjectPathToObjectName(GraphPinObj->DefaultValue);
+
+				// Fake one
+				CachedAssetData = FAssetData(FName(*PackageName), FName(*PackagePath), FName(*ObjectName), UObject::StaticClass()->GetFName());
+			}
+		}
+	}
+	else
+	{
+		if (CachedAssetData.IsValid())
+		{
+			CachedAssetData = FAssetData();
+		}
+	}
+
+	return CachedAssetData;
+}
+
 #undef LOCTEXT_NAMESPACE
diff --git a/Engine/Source/Editor/GraphEditor/Private/SGraphPin.cpp b/Engine/Source/Editor/GraphEditor/Private/SGraphPin.cpp
--- a/Engine/Source/Editor/GraphEditor/Private/SGraphPin.cpp
+++ b/Engine/Source/Editor/GraphEditor/Private/SGraphPin.cpp
@@ -68,6 +68,7 @@ TSharedPtr<SGraphPin> FGraphPinHandle::FindInGraphPanel(const SGraphPanel& InPan
 SGraphPin::SGraphPin()
 	: GraphPinObj(nullptr)
 	, bShowLabel(true)
+	, bOnlyShowDefaultValue(false)
 	, bIsMovingLinks(false)
 	, PinColorModifier(FLinearColor::White)
 	, CachedNodeOffset(FVector2D::ZeroVector)
@@ -302,16 +303,14 @@ void SGraphPin::Construct(const FArguments& InArgs, UEdGraphPin* InPin)
 
 TSharedRef<SWidget>	SGraphPin::GetDefaultValueWidget()
 {
-	return SNew(SBox);
+	return SNullWidget::NullWidget;
 }
 
-
 void SGraphPin::SetIsEditable(TAttribute<bool> InIsEditable)
 {
 	IsEditable = InIsEditable;
 }
 
-
 FReply SGraphPin::OnPinMouseDown( const FGeometry& SenderGeometry, const FPointerEvent& MouseEvent )
 {
 	bIsMovingLinks = false;
@@ -827,22 +826,17 @@ EEdGraphPinDirection SGraphPin::GetDirection() const
 
 bool SGraphPin::IsArray() const
 {
-	return GraphPinObj->PinType.bIsArray;
+	return GraphPinObj->PinType.IsArray();
 }
 
 bool SGraphPin::IsSet() const
 {
-	return GraphPinObj->PinType.bIsSet;
+	return GraphPinObj->PinType.IsSet();
 }
 
 bool SGraphPin::IsMap() const
 {
-	return GraphPinObj->PinType.bIsMap;
-}
-
-bool SGraphPin::IsByRef() const
-{
-	return GraphPinObj->PinType.bIsReference;
+	return GraphPinObj->PinType.IsMap();
 }
 
 bool SGraphPin::IsByMutableRef() const
@@ -856,11 +850,6 @@ bool SGraphPin::IsDelegate() const
 	return Schema && Schema->IsDelegateCategory(GraphPinObj->PinType.PinCategory);
 }
 
-bool SGraphPin::IsByConstRef() const
-{
-	return GraphPinObj->PinType.bIsReference && GraphPinObj->PinType.bIsConst;
-}
-
 /** @return whether this pin is connected to another pin */
 bool SGraphPin::IsConnected() const
 {
@@ -927,7 +916,7 @@ const FSlateBrush* SGraphPin::GetPinIcon() const
 
 const FSlateBrush* SGraphPin::GetSecondaryPinIcon() const
 {
-	if( !GraphPinObj->IsPendingKill() && GraphPinObj->PinType.bIsMap )
+	if( !GraphPinObj->IsPendingKill() && GraphPinObj->PinType.IsMap() )
 	{
 		return CachedImg_MapPinValue;
 	}
@@ -944,7 +933,7 @@ const FSlateBrush* SGraphPin::GetPinBorder() const
 		bIsMarkedPin = (OwnerPanelPtr->MarkedPin.Pin() == SharedThis(this));
 	}
 
-	return (IsHovered() || bIsMarkedPin || GraphPinObj->bIsDiffing) ? CachedImg_Pin_BackgroundHovered : CachedImg_Pin_Background;
+	return (IsHovered() || bIsMarkedPin || GraphPinObj->bIsDiffing || bOnlyShowDefaultValue) ? CachedImg_Pin_BackgroundHovered : CachedImg_Pin_Background;
 }
 
 
@@ -1044,6 +1033,12 @@ FReply SGraphPin::ClickedOnPinStatusIcon()
 
 EVisibility SGraphPin::GetDefaultValueVisibility() const
 {
+	// If this is only for showing default value, always show
+	if (bOnlyShowDefaultValue)
+	{
+		return EVisibility::Visible;
+	}
+
 	// First ask schema
 	const UEdGraphSchema* Schema = !GraphPinObj->IsPendingKill() ? GraphPinObj->GetSchema() : nullptr;
 	if (Schema == nullptr || Schema->ShouldHidePinDefaultValue(GraphPinObj))
@@ -1073,6 +1068,11 @@ void SGraphPin::SetShowLabel(bool bNewShowLabel)
 	bShowLabel = bNewShowLabel;
 }
 
+void SGraphPin::SetOnlyShowDefaultValue(bool bNewOnlyShowDefaultValue)
+{
+	bOnlyShowDefaultValue = bNewOnlyShowDefaultValue;
+}
+
 FText SGraphPin::GetTooltipText() const
 {
 	FText HoverText = FText::GetEmpty();
diff --git a/Engine/Source/Editor/GraphEditor/Public/GraphEditorActions.h b/Engine/Source/Editor/GraphEditor/Public/GraphEditorActions.h
--- a/Engine/Source/Editor/GraphEditor/Public/GraphEditorActions.h
+++ b/Engine/Source/Editor/GraphEditor/Public/GraphEditorActions.h
@@ -104,6 +104,7 @@ public:
 	TSharedPtr< FUICommandInfo > RecombineStructPin;
 	TSharedPtr< FUICommandInfo > StartWatchingPin;
 	TSharedPtr< FUICommandInfo > StopWatchingPin;
+	TSharedPtr< FUICommandInfo > ResetPinToDefaultValue;
 
 	// SkeletalControl specific commands
 	TSharedPtr< FUICommandInfo > SelectBone;
diff --git a/Engine/Source/Editor/GraphEditor/Public/KismetPins/SGraphPinExec.h b/Engine/Source/Editor/GraphEditor/Public/KismetPins/SGraphPinExec.h
--- a/Engine/Source/Editor/GraphEditor/Public/KismetPins/SGraphPinExec.h
+++ b/Engine/Source/Editor/GraphEditor/Public/KismetPins/SGraphPinExec.h
@@ -21,9 +21,8 @@ protected:
 	virtual const FSlateBrush* GetPinIcon() const override;
 	//~ End SGraphPin Interface
 
-protected:
-	mutable bool bWasEventPin;
+	void CachePinIcons();
 
-	mutable const FSlateBrush* CachedImg_Pin_ConnectedHovered;
-	mutable const FSlateBrush* CachedImg_Pin_DisconnectedHovered;
+	const FSlateBrush* CachedImg_Pin_ConnectedHovered;
+	const FSlateBrush* CachedImg_Pin_DisconnectedHovered;
 };
diff --git a/Engine/Source/Editor/GraphEditor/Public/KismetPins/SGraphPinObject.h b/Engine/Source/Editor/GraphEditor/Public/KismetPins/SGraphPinObject.h
--- a/Engine/Source/Editor/GraphEditor/Public/KismetPins/SGraphPinObject.h
+++ b/Engine/Source/Editor/GraphEditor/Public/KismetPins/SGraphPinObject.h
@@ -60,6 +60,9 @@ protected:
 	/** Button Color and Opacity delegate */
 	FSlateColor OnGetWidgetBackground() const;
 
+	/** Returns asset data of currently selected object, if bRuntimePath is true this will include _C for blueprint classes, for false it will point to UBlueprint instead */
+	virtual const FAssetData& GetAssetData(bool bRuntimePath) const;
+
 protected:
 	/** Object manipulator buttons. */
 	TSharedPtr<SButton> UseButton;
@@ -67,4 +70,7 @@ protected:
 
 	/** Menu anchor for opening and closing the asset picker */
 	TSharedPtr<class SMenuAnchor> AssetPickerAnchor;
+
+	/** Cached AssetData of object selected */
+	mutable FAssetData CachedAssetData;
 };
diff --git a/Engine/Source/Editor/GraphEditor/Public/NodeFactory.h b/Engine/Source/Editor/GraphEditor/Public/NodeFactory.h
new file mode 100644
--- /dev/null
+++ b/Engine/Source/Editor/GraphEditor/Public/NodeFactory.h
@@ -0,0 +1,29 @@
+// Copyright 1998-2017 Epic Games, Inc. All Rights Reserved.
+
+#pragma once
+
+#include "CoreMinimal.h"
+
+class FConnectionDrawingPolicy;
+class FSlateRect;
+class FSlateWindowElementList;
+class SGraphNode;
+class SGraphPin;
+class UEdGraph;
+class UEdGraphNode;
+class UEdGraphPin;
+class UEdGraphSchema;
+
+/** Class that decides which widget type to create for a given data object */
+class GRAPHEDITOR_API FNodeFactory
+{
+public:
+	/** Create a widget for the supplied node */
+	static TSharedPtr<SGraphNode> CreateNodeWidget(UEdGraphNode* InNode);
+
+	/** Create a widget for the supplied pin */
+	static TSharedPtr<SGraphPin> CreatePinWidget(UEdGraphPin* InPin);
+
+    /** Create a pin connection factory for the supplied schema */
+    static FConnectionDrawingPolicy* CreateConnectionPolicy(const UEdGraphSchema* Schema, int32 InBackLayerID, int32 InFrontLayerID, float ZoomFactor, const FSlateRect& InClippingRect, FSlateWindowElementList& InDrawElements, UEdGraph* InGraphObj);
+};
