Copying //UE4/Dev-Framework to //UE4/Dev-Main (Source: //UE4/Dev-Framework @ 3510040)

diff --git a/Engine/Source/Developer/AssetTools/Private/AssetTypeActions/AssetTypeActions_Blueprint.cpp b/Engine/Source/Developer/AssetTools/Private/AssetTypeActions/AssetTypeActions_Blueprint.cpp
--- a/Engine/Source/Developer/AssetTools/Private/AssetTypeActions/AssetTypeActions_Blueprint.cpp
+++ b/Engine/Source/Developer/AssetTools/Private/AssetTypeActions/AssetTypeActions_Blueprint.cpp
@@ -22,13 +22,13 @@
 
 void FAssetTypeActions_Blueprint::GetActions( const TArray<UObject*>& InObjects, FMenuBuilder& MenuBuilder )
 {
-	auto Blueprints = GetTypedWeakObjectPtrs<UBlueprint>(InObjects);
+	TArray<TWeakObjectPtr<UBlueprint>> Blueprints = GetTypedWeakObjectPtrs<UBlueprint>(InObjects);
 	
 	if (Blueprints.Num() > 1)
 	{
 		// Ensure that all the selected blueprints are actors
 		bool bCanEditSharedDefaults = true;
-		for (auto Blueprint : Blueprints)
+		for (const TWeakObjectPtr<UBlueprint>& Blueprint : Blueprints)
 		{
 			if (!Blueprint.Get()->ParentClass->IsChildOf(AActor::StaticClass()))
 			{
@@ -73,10 +73,9 @@ void FAssetTypeActions_Blueprint::OpenAssetEditor( const TArray<UObject*>& InObj
 {
 	EToolkitMode::Type Mode = EditWithinLevelEditor.IsValid() ? EToolkitMode::WorldCentric : EToolkitMode::Standalone;
 
-	for (auto ObjIt = InObjects.CreateConstIterator(); ObjIt; ++ObjIt)
+	for (UObject* Object : InObjects)
 	{
-		auto Blueprint = Cast<UBlueprint>(*ObjIt);
-		if (Blueprint)
+		if (UBlueprint* Blueprint = Cast<UBlueprint>(Object))
 		{
 			bool bLetOpen = true;
 			if (!Blueprint->SkeletonGeneratedClass || !Blueprint->GeneratedClass)
@@ -149,10 +148,9 @@ void FAssetTypeActions_Blueprint::ExecuteEditDefaults(TArray<TWeakObjectPtr<UBlu
 	FMessageLog EditorErrors("EditorErrors");
 	EditorErrors.NewPage(LOCTEXT("ExecuteEditDefaultsNewLogPage", "Loading Blueprints"));
 
-	for (auto ObjIt = Objects.CreateConstIterator(); ObjIt; ++ObjIt)
+	for (const TWeakObjectPtr<UBlueprint>& WeakBP : Objects)
 	{
-		auto Object = (*ObjIt).Get();
-		if ( Object )
+		if (UBlueprint* Object = WeakBP.Get())
 		{
 			// If the blueprint is valid, allow it to be added to the list, otherwise log the error.
 			if ( Object->SkeletonGeneratedClass && Object->GeneratedClass )
@@ -180,8 +178,7 @@ void FAssetTypeActions_Blueprint::ExecuteEditDefaults(TArray<TWeakObjectPtr<UBlu
 
 void FAssetTypeActions_Blueprint::ExecuteNewDerivedBlueprint(TWeakObjectPtr<UBlueprint> InObject)
 {
-	auto Object = InObject.Get();
-	if ( Object )
+	if (UBlueprint* Object = InObject.Get())
 	{
 		// The menu option should ONLY be available if there is only one blueprint selected, validated by the menu creation code
 		UBlueprint* TargetParentBP = Object;
@@ -219,8 +216,8 @@ FText FAssetTypeActions_Blueprint::GetNewDerivedBlueprintTooltip(TWeakObjectPtr<
 
 bool FAssetTypeActions_Blueprint::CanExecuteNewDerivedBlueprint(TWeakObjectPtr<UBlueprint> InObject)
 {
-	auto BP = InObject.Get();
-	auto BPGC = BP ? BP->GeneratedClass : NULL;
+	UBlueprint* BP = InObject.Get();
+	UClass* BPGC = BP ? BP->GeneratedClass : nullptr;
 	return BPGC && !BPGC->HasAnyClassFlags(CLASS_Deprecated);
 }
 
diff --git a/Engine/Source/Developer/BlueprintCompilerCppBackend/Private/BlueprintCompilerCppBackendValueHelper.cpp b/Engine/Source/Developer/BlueprintCompilerCppBackend/Private/BlueprintCompilerCppBackendValueHelper.cpp
--- a/Engine/Source/Developer/BlueprintCompilerCppBackend/Private/BlueprintCompilerCppBackendValueHelper.cpp
+++ b/Engine/Source/Developer/BlueprintCompilerCppBackend/Private/BlueprintCompilerCppBackendValueHelper.cpp
@@ -145,8 +145,16 @@ void FEmitDefaultValueHelper::OuterGenerate(FEmitterLocalContext& Context
 			const UObject* DefaultSubobjectInstance = nullptr;
 			if (IsInstancedSubobjectLambda(ArrayIndex, SubobjectInstance, DefaultSubobjectInstance))
 			{
+				// Emit code to create subobjects that were not originally instanced with CreateDefaultSubobject() (e.g. - 'EditInlineNew' instances).
+				if (!SubobjectInstance->HasAnyFlags(RF_DefaultSubObject) && SubobjectInstance->HasAnyFlags(RF_ArchetypeObject))
+				{
+					const uint8* ValuePtr = Property->ContainerPtrToValuePtr<uint8>(DataContainer, ArrayIndex);
+					const uint8* DefaultValuePtr = OptionalDefaultDataContainer ? Property->ContainerPtrToValuePtr<uint8>(OptionalDefaultDataContainer, ArrayIndex) : nullptr;
+					InnerGenerate(Context, Property, PathToMember, ValuePtr, DefaultValuePtr);
+				}
+
 				// Recursively emit property values for nested default subobjects.
-				if (SubobjectInstance->HasAnyFlags(RF_DefaultSubObject) && !SubobjectInstance->GetOuter()->HasAnyFlags(RF_ClassDefaultObject))
+				if (SubobjectInstance->HasAnyFlags(RF_DefaultSubObject|RF_ArchetypeObject) && !SubobjectInstance->GetOuter()->HasAnyFlags(RF_ClassDefaultObject))
 				{
 					check(SubobjectInstance != nullptr);
 
@@ -757,13 +765,17 @@ FString FEmitDefaultValueHelper::HandleSpecialTypes(FEmitterLocalContext& Contex
 				}
 			}
 
-			auto BPGC = Context.GetCurrentlyGeneratedClass();
+			UClass* BPGC = Context.GetCurrentlyGeneratedClass();
+
+			UChildActorComponent* OuterCAC = Cast<UChildActorComponent>(Object->GetOuter());
+			const bool bObjectIsCACTemplate = OuterCAC && OuterCAC->IsIn(BPGC) && OuterCAC->GetChildActorTemplate() == Object;
+
 			const bool bCreatingSubObjectsOfClass = (Context.CurrentCodeType == FEmitterLocalContext::EGeneratedCodeType::SubobjectsOfClass);
 			{
 				auto CDO = BPGC ? BPGC->GetDefaultObject(false) : nullptr;
 				if (BPGC && Object && CDO && Object->IsIn(BPGC) && !Object->IsIn(CDO) && bCreatingSubObjectsOfClass)
 				{
-					return HandleClassSubobject(Context, Object, FEmitterLocalContext::EClassSubobjectList::MiscConvertedSubobjects, true, true);
+					return HandleClassSubobject(Context, Object, FEmitterLocalContext::EClassSubobjectList::MiscConvertedSubobjects, true, true, bObjectIsCACTemplate);
 				}
 			}
 
@@ -776,17 +788,13 @@ FString FEmitDefaultValueHelper::HandleSpecialTypes(FEmitterLocalContext& Contex
 				}
 			}
 
-			if (!bCreatingSubObjectsOfClass && Object->IsIn(BPGC))
+			if (!bCreatingSubObjectsOfClass && bObjectIsCACTemplate)
 			{
-				UChildActorComponent* OuterCAC = Cast<UChildActorComponent>(Object->GetOuter());
-				if (OuterCAC && OuterCAC->IsIn(BPGC) && OuterCAC->GetChildActorTemplate() == Object)
+				Context.TemplateFromSubobjectsOfClass.AddUnique(Object);
+				const FString MappedObject = Context.FindGloballyMappedObject(Object, ObjectClassToUse);
+				if (!MappedObject.IsEmpty())
 				{
-					Context.TemplateFromSubobjectsOfClass.AddUnique(Object);
-					const FString MappedObject = Context.FindGloballyMappedObject(Object, ObjectClassToUse);
-					if (!MappedObject.IsEmpty())
-					{
-						return MappedObject;
-					}
+					return MappedObject;
 				}
 			}
 		}
@@ -1945,8 +1953,18 @@ FString FEmitDefaultValueHelper::HandleInstancedSubobject(FEmitterLocalContext&
 	{
 		if (bCreateInstance)
 		{
-			Context.AddLine(FString::Printf(TEXT("auto %s = CreateDefaultSubobject<%s>(TEXT(\"%s\"));")
-				, *LocalNativeName, *FEmitHelper::GetCppName(ObjectClass), *Object->GetName()));
+			if (Object->HasAnyFlags(RF_DefaultSubObject))
+			{
+				Context.AddLine(FString::Printf(TEXT("auto %s = CreateDefaultSubobject<%s>(TEXT(\"%s\"));")
+					, *LocalNativeName, *FEmitHelper::GetCppName(ObjectClass), *Object->GetName()));
+			}
+			else
+			{
+				check(Object->HasAnyFlags(RF_ArchetypeObject));
+
+				Context.AddLine(FString::Printf(TEXT("auto %s = NewObject<%s>(this, TEXT(\"%s\"), GetMaskedFlags(RF_PropagateToSubObjects) | RF_ArchetypeObject);")
+					, *LocalNativeName, *FEmitHelper::GetCppName(ObjectClass), *Object->GetName()));
+			}
 		}
 		else
 		{
diff --git a/Engine/Source/Developer/FunctionalTesting/Private/AutomationBlueprintFunctionLibrary.cpp b/Engine/Source/Developer/FunctionalTesting/Private/AutomationBlueprintFunctionLibrary.cpp
--- a/Engine/Source/Developer/FunctionalTesting/Private/AutomationBlueprintFunctionLibrary.cpp
+++ b/Engine/Source/Developer/FunctionalTesting/Private/AutomationBlueprintFunctionLibrary.cpp
@@ -355,7 +355,7 @@ void UAutomationBlueprintFunctionLibrary::TakeAutomationScreenshot(UObject* Worl
 {
 	if ( GIsAutomationTesting )
 	{
-		if ( UWorld* World = GEngine->GetWorldFromContextObject(WorldContextObject) )
+		if ( UWorld* World = GEngine->GetWorldFromContextObject(WorldContextObject, EGetWorldErrorMode::LogAndReturnNull) )
 		{
 			FLatentActionManager& LatentActionManager = World->GetLatentActionManager();
 			if ( LatentActionManager.FindExistingAction<FTakeScreenshotAfterTimeLatentAction>(LatentInfo.CallbackTarget, LatentInfo.UUID) == nullptr )
@@ -396,7 +396,7 @@ void UAutomationBlueprintFunctionLibrary::TakeAutomationScreenshotAtCamera(UObje
 		ScreenshotName = NameOverride;
 	}
 
-	if ( UWorld* World = GEngine->GetWorldFromContextObject(WorldContextObject) )
+	if ( UWorld* World = GEngine->GetWorldFromContextObject(WorldContextObject, EGetWorldErrorMode::LogAndReturnNull) )
 	{
 		ScreenshotName = FString::Printf(TEXT("%s_%s"), *World->GetName(), *ScreenshotName);
 
@@ -412,7 +412,7 @@ bool UAutomationBlueprintFunctionLibrary::TakeAutomationScreenshotOfUI_Immediate
 {
 	UAutomationBlueprintFunctionLibrary::FinishLoadingBeforeScreenshot();
 
-	if (UWorld* World = WorldContextObject->GetWorld())
+	if (UWorld* World = GEngine->GetWorldFromContextObject(WorldContextObject, EGetWorldErrorMode::LogAndReturnNull))
 	{
 		if (UGameViewportClient* GameViewport = WorldContextObject->GetWorld()->GetGameViewport())
 		{
@@ -425,9 +425,9 @@ bool UAutomationBlueprintFunctionLibrary::TakeAutomationScreenshotOfUI_Immediate
 				{
 #if (WITH_DEV_AUTOMATION_TESTS || WITH_PERF_AUTOMATION_TESTS)
 					// The screenshot taker deletes itself later.
-					FAutomationScreenshotTaker* TempObject = new FAutomationScreenshotTaker(GEngine->GetWorldFromContextObject(WorldContextObject), Name, Options);
+					FAutomationScreenshotTaker* TempObject = new FAutomationScreenshotTaker(World, Name, Options);
 
-					FAutomationScreenshotData Data = AutomationCommon::BuildScreenshotData(GWorld->GetName(), Name, OutSize.X, OutSize.Y);
+					FAutomationScreenshotData Data = AutomationCommon::BuildScreenshotData(World->GetName(), Name, OutSize.X, OutSize.Y);
 
 					// Copy the relevant data into the metadata for the screenshot.
 					Data.bHasComparisonRules = true;
diff --git a/Engine/Source/Developer/FunctionalTesting/Private/FunctionalTest.cpp b/Engine/Source/Developer/FunctionalTesting/Private/FunctionalTest.cpp
--- a/Engine/Source/Developer/FunctionalTesting/Private/FunctionalTest.cpp
+++ b/Engine/Source/Developer/FunctionalTesting/Private/FunctionalTest.cpp
@@ -71,7 +71,7 @@ namespace
 
 	void DelayForFramesCommon(UObject* WorldContextObject, FLatentActionInfo LatentInfo, int32 NumFrames)
 	{
-		if (UWorld* World = GEngine->GetWorldFromContextObject(WorldContextObject))
+		if (UWorld* World = GEngine->GetWorldFromContextObject(WorldContextObject, EGetWorldErrorMode::LogAndReturnNull))
 		{
 			FLatentActionManager& LatentActionManager = World->GetLatentActionManager();
 			if (LatentActionManager.FindExistingAction<FDelayForFramesLatentAction>(LatentInfo.CallbackTarget, LatentInfo.UUID) == nullptr)
diff --git a/Engine/Source/Developer/FunctionalTesting/Private/FunctionalTestingManager.cpp b/Engine/Source/Developer/FunctionalTesting/Private/FunctionalTestingManager.cpp
--- a/Engine/Source/Developer/FunctionalTesting/Private/FunctionalTestingManager.cpp
+++ b/Engine/Source/Developer/FunctionalTesting/Private/FunctionalTestingManager.cpp
@@ -9,6 +9,8 @@
 #include "AI/Navigation/NavigationSystem.h"
 #include "EngineGlobals.h"
 #include "Engine/Engine.h"
+#include "Misc/RuntimeErrors.h"
+
 #if WITH_EDITOR
 
 //----------------------------------------------------------------------//
@@ -87,6 +89,10 @@ struct FSortTestActorsByName
 bool UFunctionalTestingManager::RunAllFunctionalTests(UObject* WorldContextObject, bool bNewLog, bool bRunLooped, bool bInWaitForNavigationBuildFinish, FString ReproString)
 {
 	UFunctionalTestingManager* Manager = GetManager(WorldContextObject);
+	if (!ensureAsRuntimeWarning(Manager != nullptr))
+	{
+		return false;
+	}
 
 	if (Manager->bIsRunning)
 	{
@@ -94,7 +100,7 @@ bool UFunctionalTestingManager::RunAllFunctionalTests(UObject* WorldContextObjec
 		return true;
 	}
 	
-	UWorld* World = GEngine->GetWorldFromContextObject(WorldContextObject);
+	UWorld* World = GEngine->GetWorldFromContextObjectChecked(WorldContextObject);
 	World->ForceGarbageCollection(true);
 
 	Manager->bFinished = false;
@@ -182,14 +188,15 @@ UFunctionalTestingManager* UFunctionalTestingManager::GetManager(UObject* WorldC
 
 	if (Manager == nullptr)
 	{
-		check(WorldContext);
-		UObject* World = GEngine->GetWorldFromContextObject(WorldContext);
-		Manager = NewObject<UFunctionalTestingManager>(World);
-		FFunctionalTestingModule::Get()->SetScript(Manager);
-
-		// add to root and get notified on world cleanup to remove from root on map cleanup
-		Manager->AddToRoot();
-		FWorldDelegates::OnWorldCleanup.AddUObject(Manager, &UFunctionalTestingManager::OnWorldCleanedUp);
+		if (UObject* World = GEngine->GetWorldFromContextObject(WorldContext, EGetWorldErrorMode::LogAndReturnNull))
+		{
+			Manager = NewObject<UFunctionalTestingManager>(World);
+			FFunctionalTestingModule::Get()->SetScript(Manager);
+
+			// add to root and get notified on world cleanup to remove from root on map cleanup
+			Manager->AddToRoot();
+			FWorldDelegates::OnWorldCleanup.AddUObject(Manager, &UFunctionalTestingManager::OnWorldCleanedUp);
+		}
 	}
 
 	return Manager;
@@ -197,7 +204,7 @@ UFunctionalTestingManager* UFunctionalTestingManager::GetManager(UObject* WorldC
 
 UWorld* UFunctionalTestingManager::GetWorld() const
 {
-	return GEngine->GetWorldFromContextObject(GetOuter());
+	return GEngine->GetWorldFromContextObjectChecked(GetOuter());
 }
 
 void UFunctionalTestingManager::OnWorldCleanedUp(UWorld* World, bool bSessionEnded, bool bCleanupResources)
diff --git a/Engine/Source/Developer/FunctionalTesting/Private/TraceQueryTestResults.cpp b/Engine/Source/Developer/FunctionalTesting/Private/TraceQueryTestResults.cpp
--- a/Engine/Source/Developer/FunctionalTesting/Private/TraceQueryTestResults.cpp
+++ b/Engine/Source/Developer/FunctionalTesting/Private/TraceQueryTestResults.cpp
@@ -200,7 +200,12 @@ bool FTraceQueryTestResultsInner::AssertEqual(const FTraceQueryTestResultsInner&
 
 bool UTraceQueryTestResults::AssertEqual(const UTraceQueryTestResults* Expected, const FString& What, const UObject* ContextObject, AFunctionalTest& FunctionalTest) const
 {
-	if(Expected->BatchOptions != BatchOptions)
+	if(Expected == nullptr)
+	{
+		FunctionalTest.LogStep(ELogVerbosity::Error, FString::Printf(TEXT("Expected '%s' 'Expected != nullptr' for context '%s'"), *What, ContextObject ? *ContextObject->GetName() : TEXT("")));
+		return false;
+	}
+	else if(Expected->BatchOptions != BatchOptions)
 	{
 		FunctionalTest.LogStep(ELogVerbosity::Error, FString::Printf(TEXT("Expected '%s' to be {%s} but it was {%s} for context '%s'"), *What, *Expected->BatchOptions.ToString(), *BatchOptions.ToString(), ContextObject ? *ContextObject->GetName() : TEXT("")));
 		return false;
diff --git a/Engine/Source/Developer/LogVisualizer/Private/LogVisualizer.cpp b/Engine/Source/Developer/LogVisualizer/Private/LogVisualizer.cpp
--- a/Engine/Source/Developer/LogVisualizer/Private/LogVisualizer.cpp
+++ b/Engine/Source/Developer/LogVisualizer/Private/LogVisualizer.cpp
@@ -115,20 +115,18 @@ FLinearColor FLogVisualizer::GetColorForCategory(const FString& InFilterName) co
 
 UWorld* FLogVisualizer::GetWorld(UObject* OptionalObject)
 {
-	UWorld* World = OptionalObject != nullptr ? GEngine->GetWorldFromContextObject(OptionalObject, false) : nullptr;
+	UWorld* World = GEngine->GetWorldFromContextObject(OptionalObject, EGetWorldErrorMode::ReturnNull);
 #if WITH_EDITOR
 	if (!World && GIsEditor)
 	{
-		UEditorEngine *EEngine = Cast<UEditorEngine>(GEngine);
+		UEditorEngine* EEngine = Cast<UEditorEngine>(GEngine);
 		// lets use PlayWorld during PIE/Simulate and regular world from editor otherwise, to draw debug information
 		World = EEngine != nullptr && EEngine->PlayWorld != nullptr ? EEngine->PlayWorld : EEngine->GetEditorWorldContext().World();
-
 	}
 	else 
 #endif
 	if (!World && !GIsEditor)
 	{
-
 		World = GEngine->GetWorld();
 	}
 
diff --git a/Engine/Source/Developer/LogVisualizer/Private/SVisualLogger.cpp b/Engine/Source/Developer/LogVisualizer/Private/SVisualLogger.cpp
--- a/Engine/Source/Developer/LogVisualizer/Private/SVisualLogger.cpp
+++ b/Engine/Source/Developer/LogVisualizer/Private/SVisualLogger.cpp
@@ -66,9 +66,9 @@ namespace
 {
 	static UWorld* GetWorldForGivenObject(const UObject* Object)
 	{
-		UWorld* World = Object ? GEngine->GetWorldFromContextObject(Object, false) : nullptr;
+		UWorld* World = GEngine->GetWorldFromContextObject(Object, EGetWorldErrorMode::ReturnNull);
 #if WITH_EDITOR
-		UEditorEngine *EEngine = Cast<UEditorEngine>(GEngine);
+		UEditorEngine* EEngine = Cast<UEditorEngine>(GEngine);
 		if (GIsEditor && EEngine != nullptr && World == nullptr)
 		{
 			// lets use PlayWorld during PIE/Simulate and regular world from editor otherwise, to draw debug information
diff --git a/Engine/Source/Editor/AnimGraph/Classes/AnimGraphNode_Base.h b/Engine/Source/Editor/AnimGraph/Classes/AnimGraphNode_Base.h
--- a/Engine/Source/Editor/AnimGraph/Classes/AnimGraphNode_Base.h
+++ b/Engine/Source/Editor/AnimGraph/Classes/AnimGraphNode_Base.h
@@ -132,6 +132,7 @@ class ANIMGRAPH_API UAnimGraphNode_Base : public UK2Node
 
 	// UObject interface
 	virtual void PostEditChangeProperty(FPropertyChangedEvent& PropertyChangedEvent) override;
+	virtual void PreEditChange(UProperty* PropertyAboutToChange) override;
 	// End of UObject interface
 
 	// UEdGraphNode interface
@@ -181,7 +182,7 @@ class ANIMGRAPH_API UAnimGraphNode_Base : public UK2Node
 	virtual void BakeDataDuringCompilation(FCompilerResultsLog& MessageLog) {}
 
 	// preload asset required for this node in this function
-	virtual void PreloadRequiredAssets() {}
+	virtual void PreloadRequiredAssets() override {}
 
 	// Give the node a chance to change the display name of a pin
 	virtual void PostProcessPinName(const UEdGraphPin* Pin, FString& DisplayName) const;
@@ -308,6 +309,9 @@ protected:
 	void InternalPinCreation(TArray<UEdGraphPin*>* OldPins);
 
 	FOnNodePropertyChangedEvent PropertyChangeEvent;
+
+private:
+	TArray<FName> OldShownPins;
 };
 
 template<class AssetType>
diff --git a/Engine/Source/Editor/AnimGraph/Private/AnimGraphNode_Base.cpp b/Engine/Source/Editor/AnimGraph/Private/AnimGraphNode_Base.cpp
--- a/Engine/Source/Editor/AnimGraph/Private/AnimGraphNode_Base.cpp
+++ b/Engine/Source/Editor/AnimGraph/Private/AnimGraphNode_Base.cpp
@@ -21,12 +21,23 @@ UAnimGraphNode_Base::UAnimGraphNode_Base(const FObjectInitializer& ObjectInitial
 {
 }
 
+void UAnimGraphNode_Base::PreEditChange(UProperty* PropertyThatWillChange)
+{
+	Super::PreEditChange(PropertyThatWillChange);
+
+	if (PropertyThatWillChange && PropertyThatWillChange->GetFName() == GET_MEMBER_NAME_CHECKED(FOptionalPinFromProperty, bShowPin))
+	{
+		FOptionalPinManager::CacheShownPins(ShowPinForProperties, OldShownPins);
+	}
+}
+
 void UAnimGraphNode_Base::PostEditChangeProperty(FPropertyChangedEvent& PropertyChangedEvent)
 {
 	FName PropertyName = (PropertyChangedEvent.Property != NULL) ? PropertyChangedEvent.Property->GetFName() : NAME_None;
 
 	if ((PropertyName == GET_MEMBER_NAME_CHECKED(FOptionalPinFromProperty, bShowPin)))
 	{
+		FOptionalPinManager::EvaluateOldShownPins(ShowPinForProperties, OldShownPins, this);
 		GetSchema()->ReconstructNode(*this);
 	}
 
diff --git a/Engine/Source/Editor/AnimGraph/Private/AnimGraphNode_BlendListByEnum.cpp b/Engine/Source/Editor/AnimGraph/Private/AnimGraphNode_BlendListByEnum.cpp
--- a/Engine/Source/Editor/AnimGraph/Private/AnimGraphNode_BlendListByEnum.cpp
+++ b/Engine/Source/Editor/AnimGraph/Private/AnimGraphNode_BlendListByEnum.cpp
@@ -176,6 +176,7 @@ void UAnimGraphNode_BlendListByEnum::RemovePinFromBlendList(UEdGraphPin* Pin)
 		// setting up removed pins info 
 		RemovedPinArrayIndex = ArrayIndex;
 		Node.RemovePose(ArrayIndex);
+		Pin->bSavePinIfOrphaned = false;
 		ReconstructNode();
 		//@TODO: Just want to invalidate the visual representation currently
 		FBlueprintEditorUtils::MarkBlueprintAsStructurallyModified(GetBlueprint());
diff --git a/Engine/Source/Editor/AnimGraph/Private/AnimGraphNode_SubInstance.cpp b/Engine/Source/Editor/AnimGraph/Private/AnimGraphNode_SubInstance.cpp
--- a/Engine/Source/Editor/AnimGraph/Private/AnimGraphNode_SubInstance.cpp
+++ b/Engine/Source/Editor/AnimGraph/Private/AnimGraphNode_SubInstance.cpp
@@ -168,7 +168,7 @@ void UAnimGraphNode_SubInstance::ReallocatePinsDuringReconstruction(TArray<UEdGr
 			KnownExposableProperties.Add(PropertyName);
 		}
 
-		if(ExposedPropertyNames.Contains(PropertyName))
+		if(ExposedPropertyNames.Contains(PropertyName) && FBlueprintEditorUtils::PropertyStillExists(Property))
 		{
 			FEdGraphPinType PinType;
 
diff --git a/Engine/Source/Editor/BehaviorTreeEditor/Private/DetailCustomizations/BlackboardDecoratorDetails.cpp b/Engine/Source/Editor/BehaviorTreeEditor/Private/DetailCustomizations/BlackboardDecoratorDetails.cpp
--- a/Engine/Source/Editor/BehaviorTreeEditor/Private/DetailCustomizations/BlackboardDecoratorDetails.cpp
+++ b/Engine/Source/Editor/BehaviorTreeEditor/Private/DetailCustomizations/BlackboardDecoratorDetails.cpp
@@ -198,19 +198,17 @@ TSharedRef<SWidget> FBlackboardDecoratorDetails::OnGetEnumValueContent() const
 
 FText FBlackboardDecoratorDetails::GetCurrentEnumValueDesc() const
 {
-	if (CachedCustomObjectType && EnumPropValues.Num() > 0)
-	{
-		int32 CurrentIntValue = 0;
-		FPropertyAccess::Result Result = IntValueProperty->GetValue(CurrentIntValue);
+	FPropertyAccess::Result Result = FPropertyAccess::Fail;
+	int32 CurrentIntValue = INDEX_NONE;
 
-		if (Result == FPropertyAccess::Success)
-		{
-			int32 ClampedIdx = FMath::Clamp(CurrentIntValue, 0, EnumPropValues.Num());
-			return FText::FromString(EnumPropValues[ClampedIdx]);
-		}
+	if (CachedCustomObjectType)
+	{	
+		Result = IntValueProperty->GetValue(CurrentIntValue);
 	}
 
-	return FText::GetEmpty();
+	return (Result == FPropertyAccess::Success && EnumPropValues.IsValidIndex(CurrentIntValue))
+		? FText::FromString(EnumPropValues[CurrentIntValue])
+		: FText::GetEmpty();
 }
 
 void FBlackboardDecoratorDetails::OnEnumValueComboChange(int32 Index)
diff --git a/Engine/Source/Editor/BlueprintGraph/Classes/EdGraphSchema_K2.h b/Engine/Source/Editor/BlueprintGraph/Classes/EdGraphSchema_K2.h
--- a/Engine/Source/Editor/BlueprintGraph/Classes/EdGraphSchema_K2.h
+++ b/Engine/Source/Editor/BlueprintGraph/Classes/EdGraphSchema_K2.h
@@ -453,6 +453,8 @@ public:
 	virtual void TrySetDefaultValue(UEdGraphPin& Pin, const FString& NewDefaultValue) const override;
 	virtual void TrySetDefaultObject(UEdGraphPin& Pin, UObject* NewDefaultObject) const override;
 	virtual void TrySetDefaultText(UEdGraphPin& InPin, const FText& InNewDefaultText) const override;
+	virtual bool DoesDefaultValueMatchAutogenerated(const UEdGraphPin& InPin) const override;
+	virtual void ResetPinToAutogeneratedDefaultValue(UEdGraphPin* Pin, bool bCallModifyCallbacks = true) const override;
 	virtual bool ShouldHidePinDefaultValue(UEdGraphPin* Pin) const override;
 	virtual bool ShouldShowAssetPickerForPin(UEdGraphPin* Pin) const override;
 	virtual FLinearColor GetPinTypeColor(const FEdGraphPinType& PinType) const override;
@@ -773,7 +775,7 @@ public:
 
 	static bool IsAllowableBlueprintVariableType(const class UEnum* InEnum);
 	static bool IsAllowableBlueprintVariableType(const class UClass* InClass);
-	static bool IsAllowableBlueprintVariableType(const class UScriptStruct *InStruct);
+	static bool IsAllowableBlueprintVariableType(const class UScriptStruct *InStruct, bool bForInternalUse = false);
 
 	static bool IsPropertyExposedOnSpawn(const UProperty* Property);
 
@@ -958,9 +960,6 @@ public:
 	/** Sets the autogenerated default value for a pin using the default for that type. This will also reset the current default value to the autogenerated one */
 	virtual void SetPinAutogeneratedDefaultValueBasedOnType(UEdGraphPin* Pin) const;
 
-	/** Resets a pin back to it's autogenerated default value, optionally calling the default value change callbacks */
-	virtual void ResetPinToAutogeneratedDefaultValue(UEdGraphPin* Pin, bool bCallModifyCallbacks = true) const;
-
 	/** Sets the pin defaults, but not autogenerated defaults, at pin construction time. This is like TrySetDefaultValue but does not do validation or callbacks */
 	virtual void SetPinDefaultValueAtConstruction(UEdGraphPin* Pin, const FString& DefaultValueString) const;
 
diff --git a/Engine/Source/Editor/BlueprintGraph/Classes/EdGraphSchema_K2_Actions.h b/Engine/Source/Editor/BlueprintGraph/Classes/EdGraphSchema_K2_Actions.h
--- a/Engine/Source/Editor/BlueprintGraph/Classes/EdGraphSchema_K2_Actions.h
+++ b/Engine/Source/Editor/BlueprintGraph/Classes/EdGraphSchema_K2_Actions.h
@@ -779,7 +779,7 @@ public:
 		return OwningClass.Get();
 	}
 
-	UMulticastDelegateProperty* GetDelegatePoperty() const
+	UMulticastDelegateProperty* GetDelegateProperty() const
 	{
 		UClass* DelegateClass = GetDelegateClass();
 		return FindField<UMulticastDelegateProperty>(DelegateClass, DelegateName);
diff --git a/Engine/Source/Editor/BlueprintGraph/Classes/K2Node.h b/Engine/Source/Editor/BlueprintGraph/Classes/K2Node.h
--- a/Engine/Source/Editor/BlueprintGraph/Classes/K2Node.h
+++ b/Engine/Source/Editor/BlueprintGraph/Classes/K2Node.h
@@ -22,13 +22,13 @@ class UEdGraphSchema;
 struct FOldOptionalPinSettings
 {
 	/** TRUE if optional pin was previously visible */
-	bool bOldVisibility;
+	uint8 bOldVisibility:1;
 	/** TRUE if the optional pin's override value was previously enabled */
-	bool bIsOldOverrideEnabled;
+	uint8 bIsOldOverrideEnabled:1;
 	/** TRUE if the optional pin's value was previously editable */
-	bool bIsOldSetValuePinVisible;
+	uint8 bIsOldSetValuePinVisible:1;
 	/** TRUE if the optional pin's override value was previously editable */
-	bool bIsOldOverridePinVisible;
+	uint8 bIsOldOverridePinVisible:1;
 
 	FOldOptionalPinSettings(bool bInOldVisibility, bool bInIsOldOverrideEnabled, bool bInIsOldSetValuePinVisible, bool bInIsOldOverridePinVisible)
 		: bOldVisibility(bInOldVisibility)
@@ -142,6 +142,14 @@ public:
 
 	// Customize automatically created pins if desired
 	virtual void CustomizePinData(UEdGraphPin* Pin, FName SourcePropertyName, int32 ArrayIndex, UProperty* Property = NULL) const {}
+
+	/** Helper function to make consistent behavior between nodes that use optional pins */
+	static void CacheShownPins(const TArray<FOptionalPinFromProperty>& OptionalPins, TArray<FName>& OldShownPins);
+
+	/** Helper function to make consistent behavior between nodes that use optional pins */
+	static void EvaluateOldShownPins(const TArray<FOptionalPinFromProperty>& OptionalPins, TArray<FName>& OldShownPins, UK2Node* Node);
+
+
 protected:
 	virtual void PostInitNewPin(UEdGraphPin* Pin, FOptionalPinFromProperty& Record, int32 ArrayIndex, UProperty* Property, uint8* PropertyAddress, uint8* DefaultPropertyAddress) const {}
 	virtual void PostRemovedOldPin(FOptionalPinFromProperty& Record, int32 ArrayIndex, UProperty* Property, uint8* PropertyAddress, uint8* DefaultPropertyAddress) const {}
@@ -173,13 +181,14 @@ class UK2Node : public UEdGraphNode
 	BLUEPRINTGRAPH_API virtual FLinearColor GetNodeTitleColor() const override;
 	BLUEPRINTGRAPH_API virtual void AutowireNewNode(UEdGraphPin* FromPin) override;
 	BLUEPRINTGRAPH_API void PinConnectionListChanged(UEdGraphPin* Pin) override;
-    BLUEPRINTGRAPH_API virtual UObject* GetJumpTargetForDoubleClick() const override;
+	BLUEPRINTGRAPH_API virtual UObject* GetJumpTargetForDoubleClick() const override;
 	BLUEPRINTGRAPH_API virtual FString GetDocumentationLink() const override;
 	BLUEPRINTGRAPH_API virtual void GetPinHoverText(const UEdGraphPin& Pin, FString& HoverTextOut) const override;
 	BLUEPRINTGRAPH_API virtual bool ShowPaletteIconOnNode() const override { return true; }
 	BLUEPRINTGRAPH_API virtual bool CanSplitPin(const UEdGraphPin* Pin) const override;
 	BLUEPRINTGRAPH_API virtual UEdGraphPin* GetPassThroughPin(const UEdGraphPin* FromPin) const override;
 	BLUEPRINTGRAPH_API virtual bool IsInDevelopmentMode() const override;
+	BLUEPRINTGRAPH_API virtual void ValidateNodeDuringCompilation(class FCompilerResultsLog& MessageLog) const override;
 	// End of UEdGraphNode interface
 
 	// K2Node interface
@@ -198,7 +207,7 @@ class UK2Node : public UEdGraphNode
 	 * Returns whether or not this node has dependencies on an external structure 
 	 * If OptionalOutput isn't null, it should be filled with the known dependencies objects (Classes, Structures, Functions, etc).
 	 */
-	virtual bool HasExternalDependencies(TArray<class UStruct*>* OptionalOutput = NULL) const { return false; }
+	virtual bool HasExternalDependencies(TArray<class UStruct*>* OptionalOutput = nullptr) const { return false; }
 
 	/** Returns whether this node can have breakpoints placed on it in the debugger */
 	virtual bool CanPlaceBreakpoints() const { return !IsNodePure(); }
@@ -324,11 +333,14 @@ class UK2Node : public UEdGraphNode
 	BLUEPRINTGRAPH_API virtual bool IsConnectionDisallowed(const UEdGraphPin* MyPin, const UEdGraphPin* OtherPin, FString& OutReason) const { return false; }
 
 	/** This function if used for nodes that needs CDO for validation (Called before expansion)*/
-	BLUEPRINTGRAPH_API virtual void EarlyValidation(class FCompilerResultsLog& MessageLog) const {}
+	BLUEPRINTGRAPH_API virtual void EarlyValidation(class FCompilerResultsLog& MessageLog) const;
 
 	/** This function returns an arbitrary number of attributes that describe this node for analytics events */
 	BLUEPRINTGRAPH_API virtual void GetNodeAttributes( TArray<TKeyValuePair<FString, FString>>& OutNodeAttributes ) const;
 
+	/** Called before compilation begins, giving a blueprint time to force the linker to load data */
+	BLUEPRINTGRAPH_API virtual void PreloadRequiredAssets();
+
 	/** 
 	 * Replacement for GetMenuEntries(). Override to add specific 
 	 * UBlueprintNodeSpawners pertaining to the sub-class type. Serves as an 
@@ -424,7 +436,7 @@ protected:
 	void Message_Error(const FString& Message);
 
 
-    friend class FKismetCompilerContext;
+	friend class FKismetCompilerContext;
 
 
 protected:
@@ -439,6 +451,14 @@ protected:
 
 	void FixupPinDefaultValues();
 
+private:
+
+	/** 
+	 * Utility function to write messages about orphan nodes in to the compiler log.
+	 * bStore indicates whether to write immediately to the log, or to store as a potential message to be committed once node pruning has completed
+	 */
+	void ValidateOrphanPins(class FCompilerResultsLog& MessageLog, bool bStore) const;
+
 public:
 
 	/** Util to get the generated class from a node. */
diff --git a/Engine/Source/Editor/BlueprintGraph/Classes/K2Node_AddComponent.h b/Engine/Source/Editor/BlueprintGraph/Classes/K2Node_AddComponent.h
--- a/Engine/Source/Editor/BlueprintGraph/Classes/K2Node_AddComponent.h
+++ b/Engine/Source/Editor/BlueprintGraph/Classes/K2Node_AddComponent.h
@@ -31,7 +31,6 @@ class UK2Node_AddComponent : public UK2Node_CallFunction
 
 	//~ Begin UObject Interface
 	virtual void Serialize(FArchive& Ar) override;
-	virtual void PostDuplicate(bool bDuplicateForPIE) override;
 	//~ End UObject Interface
 
 	//~ Begin UEdGraphNode Interface
diff --git a/Engine/Source/Editor/BlueprintGraph/Classes/K2Node_AddPinInterface.h b/Engine/Source/Editor/BlueprintGraph/Classes/K2Node_AddPinInterface.h
new file mode 100644
--- /dev/null
+++ b/Engine/Source/Editor/BlueprintGraph/Classes/K2Node_AddPinInterface.h
@@ -0,0 +1,23 @@
+// Copyright 1998-2017 Epic Games, Inc. All Rights Reserved.
+
+
+#pragma once
+
+#include "CoreMinimal.h"
+#include "Interface.h"
+#include "K2Node_AddPinInterface.generated.h"
+
+UINTERFACE(meta=(CannotImplementInterfaceInBlueprint))
+class BLUEPRINTGRAPH_API UK2Node_AddPinInterface : public UInterface
+{
+	GENERATED_BODY()
+};
+
+class BLUEPRINTGRAPH_API IK2Node_AddPinInterface
+{
+	GENERATED_BODY()
+
+public:
+	virtual void AddInputPin() { }
+	virtual bool CanAddPin() const { return true; }
+};
\ No newline at end of file
diff --git a/Engine/Source/Editor/BlueprintGraph/Classes/K2Node_BreakStruct.h b/Engine/Source/Editor/BlueprintGraph/Classes/K2Node_BreakStruct.h
--- a/Engine/Source/Editor/BlueprintGraph/Classes/K2Node_BreakStruct.h
+++ b/Engine/Source/Editor/BlueprintGraph/Classes/K2Node_BreakStruct.h
@@ -27,16 +27,13 @@ class UK2Node_BreakStruct : public UK2Node_StructMemberGet
 	 *   1. The Struct has a 'native break' method
 	 * Returns true if:
 	 *   1. The Struct is tagged as BlueprintType
-	 *   or
+	 *   and
 	 *   2. The Struct has any property that is tagged as CPF_BlueprintVisible
-	 *   or
-	 *   3. The Struct has any property that is tagged as CPF_Edit and bIncludeEditAnywhere is true
-	 *
-	 * When constructing the context menu we do not allow the presence of the EditAnywhere flag to result in
-	 * a creation of a break node, although we could do so in the future. There are legacy break nodes that
-	 * rely on expansion of structs that neither have BlueprintVisible properties nor are tagged as BlueprintType
 	 */
-	static bool CanBeBroken(const UScriptStruct* Struct, bool bIncludeEditAnywhere = true, bool bMustHaveValidProperties = false);
+	static bool CanBeBroken(const UScriptStruct* Struct, bool bForInternalUse = false);
+
+	/** Can this struct be used as a split pin */
+	static bool CanBeSplit(const UScriptStruct* Struct) { return CanBeBroken(Struct); }
 
 	// UObject interface
 	virtual void Serialize(FArchive& Ar) override;
diff --git a/Engine/Source/Editor/BlueprintGraph/Classes/K2Node_CommutativeAssociativeBinaryOperator.h b/Engine/Source/Editor/BlueprintGraph/Classes/K2Node_CommutativeAssociativeBinaryOperator.h
--- a/Engine/Source/Editor/BlueprintGraph/Classes/K2Node_CommutativeAssociativeBinaryOperator.h
+++ b/Engine/Source/Editor/BlueprintGraph/Classes/K2Node_CommutativeAssociativeBinaryOperator.h
@@ -6,12 +6,13 @@
 #include "UObject/ObjectMacros.h"
 #include "EdGraph/EdGraphPin.h"
 #include "K2Node_CallFunction.h"
+#include "K2Node_AddPinInterface.h"
 #include "K2Node_CommutativeAssociativeBinaryOperator.generated.h"
 
 class UEdGraph;
 
 UCLASS(MinimalAPI)
-class UK2Node_CommutativeAssociativeBinaryOperator : public UK2Node_CallFunction
+class UK2Node_CommutativeAssociativeBinaryOperator : public UK2Node_CallFunction, public IK2Node_AddPinInterface
 {
 	GENERATED_UCLASS_BODY()
 
@@ -28,7 +29,6 @@ private:
 	FEdGraphPinType GetType() const;
 
 	void AddInputPinInner(int32 AdditionalPinIndex);
-	bool CanAddPin() const;
 	bool CanRemovePin(const UEdGraphPin* Pin) const;
 public:
 	BLUEPRINTGRAPH_API UEdGraphPin* FindOutPin() const;
@@ -37,7 +37,6 @@ public:
 	/** Get TRUE input type (self, etc.. are skipped) */
 	BLUEPRINTGRAPH_API UEdGraphPin* GetInputPin(int32 InputPinIndex);
 
-	BLUEPRINTGRAPH_API void AddInputPin();
 	BLUEPRINTGRAPH_API void RemoveInputPin(UEdGraphPin* Pin);
 
 	// UEdGraphNode interface
@@ -48,4 +47,9 @@ public:
 	virtual void GetContextMenuActions(const FGraphNodeContextMenuBuilder& Context) const override;
 	virtual void ExpandNode(class FKismetCompilerContext& CompilerContext, UEdGraph* SourceGraph) override;
 	// End of UK2Node interface
+
+	// IK2Node_AddPinInterface interface
+	BLUEPRINTGRAPH_API virtual void AddInputPin() override;
+	virtual bool CanAddPin() const override;
+	// End of IK2Node_AddPinInterface interface
 };
diff --git a/Engine/Source/Editor/BlueprintGraph/Classes/K2Node_DoOnceMultiInput.h b/Engine/Source/Editor/BlueprintGraph/Classes/K2Node_DoOnceMultiInput.h
--- a/Engine/Source/Editor/BlueprintGraph/Classes/K2Node_DoOnceMultiInput.h
+++ b/Engine/Source/Editor/BlueprintGraph/Classes/K2Node_DoOnceMultiInput.h
@@ -6,13 +6,14 @@
 #include "UObject/ObjectMacros.h"
 #include "EdGraph/EdGraphPin.h"
 #include "K2Node.h"
+#include "K2Node_AddPinInterface.h"
 #include "K2Node_DoOnceMultiInput.generated.h"
 
 class FBlueprintActionDatabaseRegistrar;
 class UEdGraph;
 
 UCLASS(MinimalAPI)
-class UK2Node_DoOnceMultiInput : public UK2Node
+class UK2Node_DoOnceMultiInput : public UK2Node, public IK2Node_AddPinInterface
 {
 	GENERATED_UCLASS_BODY()
 
@@ -39,7 +40,6 @@ private:
 	FEdGraphPinType GetOutType() const;
 
 	void AddPinsInner(int32 AdditionalPinIndex);
-	bool CanAddPin() const;
 	bool CanRemovePin(const UEdGraphPin* Pin) const;
 public:
 	BLUEPRINTGRAPH_API UEdGraphPin* FindOutPin() const;
@@ -49,7 +49,6 @@ public:
 	BLUEPRINTGRAPH_API UEdGraphPin* GetInputPin(int32 InputPinIndex);
 	BLUEPRINTGRAPH_API UEdGraphPin* GetOutputPin(int32 InputPinIndex);
 
-	BLUEPRINTGRAPH_API void AddInputPin();
 	BLUEPRINTGRAPH_API void RemoveInputPin(UEdGraphPin* Pin);
 
 	// UEdGraphNode interface
@@ -62,4 +61,9 @@ public:
 	virtual void GetMenuActions(FBlueprintActionDatabaseRegistrar& ActionRegistrar) const override;
 	virtual FText GetMenuCategory() const override;
 	// End of UK2Node interface
+
+	// IK2Node_AddPinInterface interface
+	BLUEPRINTGRAPH_API virtual void AddInputPin() override;
+	virtual bool CanAddPin() const override;
+	// End of IK2Node_AddPinInterface interface
 };
diff --git a/Engine/Source/Editor/BlueprintGraph/Classes/K2Node_ExecutionSequence.h b/Engine/Source/Editor/BlueprintGraph/Classes/K2Node_ExecutionSequence.h
--- a/Engine/Source/Editor/BlueprintGraph/Classes/K2Node_ExecutionSequence.h
+++ b/Engine/Source/Editor/BlueprintGraph/Classes/K2Node_ExecutionSequence.h
@@ -7,13 +7,14 @@
 #include "UObject/ObjectMacros.h"
 #include "Textures/SlateIcon.h"
 #include "K2Node.h"
+#include "K2Node_AddPinInterface.h"
 #include "K2Node_ExecutionSequence.generated.h"
 
 class FBlueprintActionDatabaseRegistrar;
 class UEdGraphPin;
 
 UCLASS(MinimalAPI)
-class UK2Node_ExecutionSequence : public UK2Node
+class UK2Node_ExecutionSequence : public UK2Node, public IK2Node_AddPinInterface
 {
 	GENERATED_UCLASS_BODY()
 
@@ -33,6 +34,10 @@ class UK2Node_ExecutionSequence : public UK2Node
 	virtual bool CanEverRemoveExecutionPin() const override { return true; }
 	//~ End UK2Node Interface
 
+	// IK2Node_AddPinInterface interface
+	BLUEPRINTGRAPH_API virtual void AddInputPin() override;
+	// End of IK2Node_AddPinInterface interface
+
 	//~ Begin K2Node_ExecutionSequence Interface
 
 	/** Gets a unique pin name, the next in the sequence */
@@ -41,7 +46,8 @@ class UK2Node_ExecutionSequence : public UK2Node
 	/**
 	 * Adds a new execution pin to an execution node
 	 */
-	BLUEPRINTGRAPH_API void AddPinToExecutionNode();
+	DEPRECATED(4.17, "Use AddInputPin instead.")
+	BLUEPRINTGRAPH_API void AddPinToExecutionNode() { AddInputPin(); }
 
 	/**
 	 * Removes the specified execution pin from an execution node
@@ -54,7 +60,7 @@ class UK2Node_ExecutionSequence : public UK2Node
 	BLUEPRINTGRAPH_API bool CanRemoveExecutionPin() const;
 
 	// @todo document
-	BLUEPRINTGRAPH_API UEdGraphPin * GetThenPinGivenIndex(int32 Index);
+	BLUEPRINTGRAPH_API UEdGraphPin* GetThenPinGivenIndex(int32 Index);
 
 private:
 	// Returns the exec output pin name for a given 0-based index
diff --git a/Engine/Source/Editor/BlueprintGraph/Classes/K2Node_FunctionEntry.h b/Engine/Source/Editor/BlueprintGraph/Classes/K2Node_FunctionEntry.h
--- a/Engine/Source/Editor/BlueprintGraph/Classes/K2Node_FunctionEntry.h
+++ b/Engine/Source/Editor/BlueprintGraph/Classes/K2Node_FunctionEntry.h
@@ -33,6 +33,7 @@ class UK2Node_FunctionEntry : public UK2Node_FunctionTerminator
 
 	//~ Begin UObject Interface
 	virtual void Serialize(FArchive& Ar) override;
+	virtual void PreSave(const class ITargetPlatform* TargetPlatform) override;
 	//~ End UObject Interface
 
 	//~ Begin UEdGraphNode Interface
diff --git a/Engine/Source/Editor/BlueprintGraph/Classes/K2Node_GetClassDefaults.h b/Engine/Source/Editor/BlueprintGraph/Classes/K2Node_GetClassDefaults.h
--- a/Engine/Source/Editor/BlueprintGraph/Classes/K2Node_GetClassDefaults.h
+++ b/Engine/Source/Editor/BlueprintGraph/Classes/K2Node_GetClassDefaults.h
@@ -19,6 +19,7 @@ class UK2Node_GetClassDefaults : public UK2Node
 
 	//~ Begin UObject Interface
 	virtual void PostEditChangeProperty(struct FPropertyChangedEvent& PropertyChangedEvent) override;
+	virtual void PreEditChange(UProperty* PropertyAboutToChange) override;
 	//~ End UObject Interface
 
 	//~ Begin UEdGraphNode Interface
@@ -99,6 +100,8 @@ private:
 	UPROPERTY(EditAnywhere, Category=PinOptions, EditFixedSize)
 	TArray<FOptionalPinFromProperty> ShowPinForProperties;
 
+	TArray<FName> OldShownPins;
+
 	/** Whether or not to exclude object container properties */
 	UPROPERTY()
 	bool bExcludeObjectContainers;
diff --git a/Engine/Source/Editor/BlueprintGraph/Classes/K2Node_MakeArray.h b/Engine/Source/Editor/BlueprintGraph/Classes/K2Node_MakeArray.h
--- a/Engine/Source/Editor/BlueprintGraph/Classes/K2Node_MakeArray.h
+++ b/Engine/Source/Editor/BlueprintGraph/Classes/K2Node_MakeArray.h
@@ -7,13 +7,14 @@
 #include "UObject/ObjectMacros.h"
 #include "Textures/SlateIcon.h"
 #include "K2Node.h"
+#include "K2Node_AddPinInterface.h"
 #include "K2Node_MakeArray.generated.h"
 
 class FBlueprintActionDatabaseRegistrar;
 class UEdGraphPin;
 
 UCLASS(MinimalAPI)
-class UK2Node_MakeArray : public UK2Node
+class UK2Node_MakeArray : public UK2Node, public IK2Node_AddPinInterface
 {
 	GENERATED_UCLASS_BODY()
 
@@ -22,7 +23,6 @@ class UK2Node_MakeArray : public UK2Node
 	int32 NumInputs;
 
 public:
-	BLUEPRINTGRAPH_API void AddInputPin();
 	BLUEPRINTGRAPH_API void RemoveInputPin(UEdGraphPin* Pin);
 
 	/** returns a reference to the output array pin of this node, which is responsible for defining the type */
@@ -50,7 +50,9 @@ public:
 	virtual void ReallocatePinsDuringReconstruction(TArray<UEdGraphPin*>& OldPins);
 	// End of UK2Node interface
 
-
+	// IK2Node_AddPinInterface interface
+	BLUEPRINTGRAPH_API virtual void AddInputPin() override;
+	// End of IK2Node_AddPinInterface interface
 
 protected:
 	friend class FKismetCompilerContext;
diff --git a/Engine/Source/Editor/BlueprintGraph/Classes/K2Node_MakeStruct.h b/Engine/Source/Editor/BlueprintGraph/Classes/K2Node_MakeStruct.h
--- a/Engine/Source/Editor/BlueprintGraph/Classes/K2Node_MakeStruct.h
+++ b/Engine/Source/Editor/BlueprintGraph/Classes/K2Node_MakeStruct.h
@@ -28,17 +28,14 @@ class UK2Node_MakeStruct : public UK2Node_StructMemberSet
 	*   1. The Struct has a 'native make' method
 	* Returns true if:
 	*   1. The Struct is tagged as BlueprintType
-	*   or
+	*   and
 	*   2. The Struct has any property that is tagged as CPF_BlueprintVisible and is not CPF_BlueprintReadOnly
-	*   or
-	*   3. The Struct has any property that is tagged as CPF_Edit and is not CPF_BlueprintReadOnly and bIncludeEditAnywhere is true 
-	*
-	* When constructing the context menu we do not allow the presence of the EditAnywhere flag to result in
-	* a creation of a break node, although we could do so in the future. There are legacy break nodes that
-	* rely on expansion of structs that neither have BlueprintVisible properties nor are tagged as BlueprintType
 	*/
-	BLUEPRINTGRAPH_API static bool CanBeMade(const UScriptStruct* Struct, bool bIncludeEditAnywhere = true, bool bMustHaveValidProperties = false);
+	BLUEPRINTGRAPH_API static bool CanBeMade(const UScriptStruct* Struct, bool bForInternalUse = false);
 	
+	/** Can this struct be used as a split pin */
+	BLUEPRINTGRAPH_API static bool CanBeSplit(const UScriptStruct* Struct);
+
 	// UObject interface
 	virtual void Serialize(FArchive& Ar) override;
 	// End of UObject interface
diff --git a/Engine/Source/Editor/BlueprintGraph/Classes/K2Node_SetFieldsInStruct.h b/Engine/Source/Editor/BlueprintGraph/Classes/K2Node_SetFieldsInStruct.h
--- a/Engine/Source/Editor/BlueprintGraph/Classes/K2Node_SetFieldsInStruct.h
+++ b/Engine/Source/Editor/BlueprintGraph/Classes/K2Node_SetFieldsInStruct.h
@@ -38,7 +38,7 @@ class UK2Node_SetFieldsInStruct : public UK2Node_MakeStruct
 		AllOtherPins
 	};
 
-	BLUEPRINTGRAPH_API void RemoveFieldPins(const UEdGraphPin* InGraphPin, EPinsToRemove Selection);
+	BLUEPRINTGRAPH_API void RemoveFieldPins(UEdGraphPin* InGraphPin, EPinsToRemove Selection);
 	BLUEPRINTGRAPH_API bool AllPinsAreShown() const;
 	BLUEPRINTGRAPH_API void RestoreAllPins();
 
diff --git a/Engine/Source/Editor/BlueprintGraph/Classes/K2Node_StructMemberGet.h b/Engine/Source/Editor/BlueprintGraph/Classes/K2Node_StructMemberGet.h
--- a/Engine/Source/Editor/BlueprintGraph/Classes/K2Node_StructMemberGet.h
+++ b/Engine/Source/Editor/BlueprintGraph/Classes/K2Node_StructMemberGet.h
@@ -20,6 +20,7 @@ class UK2Node_StructMemberGet : public UK2Node_StructOperation
 
 	// UObject interface
 	virtual void PostEditChangeProperty(struct FPropertyChangedEvent& PropertyChangedEvent) override;
+	virtual void PreEditChange(UProperty* PropertyAboutToChange) override;
 	// End of UObject interface
 
 	// UEdGraphNode interface
@@ -41,5 +42,7 @@ private:
 	/** Constructing FText strings can be costly, so we cache the node's title/tooltip */
 	FNodeTextCache CachedTooltip;
 	FNodeTextCache CachedNodeTitle;
+
+	TArray<FName> OldShownPins;
 };
 
diff --git a/Engine/Source/Editor/BlueprintGraph/Classes/K2Node_StructMemberSet.h b/Engine/Source/Editor/BlueprintGraph/Classes/K2Node_StructMemberSet.h
--- a/Engine/Source/Editor/BlueprintGraph/Classes/K2Node_StructMemberSet.h
+++ b/Engine/Source/Editor/BlueprintGraph/Classes/K2Node_StructMemberSet.h
@@ -19,6 +19,7 @@ class UK2Node_StructMemberSet : public UK2Node_StructOperation
 	TArray<FOptionalPinFromProperty> ShowPinForProperties;
 
 	// UObject interface
+	virtual void PreEditChange(UProperty* PropertyAboutToChange) override;
 	virtual void PostEditChangeProperty(struct FPropertyChangedEvent& PropertyChangedEvent) override;
 	// End of UObject interface
 
@@ -39,5 +40,7 @@ private:
 	/** Constructing FText strings can be costly, so we cache the node's title/tooltip */
 	FNodeTextCache CachedTooltip;
 	FNodeTextCache CachedNodeTitle;
+
+	TArray<FName> OldShownPins;
 };
 
diff --git a/Engine/Source/Editor/BlueprintGraph/Classes/K2Node_VariableGet.h b/Engine/Source/Editor/BlueprintGraph/Classes/K2Node_VariableGet.h
--- a/Engine/Source/Editor/BlueprintGraph/Classes/K2Node_VariableGet.h
+++ b/Engine/Source/Editor/BlueprintGraph/Classes/K2Node_VariableGet.h
@@ -27,6 +27,7 @@ class BLUEPRINTGRAPH_API UK2Node_VariableGet : public UK2Node_Variable
 	virtual FText GetTooltipText() const override;
 	virtual FText GetNodeTitle(ENodeTitleType::Type TitleType) const override;
 	virtual void GetContextMenuActions(const FGraphNodeContextMenuBuilder& Context) const override;
+	virtual void ValidateNodeDuringCompilation(class FCompilerResultsLog& MessageLog) const override;
 	//~ End UEdGraphNode Interface
 
 	//~ Begin K2Node Interface
diff --git a/Engine/Source/Editor/BlueprintGraph/Private/BlueprintActionDatabase.cpp b/Engine/Source/Editor/BlueprintGraph/Private/BlueprintActionDatabase.cpp
--- a/Engine/Source/Editor/BlueprintGraph/Private/BlueprintActionDatabase.cpp
+++ b/Engine/Source/Editor/BlueprintGraph/Private/BlueprintActionDatabase.cpp
@@ -740,16 +740,14 @@ static void BlueprintActionDatabaseImpl::AddBlueprintGraphActions(UBlueprint con
 {
 	using namespace FBlueprintNodeSpawnerFactory; // for MakeMacroNodeSpawner()
 
-	for (auto GraphIt = Blueprint->MacroGraphs.CreateConstIterator(); GraphIt; ++GraphIt)
+	for (UEdGraph* MacroGraph : Blueprint->MacroGraphs)
 	{
-		ActionListOut.Add(MakeMacroNodeSpawner(*GraphIt));
+		ActionListOut.Add(MakeMacroNodeSpawner(MacroGraph));
 	}
 
 	// local variables and parameters
-	for (auto GraphIt = Blueprint->FunctionGraphs.CreateConstIterator(); GraphIt; ++GraphIt)
+	for (UEdGraph* FunctionGraph : Blueprint->FunctionGraphs)
 	{
-		UEdGraph* FunctionGraph = (*GraphIt);
-
 		TArray<UK2Node_FunctionEntry*> GraphEntryNodes;
 		FunctionGraph->GetNodesOfClass<UK2Node_FunctionEntry>(GraphEntryNodes);
 
@@ -1101,11 +1099,11 @@ void FBlueprintActionDatabase::Tick(float DeltaTime)
 	{
 		auto ActionIndex = ActionPrimingQueue.CreateIterator();	
 			
-		TWeakObjectPtr<UObject> ActionsKey = ActionIndex.Key();
-		if (ActionsKey.IsValid())
+		const FObjectKey& ActionsKey = ActionIndex.Key();
+		if (ActionsKey.ResolveObjectPtr())
 		{
 			// make sure this class is still listed in the database
-			if (FActionList* ClassActionList = ActionRegistry.Find(ActionsKey.Get()))
+			if (FActionList* ClassActionList = ActionRegistry.Find(ActionsKey))
 			{
 				int32& ActionListIndex = ActionIndex.Value();
 				for (; (ActionListIndex < ClassActionList->Num()) && (PrimedCount < PrimingMaxPerFrame); ++ActionListIndex)
@@ -1130,6 +1128,17 @@ void FBlueprintActionDatabase::Tick(float DeltaTime)
 			ActionPrimingQueue.Remove(ActionsKey);
 		}
 	}
+
+	// Handle deferred removals.
+	while (ActionRemoveQueue.Num() > 0)
+	{
+		TArray<UBlueprintNodeSpawner*> NodeSpawners = ActionRegistry.FindAndRemoveChecked(ActionRemoveQueue.Pop());
+		for (UBlueprintNodeSpawner* Action : NodeSpawners)
+		{
+			check(Action);
+			Action->ClearCachedTemplateNode();
+		}
+	}
 }
 
 //------------------------------------------------------------------------------
@@ -1138,6 +1147,12 @@ TStatId FBlueprintActionDatabase::GetStatId() const
 	RETURN_QUICK_DECLARE_CYCLE_STAT(FBlueprintActionDatabase, STATGROUP_Tickables);
 }
 
+//------------------------------------------------------------------------------
+void FBlueprintActionDatabase::DeferredRemoveEntry(FObjectKey const& InKey)
+{
+	ActionRemoveQueue.AddUnique(InKey);
+}
+
 //------------------------------------------------------------------------------
 void FBlueprintActionDatabase::RefreshAll()
 {
@@ -1179,7 +1194,7 @@ void FBlueprintActionDatabase::RefreshWorlds()
 	// Add all level scripts from current world
 	const TIndirectArray<FWorldContext>& WorldContexts = GEngine->GetWorldContexts();
 
-	for( auto Context : WorldContexts )
+	for (const FWorldContext& Context : WorldContexts)
 	{
 		if( Context.WorldType == EWorldType::Editor )
 		{
@@ -1338,7 +1353,7 @@ void FBlueprintActionDatabase::RefreshAssetActions(UObject* const AssetObject)
 	UWorld* WorldAsset = Cast<UWorld>(AssetObject);
 	if (WorldAsset && WorldAsset->WorldType == EWorldType::Editor)
 	{
-		for( auto Level : WorldAsset->GetLevels() )
+		for( ULevel* Level : WorldAsset->GetLevels() )
 		{
 			UBlueprint* LevelScript = Cast<UBlueprint>(Level->GetLevelScriptBlueprint(true));
 			if (LevelScript != nullptr)
@@ -1381,8 +1396,8 @@ void FBlueprintActionDatabase::RefreshAssetActions(UObject* const AssetObject)
 	// we don't want to clear entries for blueprints, mainly because we 
 	// use the presence of an entry to know if we've set the blueprint's 
 	// OnChanged(), but also because most blueprints will have actions at some 
-	// later point. Same goes for world assets because they are used to managed level scripts blueprints.
-	else if (!BlueprintAsset && !WorldAsset)
+	// later point. Same goes for in-editor world assets because they are used to manage level script blueprints.
+	else if (!BlueprintAsset && (!WorldAsset || WorldAsset->WorldType != EWorldType::Editor))
 	{
 		ClearAssetActions(AssetObject);
 	}
@@ -1399,7 +1414,7 @@ void FBlueprintActionDatabase::RefreshComponentActions()
 	check(ComponentTypes);
 	FActionList& ClassActionList = ActionRegistry.FindOrAdd(UBlueprintComponentNodeSpawner::StaticClass());
 	ClassActionList.Empty(ComponentTypes->Num());
-	for (const auto& ComponentType : *ComponentTypes)
+	for (const FComponentTypeEntry& ComponentType : *ComponentTypes)
 	{
 		if (UBlueprintComponentNodeSpawner* NodeSpawner = UBlueprintComponentNodeSpawner::Create(ComponentType))
 		{
@@ -1447,7 +1462,7 @@ void FBlueprintActionDatabase::ClearUnloadedAssetActions(FName ObjectPath)
 	{
 		for(UBlueprintNodeSpawner* NodeSpawner : *UnloadedActionList)
 		{
-			FActionList* ActionList = ActionRegistry.Find(NodeSpawner->NodeClass);
+			FActionList* ActionList = ActionRegistry.Find(NodeSpawner->NodeClass.Get());
 
 			// Remove the NodeSpawner from the main registry, it will be replaced with the loaded version of the action
 			ActionList->Remove(NodeSpawner);
diff --git a/Engine/Source/Editor/BlueprintGraph/Private/BlueprintActionFilter.cpp b/Engine/Source/Editor/BlueprintGraph/Private/BlueprintActionFilter.cpp
--- a/Engine/Source/Editor/BlueprintGraph/Private/BlueprintActionFilter.cpp
+++ b/Engine/Source/Editor/BlueprintGraph/Private/BlueprintActionFilter.cpp
@@ -808,8 +808,8 @@ static bool BlueprintActionFilterImpl::IsPermissionNotGranted(FBlueprintActionFi
 		UClass const* const NodeClass = BlueprintAction.GetNodeClass();
 		for (UBlueprint const* Blueprint : FilterContext.Blueprints)
 		{
-			bool const bIsReadOnly = FBlueprintEditorUtils::IsPropertyReadOnlyInCurrentBlueprint(Blueprint, Property);
-			if (bIsReadOnly && NodeClass && NodeClass->IsChildOf<UK2Node_VariableSet>())
+			bool const bIsWritable = (FBlueprintEditorUtils::IsPropertyWritableInBlueprint(Blueprint, Property) == FBlueprintEditorUtils::EPropertyWritableState::Writable);
+			if (!bIsWritable && NodeClass && NodeClass->IsChildOf<UK2Node_VariableSet>())
 			{
 				bIsFilteredOut = true;
 			}
@@ -1805,7 +1805,7 @@ UClass const* FBlueprintActionInfo::GetOwnerClass()
 		}
 		else if (const UBlueprint* AsBlueprint = Cast<UBlueprint>(ActionOwner.Get()))
 		{
-			CachedOwnerClass = AsBlueprint->SkeletonGeneratedClass;
+			CachedOwnerClass = AsBlueprint->SkeletonGeneratedClass.Get();
 		}
 
 		if (CachedOwnerClass == nullptr)
@@ -1818,7 +1818,7 @@ UClass const* FBlueprintActionInfo::GetOwnerClass()
 
 		CacheFlags |= EBlueprintActionInfoFlags::CachedClass;
 	}
-	return CachedOwnerClass;
+	return CachedOwnerClass.Get();
 }
 
 //------------------------------------------------------------------------------
diff --git a/Engine/Source/Editor/BlueprintGraph/Private/CallFunctionHandler.cpp b/Engine/Source/Editor/BlueprintGraph/Private/CallFunctionHandler.cpp
--- a/Engine/Source/Editor/BlueprintGraph/Private/CallFunctionHandler.cpp
+++ b/Engine/Source/Editor/BlueprintGraph/Private/CallFunctionHandler.cpp
@@ -66,9 +66,7 @@ void FKCHandler_CallFunction::CreateFunctionCallStatement(FKismetFunctionContext
 	int32 NumErrorsAtStart = CompilerContext.MessageLog.NumErrors;
 
 	// Find the function, starting at the parent class
-	UFunction* Function = FindFunction(Context, Node);
-	
-	if (Function != NULL)
+	if (UFunction* Function = FindFunction(Context, Node))
 	{
 		CheckIfFunctionIsCallable(Function, Context, Node);
 		// Make sure the pin mapping is sound (all pins wire up to a matching function parameter, and all function parameters match a pin)
@@ -77,19 +75,10 @@ void FKCHandler_CallFunction::CreateFunctionCallStatement(FKismetFunctionContext
 		TArray<UEdGraphPin*> RemainingPins;
 		RemainingPins.Append(Node->Pins);
 
+		const UEdGraphSchema_K2* Schema = CompilerContext.GetSchema();
+
 		// Remove expected exec and self pins
-		//@TODO: Check to make sure there is exactly one exec in and one exec out, as well as one self pin
-		for (int32 i = 0; i < RemainingPins.Num(); )
-		{
-			if (CompilerContext.GetSchema()->IsMetaPin(*RemainingPins[i]))
-			{
-				RemainingPins.RemoveAtSwap(i);
-			}
-			else
-			{
-				++i;
-			}
-		}
+		RemainingPins.RemoveAllSwap([Schema](UEdGraphPin* Pin) { return (Pin->bOrphanedPin || Schema->IsMetaPin(*Pin)); }, false);
 
 		// Check for magic pins
 		const bool bIsLatent = Function->HasMetaData(FBlueprintMetadata::MD_Latent);
@@ -98,10 +87,9 @@ void FKCHandler_CallFunction::CreateFunctionCallStatement(FKismetFunctionContext
 			CompilerContext.MessageLog.Error(*LOCTEXT("ContainsLatentCall_Error", "@@ contains a latent call, which cannot exist outside of the event graph").ToString(), Node);
 		}
 
-		UEdGraphPin* LatentInfoPin = NULL;
+		UEdGraphPin* LatentInfoPin = nullptr;
 
-		TMap<FName, FString>* MetaData = UMetaData::GetMapForObject(Function);
-		if (MetaData != NULL)
+		if (TMap<FName, FString>* MetaData = UMetaData::GetMapForObject(Function))
 		{
 			for (TMap<FName, FString>::TConstIterator It(*MetaData); It; ++It)
 			{
@@ -110,13 +98,12 @@ void FKCHandler_CallFunction::CreateFunctionCallStatement(FKismetFunctionContext
 				if (Key == TEXT("LatentInfo"))
 				{
 					UEdGraphPin* Pin = Node->FindPin(It.Value());
-					if( (Pin != NULL) && (Pin->Direction == EGPD_Input) && (Pin->LinkedTo.Num() == 0))
+					if (Pin && (Pin->Direction == EGPD_Input) && (Pin->LinkedTo.Num() == 0))
 					{
 						LatentInfoPin = Pin;
 
 						UEdGraphPin* PinToTry = FEdGraphUtilities::GetNetFromPin(Pin);
-						FBPTerminal** Term = Context.NetMap.Find(PinToTry);
-						if(Term != NULL)
+						if (FBPTerminal** Term = Context.NetMap.Find(PinToTry))
 						{
 							check((*Term)->bIsLiteral);
 						
@@ -140,10 +127,10 @@ void FKCHandler_CallFunction::CreateFunctionCallStatement(FKismetFunctionContext
 		}
 
 		// Parameter info to be stored, and assigned to all function call statements generated below
-		FBPTerminal* LHSTerm = NULL;
+		FBPTerminal* LHSTerm = nullptr;
 		TArray<FBPTerminal*> RHSTerms;
-		UEdGraphPin* ThenExecPin = NULL;
-		UEdGraphNode* LatentTargetNode = NULL;
+		UEdGraphPin* ThenExecPin = nullptr;
+		UEdGraphNode* LatentTargetNode = nullptr;
 		int32 LatentTargetParamIndex = INDEX_NONE;
 
 		// Grab the special case structs that use their own literal path
@@ -167,9 +154,8 @@ void FKCHandler_CallFunction::CreateFunctionCallStatement(FKismetFunctionContext
 					if (FKismetCompilerUtilities::IsTypeCompatibleWithProperty(PinMatch, Property, CompilerContext.MessageLog, CompilerContext.GetSchema(), Context.NewClass))
 					{
 						UEdGraphPin* PinToTry = FEdGraphUtilities::GetNetFromPin(PinMatch);
-						FBPTerminal** Term = Context.NetMap.Find(PinToTry);
 
-						if (Term != NULL)
+						if (FBPTerminal** Term = Context.NetMap.Find(PinToTry))
 						{
 							// For literal structs, we have to verify the default here to make sure that it has valid formatting
 							if( (*Term)->bIsLiteral && (PinMatch != LatentInfoPin))
@@ -245,12 +231,12 @@ void FKCHandler_CallFunction::CreateFunctionCallStatement(FKismetFunctionContext
 									// Record the (latent) output impulse from this node
 									ThenExecPin = CompilerContext.GetSchema()->FindExecutionPin(*Node, EGPD_Output);
 
-									if( (ThenExecPin != NULL) && (ThenExecPin->LinkedTo.Num() > 0) )
+									if( ThenExecPin && (ThenExecPin->LinkedTo.Num() > 0) )
 									{
 										LatentTargetNode = ThenExecPin->LinkedTo[0]->GetOwningNode();
 									}
 
-									if( LatentTargetNode != NULL )
+									if (LatentTargetNode)
 									{
 										LatentTargetParamIndex = ParameterIndex;
 									}
@@ -306,7 +292,7 @@ void FKCHandler_CallFunction::CreateFunctionCallStatement(FKismetFunctionContext
 		{
 			// Build up a list of contexts that this function will be called on
 			TArray<FBPTerminal*> ContextTerms;
-			if (SelfPin != NULL)
+			if (SelfPin)
 			{
 				const bool bIsConstSelfContext = Context.IsConstFunction();
 				const bool bIsNonConstFunction = !Function->HasAnyFunctionFlags(FUNC_Const|FUNC_Static);
@@ -415,10 +401,8 @@ void FKCHandler_CallFunction::CreateFunctionCallStatement(FKismetFunctionContext
 
 			// Iterate over all the contexts this functions needs to be called on, and emit a call function statement for each
 			FBlueprintCompiledStatement* LatentStatement = nullptr;
-			for (auto TargetListIt = ContextTerms.CreateIterator(); TargetListIt; ++TargetListIt)
+			for (FBPTerminal* Target : ContextTerms)
 			{
-				FBPTerminal* Target = *TargetListIt;
-
 				FBlueprintCompiledStatement& Statement = Context.AppendStatementForNode(Node);
 				Statement.FunctionToCall = Function;
 				Statement.FunctionContext = Target;
@@ -432,7 +416,7 @@ void FKCHandler_CallFunction::CreateFunctionCallStatement(FKismetFunctionContext
 				if (!bIsLatent)
 				{
 					// Fixup ubergraph calls
-					if (pSrcEventNode != NULL)
+					if (pSrcEventNode)
 					{
 						UEdGraphPin* ExecOut = CompilerContext.GetSchema()->FindExecutionPin(**pSrcEventNode, EGPD_Output);
 
@@ -444,7 +428,7 @@ void FKCHandler_CallFunction::CreateFunctionCallStatement(FKismetFunctionContext
 				else
 				{
 					// Fixup latent functions
-					if ((LatentTargetNode != NULL) && (Target == ContextTerms.Last()))
+					if (LatentTargetNode && (Target == ContextTerms.Last()))
 					{
 						check(LatentTargetParamIndex != INDEX_NONE);
 						Statement.UbergraphCallIndex = LatentTargetParamIndex;
@@ -463,7 +447,7 @@ void FKCHandler_CallFunction::CreateFunctionCallStatement(FKismetFunctionContext
 				{
 					FBlueprintCompiledStatement& SuspendState = Context.AppendStatementForNode(Node);
 					SuspendState.Type = KCST_InstrumentedStateSuspend;
-					if (LatentTargetNode != NULL)
+					if (LatentTargetNode)
 					{
 						check(LatentTargetParamIndex != INDEX_NONE);
 						SuspendState.UbergraphCallIndex = LatentTargetParamIndex;
@@ -509,7 +493,7 @@ UClass* FKCHandler_CallFunction::GetCallingContext(FKismetFunctionContext& Conte
 	// Find the calling scope
 	UClass* SearchScope = Context.NewClass;
 	UK2Node_CallFunction* CallFuncNode = Cast<UK2Node_CallFunction>(Node);
-	if ((CallFuncNode != NULL) && CallFuncNode->bIsFinalFunction)
+	if (CallFuncNode && CallFuncNode->bIsFinalFunction)
 	{
 		if (UK2Node_CallParentFunction* ParentCall = Cast<UK2Node_CallParentFunction>(Node))
 		{
@@ -517,13 +501,13 @@ UClass* FKCHandler_CallFunction::GetCallingContext(FKismetFunctionContext& Conte
 			const FName FuncName = CallFuncNode->FunctionReference.GetMemberName();
 			UClass* SearchContext = Context.NewClass->GetSuperClass();
 
-			UFunction* ParentFunc = NULL;
-			if (SearchContext != NULL)
+			UFunction* ParentFunc = nullptr;
+			if (SearchContext)
 			{
 				ParentFunc = SearchContext->FindFunctionByName(FuncName);
 			}
 
-			return ParentFunc ? ParentFunc->GetOuterUClass() : NULL;
+			return ParentFunc ? ParentFunc->GetOuterUClass() : nullptr;
 		}
 		else
 		{
@@ -533,8 +517,7 @@ UClass* FKCHandler_CallFunction::GetCallingContext(FKismetFunctionContext& Conte
 	}
 	else
 	{
-		UEdGraphPin* SelfPin = CompilerContext.GetSchema()->FindSelfPin(*Node, EGPD_Input);
-		if (SelfPin != NULL)
+		if (UEdGraphPin* SelfPin = CompilerContext.GetSchema()->FindSelfPin(*Node, EGPD_Input))
 		{
 			SearchScope = Cast<UClass>(Context.GetScopeFromPinType(SelfPin->PinType, Context.NewClass));
 		}
@@ -545,7 +528,7 @@ UClass* FKCHandler_CallFunction::GetCallingContext(FKismetFunctionContext& Conte
 
 UClass* FKCHandler_CallFunction::GetTrueCallingClass(FKismetFunctionContext& Context, UEdGraphPin* SelfPin)
 {
-	if (SelfPin != NULL)
+	if (SelfPin)
 	{
 		UEdGraphSchema_K2 const* K2Schema = CompilerContext.GetSchema();
 
@@ -567,8 +550,7 @@ void FKCHandler_CallFunction::RegisterNets(FKismetFunctionContext& Context, UEdG
 {
 	UEdGraphSchema_K2 const* K2Schema = CompilerContext.GetSchema();
 
-	UFunction* Function = FindFunction(Context, Node);
-	if (Function != NULL)
+	if (UFunction* Function = FindFunction(Context, Node))
 	{
 		TArray<FString> DefaultToSelfParamNames;
 		TArray<FString> RequiresSetValue;
@@ -595,17 +577,16 @@ void FKCHandler_CallFunction::RegisterNets(FKismetFunctionContext& Context, UEdG
 			}
 		}
 
-		for (auto It = Node->Pins.CreateConstIterator(); It; ++It)
+		for (UEdGraphPin* Pin : Node->Pins)
 		{
-			UEdGraphPin* Pin = (*It);
 			const bool bIsConnected = (Pin->LinkedTo.Num() != 0);
 
 			// if this pin could use a default (it doesn't have a connection or default of its own)
-			if (!bIsConnected && (Pin->DefaultObject == NULL))
+			if (!bIsConnected && (Pin->DefaultObject == nullptr))
 			{
 				if (DefaultToSelfParamNames.Contains(Pin->PinName) && FKismetCompilerUtilities::ValidateSelfCompatibility(Pin, Context))
 				{
-					ensure(Pin->PinType.PinSubCategoryObject != NULL);
+					ensure(Pin->PinType.PinSubCategoryObject != nullptr);
 					ensure((Pin->PinType.PinCategory == K2Schema->PC_Object) || (Pin->PinType.PinCategory == K2Schema->PC_Interface));
 
 					FBPTerminal* Term = Context.RegisterLiteral(Pin);
diff --git a/Engine/Source/Editor/BlueprintGraph/Private/EdGraphSchema_K2.cpp b/Engine/Source/Editor/BlueprintGraph/Private/EdGraphSchema_K2.cpp
--- a/Engine/Source/Editor/BlueprintGraph/Private/EdGraphSchema_K2.cpp
+++ b/Engine/Source/Editor/BlueprintGraph/Private/EdGraphSchema_K2.cpp
@@ -1203,7 +1203,7 @@ bool UEdGraphSchema_K2::IsAllowableBlueprintVariableType(const class UClass* InC
 	return false;
 }
 
-bool UEdGraphSchema_K2::IsAllowableBlueprintVariableType(const class UScriptStruct *InStruct)
+bool UEdGraphSchema_K2::IsAllowableBlueprintVariableType(const UScriptStruct* InStruct, const bool bForInternalUse)
 {
 	if (const UUserDefinedStruct* UDStruct = Cast<const UUserDefinedStruct>(InStruct))
 	{
@@ -1212,7 +1212,7 @@ bool UEdGraphSchema_K2::IsAllowableBlueprintVariableType(const class UScriptStru
 			return false;
 		}
 	}
-	return InStruct && (InStruct->GetBoolMetaDataHierarchical(FBlueprintMetadata::MD_AllowableBlueprintVariableType));
+	return (InStruct && InStruct->GetBoolMetaDataHierarchical(FBlueprintMetadata::MD_AllowableBlueprintVariableType) && (bForInternalUse || !InStruct->GetBoolMetaData(FBlueprintMetadata::MD_BlueprintInternalUseOnly)));
 }
 
 bool UEdGraphSchema_K2::DoesGraphSupportImpureFunctions(const UEdGraph* InGraph) const
@@ -1286,7 +1286,7 @@ bool UEdGraphSchema_K2::PinHasSplittableStructType(const UEdGraphPin* InGraphPin
 		UScriptStruct* StructType = CastChecked<UScriptStruct>(InGraphPin->PinType.PinSubCategoryObject.Get());
 		if (InGraphPin->Direction == EGPD_Input)
 		{
-			bCanSplit = UK2Node_MakeStruct::CanBeMade(StructType, true, true);
+			bCanSplit = UK2Node_MakeStruct::CanBeSplit(StructType);
 			if (!bCanSplit)
 			{
 				const FString& MetaData = StructType->GetMetaData(TEXT("HasNativeMake"));
@@ -1296,7 +1296,7 @@ bool UEdGraphSchema_K2::PinHasSplittableStructType(const UEdGraphPin* InGraphPin
 		}
 		else
 		{
-			bCanSplit = UK2Node_BreakStruct::CanBeBroken(StructType, true, true);
+			bCanSplit = UK2Node_BreakStruct::CanBeSplit(StructType);
 			if (!bCanSplit)
 			{
 				const FString& MetaData = StructType->GetMetaData(TEXT("HasNativeBreak"));
@@ -2234,6 +2234,11 @@ const FPinConnectionResponse UEdGraphSchema_K2::CanCreateConnection(const UEdGra
 		return FPinConnectionResponse(CONNECT_RESPONSE_DISALLOW, TEXT("Both are on the same node"));
 	}
 
+	if (PinA->bOrphanedPin || PinB->bOrphanedPin)
+	{
+		return FPinConnectionResponse(CONNECT_RESPONSE_DISALLOW, TEXT("Cannot make new connections to orphaned pin"));
+	}
+
 	FString NodeResponseMessage;
 	// node can disallow the connection
 	{
@@ -2872,7 +2877,7 @@ bool UEdGraphSchema_K2::IsPinBeingWatched(UEdGraphPin const* Pin) const
 {
 	// Note: If you crash here; it is likely that you forgot to call Blueprint->OnBlueprintChanged.Broadcast(Blueprint) to invalidate the cached UI state
 	UBlueprint* Blueprint = FBlueprintEditorUtils::FindBlueprintForNode(Pin ? Pin->GetOwningNodeUnchecked() : nullptr);
-	return ensure(Blueprint) ? FKismetDebugUtilities::IsPinBeingWatched(Blueprint, Pin) : false;
+	return (Blueprint ? FKismetDebugUtilities::IsPinBeingWatched(Blueprint, Pin) : false);
 }
 
 void UEdGraphSchema_K2::ClearPinWatch(UEdGraphPin const* Pin) const
@@ -4541,6 +4546,12 @@ void UEdGraphSchema_K2::TrySetDefaultValue(UEdGraphPin& Pin, const FString& NewD
 		UEdGraphNode* Node = Pin.GetOwningNode();
 		Node->PinDefaultValueChanged(&Pin);
 
+		// If the default value is manually set then treat it as if the value was reset to default and remove the orphaned pin
+		if (Pin.bOrphanedPin && Pin.DoesDefaultValueMatchAutogenerated())
+		{
+			Node->PinConnectionListChanged(&Pin);
+		}
+
 		UBlueprint* Blueprint = FBlueprintEditorUtils::FindBlueprintForNodeChecked(Node);
 		FBlueprintEditorUtils::MarkBlueprintAsModified(Blueprint);
 	}
@@ -4567,6 +4578,12 @@ void UEdGraphSchema_K2::TrySetDefaultObject(UEdGraphPin& Pin, UObject* NewDefaul
 	UEdGraphNode* Node = Pin.GetOwningNode();
 	Node->PinDefaultValueChanged(&Pin);
 
+	// If the default value is manually set then treat it as if the value was reset to default and remove the orphaned pin
+	if (Pin.bOrphanedPin && Pin.DoesDefaultValueMatchAutogenerated())
+	{
+		Node->PinConnectionListChanged(&Pin);
+	}
+
 	UBlueprint* Blueprint = FBlueprintEditorUtils::FindBlueprintForNodeChecked(Node);
 	FBlueprintEditorUtils::MarkBlueprintAsModified(Blueprint);
 }
@@ -4587,14 +4604,71 @@ void UEdGraphSchema_K2::TrySetDefaultText(UEdGraphPin& InPin, const FText& InNew
 		UEdGraphNode* Node = InPin.GetOwningNode();
 		Node->PinDefaultValueChanged(&InPin);
 
+		// If the default value is manually set then treat it as if the value was reset to default and remove the orphaned pin
+		if (InPin.bOrphanedPin && InPin.DoesDefaultValueMatchAutogenerated())
+		{
+			Node->PinConnectionListChanged(&InPin);
+		}
+
 		UBlueprint* Blueprint = FBlueprintEditorUtils::FindBlueprintForNodeChecked(Node);
 		FBlueprintEditorUtils::MarkBlueprintAsModified(Blueprint);
 	}
 }
 
+bool UEdGraphSchema_K2::DoesDefaultValueMatchAutogenerated(const UEdGraphPin& InPin) const
+{
+	if (InPin.PinType.PinCategory == PC_Enum || InPin.PinType.PinCategory == PC_Byte)
+	{
+		// The autogenerated default value for an enum is left empty in case the default enum value (index 0) changes, in this case we
+		// want to validate against the actual value of the 0 index entry
+		if (InPin.AutogeneratedDefaultValue.IsEmpty())
+		{
+			const FString PinDefaultValue = InPin.GetDefaultAsString();
+			if (PinDefaultValue.IsEmpty())
+			{
+				return true;
+			}
+			else
+			{
+				if (UEnum* PinEnumType = Cast<UEnum>(InPin.PinType.PinSubCategoryObject.Get()))
+				{
+					return InPin.DefaultValue.Equals(PinEnumType->GetNameStringByIndex(0), ESearchCase::IgnoreCase);
+				}
+			}
+		}
+	}
+	else if (InPin.PinType.PinCategory == PC_Float)
+	{
+		const float AutogeneratedFloat = FCString::Atof(*InPin.AutogeneratedDefaultValue);
+		const float DefaultFloat = FCString::Atof(*InPin.DefaultValue);
+		return (AutogeneratedFloat == DefaultFloat);
+	}
+	else if (InPin.PinType.PinCategory == PC_Struct)
+	{
+		if (InPin.PinType.PinSubCategoryObject == VectorStruct)
+		{
+			FVector AutogeneratedVector = FVector::ZeroVector;
+			FVector DefaultVector = FVector::ZeroVector;
+			FDefaultValueHelper::ParseVector(InPin.AutogeneratedDefaultValue, AutogeneratedVector);
+			FDefaultValueHelper::ParseVector(InPin.DefaultValue, DefaultVector);
+			return (AutogeneratedVector == DefaultVector);
+		}
+		else if (InPin.PinType.PinSubCategoryObject == RotatorStruct)
+		{
+			FRotator AutogeneratedRotator = FRotator::ZeroRotator;
+			FRotator DefaultRotator = FRotator::ZeroRotator;
+			FDefaultValueHelper::ParseRotator(InPin.AutogeneratedDefaultValue, AutogeneratedRotator);
+			FDefaultValueHelper::ParseRotator(InPin.DefaultValue, DefaultRotator);
+			return (AutogeneratedRotator == DefaultRotator);
+		}
+	}
+
+	return Super::DoesDefaultValueMatchAutogenerated(InPin);
+}
+
 bool UEdGraphSchema_K2::IsAutoCreateRefTerm(const UEdGraphPin* Pin) const
 {
-	check(Pin != NULL);
+	check(Pin);
 
 	bool bIsAutoCreateRefTerm = false;
 	UEdGraphNode* OwningNode = Pin->GetOwningNode();
@@ -4702,20 +4776,20 @@ void UEdGraphSchema_K2::SetPinAutogeneratedDefaultValueBasedOnType(UEdGraphPin*
 	// Create a useful default value based on the pin type
 	if(Pin->PinType.IsContainer() )
 	{
-		NewValue = TEXT("");
+		NewValue = FString();
 	}
 	else if (Pin->PinType.PinCategory == PC_Int)
 	{
 		NewValue = TEXT("0");
 	}
-	else if(Pin->PinType.PinCategory == PC_Byte)
+	else if (Pin->PinType.PinCategory == PC_Byte)
 	{
 		UEnum* EnumPtr = Cast<UEnum>(Pin->PinType.PinSubCategoryObject.Get());
 		if(EnumPtr)
 		{
 			// First element of enum can change. If the enum is { A, B, C } and the default value is A, 
 			// the defult value should not change when enum will be changed into { N, A, B, C }
-			NewValue = TEXT("");
+			NewValue = FString();
 		}
 		else
 		{
@@ -4748,15 +4822,25 @@ void UEdGraphSchema_K2::SetPinAutogeneratedDefaultValueBasedOnType(UEdGraphPin*
 
 void UEdGraphSchema_K2::ResetPinToAutogeneratedDefaultValue(UEdGraphPin* Pin, bool bCallModifyCallbacks) const
 {
-	GetPinDefaultValuesFromString(Pin->PinType, Pin->GetOwningNodeUnchecked(), Pin->AutogeneratedDefaultValue, Pin->DefaultValue, Pin->DefaultObject, Pin->DefaultTextValue);
-
-	if (bCallModifyCallbacks)
+	if (Pin->bOrphanedPin)
 	{
 		UEdGraphNode* Node = Pin->GetOwningNode();
-		Node->PinDefaultValueChanged(Pin);
+		Node->PinConnectionListChanged(Pin);
+	}
+	else
+	{
+		GetPinDefaultValuesFromString(Pin->PinType, Pin->GetOwningNodeUnchecked(), Pin->AutogeneratedDefaultValue, Pin->DefaultValue, Pin->DefaultObject, Pin->DefaultTextValue);
 
-		UBlueprint* Blueprint = FBlueprintEditorUtils::FindBlueprintForNodeChecked(Node);
-		FBlueprintEditorUtils::MarkBlueprintAsModified(Blueprint);
+		if (bCallModifyCallbacks)
+		{
+			UEdGraphNode* Node = Pin->GetOwningNode();
+			Node->PinDefaultValueChanged(Pin);
+
+			if (UBlueprint* Blueprint = FBlueprintEditorUtils::FindBlueprintForNode(Node))
+			{
+				FBlueprintEditorUtils::MarkBlueprintAsModified(Blueprint);
+			}
+		}
 	}
 }
 
@@ -4827,9 +4911,8 @@ UFunction* UEdGraphSchema_K2::FindSetVariableByNameFunction(const FEdGraphPinTyp
 	{
 		static bool Is(const UObject* Obj)
 		{
-			static const FName BlueprintTypeName(TEXT("BlueprintType"));
 			const UScriptStruct* Struct = Cast<const UScriptStruct>(Obj);
-			return Struct ? Struct->GetBoolMetaData(BlueprintTypeName) : false;
+			return Struct ? Struct->GetBoolMetaData(FBlueprintMetadata::MD_AllowableBlueprintVariableType) : false;
 		}
 	};
 
@@ -5938,18 +6021,17 @@ UEdGraphNode* UEdGraphSchema_K2::CreateSubstituteNode(UEdGraphNode* Node, const
 				{
 					PreExistingNode = FBlueprintEditorUtils::FindCustomEventNode(Blueprint, EventNode->CustomFunctionName);
 				}
-				else
+				else if (UFunction* EventSignature = EventNode->FindEventSignatureFunction())
 				{
-					UFunction* EventSignature = EventNode->FindEventSignatureFunction();
-					if (ensure(EventSignature))
+					// Note: EventNode::FindEventSignatureFunction will return null if it is deleted (for instance, someone declared a 
+					// BlueprintImplementableEvent, and some blueprint implements it, but then the declaration is deleted). It also
+					// returns null if the pasted node was sourced from another asset that's not included in the destination project.
+					// This is acceptable since we've already created a substitute anyway; this is just looking to see if we actually
+					// have a valid pre-existing node that was in conflict, in which case we will emit a warning to the message log.
+					UClass* ClassOwner = EventSignature->GetOwnerClass();
+					if (ensureMsgf(ClassOwner, TEXT("Wrong class owner of signature %s in node %s"), *GetPathNameSafe(EventSignature), *GetPathNameSafe(EventNode)))
 					{
-						// EventNode::FindEventSignatureFunction will return null if it is deleted (for instance, someone declared a 
-						// BlueprintImplementableEvent, and some blueprint implements it, but then the declaration is deleted)
-						UClass* ClassOwner = EventSignature->GetOwnerClass();
-						if (ensureMsgf(ClassOwner, TEXT("Wrong class owner of signature %s in node %s"), *GetPathNameSafe(EventSignature), *GetPathNameSafe(EventNode)))
-						{
-							PreExistingNode = FBlueprintEditorUtils::FindOverrideForFunction(Blueprint, ClassOwner->GetAuthoritativeClass(), EventSignature->GetFName());
-						}
+						PreExistingNode = FBlueprintEditorUtils::FindOverrideForFunction(Blueprint, ClassOwner->GetAuthoritativeClass(), EventSignature->GetFName());
 					}
 				}
 			}
diff --git a/Engine/Source/Editor/BlueprintGraph/Private/K2Node.cpp b/Engine/Source/Editor/BlueprintGraph/Private/K2Node.cpp
--- a/Engine/Source/Editor/BlueprintGraph/Private/K2Node.cpp
+++ b/Engine/Source/Editor/BlueprintGraph/Private/K2Node.cpp
@@ -37,6 +37,7 @@ UK2Node::UK2Node(const FObjectInitializer& ObjectInitializer)
 	: Super(ObjectInitializer)
 {
 	bAllowSplitPins_DEPRECATED = true;
+	OrphanedPinSaveMode = ESaveOrphanPinMode::SaveAllButExec;
 }
 
 void UK2Node::PostLoad()
@@ -71,6 +72,8 @@ void UK2Node::PostLoad()
 
 void UK2Node::Serialize(FArchive& Ar)
 {
+	Ar.UsingCustomVersion(FFrameworkObjectVersion::GUID);
+
 	if (Ar.IsSaving())
 	{	
 		for (UEdGraphPin* Pin : Pins)
@@ -84,17 +87,12 @@ void UK2Node::Serialize(FArchive& Ar)
 
 					if (Struct)
 					{
-						int32 StructSize = Struct->GetStructureSize();
-						uint8* StructData = (uint8*)FMemory::Malloc(StructSize);
+						TSharedPtr<FStructOnScope> StructData = MakeShareable(new FStructOnScope(Struct));
 
-						// Import the literal text to a dummy struct and then serialize that
+						// Import the literal text to a dummy struct and then serialize that. Hard object references will not properly import, this is only useful for soft references!
 						FOutputDeviceNull NullOutput;
-						Struct->InitializeStruct(StructData);
-						Struct->ImportText(*Pin->DefaultValue, StructData, nullptr, PPF_None, &NullOutput, Pin->PinName);
-						Struct->SerializeItem(Ar, StructData, nullptr);
-						Struct->DestroyStruct(StructData);
-
-						FMemory::Free(StructData);
+						Struct->ImportText(*Pin->DefaultValue, StructData->GetStructMemory(), nullptr, PPF_SerializedAsImportText, &NullOutput, Pin->PinName);						
+						Struct->SerializeItem(Ar, StructData->GetStructMemory(), nullptr);
 					}
 				}
 
@@ -143,16 +141,17 @@ void UK2Node::FixupPinDefaultValues()
 			TArray<FString> RedirectPinNames;
 			GetRedirectPinNames(*Pin, RedirectPinNames);
 
-			for(auto PinNameIter = RedirectPinNames.CreateConstIterator(); PinNameIter && !bFoundPin; ++PinNameIter)
+			for (const FString& PinName : RedirectPinNames)
 			{
 				if((Pin->DefaultValue == Pin->AutogeneratedDefaultValue)
-					&& (*PinNameIter == VectorsEqualFunctionEpsilonPinName
-						|| *PinNameIter == VectorsNotEqualFunctionEpsilonPinName
-						|| *PinNameIter == RotatorsEqualFunctionEpsilonPinName
-						|| *PinNameIter == RotatorsNotEqualFunctionEpsilonPinName))
+					&& (PinName == VectorsEqualFunctionEpsilonPinName
+						|| PinName == VectorsNotEqualFunctionEpsilonPinName
+						|| PinName == RotatorsEqualFunctionEpsilonPinName
+						|| PinName == RotatorsNotEqualFunctionEpsilonPinName))
 				{
 					bFoundPin = true;
 					K2Schema->TrySetDefaultValue(*Pin, TEXT("0.0"));
+					break;
 				}
 			}
 		}
@@ -241,7 +240,7 @@ void UK2Node::AutowireNewNode(UEdGraphPin* FromPin)
 	const UEdGraphSchema_K2* K2Schema = CastChecked<UEdGraphSchema_K2>(GetSchema());
 	
 	// Do some auto-connection
-	if (FromPin != NULL)
+	if (FromPin)
 	{
 		TSet<UEdGraphNode*> NodeList;
 
@@ -317,9 +316,8 @@ void UK2Node::AutowireNewNode(UEdGraphPin* FromPin)
 		}
 
 		// Send all nodes that received a new pin connection a notification
-		for (auto It = NodeList.CreateConstIterator(); It; ++It)
+		for (UEdGraphNode* Node : NodeList)
 		{
-			UEdGraphNode* Node = (*It);
 			Node->NodeConnectionListChanged();
 		}
 	}
@@ -365,26 +363,97 @@ void UK2Node::GetNodeAttributes( TArray<TKeyValuePair<FString, FString>>& OutNod
 	OutNodeAttributes.Add( TKeyValuePair<FString, FString>( TEXT( "Name" ), GetName() ));
 }
 
+void UK2Node::PreloadRequiredAssets()
+{
+}
+
 void UK2Node::PinConnectionListChanged(UEdGraphPin* Pin) 
 {
+
 	// If the pin has been connected, clear the default values so that we don't hold on to references
-	if(Pin->LinkedTo.Num() > 0) 
+	if (Pin->LinkedTo.Num() > 0) 
 	{
-		// Verify that we have a proper outer at this point, such that the schema will be valid
-		UEdGraph* OuterGraph = GetGraph();
-		if( OuterGraph && OuterGraph->Schema )
+		// We don't want to reset Output pin defaults, that breaks Function entry nodes
+		if (!Pin->bOrphanedPin && Pin->Direction == EEdGraphPinDirection::EGPD_Input)
 		{
-			const UEdGraphSchema_K2* Schema = Cast<const UEdGraphSchema_K2>(GetSchema());
+			UEdGraph* OuterGraph = GetGraph();
 
-			// We don't want to reset Output pin defaults, that breaks Function entry nodes
-			if (Pin->Direction == EEdGraphPinDirection::EGPD_Input)
+			// Verify that we have a proper outer at this point, such that the schema will be valid
+			if (OuterGraph && OuterGraph->Schema)
 			{
+				const UEdGraphSchema_K2* Schema = CastChecked<const UEdGraphSchema_K2>(GetSchema());
 				Schema->ResetPinToAutogeneratedDefaultValue(Pin);
 			}
 		}
 	}
+	// If we're not linked and this pin should no longer exist as part of the node, remove it
+	else if (Pin->bOrphanedPin)
+	{
+		UEdGraph* OuterGraph = GetGraph();
+
+		if (OuterGraph)
+		{
+			OuterGraph->NotifyGraphChanged();
+		}
+
+		if (Pin->ParentPin == nullptr)
+		{
+			RemovePin(Pin);
+		}
+		else
+		{
+			const UEdGraphSchema_K2* Schema = CastChecked<const UEdGraphSchema_K2>(GetSchema());
+			TFunction<void(UEdGraphPin*)> RemoveNestedPin = [this, Schema, &RemoveNestedPin](UEdGraphPin* NestedPin)
+			{
+				Modify();
+				ensure(Pins.Remove(NestedPin) == 1);
 
-	NotifyPinConnectionListChanged(Pin);
+				if (UEdGraphPin* ParentPin = NestedPin->ParentPin)
+				{
+					ParentPin->Modify();
+					ensure(ParentPin->SubPins.Remove(NestedPin) == 1);
+
+					if (ParentPin->SubPins.Num() == 0)
+					{
+						if (ParentPin->bOrphanedPin)
+						{
+							RemoveNestedPin(ParentPin);
+						}
+						else
+						{
+							Schema->RecombinePin(ParentPin);
+						}
+					}
+				}
+
+				NestedPin->MarkPendingKill();
+			};
+
+			RemoveNestedPin(Pin);
+		}
+
+		bool bOrphanedPinsGone = true;
+		for (UEdGraphPin* RemainingPin : Pins)
+		{
+			if (RemainingPin->bOrphanedPin)
+			{
+				bOrphanedPinsGone = false;
+				break;
+			}
+		}
+
+		if (bOrphanedPinsGone)
+		{
+			ClearCompilerMessage();
+		}
+
+		Pin = nullptr;
+	}
+
+	if (Pin)
+	{
+		NotifyPinConnectionListChanged(Pin);
+	}
 }
 
 UObject* UK2Node::GetJumpTargetForDoubleClick() const
@@ -409,9 +478,8 @@ void UK2Node::PostReconstructNode()
 		{
 			const UEdGraphSchema_K2* Schema = Cast<const UEdGraphSchema_K2>(GetSchema());
 			// fix up any pin data if it needs to 
-			for( auto PinIt = Pins.CreateIterator(); PinIt; ++PinIt )
+			for (UEdGraphPin* CurrentPin : Pins)
 			{
-				UEdGraphPin* CurrentPin = *PinIt;
 				const FString& PinCategory = CurrentPin->PinType.PinCategory;
 
 				// fix up enum names if it exists in enum redirects
@@ -455,6 +523,9 @@ void UK2Node::ReconstructNode()
 {
 	Modify();
 
+	// Clear previously set messages
+	ErrorMsg.Reset();
+
 	UBlueprint* Blueprint = GetBlueprint();
 
 	FLinkerLoad* Linker = Blueprint->GetLinker();
@@ -491,32 +562,11 @@ void UK2Node::ReconstructNode()
 
 	// Move the existing pins to a saved array
 	TArray<UEdGraphPin*> OldPins(Pins);
-	Pins.Empty();
+	Pins.Reset();
 
 	// Recreate the new pins
 	ReallocatePinsDuringReconstruction(OldPins);
-
-	bool bDestroyOldPins = true;
-
-	if (Pins.Num() == 0)
-	{
-		//keep old pins on callfunction so that graph doesn't get broken up just because function is missing
-		if (IsA(UK2Node_CallFunction::StaticClass()) || IsA(UK2Node_MacroInstance::StaticClass()))
-		{
-			Pins = OldPins;
-			bDestroyOldPins = false;
-		}
-	}
-	else
-	{
-		RewireOldPinsToNewPins(OldPins, Pins);
-	}
-
-
-	if (bDestroyOldPins)
-	{
-		DestroyPinList(OldPins);
-	}
+	RewireOldPinsToNewPins(OldPins, Pins);
 
 	// Let subclasses do any additional work
 	PostReconstructNode();
@@ -554,12 +604,12 @@ UK2Node::ERedirectType UK2Node::ShouldRedirectParam(const TArray<FString>& OldPi
 void UK2Node::RestoreSplitPins(TArray<UEdGraphPin*>& OldPins)
 {
 	// necessary to recreate split pins and keep their wires
-	for (auto OldPin : OldPins)
+	for (UEdGraphPin* OldPin : OldPins)
 	{
 		if (OldPin->ParentPin)
 		{
 			// find the new pin that corresponds to parent, and split it if it isn't already split
-			for (auto NewPin : Pins)
+			for (UEdGraphPin* NewPin : Pins)
 			{
 				if (FCString::Stricmp(*(NewPin->PinName), *(OldPin->ParentPin->PinName)) == 0)
 				{
@@ -734,23 +784,71 @@ void UK2Node::ReconstructSinglePin(UEdGraphPin* NewPin, UEdGraphPin* OldPin, ERe
 			break;
 		}
 	}
+}
+
+void UK2Node::ValidateOrphanPins(FCompilerResultsLog& MessageLog, const bool bStore) const
+{
+	for (UEdGraphPin* Pin : Pins)
+	{
+		if (Pin && Pin->bOrphanedPin)
+		{
+			if (Pin->LinkedTo.Num())
+			{
+				const FText LinkedMessage = LOCTEXT("RemovedConnectedPin", "In use pin @@ no longer exists on this node. Please refresh node or break links to remove pin.");
+				if (bStore)
+				{
+					MessageLog.StorePotentialError(this, *LinkedMessage.ToString(), Pin);
+				}
+				else
+				{
+					MessageLog.Error(*LinkedMessage.ToString(), Pin);
+				}
+			}
+			else if (!Pin->bHidden && !Pin->DoesDefaultValueMatchAutogenerated())
+			{
+				const FText NonDefaultMessage = LOCTEXT("RemovedNonDefaultPin", "Input pin @@ specifying non-default value no longer exists on this node. Please refresh node or reset pin to default value to remove pin.");
+				if (bStore)
+				{
+					MessageLog.StorePotentialWarning(this, *NonDefaultMessage.ToString(), Pin);
+				}
+				else
+				{
+					MessageLog.Warning(*NonDefaultMessage.ToString(), Pin);
+				}
+			}
+		}
+	}
+}
 
-	OldPin->MarkPendingKill();
+void UK2Node::EarlyValidation(FCompilerResultsLog& MessageLog) const
+{
+	ValidateOrphanPins(MessageLog, true);
 }
 
-void UK2Node::RewireOldPinsToNewPins(TArray<UEdGraphPin*>& InOldPins, TArray<UEdGraphPin*>& InNewPins)
+void UK2Node::ValidateNodeDuringCompilation(FCompilerResultsLog& MessageLog) const
 {
-	// No need to attempt anything if there are no new pins
-	if (InNewPins.Num() == 0)
+	Super::ValidateNodeDuringCompilation(MessageLog);
+
+	// Since this might be an expanded node, the validation will have been done on the source node (which will return this if not expanded)
+	if (const UK2Node* SourceNode = Cast<UK2Node>(MessageLog.GetSourceNode(this)))
 	{
-		return;
+		// If we don't commit any messages for the source node, then see if this node has messages to commit.
+		// This is primarily needed for nodes generated from inside macros.
+		if (!MessageLog.CommitPotentialMessages(const_cast<UK2Node*>(SourceNode)))
+		{
+			ValidateOrphanPins(MessageLog, false);
+		}
 	}
+}
+
+void UK2Node::RewireOldPinsToNewPins(TArray<UEdGraphPin*>& InOldPins, TArray<UEdGraphPin*>& InNewPins)
+{
+	TArray<UEdGraphPin*> OrphanedOldPins;
+	TMap<UEdGraphPin*, UEdGraphPin*> MatchedPins; // Old to New
+
+	const bool bSaveUnconnectedDefaultPins = (InNewPins.Num() == 0 && (IsA<UK2Node_CallFunction>() || IsA<UK2Node_MacroInstance>()));
 
 	// Rewire any connection to pins that are matched by name (O(N^2) right now)
-	// @TODO: Can do moderately smart things here if only one pin changes name 
-	//        by looking at it's relative position, etc..., rather than just 
-	//        failing to map it and breaking the links
-	//
 	// NOTE: we iterate backwards through the list because ReconstructSinglePin()
 	//       destroys pins as we go along (clearing out parent pointers, etc.); 
 	//       we need the parent pin chain intact for DoPinsMatchForReconstruction();              
@@ -762,9 +860,10 @@ void UK2Node::RewireOldPinsToNewPins(TArray<UEdGraphPin*>& InOldPins, TArray<UEd
 		UEdGraphPin* OldPin = InOldPins[OldPinIndex];
 
 		// common case is for InOldPins and InNewPins to match, so we start searching from the current index:
+		bool bMatched = false;
 		const int32 NumNewPins = InNewPins.Num();
-		int32 NewPinIndex = OldPinIndex % InNewPins.Num();
-		for (int32 NewPinCount = InNewPins.Num()-1; NewPinCount >= 0; --NewPinCount)
+		int32 NewPinIndex = (NumNewPins ? OldPinIndex % NumNewPins : 0);
+		for (int32 NewPinCount = NumNewPins - 1; NewPinCount >= 0; --NewPinCount)
 		{
 			// if InNewPins grows in this loop then we may skip entries and fail to find a match:
 			check(NumNewPins == InNewPins.Num());
@@ -774,12 +873,114 @@ void UK2Node::RewireOldPinsToNewPins(TArray<UEdGraphPin*>& InOldPins, TArray<UEd
 			if (RedirectType != ERedirectType_None)
 			{
 				ReconstructSinglePin(NewPin, OldPin, RedirectType);
+				MatchedPins.Add(OldPin, NewPin);
+				bMatched = true;
 				break;
 			}
 
 			NewPinIndex = (NewPinIndex + 1) % InNewPins.Num();
 		}
+
+		// Orphaned pins are those that existed in the OldPins array but do not in the NewPins.
+		// We will save these pins and add the to the NewPins array if they are linked to other pins or have non-default value unless:
+		// * The node has been flagged to not save orphaned pins
+		// * The pin has been flagged not be saved if orphaned
+		// * The pin is hidden and not a split pin
+		const bool bVisibleOrSplitPin = (!OldPin->bHidden || (OldPin->SubPins.Num() > 0));
+		if (UEdGraphPin::EnableOrphanPins() && !bMatched && bVisibleOrSplitPin && OldPin->bSavePinIfOrphaned)
+		{
+			// The node can specify to save no pins, all pins, or all but exec pins. However, even if all is specified Execute and Then are never saved
+			 const bool bSaveOrphanedPin = ((OrphanedPinSaveMode == ESaveOrphanPinMode::SaveAll) ||
+											 ((OrphanedPinSaveMode == ESaveOrphanPinMode::SaveAllButExec) && !UEdGraphSchema_K2::IsExecPin(*OldPin)));
+
+			if (bSaveOrphanedPin)
+			{
+				bool bSavePin = bSaveUnconnectedDefaultPins || (OldPin->LinkedTo.Num() > 0);
+
+				if (!bSavePin && OldPin->SubPins.Num() > 0)
+				{
+					// If this is a split pin then we need to save it if any of its children are being saved
+					for (UEdGraphPin* OldSubPin : OldPin->SubPins)
+					{
+						if (OldSubPin->bOrphanedPin)
+						{
+							bSavePin = true;
+							break;
+						}
+					}
+					// Once we know we are going to be saving it we need to clean up the SubPins list to be only pins being saved
+					if (bSavePin)
+					{
+						for (int32 SubPinIndex = OldPin->SubPins.Num() - 1; SubPinIndex >= 0; --SubPinIndex)
+						{
+							UEdGraphPin* SubPin = OldPin->SubPins[SubPinIndex];
+							if (!SubPin->bOrphanedPin)
+							{
+								OldPin->SubPins.RemoveAt(SubPinIndex, 1, false);
+								SubPin->MarkPendingKill();
+							}
+						}
+					}
+				}
+
+				// Input pins with non-default value should be saved
+				if (!bSavePin && OldPin->Direction == EGPD_Input && !OldPin->DoesDefaultValueMatchAutogenerated())
+				{
+					bSavePin = true;
+				}
+
+				if (bSavePin)
+				{
+					OldPin->bOrphanedPin = true;
+					OldPin->bNotConnectable = true;
+					OrphanedOldPins.Add(OldPin);
+					InOldPins.RemoveAt(OldPinIndex, 1, false);
+				}
+			}
+		}
+	}
+
+	// The orphaned pins get placed after the rest of the new pins unless it is a child of a split pin and other
+	// children of that split pin were matched in which case it will be at the end of the list of its former siblings
+	for (int32 OrphanedIndex = OrphanedOldPins.Num() - 1; OrphanedIndex >= 0; --OrphanedIndex)
+	{
+		UEdGraphPin* OrphanedPin = OrphanedOldPins[OrphanedIndex];
+		if (OrphanedPin->ParentPin == nullptr)
+		{
+			InNewPins.Add(OrphanedPin);
+		}
+		// Otherwise we need to work out where we fit in the list
+		else 
+		{
+			UEdGraphPin* ParentPin = OrphanedPin->ParentPin;
+			if (!ParentPin->bOrphanedPin)
+			{
+				// Our parent pin was matched, so we need to go to the end of the new pins sub pin section
+				ParentPin->SubPins.Remove(OrphanedPin);
+				ParentPin = MatchedPins.FindChecked(ParentPin);
+				ParentPin->SubPins.Add(OrphanedPin);
+				OrphanedPin->ParentPin = ParentPin;
+			}
+			int32 InsertIndex = InNewPins.Find(ParentPin);
+			while (++InsertIndex < InNewPins.Num())
+			{
+				UEdGraphPin* PinToConsider = InNewPins[InsertIndex];
+				if (PinToConsider->ParentPin != ParentPin)
+				{
+					break;
+				}
+				int32 WalkOffIndex = InsertIndex + PinToConsider->SubPins.Num();
+				for (;InsertIndex < WalkOffIndex;++InsertIndex)
+				{
+					WalkOffIndex += InNewPins[WalkOffIndex]->SubPins.Num();
+				}
+			};
+
+			InNewPins.Insert(OrphanedPin, InsertIndex);
+		}
 	}
+
+	DestroyPinList(InOldPins);
 }
 
 void UK2Node::DestroyPinList(TArray<UEdGraphPin*>& InPins)
@@ -789,7 +990,7 @@ void UK2Node::DestroyPinList(TArray<UEdGraphPin*>& InPins)
 	for (UEdGraphPin* Pin : InPins)
 	{
 		Pin->Modify();
-		Pin->BreakAllPinLinks();
+		Pin->BreakAllPinLinks(!Blueprint->bIsRegeneratingOnLoad);
 
 		UEdGraphNode::DestroyPin(Pin);
 	}
@@ -821,7 +1022,7 @@ void UK2Node::ExpandSplitPin(FKismetCompilerContext* CompilerContext, UEdGraph*
 		{
 			UEdGraphPin* ExpandedPin = ExpandedNode->Pins[ExpandedPinIndex];
 
-			if (!ExpandedPin->bHidden)
+			if (!ExpandedPin->bHidden && !ExpandedPin->bOrphanedPin)
 			{
 				if (ExpandedPin->Direction == Pin->Direction)
 				{
@@ -1003,7 +1204,7 @@ void FOptionalPinManager::RebuildPropertyList(TArray<FOptionalPinFromProperty>&
 	}
 
 	// add remaining "bOverride_" properties
-	for (auto Pair : OverridesMap)
+	for (const TPair<FName, UProperty*>& Pair : OverridesMap)
 	{
 		UProperty* TestProperty = Pair.Value;
 
@@ -1048,10 +1249,8 @@ void FOptionalPinManager::CreateVisiblePins(TArray<FOptionalPinFromProperty>& Pr
 {
 	const UEdGraphSchema_K2* Schema = GetDefault<UEdGraphSchema_K2>();
 
-	for (auto ExtraPropertyIt = Properties.CreateIterator(); ExtraPropertyIt; ++ExtraPropertyIt)
+	for (FOptionalPinFromProperty& PropertyEntry : Properties)
 	{
-		FOptionalPinFromProperty& PropertyEntry = *ExtraPropertyIt;
-
 		if (UProperty* OuterProperty = FindFieldChecked<UProperty>(SourceStruct, PropertyEntry.PropertyName))
 		{
 			// Do we treat an array property as one pin, or a pin per entry in the array?
@@ -1155,6 +1354,32 @@ void FOptionalPinManager::CreateVisiblePins(TArray<FOptionalPinFromProperty>& Pr
 	}
 }
 
+void FOptionalPinManager::CacheShownPins(const TArray<FOptionalPinFromProperty>& OptionalPins, TArray<FName>& OldShownPins)
+{
+	for (const FOptionalPinFromProperty& ShowPinForProperty : OptionalPins)
+	{
+		if (ShowPinForProperty.bShowPin)
+		{
+			OldShownPins.Add(ShowPinForProperty.PropertyName);
+		}
+	}
+}
+
+void FOptionalPinManager::EvaluateOldShownPins(const TArray<FOptionalPinFromProperty>& OptionalPins, TArray<FName>& OldShownPins, UK2Node* Node)
+{
+	for (const FOptionalPinFromProperty& ShowPinForProperty : OptionalPins)
+	{
+		if (ShowPinForProperty.bShowPin == false && OldShownPins.Contains(ShowPinForProperty.PropertyName))
+		{
+			if (UEdGraphPin* Pin = Node->FindPin(ShowPinForProperty.PropertyName.ToString()))
+			{
+				Pin->bSavePinIfOrphaned = false;
+			}
+		}
+	}
+	OldShownPins.Reset();
+}
+
 UEdGraphPin* UK2Node::GetExecPin() const
 {
 	const UEdGraphSchema_K2* K2Schema = GetDefault<UEdGraphSchema_K2>();
@@ -1173,7 +1398,7 @@ UEdGraphPin* UK2Node::GetPassThroughPin(const UEdGraphPin* FromPin) const
 		if(K2Schema->IsExecPin(*FromPin))
 		{
 			// Locate the first exec pin that's opposite the given exec pin, if any.
-			for(auto Pin : Pins)
+			for(UEdGraphPin* Pin : Pins)
 			{
 				if(Pin && Pin != FromPin && Pin->Direction != FromPin->Direction && K2Schema->IsExecPin(*Pin))
 				{
@@ -1353,7 +1578,7 @@ void UK2Node::GetPinHoverText(const UEdGraphPin& Pin, FString& HoverTextOut) con
 
 UClass* UK2Node::GetBlueprintClassFromNode() const
 {
-	auto BP = FBlueprintEditorUtils::FindBlueprintForNode(this);
+	UBlueprint* BP = FBlueprintEditorUtils::FindBlueprintForNode(this);
 	UClass* BPClass = BP
 		? (BP->SkeletonGeneratedClass ? BP->SkeletonGeneratedClass : BP->GeneratedClass)
 		: nullptr;
diff --git a/Engine/Source/Editor/BlueprintGraph/Private/K2Node_AddComponent.cpp b/Engine/Source/Editor/BlueprintGraph/Private/K2Node_AddComponent.cpp
--- a/Engine/Source/Editor/BlueprintGraph/Private/K2Node_AddComponent.cpp
+++ b/Engine/Source/Editor/BlueprintGraph/Private/K2Node_AddComponent.cpp
@@ -69,19 +69,6 @@ void UK2Node_AddComponent::Serialize(FArchive& Ar)
 	}
 }
 
-void UK2Node_AddComponent::PostDuplicate(bool bDuplicateForPIE)
-{
-	Super::PostDuplicate(bDuplicateForPIE);
-
-	// Determine whether or not this node belongs to a Blueprint context.
-	UBlueprint* Blueprint = FBlueprintEditorUtils::FindBlueprintForNode(this);
-	if (!bDuplicateForPIE && Blueprint != nullptr && !Blueprint->bBeingCompiled)
-	{
-		// Create a unique component template for the duplicated node (this).
-		MakeNewComponentTemplate();
-	}
-}
-
 void UK2Node_AddComponent::AllocateDefaultPins()
 {
 	AllocateDefaultPinsWithoutExposedVariables();
@@ -113,7 +100,8 @@ void UK2Node_AddComponent::AllocatePinsForExposedVariables()
 			const bool bIsExposedToSpawn = UEdGraphSchema_K2::IsPropertyExposedOnSpawn(Property);
 			const bool bIsVisible = Property->HasAllPropertyFlags(CPF_BlueprintVisible);
 			const bool bNotParam = !Property->HasAllPropertyFlags(CPF_Parm);
-			if(bNotDelegate && bIsExposedToSpawn && bIsVisible && bNotParam)
+			const bool bStillExists = FBlueprintEditorUtils::PropertyStillExists(Property);
+			if(bNotDelegate && bIsExposedToSpawn && bIsVisible && bNotParam && bStillExists)
 			{
 				FEdGraphPinType PinType;
 				K2Schema->ConvertPropertyToPinType(Property, /*out*/ PinType);	
diff --git a/Engine/Source/Editor/BlueprintGraph/Private/K2Node_BaseAsyncTask.cpp b/Engine/Source/Editor/BlueprintGraph/Private/K2Node_BaseAsyncTask.cpp
--- a/Engine/Source/Editor/BlueprintGraph/Private/K2Node_BaseAsyncTask.cpp
+++ b/Engine/Source/Editor/BlueprintGraph/Private/K2Node_BaseAsyncTask.cpp
@@ -27,8 +27,8 @@
 UK2Node_BaseAsyncTask::UK2Node_BaseAsyncTask(const FObjectInitializer& ObjectInitializer)
 	: Super(ObjectInitializer)
 	, ProxyFactoryFunctionName(NAME_None)
-	, ProxyFactoryClass(NULL)
-	, ProxyClass(NULL)
+	, ProxyFactoryClass(nullptr)
+	, ProxyClass(nullptr)
 	, ProxyActivateFunctionName(NAME_None)
 {
 }
@@ -157,7 +157,6 @@ void UK2Node_BaseAsyncTask::AllocateDefaultPins()
 				if (PinsToHide.Contains(Pin->PinName))
 				{
 					Pin->bHidden = true;
-					Pin->DefaultValue = TEXT("0");
 				}
 			}
 
@@ -179,8 +178,7 @@ bool UK2Node_BaseAsyncTask::FBaseAsyncTaskHelper::ValidDataPin(const UEdGraphPin
 {
 	check(Schema);
 	const bool bValidDataPin = Pin
-		&& (Pin->PinName != Schema->PN_Execute)
-		&& (Pin->PinName != Schema->PN_Then)
+		&& !Pin->bOrphanedPin
 		&& (Pin->PinType.PinCategory != Schema->PC_Exec);
 
 	const bool bProperDirection = Pin && (Pin->Direction == Direction);
@@ -259,10 +257,10 @@ bool UK2Node_BaseAsyncTask::FBaseAsyncTaskHelper::HandleDelegateImplementation(
 	}
 
 	UEdGraphPin* LastActivatedNodeThen = CurrentCENode->FindPinChecked(Schema->PN_Then);
-	for (auto OutputPair : VariableOutputs) // CREATE CHAIN OF ASSIGMENTS
+	for (const FBaseAsyncTaskHelper::FOutputPinAndLocalVariable& OutputPair : VariableOutputs) // CREATE CHAIN OF ASSIGMENTS
 	{
 		UEdGraphPin* PinWithData = CurrentCENode->FindPin(OutputPair.OutputPin->PinName);
-		if (PinWithData == NULL)
+		if (PinWithData == nullptr)
 		{
 			/*FText ErrorMessage = FText::Format(LOCTEXT("MissingDataPin", "ICE: Pin @@ was expecting a data output pin named {0} on @@ (each delegate must have the same signature)"), FText::FromString(OutputPair.OutputPin->PinName));
 			CompilerContext.MessageLog.Error(*ErrorMessage.ToString(), OutputPair.OutputPin, CurrentCENode);
@@ -309,7 +307,7 @@ void UK2Node_BaseAsyncTask::ExpandNode(class FKismetCompilerContext& CompilerCon
 
 	bIsErrorFree &= CompilerContext.MovePinLinksToIntermediate(*FindPinChecked(Schema->PN_Execute), *CallCreateProxyObjectNode->FindPinChecked(Schema->PN_Execute)).CanSafeConnect();
 
-	for (auto CurrentPin : Pins)
+	for (UEdGraphPin* CurrentPin : Pins)
 	{
 		if (FBaseAsyncTaskHelper::ValidDataPin(CurrentPin, EGPD_Input, Schema))
 		{
@@ -325,7 +323,7 @@ void UK2Node_BaseAsyncTask::ExpandNode(class FKismetCompilerContext& CompilerCon
 
 	// GATHER OUTPUT PARAMETERS AND PAIR THEM WITH LOCAL VARIABLES
 	TArray<FBaseAsyncTaskHelper::FOutputPinAndLocalVariable> VariableOutputs;
-	for (auto CurrentPin : Pins)
+	for (UEdGraphPin* CurrentPin : Pins)
 	{
 		if ((OutputAsyncTaskProxy != CurrentPin) && FBaseAsyncTaskHelper::ValidDataPin(CurrentPin, EGPD_Output, Schema))
 		{
