Code standards cleanup - affects Doxygen.

diff --git a/Engine/Source/Editor/BlueprintGraph/Private/EdGraphSchema_K2.cpp b/Engine/Source/Editor/BlueprintGraph/Private/EdGraphSchema_K2.cpp
--- a/Engine/Source/Editor/BlueprintGraph/Private/EdGraphSchema_K2.cpp
+++ b/Engine/Source/Editor/BlueprintGraph/Private/EdGraphSchema_K2.cpp
@@ -1149,12 +1149,12 @@ bool UEdGraphSchema_K2::ClassHasBlueprintAccessibleMembers(const UClass* InClass
 	return false;
 }
 
-bool UEdGraphSchema_K2::IsAllowableBlueprintVariableType(const class UEnum* InEnum)
+bool UEdGraphSchema_K2::IsAllowableBlueprintVariableType(const UEnum* InEnum)
 {
 	return InEnum && (InEnum->GetBoolMetaData(FBlueprintMetadata::MD_AllowableBlueprintVariableType) || InEnum->IsA<UUserDefinedEnum>());
 }
 
-bool UEdGraphSchema_K2::IsAllowableBlueprintVariableType(const class UClass* InClass)
+bool UEdGraphSchema_K2::IsAllowableBlueprintVariableType(const UClass* InClass)
 {
 	if (InClass)
 	{
