Enums should still use special empty string checking for whether values match regardless of age of pin.

diff --git a/Engine/Source/Editor/BlueprintGraph/Private/EdGraphSchema_K2.cpp b/Engine/Source/Editor/BlueprintGraph/Private/EdGraphSchema_K2.cpp
--- a/Engine/Source/Editor/BlueprintGraph/Private/EdGraphSchema_K2.cpp
+++ b/Engine/Source/Editor/BlueprintGraph/Private/EdGraphSchema_K2.cpp
@@ -4617,30 +4617,30 @@ void UEdGraphSchema_K2::TrySetDefaultText(UEdGraphPin& InPin, const FText& InNew
 
 bool UEdGraphSchema_K2::DoesDefaultValueMatchAutogenerated(const UEdGraphPin& InPin) const
 {
-	if (!InPin.bUseBackwardsCompatForEmptyAutogeneratedValue)
+	if (InPin.PinType.PinCategory == PC_Enum || InPin.PinType.PinCategory == PC_Byte)
 	{
-		if (InPin.PinType.PinCategory == PC_Enum || InPin.PinType.PinCategory == PC_Byte)
+		// The autogenerated default value for an enum is left empty in case the default enum value (index 0) changes, in this case we
+		// want to validate against the actual value of the 0 index entry
+		if (InPin.AutogeneratedDefaultValue.IsEmpty())
 		{
-			// The autogenerated default value for an enum is left empty in case the default enum value (index 0) changes, in this case we
-			// want to validate against the actual value of the 0 index entry
-			if (InPin.AutogeneratedDefaultValue.IsEmpty())
+			const FString PinDefaultValue = InPin.GetDefaultAsString();
+			if (PinDefaultValue.IsEmpty())
 			{
-				const FString PinDefaultValue = InPin.GetDefaultAsString();
-				if (PinDefaultValue.IsEmpty())
-				{
-					return true;
-				}
-				else if (UEnum* PinEnumType = Cast<UEnum>(InPin.PinType.PinSubCategoryObject.Get()))
-				{
-					return InPin.DefaultValue.Equals(PinEnumType->GetNameStringByIndex(0), ESearchCase::IgnoreCase);
-				}
-				else if (InPin.PinType.PinCategory == PC_Byte && FCString::Atoi(*PinDefaultValue) == 0)
-				{
-					return true;
-				}
+				return true;
+			}
+			else if (UEnum* PinEnumType = Cast<UEnum>(InPin.PinType.PinSubCategoryObject.Get()))
+			{
+				return InPin.DefaultValue.Equals(PinEnumType->GetNameStringByIndex(0), ESearchCase::IgnoreCase);
+			}
+			else if (!InPin.bUseBackwardsCompatForEmptyAutogeneratedValue && InPin.PinType.PinCategory == PC_Byte && FCString::Atoi(*PinDefaultValue) == 0)
+			{
+				return true;
 			}
 		}
-		else if (InPin.PinType.PinCategory == PC_Float)
+	}
+	else if (!InPin.bUseBackwardsCompatForEmptyAutogeneratedValue)
+	{
+		if (InPin.PinType.PinCategory == PC_Float)
 		{
 			const float AutogeneratedFloat = FCString::Atof(*InPin.AutogeneratedDefaultValue);
 			const float DefaultFloat = FCString::Atof(*InPin.DefaultValue);
