Copying //UE4/Dev-Framework to //UE4/Dev-Main (Source: //UE4/Dev-Framework @ 3544039)
#rnx

=====================================
 MAJOR FEATURES + CHANGES
=====================================

Change 3343905 by [REDACTED_PERSON]

	ResolveMember optimizations and moved into cpp. ResolveMember<UFunction> now checks UClass::FuncMap before doing more expensive searches

Change 3346637 by [REDACTED_PERSON]

	Actually fix in non editor builds

Change 3355484 by [REDACTED_PERSON]

	Back out FMemberReference Optimization

Change 3425833 by [REDACTED_PERSON]

Change 3510091 by [REDACTED_PERSON]

	Expose on Spawn functional test

	#rnx

Change 3510100 by [REDACTED_PERSON]

	Fix spelling error

	#rnx

Change 3510132 by [REDACTED_PERSON]

	Fix issues with marking a widget blueprint class as abstract

Change 3510133 by [REDACTED_PERSON]

	Minor code cleanup

	#rnx

Change 3510178 by [REDACTED_PERSON]

Change 3510466 by [REDACTED_PERSON]

	Start adding basic ability system tests to enginetest, very minimal so far

Change 3511295 by [REDACTED_PERSON]

	Fix wasted work going weak -> object -> weak -> object

	#rnx

Change 3511824 by [REDACTED_PERSON]

	Fix spelling error in tooltip

	#rnx

Change 3514446 by [REDACTED_PERSON]

	Fix ActorBoundEvent and ComponentBoundEvent to always refresh their event signature from the delegate property they are bound to. This is required to correctly deal with delegate signatures being moved or renamed. Both types now do the fixup one time, in ReconstructNode.

Change 3514578 by [REDACTED_PERSON]

	Move clearing of the actor component need end of frame update mark to base class instead of just primitive component

Change 3514583 by [REDACTED_PERSON]

	Better fix to last delegate checkin that also handles moving functions between modules but not renaming

Change 3515325 by [REDACTED_PERSON]

	Fix for rare orphan pin false positive, rare exposed on spawn false positive

	#rnx

Change 3515761 by [REDACTED_PERSON]

	fix shipping configuration

	#rnx

Change 3515772 by [REDACTED_PERSON]

	Fix static analysis warnings

	#rnx

Change 3516287 by [REDACTED_PERSON]

	Fix references to instanced components not being updated when resetting component to default

	#rnx

Change 3516303 by [REDACTED_PERSON]

	Back out CL# 3516287 while an oddity is investigated

	#rnx

Change 3516563 by [REDACTED_PERSON]

	(4.17) Fix references to instanced components not being updated when resetting component to default

Change 3516637 by [REDACTED_PERSON]

	Change summary:
	- Modified UChildActorComponent::DestroyChildActor() to move the check for PendingKill/Unreachable so that we can also rename a defunct ChildActor instance out of the way in order to allow for a new ChildActor instance w/ the cached name.

Change 3517735 by [REDACTED_PERSON]

	Avoid unnecessary string copy

	#rnx

Change 3517931 by [REDACTED_PERSON]

	Small optimization to CleanupActors

Change 3518221 by [REDACTED_PERSON]

	Fix rare crash when running ConformImplementedEvents when async loading

Change 3518270 by [REDACTED_PERSON]

	Add AssetCollection to PrimaryAssetLabel that derives the bundled assets from an editor collection

Change 3518271 by [REDACTED_PERSON]

	Get rid of unnecessary construction differentiation if custom reset is being used

Change 3518310 by [REDACTED_PERSON]

	Re-adding IOS files with correct case.

Change 3518423 by [REDACTED_PERSON]

	Refactor AssetManagerSettings so it copies runtime bools into the asset manager for fast access
	Add a concept of a stalled streamable manager handle, handles can be created stalled and will not execute their async load until all needed resources have been acquired externally

Change 3518480 by [REDACTED_PERSON]

	Correctly get the variable reference for an input variable get from the member scope rather than a member variable of the same name on the class

Change 3518498 by [REDACTED_PERSON]

	Fix bug with AssetManager where requesting the same load twice in a row before the first one finishes caused the complete callback to get called too early for the second load
	Update test map to catch this

Change 3518526 by [REDACTED_PERSON]

	IOS Fix

Change 3518619 by [REDACTED_PERSON]

Change 3518747 by [REDACTED_PERSON]

	Change summary:
	- Moved UGameplayStatics::CanSpawnObjectOfClass() into UK2Node_GenericCreateObject as a local util method (per JIRA notes). This was not exposed to Blueprints and as such was inconsistent with the rest of the API.
	- Modified UGameplayStatics::SpawnObject() to no longer call CanSpawnObjectOfClass(). This seemed redundant as this will already have been called during node validation at Blueprint compile time.
	- Refactored CanSpawnObjectOfClass() into FK2Node_GenericCreateObject_Utils. Walking up the inheritance chain no longer starts out w/ the assumption that 'BlueprintType' is set by default, which was previously including a lot of engine-specific classes into the "allowed" set (e.g. UByteProperty). Also unified the 2 loop iterations that were being used to check for 'BlueprintType'/'NotBlueprintType' and 'DontUseGenericSpawnObjectName', as well as the check for whether or not the class is a derivative of AActor/UActorComponent.
	- Modified UK2Node_GenericCreateObject::EarlyValidation() to call FK2Node_GenericCreateObject_Utils::CanSpawnObjectOfClass() and emit a slightly more informative error message to the BP compiler message log.

Change 3518756 by [REDACTED_PERSON]

	(4.17) Framework: Prevent various asserts when USplineComponent methods are called on a spline with no points

Change 3518760 by [REDACTED_PERSON]

	Core: Changed FRuntimeAssetCache ensures to ensureAsRuntimeWarning

Change 3518771 by [REDACTED_PERSON]

	AI: Prevent an ensure in UBlackboardComponent::ClearValue when called on a component with a null BlackboardAsset

Change 3518818 by [REDACTED_PERSON]

	Rendering: Fixed a whitespace issue in UCanvasRenderTarget2D::RepaintCanvas()

	#rnx

Change 3518822 by [REDACTED_PERSON]

	Sequencer: Prevented crashes in some methods of UMovieSceneSequencePlayer when there is no Sequence set
	Sequencer: Prevented a crash in FMovieSceneRootEvaluationTemplateInstance::Evaluate when the instance has no template set

Change 3518824 by [REDACTED_PERSON]

	Landscape: Marked ULandscapeComponent and ULandscapeHeightfieldCollisionComponent as Within=LandscapeProxy, since they do CastChecked on their Outer all the time

Change 3519073 by [REDACTED_PERSON]

	QAGame: Fixed a crash in UQASynth::PlaySynth() if called on a directly created instance rather than using the factory method

Change 3519076 by [REDACTED_PERSON]

	Preventing crashes in UAutomationPerformaceHelper (sic) when spawned abnormally for fuzzing (assumes that the outer will have a route to a world)
	#rnx

Change 3519079 by [REDACTED_PERSON]

	Sequencer: Fixed a potential crash in UMediaPlaylist::Insert and UMediaPlaylist::RemoveAt when passed an invalid index

Change 3519081 by [REDACTED_PERSON]

	Blueprints: Added support for creating appropriate outers for objects that must be nested within another class during fuzzing (ones that specify Within=, other relationships aren't knowable yet)

Change 3519082 by [REDACTED_PERSON]

	VR: Prevent a crash in UMRMeshComponent::ConnectReconstructor when passed a null reconstructor

Change 3519084 by [REDACTED_PERSON]

	Rendering: Prevent crashes when UNiagaraComponent::GetEffectDataInterface is called on a component with no effect asset set

Change 3521889 by [REDACTED_PERSON]

	Sequencer: Prevented a bogus static analysis warning by reworking the code (FixedFrameInterval could have only been set if the pointer were valid from the line above)
	#rnx

Change 3521987 by [REDACTED_PERSON]

	Animation: Prevent a couple of potential asserts in UControlRig::GetOrAllocateSubControlRig

Change 3522101 by [REDACTED_PERSON]

	Physics: Improved the comment on UPhysicalMaterial::Friction

	#rn

Change 3522105 by [REDACTED_PERSON]

	Physics: Fixed a few crashes in UVehicleWheel when spawned directly

Change 3522106 by [REDACTED_PERSON]

	Framework: Marked ULevelStreaming as Within=World, since it does CastChecked on the Outer all the time

Change 3522109 by [REDACTED_PERSON]

	Animation: Marked UAnimInstance as Within=SkeletalMeshComponent since it assumes the outer in various places

Change 3522121 by [REDACTED_PERSON]

	Mobile: Prevent UMobileInstalledContent methods from crashing when called on a created instance in an uncooked build (no installed manifest)

Change 3522783 by [REDACTED_PERSON]

	#ue4 - Imported new simple collision for Engine/Content/BasicShaps/Cylinder.uasset which is a single convex shape (rather than being 4 shapes as before).

Change 3525477 by [REDACTED_PERSON]

	Remove Tooltip, Category, and HideCategories tooltip from the blueprint generated class if source data is cleared

Change 3526538 by [REDACTED_PERSON]

	Refresh primary asset labels if their bundles are different at all and not just if they're added or removed. This is required because they now work based on collections or directories. This fixes issue with the onboarding collection changes not correctly modifying chunks
	Copy of CL #3526501

Change 3526817 by [REDACTED_PERSON]

Change 3526873 by [REDACTED_PERSON]

	Add code to to the ubergraph frame to fall back to hard reference serialization if the reference collector doesn't support weak references, such as the cluster collector

Change 3526958 by [REDACTED_PERSON]

	(4.17) Don't copy and then break pin links when reconstructing. Instead simply move.

Change 3528916 by [REDACTED_PERSON]

	PR #3609: Adds GetKeysForAxis() to complement GetKeysForAction() in UPlayerInput (Contributed by [REDACTED_PERSON])

Change 3529080 by mason.seay

	BP asset for undetermined type bug

Change 3529381 by [REDACTED_PERSON]

	Fix ability to insert duplicates in to a set or map

Change 3529471 by [REDACTED_PERSON]

	Fix for clang 4.0 error: definition of builtin function '__rdtsc' inline unsigned long long __rdtsc()

Change 3530876 by [REDACTED_PERSON]

	Based on PR #3457: Add MakeSet BP node (Contributed by [REDACTED_PERSON])
	Also refactored MakeArray/Set to share a base MakeContainer class
	Cleaned up some dead code from MakeArray
	Added icon for MakeSet
	Added Functional Test for MakeSet

Change 3531070 by [REDACTED_PERSON]

	Change summary:
	- Modified FBlueprintEditorUtils::GetSkeletonClass() to check for NULL before attempting to check for the generating BP.

Change 3531081 by [REDACTED_PERSON]

	Remove deprecated CustomMapParamValue code

Change 3531094 by [REDACTED_PERSON]

	Change summary:
	- Modified TScriptInterface::operator=() to cast the given 'SourceObject' instance to the 'InterfaceType' type before assigning to 'SourceInterface'. This was necessary because if the caller (in this case nativized codegen) passes in a UObject* that does not explicitly inherit from 'InterfaceType', then it will need to go through the object's GetInterfaceAddress() API instead and cast the result back to an 'InterfaceType' pointer.

Change 3531186 by [REDACTED_PERSON]

	Back out changelist 3531094 (temp CIS fix).

	#rnx

Change 3532082 by [REDACTED_PERSON]

	Move garbage collection timers and other management to UEngine instead of UWorld
	Fixes CollectGarbage blueprint node not working in shipping

Change 3532134 by [REDACTED_PERSON]

	Restored changelist 3531094 w/ fix for non-unity.

	- Mirrored from //UE4/Release-4.17 (CL# 3531232).

	#rnx

Change 3533009 by [REDACTED_PERSON]

	Fixup missing function and deprecation warnings

Change 3534056 by [REDACTED_PERSON]

	(4.17) Fix expose on spawn of map and sets to work

Change 3534761 by [REDACTED_PERSON]

	(4.17) Apply code review changes to Dev-Framework as well

	#rnx

Change 3535147 by [REDACTED_PERSON]

	Build fix, already made in 4.17

	#rnx

Change 3535530 by mason.seay

	Resaving to remove error when opening level blueprint

Change 3535581 by [REDACTED_PERSON]

	Class Properties are only identical if they are literally the same object. Do not consider the deep compare port flags as object property base does.

Change 3535583 by [REDACTED_PERSON]

	When properties are imported in to a child actor component the cached instance data is invalidated, so clear it.

Change 3535617 by [REDACTED_PERSON]

	PR #3788: UE-39237: Prevent (im-)pure casting during BP debugging (Contributed by [REDACTED_PERSON])

Change 3535671 by [REDACTED_PERSON]

	Change NodeFactory to look at interface to use sequence node instead of each node having to add itself

Change 3535955 by [REDACTED_PERSON]

	Prevent MakeSet from removing split pins

Change 3536114 by [REDACTED_PERSON]

	Paper2D: Removing deprecated code from 4.3/4.4 era

	#rnx

Change 3536120 by [REDACTED_PERSON]

	Animation: Removed deprecated FTAlphaBlend class and AlphaBlendType.h header

Change 3536124 by [REDACTED_PERSON]

	Physics: Removed deprecated methods that were replaced by _AssumesLocked variations

Change 3536131 by [REDACTED_PERSON]

	Slate: Converting remaining uses of EKeyboardFocusCause to EFocusCause and properly deprecating it

Change 3536138 by [REDACTED_PERSON]

	Slate: Removed any deprecated code older than 4.10 that didn't affect content compatibility

Change 3536167 by [REDACTED_PERSON]

	When a client provides a skeleton class as the self scope, make sure we also use a skel class for non-self scopes - but only if using the compilation manager. Skel classes are not reliably up to date when not using the compilation manager

Change 3536221 by [REDACTED_PERSON]

	Editor: Removing deprecated code from 4.9 or earlier

Change 3536240 by [REDACTED_PERSON]

	Blueprints: Removed long-deprecated TypeToString method from the K2 schema
	#rnx

Change 3536243 by [REDACTED_PERSON]

	AI: Prevent crashes if UMockTask_Log is created manually rather than via the CreateTask factory method

Change 3536244 by [REDACTED_PERSON]

	Core: Prevent FScopedExternalProfilerBase::StopScopedTimer() from asserting if called an unmatched number of times with StartScopedTimer, as both are exposed to BPs now

Change 3536250 by [REDACTED_PERSON]

	CoreUObject: Removed any deprecated code older than 4.10 that didn't affect content compatibility

Change 3536253 by [REDACTED_PERSON]

	Core: Removed any deprecated code older than 4.10 that didn't affect content compatibility

Change 3536310 by [REDACTED_PERSON]

	Engine: Removed any deprecated code older than 4.10 that didn't affect content compatibility

Change 3536397 by [REDACTED_PERSON]

	Fixed UCrowdFollowingComponent::UpdateCachedDirections crashing when CharacterMovement is not set #UE4

Change 3536404 by [REDACTED_PERSON]

	Platform: Added a warning for others when they try to remove this 'deprecated' method

Change 3536639 by [REDACTED_PERSON]

	CharacterMovement: Changed the name of a variable introduced in CL# 3536397 to better match intent
	#rnx

Change 3536893 by [REDACTED_PERSON]

	Blueprints: Clear the stale value on the value pin when a map find node fails to find an item

Change 3536902 by [REDACTED_PERSON]

	Framework: Killed a couple of more deprecated methods that were not exposed to Blueprints

	#rnx

Change 3537038 by [REDACTED_PERSON]

	Fixing case of iOS directories, pt1

Change 3537039 by [REDACTED_PERSON]

	Fixing case of iOS directories, pt2

Change 3538246 by [REDACTED_PERSON]

	UnrealTournament: Fixing issues with renamed enum

	#rnx

Change 3538618 by [REDACTED_PERSON]

	Fix ensure when closing sequencer transform UI

Change 3540213 by [REDACTED_PERSON]

Change 3540253 by [REDACTED_PERSON]

	Only copy default values for input pins as output pins do not have them

	#rnx

Change 3540376 by [REDACTED_PERSON]

	Add utility FromPinType for FEdGraphTerminalType

	#rnx

Change 3540433 by [REDACTED_PERSON]

	Add MakeMap
	Unify IsConnectionDisallowed for containers and fix static analysis warning

Change 3540585 by [REDACTED_PERSON]

	Change summary:
	- Modified FEmitDefaultValueHelper::HandleSpecialTypes() to only direct HandleInstancedSubobject() to emit code to create the instanced subobject if it's not a default subobject. This was previously being incorrectly interpreted as an object having the 'RF_ArchetypeObject' flag set; however, default subobjects will also have that flag set in addition to the 'RF_DefaultSubobject' flag.
	- Modified FEmitDefaultValueHelper::HandleInstancedSubobject() to assert in the 'GetDefaultSubobjectByName' case if the given object is not also a default subobject.

Change 3541147 by [REDACTED_PERSON]

	Fix for not being able to override custom events when using the compilation manager post 3536167
	#rnx

Change 3541177 by [REDACTED_PERSON]

	Copy of CL #3541027

Change 3541325 by [REDACTED_PERSON]

	K2node data table data needs to preload data before the compilation queue is flushed

	#rnx

Change 3541409 by [REDACTED_PERSON]

	Blueprints: Added code to reapply any active breakpoints after recompilation when using the BP compilation manager

	[reimplementing CL# 3541404 in Dev-Framework]

Change 3541418 by [REDACTED_PERSON]

	Fix for bad SKEL_ CDO reference in blueprint bytecode

	#rnx

Change 3541482 by [REDACTED_PERSON]

	Blanket fix up of preload calls that are being done in AllocateDefaultPins. AllocatDefaultPins is not called until compile, meaning if these preload calls load blueprints they will be loaded while the compilation manager is compiling blueprints

	#rnx

Change 3541817 by [REDACTED_PERSON]

	Fix static analysis warnings

	#rnx

Change 3542299 by [REDACTED_PERSON]

	Blueprints: Speculative fix for static analysis warning
	#rnx

Change 3542406 by [REDACTED_PERSON]

	Use a check slow to avoid any cost

	#rnx

Change 3542486 by [REDACTED_PERSON]

	Asset Manager: Removing an unnecessary ensure (it's a potentially expected case)

Change 3542659 by [REDACTED_PERSON]

	Blueprints: Clear out null entries in the LastEditedDocuments list during PostLoad() and remove entries when a graph is being deleted to prevent their generation in the first place

Change 3543620 by [REDACTED_PERSON]

	Remove overzealous ensure - we may recompile blueprints that are asynchronously loading when a user triggers a synchronous compile
	#rnx

Change 3518415 by [REDACTED_PERSON]

	Replace with AddChunkInstallDelegate, which supports a bool error code and is bound once instead of separately for each chunk. All implementations support this delegate at a basic level, although several could be improved to call the failure delegate in more cases

Change 3534339 by [REDACTED_PERSON]

	Platforms: Changed DEPRECATED() macro description to use 4.xx rather than a speciifc version in examples, so it doesn't show up when removing deprecated code

diff --git a/Engine/Plugins/2D/Paper2D/Source/Paper2D/Classes/PaperFlipbookComponent.h b/Engine/Plugins/2D/Paper2D/Source/Paper2D/Classes/PaperFlipbookComponent.h
--- a/Engine/Plugins/2D/Paper2D/Source/Paper2D/Classes/PaperFlipbookComponent.h
+++ b/Engine/Plugins/2D/Paper2D/Source/Paper2D/Classes/PaperFlipbookComponent.h
@@ -77,12 +77,7 @@ public:
 	/** Gets the flipbook used by this instance. */
 	UFUNCTION(BlueprintPure, Category="Sprite")
 	virtual UPaperFlipbook* GetFlipbook();
-
-	/** DEPRECATED! Use GetMaterial() instead.  Gets the material used by this instance */
-	DEPRECATED(4.5, "GetSpriteMaterial has been replaced by GetMaterial")
-	UFUNCTION(BlueprintCallable, Category="Sprite", meta=(DeprecatedFunction, DeprecationMessage="Use GetMaterial instead"))
-	UMaterialInterface* GetSpriteMaterial() const;
-
+	
 	/** Set color of the sprite */
 	UFUNCTION(BlueprintCallable, Category="Sprite")
 	void SetSpriteColor(FLinearColor NewColor);
@@ -206,7 +201,3 @@ public:
 	virtual UBodySetup* GetBodySetup() override;
 	// End of UPrimitiveComponent interface
 };
-
-// Allow the old name to continue to work for one release
-DEPRECATED(4.3, "UPaperAnimatedRenderComponent has been renamed to UPaperFlipbookComponent")
-typedef UPaperFlipbookComponent UPaperAnimatedRenderComponent;
diff --git a/Engine/Plugins/2D/Paper2D/Source/Paper2D/Classes/PaperSpriteActor.h b/Engine/Plugins/2D/Paper2D/Source/Paper2D/Classes/PaperSpriteActor.h
--- a/Engine/Plugins/2D/Paper2D/Source/Paper2D/Classes/PaperSpriteActor.h
+++ b/Engine/Plugins/2D/Paper2D/Source/Paper2D/Classes/PaperSpriteActor.h
@@ -32,7 +32,3 @@ public:
 	/** Returns RenderComponent subobject **/
 	FORCEINLINE class UPaperSpriteComponent* GetRenderComponent() const { return RenderComponent; }
 };
-
-// Allow the old name to continue to work for one release
-DEPRECATED(4.3, "APaperRenderActor has been renamed to APaperSpriteActor")
-typedef APaperSpriteActor APaperRenderActor;
diff --git a/Engine/Plugins/2D/Paper2D/Source/Paper2D/Classes/PaperSpriteComponent.h b/Engine/Plugins/2D/Paper2D/Source/Paper2D/Classes/PaperSpriteComponent.h
--- a/Engine/Plugins/2D/Paper2D/Source/Paper2D/Classes/PaperSpriteComponent.h
+++ b/Engine/Plugins/2D/Paper2D/Source/Paper2D/Classes/PaperSpriteComponent.h
@@ -102,7 +102,3 @@ public:
 protected:
 	friend class FPaperSpriteSceneProxy;
 };
-
-// Allow the old name to continue to work for one release
-DEPRECATED(4.3, "UPaperRenderComponent has been renamed to UPaperSpriteComponent")
-typedef UPaperSpriteComponent UPaperRenderComponent;
diff --git a/Engine/Plugins/2D/Paper2D/Source/Paper2D/Classes/PaperTileMapComponent.h b/Engine/Plugins/2D/Paper2D/Source/Paper2D/Classes/PaperTileMapComponent.h
--- a/Engine/Plugins/2D/Paper2D/Source/Paper2D/Classes/PaperTileMapComponent.h
+++ b/Engine/Plugins/2D/Paper2D/Source/Paper2D/Classes/PaperTileMapComponent.h
@@ -227,7 +227,3 @@ public:
 	void GetRenderingStats(int32& OutNumTriangles, int32& OutNumBatches) const;
 #endif
 };
-
-// Allow the old name to continue to work for one release
-DEPRECATED(4.7, "UPaperTileMapRenderComponent has been renamed to UPaperTileMapComponent")
-typedef UPaperTileMapComponent UPaperTileMapRenderComponent;
diff --git a/Engine/Plugins/2D/Paper2D/Source/Paper2D/Classes/SpriteEditorOnlyTypes.h b/Engine/Plugins/2D/Paper2D/Source/Paper2D/Classes/SpriteEditorOnlyTypes.h
--- a/Engine/Plugins/2D/Paper2D/Source/Paper2D/Classes/SpriteEditorOnlyTypes.h
+++ b/Engine/Plugins/2D/Paper2D/Source/Paper2D/Classes/SpriteEditorOnlyTypes.h
@@ -334,10 +334,3 @@ namespace ESpritePivotMode
 		Custom
 	};
 }
-
-// Allow the old name to continue to work for one release
-DEPRECATED(4.7, "FSpritePolygon has been renamed to FSpriteGeometryShape")
-typedef FSpriteGeometryShape FSpritePolygon;
-
-DEPRECATED(4.7, "FSpritePolygon has been renamed to FSpriteGeometryShape")
-typedef FSpriteGeometryCollection FSpritePolygonCollection;
diff --git a/Engine/Plugins/2D/Paper2D/Source/Paper2D/Private/PaperFlipbookComponent.cpp b/Engine/Plugins/2D/Paper2D/Source/Paper2D/Private/PaperFlipbookComponent.cpp
--- a/Engine/Plugins/2D/Paper2D/Source/Paper2D/Private/PaperFlipbookComponent.cpp
+++ b/Engine/Plugins/2D/Paper2D/Source/Paper2D/Private/PaperFlipbookComponent.cpp
@@ -408,11 +408,6 @@ UPaperFlipbook* UPaperFlipbookComponent::GetFlipbook()
 	return SourceFlipbook;
 }
 
-UMaterialInterface* UPaperFlipbookComponent::GetSpriteMaterial() const
-{
-	return GetMaterial(0);
-}
-
 void UPaperFlipbookComponent::SetSpriteColor(FLinearColor NewColor)
 {
 	// Can't set color on a static component
diff --git a/Engine/Plugins/Experimental/ControlRig/Source/ControlRig/Private/ControlRig.cpp b/Engine/Plugins/Experimental/ControlRig/Source/ControlRig/Private/ControlRig.cpp
--- a/Engine/Plugins/Experimental/ControlRig/Source/ControlRig/Private/ControlRig.cpp
+++ b/Engine/Plugins/Experimental/ControlRig/Source/ControlRig/Private/ControlRig.cpp
@@ -107,21 +107,26 @@ float UControlRig::GetDeltaTime() const
 
 UControlRig* UControlRig::GetOrAllocateSubControlRig(TSubclassOf<UControlRig> ControlRigClass, int32 AllocationIndex)
 {
-	// Must have a valid index
-	check(AllocationIndex >= 0);
-
-	// reallocate the sub ControlRig reserve to accommodate this index
-	if (!SubControlRigs.IsValidIndex(AllocationIndex))
+	if (ensureAsRuntimeWarning(ControlRigClass != nullptr))
 	{
-		SubControlRigs.SetNum(FMath::Max(SubControlRigs.Num(), AllocationIndex + 1), false);
+		if (ensureAsRuntimeWarning(AllocationIndex >= 0))
+		{
+			// reallocate the sub ControlRig reserve to accommodate this index
+			if (!SubControlRigs.IsValidIndex(AllocationIndex))
+			{
+				SubControlRigs.SetNum(FMath::Max(SubControlRigs.Num(), AllocationIndex + 1), false);
+			}
+
+			if (SubControlRigs[AllocationIndex] == nullptr)
+			{
+				SubControlRigs[AllocationIndex] = NewObject<UControlRig>(this, ControlRigClass.Get(), "SubControlRig");
+			}
+
+			return SubControlRigs[AllocationIndex];
+		}
 	}
 
-	if(SubControlRigs[AllocationIndex] == nullptr)
-	{
-		SubControlRigs[AllocationIndex] = NewObject<UControlRig>(this, ControlRigClass.Get(), "SubControlRig");
-	}
-
-	return SubControlRigs[AllocationIndex];
+	return nullptr;
 }
 
 AActor* UControlRig::GetHostingActor() const
diff --git a/Engine/Plugins/Runtime/HTTPChunkInstaller/Source/Private/HTTPChunkInstaller.cpp b/Engine/Plugins/Runtime/HTTPChunkInstaller/Source/Private/HTTPChunkInstaller.cpp
--- a/Engine/Plugins/Runtime/HTTPChunkInstaller/Source/Private/HTTPChunkInstaller.cpp
+++ b/Engine/Plugins/Runtime/HTTPChunkInstaller/Source/Private/HTTPChunkInstaller.cpp
@@ -1176,6 +1176,8 @@ bool FHTTPChunkInstall::Tick(float DeltaSeconds)
 				{
 					FoundDelegate->Broadcast(InstallingChunkID);
 				}
+
+				InstallDelegate.Broadcast(InstallingChunkID, true);
 			}
 			EndInstall();
 
diff --git a/Engine/Plugins/Runtime/HTTPChunkInstaller/Source/Public/HTTPChunkInstaller.h b/Engine/Plugins/Runtime/HTTPChunkInstaller/Source/Public/HTTPChunkInstaller.h
--- a/Engine/Plugins/Runtime/HTTPChunkInstaller/Source/Public/HTTPChunkInstaller.h
+++ b/Engine/Plugins/Runtime/HTTPChunkInstaller/Source/Public/HTTPChunkInstaller.h
@@ -12,7 +12,7 @@
 /**
 * HTTP based implementation of chunk based install
 **/
-class HTTPCHUNKINSTALLER_API FHTTPChunkInstall : public IPlatformChunkInstall, public FTickerObjectBase
+class HTTPCHUNKINSTALLER_API FHTTPChunkInstall : public FGenericPlatformChunkInstall, public FTickerObjectBase
 {
 public:
 	FHTTPChunkInstall();
diff --git a/Engine/Plugins/Runtime/MobilePatchingUtils/Source/MobilePatchingUtils/Private/MobilePatchingLibrary.cpp b/Engine/Plugins/Runtime/MobilePatchingUtils/Source/MobilePatchingUtils/Private/MobilePatchingLibrary.cpp
--- a/Engine/Plugins/Runtime/MobilePatchingUtils/Source/MobilePatchingUtils/Private/MobilePatchingLibrary.cpp
+++ b/Engine/Plugins/Runtime/MobilePatchingUtils/Source/MobilePatchingUtils/Private/MobilePatchingLibrary.cpp
@@ -60,10 +60,8 @@ static const FText& GetRequestContentErrorText(ERequestContentError ErrorCode)
 
 float UMobileInstalledContent::GetInstalledContentSize()
 {
-	check(InstalledManifest.IsValid())
-	
-	double BuildSize = (double)InstalledManifest->GetBuildSize();
-	double BuildSizeMB = FUnitConversion::Convert(BuildSize, EUnit::Bytes, EUnit::Megabytes);
+	const double BuildSize = InstalledManifest.IsValid() ? (double)InstalledManifest->GetBuildSize() : 0.0;
+	const double BuildSizeMB = FUnitConversion::Convert(BuildSize, EUnit::Bytes, EUnit::Megabytes);
 	return (float)BuildSizeMB;
 }
 
@@ -93,26 +91,33 @@ bool UMobileInstalledContent::Mount(int32 InPakOrder, const FString& InMountPoin
 	uint32 PakOrder = (uint32)FMath::Max(0, InPakOrder);
 	const TCHAR* MountPount = InMountPoint.GetCharArray().GetData();
 
-	TArray<FString> InstalledFileNames = InstalledManifest->GetBuildFileList();
-	for (const FString& FileName : InstalledFileNames)
+	if (InstalledManifest.IsValid())
 	{
-		if (FPaths::GetExtension(FileName) == TEXT("pak"))
+		TArray<FString> InstalledFileNames = InstalledManifest->GetBuildFileList();
+		for (const FString& FileName : InstalledFileNames)
 		{
-			FString PakFullName = InstallDir / FileName;
-			if (PakFileMgr->Mount(*PakFullName, PakOrder, MountPount))
+			if (FPaths::GetExtension(FileName) == TEXT("pak"))
 			{
-				UE_LOG(LogMobilePatchingUtils, Log, TEXT("Mounted = %s, Order = %d, MountPoint = %s"), *PakFullName, PakOrder, !MountPount ? TEXT("(null)") : MountPount);
-				bMounted = true;
-			}
-			else
-			{
-				UE_LOG(LogMobilePatchingUtils, Error, TEXT("Failed to mount pak = %s"), *PakFullName);
-				bMounted = false;
-				break;
+				FString PakFullName = InstallDir / FileName;
+				if (PakFileMgr->Mount(*PakFullName, PakOrder, MountPount))
+				{
+					UE_LOG(LogMobilePatchingUtils, Log, TEXT("Mounted = %s, Order = %d, MountPoint = %s"), *PakFullName, PakOrder, !MountPount ? TEXT("(null)") : MountPount);
+					bMounted = true;
+				}
+				else
+				{
+					UE_LOG(LogMobilePatchingUtils, Error, TEXT("Failed to mount pak = %s"), *PakFullName);
+					bMounted = false;
+					break;
+				}
 			}
 		}
 	}
-	
+	else
+	{
+		UE_LOG(LogMobilePatchingUtils, Log, TEXT("No installed manifest, failed to mount"));
+	}
+
 	return bMounted;
 }
 
diff --git a/Engine/Plugins/Runtime/PhysXVehicles/Source/PhysXVehicles/Private/VehicleWheel.cpp b/Engine/Plugins/Runtime/PhysXVehicles/Source/PhysXVehicles/Private/VehicleWheel.cpp
--- a/Engine/Plugins/Runtime/PhysXVehicles/Source/PhysXVehicles/Private/VehicleWheel.cpp
+++ b/Engine/Plugins/Runtime/PhysXVehicles/Source/PhysXVehicles/Private/VehicleWheel.cpp
@@ -3,6 +3,7 @@
 #include "VehicleWheel.h"
 #include "UObject/ConstructorHelpers.h"
 #include "Engine/World.h"
+#include "Engine/Engine.h"
 #include "PhysxUserData.h"
 #include "Engine/StaticMesh.h"
 #include "Vehicles/TireType.h"
@@ -11,7 +12,6 @@
 #include "PhysXVehicleManager.h"
 #include "PhysXPublic.h"
 
-
 UVehicleWheel::UVehicleWheel(const FObjectInitializer& ObjectInitializer)
 	: Super(ObjectInitializer)
 {
@@ -38,37 +38,55 @@ UVehicleWheel::UVehicleWheel(const FObjectInitializer& ObjectInitializer)
 	SweepType = EWheelSweepType::SimpleAndComplex;
 }
 
+FPhysXVehicleManager* UVehicleWheel::GetVehicleManager() const
+{
+	UWorld* World = GEngine->GetWorldFromContextObject(VehicleSim, EGetWorldErrorMode::LogAndReturnNull);
+	return World ? FPhysXVehicleManager::GetVehicleManagerFromScene(World->GetPhysicsScene()) : nullptr;
+}
+
 float UVehicleWheel::GetSteerAngle() const
 {
-	FPhysXVehicleManager* VehicleManager = FPhysXVehicleManager::GetVehicleManagerFromScene(VehicleSim->GetWorld()->GetPhysicsScene());
-	SCOPED_SCENE_READ_LOCK(VehicleManager->GetScene());
-	return FMath::RadiansToDegrees( VehicleManager->GetWheelsStates_AssumesLocked(VehicleSim)[WheelIndex].steerAngle );
+	if (FPhysXVehicleManager* VehicleManager = GetVehicleManager())
+	{
+		SCOPED_SCENE_READ_LOCK(VehicleManager->GetScene());
+		return FMath::RadiansToDegrees(VehicleManager->GetWheelsStates_AssumesLocked(VehicleSim)[WheelIndex].steerAngle);
+	}
+	return 0.0f;
 }
 
 float UVehicleWheel::GetRotationAngle() const
 {
-	FPhysXVehicleManager* VehicleManager = FPhysXVehicleManager::GetVehicleManagerFromScene(VehicleSim->GetWorld()->GetPhysicsScene());
-	SCOPED_SCENE_READ_LOCK(VehicleManager->GetScene());
+	if (FPhysXVehicleManager* VehicleManager = GetVehicleManager())
+	{
+		SCOPED_SCENE_READ_LOCK(VehicleManager->GetScene());
 
-	float RotationAngle = -1.0f * FMath::RadiansToDegrees( VehicleSim->PVehicle->mWheelsDynData.getWheelRotationAngle( WheelIndex ) );
-	check(!FMath::IsNaN(RotationAngle));
-	return RotationAngle;
+		float RotationAngle = -1.0f * FMath::RadiansToDegrees(VehicleSim->PVehicle->mWheelsDynData.getWheelRotationAngle(WheelIndex));
+		ensure(!FMath::IsNaN(RotationAngle));
+		return RotationAngle;
+	}
+	return 0.0f;
 }
 
 float UVehicleWheel::GetSuspensionOffset() const
 {
-	FPhysXVehicleManager* VehicleManager = FPhysXVehicleManager::GetVehicleManagerFromScene(VehicleSim->GetWorld()->GetPhysicsScene());
-	SCOPED_SCENE_READ_LOCK(VehicleManager->GetScene());
+	if (FPhysXVehicleManager* VehicleManager = GetVehicleManager())
+	{
+		SCOPED_SCENE_READ_LOCK(VehicleManager->GetScene());
 
-	return VehicleManager->GetWheelsStates_AssumesLocked(VehicleSim)[WheelIndex].suspJounce;
+		return VehicleManager->GetWheelsStates_AssumesLocked(VehicleSim)[WheelIndex].suspJounce;
+	}
+	return 0.0f;
 }
 
 bool UVehicleWheel::IsInAir() const
 {
-	FPhysXVehicleManager* VehicleManager = FPhysXVehicleManager::GetVehicleManagerFromScene(VehicleSim->GetWorld()->GetPhysicsScene());
-	SCOPED_SCENE_READ_LOCK(VehicleManager->GetScene());
+	if (FPhysXVehicleManager* VehicleManager = GetVehicleManager())
+	{
+		SCOPED_SCENE_READ_LOCK(VehicleManager->GetScene());
 
-	return VehicleManager->GetWheelsStates_AssumesLocked(VehicleSim)[WheelIndex].isInAir;
+		return VehicleManager->GetWheelsStates_AssumesLocked(VehicleSim)[WheelIndex].isInAir;
+	}
+	return false;
 }
 
 void UVehicleWheel::Init( UWheeledVehicleMovementComponent* InVehicleSim, int32 InWheelIndex )
@@ -114,14 +132,16 @@ FVector UVehicleWheel::GetPhysicsLocation()
 {
 	if ( WheelShape )
 	{
-		FPhysXVehicleManager* VehicleManager = FPhysXVehicleManager::GetVehicleManagerFromScene(VehicleSim->GetWorld()->GetPhysicsScene());
-		SCOPED_SCENE_READ_LOCK(VehicleManager->GetScene());
+		if (FPhysXVehicleManager* VehicleManager = GetVehicleManager())
+		{
+			SCOPED_SCENE_READ_LOCK(VehicleManager->GetScene());
 
-		PxVec3 PLocation = VehicleSim->PVehicle->getRigidDynamicActor()->getGlobalPose().transform( WheelShape->getLocalPose() ).p;
-		return P2UVector( PLocation );
+			PxVec3 PLocation = VehicleSim->PVehicle->getRigidDynamicActor()->getGlobalPose().transform(WheelShape->getLocalPose()).p;
+			return P2UVector(PLocation);
+		}
 	}
 
-	return FVector(0.0f);
+	return FVector::ZeroVector;
 }
 
 #if WITH_EDITOR
diff --git a/Engine/Plugins/Runtime/PhysXVehicles/Source/PhysXVehicles/Public/VehicleWheel.h b/Engine/Plugins/Runtime/PhysXVehicles/Source/PhysXVehicles/Public/VehicleWheel.h
--- a/Engine/Plugins/Runtime/PhysXVehicles/Source/PhysXVehicles/Public/VehicleWheel.h
+++ b/Engine/Plugins/Runtime/PhysXVehicles/Source/PhysXVehicles/Public/VehicleWheel.h
@@ -14,6 +14,7 @@
 #include "VehicleWheel.generated.h"
 
 class UPhysicalMaterial;
+class FPhysXVehicleManager;
 
 namespace physx
 {
@@ -243,9 +244,11 @@ protected:
 	 * Get the wheel's location in physics land
 	 */
 	FVector GetPhysicsLocation();
-
 #endif // WITH_PHYSX
 
+private:
+	FPhysXVehicleManager* GetVehicleManager() const;
+
 public:
 
 	/** Get contact surface material */
diff --git a/Engine/Plugins/Tests/RuntimeTests/Source/RuntimeTests/Private/Scripting/ScriptFuzzingTests.cpp b/Engine/Plugins/Tests/RuntimeTests/Source/RuntimeTests/Private/Scripting/ScriptFuzzingTests.cpp
--- a/Engine/Plugins/Tests/RuntimeTests/Source/RuntimeTests/Private/Scripting/ScriptFuzzingTests.cpp
+++ b/Engine/Plugins/Tests/RuntimeTests/Source/RuntimeTests/Private/Scripting/ScriptFuzzingTests.cpp
@@ -3,13 +3,76 @@
 #include "CoreMinimal.h"
 #include "Misc/AutomationTest.h"
 
-// #include "UObject/Script.h"
 #include "UObject/UObjectIterator.h"
 #include "UObject/Class.h"
 #include "UObject/UnrealType.h"
 #include "GameFramework/Actor.h"
 #include "Engine/World.h"
 #include "UObject/Package.h"
+#include "HAL/ExceptionHandling.h"
+
+UObject* CreateFuzzingHostObject(UWorld* World, UClass* DesiredClass, TArray<UObject*>& CreatedHosts)
+{
+	UObject* Result = nullptr;
+	if (!DesiredClass->HasAnyClassFlags(CLASS_Abstract))
+	{
+		if (DesiredClass->IsChildOf(AActor::StaticClass()))
+		{
+			if (DesiredClass->HasAnyClassFlags(CLASS_NotPlaceable))
+			{
+				UE_LOG(LogTemp, Warning, TEXT("Cannot fuzz non-static methods in %s%s as it is marked NotPlaceable"), DesiredClass->GetPrefixCPP(), *DesiredClass->GetName());
+			}
+			else
+			{
+				FActorSpawnParameters SpawnParams;
+				SpawnParams.bNoFail = true;
+				SpawnParams.SpawnCollisionHandlingOverride = ESpawnActorCollisionHandlingMethod::AlwaysSpawn;
+				Result = World->SpawnActor<AActor>(DesiredClass, FTransform::Identity, SpawnParams);
+			}
+		}
+		else
+		{
+			UObject* Outer = GetTransientPackage();
+			UClass* DesiredWithin = DesiredClass->ClassWithin;
+			if (DesiredWithin == UObject::StaticClass())
+			{
+				DesiredWithin = nullptr;
+			}
+
+			if ((DesiredWithin == nullptr) && DesiredClass->IsChildOf(UActorComponent::StaticClass()))
+			{
+				DesiredWithin = AActor::StaticClass();
+			}
+
+			if ((DesiredWithin != nullptr) && !Outer->GetClass()->IsChildOf(DesiredWithin))
+			{
+				UE_LOG(LogTemp, Log, TEXT("%s%s has a desired class within of %s"), DesiredClass->GetPrefixCPP(), *DesiredClass->GetName(), *DesiredWithin->GetName());
+				Outer = CreateFuzzingHostObject(World, DesiredWithin, /*inout*/ CreatedHosts);
+			}
+
+			if (Outer != nullptr)
+			{
+				 Result = NewObject<UObject>(Outer, DesiredClass);
+			}
+			else
+			{
+				UE_LOG(LogTemp, Warning, TEXT("Cannot fuzz non-static methods in %s%s, was unable to create an appropriate outer to satisify a ClassWithin constraint"), DesiredClass->GetPrefixCPP(), *DesiredClass->GetName());
+			}
+		}
+	}
+	else
+	{
+		UE_LOG(LogTemp, Warning, TEXT("Cannot fuzz non-static methods in %s%s as it is marked Abstract; we might be able to in derived classes"), DesiredClass->GetPrefixCPP(), *DesiredClass->GetName());
+	}
+
+	if (Result != nullptr)
+	{
+		CreatedHosts.Add(Result);
+	}
+
+	return Result;
+}
+
 
 FAutoConsoleCommandWithWorldAndArgs GScriptFuzzingCommand(
 	TEXT("Test.ScriptFuzzing"),
@@ -19,14 +82,37 @@ FAutoConsoleCommandWithWorldAndArgs GScriptFuzzingCommand(
 {
 	FEditorScriptExecutionGuard AllowScriptExec;
 
+	//@TODO: Ensures are still bad in BP callable methods, this is just to help get to the crashes first
+	//@TODO: THis currently controls both ensures and checks, so it's no good
+	//	TGuardValue<bool> DisableEnsures(GIgnoreDebugger, true);
+
+	//@TODO: Need to fuzz these too
+	TArray<FString> BannedPackageNames;
+	// does lots of terrible things right now (UnrealEd) 
+	BannedPackageNames.Add(TEXT("UnrealEd"));
+	// slots have complicated lifecycle rules that aren't exposed programmatically (UMG)
+	BannedPackageNames.Add(TEXT("UMG"));
+	// DLL badness when not installed?  TBD
+	BannedPackageNames.Add(TEXT("OculusHMD"));
+
 	// Run thru all native classes
 	for (TObjectIterator<UClass> ClassIt; ClassIt; ++ClassIt)
 	{
 		UClass* TestClass = *ClassIt;
 		if (TestClass->HasAnyClassFlags(CLASS_Native))
 		{
-			//@TODO: Need to fuzz these too, but they do lots of terrible things right now
-			if (TestClass->GetOutermost()->GetName().Contains(TEXT("UnrealEd")))
+			// See if the class is in a banned package
+			const FString OutermostPackageName = TestClass->GetOutermost()->GetName();
+			bool bBanned = false;
+			for (const FString& BannedPackageName : BannedPackageNames)
+			{
+				if (OutermostPackageName.Contains(BannedPackageName))
+				{
+					bBanned = true;
+					break;
+				}
+			}
+			if (bBanned)
 			{
 				continue;
 			}
@@ -50,43 +136,15 @@ FAutoConsoleCommandWithWorldAndArgs GScriptFuzzingCommand(
 			if (bHasFunctions)
 			{
 				// Create an instance of the object if necessary (function libraries can use the CDO instead)
+				TArray<UObject*> CreatedHosts;
 				UObject* CreatedInstance = nullptr;
 				if (bHasNonStaticFunctions)
 				{
-					if (!TestClass->HasAnyClassFlags(CLASS_Abstract))
-					{
-						if (TestClass->IsChildOf(AActor::StaticClass()))
-						{
-							if (TestClass->HasAnyClassFlags(CLASS_NotPlaceable))
-							{
-								UE_LOG(LogTemp, Warning, TEXT("Cannot fuzz non-static methods in %s%s as it is marked NotPlaceable"), TestClass->GetPrefixCPP(), *TestClass->GetName());
-							}
-							else
-							{
-								FActorSpawnParameters SpawnParams;
-								SpawnParams.bNoFail = true;
-								SpawnParams.SpawnCollisionHandlingOverride = ESpawnActorCollisionHandlingMethod::AlwaysSpawn;
-								CreatedInstance = World->SpawnActor<AActor>(TestClass, FTransform::Identity, SpawnParams);
-							}
-						}
-						else
-						{
-							if (TestClass->ClassWithin != nullptr)
-							{
-								UE_LOG(LogTemp, Warning, TEXT("Cannot fuzz non-static methods in %s%s, no support for ClassWithin != nullptr yet"), TestClass->GetPrefixCPP(), *TestClass->GetName());
-							}
-							else
-							{
-								CreatedInstance = NewObject<UObject>(GetTransientPackage(), TestClass);
-							}
-						}
-					}
-					else
-					{
-						UE_LOG(LogTemp, Warning, TEXT("Cannot fuzz non-static methods in %s%s as it is marked Abstract; we might be able to in derived classes"), TestClass->GetPrefixCPP(), *TestClass->GetName());
-					}
+					CreatedInstance = CreateFuzzingHostObject(World, TestClass, /*out*/ CreatedHosts);
 				}
 
+				int32 NumFunctionsWeFailedToFuzzDueToNoObject = 0;
+
 				// Run thru all functions and fuzz them
 				for (TFieldIterator<UFunction> FuncIt(TestClass); FuncIt; ++FuncIt)
 				{
@@ -110,19 +168,23 @@ FAutoConsoleCommandWithWorldAndArgs GScriptFuzzingCommand(
 						{
 							if (!TestClass->HasAnyClassFlags(CLASS_Abstract))
 							{
-								UE_LOG(LogTemp, Warning, TEXT("Failed to fuzz %s%s::%s() because we could not make an object to test it on"), TestClass->GetPrefixCPP(), *TestClass->GetName(), *Function->GetName());
+								++NumFunctionsWeFailedToFuzzDueToNoObject;
 							}
 						}
 					}
 				}
 
+				if (NumFunctionsWeFailedToFuzzDueToNoObject > 0)
+				{
+					UE_LOG(LogTemp, Warning, TEXT("Failed to fuzz %d non-static functions on %s%s because we could not make an object to test it on"), NumFunctionsWeFailedToFuzzDueToNoObject, TestClass->GetPrefixCPP(), *TestClass->GetName());
+				}
+
 				// Tear down the created instance
-				if (CreatedInstance != nullptr)
+				for (UObject* CreatedHost : CreatedHosts)
 				{
-					if (TestClass->IsChildOf(AActor::StaticClass()))
+					if (CreatedHost->IsA(AActor::StaticClass()))
 					{
-						World->DestroyActor(CastChecked<AActor>(CreatedInstance));
-						CreatedInstance = nullptr;
+						World->DestroyActor(CastChecked<AActor>(CreatedHost));
 					}
 				}
 			}
diff --git a/Engine/Source/Developer/AITestSuite/Private/MockAI/MockGameplayTasks.cpp b/Engine/Source/Developer/AITestSuite/Private/MockAI/MockGameplayTasks.cpp
--- a/Engine/Source/Developer/AITestSuite/Private/MockAI/MockGameplayTasks.cpp
+++ b/Engine/Source/Developer/AITestSuite/Private/MockAI/MockGameplayTasks.cpp
@@ -29,7 +29,10 @@ UMockTask_Log* UMockTask_Log::CreateTask(IGameplayTaskOwnerInterface& TaskOwner,
 
 void UMockTask_Log::Activate()
 {
-	Logger->Log(ETestTaskMessage::Activate);
+	if (Logger)
+	{
+		Logger->Log(ETestTaskMessage::Activate);
+	}
 	Super::Activate();
 	if (bShoudEndAsPartOfActivation)
 	{
@@ -39,25 +42,37 @@ void UMockTask_Log::Activate()
 
 void UMockTask_Log::OnDestroy(bool bInOwnerFinished)
 {
-	Logger->Log(ETestTaskMessage::Ended);
+	if (Logger)
+	{
+		Logger->Log(ETestTaskMessage::Ended);
+	}
 	Super::OnDestroy(bInOwnerFinished);
 }
 
 void UMockTask_Log::TickTask(float DeltaTime)
 {
-	Logger->Log(ETestTaskMessage::Tick);
+	if (Logger)
+	{
+		Logger->Log(ETestTaskMessage::Tick);
+	}
 	Super::TickTask(DeltaTime);
 }
 
 void UMockTask_Log::ExternalConfirm(bool bEndTask)
 {
-	Logger->Log(ETestTaskMessage::ExternalConfirm);
+	if (Logger)
+	{
+		Logger->Log(ETestTaskMessage::ExternalConfirm);
+	}
 	Super::ExternalConfirm(bEndTask);
 }
 
 void UMockTask_Log::ExternalCancel()
 {
-	Logger->Log(ETestTaskMessage::ExternalCancel);
+	if (Logger)
+	{
+		Logger->Log(ETestTaskMessage::ExternalCancel);
+	}
 	Super::ExternalCancel();
 }
 
diff --git a/Engine/Source/Developer/BlueprintCompilerCppBackend/Private/BlueprintCompilerCppBackendValueHelper.cpp b/Engine/Source/Developer/BlueprintCompilerCppBackend/Private/BlueprintCompilerCppBackendValueHelper.cpp
--- a/Engine/Source/Developer/BlueprintCompilerCppBackend/Private/BlueprintCompilerCppBackendValueHelper.cpp
+++ b/Engine/Source/Developer/BlueprintCompilerCppBackend/Private/BlueprintCompilerCppBackendValueHelper.cpp
@@ -781,10 +781,12 @@ FString FEmitDefaultValueHelper::HandleSpecialTypes(FEmitterLocalContext& Contex
 
 			if (!bCreatingSubObjectsOfClass && Property->HasAnyPropertyFlags(CPF_InstancedReference))
 			{
-				const FString CreateAsInstancedSubobject = HandleInstancedSubobject(Context, Object, Object->HasAnyFlags(RF_ArchetypeObject));
-				if (!CreateAsInstancedSubobject.IsEmpty())
+				// Emit ctor code to create the instance only if it's not a default subobject; otherwise, just access it.
+				const bool bCreateInstance = !Object->IsDefaultSubobject();
+				const FString GetOrCreateAsInstancedSubobject = HandleInstancedSubobject(Context, Object, bCreateInstance);
+				if (!GetOrCreateAsInstancedSubobject.IsEmpty())
 				{
-					return CreateAsInstancedSubobject;
+					return GetOrCreateAsInstancedSubobject;
 				}
 			}
 
@@ -1968,6 +1970,8 @@ FString FEmitDefaultValueHelper::HandleInstancedSubobject(FEmitterLocalContext&
 		}
 		else
 		{
+			check(Object->IsDefaultSubobject());
+
 			Context.AddLine(FString::Printf(TEXT("auto %s = CastChecked<%s>(%s(TEXT(\"%s\")));")
 				, *LocalNativeName
 				, *FEmitHelper::GetCppName(ObjectClass)
diff --git a/Engine/Source/Developer/FunctionalTesting/Classes/FunctionalTest.h b/Engine/Source/Developer/FunctionalTesting/Classes/FunctionalTest.h
--- a/Engine/Source/Developer/FunctionalTesting/Classes/FunctionalTest.h
+++ b/Engine/Source/Developer/FunctionalTesting/Classes/FunctionalTest.h
@@ -14,6 +14,7 @@
 class Error;
 class UBillboardComponent;
 class UTraceQueryTestResults;
+class UWorld;
 
 //Experimental effort at automated cpu captures from the functional testing.
 class FFunctionalTestExternalProfiler : public FScopedExternalProfilerBase
@@ -159,6 +160,10 @@ public:
 
 	UAutomationPerformaceHelper();
 
+	// UObject interface
+	virtual UWorld* GetWorld() const override;
+	// End of UObject interface
+
 	//Begin basic stat recording
 
 	UFUNCTION(BlueprintCallable, Category = Perf)
diff --git a/Engine/Source/Developer/FunctionalTesting/Private/FunctionalTest.cpp b/Engine/Source/Developer/FunctionalTesting/Private/FunctionalTest.cpp
--- a/Engine/Source/Developer/FunctionalTesting/Private/FunctionalTest.cpp
+++ b/Engine/Source/Developer/FunctionalTesting/Private/FunctionalTest.cpp
@@ -23,6 +23,7 @@
 #include "DelayForFramesLatentAction.h"
 #include "Engine/DebugCameraController.h"
 #include "TraceQueryTestResults.h"
+#include "Misc/RuntimeErrors.h"
 
 namespace
 {
@@ -198,7 +199,7 @@ bool AFunctionalTest::RunTest(const TArray<FString>& Params)
 	FailureMessage = TEXT("");
 	
 	//Do not collect garbage during the test. We force GC at the end.
-	GetWorld()->DelayGarbageCollection();
+	GEngine->DelayGarbageCollection();
 
 	RunFrame = GFrameNumber;
 	RunTime = GetWorld()->GetTimeSeconds();
@@ -249,7 +250,7 @@ void AFunctionalTest::Tick(float DeltaSeconds)
 	}
 
 	//Do not collect garbage during the test. We force GC at the end.
-	GetWorld()->DelayGarbageCollection();
+	GEngine->DelayGarbageCollection();
 
 	if ( !bIsReady )
 	{
@@ -300,7 +301,7 @@ void AFunctionalTest::FinishTest(EFunctionalTestResult TestResult, const FString
 	}
 
 	//Force GC at the end of every test.
-	GetWorld()->ForceGarbageCollection();
+	GEngine->ForceGarbageCollection();
 
 	Result = TestResult;
 
@@ -1031,13 +1032,23 @@ UAutomationPerformaceHelper::UAutomationPerformaceHelper()
 {
 }
 
+UWorld* UAutomationPerformaceHelper::GetWorld() const
+{
+	UWorld* OuterWorld = GetOuter()->GetWorld();
+	ensureAsRuntimeWarning(OuterWorld != nullptr);
+	return OuterWorld;
+}
+
 void UAutomationPerformaceHelper::BeginRecordingBaseline(FString RecordName)
 {
-	bRecordingBasicStats = true;
-	bRecordingBaselineBasicStats = true;
-	bGPUTraceIfBelowBudget = false;
-	Records.Add(FPerfStatsRecord(RecordName));
-	GEngine->SetEngineStat(GetOuter()->GetWorld(), GetOuter()->GetWorld()->GetGameViewport(), TEXT("Unit"), true);
+	if (UWorld* World = GetWorld())
+	{
+		bRecordingBasicStats = true;
+		bRecordingBaselineBasicStats = true;
+		bGPUTraceIfBelowBudget = false;
+		Records.Add(FPerfStatsRecord(RecordName));
+		GEngine->SetEngineStat(World, World->GetGameViewport(), TEXT("Unit"), true);
+	}
 }
 
 void UAutomationPerformaceHelper::EndRecordingBaseline()
@@ -1048,21 +1059,24 @@ void UAutomationPerformaceHelper::EndRecordingBaseline()
 
 void UAutomationPerformaceHelper::BeginRecording(FString RecordName, float InGPUBudget, float InRenderThreadBudget, float InGameThreadBudget)
 {
-	//Ensure we're recording engine stats.
-	GEngine->SetEngineStat(GetOuter()->GetWorld(), GetOuter()->GetWorld()->GetGameViewport(), TEXT("Unit"), true);
-	bRecordingBasicStats = true;
-	bRecordingBaselineBasicStats = false;
-	bGPUTraceIfBelowBudget = false;
-
-	FPerfStatsRecord* CurrRecord = GetCurrentRecord();
-	if (!CurrRecord || CurrRecord->Name != RecordName)
+	if (UWorld* World = GetWorld())
 	{
-		Records.Add(FPerfStatsRecord(RecordName));
-		CurrRecord = GetCurrentRecord();
-	}
+		//Ensure we're recording engine stats.
+		GEngine->SetEngineStat(World, World->GetGameViewport(), TEXT("Unit"), true);
+		bRecordingBasicStats = true;
+		bRecordingBaselineBasicStats = false;
+		bGPUTraceIfBelowBudget = false;
 
-	check(CurrRecord);
-	CurrRecord->SetBudgets(InGPUBudget, InRenderThreadBudget, InGameThreadBudget);
+		FPerfStatsRecord* CurrRecord = GetCurrentRecord();
+		if (!CurrRecord || CurrRecord->Name != RecordName)
+		{
+			Records.Add(FPerfStatsRecord(RecordName));
+			CurrRecord = GetCurrentRecord();
+		}
+
+		check(CurrRecord);
+		CurrRecord->SetBudgets(InGPUBudget, InRenderThreadBudget, InGameThreadBudget);
+	}
 }
 
 void UAutomationPerformaceHelper::EndRecording()
@@ -1101,10 +1115,13 @@ void UAutomationPerformaceHelper::Tick(float DeltaSeconds)
 
 void UAutomationPerformaceHelper::Sample(float DeltaSeconds)
 {
-	int32 Index = Records.Num() - 1;
-	if (Index >= 0 && bRecordingBasicStats)
+	if (UWorld* World = GetWorld())
 	{
-		Records[Index].Sample(GetOuter()->GetWorld(), DeltaSeconds, bRecordingBaselineBasicStats);
+		int32 Index = Records.Num() - 1;
+		if (Index >= 0 && bRecordingBasicStats)
+		{
+			Records[Index].Sample(World, DeltaSeconds, bRecordingBaselineBasicStats);
+		}
 	}
 }
 
@@ -1262,12 +1279,18 @@ void UAutomationPerformaceHelper::TriggerGPUTraceIfRecordFallsBelowBudget()
 
 void UAutomationPerformaceHelper::BeginStatsFile(const FString& RecordName)
 {
-	FString MapName = GetOuter()->GetWorld()->GetMapName();
-	FString Cmd = FString::Printf(TEXT("Stat StartFile %s-%s/%s.ue4stats"), *MapName, *StartOfTestingTime, *RecordName);
-	GEngine->Exec(GetOuter()->GetWorld(), *Cmd);
+	if (UWorld* World = GetWorld())
+	{
+		FString MapName = World->GetMapName();
+		FString Cmd = FString::Printf(TEXT("Stat StartFile %s-%s/%s.ue4stats"), *MapName, *StartOfTestingTime, *RecordName);
+		GEngine->Exec(World, *Cmd);
+	}
 }
 
 void UAutomationPerformaceHelper::EndStatsFile()
 {
-	GEngine->Exec(GetOuter()->GetWorld(), TEXT("Stat StopFile"));
+	if (UWorld* World = GetWorld())
+	{
+		GEngine->Exec(World, TEXT("Stat StopFile"));
+	}
 }
diff --git a/Engine/Source/Developer/FunctionalTesting/Private/FunctionalTestingManager.cpp b/Engine/Source/Developer/FunctionalTesting/Private/FunctionalTestingManager.cpp
--- a/Engine/Source/Developer/FunctionalTesting/Private/FunctionalTestingManager.cpp
+++ b/Engine/Source/Developer/FunctionalTesting/Private/FunctionalTestingManager.cpp
@@ -101,7 +101,7 @@ bool UFunctionalTestingManager::RunAllFunctionalTests(UObject* WorldContextObjec
 	}
 	
 	UWorld* World = GEngine->GetWorldFromContextObjectChecked(WorldContextObject);
-	World->ForceGarbageCollection(true);
+	GEngine->ForceGarbageCollection(true);
 
 	Manager->bFinished = false;
 	Manager->bLooped = bRunLooped;
diff --git a/Engine/Source/Developer/TreeMap/Public/STreeMap.h b/Engine/Source/Developer/TreeMap/Public/STreeMap.h
--- a/Engine/Source/Developer/TreeMap/Public/STreeMap.h
+++ b/Engine/Source/Developer/TreeMap/Public/STreeMap.h
@@ -184,10 +184,6 @@ protected:
 
 	/** Recursively applies the active visualization to all nodes, such as size by attribute, or color by attribute */
 	void ApplyVisualizationToNodesRecursively( const FTreeMapNodeDataRef& Node, const FLinearColor& DefaultColor, const int32 TreeDepth );
-		
-	/** Displays a context menu at the specified location with options for configuring the tree display */
-	DEPRECATED(4.9, "ShowOptionsMenuAt() taking a FVector2D param is deprecated. Use ShowOptionsMenuAt() that takes an FPointerEvent instead.")
-	void ShowOptionsMenuAt(const FVector2D& ScreenSpacePosition);
 
 	/** Displays a context menu at the specified location with options for configuring the tree display */
 	void ShowOptionsMenuAt( const FPointerEvent& InMouseEvent );
diff --git a/Engine/Source/Developer/TreeMap/STreeMap.cpp b/Engine/Source/Developer/TreeMap/STreeMap.cpp
--- a/Engine/Source/Developer/TreeMap/STreeMap.cpp
+++ b/Engine/Source/Developer/TreeMap/STreeMap.cpp
@@ -1171,7 +1171,7 @@ void STreeMap::StartRenamingNode( const FGeometry& MyGeometry, const FTreeMapNod
 	FSlateApplication::Get().PushMenu( AsShared(), FWidgetPath(), RenamerWidget, MyGeometry.LocalToAbsolute( RelativePosition ), FPopupTransitionEffect::None, bFocusImmediately );
 
 	// Focus the text box right after we spawn it so that the user can start typing
-	FSlateApplication::Get().SetKeyboardFocus( EditableText, EKeyboardFocusCause::SetDirectly );
+	FSlateApplication::Get().SetKeyboardFocus( EditableText, EFocusCause::SetDirectly );
 }
 
 
@@ -1298,13 +1298,6 @@ void STreeMap::ApplyVisualizationToNodesRecursively( const FTreeMapNodeDataRef&
 
 		ApplyVisualizationToNodesRecursively( ChildNode.ToSharedRef(), ChildColor, TreeDepth + 1 );
 	}
-
-}
-
-void STreeMap::ShowOptionsMenuAt(const FVector2D& ScreenSpacePosition)
-{
-	// deprecated
-	ShowOptionsMenuAtInternal(ScreenSpacePosition, FWidgetPath());
 }
 
 void STreeMap::ShowOptionsMenuAt(const FPointerEvent& InMouseEvent)
diff --git a/Engine/Source/Editor/AnimGraph/Classes/AnimGraphNode_Base.h b/Engine/Source/Editor/AnimGraph/Classes/AnimGraphNode_Base.h
--- a/Engine/Source/Editor/AnimGraph/Classes/AnimGraphNode_Base.h
+++ b/Engine/Source/Editor/AnimGraph/Classes/AnimGraphNode_Base.h
@@ -181,9 +181,6 @@ class ANIMGRAPH_API UAnimGraphNode_Base : public UK2Node
 	// Gives each visual node a chance to update the node template before it is inserted in the compiled class
 	virtual void BakeDataDuringCompilation(FCompilerResultsLog& MessageLog) {}
 
-	// preload asset required for this node in this function
-	virtual void PreloadRequiredAssets() override {}
-
 	// Give the node a chance to change the display name of a pin
 	virtual void PostProcessPinName(const UEdGraphPin* Pin, FString& DisplayName) const;
 
diff --git a/Engine/Source/Editor/AnimGraph/Classes/K2Node_TransitionRuleGetter.h b/Engine/Source/Editor/AnimGraph/Classes/K2Node_TransitionRuleGetter.h
--- a/Engine/Source/Editor/AnimGraph/Classes/K2Node_TransitionRuleGetter.h
+++ b/Engine/Source/Editor/AnimGraph/Classes/K2Node_TransitionRuleGetter.h
@@ -50,6 +50,7 @@ class UK2Node_TransitionRuleGetter : public UK2Node
 
 	// UEdGraphNode interface
 	virtual void AllocateDefaultPins() override;
+	virtual void PreloadRequiredAssets() override;
 	virtual FText GetNodeTitle(ENodeTitleType::Type TitleType) const override;
 	virtual FText GetTooltipText() const override;
 	virtual bool ShowPaletteIconOnNode() const override { return false; }
diff --git a/Engine/Source/Editor/AnimGraph/Private/K2Node_TransitionRuleGetter.cpp b/Engine/Source/Editor/AnimGraph/Private/K2Node_TransitionRuleGetter.cpp
--- a/Engine/Source/Editor/AnimGraph/Private/K2Node_TransitionRuleGetter.cpp
+++ b/Engine/Source/Editor/AnimGraph/Private/K2Node_TransitionRuleGetter.cpp
@@ -37,6 +37,13 @@ void UK2Node_TransitionRuleGetter::AllocateDefaultPins()
 	Super::AllocateDefaultPins();
 }
 
+void UK2Node_TransitionRuleGetter::PreloadRequiredAssets()
+{
+	PreloadObject(AssociatedAnimAssetPlayerNode);
+	PreloadObject(AssociatedStateNode);
+	Super::PreloadRequiredAssets();
+}
+
 FText UK2Node_TransitionRuleGetter::GetFriendlyName(ETransitionGetter::Type TypeID)
 {
 	FText FriendlyName;
diff --git a/Engine/Source/Editor/AnimationBlueprintEditor/Private/AnimationBlueprintEditor.cpp b/Engine/Source/Editor/AnimationBlueprintEditor/Private/AnimationBlueprintEditor.cpp
--- a/Engine/Source/Editor/AnimationBlueprintEditor/Private/AnimationBlueprintEditor.cpp
+++ b/Engine/Source/Editor/AnimationBlueprintEditor/Private/AnimationBlueprintEditor.cpp
@@ -534,7 +534,7 @@ void FAnimationBlueprintEditor::OnConvertToSequenceEvaluator()
 
 				if (ensure(OldPosePin && NewPosePin))
 				{
-					NewPosePin->CopyPersistentDataFromOldPin(*OldPosePin);
+					NewPosePin->MovePersistentDataFromOldPin(*OldPosePin);
 				}
 
 				// remove from selection and from graph
@@ -588,7 +588,7 @@ void FAnimationBlueprintEditor::OnConvertToSequencePlayer()
 
 				if (ensure(OldPosePin && NewPosePin))
 				{
-					NewPosePin->CopyPersistentDataFromOldPin(*OldPosePin);
+					NewPosePin->MovePersistentDataFromOldPin(*OldPosePin);
 				}
 
 				// remove from selection and from graph
@@ -644,7 +644,7 @@ void FAnimationBlueprintEditor::OnConvertToBlendSpaceEvaluator()
 
 				if (ensure(OldPosePin && NewPosePin))
 				{
-					NewPosePin->CopyPersistentDataFromOldPin(*OldPosePin);
+					NewPosePin->MovePersistentDataFromOldPin(*OldPosePin);
 				}
 
 				OldPosePin = OldNode->FindPin(TEXT("Y"));
@@ -652,7 +652,7 @@ void FAnimationBlueprintEditor::OnConvertToBlendSpaceEvaluator()
 
 				if (ensure(OldPosePin && NewPosePin))
 				{
-					NewPosePin->CopyPersistentDataFromOldPin(*OldPosePin);
+					NewPosePin->MovePersistentDataFromOldPin(*OldPosePin);
 				}
 
 
@@ -661,7 +661,7 @@ void FAnimationBlueprintEditor::OnConvertToBlendSpaceEvaluator()
 
 				if (ensure(OldPosePin && NewPosePin))
 				{
-					NewPosePin->CopyPersistentDataFromOldPin(*OldPosePin);
+					NewPosePin->MovePersistentDataFromOldPin(*OldPosePin);
 				}
 
 				// remove from selection and from graph
@@ -714,7 +714,7 @@ void FAnimationBlueprintEditor::OnConvertToBlendSpacePlayer()
 
 				if (ensure(OldPosePin && NewPosePin))
 				{
-					NewPosePin->CopyPersistentDataFromOldPin(*OldPosePin);
+					NewPosePin->MovePersistentDataFromOldPin(*OldPosePin);
 				}
 
 				OldPosePin = OldNode->FindPin(TEXT("Y"));
@@ -722,7 +722,7 @@ void FAnimationBlueprintEditor::OnConvertToBlendSpacePlayer()
 
 				if (ensure(OldPosePin && NewPosePin))
 				{
-					NewPosePin->CopyPersistentDataFromOldPin(*OldPosePin);
+					NewPosePin->MovePersistentDataFromOldPin(*OldPosePin);
 				}
 
 
@@ -731,7 +731,7 @@ void FAnimationBlueprintEditor::OnConvertToBlendSpacePlayer()
 
 				if (ensure(OldPosePin && NewPosePin))
 				{
-					NewPosePin->CopyPersistentDataFromOldPin(*OldPosePin);
+					NewPosePin->MovePersistentDataFromOldPin(*OldPosePin);
 				}
 
 				// remove from selection and from graph
@@ -783,7 +783,7 @@ void FAnimationBlueprintEditor::OnConvertToPoseBlender()
 
 				if (ensure(OldPosePin && NewPosePin))
 				{
-					NewPosePin->CopyPersistentDataFromOldPin(*OldPosePin);
+					NewPosePin->MovePersistentDataFromOldPin(*OldPosePin);
 				}
 
 				// remove from selection and from graph
@@ -837,7 +837,7 @@ void FAnimationBlueprintEditor::OnConvertToPoseByName()
 
 				if (ensure(OldPosePin && NewPosePin))
 				{
-					NewPosePin->CopyPersistentDataFromOldPin(*OldPosePin);
+					NewPosePin->MovePersistentDataFromOldPin(*OldPosePin);
 				}
 
 				// remove from selection and from graph
@@ -893,7 +893,7 @@ void FAnimationBlueprintEditor::OnConvertToAimOffsetLookAt()
 
 				if (ensure(OldPosePin && NewPosePin))
 				{
-					NewPosePin->CopyPersistentDataFromOldPin(*OldPosePin);
+					NewPosePin->MovePersistentDataFromOldPin(*OldPosePin);
 				}
 
 				OldPosePin = OldNode->FindPin(TEXT("BasePose"));
@@ -901,7 +901,7 @@ void FAnimationBlueprintEditor::OnConvertToAimOffsetLookAt()
 
 				if (ensure(OldPosePin && NewPosePin))
 				{
-					NewPosePin->CopyPersistentDataFromOldPin(*OldPosePin);
+					NewPosePin->MovePersistentDataFromOldPin(*OldPosePin);
 				}
 
 				// remove from selection and from graph
@@ -955,7 +955,7 @@ void FAnimationBlueprintEditor::OnConvertToAimOffsetSimple()
 
 				if (ensure(OldPosePin && NewPosePin))
 				{
-					NewPosePin->CopyPersistentDataFromOldPin(*OldPosePin);
+					NewPosePin->MovePersistentDataFromOldPin(*OldPosePin);
 				}
 
 				OldPosePin = OldNode->FindPin(TEXT("BasePose"));
@@ -963,7 +963,7 @@ void FAnimationBlueprintEditor::OnConvertToAimOffsetSimple()
 
 				if (ensure(OldPosePin && NewPosePin))
 				{
-					NewPosePin->CopyPersistentDataFromOldPin(*OldPosePin);
+					NewPosePin->MovePersistentDataFromOldPin(*OldPosePin);
 				}
 
 				// remove from selection and from graph
diff --git a/Engine/Source/Editor/AudioEditor/Private/SoundCueGraphNode_Base.cpp b/Engine/Source/Editor/AudioEditor/Private/SoundCueGraphNode_Base.cpp
--- a/Engine/Source/Editor/AudioEditor/Private/SoundCueGraphNode_Base.cpp
+++ b/Engine/Source/Editor/AudioEditor/Private/SoundCueGraphNode_Base.cpp
@@ -153,7 +153,7 @@ void USoundCueGraphNode_Base::ReconstructNode()
 
 	// Move the existing pins to a saved array
 	TArray<UEdGraphPin*> OldPins(Pins);
-	Pins.Empty();
+	Pins.Reset();
 
 	// Recreate the new pins
 	AllocateDefaultPins();
@@ -167,27 +167,21 @@ void USoundCueGraphNode_Base::ReconstructNode()
 	{
 		if (PinIndex < NewInputPins.Num())
 		{
-			NewInputPins[PinIndex]->CopyPersistentDataFromOldPin(*OldInputPins[PinIndex]);
+			NewInputPins[PinIndex]->MovePersistentDataFromOldPin(*OldInputPins[PinIndex]);
 		}
 	}
 
 	if (OldOutputPin)
 	{
-		NewOutputPin->CopyPersistentDataFromOldPin(*OldOutputPin);
-		OldInputPins.Empty();
-		OldOutputPin = NULL;
+		NewOutputPin->MovePersistentDataFromOldPin(*OldOutputPin);
 	}
 
 	// Throw away the original pins
-	for (int32 OldPinIndex = 0; OldPinIndex < OldPins.Num(); ++OldPinIndex)
+	for (UEdGraphPin* OldPin : OldPins)
 	{
-		UEdGraphPin* OldPin = OldPins[OldPinIndex];
 		OldPin->Modify();
-		OldPin->BreakAllPinLinks();
-
 		UEdGraphNode::DestroyPin(OldPin);
 	}
-	OldPins.Empty();
 }
 
 void USoundCueGraphNode_Base::AutowireNewNode(UEdGraphPin* FromPin)
diff --git a/Engine/Source/Editor/BlueprintGraph/Classes/EdGraphSchema_K2.h b/Engine/Source/Editor/BlueprintGraph/Classes/EdGraphSchema_K2.h
--- a/Engine/Source/Editor/BlueprintGraph/Classes/EdGraphSchema_K2.h
+++ b/Engine/Source/Editor/BlueprintGraph/Classes/EdGraphSchema_K2.h
@@ -843,26 +843,6 @@ public:
 	 */
 	virtual void CreateFunctionGraphTerminators(UEdGraph& Graph, UFunction* FunctionSignature) const;
 
-	/**
-	 * Converts a pin type into a fully qualified string (e.g., object'ObjectName').
-	 *
-	 * @param	Type	The type to convert into a string.
-	 *
-	 * @return	The converted type string.
-	 */
-	DEPRECATED(4.5, "UEdGraphSchema_K2::TypeToString is deprecated.  Use TypeToText instead.")
-	static FString TypeToString(const FEdGraphPinType& Type);
-
-	/**
-	 * Converts the type of a property into a fully qualified string (e.g., object'ObjectName').
-	 *
-	 * @param	Property	The property to convert into a string.
-	 *
-	 * @return	The converted type string.
-	 */
-	DEPRECATED(4.5, "UEdGraphSchema_K2::TypeToString is deprecated.  Use TypeToText instead.")
-	static FString TypeToString(UProperty* const Property);
-
 	/**
 	 * Converts the type of a property into a fully qualified string (e.g., object'ObjectName').
 	 *
diff --git a/Engine/Source/Editor/BlueprintGraph/Classes/K2Node.h b/Engine/Source/Editor/BlueprintGraph/Classes/K2Node.h
--- a/Engine/Source/Editor/BlueprintGraph/Classes/K2Node.h
+++ b/Engine/Source/Editor/BlueprintGraph/Classes/K2Node.h
@@ -339,7 +339,7 @@ class UK2Node : public UEdGraphNode
 	BLUEPRINTGRAPH_API virtual void GetNodeAttributes( TArray<TKeyValuePair<FString, FString>>& OutNodeAttributes ) const;
 
 	/** Called before compilation begins, giving a blueprint time to force the linker to load data */
-	BLUEPRINTGRAPH_API virtual void PreloadRequiredAssets();
+	BLUEPRINTGRAPH_API virtual void PreloadRequiredAssets() { }
 
 	/** 
 	 * Replacement for GetMenuEntries(). Override to add specific 
@@ -395,8 +395,7 @@ protected:
 	{
 		ERedirectType_None,
 		ERedirectType_Name,
-		ERedirectType_Value,
-		ERedirectType_Custom // DEPRECATED
+		ERedirectType_Value
 	};
 
 	// Handles the actual reconstruction (copying data, links, name, etc...) from two pins that have already been matched together
@@ -408,9 +407,6 @@ protected:
 	// Helper function to properly destroy a set of pins
 	BLUEPRINTGRAPH_API void DestroyPinList(TArray<UEdGraphPin*>& InPins);
 
-	DEPRECATED(4.16, "This function is deprecated, override ReconstructSinglePin instead ")
-	BLUEPRINTGRAPH_API virtual void CustomMapParamValue(UEdGraphPin& Pin);
-
 	/** Whether or not two pins match for purposes of reconnection after reconstruction.  This allows pins that may have had their names changed via reconstruction to be matched to their old values on a node-by-node basis, if needed*/
 	BLUEPRINTGRAPH_API virtual ERedirectType DoPinsMatchForReconstruction(const UEdGraphPin* NewPin, int32 NewPinIndex, const UEdGraphPin* OldPin, int32 OldPinIndex) const;
 
diff --git a/Engine/Source/Editor/BlueprintGraph/Classes/K2Node_ActorBoundEvent.h b/Engine/Source/Editor/BlueprintGraph/Classes/K2Node_ActorBoundEvent.h
--- a/Engine/Source/Editor/BlueprintGraph/Classes/K2Node_ActorBoundEvent.h
+++ b/Engine/Source/Editor/BlueprintGraph/Classes/K2Node_ActorBoundEvent.h
@@ -62,12 +62,10 @@ class UK2Node_ActorBoundEvent : public UK2Node_Event
 	BLUEPRINTGRAPH_API void InitializeActorBoundEventParams(AActor* InEventOwner, const UMulticastDelegateProperty* InDelegateProperty);
 
 	/** Return the delegate property that this event is bound to */
-	BLUEPRINTGRAPH_API UMulticastDelegateProperty* GetTargetDelegateProperty();
-	/** Const version of GetTargetDelegateProperty that does not search redirect table */
-	BLUEPRINTGRAPH_API UMulticastDelegateProperty* GetTargetDelegatePropertyConst() const;
+	BLUEPRINTGRAPH_API UMulticastDelegateProperty* GetTargetDelegateProperty() const;
 
-	/** Returns the delegate that this event is bound in */
-	BLUEPRINTGRAPH_API FMulticastScriptDelegate* GetTargetDelegate();
+	/** Returns the delegate that this event is bound in, for actors bound in level script */
+	BLUEPRINTGRAPH_API FMulticastScriptDelegate* GetTargetDelegate() const;
 
 private:
 	/** Constructing FText strings can be costly, so we cache the node's title */
diff --git a/Engine/Source/Editor/BlueprintGraph/Classes/K2Node_BreakStruct.h b/Engine/Source/Editor/BlueprintGraph/Classes/K2Node_BreakStruct.h
--- a/Engine/Source/Editor/BlueprintGraph/Classes/K2Node_BreakStruct.h
+++ b/Engine/Source/Editor/BlueprintGraph/Classes/K2Node_BreakStruct.h
@@ -41,6 +41,7 @@ class UK2Node_BreakStruct : public UK2Node_StructMemberGet
 
 	//~ Begin UEdGraphNode Interface
 	virtual void AllocateDefaultPins() override;
+	virtual void PreloadRequiredAssets() override;
 	virtual FText GetNodeTitle(ENodeTitleType::Type TitleType) const override;
 	virtual FLinearColor GetNodeTitleColor() const override;
 	virtual FText GetTooltipText() const override;
diff --git a/Engine/Source/Editor/BlueprintGraph/Classes/K2Node_ComponentBoundEvent.h b/Engine/Source/Editor/BlueprintGraph/Classes/K2Node_ComponentBoundEvent.h
--- a/Engine/Source/Editor/BlueprintGraph/Classes/K2Node_ComponentBoundEvent.h
+++ b/Engine/Source/Editor/BlueprintGraph/Classes/K2Node_ComponentBoundEvent.h
@@ -36,6 +36,7 @@ class UK2Node_ComponentBoundEvent : public UK2Node_Event
 	//~ End UObject Interface
 
 	//~ Begin UEdGraphNode Interface
+	virtual void ReconstructNode() override;
 	virtual bool CanPasteHere(const UEdGraph* TargetGraph) const override;
 	virtual FText GetNodeTitle(ENodeTitleType::Type TitleType) const override;
 	virtual FText GetTooltipText() const override;
diff --git a/Engine/Source/Editor/BlueprintGraph/Classes/K2Node_MakeArray.h b/Engine/Source/Editor/BlueprintGraph/Classes/K2Node_MakeArray.h
--- a/Engine/Source/Editor/BlueprintGraph/Classes/K2Node_MakeArray.h
+++ b/Engine/Source/Editor/BlueprintGraph/Classes/K2Node_MakeArray.h
@@ -3,73 +3,32 @@
 
 #pragma once
 
-#include "CoreMinimal.h"
-#include "UObject/ObjectMacros.h"
-#include "Textures/SlateIcon.h"
-#include "K2Node.h"
-#include "K2Node_AddPinInterface.h"
+#include "K2Node_MakeContainer.h"
 #include "K2Node_MakeArray.generated.h"
 
 class FBlueprintActionDatabaseRegistrar;
 class UEdGraphPin;
 
 UCLASS(MinimalAPI)
-class UK2Node_MakeArray : public UK2Node, public IK2Node_AddPinInterface
+class UK2Node_MakeArray : public UK2Node_MakeContainer
 {
 	GENERATED_UCLASS_BODY()
 
-	/** The number of input pins to generate for this node */
-	UPROPERTY()
-	int32 NumInputs;
-
-public:
-	BLUEPRINTGRAPH_API void RemoveInputPin(UEdGraphPin* Pin);
-
-	/** returns a reference to the output array pin of this node, which is responsible for defining the type */
-	BLUEPRINTGRAPH_API UEdGraphPin* GetOutputPin() const;
-
 public:
 	// UEdGraphNode interface
 	virtual FText GetNodeTitle(ENodeTitleType::Type TitleType) const override;
-	virtual void AllocateDefaultPins() override;
-	virtual void PostReconstructNode() override;
 	virtual FText GetTooltipText() const override;
 	virtual FSlateIcon GetIconAndTint(FLinearColor& OutColor) const override;
 	virtual void ValidateNodeDuringCompilation(class FCompilerResultsLog& MessageLog) const override;
 	// End of UEdGraphNode interface
 
 	// UK2Node interface
-	virtual bool IsNodePure() const override { return true; }
-	virtual void NotifyPinConnectionListChanged(UEdGraphPin* Pin) override;
 	virtual void GetContextMenuActions(const FGraphNodeContextMenuBuilder& Context) const override;
 	virtual class FNodeHandlingFunctor* CreateNodeHandler(class FKismetCompilerContext& CompilerContext) const override;
-	virtual bool IsConnectionDisallowed(const UEdGraphPin* MyPin, const UEdGraphPin* OtherPin, FString& OutReason) const override;
-	virtual void GetMenuActions(FBlueprintActionDatabaseRegistrar& ActionRegistrar) const override;
 	virtual FText GetMenuCategory() const override;
-	virtual int32 GetNodeRefreshPriority() const override { return EBaseNodeRefreshPriority::Low_UsesDependentWildcard; }
-	virtual void ReallocatePinsDuringReconstruction(TArray<UEdGraphPin*>& OldPins);
 	// End of UK2Node interface
 
-	// IK2Node_AddPinInterface interface
-	BLUEPRINTGRAPH_API virtual void AddInputPin() override;
-	// End of IK2Node_AddPinInterface interface
-
-protected:
-	friend class FKismetCompilerContext;
-
-	/** If needed, will clear all pins to be wildcards */
-	void ClearPinTypeToWildcard();
-
-	bool CanResetToWildcard() const;
-
-	/** Propagates the pin type from the output (array) pin to the inputs, to make sure types are consistent */
-	void PropagatePinType();
-
-	/** Returns the function to be used to clear the array */
-	BLUEPRINTGRAPH_API UFunction* GetArrayClearFunction() const;
-
-	/** Returns the function to be used to add a function to the array */
-	BLUEPRINTGRAPH_API UFunction* GetArrayAddFunction() const;
-
-	void SyncPinNames();
+	// UK2Node_MakeContainer interface
+	virtual FString GetOutputPinName() const override;
+	// UK2Node_MakeContainer interface
 };
diff --git a/Engine/Source/Editor/BlueprintGraph/Classes/K2Node_MakeContainer.h b/Engine/Source/Editor/BlueprintGraph/Classes/K2Node_MakeContainer.h
new file mode 100644
--- /dev/null
+++ b/Engine/Source/Editor/BlueprintGraph/Classes/K2Node_MakeContainer.h
@@ -0,0 +1,86 @@
+// Copyright 1998-2017 Epic Games, Inc. All Rights Reserved.
+
+
+#pragma once
+
+#include "CoreMinimal.h"
+#include "UObject/ObjectMacros.h"
+#include "Textures/SlateIcon.h"
+#include "K2Node.h"
+#include "K2Node_AddPinInterface.h"
+#include "KismetCompilerMisc.h"
+#include "K2Node_MakeContainer.generated.h"
+
+class FBlueprintActionDatabaseRegistrar;
+class UEdGraphPin;
+
+class FKCHandler_MakeContainer : public FNodeHandlingFunctor
+{
+public:
+	FKCHandler_MakeContainer(FKismetCompilerContext& InCompilerContext)
+		: FNodeHandlingFunctor(InCompilerContext)
+	{
+	}
+
+	virtual void RegisterNets(FKismetFunctionContext& Context, UEdGraphNode* Node) override;
+	virtual void Compile(FKismetFunctionContext& Context, UEdGraphNode* Node) override;
+
+protected:
+	EKismetCompiledStatementType CompiledStatementType;
+};
+
+
+UCLASS(MinimalAPI, abstract)
+class UK2Node_MakeContainer : public UK2Node, public IK2Node_AddPinInterface
+{
+	GENERATED_UCLASS_BODY()
+
+	/** The number of input pins to generate for this node */
+	UPROPERTY()
+	int32 NumInputs;
+
+public:
+	BLUEPRINTGRAPH_API void RemoveInputPin(UEdGraphPin* Pin);
+
+	BLUEPRINTGRAPH_API UEdGraphPin* GetOutputPin() const;
+
+	/** returns a reference to the output array pin of this node, which is responsible for defining the type */
+	virtual FString GetOutputPinName() const PURE_VIRTUAL(UK2Node_MakeContainer::GetOutputPinName, return FString(););
+	virtual FString GetPinName(int32 PinIndex) const;
+	virtual void GetKeyAndValuePins(TArray<UEdGraphPin*>& KeyPins, TArray<UEdGraphPin*>& ValuePins) const;
+
+public:
+	// UEdGraphNode interface
+	virtual void AllocateDefaultPins() override;
+	virtual void PostReconstructNode() override;
+	// End of UEdGraphNode interface
+
+	// UK2Node interface
+	virtual bool IsNodePure() const override { return true; }
+	virtual void NotifyPinConnectionListChanged(UEdGraphPin* Pin) override;
+	virtual class FNodeHandlingFunctor* CreateNodeHandler(class FKismetCompilerContext& CompilerContext) const override PURE_VIRTUAL(UK2Node_MakeContainer::CreateNodeHandler, return nullptr;);
+	virtual void GetMenuActions(FBlueprintActionDatabaseRegistrar& ActionRegistrar) const override;
+	virtual int32 GetNodeRefreshPriority() const override { return EBaseNodeRefreshPriority::Low_UsesDependentWildcard; }
+	virtual void ReallocatePinsDuringReconstruction(TArray<UEdGraphPin*>& OldPins);
+	virtual bool IsConnectionDisallowed(const UEdGraphPin* MyPin, const UEdGraphPin* OtherPin, FString& OutReason) const override;
+	// End of UK2Node interface
+
+	// IK2Node_AddPinInterface interface
+	BLUEPRINTGRAPH_API virtual void AddInputPin() override;
+	// End of IK2Node_AddPinInterface interface
+
+protected:
+	friend class FKismetCompilerContext;
+
+	/** If needed, will clear all pins to be wildcards */
+	void ClearPinTypeToWildcard();
+
+	bool CanResetToWildcard() const;
+
+	/** Propagates the pin type from the output (set) pin to the inputs, to make sure types are consistent */
+	void PropagatePinType();
+
+	void SyncPinNames();
+
+	EPinContainerType ContainerType;
+};
diff --git a/Engine/Source/Editor/BlueprintGraph/Classes/K2Node_MakeMap.h b/Engine/Source/Editor/BlueprintGraph/Classes/K2Node_MakeMap.h
new file mode 100644
--- /dev/null
+++ b/Engine/Source/Editor/BlueprintGraph/Classes/K2Node_MakeMap.h
@@ -0,0 +1,40 @@
+// Copyright 1998-2017 Epic Games, Inc. All Rights Reserved.
+
+
+#pragma once
+
+#include "K2Node_MakeContainer.h"
+#include "K2Node_MakeMap.generated.h"
+
+class UEdGraphPin;
+
+UCLASS(MinimalAPI)
+class UK2Node_MakeMap : public UK2Node_MakeContainer
+{
+	GENERATED_UCLASS_BODY()
+
+public:
+	// UEdGraphNode interface
+	virtual void AllocateDefaultPins() override;
+	virtual FText GetNodeTitle(ENodeTitleType::Type TitleType) const override;
+	virtual FText GetTooltipText() const override;
+	virtual FSlateIcon GetIconAndTint(FLinearColor& OutColor) const override;
+	virtual void ValidateNodeDuringCompilation(class FCompilerResultsLog& MessageLog) const override;
+	// End of UEdGraphNode interface
+
+	// UK2Node interface
+	virtual void GetContextMenuActions(const FGraphNodeContextMenuBuilder& Context) const override;
+	virtual class FNodeHandlingFunctor* CreateNodeHandler(class FKismetCompilerContext& CompilerContext) const override;
+	virtual FText GetMenuCategory() const override;
+	// End of UK2Node interface
+
+	// UK2Node_MakeContainer interface
+	virtual FString GetOutputPinName() const override;
+	virtual FString GetPinName(int32 PinIndex) const override;
+	virtual void GetKeyAndValuePins(TArray<UEdGraphPin*>& KeyPins, TArray<UEdGraphPin*>& ValuePins) const override;
+	// UK2Node_MakeContainer interface
+
+	// IK2Node_AddPinInterface interface
+	BLUEPRINTGRAPH_API virtual void AddInputPin() override;
+	// End of IK2Node_AddPinInterface interface
+};
diff --git a/Engine/Source/Editor/BlueprintGraph/Classes/K2Node_Variable.h b/Engine/Source/Editor/BlueprintGraph/Classes/K2Node_Variable.h
--- a/Engine/Source/Editor/BlueprintGraph/Classes/K2Node_Variable.h
+++ b/Engine/Source/Editor/BlueprintGraph/Classes/K2Node_Variable.h
@@ -138,6 +138,11 @@ public:
 	UProperty* GetPropertyForVariable() const;
 	UProperty* GetPropertyForVariableFromSkeleton() const;
 
+private:
+	UProperty* GetPropertyForVariable_Internal(UClass* OwningClass) const;
+
+public:
+
 	/** Accessor for the value output pin of the node */
 	UEdGraphPin* GetValuePin() const;
 
diff --git a/Engine/Source/Editor/BlueprintGraph/Private/K2Node.cpp b/Engine/Source/Editor/BlueprintGraph/Private/K2Node.cpp
--- a/Engine/Source/Editor/BlueprintGraph/Private/K2Node.cpp
+++ b/Engine/Source/Editor/BlueprintGraph/Private/K2Node.cpp
@@ -2,6 +2,7 @@
 
 
 #include "K2Node.h"
+#include "BlueprintCompilationManager.h"
 #include "UObject/UnrealType.h"
 #include "UObject/CoreRedirects.h"
 #include "EdGraph/EdGraphPin.h"
@@ -210,6 +211,9 @@ bool UK2Node::CreatePinsForFunctionEntryExit(const UFunction* Function, bool bFo
 {
 	const UEdGraphSchema_K2* K2Schema = GetDefault<UEdGraphSchema_K2>();
 
+	// if the generated class is not up to date, use the skeleton's class function to create pins:
+	Function = FBlueprintEditorUtils::GetMostUpToDateFunction(Function);
+
 	// Create the inputs and outputs
 	bool bAllPinsGood = true;
 	for (TFieldIterator<UProperty> PropIt(Function); PropIt && (PropIt->PropertyFlags & CPF_Parm); ++PropIt)
@@ -363,10 +367,6 @@ void UK2Node::GetNodeAttributes( TArray<TKeyValuePair<FString, FString>>& OutNod
 	OutNodeAttributes.Add( TKeyValuePair<FString, FString>( TEXT( "Name" ), GetName() ));
 }
 
-void UK2Node::PreloadRequiredAssets()
-{
-}
-
 void UK2Node::PinConnectionListChanged(UEdGraphPin* Pin) 
 {
 
@@ -729,10 +729,6 @@ UK2Node::ERedirectType UK2Node::DoPinsMatchForReconstruction(const UEdGraphPin*
 	return RedirectType;
 }
 
-void UK2Node::CustomMapParamValue(UEdGraphPin& Pin)
-{
-}
-
 void UK2Node::ReconstructSinglePin(UEdGraphPin* NewPin, UEdGraphPin* OldPin, ERedirectType RedirectType)
 {
 	DECLARE_SCOPE_CYCLE_COUNTER(TEXT("UK2Node::ReconstructSinglePin"), STAT_LinkerLoad_ReconstructSinglePin, STATGROUP_LoadTimeVerbose);
@@ -766,12 +762,6 @@ void UK2Node::ReconstructSinglePin(UEdGraphPin* NewPin, UEdGraphPin* OldPin, ERe
 				}
 			}
 		}
-		else if (RedirectType == ERedirectType_Custom)
-		{
-			PRAGMA_DISABLE_DEPRECATION_WARNINGS
-			CustomMapParamValue(*NewPin);
-			PRAGMA_ENABLE_DEPRECATION_WARNINGS
-		}
 	}
 
 	// Update the blueprints watched pins as the old pin will be going the way of the dodo
diff --git a/Engine/Source/Editor/BlueprintGraph/Private/K2Node_ActorBoundEvent.cpp b/Engine/Source/Editor/BlueprintGraph/Private/K2Node_ActorBoundEvent.cpp
--- a/Engine/Source/Editor/BlueprintGraph/Private/K2Node_ActorBoundEvent.cpp
+++ b/Engine/Source/Editor/BlueprintGraph/Private/K2Node_ActorBoundEvent.cpp
@@ -50,8 +50,26 @@ FNodeHandlingFunctor* UK2Node_ActorBoundEvent::CreateNodeHandler(FKismetCompiler
 
 void UK2Node_ActorBoundEvent::ReconstructNode()
 {
-	// Ensure that we try to update the delegate we're bound to if its moved
-	GetTargetDelegate();
+	// We need to fixup our event reference as it may have changed or been redirected
+	UMulticastDelegateProperty* TargetDelegateProp = GetTargetDelegateProperty();
+
+	// If we couldn't find the target delegate, then try to find it in the property remap table
+	if (!TargetDelegateProp)
+	{
+		UMulticastDelegateProperty* NewProperty = FMemberReference::FindRemappedField<UMulticastDelegateProperty>(DelegateOwnerClass, DelegatePropertyName);
+		if (NewProperty)
+		{
+			// Found a remapped property, update the node
+			TargetDelegateProp = NewProperty;
+			DelegatePropertyName = NewProperty->GetFName();
+			CachedNodeTitle.MarkDirty();
+		}
+	}
+
+	if (TargetDelegateProp && TargetDelegateProp->SignatureFunction)
+	{
+		EventReference.SetFromField<UFunction>(TargetDelegateProp->SignatureFunction, false);
+	}
 
 	CachedNodeTitle.MarkDirty();
 
@@ -133,7 +151,7 @@ FText UK2Node_ActorBoundEvent::GetNodeTitle(ENodeTitleType::Type TitleType) cons
 
 FText UK2Node_ActorBoundEvent::GetTooltipText() const
 {
-	UMulticastDelegateProperty* TargetDelegateProp = GetTargetDelegatePropertyConst();
+	UMulticastDelegateProperty* TargetDelegateProp = GetTargetDelegateProperty();
 	if(TargetDelegateProp)
 	{
 		return TargetDelegateProp->GetToolTipText();
@@ -180,32 +198,12 @@ void UK2Node_ActorBoundEvent::InitializeActorBoundEventParams(AActor* InEventOwn
 	}
 }
 
-UMulticastDelegateProperty* UK2Node_ActorBoundEvent::GetTargetDelegatePropertyConst() const
+UMulticastDelegateProperty* UK2Node_ActorBoundEvent::GetTargetDelegateProperty() const
 {
 	return Cast<UMulticastDelegateProperty>(FindField<UMulticastDelegateProperty>(DelegateOwnerClass, DelegatePropertyName));
 }
 
-UMulticastDelegateProperty* UK2Node_ActorBoundEvent::GetTargetDelegateProperty()
-{
-	UMulticastDelegateProperty* TargetDelegateProp = Cast<UMulticastDelegateProperty>(FindField<UMulticastDelegateProperty>(DelegateOwnerClass, DelegatePropertyName));
-
-	// If we couldn't find the target delegate, then try to find it in the property remap table
-	if (!TargetDelegateProp)
-	{
-		UMulticastDelegateProperty* NewProperty = FMemberReference::FindRemappedField<UMulticastDelegateProperty>(DelegateOwnerClass, DelegatePropertyName);
-		if (NewProperty)
-		{
-			// Found a remapped property, update the node
-			TargetDelegateProp = NewProperty;
-			DelegatePropertyName = NewProperty->GetFName();
-			CachedNodeTitle.MarkDirty();
-		}
-	}
-
-	return TargetDelegateProp;
-}
-
-FMulticastScriptDelegate* UK2Node_ActorBoundEvent::GetTargetDelegate()
+FMulticastScriptDelegate* UK2Node_ActorBoundEvent::GetTargetDelegate() const 
 {
 	if( EventOwner )
 	{
@@ -234,10 +232,7 @@ void UK2Node_ActorBoundEvent::Serialize(FArchive& Ar)
 	{
 		if(Ar.UE4Ver() < VER_UE4_K2NODE_EVENT_MEMBER_REFERENCE)
 		{
-			DelegateOwnerClass = EventSignatureClass_DEPRECATED;
-			UMulticastDelegateProperty* TargetDelegateProp = GetTargetDelegateProperty();
-			check(TargetDelegateProp);
-			EventReference.SetFromField<UFunction>(TargetDelegateProp->SignatureFunction, false);
+			DelegateOwnerClass = EventSignatureClass_DEPRECATED;	
 		}
 	}
 }
diff --git a/Engine/Source/Editor/BlueprintGraph/Private/K2Node_BreakStruct.cpp b/Engine/Source/Editor/BlueprintGraph/Private/K2Node_BreakStruct.cpp
--- a/Engine/Source/Editor/BlueprintGraph/Private/K2Node_BreakStruct.cpp
+++ b/Engine/Source/Editor/BlueprintGraph/Private/K2Node_BreakStruct.cpp
@@ -208,6 +208,12 @@ void UK2Node_BreakStruct::AllocateDefaultPins()
 	}
 }
 
+void UK2Node_BreakStruct::PreloadRequiredAssets()
+{
+	PreloadObject(StructType);
+	Super::PreloadRequiredAssets();
+}
+
 FText UK2Node_BreakStruct::GetNodeTitle(ENodeTitleType::Type TitleType) const
 {
 	if (StructType == nullptr)
diff --git a/Engine/Source/Editor/BlueprintGraph/Private/K2Node_CallFunction.cpp b/Engine/Source/Editor/BlueprintGraph/Private/K2Node_CallFunction.cpp
--- a/Engine/Source/Editor/BlueprintGraph/Private/K2Node_CallFunction.cpp
+++ b/Engine/Source/Editor/BlueprintGraph/Private/K2Node_CallFunction.cpp
@@ -1009,7 +1009,7 @@ void UK2Node_CallFunction::PostReconstructNode()
 			UClass* FunctionOwnerClass = Function->GetOuterUClass();
 			if (!BP->SkeletonGeneratedClass->IsChildOf(FunctionOwnerClass))
 			{
-				SelfPin->DefaultObject = FunctionOwnerClass->GetDefaultObject();
+				SelfPin->DefaultObject = FunctionOwnerClass->GetAuthoritativeClass()->GetDefaultObject();
 			}
 			else
 			{
diff --git a/Engine/Source/Editor/BlueprintGraph/Private/K2Node_ComponentBoundEvent.cpp b/Engine/Source/Editor/BlueprintGraph/Private/K2Node_ComponentBoundEvent.cpp
--- a/Engine/Source/Editor/BlueprintGraph/Private/K2Node_ComponentBoundEvent.cpp
+++ b/Engine/Source/Editor/BlueprintGraph/Private/K2Node_ComponentBoundEvent.cpp
@@ -133,6 +133,34 @@ FString UK2Node_ComponentBoundEvent::GetDocumentationExcerptName() const
 	return DelegatePropertyName.ToString();
 }
 
+void UK2Node_ComponentBoundEvent::ReconstructNode()
+{
+	// We need to fixup our event reference as it may have changed or been redirected
+	UMulticastDelegateProperty* TargetDelegateProp = GetTargetDelegateProperty();
+
+	// If we couldn't find the target delegate, then try to find it in the property remap table
+	if (!TargetDelegateProp)
+	{
+		UMulticastDelegateProperty* NewProperty = FMemberReference::FindRemappedField<UMulticastDelegateProperty>(DelegateOwnerClass, DelegatePropertyName);
+		if (NewProperty)
+		{
+			// Found a remapped property, update the node
+			TargetDelegateProp = NewProperty;
+			DelegatePropertyName = NewProperty->GetFName();
+			CachedNodeTitle.MarkDirty();
+		}
+	}
+
+	if (TargetDelegateProp && TargetDelegateProp->SignatureFunction)
+	{
+		EventReference.SetFromField<UFunction>(TargetDelegateProp->SignatureFunction, false);
+	}
+
+	CachedNodeTitle.MarkDirty();
+
+	Super::ReconstructNode();
+}
+
 void UK2Node_ComponentBoundEvent::Serialize(FArchive& Ar)
 {
 	Super::Serialize(Ar);
@@ -140,17 +168,14 @@ void UK2Node_ComponentBoundEvent::Serialize(FArchive& Ar)
 	// Fix up legacy nodes that may not yet have a delegate pin
 	if(Ar.IsLoading())
 	{
-		bool bNeedsFixup = false;
 		if(Ar.UE4Ver() < VER_UE4_K2NODE_EVENT_MEMBER_REFERENCE)
 		{
 			DelegateOwnerClass = EventSignatureClass_DEPRECATED;
-			bNeedsFixup = true;
 		}
 
 		// Recover from the period where DelegateOwnerClass was transient
 		if (!DelegateOwnerClass && HasValidBlueprint())
 		{
-			bNeedsFixup = true;
 			// Search for a component property on the owning class, this should work in most cases
 			UBlueprint* ParentBlueprint = GetBlueprint();
 			UClass* ParentClass = ParentBlueprint ? ParentBlueprint->GeneratedClass : NULL;
@@ -168,27 +193,6 @@ void UK2Node_ComponentBoundEvent::Serialize(FArchive& Ar)
 				DelegateOwnerClass = ComponentProperty->PropertyClass;
 			}
 		}
-
-		if (bNeedsFixup)
-		{
-			// We need to fixup our event reference as it may have been saved incorrectly
-			UMulticastDelegateProperty* TargetDelegateProp = GetTargetDelegateProperty();
-			if (TargetDelegateProp && TargetDelegateProp->SignatureFunction)
-			{
-				FName ReferenceName = TargetDelegateProp->SignatureFunction->GetFName();
-				UClass* ReferenceClass = TargetDelegateProp->SignatureFunction->GetOwnerClass();
-
-				if (EventReference.GetMemberName() != ReferenceName || EventReference.GetMemberParentClass() != ReferenceClass)
-				{
-					// Set the reference if it wasn't already set properly, owner class may end up being NULL for native delegates
-					EventReference.SetExternalMember(ReferenceName, ReferenceClass);
-				}
-			}
-			else
-			{
-				UE_LOG(LogBlueprint, Warning, TEXT("Loaded invalid component bound event in node %s."), *GetPathName());
-			}
-		}
 	}
 }
 
diff --git a/Engine/Source/Editor/BlueprintGraph/Private/K2Node_DynamicCast.cpp b/Engine/Source/Editor/BlueprintGraph/Private/K2Node_DynamicCast.cpp
--- a/Engine/Source/Editor/BlueprintGraph/Private/K2Node_DynamicCast.cpp
+++ b/Engine/Source/Editor/BlueprintGraph/Private/K2Node_DynamicCast.cpp
@@ -121,44 +121,47 @@ void UK2Node_DynamicCast::GetContextMenuActions(const FGraphNodeContextMenuBuild
 {
 	Super::GetContextMenuActions(Context);
 
-	Context.MenuBuilder->BeginSection("K2NodeDynamicCast", LOCTEXT("DynamicCastHeader", "Cast"));
+	if (!Context.bIsDebugging)
 	{
-		FText MenuEntryTitle   = LOCTEXT("MakePureTitle",   "Convert to pure cast");
-		FText MenuEntryTooltip = LOCTEXT("MakePureTooltip", "Removes the execution pins to make the node more versatile (NOTE: the cast could still fail, resulting in an invalid output).");
-
-		bool bCanTogglePurity = true;
-		auto CanExecutePurityToggle = [](bool const bInCanTogglePurity)->bool
+		Context.MenuBuilder->BeginSection("K2NodeDynamicCast", LOCTEXT("DynamicCastHeader", "Cast"));
 		{
-			return bInCanTogglePurity;
-		};
+			FText MenuEntryTitle = LOCTEXT("MakePureTitle", "Convert to pure cast");
+			FText MenuEntryTooltip = LOCTEXT("MakePureTooltip", "Removes the execution pins to make the node more versatile (NOTE: the cast could still fail, resulting in an invalid output).");
 
-		if (bIsPureCast)
-		{
-			MenuEntryTitle   = LOCTEXT("MakeImpureTitle",   "Convert to impure cast");
-			MenuEntryTooltip = LOCTEXT("MakeImpureTooltip", "Adds in branching execution pins so that you can separatly handle when the cast fails/succeeds.");
-		
-			const UEdGraphSchema_K2* K2Schema = Cast<UEdGraphSchema_K2>(GetSchema());
-			check(K2Schema != nullptr);
+			bool bCanTogglePurity = true;
+			auto CanExecutePurityToggle = [](bool const bInCanTogglePurity)->bool
+			{
+				return bInCanTogglePurity;
+			};
 
-			bCanTogglePurity = K2Schema->DoesGraphSupportImpureFunctions(GetGraph());
-			if (!bCanTogglePurity)
+			if (bIsPureCast)
 			{
-				MenuEntryTooltip = LOCTEXT("CannotMakeImpureTooltip", "This graph does not support impure calls (and you should therefore test the cast's result for validity).");
+				MenuEntryTitle = LOCTEXT("MakeImpureTitle", "Convert to impure cast");
+				MenuEntryTooltip = LOCTEXT("MakeImpureTooltip", "Adds in branching execution pins so that you can separatly handle when the cast fails/succeeds.");
+
+				const UEdGraphSchema_K2* K2Schema = Cast<UEdGraphSchema_K2>(GetSchema());
+				check(K2Schema != nullptr);
+
+				bCanTogglePurity = K2Schema->DoesGraphSupportImpureFunctions(GetGraph());
+				if (!bCanTogglePurity)
+				{
+					MenuEntryTooltip = LOCTEXT("CannotMakeImpureTooltip", "This graph does not support impure calls (and you should therefore test the cast's result for validity).");
+				}
 			}
+
+			Context.MenuBuilder->AddMenuEntry(
+				MenuEntryTitle,
+				MenuEntryTooltip,
+				FSlateIcon(),
+				FUIAction(
+					FExecuteAction::CreateUObject(this, &UK2Node_DynamicCast::TogglePurity),
+					FCanExecuteAction::CreateStatic(CanExecutePurityToggle, bCanTogglePurity),
+					FIsActionChecked()
+				)
+			);
 		}
-		
-		Context.MenuBuilder->AddMenuEntry(
-			MenuEntryTitle,
-			MenuEntryTooltip,
-			FSlateIcon(),
-			FUIAction(
-				FExecuteAction::CreateUObject(this, &UK2Node_DynamicCast::TogglePurity),
-				FCanExecuteAction::CreateStatic(CanExecutePurityToggle, bCanTogglePurity),
-				FIsActionChecked()
-			)
-		);
+		Context.MenuBuilder->EndSection();
 	}
-	Context.MenuBuilder->EndSection();
 }
 
 void UK2Node_DynamicCast::PostReconstructNode()
diff --git a/Engine/Source/Editor/BlueprintGraph/Private/K2Node_GetDataTableRow.cpp b/Engine/Source/Editor/BlueprintGraph/Private/K2Node_GetDataTableRow.cpp
--- a/Engine/Source/Editor/BlueprintGraph/Private/K2Node_GetDataTableRow.cpp
+++ b/Engine/Source/Editor/BlueprintGraph/Private/K2Node_GetDataTableRow.cpp
@@ -452,6 +452,20 @@ void UK2Node_GetDataTableRow::EarlyValidation(class FCompilerResultsLog& Message
 	}	
 }
 
+void UK2Node_GetDataTableRow::PreloadRequiredAssets()
+{
+	if (UEdGraphPin* DataTablePin = GetDataTablePin())
+	{
+		if (UDataTable* DataTable = Cast<UDataTable>(DataTablePin->DefaultObject))
+		{
+			// make sure to properly load the data-table object so that we can 
+			// farm the "RowStruct" property from it (below, in GetDataTableRowStructType)
+			PreloadObject(DataTable);
+		}
+	}
+	return Super::PreloadRequiredAssets();
+}
+
 void UK2Node_GetDataTableRow::NotifyPinConnectionListChanged(UEdGraphPin* Pin)
 {
 	Super::NotifyPinConnectionListChanged(Pin);
diff --git a/Engine/Source/Editor/BlueprintGraph/Private/K2Node_MacroInstance.cpp b/Engine/Source/Editor/BlueprintGraph/Private/K2Node_MacroInstance.cpp
--- a/Engine/Source/Editor/BlueprintGraph/Private/K2Node_MacroInstance.cpp
+++ b/Engine/Source/Editor/BlueprintGraph/Private/K2Node_MacroInstance.cpp
@@ -118,6 +118,14 @@ void UK2Node_MacroInstance::AllocateDefaultPins()
 	}
 }
 
+void UK2Node_MacroInstance::PreloadRequiredAssets()
+{
+	PreloadObject(MacroGraphReference.GetBlueprint());
+	UEdGraph* MacroGraph = MacroGraphReference.GetGraph();
+	PreloadObject(MacroGraph);
+	Super::PreloadRequiredAssets();
+}
+
 FText UK2Node_MacroInstance::GetTooltipText() const
 {
 	UEdGraph* MacroGraph = MacroGraphReference.GetGraph();
diff --git a/Engine/Source/Editor/BlueprintGraph/Private/K2Node_MakeArray.cpp b/Engine/Source/Editor/BlueprintGraph/Private/K2Node_MakeArray.cpp
--- a/Engine/Source/Editor/BlueprintGraph/Private/K2Node_MakeArray.cpp
+++ b/Engine/Source/Editor/BlueprintGraph/Private/K2Node_MakeArray.cpp
@@ -10,7 +10,6 @@
 #include "EdGraph/EdGraphNodeUtils.h"
 #include "Kismet2/BlueprintEditorUtils.h"
 
-#include "Kismet/KismetArrayLibrary.h"
 #include "ScopedTransaction.h"
 #include "EdGraphUtilities.h"
 #include "KismetCompiledFunctionContext.h"
@@ -19,57 +18,22 @@
 #include "EditorCategoryUtils.h"
 #include "BlueprintActionDatabaseRegistrar.h"
 
-static const FString OutputPinName = FString(TEXT("Array"));
+namespace MakeArrayLiterals
+{
+	static const FString OutputPinName = FString(TEXT("Array"));
+};
 
 #define LOCTEXT_NAMESPACE "MakeArrayNode"
 
 /////////////////////////////////////////////////////
 // FKCHandler_MakeArray
-class FKCHandler_MakeArray : public FNodeHandlingFunctor
+class FKCHandler_MakeArray : public FKCHandler_MakeContainer
 {
 public:
 	FKCHandler_MakeArray(FKismetCompilerContext& InCompilerContext)
-		: FNodeHandlingFunctor(InCompilerContext)
-	{
-	}
-
-	virtual void RegisterNets(FKismetFunctionContext& Context, UEdGraphNode* Node) override
+		: FKCHandler_MakeContainer(InCompilerContext)
 	{
- 		UK2Node_MakeArray* ArrayNode = CastChecked<UK2Node_MakeArray>(Node);
- 		UEdGraphPin* OutputPin = ArrayNode->GetOutputPin();
-
-		FNodeHandlingFunctor::RegisterNets(Context, Node);
-
-		// Create a local term to drop the array into
-		FBPTerminal* Term = Context.CreateLocalTerminalFromPinAutoChooseScope(OutputPin, Context.NetNameMap->MakeValidName(OutputPin));
-		Term->bPassedByReference = false;
-		Term->Source = Node;
- 		Context.NetMap.Add(OutputPin, Term);
-	}
-
-	virtual void Compile(FKismetFunctionContext& Context, UEdGraphNode* Node) override
-	{
-		UK2Node_MakeArray* ArrayNode = CastChecked<UK2Node_MakeArray>(Node);
-		UEdGraphPin* OutputPin = ArrayNode->GetOutputPin();
-
-		FBPTerminal** ArrayTerm = Context.NetMap.Find(OutputPin);
-		check(ArrayTerm);
-
-		FBlueprintCompiledStatement& CreateArrayStatement = Context.AppendStatementForNode(Node);
-		CreateArrayStatement.Type = KCST_CreateArray;
-		CreateArrayStatement.LHS = *ArrayTerm;
-
-		for (UEdGraphPin* Pin : Node->Pins)
-		{
-			if(Pin && Pin->Direction == EGPD_Input)
-			{
-				FBPTerminal** InputTerm = Context.NetMap.Find(FEdGraphUtilities::GetNetFromPin(Pin));
-				if( InputTerm )
-				{
-					CreateArrayStatement.RHS.Add(*InputTerm);
-				}
-			}
-		}
+		CompiledStatementType = KCST_CreateArray;
 	}
 };
 
@@ -79,8 +43,7 @@ public:
 UK2Node_MakeArray::UK2Node_MakeArray(const FObjectInitializer& ObjectInitializer)
 	: Super(ObjectInitializer)
 {
-
-	NumInputs = 1;
+	ContainerType = EPinContainerType::Array;
 }
 
 FNodeHandlingFunctor* UK2Node_MakeArray::CreateNodeHandler(FKismetCompilerContext& CompilerContext) const
@@ -93,284 +56,9 @@ FText UK2Node_MakeArray::GetNodeTitle(ENodeTitleType::Type TitleType) const
 	return LOCTEXT("NodeTitle", "Make Array");
 }
 
-UEdGraphPin* UK2Node_MakeArray::GetOutputPin() const
-{
-	return FindPin(OutputPinName);
-}
-
-void UK2Node_MakeArray::ReallocatePinsDuringReconstruction(TArray<UEdGraphPin*>& OldPins)
-{
-	Super::ReallocatePinsDuringReconstruction(OldPins);
-
-	// This is necessary to retain type information after pasting or loading from disc
-	UEdGraphPin* OutputPin = GetOutputPin();
-	if (OutputPin)
-	{
-		// Only update the output pin if it is currently a wildcard
-		if (OutputPin->PinType.PinCategory == UEdGraphSchema_K2::PC_Wildcard)
-		{
-			// Find the matching Old Pin if it exists
-			for (UEdGraphPin* OldPin : OldPins)
-			{
-				if (OldPin->Direction == EGPD_Output)
-				{
-					// Update our output pin with the old type information and then propagate it to our input pins
-					OutputPin->PinType = OldPin->PinType;
-					PropagatePinType();
-					break;
-				}
-			}
-		}
-	}
-}
-
-void UK2Node_MakeArray::AllocateDefaultPins()
-{
-	// Create the output pin
-	CreatePin(EGPD_Output, UEdGraphSchema_K2::PC_Wildcard, FString(), nullptr, *OutputPinName, EPinContainerType::Array);
-
-	// Create the input pins to create the arrays from
-	for (int32 i = 0; i < NumInputs; ++i)
-	{
-		CreatePin(EGPD_Input, UEdGraphSchema_K2::PC_Wildcard, FString(), nullptr, *FString::Printf(TEXT("[%d]"), i));
-	}
-}
-
-bool UK2Node_MakeArray::CanResetToWildcard() const
-{
-	bool bClearPinsToWildcard = true;
-
-	// Check to see if we want to clear the wildcards.
-	for (const UEdGraphPin* Pin : Pins)
-	{
-		if( Pin->LinkedTo.Num() > 0 )
-		{
-			// One of the pins is still linked, we will not be clearing the types.
-			bClearPinsToWildcard = false;
-			break;
-		}
-	}
-
-	return bClearPinsToWildcard;
-}
-
-void UK2Node_MakeArray::ClearPinTypeToWildcard()
-{
-	if( CanResetToWildcard() )
-	{
-		UEdGraphPin* OutputPin = GetOutputPin();
-		OutputPin->PinType.PinCategory = UEdGraphSchema_K2::PC_Wildcard;
-		OutputPin->PinType.PinSubCategory.Reset();
-		OutputPin->PinType.PinSubCategoryObject = nullptr;
-
-		PropagatePinType();
-	}
-}
-
-void UK2Node_MakeArray::NotifyPinConnectionListChanged(UEdGraphPin* Pin)
-{
-	Super::NotifyPinConnectionListChanged(Pin);
-
-	// Was this the first or last connection?
-	int32 NumPinsWithLinks = 0;
-	// Array to cache the input pins we might want to find these if we are removing the last link
-	TArray< UEdGraphPin* > InputPins;
-	for (int32 PinIndex = 0; PinIndex < Pins.Num(); ++PinIndex)
-	{
-		NumPinsWithLinks += (Pins[PinIndex]->LinkedTo.Num() > 0) ? 1 : 0;
-		if( Pins[PinIndex]->Direction == EGPD_Input )
-		{
-			InputPins.Add(Pins[PinIndex]);
-		}
-	}
-
-	UEdGraphPin* OutputPin = GetOutputPin();
-
-	if (Pin->LinkedTo.Num() > 0)
-	{
-		// Just made a connection, was it the first and was it the parent pin?
-		if (Pin->ParentPin == nullptr && NumPinsWithLinks == 1 && (OutputPin->PinType.PinCategory == UEdGraphSchema_K2::PC_Wildcard || Pin->LinkedTo[0]->PinType.PinCategory != UEdGraphSchema_K2::PC_Wildcard))
-		{
-			// Update the types on all the pins
-			OutputPin->PinType = Pin->LinkedTo[0]->PinType;
-			OutputPin->PinType.ContainerType = EPinContainerType::Array;
-			PropagatePinType();
-			GetGraph()->NotifyGraphChanged();
-		}
-	}
-	else
-	{
-		// Just broke a connection, was it the last?
-		if (NumPinsWithLinks == 0)
-		{
-			// Return to wildcard if theres nothing in any of the input pins
-			bool bResetOutputPin = true;
-			for (int32 PinIndex = 0; PinIndex < InputPins.Num(); ++PinIndex)
-			{
-				UEdGraphPin*& CurrentPin = InputPins[PinIndex];
-				if (!CurrentPin->DoesDefaultValueMatchAutogenerated())
-				{
-					bResetOutputPin = false;
-				}
-			}
-
-			if( bResetOutputPin == true )
-			{
-				OutputPin->PinType.PinCategory = UEdGraphSchema_K2::PC_Wildcard;
-				OutputPin->PinType.PinSubCategory.Reset();
-				OutputPin->PinType.PinSubCategoryObject = nullptr;
-
-				PropagatePinType();
-				GetGraph()->NotifyGraphChanged();
-			}
-		}
-	}
-}
-
-void UK2Node_MakeArray::PropagatePinType()
-{
-	const UEdGraphPin* OutputPin = GetOutputPin();
-
-	if (OutputPin)
-	{
-		UClass const* CallingContext = NULL;
-		if (UBlueprint const* Blueprint = GetBlueprint())
-		{
-			CallingContext = Blueprint->GeneratedClass;
-			if (CallingContext == NULL)
-			{
-				CallingContext = Blueprint->ParentClass;
-			}
-		}
-
-		const UEdGraphSchema_K2* Schema = GetDefault<UEdGraphSchema_K2>();
-		bool bWantRefresh = false;
-		// Propagate pin type info (except for array info!) to pins with dependent types
-		for (TArray<UEdGraphPin*>::TIterator it(Pins); it; ++it)
-		{
-			UEdGraphPin* CurrentPin = *it;
-			check(CurrentPin);
-
-			if (CurrentPin != OutputPin)
-			{
-				// sub pins will be updated by their parent pin, so if we have a parent pin just do nothing
-				if (CurrentPin->ParentPin == nullptr)
-				{
-					bWantRefresh = true;
-
-					// if we've reset to wild card or the parentpin no longer matches we need to collapse the split pin(s)
-					// otherwise everything should be OK:
-					if (CurrentPin->SubPins.Num() != 0 &&
-						(	CurrentPin->PinType.PinCategory != OutputPin->PinType.PinCategory ||
-							CurrentPin->PinType.PinSubCategory != OutputPin->PinType.PinSubCategory ||
-							CurrentPin->PinType.PinSubCategoryObject != OutputPin->PinType.PinSubCategoryObject )
-						)
-					{
-						// this is a little dicey, but should be fine.. relies on the fact that RecombinePin will only remove pins that
-						// are placed after CurrentPin in the Pins member:
-						Schema->RecombinePin(CurrentPin->SubPins[0]);
-					}
-
-					CurrentPin->PinType.PinCategory = OutputPin->PinType.PinCategory;
-					CurrentPin->PinType.PinSubCategory = OutputPin->PinType.PinSubCategory;
-					CurrentPin->PinType.PinSubCategoryObject = OutputPin->PinType.PinSubCategoryObject;
-				}
-
-				if (CurrentPin->PinType.PinCategory == UEdGraphSchema_K2::PC_Wildcard || CurrentPin->GetDefaultAsString().IsEmpty() == true)
-				{
-					// Only reset default value if there isn't one set or it is a wildcard. Otherwise this deletes data!
-					Schema->SetPinAutogeneratedDefaultValueBasedOnType(CurrentPin);
-				}
-
-				// Verify that all previous connections to this pin are still valid with the new type
-				for (TArray<UEdGraphPin*>::TIterator ConnectionIt(CurrentPin->LinkedTo); ConnectionIt; ++ConnectionIt)
-				{
-					UEdGraphPin* ConnectedPin = *ConnectionIt;
-					if (!Schema->ArePinsCompatible(CurrentPin, ConnectedPin, CallingContext))
-					{
-						CurrentPin->BreakLinkTo(ConnectedPin);
-					}
-				}
-			}
-		}
-		// If we have a valid graph we should refresh it now to refelect any changes we made
-		if( (bWantRefresh == true ) && ( OutputPin->GetOwningNode() != NULL ) && ( OutputPin->GetOwningNode()->GetGraph() != NULL ) )
-		{
-			OutputPin->GetOwningNode()->GetGraph()->NotifyGraphChanged();
-		}
-	}
-}
-
-UFunction* UK2Node_MakeArray::GetArrayClearFunction() const
-{
-	UClass* ArrayLibClass = UKismetArrayLibrary::StaticClass();
-	UFunction* ReturnFunction = ArrayLibClass->FindFunctionByName(FName(TEXT("Array_Clear")));
-	check(ReturnFunction);
-	return ReturnFunction;
-}
-
-UFunction* UK2Node_MakeArray::GetArrayAddFunction() const
+FString UK2Node_MakeArray::GetOutputPinName() const
 {
-	UClass* ArrayLibClass = UKismetArrayLibrary::StaticClass();
-	UFunction* ReturnFunction = ArrayLibClass->FindFunctionByName(FName(TEXT("Array_Add")));
-	check(ReturnFunction);
-	return ReturnFunction;
-}
-
-void UK2Node_MakeArray::PostReconstructNode()
-{
-	// Find a pin that has connections to use to jumpstart the wildcard process
-	bool bResetToWildcard = true;
-	for (int32 PinIndex = 0; PinIndex < Pins.Num(); ++PinIndex)
-	{
-		if (Pins[PinIndex]->LinkedTo.Num() > 0)
-		{
-			// The pin is linked, continue to use its type as the type for all pins.
-
-			// Update the types on all the pins
-			UEdGraphPin* OutputPin = GetOutputPin();
-			if (Pins[PinIndex]->ParentPin != nullptr)
-			{ 
-				// This is a split pin, so we want to base the pin type on the parent rather than the linked to pin
-				OutputPin->PinType = Pins[PinIndex]->ParentPin->PinType;
-			}
-			else
-			{
-				OutputPin->PinType = Pins[PinIndex]->LinkedTo[0]->PinType;
-			}
-			OutputPin->PinType.ContainerType = EPinContainerType::Array;
-			PropagatePinType();
-			bResetToWildcard = false;
-			break;
-		}
-		else if(!Pins[PinIndex]->DoesDefaultValueMatchAutogenerated())
-		{
-			// The pin has user data in it, continue to use its type as the type for all pins.
-
-			// Update the types on all the pins
-			UEdGraphPin* OutputPin = GetOutputPin();
-			if (Pins[PinIndex]->ParentPin != nullptr)
-			{
-				// This is a split pin, so we want to base the pin type on the parent instead of this pin
-				OutputPin->PinType = Pins[PinIndex]->ParentPin->PinType;
-			}
-			else
-			{
-				OutputPin->PinType = Pins[PinIndex]->PinType;
-			}
-			OutputPin->PinType.ContainerType = EPinContainerType::Array;
-			PropagatePinType();
-			bResetToWildcard = false;
-			break;
-		}
-	}
-
-	if (bResetToWildcard)
-	{
-		ClearPinTypeToWildcard();
-	}
-
-	Super::PostReconstructNode();
+	return MakeArrayLiterals::OutputPinName;
 }
 
 FText UK2Node_MakeArray::GetTooltipText() const
@@ -384,92 +72,6 @@ FSlateIcon UK2Node_MakeArray::GetIconAndTint(FLinearColor& OutColor) const
 	return Icon;
 }
 
-void UK2Node_MakeArray::AddInputPin()
-{
-	Modify();
-
-	++NumInputs;
-	FEdGraphPinType OutputPinType = GetOutputPin()->PinType;
-	UEdGraphPin* Pin = CreatePin(EGPD_Input, OutputPinType.PinCategory, OutputPinType.PinSubCategory, OutputPinType.PinSubCategoryObject.Get(), *FString::Printf(TEXT("[%d]"), (NumInputs-1)));
-	GetDefault<UEdGraphSchema_K2>()->SetPinAutogeneratedDefaultValueBasedOnType(Pin);
-
-	const bool bIsCompiling = GetBlueprint()->bBeingCompiled;
-	if( !bIsCompiling )
-	{
-		FBlueprintEditorUtils::MarkBlueprintAsStructurallyModified(GetBlueprint());
-	}
-}
-
-void UK2Node_MakeArray::SyncPinNames()
-{
-	int32 CurrentNumParentPins = 0;
-	for (int32 PinIndex = 0; PinIndex < Pins.Num(); ++PinIndex)
-	{
-		UEdGraphPin*& CurrentPin = Pins[PinIndex];
-		if (CurrentPin->Direction == EGPD_Input &&
-			CurrentPin->ParentPin == nullptr)
-		{
-			FString OldName = CurrentPin->PinName;
-			FString ElementName = FString::Printf(TEXT("[%d]"), CurrentNumParentPins++);
-
-			CurrentPin->Modify();
-			CurrentPin->PinName = ElementName;
-
-			if (CurrentPin->SubPins.Num() > 0)
-			{
-				FString OldFriendlyName = OldName;
-				FString ElementFriendlyName = ElementName;
-
-				// SubPin Friendly Name has an extra space in it so we need to account for that
-				OldFriendlyName.InsertAt(1, " ");
-				ElementFriendlyName.InsertAt(1, " ");
-
-				for (UEdGraphPin*& SubPin : CurrentPin->SubPins)
-				{
-					FString SubPinFriendlyName = SubPin->PinFriendlyName.ToString();
-					SubPinFriendlyName = SubPinFriendlyName.Replace(*OldFriendlyName, *ElementFriendlyName);
-
-					SubPin->Modify();
-					SubPin->PinName = SubPin->PinName.Replace(*OldName, *ElementName);
-					SubPin->PinFriendlyName = FText::FromString(SubPinFriendlyName);
-				}
-			}
-		}
-	}
-}
-
-void UK2Node_MakeArray::RemoveInputPin(UEdGraphPin* Pin)
-{
-	FScopedTransaction Transaction( LOCTEXT("RemovePinTx", "RemovePin") );
-	Modify();
-
-	check(Pin->Direction == EGPD_Input);
-
-	int32 PinRemovalIndex = INDEX_NONE;
-	if (Pins.Find(Pin, /*out*/ PinRemovalIndex))
-	{
-		UEdGraphPin* ParentPin = Pin->ParentPin;
-		Pins.RemoveAt(PinRemovalIndex);
-		Pin->MarkPendingKill();
-
-		if (!ParentPin)
-		{
-			--NumInputs;
-		}
-		else if (ParentPin->SubPins.Num() == 0)
-		{
-			// We have removed all of the subpins, so we need to remove the parent as well
-			Pins.Remove(ParentPin);
-			ParentPin->MarkPendingKill();
-			--NumInputs;
-		}
-
-		SyncPinNames();
-
-		FBlueprintEditorUtils::MarkBlueprintAsStructurallyModified(GetBlueprint());
-	}
-}
-
 void UK2Node_MakeArray::GetContextMenuActions(const FGraphNodeContextMenuBuilder& Context) const
 {
 	Super::GetContextMenuActions(Context);
@@ -480,10 +82,8 @@ void UK2Node_MakeArray::GetContextMenuActions(const FGraphNodeContextMenuBuilder
 
 		if (Context.Pin != NULL)
 		{
-			// we only do this for normal BlendList/BlendList by enum, BlendList by Bool doesn't support add/remove pins
 			if (Context.Pin->Direction == EGPD_Input)
 			{
-				//@TODO: Only offer this option on arrayed pins
 				Context.MenuBuilder->AddMenuEntry(
 					LOCTEXT("RemovePin", "Remove array element pin"),
 					LOCTEXT("RemovePinTooltip", "Remove this array element pin"),
@@ -520,25 +120,6 @@ void UK2Node_MakeArray::GetContextMenuActions(const FGraphNodeContextMenuBuilder
 	}
 }
 
-bool UK2Node_MakeArray::IsConnectionDisallowed(const UEdGraphPin* MyPin, const UEdGraphPin* OtherPin, FString& OutReason) const
-{
-	// if MyPin has a ParentPin then we are dealing with a split pin and we should evaluate it with default behavior
-	if(MyPin->ParentPin == nullptr && OtherPin->PinType.IsContainer() == true && MyPin->Direction == EGPD_Input) //-V595
-	{
-		OutReason = NSLOCTEXT("K2Node", "MakeArray_InputIsContainer", "Cannot make an array with an input of a container!").ToString();
-		return true;
-	}
-
-	const UEdGraphSchema_K2* Schema = Cast<const UEdGraphSchema_K2>(GetSchema());
-	if (!ensure(Schema) || (ensure(OtherPin) && Schema->IsExecPin(*OtherPin)))
-	{
-		OutReason = NSLOCTEXT("K2Node", "MakeArray_InputIsExec", "Cannot make an array with an execution input!").ToString();
-		return true;
-	}
-
-	return false;
-}
-
 void UK2Node_MakeArray::ValidateNodeDuringCompilation(class FCompilerResultsLog& MessageLog) const
 {
 	Super::ValidateNodeDuringCompilation(MessageLog);
@@ -547,31 +128,7 @@ void UK2Node_MakeArray::ValidateNodeDuringCompilation(class FCompilerResultsLog&
 	UEdGraphPin* OutputPin = GetOutputPin();
 	if (!ensure(Schema) || !ensure(OutputPin) || Schema->IsExecPin(*OutputPin))
 	{
-		MessageLog.Error(*NSLOCTEXT("K2Node", "MakeArray_OutputIsExec", "Uaccepted array type in @@").ToString(), this);
-	}
-}
-
-void UK2Node_MakeArray::GetMenuActions(FBlueprintActionDatabaseRegistrar& ActionRegistrar) const
-{
-	// actions get registered under specific object-keys; the idea is that 
-	// actions might have to be updated (or deleted) if their object-key is  
-	// mutated (or removed)... here we use the node's class (so if the node 
-	// type disappears, then the action should go with it)
-	// actions get registered under specific object-keys; the idea is that 
-	// actions might have to be updated (or deleted) if their object-key is  
-	// mutated (or removed)... here we use the node's class (so if the node 
-	// type disappears, then the action should go with it)
-	UClass* ActionKey = GetClass();
-	// to keep from needlessly instantiating a UBlueprintNodeSpawner, first   
-	// check to make sure that the registrar is looking for actions of this type
-	// (could be regenerating actions for a specific asset, and therefore the 
-	// registrar would only accept actions corresponding to that asset)
-	if (ActionRegistrar.IsOpenForRegistration(ActionKey))
-	{
-		UBlueprintNodeSpawner* NodeSpawner = UBlueprintNodeSpawner::Create(GetClass());
-		check(NodeSpawner != nullptr);
-
-		ActionRegistrar.AddBlueprintAction(ActionKey, NodeSpawner);
+		MessageLog.Error(*NSLOCTEXT("K2Node", "MakeArray_OutputIsExec", "Unacceptable array type in @@").ToString(), this);
 	}
 }
 
diff --git a/Engine/Source/Editor/BlueprintGraph/Private/K2Node_MakeMap.cpp b/Engine/Source/Editor/BlueprintGraph/Private/K2Node_MakeMap.cpp
new file mode 100644
--- /dev/null
+++ b/Engine/Source/Editor/BlueprintGraph/Private/K2Node_MakeMap.cpp
@@ -0,0 +1,216 @@
+// Copyright 1998-2017 Epic Games, Inc. All Rights Reserved.
+
+
+#include "K2Node_MakeMap.h"
+#include "EdGraph/EdGraphPin.h"
+#include "Engine/Blueprint.h"
+#include "Framework/Commands/UIAction.h"
+#include "Framework/MultiBox/MultiBoxBuilder.h"
+#include "EdGraphSchema_K2.h"
+#include "EdGraph/EdGraphNodeUtils.h"
+#include "Kismet2/BlueprintEditorUtils.h"
+
+#include "ScopedTransaction.h"
+#include "EdGraphUtilities.h"
+#include "KismetCompiledFunctionContext.h"
+#include "KismetCompilerMisc.h"
+#include "BlueprintNodeSpawner.h"
+#include "EditorCategoryUtils.h"
+#include "BlueprintActionDatabaseRegistrar.h"
+
+namespace MakeMapLiterals
+{
+	static const FString OutputPinName = FString(TEXT("Map"));
+};
+
+#define LOCTEXT_NAMESPACE "MakeMapNode"
+
+/////////////////////////////////////////////////////
+// FKCHandler_MakeMap
+class FKCHandler_MakeMap : public FKCHandler_MakeContainer
+{
+public:
+	FKCHandler_MakeMap(FKismetCompilerContext& InCompilerContext)
+		: FKCHandler_MakeContainer(InCompilerContext)
+	{
+		CompiledStatementType = KCST_CreateMap;
+	}
+};
+
+/////////////////////////////////////////////////////
+// UK2Node_MakeMap
+
+UK2Node_MakeMap::UK2Node_MakeMap(const FObjectInitializer& ObjectInitializer)
+	: Super(ObjectInitializer)
+{
+	ContainerType = EPinContainerType::Map;
+}
+
+FNodeHandlingFunctor* UK2Node_MakeMap::CreateNodeHandler(FKismetCompilerContext& CompilerContext) const
+{
+	return new FKCHandler_MakeMap(CompilerContext);
+}
+
+void UK2Node_MakeMap::AllocateDefaultPins()
+{
+	// Create the output pin
+	FEdGraphTerminalType TerminalType;
+	TerminalType.TerminalCategory = UEdGraphSchema_K2::PC_Wildcard;
+	CreatePin(EGPD_Output, UEdGraphSchema_K2::PC_Wildcard, FString(), nullptr, *GetOutputPinName(), ContainerType, false, false, INDEX_NONE, TerminalType);
+
+	// Create the input pins to create the container from
+	for (int32 i = 0; i < NumInputs; ++i)
+	{
+		CreatePin(EGPD_Input, UEdGraphSchema_K2::PC_Wildcard, FString(), nullptr, *GetPinName(i * 2));
+		CreatePin(EGPD_Input, UEdGraphSchema_K2::PC_Wildcard, FString(), nullptr, *GetPinName((i * 2) + 1));
+	}
+}
+
+void UK2Node_MakeMap::GetKeyAndValuePins(TArray<UEdGraphPin*>& KeyPins, TArray<UEdGraphPin*>& ValuePins) const
+{
+	for (UEdGraphPin* CurrentPin : Pins)
+	{
+		if (CurrentPin->Direction == EGPD_Input && CurrentPin->ParentPin == nullptr)
+		{
+			// Key/Value pins alternate so if this is a map and the counts are even then this is a key
+			if (KeyPins.Num() == ValuePins.Num())
+			{
+				KeyPins.Add(CurrentPin);
+			}
+			else
+			{
+				ValuePins.Add(CurrentPin);
+			}
+		}
+	}
+	check(KeyPins.Num() == ValuePins.Num());
+}
+
+void UK2Node_MakeMap::AddInputPin()
+{
+	Modify();
+
+	const UEdGraphSchema_K2* Schema = GetDefault<UEdGraphSchema_K2>();
+
+	++NumInputs;
+	const FEdGraphPinType& OutputPinType = GetOutputPin()->PinType;
+	UEdGraphPin* Pin = CreatePin(EGPD_Input, OutputPinType.PinCategory, OutputPinType.PinSubCategory, OutputPinType.PinSubCategoryObject.Get(), *GetPinName((NumInputs - 1) * 2));
+	Schema->SetPinAutogeneratedDefaultValueBasedOnType(Pin);
+
+	const FEdGraphPinType ValuePinType = FEdGraphPinType::GetPinTypeForTerminalType(OutputPinType.PinValueType);
+	Pin = CreatePin(EGPD_Input, ValuePinType.PinCategory, ValuePinType.PinSubCategory, ValuePinType.PinSubCategoryObject.Get(), *GetPinName(((NumInputs - 1) * 2) + 1));
+	Schema->SetPinAutogeneratedDefaultValueBasedOnType(Pin);
+
+	const bool bIsCompiling = GetBlueprint()->bBeingCompiled;
+	if (!bIsCompiling)
+	{
+		FBlueprintEditorUtils::MarkBlueprintAsStructurallyModified(GetBlueprint());
+	}
+}
+
+FText UK2Node_MakeMap::GetNodeTitle(ENodeTitleType::Type TitleType) const
+{
+	return LOCTEXT("NodeTitle", "Make Map");
+}
+
+FString UK2Node_MakeMap::GetPinName(const int32 PinIndex) const
+{
+	const int32 PairIndex = PinIndex / 2;
+	if (PinIndex % 2 == 0)
+	{
+		return FString::Printf(TEXT("Key %d"), PairIndex);
+	}
+	else
+	{
+		return FString::Printf(TEXT("Value %d"), PairIndex);
+	}
+}
+
+FString UK2Node_MakeMap::GetOutputPinName() const
+{
+	return MakeMapLiterals::OutputPinName;
+}
+
+FText UK2Node_MakeMap::GetTooltipText() const
+{
+	return LOCTEXT("MakeMapTooltip", "Create a map from a series of key/value items.");
+}
+
+FSlateIcon UK2Node_MakeMap::GetIconAndTint(FLinearColor& OutColor) const
+{
+	static FSlateIcon Icon("EditorStyle", "GraphEditor.MakeMap_16x");
+	return Icon;
+}
+
+void UK2Node_MakeMap::GetContextMenuActions(const FGraphNodeContextMenuBuilder& Context) const
+{
+	Super::GetContextMenuActions(Context);
+
+	if (!Context.bIsDebugging)
+	{
+		Context.MenuBuilder->BeginSection("K2NodeMakeMap", NSLOCTEXT("K2Nodes", "MakeMapHeader", "MakeMap"));
+
+		if (Context.Pin)
+		{
+			if (Context.Pin->Direction == EGPD_Input && Context.Pin->ParentPin == nullptr)
+			{
+				Context.MenuBuilder->AddMenuEntry(
+					LOCTEXT("RemovePin", "Remove key/value pair"),
+					LOCTEXT("RemovePinTooltip", "Remove this pin and its corresponding key/value pin"),
+					FSlateIcon(),
+					FUIAction(
+						FExecuteAction::CreateUObject(this, &UK2Node_MakeMap::RemoveInputPin, const_cast<UEdGraphPin*>(Context.Pin))
+					)
+				);
+			}
+		}
+		else
+		{
+			Context.MenuBuilder->AddMenuEntry(
+				LOCTEXT("AddPin", "Add key/value pair"),
+				LOCTEXT("AddPinTooltip", "Add another pair of key/value pins"),
+				FSlateIcon(),
+				FUIAction(
+					FExecuteAction::CreateUObject(this, &UK2Node_MakeMap::AddInputPin)
+				)
+			);
+		}
+
+		Context.MenuBuilder->AddMenuEntry(
+			LOCTEXT("ResetToWildcard", "Reset to wildcard"),
+			LOCTEXT("ResetToWildcardTooltip", "Reset the node to have wildcard input/outputs. Requires no pins are connected."),
+			FSlateIcon(),
+			FUIAction(
+				FExecuteAction::CreateUObject(this, &UK2Node_MakeMap::ClearPinTypeToWildcard),
+				FCanExecuteAction::CreateUObject(this, &UK2Node_MakeMap::CanResetToWildcard)
+			)
+		);
+
+		Context.MenuBuilder->EndSection();
+	}
+}
+
+void UK2Node_MakeMap::ValidateNodeDuringCompilation(class FCompilerResultsLog& MessageLog) const
+{
+	Super::ValidateNodeDuringCompilation(MessageLog);
+
+	const UEdGraphSchema_K2* Schema = Cast<const UEdGraphSchema_K2>(GetSchema());
+	UEdGraphPin* OutputPin = GetOutputPin();
+	if (!ensure(Schema) || !ensure(OutputPin) || Schema->IsExecPin(*OutputPin))
+	{
+		MessageLog.Error(*NSLOCTEXT("K2Node", "MakeMap_OutputIsExec", "Unacceptable map type in @@").ToString(), this);
+	}
+}
+
+FText UK2Node_MakeMap::GetMenuCategory() const
+{
+	static FNodeTextCache CachedCategory;
+	if (CachedCategory.IsOutOfDate(this))
+	{
+		// FText::Format() is slow, so we cache this to save on performance
+		CachedCategory.SetCachedText(FEditorCategoryUtils::BuildCategoryString(FCommonEditorCategory::Utilities, LOCTEXT("ActionMenuCategory", "Map")), this);
+	}
+	return CachedCategory;
+}
+
+#undef LOCTEXT_NAMESPACE
