Missed file from CL# 3544050
#rnx

diff --git a/Engine/Source/Editor/BlueprintGraph/Private/EdGraphSchema_K2.cpp b/Engine/Source/Editor/BlueprintGraph/Private/EdGraphSchema_K2.cpp
--- a/Engine/Source/Editor/BlueprintGraph/Private/EdGraphSchema_K2.cpp
+++ b/Engine/Source/Editor/BlueprintGraph/Private/EdGraphSchema_K2.cpp
@@ -1,6 +1,7 @@
 // Copyright 1998-2017 Epic Games, Inc. All Rights Reserved.
 
 #include "EdGraphSchema_K2.h"
+#include "BlueprintCompilationManager.h"
 #include "Modules/ModuleManager.h"
 #include "UObject/Interface.h"
 #include "UObject/UnrealType.h"
@@ -64,6 +65,7 @@
 #include "EditorStyleSettings.h"
 #include "Editor.h"
 
+#include "Kismet/BlueprintMapLibrary.h"
 #include "Kismet/BlueprintSetLibrary.h"
 #include "Kismet/KismetArrayLibrary.h"
 #include "Kismet/KismetMathLibrary.h"
@@ -1225,6 +1227,7 @@ bool UEdGraphSchema_K2::DoesGraphSupportImpureFunctions(const UEdGraph* InGraph)
 
 bool UEdGraphSchema_K2::IsPropertyExposedOnSpawn(const UProperty* Property)
 {
+	Property = FBlueprintEditorUtils::GetMostUpToDateProperty(Property);
 	if (Property)
 	{
 		const bool bMeta = Property->HasMetaData(FBlueprintMetadata::MD_ExposeOnSpawn);
@@ -3505,16 +3508,6 @@ bool UEdGraphSchema_K2::IsWildcardProperty(const UProperty* Property)
 		|| FEdGraphUtilities::IsArrayDependentParam(Function, Property->GetName()) );
 }
 
-FString UEdGraphSchema_K2::TypeToString(const FEdGraphPinType& Type)
-{
-	return TypeToText(Type).ToString();
-}
-
-FString UEdGraphSchema_K2::TypeToString(UProperty* const Property)
-{
-	return TypeToText(Property).ToString();
-}
-
 FText UEdGraphSchema_K2::TypeToText(UProperty* const Property)
 {
 	if (UStructProperty* Struct = Cast<UStructProperty>(Property))
@@ -4457,6 +4450,9 @@ void UEdGraphSchema_K2::HandleGraphBeingDeleted(UEdGraph& GraphBeingRemoved) con
 			NodeToDelete->Modify();
 			NodeToDelete->DestroyNode();
 		}
+
+		// Remove from the list of recently edited documents
+		Blueprint->LastEditedDocuments.RemoveAll([&GraphBeingRemoved](const FEditedDocumentInfo& TestDoc) { return TestDoc.EditedObject == &GraphBeingRemoved; });
 	}
 }
 
@@ -4930,6 +4926,31 @@ void UEdGraphSchema_K2::ValidateExistingConnections(UEdGraphPin* Pin)
 	}
 }
 
+namespace FSetVariableByNameFunctionNames
+{
+	static const FName SetIntName(GET_FUNCTION_NAME_CHECKED(UKismetSystemLibrary, SetIntPropertyByName));
+	static const FName SetByteName(GET_FUNCTION_NAME_CHECKED(UKismetSystemLibrary, SetBytePropertyByName));
+	static const FName SetFloatName(GET_FUNCTION_NAME_CHECKED(UKismetSystemLibrary, SetFloatPropertyByName));
+	static const FName SetBoolName(GET_FUNCTION_NAME_CHECKED(UKismetSystemLibrary, SetBoolPropertyByName));
+	static const FName SetObjectName(GET_FUNCTION_NAME_CHECKED(UKismetSystemLibrary, SetObjectPropertyByName));
+	static const FName SetClassName(GET_FUNCTION_NAME_CHECKED(UKismetSystemLibrary, SetClassPropertyByName));
+	static const FName SetInterfaceName(GET_FUNCTION_NAME_CHECKED(UKismetSystemLibrary, SetInterfacePropertyByName));
+	static const FName SetStringName(GET_FUNCTION_NAME_CHECKED(UKismetSystemLibrary, SetStringPropertyByName));
+	static const FName SetTextName(GET_FUNCTION_NAME_CHECKED(UKismetSystemLibrary, SetTextPropertyByName));
+	static const FName SetAssetName(GET_FUNCTION_NAME_CHECKED(UKismetSystemLibrary, SetAssetPropertyByName));
+	static const FName SetAssetClassName(GET_FUNCTION_NAME_CHECKED(UKismetSystemLibrary, SetAssetClassPropertyByName));
+	static const FName SetNameName(GET_FUNCTION_NAME_CHECKED(UKismetSystemLibrary, SetNamePropertyByName));
+	static const FName SetVectorName(GET_FUNCTION_NAME_CHECKED(UKismetSystemLibrary, SetVectorPropertyByName));
+	static const FName SetRotatorName(GET_FUNCTION_NAME_CHECKED(UKismetSystemLibrary, SetRotatorPropertyByName));
+	static const FName SetLinearColorName(GET_FUNCTION_NAME_CHECKED(UKismetSystemLibrary, SetLinearColorPropertyByName));
+	static const FName SetTransformName(GET_FUNCTION_NAME_CHECKED(UKismetSystemLibrary, SetTransformPropertyByName));
+	static const FName SetCollisionProfileName(GET_FUNCTION_NAME_CHECKED(UKismetSystemLibrary, SetCollisionProfileNameProperty));
+	static const FName SetStructureName(GET_FUNCTION_NAME_CHECKED(UKismetSystemLibrary, SetStructurePropertyByName));
+	static const FName SetArrayName(GET_FUNCTION_NAME_CHECKED(UKismetArrayLibrary, SetArrayPropertyByName));
+	static const FName SetSetName(GET_FUNCTION_NAME_CHECKED(UBlueprintSetLibrary, SetSetPropertyByName));
+	static const FName SetMapName(GET_FUNCTION_NAME_CHECKED(UBlueprintMapLibrary, SetMapPropertyByName));
+};
+
 UFunction* UEdGraphSchema_K2::FindSetVariableByNameFunction(const FEdGraphPinType& PinType)
 {
 	//!!!! Keep this function synced with FExposeOnSpawnValidator::IsSupported !!!!
@@ -4943,112 +4964,100 @@ UFunction* UEdGraphSchema_K2::FindSetVariableByNameFunction(const FEdGraphPinTyp
 		}
 	};
 
-	const UEdGraphSchema_K2* K2Schema = GetDefault<UEdGraphSchema_K2>();
-
+	UClass* SetFunctionLibraryClass = UKismetSystemLibrary::StaticClass();
 	FName SetFunctionName = NAME_None;
-	if(PinType.PinCategory == K2Schema->PC_Int)
+	if (PinType.ContainerType == EPinContainerType::Array)
 	{
-		static FName SetIntName(GET_FUNCTION_NAME_CHECKED(UKismetSystemLibrary, SetIntPropertyByName));
-		SetFunctionName = SetIntName;
+		SetFunctionName = FSetVariableByNameFunctionNames::SetArrayName;
+		SetFunctionLibraryClass = UKismetArrayLibrary::StaticClass();
 	}
-	else if(PinType.PinCategory == K2Schema->PC_Byte)
+	else if (PinType.ContainerType == EPinContainerType::Set)
 	{
-		static FName SetByteName(GET_FUNCTION_NAME_CHECKED(UKismetSystemLibrary, SetBytePropertyByName));
-		SetFunctionName = SetByteName;
+		SetFunctionName = FSetVariableByNameFunctionNames::SetSetName;
+		SetFunctionLibraryClass = UBlueprintSetLibrary::StaticClass();
 	}
-	else if(PinType.PinCategory == K2Schema->PC_Float)
+	else if (PinType.ContainerType == EPinContainerType::Map)
 	{
-		static FName SetFloatName(GET_FUNCTION_NAME_CHECKED(UKismetSystemLibrary, SetFloatPropertyByName));
-		SetFunctionName = SetFloatName;
+		SetFunctionName = FSetVariableByNameFunctionNames::SetMapName;
+		SetFunctionLibraryClass = UBlueprintMapLibrary::StaticClass();
 	}
-	else if(PinType.PinCategory == K2Schema->PC_Boolean)
+	else if(PinType.PinCategory == UEdGraphSchema_K2::PC_Int)
 	{
-		static FName SetBoolName(GET_FUNCTION_NAME_CHECKED(UKismetSystemLibrary, SetBoolPropertyByName));
-		SetFunctionName = SetBoolName;
+		SetFunctionName = FSetVariableByNameFunctionNames::SetIntName;
 	}
-	else if(PinType.PinCategory == K2Schema->PC_Object)
+	else if(PinType.PinCategory == UEdGraphSchema_K2::PC_Byte)
 	{
-		static FName SetObjectName(GET_FUNCTION_NAME_CHECKED(UKismetSystemLibrary, SetObjectPropertyByName));
-		SetFunctionName = SetObjectName;
+		SetFunctionName = FSetVariableByNameFunctionNames::SetByteName;
 	}
-	else if (PinType.PinCategory == K2Schema->PC_Class)
+	else if(PinType.PinCategory == UEdGraphSchema_K2::PC_Float)
 	{
-		static FName SetObjectName(GET_FUNCTION_NAME_CHECKED(UKismetSystemLibrary, SetClassPropertyByName));
-		SetFunctionName = SetObjectName;
+		SetFunctionName = FSetVariableByNameFunctionNames::SetFloatName;
 	}
-	else if(PinType.PinCategory == K2Schema->PC_Interface)
+	else if(PinType.PinCategory == UEdGraphSchema_K2::PC_Boolean)
 	{
-		static FName SetObjectName(GET_FUNCTION_NAME_CHECKED(UKismetSystemLibrary, SetInterfacePropertyByName));
-		SetFunctionName = SetObjectName;
+		SetFunctionName = FSetVariableByNameFunctionNames::SetBoolName;
 	}
-	else if(PinType.PinCategory == K2Schema->PC_String)
+	else if(PinType.PinCategory == UEdGraphSchema_K2::PC_Object)
 	{
-		static FName SetStringName(GET_FUNCTION_NAME_CHECKED(UKismetSystemLibrary, SetStringPropertyByName));
-		SetFunctionName = SetStringName;
+		SetFunctionName = FSetVariableByNameFunctionNames::SetObjectName;
 	}
-	else if ( PinType.PinCategory == K2Schema->PC_Text)
+	else if (PinType.PinCategory == UEdGraphSchema_K2::PC_Class)
 	{
-		static FName SetTextName(GET_FUNCTION_NAME_CHECKED(UKismetSystemLibrary, SetTextPropertyByName));
-		SetFunctionName = SetTextName;
+		SetFunctionName = FSetVariableByNameFunctionNames::SetClassName;
 	}
-	else if (PinType.PinCategory == K2Schema->PC_Asset)
+	else if(PinType.PinCategory == UEdGraphSchema_K2::PC_Interface)
 	{
-		static FName SetAssetName(GET_FUNCTION_NAME_CHECKED(UKismetSystemLibrary, SetAssetPropertyByName));
-		SetFunctionName = SetAssetName;
+		SetFunctionName = FSetVariableByNameFunctionNames::SetInterfaceName;
 	}
-	else if (PinType.PinCategory == K2Schema->PC_AssetClass)
+	else if(PinType.PinCategory == UEdGraphSchema_K2::PC_String)
 	{
-		static FName SetAssetClassName(GET_FUNCTION_NAME_CHECKED(UKismetSystemLibrary, SetAssetClassPropertyByName));
-		SetFunctionName = SetAssetClassName;
+		SetFunctionName = FSetVariableByNameFunctionNames::SetStringName;
 	}
-	else if(PinType.PinCategory == K2Schema->PC_Name)
+	else if ( PinType.PinCategory == UEdGraphSchema_K2::PC_Text)
 	{
-		static FName SetNameName(GET_FUNCTION_NAME_CHECKED(UKismetSystemLibrary, SetNamePropertyByName));
-		SetFunctionName = SetNameName;
+		SetFunctionName = FSetVariableByNameFunctionNames::SetTextName;
 	}
-	else if(PinType.PinCategory == K2Schema->PC_Struct && PinType.PinSubCategoryObject == VectorStruct)
+	else if (PinType.PinCategory == UEdGraphSchema_K2::PC_Asset)
 	{
-		static FName SetVectorName(GET_FUNCTION_NAME_CHECKED(UKismetSystemLibrary, SetVectorPropertyByName));
-		SetFunctionName = SetVectorName;
+		SetFunctionName = FSetVariableByNameFunctionNames::SetAssetName;
 	}
-	else if(PinType.PinCategory == K2Schema->PC_Struct && PinType.PinSubCategoryObject == RotatorStruct)
+	else if (PinType.PinCategory == UEdGraphSchema_K2::PC_AssetClass)
 	{
-		static FName SetRotatorName(GET_FUNCTION_NAME_CHECKED(UKismetSystemLibrary, SetRotatorPropertyByName));
-		SetFunctionName = SetRotatorName;
+		SetFunctionName = FSetVariableByNameFunctionNames::SetAssetClassName;
 	}
-	else if(PinType.PinCategory == K2Schema->PC_Struct && PinType.PinSubCategoryObject == ColorStruct)
+	else if(PinType.PinCategory == UEdGraphSchema_K2::PC_Name)
 	{
-		static FName SetLinearColorName(GET_FUNCTION_NAME_CHECKED(UKismetSystemLibrary, SetLinearColorPropertyByName));
-		SetFunctionName = SetLinearColorName;
+		SetFunctionName = FSetVariableByNameFunctionNames::SetNameName;
 	}
-	else if(PinType.PinCategory == K2Schema->PC_Struct && PinType.PinSubCategoryObject == TransformStruct)
+	else if(PinType.PinCategory == UEdGraphSchema_K2::PC_Struct && PinType.PinSubCategoryObject == VectorStruct)
 	{
-		static FName SetTransformName(GET_FUNCTION_NAME_CHECKED(UKismetSystemLibrary, SetTransformPropertyByName));
-		SetFunctionName = SetTransformName;
+		SetFunctionName = FSetVariableByNameFunctionNames::SetVectorName;
 	}
-	else if (PinType.PinCategory == K2Schema->PC_Struct && PinType.PinSubCategoryObject == FCollisionProfileName::StaticStruct())
+	else if(PinType.PinCategory == UEdGraphSchema_K2::PC_Struct && PinType.PinSubCategoryObject == RotatorStruct)
 	{
-		static FName SetStructureName(GET_FUNCTION_NAME_CHECKED(UKismetSystemLibrary, SetCollisionProfileNameProperty));
-		SetFunctionName = SetStructureName;
+		SetFunctionName = FSetVariableByNameFunctionNames::SetRotatorName;
 	}
-	else if (PinType.PinCategory == K2Schema->PC_Struct && FIsCustomStructureParamHelper::Is(PinType.PinSubCategoryObject.Get()))
+	else if(PinType.PinCategory == UEdGraphSchema_K2::PC_Struct && PinType.PinSubCategoryObject == ColorStruct)
 	{
-		static FName SetStructureName(GET_FUNCTION_NAME_CHECKED(UKismetSystemLibrary, SetStructurePropertyByName));
-		SetFunctionName = SetStructureName;
+		SetFunctionName = FSetVariableByNameFunctionNames::SetLinearColorName;
+	}
+	else if(PinType.PinCategory == UEdGraphSchema_K2::PC_Struct && PinType.PinSubCategoryObject == TransformStruct)
+	{
+		SetFunctionName = FSetVariableByNameFunctionNames::SetTransformName;
+	}
+	else if (PinType.PinCategory == UEdGraphSchema_K2::PC_Struct && PinType.PinSubCategoryObject == FCollisionProfileName::StaticStruct())
+	{
+		SetFunctionName = FSetVariableByNameFunctionNames::SetCollisionProfileName;
+	}
+	else if (PinType.PinCategory == UEdGraphSchema_K2::PC_Struct && FIsCustomStructureParamHelper::Is(PinType.PinSubCategoryObject.Get()))
+	{
+		SetFunctionName = FSetVariableByNameFunctionNames::SetStructureName;
 	}
 
-	UFunction* Function = NULL;
-	if(SetFunctionName != NAME_None)
+	UFunction* Function = nullptr;
+	if (!SetFunctionName.IsNone())
 	{
-		if(PinType.IsArray())
-		{
-			static FName SetArrayName(GET_FUNCTION_NAME_CHECKED(UKismetArrayLibrary, SetArrayPropertyByName));
-			Function = FindField<UFunction>(UKismetArrayLibrary::StaticClass(), SetArrayName);
-		}
-		else
-		{
-			Function = FindField<UFunction>(UKismetSystemLibrary::StaticClass(), SetFunctionName);
-		}
+		Function = SetFunctionLibraryClass->FindFunctionByName(SetFunctionName);
 	}
 
 	return Function;
@@ -5994,7 +6003,7 @@ UEdGraphNode* UEdGraphSchema_K2::CreateSubstituteNode(UEdGraphNode* Node, const
 			UObject* Found = FindObject<UObject>(EventNode->GetOuter(), *ObjName.ToString());
 			if(Found)
 			{
-				Found->Rename(NULL, NULL, REN_DontCreateRedirectors | RenameFlags | (Found->HasAnyFlags(RF_NeedLoad | RF_NeedPostLoad | RF_NeedPostLoadSubobjects) ? REN_ForceNoResetLoaders : RF_NoFlags));
+				Found->Rename(NULL, NULL, REN_DontCreateRedirectors | RenameFlags | ((IsAsyncLoading() || Found->HasAnyFlags(RF_NeedLoad | RF_NeedPostLoad | RF_NeedPostLoadSubobjects)) ? REN_ForceNoResetLoaders : RF_NoFlags));
 			}
 
 			// Create a custom event node to replace the original event node imported from text
