Copying //UE4/Dev-Framework to //UE4/Dev-Main (Source: //UE4/Dev-Framework @ 3582324)

diff --git a/Engine/Source/Editor/BlueprintGraph/Private/EdGraphSchema_K2.cpp b/Engine/Source/Editor/BlueprintGraph/Private/EdGraphSchema_K2.cpp
--- a/Engine/Source/Editor/BlueprintGraph/Private/EdGraphSchema_K2.cpp
+++ b/Engine/Source/Editor/BlueprintGraph/Private/EdGraphSchema_K2.cpp
@@ -189,7 +189,7 @@ UEdGraphSchema_K2::FPinTypeTreeInfo::FPinTypeTreeInfo(const FText& InFriendlyNam
 
 struct FUnloadedAssetData
 {
-	FStringAssetReference StringAssetReference;
+	FSoftObjectPath SoftObjectPath;
 	FText AssetFriendlyName;
 	FText Tooltip;
 	uint8 PossibleObjectReferenceTypes;
@@ -199,7 +199,7 @@ struct FUnloadedAssetData
 	{}
 
 	FUnloadedAssetData(const FAssetData& InAsset, uint8 InPossibleObjectReferenceTypes = 0)
-		: StringAssetReference(InAsset.ToStringReference())
+		: SoftObjectPath(InAsset.ToSoftObjectPath())
 		, AssetFriendlyName(FText::FromString(FName::NameToDisplayString(InAsset.AssetName.ToString(), false)))
 		, PossibleObjectReferenceTypes(InPossibleObjectReferenceTypes)
 	{
@@ -290,8 +290,8 @@ public:
 				TypesDatabase.LoadedTypesMap.Add(UEdGraphSchema_K2::PC_Struct, LoadedTypesList);
 			}
 
-			//(Type == UEdGraphSchema_K2::PC_Class || Type == UEdGraphSchema_K2::PC_AssetClass) UEdGraphSchema_K2::PC_Interface)
-			//(Type == UEdGraphSchema_K2::PC_Object || Type == UEdGraphSchema_K2::PC_Asset)
+			//(Type == UEdGraphSchema_K2::PC_Class || Type == UEdGraphSchema_K2::PC_SoftClass) UEdGraphSchema_K2::PC_Interface)
+			//(Type == UEdGraphSchema_K2::PC_Object || Type == UEdGraphSchema_K2::PC_SoftObject)
 			{
 				FTypesDatabase::FLoadedTypesList InterfaceLoadedTypesList = MakeShareable(new TArray<FLoadedAssetData>());
 				FTypesDatabase::FLoadedTypesList AllObjectLoadedTypesList = MakeShareable(new TArray<FLoadedAssetData>());
@@ -364,7 +364,7 @@ public:
 				TypesDatabase.UnLoadedTypesMap.Add(UEdGraphSchema_K2::PC_Struct, UnLoadedTypesList);
 			}
 
-			//else if (Schema->PC_Object == CategoryName || Schema->PC_Class == CategoryName || Schema->PC_Interface == CategoryName || Schema->PC_Asset == CategoryName || Schema->PC_AssetClass == CategoryName)
+			//else if (Schema->PC_Object == CategoryName || Schema->PC_Class == CategoryName || Schema->PC_Interface == CategoryName || Schema->PC_SoftObject == CategoryName || Schema->PC_SoftClass == CategoryName)
 			{
 				TArray<FAssetData> AssetData;
 				AssetRegistryModule.Get().GetAssetsByClass(UBlueprint::StaticClass()->GetFName(), AssetData);
@@ -442,7 +442,7 @@ public:
 			{
 				FPinTypeTreeInfoPtr TypeTreeInfo = MakeShareable(new UEdGraphSchema_K2::FPinTypeTreeInfo(It.AssetFriendlyName
 					, CategoryName
-					, It.StringAssetReference
+					, It.SoftObjectPath
 					, It.Tooltip
 					, false
 					, It.PossibleObjectReferenceTypes));
@@ -454,7 +454,7 @@ public:
 	}
 
 	/** Loads an asset based on the AssetReference through the asset registry */
-	static UObject* LoadAsset(const FStringAssetReference& AssetReference)
+	static UObject* LoadAsset(const FSoftObjectPath& AssetReference)
 	{
 		if (AssetReference.IsValid())
 		{
@@ -471,7 +471,7 @@ const FEdGraphPinType& UEdGraphSchema_K2::FPinTypeTreeInfo::GetPinType(bool bFor
 	if (bForceLoadedSubCategoryObject)
 	{
 		// Only attempt to load the sub category object if we need to
-		if ( SubCategoryObjectAssetReference.IsValid() && (!PinType.PinSubCategoryObject.IsValid() || FStringAssetReference(PinType.PinSubCategoryObject.Get()) != SubCategoryObjectAssetReference) )
+		if ( SubCategoryObjectAssetReference.IsValid() && (!PinType.PinSubCategoryObject.IsValid() || FSoftObjectPath(PinType.PinSubCategoryObject.Get()) != SubCategoryObjectAssetReference) )
 		{
 			UObject* LoadedObject = FGatherTypesHelper::LoadAsset(SubCategoryObjectAssetReference);
 
@@ -553,7 +553,7 @@ UEdGraphSchema_K2::FPinTypeTreeInfo::FPinTypeTreeInfo(const FString& CategoryNam
 	CachedDescription = GenerateDescription();
 }
 
-UEdGraphSchema_K2::FPinTypeTreeInfo::FPinTypeTreeInfo(const FText& InFriendlyName, const FString& CategoryName, const FStringAssetReference& SubCategoryObject, const FText& InTooltip, bool bInReadOnly, uint8 InPossibleObjectReferenceTypes)
+UEdGraphSchema_K2::FPinTypeTreeInfo::FPinTypeTreeInfo(const FText& InFriendlyName, const FString& CategoryName, const FSoftObjectPath& SubCategoryObject, const FText& InTooltip, bool bInReadOnly, uint8 InPossibleObjectReferenceTypes)
 	: PossibleObjectReferenceTypes(InPossibleObjectReferenceTypes)
 {
 	FriendlyName = InFriendlyName;
@@ -618,8 +618,8 @@ const FString UEdGraphSchema_K2::PC_Text(TEXT("text"));
 const FString UEdGraphSchema_K2::PC_Struct(TEXT("struct"));
 const FString UEdGraphSchema_K2::PC_Wildcard(TEXT("wildcard"));
 const FString UEdGraphSchema_K2::PC_Enum(TEXT("enum"));
-const FString UEdGraphSchema_K2::PC_Asset(TEXT("asset"));
-const FString UEdGraphSchema_K2::PC_AssetClass(TEXT("assetclass"));
+const FString UEdGraphSchema_K2::PC_SoftObject(TEXT("softobject"));
+const FString UEdGraphSchema_K2::PC_SoftClass(TEXT("softclass"));
 const FString UEdGraphSchema_K2::PSC_Self(TEXT("self"));
 const FString UEdGraphSchema_K2::PSC_Index(TEXT("index"));
 const FString UEdGraphSchema_K2::PSC_Bitmask(TEXT("bitmask"));
@@ -1258,7 +1258,7 @@ static bool IsUsingNonExistantVariable(const UEdGraphNode* InGraphNode, UBluepri
 		{
 			if (Variable->VariableReference.IsSelfContext())
 			{
-				TArray<FName> CurrentVars;
+				TSet<FName> CurrentVars;
 				FBlueprintEditorUtils::GetClassVariableList(OwnerBlueprint, CurrentVars);
 				if ( false == CurrentVars.Contains(Variable->GetVarName()) )
 				{
@@ -1577,34 +1577,22 @@ void UEdGraphSchema_K2::GetContextMenuActions(const UEdGraph* CurrentGraph, cons
 					{
 						MenuBuilder->AddMenuEntry(FGraphEditorCommands::Get().AssignReferencedActor);
 					}
+				}
 
-					// Add the goto source code action for native functions
-					if (InGraphNode->IsA(UK2Node_CallFunction::StaticClass()))
-					{
-						const UEdGraphNode* ResultEventNode= NULL;
-						if(Cast<UK2Node_CallFunction>(InGraphNode)->GetFunctionGraph(ResultEventNode))
-						{
-							MenuBuilder->AddMenuEntry(FGraphEditorCommands::Get().GoToDefinition);
-						}
-						else
-						{
-							MenuBuilder->AddMenuEntry(FGraphEditorCommands::Get().GotoNativeFunctionDefinition);
-						}
-					}
-
-					// Functions, macros, and composite nodes support going to a definition
-					if (InGraphNode->IsA(UK2Node_MacroInstance::StaticClass()) || InGraphNode->IsA(UK2Node_Composite::StaticClass()))
-					{
-						MenuBuilder->AddMenuEntry(FGraphEditorCommands::Get().GoToDefinition);
-					}
+				// If the node has an associated definition (for some loose sense of the word), allow going to it (same action as double-clicking on a node)
+				if (InGraphNode->CanJumpToDefinition())
+				{
+					MenuBuilder->AddMenuEntry(FGraphEditorCommands::Get().GoToDefinition);
+				}
 
-					MenuBuilder->AddMenuEntry(FGraphEditorCommands::Get().FindReferences);
+				// show search for references for everyone
+				MenuBuilder->AddMenuEntry(FGraphEditorCommands::Get().FindReferences);
 
-					// show search for references for variable nodes and goto source code action
+				if (!bIsDebugging)
+				{
 					if (InGraphNode->IsA(UK2Node_Variable::StaticClass()))
 					{
-						MenuBuilder->AddMenuEntry(FGraphEditorCommands::Get().GotoNativeVariableDefinition);
-						GetReplaceVariableMenu(InGraphNode,OwnerBlueprint, MenuBuilder, true);
+						GetReplaceVariableMenu(InGraphNode, OwnerBlueprint, MenuBuilder, true);
 					}
 
 					if (InGraphNode->IsA(UK2Node_SetFieldsInStruct::StaticClass()))
@@ -1663,11 +1651,10 @@ void UEdGraphSchema_K2::GetContextMenuActions(const UEdGraph* CurrentGraph, cons
 			{
 				if (!K2Node->IsNodePure())
 				{
-					const UBlueprintEditorSettings* EditorSettings = GetDefault<UBlueprintEditorSettings>();
-					if (EditorSettings && EditorSettings->bAllowExplicitImpureNodeDisabling)
+					if (!bIsDebugging && GetDefault<UBlueprintEditorSettings>()->bAllowExplicitImpureNodeDisabling)
 					{
 						// Don't expose the enabled state for disabled nodes that were not explicitly disabled by the user
-						if (!bIsDebugging && (K2Node->bUserSetEnabledState || K2Node->EnabledState != ENodeEnabledState::Disabled))
+						if (!K2Node->IsAutomaticallyPlacedGhostNode())
 						{
 							// Add compile options
 							MenuBuilder->BeginSection("EdGraphSchemaCompileOptions", LOCTEXT("CompileOptionsHeader", "Compile Options"));
@@ -2699,10 +2686,10 @@ bool UEdGraphSchema_K2::FindSpecializedConversionNode(const UEdGraphPin* OutputP
 
 			if (!bCanConvert && InputClass && OutputClass && OutputClass->IsChildOf(InputClass))
 			{
-				const bool bConvertAsset = (OutputType.PinCategory == PC_Asset) && (InputType.PinCategory == PC_Object);
-				const bool bConvertAssetClass = (OutputType.PinCategory == PC_AssetClass) && (InputType.PinCategory == PC_Class);
-				const bool bConvertToAsset = (OutputType.PinCategory == PC_Object) && (InputType.PinCategory == PC_Asset);
-				const bool bConvertToAssetClass = (OutputType.PinCategory == PC_Class) && (InputType.PinCategory == PC_AssetClass);
+				const bool bConvertAsset = (OutputType.PinCategory == PC_SoftObject) && (InputType.PinCategory == PC_Object);
+				const bool bConvertAssetClass = (OutputType.PinCategory == PC_SoftClass) && (InputType.PinCategory == PC_Class);
+				const bool bConvertToAsset = (OutputType.PinCategory == PC_Object) && (InputType.PinCategory == PC_SoftObject);
+				const bool bConvertToAssetClass = (OutputType.PinCategory == PC_Class) && (InputType.PinCategory == PC_SoftClass);
 
 				if (bConvertAsset || bConvertAssetClass || bConvertToAsset || bConvertToAssetClass)
 				{
@@ -2967,13 +2954,13 @@ FLinearColor UEdGraphSchema_K2::GetPinTypeColor(const FEdGraphPinType& PinType)
 	{
 		return Settings->NamePinTypeColor;
 	}
-	else if (TypeString == PC_Asset)
+	else if (TypeString == PC_SoftObject)
 	{
-		return Settings->AssetPinTypeColor;
+		return Settings->SoftObjectPinTypeColor;
 	}
-	else if (TypeString == PC_AssetClass)
+	else if (TypeString == PC_SoftClass)
 	{
-		return Settings->AssetClassPinTypeColor;
+		return Settings->SoftClassPinTypeColor;
 	}
 	else if (TypeString == PC_Delegate)
 	{
@@ -3314,15 +3301,15 @@ bool UEdGraphSchema_K2::GetPropertyCategoryInfo(const UProperty* TestProperty, F
 		OutCategory = PC_Class;
 		OutSubCategoryObject = ClassProperty->MetaClass;
 	}
-	else if (const UAssetClassProperty* AssetClassProperty = Cast<const UAssetClassProperty>(TestProperty))
+	else if (const USoftClassProperty* SoftClassProperty = Cast<const USoftClassProperty>(TestProperty))
 	{
-		OutCategory = PC_AssetClass;
-		OutSubCategoryObject = AssetClassProperty->MetaClass;
+		OutCategory = PC_SoftClass;
+		OutSubCategoryObject = SoftClassProperty->MetaClass;
 	}
-	else if (const UAssetObjectProperty* AssetObjectProperty = Cast<const UAssetObjectProperty>(TestProperty))
+	else if (const USoftObjectProperty* SoftObjectProperty = Cast<const USoftObjectProperty>(TestProperty))
 	{
-		OutCategory = PC_Asset;
-		OutSubCategoryObject = AssetObjectProperty->PropertyClass;
+		OutCategory = PC_SoftObject;
+		OutSubCategoryObject = SoftObjectProperty->PropertyClass;
 	}
 	else if (const UObjectPropertyBase* ObjectProperty = Cast<const UObjectPropertyBase>(TestProperty))
 	{
@@ -3611,8 +3598,8 @@ FText UEdGraphSchema_K2::GetCategoryText(const FString& Category, const bool bFo
 		CategoryDescriptions.Add(PC_Struct, LOCTEXT("StructCategory","Structure"));
 		CategoryDescriptions.Add(PC_Wildcard, LOCTEXT("WildcardCategory","Wildcard"));
 		CategoryDescriptions.Add(PC_Enum, LOCTEXT("EnumCategory","Enum"));
-		CategoryDescriptions.Add(PC_Asset, LOCTEXT("SoftObjectReferenceCategory", "Soft Object Reference"));
-		CategoryDescriptions.Add(PC_AssetClass, LOCTEXT("SoftClassReferenceCategory", "Soft Class Reference"));
+		CategoryDescriptions.Add(PC_SoftObject, LOCTEXT("SoftObjectReferenceCategory", "Soft Object Reference"));
+		CategoryDescriptions.Add(PC_SoftClass, LOCTEXT("SoftClassReferenceCategory", "Soft Class Reference"));
 		CategoryDescriptions.Add(AllObjectTypes, LOCTEXT("AllObjectTypes", "Object Types"));
 	}
 
@@ -3807,8 +3794,8 @@ void UEdGraphSchema_K2::GetVariableTypeTree(TArray< TSharedPtr<FPinTypeTreeInfo>
 		{
 			TypeTree.Add(MakeShareable(new FPinTypeTreeInfo(GetCategoryText(PC_Object, true), PC_Object, this, LOCTEXT("ObjectTypeHardReference", "Hard reference to an Object"), true, TypesDatabasePtr)));
 			TypeTree.Add(MakeShareable(new FPinTypeTreeInfo(GetCategoryText(PC_Class, true), PC_Class, this, LOCTEXT("ClassType", "Hard reference to a Class"), true, TypesDatabasePtr)));
-			TypeTree.Add(MakeShareable(new FPinTypeTreeInfo(GetCategoryText(PC_Asset, true), PC_Asset, this, LOCTEXT("AssetType", "Soft reference to an Object"), true, TypesDatabasePtr)));
-			TypeTree.Add(MakeShareable(new FPinTypeTreeInfo(GetCategoryText(PC_AssetClass, true), PC_AssetClass, this, LOCTEXT("AssetClassType", "Soft reference to a Class"), true, TypesDatabasePtr)));
+			TypeTree.Add(MakeShareable(new FPinTypeTreeInfo(GetCategoryText(PC_SoftObject, true), PC_SoftObject, this, LOCTEXT("SoftObjectType", "Soft reference to an Object"), true, TypesDatabasePtr)));
+			TypeTree.Add(MakeShareable(new FPinTypeTreeInfo(GetCategoryText(PC_SoftClass, true), PC_SoftClass, this, LOCTEXT("SoftClassType", "Soft reference to a Class"), true, TypesDatabasePtr)));
 		}
 	}
 	TypeTree.Add( MakeShareable( new FPinTypeTreeInfo(GetCategoryText(PC_Enum, true), PC_Enum, this, LOCTEXT("EnumType", "Enumeration types."), true, TypesDatabasePtr) ) );
@@ -3822,7 +3809,7 @@ void UEdGraphSchema_K2::GetVariableTypeTree(TArray< TSharedPtr<FPinTypeTreeInfo>
 
 bool UEdGraphSchema_K2::DoesTypeHaveSubtypes(const FString& Category) const
 {
-	return (Category == PC_Struct) || (Category == PC_Object) || (Category == PC_Asset) || (Category == PC_AssetClass) || (Category == PC_Interface) || (Category == PC_Class) || (Category == PC_Enum) || (Category == AllObjectTypes);
+	return (Category == PC_Struct) || (Category == PC_Object) || (Category == PC_SoftObject) || (Category == PC_SoftClass) || (Category == PC_Interface) || (Category == PC_Class) || (Category == PC_Enum) || (Category == AllObjectTypes);
 }
 
 struct FWildcardArrayPinHelper
@@ -4094,7 +4081,7 @@ bool UEdGraphSchema_K2::DefaultValueSimpleValidation(const FEdGraphPinType& PinT
 			DVSV_RETURN_MSG(FString::Printf(TEXT("%s isn't a %s (specified on pin %s)"), *NewDefaultObject->GetPathName(), *ObjectClass->GetName(), *(PinName)));
 		}
 	}
-	else if ((PinCategory == PC_Asset) || (PinCategory == PC_AssetClass))
+	else if ((PinCategory == PC_SoftObject) || (PinCategory == PC_SoftClass))
 	{
 		// Should not have an object set, should be converted to string before getting here
 		if (NewDefaultObject)
@@ -4234,8 +4221,8 @@ bool UEdGraphSchema_K2::ArePinTypesCompatible(const FEdGraphPinType& Output, con
 
 			return ExtendedIsChildOf(OutputClass->GetAuthoritativeClass(), InputClass->GetAuthoritativeClass());
 		}
-		else if (((Output.PinCategory == PC_Asset) && (Input.PinCategory == PC_Asset))
-			|| ((Output.PinCategory == PC_AssetClass) && (Input.PinCategory == PC_AssetClass)))
+		else if (((Output.PinCategory == PC_SoftObject) && (Input.PinCategory == PC_SoftObject))
+			|| ((Output.PinCategory == PC_SoftClass) && (Input.PinCategory == PC_SoftClass)))
 		{
 			const UClass* OutputObject = (Output.PinSubCategory == PSC_Self) ? CallingContext : Cast<const UClass>(Output.PinSubCategoryObject.Get());
 			const UClass* InputObject = (Input.PinSubCategory == PSC_Self) ? CallingContext : Cast<const UClass>(Input.PinSubCategoryObject.Get());
@@ -4335,6 +4322,11 @@ bool UEdGraphSchema_K2::ArePinTypesCompatible(const FEdGraphPinType& Output, con
 	}
 	else if (Output.PinCategory == PC_Wildcard || Input.PinCategory == PC_Wildcard)
 	{
+		if (Output.PinCategory == PC_Exec || Input.PinCategory == PC_Exec)
+		{
+			return false;
+		}
+
 		// If this is an Index Wildcard we have to check compatibility for indexing types
 		if (Output.PinSubCategory == PSC_Index)
 		{
@@ -4468,7 +4460,7 @@ void UEdGraphSchema_K2::GetPinDefaultValuesFromString(const FEdGraphPinType& Pin
 		// If this is not a full object path it's a relative path so should be saved as a string
 		if (FPackageName::IsValidObjectPath(ObjectPathLocal))
 		{
-			FStringAssetReference AssetRef = ObjectPathLocal;
+			FSoftObjectPath AssetRef = ObjectPathLocal;
 			UseDefaultValue.Empty();
 			UseDefaultObject = AssetRef.TryLoad();
 			UseDefaultText = FText::GetEmpty();
@@ -4517,7 +4509,7 @@ void UEdGraphSchema_K2::GetPinDefaultValuesFromString(const FEdGraphPinType& Pin
 				UseDefaultValue = EnumPtr->GetNameStringByIndex(0);
 			}
 		}
-		else if ((PinType.PinCategory == PC_Asset) || (PinType.PinCategory == PC_AssetClass))
+		else if ((PinType.PinCategory == PC_SoftObject) || (PinType.PinCategory == PC_SoftClass))
 		{
 			ConstructorHelpers::StripObjectClass(UseDefaultValue);
 		}
@@ -4557,7 +4549,7 @@ void UEdGraphSchema_K2::TrySetDefaultObject(UEdGraphPin& Pin, UObject* NewDefaul
 {
 	FText UseDefaultText;
 
-	if ((Pin.PinType.PinCategory == PC_Asset) || (Pin.PinType.PinCategory == PC_AssetClass))
+	if ((Pin.PinType.PinCategory == PC_SoftObject) || (Pin.PinType.PinCategory == PC_SoftClass))
 	{
 		TrySetDefaultValue(Pin, NewDefaultObject ? NewDefaultObject->GetPathName() : FString());
 		return;
@@ -4937,8 +4929,8 @@ namespace FSetVariableByNameFunctionNames
 	static const FName SetInterfaceName(GET_FUNCTION_NAME_CHECKED(UKismetSystemLibrary, SetInterfacePropertyByName));
 	static const FName SetStringName(GET_FUNCTION_NAME_CHECKED(UKismetSystemLibrary, SetStringPropertyByName));
 	static const FName SetTextName(GET_FUNCTION_NAME_CHECKED(UKismetSystemLibrary, SetTextPropertyByName));
-	static const FName SetAssetName(GET_FUNCTION_NAME_CHECKED(UKismetSystemLibrary, SetAssetPropertyByName));
-	static const FName SetAssetClassName(GET_FUNCTION_NAME_CHECKED(UKismetSystemLibrary, SetAssetClassPropertyByName));
+	static const FName SetSoftObjectName(GET_FUNCTION_NAME_CHECKED(UKismetSystemLibrary, SetSoftObjectPropertyByName));
+	static const FName SetSoftClassName(GET_FUNCTION_NAME_CHECKED(UKismetSystemLibrary, SetSoftClassPropertyByName));
 	static const FName SetNameName(GET_FUNCTION_NAME_CHECKED(UKismetSystemLibrary, SetNamePropertyByName));
 	static const FName SetVectorName(GET_FUNCTION_NAME_CHECKED(UKismetSystemLibrary, SetVectorPropertyByName));
 	static const FName SetRotatorName(GET_FUNCTION_NAME_CHECKED(UKismetSystemLibrary, SetRotatorPropertyByName));
@@ -5017,13 +5009,13 @@ UFunction* UEdGraphSchema_K2::FindSetVariableByNameFunction(const FEdGraphPinTyp
 	{
 		SetFunctionName = FSetVariableByNameFunctionNames::SetTextName;
 	}
-	else if (PinType.PinCategory == UEdGraphSchema_K2::PC_Asset)
+	else if (PinType.PinCategory == UEdGraphSchema_K2::PC_SoftObject)
 	{
-		SetFunctionName = FSetVariableByNameFunctionNames::SetAssetName;
+		SetFunctionName = FSetVariableByNameFunctionNames::SetSoftObjectName;
 	}
-	else if (PinType.PinCategory == UEdGraphSchema_K2::PC_AssetClass)
+	else if (PinType.PinCategory == UEdGraphSchema_K2::PC_SoftClass)
 	{
-		SetFunctionName = FSetVariableByNameFunctionNames::SetAssetClassName;
+		SetFunctionName = FSetVariableByNameFunctionNames::SetSoftClassName;
 	}
 	else if(PinType.PinCategory == UEdGraphSchema_K2::PC_Name)
 	{
@@ -5390,17 +5382,26 @@ void UEdGraphSchema_K2::DroppedAssetsOnGraph(const TArray<FAssetData>& Assets, c
 				AssetClass = BlueprintAsset->GeneratedClass;
 			}
 
-			TSubclassOf<UActorComponent> DestinationComponentType = FComponentAssetBrokerage::GetPrimaryComponentForAsset(AssetClass);
-			if ((DestinationComponentType == NULL) && AssetClass->IsChildOf(AActor::StaticClass()))
+			TSubclassOf<UActorComponent> DestinationComponentType;
+			if (AssetClass->IsChildOf(UActorComponent::StaticClass()) && IsAllowableBlueprintVariableType(AssetClass))
 			{
-				DestinationComponentType = UChildActorComponent::StaticClass();
+				// If it's an actor component subclass that is a BlueprintableComponent, we're good to go
+				DestinationComponentType = AssetClass;
+			}
+			else
+			{
+				// Otherwise see if we can factory a component from the asset
+				DestinationComponentType = FComponentAssetBrokerage::GetPrimaryComponentForAsset(AssetClass);
+				if ((DestinationComponentType == nullptr) && AssetClass->IsChildOf(AActor::StaticClass()))
+				{
+					DestinationComponentType = UChildActorComponent::StaticClass();
+				}
 			}
 
 			// Make sure we have an asset type that's registered with the component list
-			if (DestinationComponentType != NULL)
+			if (DestinationComponentType != nullptr)
 			{
-				UEdGraph* TempOuter = NewObject<UEdGraph>((UObject*)Blueprint);
-				TempOuter->SetFlags(RF_Transient);
+				const FScopedTransaction Transaction(LOCTEXT("CreateAddComponentFromAsset", "Add Component From Asset"));
 
 				FComponentTypeEntry ComponentType = { FString(), FString(), DestinationComponentType };
 
@@ -5473,11 +5474,11 @@ void UEdGraphSchema_K2::GetAssetsPinHoverMessage(const TArray<FAssetData>& Asset
 void UEdGraphSchema_K2::GetAssetsGraphHoverMessage(const TArray<FAssetData>& Assets, const UEdGraph* HoverGraph, FString& OutTooltipText, bool& OutOkIcon) const
 {
 	OutOkIcon = false;
-	OutTooltipText = LOCTEXT("CannotCreateComponentsInNonActorBlueprints", "Cannot create components from assets in a non-Actor blueprint").ToString();
 
 	UBlueprint* Blueprint = FBlueprintEditorUtils::FindBlueprintForGraph(HoverGraph);
 	if ((Blueprint != nullptr) && FBlueprintEditorUtils::IsActorBased(Blueprint))
 	{
+		OutTooltipText = LOCTEXT("UnsupportedAssetTypeForGraphDragDrop", "Cannot create a node from this type of asset").ToString();
 		for (const FAssetData& AssetData : Assets)
 		{
 			if (UObject* Asset = AssetData.GetAsset())
@@ -5488,10 +5489,20 @@ void UEdGraphSchema_K2::GetAssetsGraphHoverMessage(const TArray<FAssetData>& Ass
 					AssetClass = BlueprintAsset->GeneratedClass;
 				}
 
-				TSubclassOf<UActorComponent> DestinationComponentType = FComponentAssetBrokerage::GetPrimaryComponentForAsset(AssetClass);
-				if ((DestinationComponentType == nullptr) && AssetClass->IsChildOf(AActor::StaticClass()))
+				TSubclassOf<UActorComponent> DestinationComponentType;
+				if (AssetClass->IsChildOf(UActorComponent::StaticClass()) && IsAllowableBlueprintVariableType(AssetClass))
 				{
-					DestinationComponentType = UChildActorComponent::StaticClass();
+					// If it's an actor component subclass that is a BlueprintableComponent, we're good to go
+					DestinationComponentType = AssetClass;
+				}
+				else
+				{
+					// Otherwise, see if we have a way to make a component out of the specified asset
+					DestinationComponentType = FComponentAssetBrokerage::GetPrimaryComponentForAsset(AssetClass);
+					if ((DestinationComponentType == nullptr) && AssetClass->IsChildOf(AActor::StaticClass()))
+					{
+						DestinationComponentType = UChildActorComponent::StaticClass();
+					}
 				}
 
 				if (DestinationComponentType != nullptr)
@@ -5503,6 +5514,10 @@ void UEdGraphSchema_K2::GetAssetsGraphHoverMessage(const TArray<FAssetData>& Ass
 			}
 		}
 	}
+	else
+	{
+		OutTooltipText = LOCTEXT("CannotCreateComponentsInNonActorBlueprints", "Cannot create components from assets in a non-Actor blueprint").ToString();
+	}
 }
 
 bool UEdGraphSchema_K2::FadeNodeWhenDraggingOffPin(const UEdGraphNode* Node, const UEdGraphPin* Pin) const
@@ -5969,7 +5984,7 @@ void UEdGraphSchema_K2::BackwardCompatibilityNodeConversion(UEdGraph* Graph, boo
 	}
 }
 
-UEdGraphNode* UEdGraphSchema_K2::CreateSubstituteNode(UEdGraphNode* Node, const UEdGraph* Graph, FObjectInstancingGraph* InstanceGraph, TArray<FName>& InOutExtraNames) const
+UEdGraphNode* UEdGraphSchema_K2::CreateSubstituteNode(UEdGraphNode* Node, const UEdGraph* Graph, FObjectInstancingGraph* InstanceGraph, TSet<FName>& InOutExtraNames) const
 {
 	// If this is an event node, create a unique custom event node as a substitute
 	UK2Node_Event* EventNode = Cast<UK2Node_Event>(Node);
@@ -5992,7 +6007,7 @@ UEdGraphNode* UEdGraphSchema_K2::CreateSubstituteNode(UEdGraphNode* Node, const
 		if(Blueprint && Blueprint->SkeletonGeneratedClass)
 		{
 			// Gather all names in use by the Blueprint class
-			TArray<FName> ExistingNamesInUse = InOutExtraNames;
+			TSet<FName> ExistingNamesInUse = InOutExtraNames;
 			FBlueprintEditorUtils::GetFunctionNameList(Blueprint, ExistingNamesInUse);
 			FBlueprintEditorUtils::GetClassVariableList(Blueprint, ExistingNamesInUse);
 
@@ -6572,7 +6587,7 @@ UK2Node* UEdGraphSchema_K2::CreateSplitPinNode(UEdGraphPin* Pin, FKismetCompiler
 	return SplitPinNode;
 }
 
-void UEdGraphSchema_K2::SplitPin(UEdGraphPin* Pin) const
+void UEdGraphSchema_K2::SplitPin(UEdGraphPin* Pin, const bool bNotify) const
 {
 	// Under some circumstances we can get here when PinSubCategoryObject is not set, so we just can't split the pin in that case
 	UScriptStruct* StructType = Cast<UScriptStruct>(Pin->PinType.PinSubCategoryObject.Get());
@@ -6681,10 +6696,13 @@ void UEdGraphSchema_K2::SplitPin(UEdGraphPin* Pin) const
 
 	GraphNode->Pins.Insert(Pin->SubPins, GraphNode->Pins.Find(Pin) + 1);
 
-	Graph->NotifyGraphChanged();
+	if (bNotify)
+	{
+		Graph->NotifyGraphChanged();
 
-	UBlueprint* Blueprint = FBlueprintEditorUtils::FindBlueprintForGraphChecked(Graph);
-	FBlueprintEditorUtils::MarkBlueprintAsModified(Blueprint);
+		UBlueprint* Blueprint = FBlueprintEditorUtils::FindBlueprintForGraphChecked(Graph);
+		FBlueprintEditorUtils::MarkBlueprintAsModified(Blueprint);
+	}
 }
 
 void UEdGraphSchema_K2::RecombinePin(UEdGraphPin* Pin) const
@@ -6715,40 +6733,40 @@ void UEdGraphSchema_K2::RecombinePin(UEdGraphPin* Pin) const
 
 	if (Pin->Direction == EGPD_Input)
 	{
-		UScriptStruct* StructType = CastChecked<UScriptStruct>(ParentPin->PinType.PinSubCategoryObject.Get());
-
-		TArray<FString> OriginalDefaults;
-		if (StructType == TBaseStructure<FVector>::Get())
-		{
-			ParentPin->DefaultValue = ParentPin->SubPins[0]->DefaultValue + TEXT(",") 
-									+ ParentPin->SubPins[1]->DefaultValue + TEXT(",")
-									+ ParentPin->SubPins[2]->DefaultValue;
-		}
-		else if (StructType == TBaseStructure<FRotator>::Get())
+		if (UScriptStruct* StructType = Cast<UScriptStruct>(ParentPin->PinType.PinSubCategoryObject.Get()))
 		{
-			// Our pins are in the form X,Y,Z but the Rotator pin type expects the form Y,Z,X
-			// so we need to make sure they are added in that order here
-			ParentPin->DefaultValue = ParentPin->SubPins[1]->DefaultValue + TEXT(",")
-									+ ParentPin->SubPins[2]->DefaultValue + TEXT(",")
-									+ ParentPin->SubPins[0]->DefaultValue;
-		}
-		else if (StructType == TBaseStructure<FVector2D>::Get())
-		{
-			FVector2D V2D;
-			V2D.X = FCString::Atof(*ParentPin->SubPins[0]->DefaultValue);
-			V2D.Y = FCString::Atof(*ParentPin->SubPins[1]->DefaultValue);
+			if (StructType == TBaseStructure<FVector>::Get())
+			{
+				ParentPin->DefaultValue = ParentPin->SubPins[0]->DefaultValue + TEXT(",") 
+										+ ParentPin->SubPins[1]->DefaultValue + TEXT(",")
+										+ ParentPin->SubPins[2]->DefaultValue;
+			}
+			else if (StructType == TBaseStructure<FRotator>::Get())
+			{
+				// Our pins are in the form X,Y,Z but the Rotator pin type expects the form Y,Z,X
+				// so we need to make sure they are added in that order here
+				ParentPin->DefaultValue = ParentPin->SubPins[1]->DefaultValue + TEXT(",")
+										+ ParentPin->SubPins[2]->DefaultValue + TEXT(",")
+										+ ParentPin->SubPins[0]->DefaultValue;
+			}
+			else if (StructType == TBaseStructure<FVector2D>::Get())
+			{
+				FVector2D V2D;
+				V2D.X = FCString::Atof(*ParentPin->SubPins[0]->DefaultValue);
+				V2D.Y = FCString::Atof(*ParentPin->SubPins[1]->DefaultValue);
 			
-			ParentPin->DefaultValue = V2D.ToString();
-		}
-		else if (StructType == TBaseStructure<FLinearColor>::Get())
-		{
-			FLinearColor LC;
-			LC.R = FCString::Atof(*ParentPin->SubPins[0]->DefaultValue);
-			LC.G = FCString::Atof(*ParentPin->SubPins[1]->DefaultValue);
-			LC.B = FCString::Atof(*ParentPin->SubPins[2]->DefaultValue);
-			LC.A = FCString::Atof(*ParentPin->SubPins[3]->DefaultValue);
+				ParentPin->DefaultValue = V2D.ToString();
+			}
+			else if (StructType == TBaseStructure<FLinearColor>::Get())
+			{
+				FLinearColor LC;
+				LC.R = FCString::Atof(*ParentPin->SubPins[0]->DefaultValue);
+				LC.G = FCString::Atof(*ParentPin->SubPins[1]->DefaultValue);
+				LC.B = FCString::Atof(*ParentPin->SubPins[2]->DefaultValue);
+				LC.A = FCString::Atof(*ParentPin->SubPins[3]->DefaultValue);
 
-			ParentPin->DefaultValue = LC.ToString();
+				ParentPin->DefaultValue = LC.ToString();
+			}
 		}
 	}
 
diff --git a/Engine/Source/Editor/BlueprintGraph/Private/EdGraphSchema_K2_Actions.cpp b/Engine/Source/Editor/BlueprintGraph/Private/EdGraphSchema_K2_Actions.cpp
--- a/Engine/Source/Editor/BlueprintGraph/Private/EdGraphSchema_K2_Actions.cpp
+++ b/Engine/Source/Editor/BlueprintGraph/Private/EdGraphSchema_K2_Actions.cpp
@@ -37,6 +37,148 @@ namespace
 	const float NodeLiteralHeight = 48.0f;
 }
 
+/////////////////////////////////////////////////////
+// FEdGraphSchemaAction_BlueprintVariableBase
+
+void FEdGraphSchemaAction_BlueprintVariableBase::MovePersistentItemToCategory(const FText& NewCategoryName)
+{
+	FBlueprintEditorUtils::SetBlueprintVariableCategory(GetSourceBlueprint(), VarName, GetVariableScope(), NewCategoryName);
+}
+
+int32 FEdGraphSchemaAction_BlueprintVariableBase::GetReorderIndexInContainer() const
+{
+	if (UBlueprint* SourceBlueprint = GetSourceBlueprint())
+	{
+		return FBlueprintEditorUtils::FindNewVariableIndex(SourceBlueprint, VarName);
+	}
+
+	return INDEX_NONE;
+}
+
+bool FEdGraphSchemaAction_BlueprintVariableBase::ReorderToBeforeAction(TSharedRef<FEdGraphSchemaAction> OtherAction)
+{
+	if ((OtherAction->GetTypeId() == StaticGetTypeId()) && (OtherAction->GetPersistentItemDefiningObject() == GetPersistentItemDefiningObject()))
+	{
+		FEdGraphSchemaAction_BlueprintVariableBase* VarAction = (FEdGraphSchemaAction_BlueprintVariableBase*)&OtherAction.Get();
+
+		// Only let you drag and drop if variables are from same BP class, and not onto itself
+		UBlueprint* BP = GetSourceBlueprint();
+		FName TargetVarName = VarAction->GetVariableName();
+		if ((BP != nullptr) && (VarName != TargetVarName) && (VariableSource == VarAction->GetVariableClass()))
+		{
+			if (FBlueprintEditorUtils::MoveVariableBeforeVariable(BP, VarName, TargetVarName, true))
+			{
+				// Change category of var to match the one we dragged on to as well
+				FText TargetVarCategory = FBlueprintEditorUtils::GetBlueprintVariableCategory(BP, TargetVarName, GetVariableScope());
+				MovePersistentItemToCategory(TargetVarCategory);
+
+				// Update Blueprint after changes so they reflect in My Blueprint tab.
+				FBlueprintEditorUtils::MarkBlueprintAsStructurallyModified(BP);
+
+				return true;
+			}
+		}
+	}
+
+	return false;
+}
+
+
+FEdGraphSchemaActionDefiningObject FEdGraphSchemaAction_BlueprintVariableBase::GetPersistentItemDefiningObject() const
+{
+	UObject* DefiningObject = GetSourceBlueprint();
+	if (UProperty* Prop = GetProperty())
+	{
+		DefiningObject = Prop->GetOwnerStruct();
+	}
+	return FEdGraphSchemaActionDefiningObject(DefiningObject);
+}
+
+UBlueprint* FEdGraphSchemaAction_BlueprintVariableBase::GetSourceBlueprint() const
+{
+	UClass* ClassToCheck = GetVariableClass();
+	if (ClassToCheck == nullptr)
+	{
+		if (UFunction* Function = Cast<UFunction>(GetVariableScope()))
+		{
+			ClassToCheck = Function->GetOuterUClass();
+		}
+	}
+	return UBlueprint::GetBlueprintFromClass(ClassToCheck);
+}
+
+/////////////////////////////////////////////////////
+// FEdGraphSchemaAction_K2Graph
+
+void FEdGraphSchemaAction_K2Graph::MovePersistentItemToCategory(const FText& NewCategoryName)
+{
+	if ((GraphType == EEdGraphSchemaAction_K2Graph::Function) || (GraphType == EEdGraphSchemaAction_K2Graph::Macro))
+	{
+		FBlueprintEditorUtils::SetBlueprintFunctionOrMacroCategory(EdGraph, NewCategoryName);
+	}
+}
+
+int32 FEdGraphSchemaAction_K2Graph::GetReorderIndexInContainer() const
+{
+	return FBlueprintEditorUtils::FindIndexOfGraphInParent(EdGraph);
+}
+
+bool FEdGraphSchemaAction_K2Graph::ReorderToBeforeAction(TSharedRef<FEdGraphSchemaAction> OtherAction)
+{
+	if ((OtherAction->GetTypeId() == GetTypeId()) && (OtherAction->GetPersistentItemDefiningObject() == GetPersistentItemDefiningObject()))
+	{
+		const int32 OldIndex = GetReorderIndexInContainer();
+		const int32 NewIndexToGoBefore = OtherAction->GetReorderIndexInContainer();
+
+		if ((OldIndex != INDEX_NONE) && (OldIndex != NewIndexToGoBefore))
+		{
+			if (FBlueprintEditorUtils::MoveGraphBeforeOtherGraph(EdGraph, NewIndexToGoBefore, true))
+			{
+				// Change category to match the one we dragged on to as well
+				MovePersistentItemToCategory(OtherAction->GetCategory());
+
+				return true;
+			}
+		}
+	}
+
+	return false;
+}
+
+FEdGraphSchemaActionDefiningObject FEdGraphSchemaAction_K2Graph::GetPersistentItemDefiningObject() const
+{
+	UObject* DefiningObject = GetSourceBlueprint();
+	if (UFunction* Func = GetFunction())
+	{
+		DefiningObject = Func->GetOwnerStruct();
+	}
+	return FEdGraphSchemaActionDefiningObject(DefiningObject, (void*)GraphType);
+}
+
+UBlueprint* FEdGraphSchemaAction_K2Graph::GetSourceBlueprint() const
+{
+	return FBlueprintEditorUtils::FindBlueprintForGraph(EdGraph);
+}
+
+UFunction* FEdGraphSchemaAction_K2Graph::GetFunction() const
+{
+	if (GraphType == EEdGraphSchemaAction_K2Graph::Function)
+	{
+		if (UBlueprint* SourceBlueprint = GetSourceBlueprint())
+		{
+			if (FuncName != NAME_None)
+			{
+				return FindField<UFunction>(SourceBlueprint->SkeletonGeneratedClass, FuncName);
+			}
+		}
+	}
+
+	return nullptr;
+}
+
+/////////////////////////////////////////////////////
+// FEdGraphSchemaAction_K2Delegate
+
 /////////////////////////////////////////////////////
 // FEdGraphSchemaAction_K2ViewNode
 
@@ -87,10 +229,10 @@ UEdGraphNode* FEdGraphSchemaAction_K2NewNode::CreateNode(class UEdGraph* ParentG
 
 UEdGraphNode* FEdGraphSchemaAction_K2NewNode::PerformAction(class UEdGraph* ParentGraph, UEdGraphPin* FromPin, const FVector2D Location, bool bSelectNewNode/* = true*/)
 {
-	UEdGraphNode* ResultNode = NULL;
+	UEdGraphNode* ResultNode = nullptr;
 
 	// If there is a template, we actually use it
-	if (NodeTemplate != NULL)
+	if (NodeTemplate != nullptr)
 	{
 		const FScopedTransaction Transaction( NSLOCTEXT("UnrealEd", "K2_AddNode", "Add Node") );
 		ParentGraph->Modify();
@@ -131,7 +273,7 @@ UEdGraphNode* FEdGraphSchemaAction_K2NewNode::PerformAction(class UEdGraph* Pare
 
 UEdGraphNode* FEdGraphSchemaAction_K2NewNode::PerformAction(class UEdGraph* ParentGraph, TArray<UEdGraphPin*>& FromPins, const FVector2D Location, bool bSelectNewNode/* = true*/) 
 {
-	UEdGraphNode* ResultNode = NULL;
+	UEdGraphNode* ResultNode = nullptr;
 
 	if (FromPins.Num() > 0)
 	{
@@ -145,7 +287,7 @@ UEdGraphNode* FEdGraphSchemaAction_K2NewNode::PerformAction(class UEdGraph* Pare
 	}
 	else
 	{
-		ResultNode = PerformAction(ParentGraph, NULL, Location, bSelectNewNode);
+		ResultNode = PerformAction(ParentGraph, nullptr, Location, bSelectNewNode);
 	}
 
 	return ResultNode;
@@ -171,13 +313,13 @@ UEdGraphNode* FEdGraphSchemaAction_K2ViewNode::PerformAction(class UEdGraph* Par
 		FKismetEditorUtilities::BringKismetToFocusAttentionOnObject(NodePtr);
 	}
 
-	return NULL;
+	return nullptr;
 }
 
 UEdGraphNode* FEdGraphSchemaAction_K2ViewNode::PerformAction(class UEdGraph* ParentGraph, TArray<UEdGraphPin*>& FromPins, const FVector2D Location, bool bSelectNewNode/* = true*/) 
 {
-	PerformAction(ParentGraph, NULL, Location, bSelectNewNode);
-	return NULL;
+	PerformAction(ParentGraph, nullptr, Location, bSelectNewNode);
+	return nullptr;
 }
 
 /////////////////////////////////////////////////////
@@ -185,7 +327,7 @@ UEdGraphNode* FEdGraphSchemaAction_K2ViewNode::PerformAction(class UEdGraph* Par
 
 UEdGraphNode* FEdGraphSchemaAction_K2AssignDelegate::AssignDelegate(class UK2Node* NodeTemplate, class UEdGraph* ParentGraph, UEdGraphPin* FromPin, const FVector2D Location, bool bSelectNewNode)
 {
-	UK2Node_AddDelegate* BindNode = NULL;
+	UK2Node_AddDelegate* BindNode = nullptr;
 	if (UK2Node_AddDelegate* AddDelegateTemplate = Cast<UK2Node_AddDelegate>(NodeTemplate))
 	{
 		const FScopedTransaction Transaction( NSLOCTEXT("UnrealEd", "K2_AddNode", "Add Node") );
@@ -196,7 +338,7 @@ UEdGraphNode* FEdGraphSchemaAction_K2AssignDelegate::AssignDelegate(class UK2Nod
 		}
 
 		BindNode = Cast<UK2Node_AddDelegate>(CreateNode(ParentGraph, FromPin, Location, NodeTemplate, bSelectNewNode));
-		UMulticastDelegateProperty* DelegateProperty = BindNode ? Cast<UMulticastDelegateProperty>(BindNode->GetProperty()) : NULL;
+		UMulticastDelegateProperty* DelegateProperty = BindNode ? Cast<UMulticastDelegateProperty>(BindNode->GetProperty()) : nullptr;
 		if(DelegateProperty)
 		{
 			const FString FunctionName = FString::Printf(TEXT("%s_Event"), *DelegateProperty->GetName());
@@ -228,15 +370,15 @@ UEdGraphNode* FEdGraphSchemaAction_K2AssignDelegate::PerformAction(class UEdGrap
 
 UEdGraphNode* FEdGraphSchemaAction_EventFromFunction::PerformAction(class UEdGraph* ParentGraph, UEdGraphPin* FromPin, const FVector2D Location, bool bSelectNewNode/* = true*/)
 {
-	UK2Node_CustomEvent* EventNode = NULL;
+	UK2Node_CustomEvent* EventNode = nullptr;
 	if (SignatureFunction)
 	{
 		if (FromPin)
 		{
 			// Make sure, that function is latest, so the names of parameters are proper.
 			UK2Node_BaseMCDelegate* MCDelegateNode = Cast<UK2Node_BaseMCDelegate>(FromPin->GetOwningNode());
-			UEdGraphPin* InputDelegatePin = MCDelegateNode ? MCDelegateNode->GetDelegatePin() : NULL;
-			UFunction* OriginalFunction = MCDelegateNode ? MCDelegateNode->GetDelegateSignature() : NULL;
+			UEdGraphPin* InputDelegatePin = MCDelegateNode ? MCDelegateNode->GetDelegatePin() : nullptr;
+			UFunction* OriginalFunction = MCDelegateNode ? MCDelegateNode->GetDelegateSignature() : nullptr;
 			if (OriginalFunction && 
 				(OriginalFunction != SignatureFunction) && 
 				(FromPin == InputDelegatePin) &&
@@ -264,7 +406,7 @@ UEdGraphNode* FEdGraphSchemaAction_EventFromFunction::PerformAction(class UEdGra
 
 UEdGraphNode* FEdGraphSchemaAction_EventFromFunction::PerformAction(class UEdGraph* ParentGraph, TArray<UEdGraphPin*>& FromPins, const FVector2D Location, bool bSelectNewNode/* = true*/)
 {
-	UEdGraphNode* ResultNode = NULL;
+	UEdGraphNode* ResultNode = nullptr;
 
 	if (FromPins.Num() > 0)
 	{
@@ -278,7 +420,7 @@ UEdGraphNode* FEdGraphSchemaAction_EventFromFunction::PerformAction(class UEdGra
 	}
 	else
 	{
-		ResultNode = PerformAction(ParentGraph, NULL, Location, bSelectNewNode);
+		ResultNode = PerformAction(ParentGraph, nullptr, Location, bSelectNewNode);
 	}
 
 	return ResultNode;
@@ -297,18 +439,18 @@ void FEdGraphSchemaAction_EventFromFunction::AddReferencedObjects(FReferenceColl
 
 UEdGraphNode* FEdGraphSchemaAction_K2AddComponent::PerformAction(class UEdGraph* ParentGraph, UEdGraphPin* FromPin, const FVector2D Location, bool bSelectNewNode/* = true*/)
 {
-	if (ComponentClass == NULL)
+	if (ComponentClass == nullptr)
 	{
-		return NULL;
+		return nullptr;
 	}
 
 	UBlueprint* Blueprint = FBlueprintEditorUtils::FindBlueprintForGraphChecked(ParentGraph);
 	UEdGraphNode* NewNode = FEdGraphSchemaAction_K2NewNode::PerformAction(ParentGraph, FromPin, Location, bSelectNewNode);
-	if ((NewNode != NULL) && (Blueprint != NULL))
+	if ((NewNode != nullptr) && (Blueprint != nullptr))
 	{
 		UK2Node_AddComponent* AddCompNode = CastChecked<UK2Node_AddComponent>(NewNode);
 
-		ensure(NULL != Cast<UBlueprintGeneratedClass>(Blueprint->GeneratedClass));
+		ensure(nullptr != Cast<UBlueprintGeneratedClass>(Blueprint->GeneratedClass));
 		// Then create a new template object, and add to array in
 		UActorComponent* NewTemplate = NewObject<UActorComponent>(Blueprint->GeneratedClass, ComponentClass, NAME_None, RF_ArchetypeObject | RF_Public);
 		Blueprint->ComponentTemplates.Add(NewTemplate);
@@ -328,7 +470,7 @@ UEdGraphNode* FEdGraphSchemaAction_K2AddComponent::PerformAction(class UEdGraph*
 		}
 
 		// Set the asset
-		if(ComponentAsset != NULL)
+		if(ComponentAsset != nullptr)
 		{
 			FComponentAssetBrokerage::AssignAssetToComponent(NewTemplate, ComponentAsset);
 		}
@@ -349,46 +491,20 @@ void FEdGraphSchemaAction_K2AddComponent::AddReferencedObjects( FReferenceCollec
 	Collector.AddReferencedObject( ComponentAsset );
 }
 
-/////////////////////////////////////////////////////
-// FEdGraphSchemaAction_K2AddTimeline
-
-UEdGraphNode* FEdGraphSchemaAction_K2AddTimeline::PerformAction(class UEdGraph* ParentGraph, UEdGraphPin* FromPin, const FVector2D Location, bool bSelectNewNode/* = true*/)
-{
-	const FScopedTransaction Transaction( NSLOCTEXT("UnrealEd", "K2_AddTimeline", "Add Timeline") );
-	UBlueprint* Blueprint = FBlueprintEditorUtils::FindBlueprintForGraphChecked(ParentGraph);
-	UEdGraphNode* NewNode = FEdGraphSchemaAction_K2NewNode::PerformAction(ParentGraph, FromPin, Location, bSelectNewNode);
-
-	// Set the name to be a unique timeline name, so the generated node will have a default name that is already validated
-	UK2Node_Timeline* TimelineNode = CastChecked<UK2Node_Timeline>(NewNode);
-	TimelineNode->TimelineName = FBlueprintEditorUtils::FindUniqueTimelineName(Blueprint);
-
-	if(Blueprint != NULL)
-	{
-		if (FBlueprintEditorUtils::AddNewTimeline(Blueprint, TimelineNode->TimelineName) != NULL)
-		{
-			// Clear off any existing error message now the timeline has been added
-			TimelineNode->ErrorMsg.Empty();
-			TimelineNode->bHasCompilerMessage = false;
-		}
-	}
-
-	return NewNode;
-}
-
 /////////////////////////////////////////////////////
 // FEdGraphSchemaAction_K2AddEvent
 
 UEdGraphNode* FEdGraphSchemaAction_K2AddEvent::PerformAction(class UEdGraph* ParentGraph, UEdGraphPin* FromPin, const FVector2D Location, bool bSelectNewNode/* = true*/)
 {
-	UEdGraphNode* NewNode = NULL;
+	UEdGraphNode* NewNode = nullptr;
 	const FScopedTransaction Transaction( NSLOCTEXT("UnrealEd", "K2_Event", "Add Event") );	
 
 	UBlueprint const* Blueprint = FBlueprintEditorUtils::FindBlueprintForGraphChecked(ParentGraph);
 
-	UK2Node_Event const* ExistingEvent = NULL;
+	UK2Node_Event const* ExistingEvent = nullptr;
 	if (EventHasAlreadyBeenPlaced(Blueprint, &ExistingEvent))
 	{
-		check(ExistingEvent != NULL);
+		check(ExistingEvent != nullptr);
 		FKismetEditorUtilities::BringKismetToFocusAttentionOnObject(ExistingEvent);
 	}
 	else
@@ -399,22 +515,22 @@ UEdGraphNode* FEdGraphSchemaAction_K2AddEvent::PerformAction(class UEdGraph* Par
 	return NewNode;
 }
 
-bool FEdGraphSchemaAction_K2AddEvent::EventHasAlreadyBeenPlaced(UBlueprint const* Blueprint, UK2Node_Event const** FoundEventOut /*= NULL*/) const
+bool FEdGraphSchemaAction_K2AddEvent::EventHasAlreadyBeenPlaced(UBlueprint const* Blueprint, UK2Node_Event const** FoundEventOut /*= nullptr*/) const
 {
-	UK2Node_Event* ExistingEvent = NULL;
+	UK2Node_Event* ExistingEvent = nullptr;
 
-	if (Blueprint != NULL)
+	if (Blueprint != nullptr)
 	{
 		UK2Node_Event const* EventTemplate = Cast<UK2Node_Event const>(NodeTemplate);
 		ExistingEvent = FBlueprintEditorUtils::FindOverrideForFunction(Blueprint, EventTemplate->EventReference.GetMemberParentClass(EventTemplate->GetBlueprintClassFromNode()), EventTemplate->EventReference.GetMemberName());
 	}
 
-	if (FoundEventOut != NULL)
+	if (FoundEventOut != nullptr)
 	{
 		*FoundEventOut = ExistingEvent;
 	}
 
-	return (ExistingEvent != NULL);
+	return (ExistingEvent != nullptr);
 }
 
 /////////////////////////////////////////////////////
@@ -464,7 +580,7 @@ UEdGraphNode* FEdGraphSchemaAction_K2AddCallOnActor::PerformAction(class UEdGrap
 	{
 		AActor* LevelActor = LevelActors[ActorIndex];
 
-		if(LevelActor != NULL)
+		if(LevelActor != nullptr)
 		{
 			UK2Node_Literal* LiteralNode =  NewObject<UK2Node_Literal>(ParentGraph);
 			ParentGraph->AddNode(LiteralNode, false, bSelectNewNode);
@@ -483,7 +599,7 @@ UEdGraphNode* FEdGraphSchemaAction_K2AddCallOnActor::PerformAction(class UEdGrap
 			// Connect the literal out to the self of the call
 			UEdGraphPin* LiteralOutput = LiteralNode->GetValuePin();
 			UEdGraphPin* CallSelfInput = CallNode->FindPin(K2Schema->PN_Self);
-			if(LiteralOutput != NULL && CallSelfInput != NULL)
+			if(LiteralOutput != nullptr && CallSelfInput != nullptr)
 			{
 				LiteralOutput->MakeLinkTo(CallSelfInput);
 			}
@@ -513,7 +629,7 @@ UEdGraphNode* FEdGraphSchemaAction_K2AddComment::PerformAction(class UEdGraph* P
 	FVector2D SpawnLocation = Location;
 
 	FSlateRect Bounds;
-	if ((Blueprint != NULL) && FKismetEditorUtilities::GetBoundsForSelectedNodes(Blueprint, Bounds, 50.0f))
+	if ((Blueprint != nullptr) && FKismetEditorUtilities::GetBoundsForSelectedNodes(Blueprint, Bounds, 50.0f))
 	{
 		CommentTemplate->SetBounds(Bounds);
 		SpawnLocation.X = CommentTemplate->NodePosX;
@@ -538,7 +654,7 @@ UEdGraphNode* FEdGraphSchemaAction_K2AddComment::PerformAction(class UEdGraph* P
 UEdGraphNode* FEdGraphSchemaAction_K2TargetNode::PerformAction( class UEdGraph* ParentGraph, UEdGraphPin* FromPin, const FVector2D Location, bool bSelectNewNode/* = true*/ ) 
 {
 	FKismetEditorUtilities::BringKismetToFocusAttentionOnObject(NodeTemplate);
-	return NULL;
+	return nullptr;
 }
 
 /////////////////////////////////////////////////////
@@ -547,6 +663,6 @@ UEdGraphNode* FEdGraphSchemaAction_K2TargetNode::PerformAction( class UEdGraph*
 UEdGraphNode* FEdGraphSchemaAction_K2PasteHere::PerformAction( class UEdGraph* ParentGraph, UEdGraphPin* FromPin, const FVector2D Location, bool bSelectNewNode/* = true*/ ) 
 {
 	FKismetEditorUtilities::PasteNodesHere(ParentGraph, Location);
-	return NULL;
+	return nullptr;
 }
 
diff --git a/Engine/Source/Editor/BlueprintGraph/Private/K2Node.cpp b/Engine/Source/Editor/BlueprintGraph/Private/K2Node.cpp
--- a/Engine/Source/Editor/BlueprintGraph/Private/K2Node.cpp
+++ b/Engine/Source/Editor/BlueprintGraph/Private/K2Node.cpp
@@ -97,9 +97,9 @@ void UK2Node::Serialize(FArchive& Ar)
 					}
 				}
 
-				if (Pin->PinType.PinCategory == UEdGraphSchema_K2::PC_Asset || Pin->PinType.PinCategory == UEdGraphSchema_K2::PC_AssetClass)
+				if (Pin->PinType.PinCategory == UEdGraphSchema_K2::PC_SoftObject || Pin->PinType.PinCategory == UEdGraphSchema_K2::PC_SoftClass)
 				{
-					FStringAssetReference TempRef(Pin->DefaultValue);
+					FSoftObjectPath TempRef(Pin->DefaultValue);
 
 					// Serialize the asset reference, this will do the save fixup. It won't actually serialize the string if this is a real archive like linkersave
 					TempRef.SerializePath(Ar, true);
@@ -166,7 +166,7 @@ void UK2Node::FixupPinDefaultValues()
 		{
 			UEdGraphPin* Pin = Pins[i];
 
-			if (Pin->PinType.PinCategory == K2Schema->PC_Asset || Pin->PinType.PinCategory == K2Schema->PC_AssetClass)
+			if (Pin->PinType.PinCategory == K2Schema->PC_SoftObject || Pin->PinType.PinCategory == K2Schema->PC_SoftClass)
 			{
 				if (Pin->DefaultObject && Pin->DefaultValue.IsEmpty())
 				{
@@ -264,33 +264,36 @@ void UK2Node::AutowireNewNode(UEdGraphPin* FromPin)
 			}
 
 			ECanCreateConnectionResponse ConnectResponse = K2Schema->CanCreateConnection(FromPin, Pin).Response;
-			if (ConnectResponse == ECanCreateConnectionResponse::CONNECT_RESPONSE_MAKE)
+			if ((FromPin->PinType.PinCategory == K2Schema->PC_Exec) && (ConnectResponse == ECanCreateConnectionResponse::CONNECT_RESPONSE_BREAK_OTHERS_A))
 			{
-				if (K2Schema->TryCreateConnection(FromPin, Pin))
-				{
-					NodeList.Add(FromPin->GetOwningNode());
-					NodeList.Add(this);
-				}
+				InsertNewNode(FromPin, Pin, NodeList);
 
 				// null out the backup connection (so we don't attempt to make it 
 				// once we exit the loop... we successfully made this connection!)
 				BackupConnection = NULL;
 				break;
 			}
-			else if((FromPin->PinType.PinCategory == K2Schema->PC_Exec) && (ConnectResponse == ECanCreateConnectionResponse::CONNECT_RESPONSE_BREAK_OTHERS_A))
+			else if ((BackupConnection == NULL) && (ConnectResponse == ECanCreateConnectionResponse::CONNECT_RESPONSE_MAKE_WITH_CONVERSION_NODE))
 			{
-				InsertNewNode(FromPin, Pin, NodeList);
+				// save this off, in-case we don't make any connection at all
+				BackupConnection = Pin;
+			}
+			else if ((ConnectResponse == ECanCreateConnectionResponse::CONNECT_RESPONSE_MAKE) ||
+				(ConnectResponse == ECanCreateConnectionResponse::CONNECT_RESPONSE_BREAK_OTHERS_A) ||
+				(ConnectResponse == ECanCreateConnectionResponse::CONNECT_RESPONSE_BREAK_OTHERS_B) ||
+				(ConnectResponse == ECanCreateConnectionResponse::CONNECT_RESPONSE_BREAK_OTHERS_AB))
+			{
+				if (K2Schema->TryCreateConnection(FromPin, Pin))
+				{
+					NodeList.Add(FromPin->GetOwningNode());
+					NodeList.Add(this);
+				}
 
 				// null out the backup connection (so we don't attempt to make it 
 				// once we exit the loop... we successfully made this connection!)
 				BackupConnection = NULL;
 				break;
 			}
-			else if ((BackupConnection == NULL) && (ConnectResponse == ECanCreateConnectionResponse::CONNECT_RESPONSE_MAKE_WITH_CONVERSION_NODE))
-			{
-				// save this off, in-case we don't make any connection at all
-				BackupConnection = Pin;
-			}
 		}
 
 		// if we didn't find an ideal connection, then lets connect this pin to 
@@ -461,6 +464,19 @@ UObject* UK2Node::GetJumpTargetForDoubleClick() const
     return GetReferencedLevelActor();
 }
 
+bool UK2Node::CanJumpToDefinition() const
+{
+	return GetJumpTargetForDoubleClick() != nullptr;
+}
+
+void UK2Node::JumpToDefinition() const
+{
+	if (UObject* HyperlinkTarget = GetJumpTargetForDoubleClick())
+	{
+		FKismetEditorUtilities::BringKismetToFocusAttentionOnObject(HyperlinkTarget);
+	}
+}
+
 void UK2Node::ReallocatePinsDuringReconstruction(TArray<UEdGraphPin*>& OldPins)
 {
 	AllocateDefaultPins();
@@ -611,14 +627,21 @@ void UK2Node::RestoreSplitPins(TArray<UEdGraphPin*>& OldPins)
 			// find the new pin that corresponds to parent, and split it if it isn't already split
 			for (UEdGraphPin* NewPin : Pins)
 			{
-				if (FCString::Stricmp(*(NewPin->PinName), *(OldPin->ParentPin->PinName)) == 0)
+				// The pin we're searching for has the same direction, is not a container, has the same name as our parent pin (TODO: does this handle redirects?), and is either a wildcard or a struct
+				// We allow sub categories of struct to change because it may be changing to a type that has the same members
+				if ((NewPin->Direction == OldPin->Direction) && !NewPin->PinType.IsContainer() && (NewPin->PinName == OldPin->ParentPin->PinName) 
+					&& (NewPin->PinType.PinCategory == UEdGraphSchema_K2::PC_Wildcard || NewPin->PinType.PinCategory == UEdGraphSchema_K2::PC_Struct))
 				{
 					// Make sure we're not dealing with a menu node
 					UEdGraph* OuterGraph = GetGraph();
 					if (OuterGraph && OuterGraph->Schema && NewPin->SubPins.Num() == 0)
 					{
-						NewPin->PinType = OldPin->ParentPin->PinType;
-						GetSchema()->SplitPin(NewPin);
+						if (NewPin->PinType.PinCategory == UEdGraphSchema_K2::PC_Wildcard)
+						{
+							NewPin->PinType = OldPin->ParentPin->PinType;
+						}
+						
+						GetSchema()->SplitPin(NewPin, false);
 						break;
 					}
 				}
@@ -1382,21 +1405,45 @@ UEdGraphPin* UK2Node::GetExecPin() const
 UEdGraphPin* UK2Node::GetPassThroughPin(const UEdGraphPin* FromPin) const
 {
 	UEdGraphPin* PassThroughPin = nullptr;
-	if(FromPin && Pins.Contains(FromPin))
+
+	if (FromPin && UEdGraphSchema_K2::IsExecPin(*FromPin))
 	{
-		const UEdGraphSchema_K2* K2Schema = GetDefault<UEdGraphSchema_K2>();
-		if(K2Schema->IsExecPin(*FromPin))
+		// We only allow execution passing if there is exactly one input and one output exec pin otherwise there is
+		// ambiguity (e.g., on a branch or sequence node or timeline) so we want no passthrough and commenting it out will kill subsequent code
+		bool bFoundFromPin = false;
+		int32 NumInputs = 0;
+		int32 NumOutputs = 0;
+		UEdGraphPin* PotentialResult = nullptr;
+
+		for (UEdGraphPin* Pin : Pins)
 		{
-			// Locate the first exec pin that's opposite the given exec pin, if any.
-			for(UEdGraphPin* Pin : Pins)
+			if (Pin == FromPin)
+			{
+				bFoundFromPin = true;
+			}
+
+			if (UEdGraphSchema_K2::IsExecPin(*Pin))
 			{
-				if(Pin && Pin != FromPin && Pin->Direction != FromPin->Direction && K2Schema->IsExecPin(*Pin))
+				if (Pin->Direction == EGPD_Input)
 				{
-					PassThroughPin = Pin;
-					break;
+					++NumInputs;
+				}
+				else
+				{
+					++NumOutputs;
+				}
+
+				if (Pin->Direction != FromPin->Direction)
+				{
+					PotentialResult = Pin;
 				}
 			}
 		}
+
+		if ((NumInputs == 1) && (NumOutputs == 1) && bFoundFromPin)
+		{
+			PassThroughPin = PotentialResult;
+		}
 	}
 
 	return PassThroughPin;
diff --git a/Engine/Source/Editor/BlueprintGraph/Private/K2Node_ActorBoundEvent.cpp b/Engine/Source/Editor/BlueprintGraph/Private/K2Node_ActorBoundEvent.cpp
--- a/Engine/Source/Editor/BlueprintGraph/Private/K2Node_ActorBoundEvent.cpp
+++ b/Engine/Source/Editor/BlueprintGraph/Private/K2Node_ActorBoundEvent.cpp
@@ -183,6 +183,38 @@ AActor* UK2Node_ActorBoundEvent::GetReferencedLevelActor() const
 	return EventOwner;
 }
 
+void UK2Node_ActorBoundEvent::ValidateNodeDuringCompilation(class FCompilerResultsLog& MessageLog) const
+{
+	Super::ValidateNodeDuringCompilation(MessageLog);
+
+	// make sure that the actor still exists:
+	AActor* TargetActor = GetReferencedLevelActor();
+	if(!TargetActor)
+	{
+		MessageLog.Warning(
+			*NSLOCTEXT("KismetCompiler", "MissingActor_ActorBoundEvent", "@@ is referencing an Actor that no longer exists. Attached logic will never execute.").ToString(), 
+			this
+		);
+	}
+	else if(DelegateOwnerClass == nullptr)
+	{
+		MessageLog.Warning(
+			*NSLOCTEXT("KismetCompiler", "MissingClass_ActorBoundEvent", "@@ is trying to find an Event Dispatcher named @@ in a class that no longer exists. Attached logic will never execute.").ToString(), 
+			this,
+			*DelegatePropertyName.ToString()
+		);
+	}
+	else if( GetTargetDelegatePropertyFromSkel() == nullptr )
+	{
+		MessageLog.Warning(
+			*NSLOCTEXT("KismetCompiler", "MissingDelegate_ActorBoundEvent", "@@ is referencing an Event Dispatcher named @@ that no longer exists in class @@. Attached logic will never execute.").ToString(),
+			this,
+			*DelegatePropertyName.ToString(),
+			DelegateOwnerClass
+		);
+	}
+}
+
 void UK2Node_ActorBoundEvent::InitializeActorBoundEventParams(AActor* InEventOwner, const UMulticastDelegateProperty* InDelegateProperty)
 {
 	if (InEventOwner && InDelegateProperty)
@@ -203,6 +235,11 @@ UMulticastDelegateProperty* UK2Node_ActorBoundEvent::GetTargetDelegateProperty()
 	return Cast<UMulticastDelegateProperty>(FindField<UMulticastDelegateProperty>(DelegateOwnerClass, DelegatePropertyName));
 }
 
+UMulticastDelegateProperty* UK2Node_ActorBoundEvent::GetTargetDelegatePropertyFromSkel() const
+{
+	return Cast<UMulticastDelegateProperty>(FindField<UMulticastDelegateProperty>(FBlueprintEditorUtils::GetMostUpToDateClass(DelegateOwnerClass), DelegatePropertyName));
+}
+
 FMulticastScriptDelegate* UK2Node_ActorBoundEvent::GetTargetDelegate() const 
 {
 	if( EventOwner )
diff --git a/Engine/Source/Editor/BlueprintGraph/Private/K2Node_CallFunction.cpp b/Engine/Source/Editor/BlueprintGraph/Private/K2Node_CallFunction.cpp
--- a/Engine/Source/Editor/BlueprintGraph/Private/K2Node_CallFunction.cpp
+++ b/Engine/Source/Editor/BlueprintGraph/Private/K2Node_CallFunction.cpp
@@ -32,6 +32,7 @@
 #include "BlueprintActionFilter.h"
 #include "FindInBlueprintManager.h"
 #include "SPinTypeSelector.h"
+#include "SourceCodeNavigation.h"
 
 #define LOCTEXT_NAMESPACE "K2Node"
 
@@ -1896,23 +1897,19 @@ void UK2Node_CallFunction::Serialize(FArchive& Ar)
 		{
 			// Don't validate the enabled state if the user has explicitly set it. Also skip validation if we're just duplicating this node.
 			const bool bIsDuplicating = (Ar.GetPortFlags() & PPF_Duplicate) != 0;
-			if (!bIsDuplicating && !bUserSetEnabledState)
+			if (!bIsDuplicating && !HasUserSetTheEnabledState())
 			{
-				UClass* SelfScope = GetBlueprintClassFromNode();
-				if (!FunctionReference.IsSelfContext() || SelfScope != nullptr)
+				if (const UFunction* Function = GetTargetFunction())
 				{
-					if (const UFunction* Function = FunctionReference.ResolveMember<UFunction>(SelfScope))
+					// Enable as development-only if specified in metadata. This way existing functions that have the metadata added to them will get their enabled state fixed up on load.
+					if (GetDesiredEnabledState() == ENodeEnabledState::Enabled && Function->HasMetaData(FBlueprintMetadata::MD_DevelopmentOnly))
 					{
-						// Enable as development-only if specified in metadata. This way existing functions that have the metadata added to them will get their enabled state fixed up on load.
-						if (EnabledState == ENodeEnabledState::Enabled && Function->HasMetaData(FBlueprintMetadata::MD_DevelopmentOnly))
-						{
-							EnabledState = ENodeEnabledState::DevelopmentOnly;
-						}
-						// Ensure that if the metadata is removed, we also fix up the enabled state to avoid leaving it set as development-only in that case.
-						else if (EnabledState == ENodeEnabledState::DevelopmentOnly && !Function->HasMetaData(FBlueprintMetadata::MD_DevelopmentOnly))
-						{
-							EnabledState = ENodeEnabledState::Enabled;
-						}
+						SetEnabledState(ENodeEnabledState::DevelopmentOnly, /*bUserAction=*/ false);
+					}
+					// Ensure that if the metadata is removed, we also fix up the enabled state to avoid leaving it set as development-only in that case.
+					else if (GetDesiredEnabledState() == ENodeEnabledState::DevelopmentOnly && !Function->HasMetaData(FBlueprintMetadata::MD_DevelopmentOnly))
+					{
+						SetEnabledState(ENodeEnabledState::Enabled, /*bUserAction=*/ false);
 					}
 				}
 			}
@@ -1927,13 +1924,13 @@ void UK2Node_CallFunction::PostPlacedNewNode()
 	// Try re-setting the function given our new parent scope, in case it turns an external to an internal, or vis versa
 	FunctionReference.RefreshGivenNewSelfScope<UFunction>(GetBlueprintClassFromNode());
 
-	// Re-enable for development only if specified in metadata.
-	if(EnabledState == ENodeEnabledState::Enabled && !bUserSetEnabledState)
+	// Set the node to development only if the function specifies that
+	check(!HasUserSetTheEnabledState());
+	if (const UFunction* Function = GetTargetFunction())
 	{
-		const UFunction* Function = GetTargetFunction();
-		if (Function && Function->HasMetaData(FBlueprintMetadata::MD_DevelopmentOnly))
+		if (Function->HasMetaData(FBlueprintMetadata::MD_DevelopmentOnly))
 		{
-			EnabledState = ENodeEnabledState::DevelopmentOnly;
+			SetEnabledState(ENodeEnabledState::DevelopmentOnly, /*bUserAction=*/ false);
 		}
 	}
 }
@@ -2730,6 +2727,44 @@ TSharedPtr<SWidget> UK2Node_CallFunction::CreateNodeImage() const
 	return TSharedPtr<SWidget>();
 }
 
+UObject* UK2Node_CallFunction::GetJumpTargetForDoubleClick() const
+{
+	// If there is an event node, jump to it, otherwise jump to the function graph
+	const UEdGraphNode* ResultEventNode = nullptr;
+	UEdGraph* FunctionGraph = GetFunctionGraph(/*out*/ ResultEventNode);
+	if (ResultEventNode != nullptr)
+	{
+		return const_cast<UEdGraphNode*>(ResultEventNode);
+	}
+	else
+	{
+		return FunctionGraph;
+	}
+}
+
+bool UK2Node_CallFunction::CanJumpToDefinition() const
+{
+	const UFunction* TargetFunction = GetTargetFunction();
+	const bool bNativeFunction = (TargetFunction != nullptr) && (TargetFunction->IsNative());
+	return bNativeFunction || (GetJumpTargetForDoubleClick() != nullptr);
+}
+
+void UK2Node_CallFunction::JumpToDefinition() const
+{
+	// For native functions, try going to the function definition in C++ if available
+	if (UFunction* TargetFunction = GetTargetFunction())
+	{
+		if (TargetFunction->IsNative())
+		{
+			FSourceCodeNavigation::NavigateToFunctionAsync(TargetFunction);
+			return;
+		}
+	}
+
+	// Otherwise, fall back to the inherited behavior which should go to the function entry node
+	Super::JumpToDefinition();
+}
+
 bool UK2Node_CallFunction::IsConnectionDisallowed(const UEdGraphPin* MyPin, const UEdGraphPin* OtherPin, FString& OutReason) const
 {
 	bool bIsDisallowed = Super::IsConnectionDisallowed(MyPin, OtherPin, OutReason);
diff --git a/Engine/Source/Editor/BlueprintGraph/Private/K2Node_ConstructObjectFromClass.cpp b/Engine/Source/Editor/BlueprintGraph/Private/K2Node_ConstructObjectFromClass.cpp
--- a/Engine/Source/Editor/BlueprintGraph/Private/K2Node_ConstructObjectFromClass.cpp
+++ b/Engine/Source/Editor/BlueprintGraph/Private/K2Node_ConstructObjectFromClass.cpp
@@ -34,34 +34,32 @@ UClass* UK2Node_ConstructObjectFromClass::GetClassPinBaseClass() const
 
 bool UK2Node_ConstructObjectFromClass::UseWorldContext() const
 {
-	auto BP = GetBlueprint();
+	UBlueprint* BP = GetBlueprint();
 	const UClass* ParentClass = BP ? BP->ParentClass : nullptr;
 	return ParentClass ? ParentClass->HasMetaDataHierarchical(FBlueprintMetadata::MD_ShowWorldContextPin) != nullptr : false;
 }
 
 void UK2Node_ConstructObjectFromClass::AllocateDefaultPins()
 {
-	const UEdGraphSchema_K2* K2Schema = GetDefault<UEdGraphSchema_K2>();
-
 	// Add execution pins
-	CreatePin(EGPD_Input, K2Schema->PC_Exec, FString(), nullptr, K2Schema->PN_Execute);
-	CreatePin(EGPD_Output, K2Schema->PC_Exec, FString(), nullptr, K2Schema->PN_Then);
+	CreatePin(EGPD_Input, UEdGraphSchema_K2::PC_Exec, FString(), nullptr, UEdGraphSchema_K2::PN_Execute);
+	CreatePin(EGPD_Output, UEdGraphSchema_K2::PC_Exec, FString(), nullptr, UEdGraphSchema_K2::PN_Then);
 
 	// If required add the world context pin
 	if (UseWorldContext())
 	{
-		CreatePin(EGPD_Input, K2Schema->PC_Object, FString(), UObject::StaticClass(), FK2Node_ConstructObjectFromClassHelper::WorldContextPinName);
+		CreatePin(EGPD_Input, UEdGraphSchema_K2::PC_Object, FString(), UObject::StaticClass(), FK2Node_ConstructObjectFromClassHelper::WorldContextPinName);
 	}
 
 	// Add blueprint pin
-	UEdGraphPin* ClassPin = CreatePin(EGPD_Input, K2Schema->PC_Class, FString(), GetClassPinBaseClass(), FK2Node_ConstructObjectFromClassHelper::ClassPinName);
+	UEdGraphPin* ClassPin = CreatePin(EGPD_Input, UEdGraphSchema_K2::PC_Class, FString(), GetClassPinBaseClass(), FK2Node_ConstructObjectFromClassHelper::ClassPinName);
 	
 	// Result pin
-	UEdGraphPin* ResultPin = CreatePin(EGPD_Output, K2Schema->PC_Object, FString(), GetClassPinBaseClass(), K2Schema->PN_ReturnValue);
+	UEdGraphPin* ResultPin = CreatePin(EGPD_Output, UEdGraphSchema_K2::PC_Object, FString(), GetClassPinBaseClass(), UEdGraphSchema_K2::PN_ReturnValue);
 	
 	if (UseOuter())
 	{
-		UEdGraphPin* OuterPin = CreatePin(EGPD_Input, K2Schema->PC_Object, FString(), UObject::StaticClass(), FK2Node_ConstructObjectFromClassHelper::OuterPinName);
+		UEdGraphPin* OuterPin = CreatePin(EGPD_Input, UEdGraphSchema_K2::PC_Object, FString(), UObject::StaticClass(), FK2Node_ConstructObjectFromClassHelper::OuterPinName);
 	}
 
 	Super::AllocateDefaultPins();
@@ -90,7 +88,7 @@ void UK2Node_ConstructObjectFromClass::SetPinToolTip(UEdGraphPin& MutatablePin,
 
 void UK2Node_ConstructObjectFromClass::CreatePinsForClass(UClass* InClass, TArray<UEdGraphPin*>* OutClassPins)
 {
-	check(InClass != NULL);
+	check(InClass);
 
 	const UEdGraphSchema_K2* K2Schema = GetDefault<UEdGraphSchema_K2>();
 
@@ -112,25 +110,23 @@ void UK2Node_ConstructObjectFromClass::CreatePinsForClass(UClass* InClass, TArra
 			(nullptr == FindPin(Property->GetName()) ) &&
 			FBlueprintEditorUtils::PropertyStillExists(Property))
 		{
-			UEdGraphPin* Pin = CreatePin(EGPD_Input, FString(), FString(), nullptr, Property->GetName());
-			const bool bPinGood = (Pin != NULL) && K2Schema->ConvertPropertyToPinType(Property, /*out*/ Pin->PinType);
-			
-			if (OutClassPins && Pin)
-			{
-				OutClassPins->Add(Pin);
-			}
-
-			if (ClassDefaultObject && Pin != NULL && K2Schema->PinDefaultValueIsEditable(*Pin))
-			{
-				FString DefaultValueAsString;
-				const bool bDefaultValueSet = FBlueprintEditorUtils::PropertyValueToString(Property, reinterpret_cast<const uint8*>(ClassDefaultObject), DefaultValueAsString);
-				check( bDefaultValueSet );
-				K2Schema->SetPinAutogeneratedDefaultValue(Pin, DefaultValueAsString);
-			}
-
-			// Copy tooltip from the property.
-			if (Pin != nullptr)
+			if (UEdGraphPin* Pin = CreatePin(EGPD_Input, FString(), FString(), nullptr, Property->GetName()))
 			{
+				K2Schema->ConvertPropertyToPinType(Property, /*out*/ Pin->PinType);
+				if (OutClassPins)
+				{
+					OutClassPins->Add(Pin);
+				}
+
+				if (ClassDefaultObject && K2Schema->PinDefaultValueIsEditable(*Pin))
+				{
+					FString DefaultValueAsString;
+					const bool bDefaultValueSet = FBlueprintEditorUtils::PropertyValueToString(Property, reinterpret_cast<const uint8*>(ClassDefaultObject), DefaultValueAsString);
+					check(bDefaultValueSet);
+					K2Schema->SetPinAutogeneratedDefaultValue(Pin, DefaultValueAsString);
+				}
+
+				// Copy tooltip from the property.
 				K2Schema->ConstructBasicPinTooltip(*Pin, Property->GetToolTipText(), Pin->PinToolTip);
 			}
 		}
@@ -143,17 +139,17 @@ void UK2Node_ConstructObjectFromClass::CreatePinsForClass(UClass* InClass, TArra
 
 UClass* UK2Node_ConstructObjectFromClass::GetClassToSpawn(const TArray<UEdGraphPin*>* InPinsToSearch /*=NULL*/) const
 {
-	UClass* UseSpawnClass = NULL;
+	UClass* UseSpawnClass = nullptr;
 	const TArray<UEdGraphPin*>* PinsToSearch = InPinsToSearch ? InPinsToSearch : &Pins;
 
 	UEdGraphPin* ClassPin = GetClassPin(PinsToSearch);
-	if(ClassPin && ClassPin->DefaultObject != NULL && ClassPin->LinkedTo.Num() == 0)
+	if (ClassPin && ClassPin->DefaultObject && ClassPin->LinkedTo.Num() == 0)
 	{
 		UseSpawnClass = CastChecked<UClass>(ClassPin->DefaultObject);
 	}
 	else if (ClassPin && ClassPin->LinkedTo.Num())
 	{
-		auto ClassSource = ClassPin->LinkedTo[0];
+		UEdGraphPin* ClassSource = ClassPin->LinkedTo[0];
 		UseSpawnClass = ClassSource ? Cast<UClass>(ClassSource->PinType.PinSubCategoryObject.Get()) : nullptr;
 	}
 
@@ -163,21 +159,28 @@ UClass* UK2Node_ConstructObjectFromClass::GetClassToSpawn(const TArray<UEdGraphP
 void UK2Node_ConstructObjectFromClass::ReallocatePinsDuringReconstruction(TArray<UEdGraphPin*>& OldPins) 
 {
 	AllocateDefaultPins();
-	UClass* UseSpawnClass = GetClassToSpawn(&OldPins);
-	if( UseSpawnClass != NULL )
+	if (UClass* UseSpawnClass = GetClassToSpawn(&OldPins))
 	{
 		CreatePinsForClass(UseSpawnClass);
 	}
 	RestoreSplitPins(OldPins);
 }
 
-bool UK2Node_ConstructObjectFromClass::IsSpawnVarPin(UEdGraphPin* Pin)
+void UK2Node_ConstructObjectFromClass::PostPlacedNewNode()
 {
-	const UEdGraphSchema_K2* K2Schema = GetDefault<UEdGraphSchema_K2>();
+	Super::PostPlacedNewNode();
 
-	return(	Pin->PinName != K2Schema->PN_Execute &&
-			Pin->PinName != K2Schema->PN_Then &&
-			Pin->PinName != K2Schema->PN_ReturnValue &&
+	if (UClass* UseSpawnClass = GetClassToSpawn())
+	{
+		CreatePinsForClass(UseSpawnClass);
+	}
+}
+
+bool UK2Node_ConstructObjectFromClass::IsSpawnVarPin(UEdGraphPin* Pin) const
+{
+	return(	Pin->PinName != UEdGraphSchema_K2::PN_Execute &&
+			Pin->PinName != UEdGraphSchema_K2::PN_Then &&
+			Pin->PinName != UEdGraphSchema_K2::PN_ReturnValue &&
 			Pin->PinName != FK2Node_ConstructObjectFromClassHelper::ClassPinName &&
 			Pin->PinName != FK2Node_ConstructObjectFromClassHelper::WorldContextPinName &&
 			Pin->PinName != FK2Node_ConstructObjectFromClassHelper::OuterPinName);
@@ -185,36 +188,48 @@ bool UK2Node_ConstructObjectFromClass::IsSpawnVarPin(UEdGraphPin* Pin)
 
 void UK2Node_ConstructObjectFromClass::OnClassPinChanged()
 {
-	const UEdGraphSchema_K2* K2Schema = GetDefault<UEdGraphSchema_K2>();
-
 	// Remove all pins related to archetype variables
+	TArray<UEdGraphPin*> OldPins = Pins;
 	TArray<UEdGraphPin*> OldClassPins;
 
-	for (int32 i = Pins.Num() - 1; i >= 0; --i)
+	for (UEdGraphPin* OldPin : OldPins)
 	{
-		UEdGraphPin* OldPin = Pins[i];
 		if (IsSpawnVarPin(OldPin))
 		{
-			Pins.RemoveAt(i, 1, false);
+			Pins.Remove(OldPin);
 			OldClassPins.Add(OldPin);
 		}
 	}
 
 	CachedNodeTitle.MarkDirty();
 
-	UClass* UseSpawnClass = GetClassToSpawn();
 	TArray<UEdGraphPin*> NewClassPins;
-	if (UseSpawnClass)
+	if (UClass* UseSpawnClass = GetClassToSpawn())
 	{
 		CreatePinsForClass(UseSpawnClass, &NewClassPins);
 	}
 
+	RestoreSplitPins(OldPins);
+
+	UEdGraphPin* ResultPin = GetResultPin();
+	// Cache all the pin connections to the ResultPin, we will attempt to recreate them
+	TArray<UEdGraphPin*> ResultPinConnectionList = ResultPin->LinkedTo;
+	// Because the archetype has changed, we break the output link as the output pin type will change
+	ResultPin->BreakAllPinLinks();
+
+	const UEdGraphSchema_K2* K2Schema = GetDefault<UEdGraphSchema_K2>();
+
+	// Recreate any pin links to the Result pin that are still valid
+	for (UEdGraphPin* Connections : ResultPinConnectionList)
+	{
+		K2Schema->TryCreateConnection(ResultPin, Connections);
+	}
+
 	// Rewire the old pins to the new pins so connections are maintained if possible
-	RewireOldPinsToNewPins(OldClassPins, NewClassPins);
+	RewireOldPinsToNewPins(OldClassPins, Pins);
 
 	// Refresh the UI for the graph so the pin changes show up
-	UEdGraph* Graph = GetGraph();
-	Graph->NotifyGraphChanged();
+	GetGraph()->NotifyGraphChanged();
 
 	// Mark dirty
 	FBlueprintEditorUtils::MarkBlueprintAsModified(GetBlueprint());
@@ -230,18 +245,15 @@ void UK2Node_ConstructObjectFromClass::PinConnectionListChanged(UEdGraphPin* Pin
 
 void UK2Node_ConstructObjectFromClass::GetPinHoverText(const UEdGraphPin& Pin, FString& HoverTextOut) const
 {
-	UEdGraphPin* ClassPin = GetClassPin();
-	if (ClassPin)
+	if (UEdGraphPin* ClassPin = GetClassPin())
 	{
 		SetPinToolTip(*ClassPin, LOCTEXT("ClassPinDescription", "The object class you want to construct"));
 	}
-	UEdGraphPin* ResultPin = GetResultPin();
-	if (ResultPin)
+	if (UEdGraphPin* ResultPin = GetResultPin())
 	{
 		SetPinToolTip(*ResultPin, LOCTEXT("ResultPinDescription", "The constructed object"));
 	}
-	UEdGraphPin* OuterPin = UseOuter() ? GetOuterPin() : nullptr;
-	if (OuterPin)
+	if (UEdGraphPin* OuterPin = (UseOuter() ? GetOuterPin() : nullptr))
 	{
 		SetPinToolTip(*OuterPin, LOCTEXT("OuterPinDescription", "Owner of the constructed object"));
 	}
@@ -264,9 +276,7 @@ FText UK2Node_ConstructObjectFromClass::GetTooltipText() const
 
 UEdGraphPin* UK2Node_ConstructObjectFromClass::GetThenPin()const
 {
-	const UEdGraphSchema_K2* K2Schema = GetDefault<UEdGraphSchema_K2>();
-
-	UEdGraphPin* Pin = FindPinChecked(K2Schema->PN_Then);
+	UEdGraphPin* Pin = FindPinChecked(UEdGraphSchema_K2::PN_Then);
 	check(Pin->Direction == EGPD_Output);
 	return Pin;
 }
@@ -275,41 +285,33 @@ UEdGraphPin* UK2Node_ConstructObjectFromClass::GetClassPin(const TArray<UEdGraph
 {
 	const TArray<UEdGraphPin*>* PinsToSearch = InPinsToSearch ? InPinsToSearch : &Pins;
 
-	UEdGraphPin* Pin = NULL;
-	for( auto PinIt = PinsToSearch->CreateConstIterator(); PinIt; ++PinIt )
+	UEdGraphPin* Pin = nullptr;
+	for (UEdGraphPin* TestPin : *PinsToSearch)
 	{
-		UEdGraphPin* TestPin = *PinIt;
-		if( TestPin && TestPin->PinName == FK2Node_ConstructObjectFromClassHelper::ClassPinName )
+		if (TestPin && TestPin->PinName == FK2Node_ConstructObjectFromClassHelper::ClassPinName)
 		{
 			Pin = TestPin;
 			break;
 		}
 	}
-	check(Pin == NULL || Pin->Direction == EGPD_Input);
+	check(Pin == nullptr || Pin->Direction == EGPD_Input);
 	return Pin;
 }
 
 UEdGraphPin* UK2Node_ConstructObjectFromClass::GetWorldContextPin() const
 {
 	UEdGraphPin* Pin = FindPin(FK2Node_ConstructObjectFromClassHelper::WorldContextPinName);
-	check(Pin == NULL || Pin->Direction == EGPD_Input);
+	check(Pin == nullptr || Pin->Direction == EGPD_Input);
 	return Pin;
 }
 
 UEdGraphPin* UK2Node_ConstructObjectFromClass::GetResultPin() const
 {
-	const UEdGraphSchema_K2* K2Schema = GetDefault<UEdGraphSchema_K2>();
-
-	UEdGraphPin* Pin = FindPinChecked(K2Schema->PN_ReturnValue);
+	UEdGraphPin* Pin = FindPinChecked(UEdGraphSchema_K2::PN_ReturnValue);
 	check(Pin->Direction == EGPD_Output);
 	return Pin;
 }
 
-FLinearColor UK2Node_ConstructObjectFromClass::GetNodeTitleColor() const
-{
-	return Super::GetNodeTitleColor();
-}
-
 FText UK2Node_ConstructObjectFromClass::GetBaseNodeTitle() const
 {
 	return NSLOCTEXT("K2Node", "ConstructObject_BaseTitle", "Construct Object from Class");
@@ -326,7 +328,7 @@ FText UK2Node_ConstructObjectFromClass::GetNodeTitle(ENodeTitleType::Type TitleT
 	{
 		return GetBaseNodeTitle();
 	}
-	else if (auto ClassToSpawn = GetClassToSpawn())
+	else if (UClass* ClassToSpawn = GetClassToSpawn())
 	{
 		if (CachedNodeTitle.IsOutOfDate(this))
 		{
@@ -385,7 +387,7 @@ bool UK2Node_ConstructObjectFromClass::HasExternalDependencies(TArray<class UStr
 {
 	UClass* SourceClass = GetClassToSpawn();
 	const UBlueprint* SourceBlueprint = GetBlueprint();
-	const bool bResult = (SourceClass != NULL) && (SourceClass->ClassGeneratedBy != SourceBlueprint);
+	const bool bResult = (SourceClass && (SourceClass->ClassGeneratedBy != SourceBlueprint));
 	if (bResult && OptionalOutput)
 	{
 		OptionalOutput->AddUnique(SourceClass);
diff --git a/Engine/Source/Editor/BlueprintGraph/Private/K2Node_ConvertAsset.cpp b/Engine/Source/Editor/BlueprintGraph/Private/K2Node_ConvertAsset.cpp
--- a/Engine/Source/Editor/BlueprintGraph/Private/K2Node_ConvertAsset.cpp
+++ b/Engine/Source/Editor/BlueprintGraph/Private/K2Node_ConvertAsset.cpp
@@ -35,7 +35,7 @@ bool UK2Node_ConvertAsset::IsAssetClassType() const
 	UEdGraphPin* InputPin = FindPin(UK2Node_ConvertAssetImpl::InputPinName);
 	bool bIsConnected = InputPin && InputPin->LinkedTo.Num() && InputPin->LinkedTo[0];
 	UEdGraphPin* SourcePin = bIsConnected ? InputPin->LinkedTo[0] : nullptr;
-	return SourcePin ? (SourcePin->PinType.PinCategory == UEdGraphSchema_K2::PC_AssetClass || SourcePin->PinType.PinCategory == UEdGraphSchema_K2::PC_Class) : false;
+	return SourcePin ? (SourcePin->PinType.PinCategory == UEdGraphSchema_K2::PC_SoftClass || SourcePin->PinType.PinCategory == UEdGraphSchema_K2::PC_Class) : false;
 }
 
 bool UK2Node_ConvertAsset::IsConvertToAsset() const
@@ -52,8 +52,8 @@ bool UK2Node_ConvertAsset::IsConnectionDisallowed(const UEdGraphPin* MyPin, cons
 	UEdGraphPin* InputPin = FindPin(UK2Node_ConvertAssetImpl::InputPinName);
 	if (InputPin && OtherPin && (InputPin == MyPin) && (MyPin->PinType.PinCategory == UEdGraphSchema_K2::PC_Wildcard))
 	{
-		if ((OtherPin->PinType.PinCategory != UEdGraphSchema_K2::PC_Asset) &&
-			(OtherPin->PinType.PinCategory != UEdGraphSchema_K2::PC_AssetClass) &&
+		if ((OtherPin->PinType.PinCategory != UEdGraphSchema_K2::PC_SoftObject) &&
+			(OtherPin->PinType.PinCategory != UEdGraphSchema_K2::PC_SoftClass) &&
 			(OtherPin->PinType.PinCategory != UEdGraphSchema_K2::PC_Object) &&
 			(OtherPin->PinType.PinCategory != UEdGraphSchema_K2::PC_Class))
 		{
@@ -83,11 +83,11 @@ void UK2Node_ConvertAsset::RefreshPinTypes()
 			if (bIsConvertToAsset)
 			{
 				InputCategory = (bIsAssetClass ? UEdGraphSchema_K2::PC_Class : UEdGraphSchema_K2::PC_Object);
-				OutputCategory = (bIsAssetClass ? UEdGraphSchema_K2::PC_AssetClass : UEdGraphSchema_K2::PC_Asset);
+				OutputCategory = (bIsAssetClass ? UEdGraphSchema_K2::PC_SoftClass : UEdGraphSchema_K2::PC_SoftObject);
 			}
 			else
 			{
-				InputCategory = (bIsAssetClass ? UEdGraphSchema_K2::PC_AssetClass : UEdGraphSchema_K2::PC_Asset);
+				InputCategory = (bIsAssetClass ? UEdGraphSchema_K2::PC_SoftClass : UEdGraphSchema_K2::PC_SoftObject);
 				OutputCategory = (bIsAssetClass ? UEdGraphSchema_K2::PC_Class : UEdGraphSchema_K2::PC_Object);
 			}
 		}
@@ -235,8 +235,8 @@ void UK2Node_ConvertAsset::ExpandNode(class FKismetCompilerContext& CompilerCont
 			//Connect input to convert
 			UEdGraphPin* InputPin = FindPin(UK2Node_ConvertAssetImpl::InputPinName);
 			const FString ConvertInputName = bIsAssetClass
-				? FString(TEXT("AssetClass"))
-				: FString(TEXT("Asset"));
+				? FString(TEXT("SoftClass"))
+				: FString(TEXT("SoftObject"));
 			UEdGraphPin* ConvertInput = ConvertToObjectFunc->FindPin(ConvertInputName);
 			bIsErrorFree = InputPin && ConvertInput && CompilerContext.MovePinLinksToIntermediate(*InputPin, *ConvertInput).CanSafeConnect();
 
diff --git a/Engine/Source/Editor/BlueprintGraph/Private/K2Node_Copy.cpp b/Engine/Source/Editor/BlueprintGraph/Private/K2Node_Copy.cpp
--- a/Engine/Source/Editor/BlueprintGraph/Private/K2Node_Copy.cpp
+++ b/Engine/Source/Editor/BlueprintGraph/Private/K2Node_Copy.cpp
@@ -184,7 +184,7 @@ bool UK2Node_Copy::IsConnectionDisallowed(const UEdGraphPin* MyPin, const UEdGra
 			return true;
 		}
 		else if (OtherPin->PinType.PinCategory == K2Schema->PC_Object || OtherPin->PinType.PinCategory == K2Schema->PC_Class
-			|| OtherPin->PinType.PinCategory == K2Schema->PC_Asset || OtherPin->PinType.PinCategory == K2Schema->PC_AssetClass
+			|| OtherPin->PinType.PinCategory == K2Schema->PC_SoftObject || OtherPin->PinType.PinCategory == K2Schema->PC_SoftClass
 			|| OtherPin->PinType.PinCategory == K2Schema->PC_Interface)
 		{
 			OutReason = FText::Format(LOCTEXT("ObjectConnectionDisallowed", "Cannot connect with {0} pin."), FText::FromString(OtherPin->PinType.PinCategory)).ToString();
diff --git a/Engine/Source/Editor/BlueprintGraph/Private/K2Node_CreateDelegate.cpp b/Engine/Source/Editor/BlueprintGraph/Private/K2Node_CreateDelegate.cpp
--- a/Engine/Source/Editor/BlueprintGraph/Private/K2Node_CreateDelegate.cpp
+++ b/Engine/Source/Editor/BlueprintGraph/Private/K2Node_CreateDelegate.cpp
@@ -405,28 +405,13 @@ FText UK2Node_CreateDelegate::GetNodeTitle(ENodeTitleType::Type TitleType) const
 UObject* UK2Node_CreateDelegate::GetJumpTargetForDoubleClick() const
 {
 	const UEdGraphSchema_K2* K2Schema = GetDefault<UEdGraphSchema_K2>();
-	UBlueprint* ScopeClassBlueprint = NULL;
+	UClass* ScopeClass = GetScopeClass();
 
-	UEdGraphPin* Pin = FindPinChecked(K2Schema->PN_Self);
-	UEdGraphPin* ResultPin = Pin->LinkedTo.Num() ? Pin->LinkedTo[0] : NULL;
-	if (ResultPin)
-	{
-		ensure(K2Schema->PC_Object == ResultPin->PinType.PinCategory);
-		if (UClass* TrueScopeClass = Cast<UClass>(ResultPin->PinType.PinSubCategoryObject.Get()))
-		{
-			ScopeClassBlueprint = Cast<UBlueprint>(TrueScopeClass->ClassGeneratedBy);
-		}
-		else if (K2Schema->PN_Self == ResultPin->PinType.PinSubCategory)
-		{
-			ScopeClassBlueprint = GetBlueprint();
-		}
-	}
-
-	if (ScopeClassBlueprint)
+	if (UBlueprint* ScopeClassBlueprint = (ScopeClass != nullptr) ? Cast<UBlueprint>(ScopeClass->ClassGeneratedBy) : nullptr)
 	{
 		if (UEdGraph* FoundGraph = FindObject<UEdGraph>(ScopeClassBlueprint, *GetFunctionName().ToString()))
 		{
-			if(!FBlueprintEditorUtils::IsGraphIntermediate(FoundGraph))
+			if (!FBlueprintEditorUtils::IsGraphIntermediate(FoundGraph))
 			{
 				return FoundGraph;
 			}
@@ -434,21 +419,22 @@ UObject* UK2Node_CreateDelegate::GetJumpTargetForDoubleClick() const
 		for (auto UbergraphIt = ScopeClassBlueprint->UbergraphPages.CreateIterator(); UbergraphIt; ++UbergraphIt)
 		{
 			UEdGraph* Graph = (*UbergraphIt);
-			if(!FBlueprintEditorUtils::IsGraphIntermediate(Graph))
+			if (!FBlueprintEditorUtils::IsGraphIntermediate(Graph))
 			{
 				TArray<UK2Node_Event*> EventNodes;
 				Graph->GetNodesOfClass(EventNodes);
-				for (auto EventIt = EventNodes.CreateIterator(); EventIt; ++EventIt)
+				for (UK2Node_Event* EventNode : EventNodes)
 				{
-					if(GetFunctionName() == (*EventIt)->GetFunctionName())
+					if (GetFunctionName() == EventNode->GetFunctionName())
 					{
-						return *EventIt;
+						return EventNode;
 					}
 				}
 			}
 		}
 	}
-	return NULL;
+
+	return GetDelegateSignature();
 }
 
 void UK2Node_CreateDelegate::AddSearchMetaDataInfo(TArray<struct FSearchTagDataPair>& OutTaggedMetaData) const
diff --git a/Engine/Source/Editor/BlueprintGraph/Private/K2Node_DynamicCast.cpp b/Engine/Source/Editor/BlueprintGraph/Private/K2Node_DynamicCast.cpp
--- a/Engine/Source/Editor/BlueprintGraph/Private/K2Node_DynamicCast.cpp
+++ b/Engine/Source/Editor/BlueprintGraph/Private/K2Node_DynamicCast.cpp
@@ -252,15 +252,7 @@ void UK2Node_DynamicCast::SetPurity(bool bNewPurity)
 
 void UK2Node_DynamicCast::TogglePurity()
 {
-	FText TransactionTitle;
-	if(bIsPureCast)
-	{
-		TransactionTitle = LOCTEXT("TogglePure", "Convert to Pure Cast");
-	}
-	else
-	{
-		TransactionTitle = LOCTEXT("ToggleImpure", "Convert to Impure Cast");
-	}
+	const FText TransactionTitle = bIsPureCast ? LOCTEXT("TogglePurityToImpure", "Convert to Impure Cast") : LOCTEXT("TogglePurityToPure", "Convert to Pure Cast");
 	const FScopedTransaction Transaction( TransactionTitle );
 	Modify();
 
diff --git a/Engine/Source/Editor/BlueprintGraph/Private/K2Node_Event.cpp b/Engine/Source/Editor/BlueprintGraph/Private/K2Node_Event.cpp
--- a/Engine/Source/Editor/BlueprintGraph/Private/K2Node_Event.cpp
+++ b/Engine/Source/Editor/BlueprintGraph/Private/K2Node_Event.cpp
@@ -483,7 +483,7 @@ bool UK2Node_Event::CanPasteHere(const UEdGraph* TargetGraph) const
 		UBlueprint* Blueprint = FBlueprintEditorUtils::FindBlueprintForGraph(TargetGraph);
 		if(Blueprint && Blueprint->SkeletonGeneratedClass)
 		{
-			TArray<FName> ExistingNamesInUse;
+			TSet<FName> ExistingNamesInUse;
 			TArray<FString> ExcludedEventNames;
 			TArray<UK2Node_Event*> ExistingEventNodes;
 			TArray<UClass*> ImplementedInterfaceClasses;
diff --git a/Engine/Source/Editor/BlueprintGraph/Private/K2Node_ExecutionSequence.cpp b/Engine/Source/Editor/BlueprintGraph/Private/K2Node_ExecutionSequence.cpp
--- a/Engine/Source/Editor/BlueprintGraph/Private/K2Node_ExecutionSequence.cpp
+++ b/Engine/Source/Editor/BlueprintGraph/Private/K2Node_ExecutionSequence.cpp
@@ -60,7 +60,6 @@ public:
 			if (Context.IsDebuggingOrInstrumentationRequired() && (OutputPins.Num() > 1))
 			{
 				const FString NodeComment = Node->NodeComment.IsEmpty() ? Node->GetName() : Node->NodeComment;
-				const bool bInstrumentedVersion = Context.bInstrumentScriptCode;
 
 				// Assuming sequence X goes to A, B, C, we want to emit:
 				//   X: push X1
@@ -74,14 +73,6 @@ public:
 				// A push statement we need to patch up on the next pass (e.g., push X1 before we know where X1 is)
 				FBlueprintCompiledStatement* LastPushStatement = NULL;
 
-				if (bInstrumentedVersion)
-				{
-					// Emit a sequence instruction
-					FBlueprintCompiledStatement& SequenceStart = Context.AppendStatementForNode(Node);
-					SequenceStart.Type = KCST_InstrumentedStatePush;
-					SequenceStart.Comment = NodeComment;
-				}
-
 				for (int32 i = 0; i < OutputPins.Num(); ++i)
 				{
 					// Emit the debug site and patch up the previous jump if we're on subsequent steps
@@ -90,7 +81,7 @@ public:
 					{
 						// Emit a debug site
 						FBlueprintCompiledStatement& DebugSiteAndJumpTarget = Context.AppendStatementForNode(Node);
-						DebugSiteAndJumpTarget.Type = bInstrumentedVersion ? KCST_InstrumentedStateRestore : Context.GetBreakpointType();
+						DebugSiteAndJumpTarget.Type = Context.GetBreakpointType();
 						DebugSiteAndJumpTarget.Comment = NodeComment;
 						DebugSiteAndJumpTarget.bIsJumpTarget = true;
 
@@ -100,7 +91,7 @@ public:
 					}
 
 					// Emit a push to get to the next step in the sequence, unless we're the last one or this is an instrumented build
-					const bool bNotLastIndex = ((i + 1) < OutputPins.Num()) || bInstrumentedVersion;
+					const bool bNotLastIndex = ((i + 1) < OutputPins.Num());
 					if (bNotLastIndex)
 					{
 						FBlueprintCompiledStatement& PushExecutionState = Context.AppendStatementForNode(Node);
@@ -115,18 +106,6 @@ public:
 				}
 
 				check(LastPushStatement);
-
-				if (bInstrumentedVersion)
-				{
-					// For instrumented builds jump back to node and pop the execution state then exit the thread.
-					FBlueprintCompiledStatement& SequenceEnd = Context.AppendStatementForNode(Node);
-					SequenceEnd.Type = KCST_InstrumentedStatePop;
-					SequenceEnd.Comment = NodeComment;
-					SequenceEnd.bIsJumpTarget = true;
-					LastPushStatement->TargetLabel = &SequenceEnd;
-					FBlueprintCompiledStatement& NextExecutionState = Context.AppendStatementForNode(Node);
-					NextExecutionState.Type = KCST_EndOfThread;
-				}
 			}
 			else
 			{
diff --git a/Engine/Source/Editor/BlueprintGraph/Private/K2Node_FunctionResult.cpp b/Engine/Source/Editor/BlueprintGraph/Private/K2Node_FunctionResult.cpp
--- a/Engine/Source/Editor/BlueprintGraph/Private/K2Node_FunctionResult.cpp
+++ b/Engine/Source/Editor/BlueprintGraph/Private/K2Node_FunctionResult.cpp
@@ -97,7 +97,7 @@ public:
 		{
 			GenerateAssigments(Context, Node);
 
-			if (Context.IsDebuggingOrInstrumentationRequired() && Node && !Context.IsExpansionNode(Node))
+			if (Context.IsDebuggingOrInstrumentationRequired() && Node)
 			{
 				FBlueprintCompiledStatement& TraceStatement = Context.AppendStatementForNode(Node);
 				TraceStatement.Type = Context.GetWireTraceType();
diff --git a/Engine/Source/Editor/BlueprintGraph/Private/K2Node_GenericCreateObject.cpp b/Engine/Source/Editor/BlueprintGraph/Private/K2Node_GenericCreateObject.cpp
--- a/Engine/Source/Editor/BlueprintGraph/Private/K2Node_GenericCreateObject.cpp
+++ b/Engine/Source/Editor/BlueprintGraph/Private/K2Node_GenericCreateObject.cpp
@@ -65,6 +65,8 @@ struct FK2Node_GenericCreateObject_Utils
 
 void UK2Node_GenericCreateObject::ExpandNode(class FKismetCompilerContext& CompilerContext, UEdGraph* SourceGraph)
 {
+	Super::ExpandNode(CompilerContext, SourceGraph);
+
 	UK2Node_CallFunction* CallCreateNode = CompilerContext.SpawnIntermediateNode<UK2Node_CallFunction>(this, SourceGraph);
 	CallCreateNode->FunctionReference.SetExternalMember(GET_FUNCTION_NAME_CHECKED(UGameplayStatics, SpawnObject), UGameplayStatics::StaticClass());
 	CallCreateNode->AllocateDefaultPins();
diff --git a/Engine/Source/Editor/BlueprintGraph/Private/K2Node_GetArrayItem.cpp b/Engine/Source/Editor/BlueprintGraph/Private/K2Node_GetArrayItem.cpp
--- a/Engine/Source/Editor/BlueprintGraph/Private/K2Node_GetArrayItem.cpp
+++ b/Engine/Source/Editor/BlueprintGraph/Private/K2Node_GetArrayItem.cpp
@@ -93,8 +93,8 @@ namespace K2Node_GetArrayItem_Impl
 	{
 		return !(PinType.PinCategory == UEdGraphSchema_K2::PC_Object ||
 			PinType.PinCategory == UEdGraphSchema_K2::PC_Class ||
-			PinType.PinCategory == UEdGraphSchema_K2::PC_Asset ||
-			PinType.PinCategory == UEdGraphSchema_K2::PC_AssetClass ||
+			PinType.PinCategory == UEdGraphSchema_K2::PC_SoftObject ||
+			PinType.PinCategory == UEdGraphSchema_K2::PC_SoftClass ||
 			PinType.PinCategory == UEdGraphSchema_K2::PC_Interface);
 	}
 
diff --git a/Engine/Source/Editor/BlueprintGraph/Private/K2Node_GetClassDefaults.cpp b/Engine/Source/Editor/BlueprintGraph/Private/K2Node_GetClassDefaults.cpp
--- a/Engine/Source/Editor/BlueprintGraph/Private/K2Node_GetClassDefaults.cpp
+++ b/Engine/Source/Editor/BlueprintGraph/Private/K2Node_GetClassDefaults.cpp
@@ -337,6 +337,8 @@ FNodeHandlingFunctor* UK2Node_GetClassDefaults::CreateNodeHandler(FKismetCompile
 
 void UK2Node_GetClassDefaults::ExpandNode(class FKismetCompilerContext& CompilerContext, UEdGraph* SourceGraph)
 {
+	Super::ExpandNode(CompilerContext, SourceGraph);
+
 	const UClass* ClassType = GetInputClass();
 
 	// @TODO - Remove if/when we support 'const' input pins.
diff --git a/Engine/Source/Editor/BlueprintGraph/Private/K2Node_LoadAsset.cpp b/Engine/Source/Editor/BlueprintGraph/Private/K2Node_LoadAsset.cpp
--- a/Engine/Source/Editor/BlueprintGraph/Private/K2Node_LoadAsset.cpp
+++ b/Engine/Source/Editor/BlueprintGraph/Private/K2Node_LoadAsset.cpp
@@ -253,7 +253,7 @@ FText UK2Node_LoadAsset::GetMenuCategory() const
 
 const FString& UK2Node_LoadAsset::GetInputCategory() const
 {
-	return UEdGraphSchema_K2::PC_Asset;
+	return UEdGraphSchema_K2::PC_SoftObject;
 }
 
 const FString& UK2Node_LoadAsset::GetOutputCategory() const
@@ -298,7 +298,7 @@ FText UK2Node_LoadAssetClass::GetNodeTitle(ENodeTitleType::Type TitleType) const
 
 const FString& UK2Node_LoadAssetClass::GetInputCategory() const
 {
-	return UEdGraphSchema_K2::PC_AssetClass;
+	return UEdGraphSchema_K2::PC_SoftClass;
 }
 
 const FString& UK2Node_LoadAssetClass::GetOutputCategory() const
diff --git a/Engine/Source/Editor/BlueprintGraph/Private/K2Node_MakeStruct.cpp b/Engine/Source/Editor/BlueprintGraph/Private/K2Node_MakeStruct.cpp
--- a/Engine/Source/Editor/BlueprintGraph/Private/K2Node_MakeStruct.cpp
+++ b/Engine/Source/Editor/BlueprintGraph/Private/K2Node_MakeStruct.cpp
@@ -51,7 +51,7 @@ void UK2Node_MakeStruct::FMakeStructPinManager::CustomizePinData(UEdGraphPin* Pi
 
 		const bool bIsObject = Property->IsA<UObjectPropertyBase>();
 		checkSlow(bIsObject == ((Schema->PC_Object == Pin->PinType.PinCategory || Schema->PC_Class == Pin->PinType.PinCategory || 
-			Schema->PC_Asset == Pin->PinType.PinCategory || Schema->PC_AssetClass == Pin->PinType.PinCategory) && !Pin->PinType.IsContainer()));
+			Schema->PC_SoftObject == Pin->PinType.PinCategory || Schema->PC_SoftClass == Pin->PinType.PinCategory) && !Pin->PinType.IsContainer()));
 
 		if (Property->HasAnyPropertyFlags(CPF_AdvancedDisplay))
 		{
