Copying //UE4/Dev-Framework to //UE4/Dev-Main (Source: //UE4/Dev-Framework @ 3582324)

diff --git a/Engine/Source/Editor/BlueprintGraph/Private/K2Node_SpawnActorFromClass.cpp b/Engine/Source/Editor/BlueprintGraph/Private/K2Node_SpawnActorFromClass.cpp
--- a/Engine/Source/Editor/BlueprintGraph/Private/K2Node_SpawnActorFromClass.cpp
+++ b/Engine/Source/Editor/BlueprintGraph/Private/K2Node_SpawnActorFromClass.cpp
@@ -19,8 +19,6 @@
 
 struct FK2Node_SpawnActorFromClassHelper
 {
-	static FString WorldContextPinName;
-	static FString ClassPinName;
 	static FString SpawnTransformPinName;
 	static FString SpawnEvenIfCollidingPinName;
 	static FString NoCollisionFailPinName;
@@ -28,8 +26,6 @@ struct FK2Node_SpawnActorFromClassHelper
 	static FString OwnerPinName;
 };
 
-FString FK2Node_SpawnActorFromClassHelper::WorldContextPinName(TEXT("WorldContextObject"));
-FString FK2Node_SpawnActorFromClassHelper::ClassPinName(TEXT("Class"));
 FString FK2Node_SpawnActorFromClassHelper::SpawnTransformPinName(TEXT("SpawnTransform"));
 FString FK2Node_SpawnActorFromClassHelper::SpawnEvenIfCollidingPinName(TEXT("SpawnEvenIfColliding"));		// deprecated pin, name kept for backwards compat
 FString FK2Node_SpawnActorFromClassHelper::NoCollisionFailPinName(TEXT("bNoCollisionFail"));		// deprecated pin, name kept for backwards compat
@@ -44,111 +40,32 @@ UK2Node_SpawnActorFromClass::UK2Node_SpawnActorFromClass(const FObjectInitialize
 	NodeTooltip = LOCTEXT("NodeTooltip", "Attempts to spawn a new Actor with the specified transform");
 }
 
-void UK2Node_SpawnActorFromClass::AllocateDefaultPins()
+UClass* UK2Node_SpawnActorFromClass::GetClassPinBaseClass() const
 {
-	const UEdGraphSchema_K2* K2Schema = GetDefault<UEdGraphSchema_K2>();
-
-	// Add execution pins
-	CreatePin(EGPD_Input, K2Schema->PC_Exec, FString(), nullptr, K2Schema->PN_Execute);
-	CreatePin(EGPD_Output, K2Schema->PC_Exec, FString(), nullptr, K2Schema->PN_Then);
+	return AActor::StaticClass();
+}
 
-	// If required add the world context pin
-	if (GetBlueprint()->ParentClass->HasMetaDataHierarchical(FBlueprintMetadata::MD_ShowWorldContextPin))
-	{
-		CreatePin(EGPD_Input, K2Schema->PC_Object, FString(), UObject::StaticClass(), FK2Node_SpawnActorFromClassHelper::WorldContextPinName);
-	}
+void UK2Node_SpawnActorFromClass::AllocateDefaultPins()
+{
+	Super::AllocateDefaultPins();
 
-	// Add blueprint pin
-	UEdGraphPin* ClassPin = CreatePin(EGPD_Input, K2Schema->PC_Class, FString(), AActor::StaticClass(), FK2Node_SpawnActorFromClassHelper::ClassPinName);
+	const UEdGraphSchema_K2* K2Schema = GetDefault<UEdGraphSchema_K2>();
 
 	// Transform pin
 	UScriptStruct* TransformStruct = TBaseStructure<FTransform>::Get();
-	UEdGraphPin* TransformPin = CreatePin(EGPD_Input, K2Schema->PC_Struct, FString(), TransformStruct, FK2Node_SpawnActorFromClassHelper::SpawnTransformPinName);
+	UEdGraphPin* TransformPin = CreatePin(EGPD_Input, UEdGraphSchema_K2::PC_Struct, FString(), TransformStruct, FK2Node_SpawnActorFromClassHelper::SpawnTransformPinName);
 
 	// Collision handling method pin
 	UEnum* const MethodEnum = FindObjectChecked<UEnum>(ANY_PACKAGE, TEXT("ESpawnActorCollisionHandlingMethod"), true);
-	UEdGraphPin* const CollisionHandlingOverridePin = CreatePin(EGPD_Input, K2Schema->PC_Byte, FString(), MethodEnum, FK2Node_SpawnActorFromClassHelper::CollisionHandlingOverridePinName);
+	UEdGraphPin* const CollisionHandlingOverridePin = CreatePin(EGPD_Input, UEdGraphSchema_K2::PC_Byte, FString(), MethodEnum, FK2Node_SpawnActorFromClassHelper::CollisionHandlingOverridePinName);
 	CollisionHandlingOverridePin->DefaultValue = MethodEnum->GetNameStringByValue(static_cast<int>(ESpawnActorCollisionHandlingMethod::Undefined));
 
-	UEdGraphPin* OwnerPin = CreatePin(EGPD_Input, K2Schema->PC_Object, FString(), AActor::StaticClass(), FK2Node_SpawnActorFromClassHelper::OwnerPinName);
+	UEdGraphPin* OwnerPin = CreatePin(EGPD_Input, UEdGraphSchema_K2::PC_Object, FString(), AActor::StaticClass(), FK2Node_SpawnActorFromClassHelper::OwnerPinName);
 	OwnerPin->bAdvancedView = true;
 	if (ENodeAdvancedPins::NoPins == AdvancedPinDisplay)
 	{
 		AdvancedPinDisplay = ENodeAdvancedPins::Hidden;
 	}
-
-	// Result pin
-	UEdGraphPin* ResultPin = CreatePin(EGPD_Output, K2Schema->PC_Object, FString(), AActor::StaticClass(), K2Schema->PN_ReturnValue);
-
-	Super::AllocateDefaultPins();
-}
-
-void UK2Node_SpawnActorFromClass::CreatePinsForClass(UClass* InClass, TArray<UEdGraphPin*>& OutClassPins)
-{
-	check(InClass != NULL);
-
-	const UEdGraphSchema_K2* K2Schema = GetDefault<UEdGraphSchema_K2>();
-
-	const UObject* const ClassDefaultObject = InClass->GetDefaultObject(false);
-
-	for (TFieldIterator<UProperty> PropertyIt(InClass, EFieldIteratorFlags::IncludeSuper); PropertyIt; ++PropertyIt)
-	{
-		UProperty* Property = *PropertyIt;
-		UClass* PropertyClass = CastChecked<UClass>(Property->GetOuter());
-		const bool bIsDelegate = Property->IsA(UMulticastDelegateProperty::StaticClass());
-		const bool bIsExposedToSpawn = UEdGraphSchema_K2::IsPropertyExposedOnSpawn(Property);
-		const bool bIsSettableExternally = !Property->HasAnyPropertyFlags(CPF_DisableEditOnInstance);
-
-		if(	bIsExposedToSpawn &&
-			!Property->HasAnyPropertyFlags(CPF_Parm) && 
-			bIsSettableExternally &&
-			Property->HasAllPropertyFlags(CPF_BlueprintVisible) &&
-			!bIsDelegate &&
-			(nullptr == FindPin(Property->GetName()) ) &&
-			FBlueprintEditorUtils::PropertyStillExists(Property) )
-		{
-			UEdGraphPin* Pin = CreatePin(EGPD_Input, FString(), FString(), nullptr, Property->GetName());
-			const bool bPinGood = (Pin != nullptr) && K2Schema->ConvertPropertyToPinType(Property, /*out*/ Pin->PinType);
-			OutClassPins.Add(Pin);
-
-			if (ClassDefaultObject && Pin != nullptr && K2Schema->PinDefaultValueIsEditable(*Pin))
-			{
-				FString DefaultValueAsString;
-				const bool bDefaultValueSet = FBlueprintEditorUtils::PropertyValueToString(Property, reinterpret_cast<const uint8*>(ClassDefaultObject), DefaultValueAsString);
-				check( bDefaultValueSet );
-				K2Schema->SetPinAutogeneratedDefaultValue(Pin, DefaultValueAsString);
-			}
-
-			// Copy tooltip from the property.
-			if (Pin != nullptr)
-			{
-				K2Schema->ConstructBasicPinTooltip(*Pin, Property->GetToolTipText(), Pin->PinToolTip);
-			}
-		}
-	}
-
-	// Change class of output pin
-	UEdGraphPin* ResultPin = GetResultPin();
-	ResultPin->PinType.PinSubCategoryObject = InClass;
-}
-
-UClass* UK2Node_SpawnActorFromClass::GetClassToSpawn(const TArray<UEdGraphPin*>* InPinsToSearch /*=NULL*/) const
-{
-	UClass* UseSpawnClass = nullptr;
-	const TArray<UEdGraphPin*>* PinsToSearch = InPinsToSearch ? InPinsToSearch : &Pins;
-
-	UEdGraphPin* ClassPin = GetClassPin(PinsToSearch);
-	if(ClassPin && ClassPin->DefaultObject != NULL && ClassPin->LinkedTo.Num() == 0)
-	{
-		UseSpawnClass = CastChecked<UClass>(ClassPin->DefaultObject);
-	}
-	else if (ClassPin && (1 == ClassPin->LinkedTo.Num()))
-	{
-		auto SourcePin = ClassPin->LinkedTo[0];
-		UseSpawnClass = SourcePin ? Cast<UClass>(SourcePin->PinType.PinSubCategoryObject.Get()) : nullptr;
-	}
-
-	return UseSpawnClass;
 }
 
 void UK2Node_SpawnActorFromClass::MaybeUpdateCollisionPin(TArray<UEdGraphPin*>& OldPins)
@@ -274,34 +191,13 @@ void UK2Node_SpawnActorFromClass::MaybeUpdateCollisionPin(TArray<UEdGraphPin*>&
 
 void UK2Node_SpawnActorFromClass::ReallocatePinsDuringReconstruction(TArray<UEdGraphPin*>& OldPins) 
 {
-	AllocateDefaultPins();
-	UClass* UseSpawnClass = GetClassToSpawn(&OldPins);
-	if( UseSpawnClass != NULL )
-	{
-		TArray<UEdGraphPin*> ClassPins;
-		CreatePinsForClass(UseSpawnClass, ClassPins);
-	}
+	Super::ReallocatePinsDuringReconstruction(OldPins);
 
 	MaybeUpdateCollisionPin(OldPins);
-	RestoreSplitPins(OldPins);
-}
-
-void UK2Node_SpawnActorFromClass::PostPlacedNewNode()
-{
-	Super::PostPlacedNewNode();
-
-	UClass* UseSpawnClass = GetClassToSpawn();
-	if( UseSpawnClass != NULL )
-	{
-		TArray<UEdGraphPin*> ClassPins;
-		CreatePinsForClass(UseSpawnClass, ClassPins);
-	}
 }
 
-bool UK2Node_SpawnActorFromClass::IsSpawnVarPin(UEdGraphPin* Pin)
+bool UK2Node_SpawnActorFromClass::IsSpawnVarPin(UEdGraphPin* Pin) const
 {
-	const UEdGraphSchema_K2* K2Schema = GetDefault<UEdGraphSchema_K2>();
-
 	UEdGraphPin* ParentPin = Pin->ParentPin;
 	while (ParentPin)
 	{
@@ -312,82 +208,16 @@ bool UK2Node_SpawnActorFromClass::IsSpawnVarPin(UEdGraphPin* Pin)
 		ParentPin = ParentPin->ParentPin;
 	}
 
-	return(	Pin->PinName != K2Schema->PN_Execute &&
-			Pin->PinName != K2Schema->PN_Then &&
-			Pin->PinName != K2Schema->PN_ReturnValue &&
-			Pin->PinName != FK2Node_SpawnActorFromClassHelper::ClassPinName &&
-			Pin->PinName != FK2Node_SpawnActorFromClassHelper::WorldContextPinName &&
+	return(	Super::IsSpawnVarPin(Pin) &&
 			Pin->PinName != FK2Node_SpawnActorFromClassHelper::CollisionHandlingOverridePinName &&
 			Pin->PinName != FK2Node_SpawnActorFromClassHelper::SpawnTransformPinName && 
 			Pin->PinName != FK2Node_SpawnActorFromClassHelper::OwnerPinName );
 }
 
-void UK2Node_SpawnActorFromClass::OnClassPinChanged()
-{
- 	const UEdGraphSchema_K2* K2Schema = GetDefault<UEdGraphSchema_K2>();
-
-	// Remove all pins related to archetype variables
-	TArray<UEdGraphPin*> OldPins = Pins;
-	TArray<UEdGraphPin*> OldClassPins;
-
-	for (int32 i = 0; i < OldPins.Num(); i++)
-	{
-		UEdGraphPin* OldPin = OldPins[i];
-		if (IsSpawnVarPin(OldPin))
-		{
-			Pins.Remove(OldPin);
-			OldClassPins.Add(OldPin);
-		}
-	}
-
-	CachedNodeTitle.MarkDirty();
-
-	UClass* UseSpawnClass = GetClassToSpawn();
-	TArray<UEdGraphPin*> NewClassPins;
-	if (UseSpawnClass != NULL)
-	{
-		CreatePinsForClass(UseSpawnClass, NewClassPins);
-	}
-
-	UEdGraphPin* ResultPin = GetResultPin();
-	// Cache all the pin connections to the ResultPin, we will attempt to recreate them
-	TArray<UEdGraphPin*> ResultPinConnectionList = ResultPin->LinkedTo;
-	// Because the archetype has changed, we break the output link as the output pin type will change
-	ResultPin->BreakAllPinLinks();
-
-	// Recreate any pin links to the Result pin that are still valid
-	for (UEdGraphPin* Connections : ResultPinConnectionList)
-	{
-		K2Schema->TryCreateConnection(ResultPin, Connections);
-	}
-
-	// Rewire the old pins to the new pins so connections are maintained if possible
-	RewireOldPinsToNewPins(OldClassPins, NewClassPins);
-
-	// Refresh the UI for the graph so the pin changes show up
-	UEdGraph* Graph = GetGraph();
-	Graph->NotifyGraphChanged();
-
-	// Mark dirty
-	FBlueprintEditorUtils::MarkBlueprintAsModified(GetBlueprint());
-}
-
-void UK2Node_SpawnActorFromClass::PinConnectionListChanged(UEdGraphPin* ChangedPin)
-{
-	if (ChangedPin && (ChangedPin->PinName == FK2Node_SpawnActorFromClassHelper::ClassPinName))
-	{
-		OnClassPinChanged();
-	}
-}
-
 void UK2Node_SpawnActorFromClass::GetPinHoverText(const UEdGraphPin& Pin, FString& HoverTextOut) const
 {
 	const UEdGraphSchema_K2* K2Schema = GetDefault<UEdGraphSchema_K2>();
 
-	if (UEdGraphPin* ClassPin = GetClassPin())
-	{
-		K2Schema->ConstructBasicPinTooltip(*ClassPin, LOCTEXT("ClassPinDescription", "The Actor class you want to spawn"), ClassPin->PinToolTip);
-	}
 	if (UEdGraphPin* TransformPin = GetSpawnTransformPin())
 	{
 		K2Schema->ConstructBasicPinTooltip(*TransformPin, LOCTEXT("TransformPinDescription", "The transform to spawn the Actor with"), TransformPin->PinToolTip);
@@ -400,60 +230,16 @@ void UK2Node_SpawnActorFromClass::GetPinHoverText(const UEdGraphPin& Pin, FStrin
 	{
 		K2Schema->ConstructBasicPinTooltip(*OwnerPin, LOCTEXT("OwnerPinDescription", "Can be left empty; primarily used for replication (bNetUseOwnerRelevancy and bOnlyRelevantToOwner), or visibility (PrimitiveComponent's bOwnerNoSee/bOnlyOwnerSee)"), OwnerPin->PinToolTip);
 	}
-	if (UEdGraphPin* ResultPin = GetResultPin())
-	{
-		K2Schema->ConstructBasicPinTooltip(*ResultPin, LOCTEXT("ResultPinDescription", "The spawned Actor"), ResultPin->PinToolTip);
-	}
 
 	return Super::GetPinHoverText(Pin, HoverTextOut);
 }
 
-void UK2Node_SpawnActorFromClass::PinDefaultValueChanged(UEdGraphPin* ChangedPin) 
-{
-	if (ChangedPin && (ChangedPin->PinName == FK2Node_SpawnActorFromClassHelper::ClassPinName))
-	{
-		OnClassPinChanged();
-	}
-}
-
-FText UK2Node_SpawnActorFromClass::GetTooltipText() const
-{
-	return NodeTooltip;
-}
-
 FSlateIcon UK2Node_SpawnActorFromClass::GetIconAndTint(FLinearColor& OutColor) const
 {
 	static FSlateIcon Icon("EditorStyle", "GraphEditor.SpawnActor_16x");
 	return Icon;
 }
 
-UEdGraphPin* UK2Node_SpawnActorFromClass::GetThenPin()const
-{
-	const UEdGraphSchema_K2* K2Schema = GetDefault<UEdGraphSchema_K2>();
-
-	UEdGraphPin* Pin = FindPinChecked(K2Schema->PN_Then);
-	check(Pin->Direction == EGPD_Output);
-	return Pin;
-}
-
-UEdGraphPin* UK2Node_SpawnActorFromClass::GetClassPin(const TArray<UEdGraphPin*>* InPinsToSearch /*= NULL*/) const
-{
-	const TArray<UEdGraphPin*>* PinsToSearch = InPinsToSearch ? InPinsToSearch : &Pins;
-
-	UEdGraphPin* Pin = NULL;
-	for( auto PinIt = PinsToSearch->CreateConstIterator(); PinIt; ++PinIt )
-	{
-		UEdGraphPin* TestPin = *PinIt;
-		if( TestPin && TestPin->PinName == FK2Node_SpawnActorFromClassHelper::ClassPinName )
-		{
-			Pin = TestPin;
-			break;
-		}
-	}
-	check(Pin == NULL || Pin->Direction == EGPD_Input);
-	return Pin;
-}
-
 UEdGraphPin* UK2Node_SpawnActorFromClass::GetSpawnTransformPin() const
 {
 	UEdGraphPin* Pin = FindPinChecked(FK2Node_SpawnActorFromClassHelper::SpawnTransformPinName);
@@ -471,32 +257,10 @@ UEdGraphPin* UK2Node_SpawnActorFromClass::GetCollisionHandlingOverridePin() cons
 UEdGraphPin* UK2Node_SpawnActorFromClass::GetOwnerPin() const
 {
 	UEdGraphPin* Pin = FindPin(FK2Node_SpawnActorFromClassHelper::OwnerPinName);
-	check(Pin == NULL || Pin->Direction == EGPD_Input);
+	check(Pin == nullptr || Pin->Direction == EGPD_Input);
 	return Pin;
 }
 
-UEdGraphPin* UK2Node_SpawnActorFromClass::GetWorldContextPin() const
-{
-	UEdGraphPin* Pin = FindPin(FK2Node_SpawnActorFromClassHelper::WorldContextPinName);
-	check(Pin == NULL || Pin->Direction == EGPD_Input);
-	return Pin;
-}
-
-UEdGraphPin* UK2Node_SpawnActorFromClass::GetResultPin() const
-{
-	const UEdGraphSchema_K2* K2Schema = GetDefault<UEdGraphSchema_K2>();
-
-	UEdGraphPin* Pin = FindPinChecked(K2Schema->PN_ReturnValue);
-	check(Pin->Direction == EGPD_Output);
-	return Pin;
-}
-
-FLinearColor UK2Node_SpawnActorFromClass::GetNodeTitleColor() const
-{
-	return Super::GetNodeTitleColor();
-}
-
-
 FText UK2Node_SpawnActorFromClass::GetNodeTitle(ENodeTitleType::Type TitleType) const
 {
 	FText NodeTitle = NSLOCTEXT("K2Node", "SpawnActor_BaseTitle", "Spawn Actor from Class");
@@ -556,31 +320,6 @@ void UK2Node_SpawnActorFromClass::GetNodeAttributes( TArray<TKeyValuePair<FStrin
 	OutNodeAttributes.Add( TKeyValuePair<FString, FString>( TEXT( "ActorClass" ), ClassToSpawnStr ));
 }
 
-void UK2Node_SpawnActorFromClass::GetMenuActions(FBlueprintActionDatabaseRegistrar& ActionRegistrar) const
-{
-	// actions get registered under specific object-keys; the idea is that 
-	// actions might have to be updated (or deleted) if their object-key is  
-	// mutated (or removed)... here we use the node's class (so if the node 
-	// type disappears, then the action should go with it)
-	UClass* ActionKey = GetClass();
-	// to keep from needlessly instantiating a UBlueprintNodeSpawner, first   
-	// check to make sure that the registrar is looking for actions of this type
-	// (could be regenerating actions for a specific asset, and therefore the 
-	// registrar would only accept actions corresponding to that asset)
-	if (ActionRegistrar.IsOpenForRegistration(ActionKey))
-	{
-		UBlueprintNodeSpawner* NodeSpawner = UBlueprintNodeSpawner::Create(GetClass());
-		check(NodeSpawner != nullptr);
-
-		ActionRegistrar.AddBlueprintAction(ActionKey, NodeSpawner);
-	}
-}
-
-FText UK2Node_SpawnActorFromClass::GetMenuCategory() const
-{
-	return FEditorCategoryUtils::GetCommonCategory(FCommonEditorCategory::Gameplay);
-}
-
 FNodeHandlingFunctor* UK2Node_SpawnActorFromClass::CreateNodeHandler(FKismetCompilerContext& CompilerContext) const
 {
 	return new FNodeHandlingFunctor(CompilerContext);
@@ -590,19 +329,19 @@ void UK2Node_SpawnActorFromClass::ExpandNode(class FKismetCompilerContext& Compi
 {
 	Super::ExpandNode(CompilerContext, SourceGraph);
 
-	static FName BeginSpawningBlueprintFuncName = GET_FUNCTION_NAME_CHECKED(UGameplayStatics, BeginDeferredActorSpawnFromClass);
-	static FString ActorClassParamName = FString(TEXT("ActorClass"));
-	static FString WorldContextParamName = FString(TEXT("WorldContextObject"));
+	static const FName BeginSpawningBlueprintFuncName = GET_FUNCTION_NAME_CHECKED(UGameplayStatics, BeginDeferredActorSpawnFromClass);
+	static const FString ActorClassParamName = FString(TEXT("ActorClass"));
+	static const FString WorldContextParamName = FString(TEXT("WorldContextObject"));
 
-	static FName FinishSpawningFuncName = GET_FUNCTION_NAME_CHECKED(UGameplayStatics, FinishSpawningActor);
-	static FString ActorParamName = FString(TEXT("Actor"));
-	static FString TransformParamName = FString(TEXT("SpawnTransform"));
-	static FString CollisionHandlingOverrideParamName = FString(TEXT("CollisionHandlingOverride"));
-	static FString OwnerParamName = FString(TEXT("Owner"));
+	static const FName FinishSpawningFuncName = GET_FUNCTION_NAME_CHECKED(UGameplayStatics, FinishSpawningActor);
+	static const FString ActorParamName = FString(TEXT("Actor"));
+	static const FString TransformParamName = FString(TEXT("SpawnTransform"));
+	static const FString CollisionHandlingOverrideParamName = FString(TEXT("CollisionHandlingOverride"));
+	static const FString OwnerParamName = FString(TEXT("Owner"));
 
-	static FString ObjectParamName = FString(TEXT("Object"));
-	static FString ValueParamName = FString(TEXT("Value"));
-	static FString PropertyNameParamName = FString(TEXT("PropertyName"));
+	static const FString ObjectParamName = FString(TEXT("Object"));
+	static const FString ValueParamName = FString(TEXT("Value"));
+	static const FString PropertyNameParamName = FString(TEXT("PropertyName"));
 
 	UK2Node_SpawnActorFromClass* SpawnNode = this;
 	UEdGraphPin* SpawnNodeExec = SpawnNode->GetExecPin();
diff --git a/Engine/Source/Editor/BlueprintGraph/Private/K2Node_Timeline.cpp b/Engine/Source/Editor/BlueprintGraph/Private/K2Node_Timeline.cpp
--- a/Engine/Source/Editor/BlueprintGraph/Private/K2Node_Timeline.cpp
+++ b/Engine/Source/Editor/BlueprintGraph/Private/K2Node_Timeline.cpp
@@ -612,6 +612,14 @@ FSlateIcon UK2Node_Timeline::GetIconAndTint(FLinearColor& OutColor) const
 	return Icon;
 }
 
+UObject* UK2Node_Timeline::GetJumpTargetForDoubleClick() const
+{
+	UBlueprint* Blueprint = GetBlueprint();
+	check(Blueprint);
+	UTimelineTemplate* Timeline = Blueprint->FindTimelineTemplateByVariableName(TimelineName);
+	return Timeline;
+}
+
 FString UK2Node_Timeline::GetDocumentationExcerptName() const
 {
 	return TEXT("UK2Node_Timeline");
diff --git a/Engine/Source/Editor/BlueprintGraph/Private/K2Node_TunnelBoundary.cpp b/Engine/Source/Editor/BlueprintGraph/Private/K2Node_TunnelBoundary.cpp
--- a/Engine/Source/Editor/BlueprintGraph/Private/K2Node_TunnelBoundary.cpp
+++ b/Engine/Source/Editor/BlueprintGraph/Private/K2Node_TunnelBoundary.cpp
@@ -306,7 +306,6 @@ void UK2Node_TunnelBoundary::WireUpEntry(UEdGraphNode* SourceNode, UEdGraphPin*
 		UEdGraphPin* InputPin = CreatePin(EGPD_Input, SourcePin->PinType.PinCategory, SourcePin->PinType.PinSubCategory, SourcePin->PinType.PinSubCategoryObject.Get(), TEXT("EntryBoundary"), SourcePin->PinType.ContainerType, SourcePin->PinType.bIsReference, SourcePin->PinType.bIsConst, INDEX_NONE, SourcePin->PinType.PinValueType);
 		MessageLog.NotifyIntermediatePinCreation(InputPin, SourcePin);
 		UEdGraphPin* OutputPin = CreatePin(EGPD_Output, SourcePin->PinType.PinCategory, SourcePin->PinType.PinSubCategory, SourcePin->PinType.PinSubCategoryObject.Get(), SourcePin->PinName, SourcePin->PinType.ContainerType, SourcePin->PinType.bIsReference, SourcePin->PinType.bIsConst, INDEX_NONE, SourcePin->PinType.PinValueType);
-		MessageLog.AddPinTraceFilter(OutputPin);
 		for (UEdGraphPin* EntryPin : EntryPins)
 		{
 			for (UEdGraphPin* LinkedPin : EntryPin->LinkedTo)
@@ -333,7 +332,6 @@ void UK2Node_TunnelBoundary::WireUpExit(UEdGraphNode* SourceNode, UEdGraphPin* S
 		UEdGraphPin* OutputPin = CreatePin(EGPD_Output, SourcePin->PinType.PinCategory, SourcePin->PinType.PinSubCategory, SourcePin->PinType.PinSubCategoryObject.Get(), SourcePin->PinName, SourcePin->PinType.ContainerType, SourcePin->PinType.bIsReference, SourcePin->PinType.bIsConst, INDEX_NONE, SourcePin->PinType.PinValueType);
 		MessageLog.NotifyIntermediatePinCreation(OutputPin, SourcePin);
 		UEdGraphPin* InputPin = CreatePin(EGPD_Input, SourcePin->PinType.PinCategory, SourcePin->PinType.PinSubCategory, SourcePin->PinType.PinSubCategoryObject.Get(), TEXT("ExitBoundary"), SourcePin->PinType.ContainerType, SourcePin->PinType.bIsReference, SourcePin->PinType.bIsConst, INDEX_NONE, SourcePin->PinType.PinValueType);
-		MessageLog.AddPinTraceFilter(InputPin);
 		for (UEdGraphPin* ExitPin : ExitPins)
 		{
 			for (UEdGraphPin* LinkedPin : ExitPin->LinkedTo)
diff --git a/Engine/Source/Editor/BlueprintGraph/Private/K2Node_Variable.cpp b/Engine/Source/Editor/BlueprintGraph/Private/K2Node_Variable.cpp
--- a/Engine/Source/Editor/BlueprintGraph/Private/K2Node_Variable.cpp
+++ b/Engine/Source/Editor/BlueprintGraph/Private/K2Node_Variable.cpp
@@ -13,6 +13,7 @@
 #include "Kismet2/CompilerResultsLog.h"
 #include "Styling/SlateIconFinder.h"
 #include "Logging/MessageLog.h"
+#include "SourceCodeNavigation.h"
 
 #define LOCTEXT_NAMESPACE "K2Node"
 
@@ -970,4 +971,47 @@ FString UK2Node_Variable::GetDeprecationMessage() const
 	return Super::GetDeprecationMessage();
 }
 
+UObject* UK2Node_Variable::GetJumpTargetForDoubleClick() const
+{
+	// Jump to the RepNotify function graph if one exists
+	UBlueprint* OwnerBlueprint = GetBlueprint();
+
+	FName RepNotifyFunc = FBlueprintEditorUtils::GetBlueprintVariableRepNotifyFunc(OwnerBlueprint, GetVarName());
+	if (RepNotifyFunc != NAME_None)
+	{
+		for (UEdGraph* Graph : OwnerBlueprint->FunctionGraphs)
+		{
+			if (Graph->GetFName() == RepNotifyFunc)
+			{
+				return Graph;
+			}
+		}
+	}
+
+	return nullptr;
+}
+
+bool UK2Node_Variable::CanJumpToDefinition() const
+{
+	const UProperty* VariableProperty = GetPropertyForVariable();
+	const bool bNativeVariable = (VariableProperty != nullptr) && (VariableProperty->IsNative());
+	return bNativeVariable || (GetJumpTargetForDoubleClick() != nullptr);
+}
+
+void UK2Node_Variable::JumpToDefinition() const
+{
+	// For native variables, try going to the variable definition in C++ if available
+	if (UProperty* VariableProperty = GetPropertyForVariable())
+	{
+		if (VariableProperty->IsNative())
+		{
+			FSourceCodeNavigation::NavigateToProperty(VariableProperty);
+			return;
+		}
+	}
+
+	// Otherwise, fall back to the inherited behavior
+	Super::JumpToDefinition();
+}
+
 #undef LOCTEXT_NAMESPACE
diff --git a/Engine/Source/Editor/BlueprintGraph/Private/K2Node_VariableGet.cpp b/Engine/Source/Editor/BlueprintGraph/Private/K2Node_VariableGet.cpp
--- a/Engine/Source/Editor/BlueprintGraph/Private/K2Node_VariableGet.cpp
+++ b/Engine/Source/Editor/BlueprintGraph/Private/K2Node_VariableGet.cpp
@@ -367,7 +367,7 @@ void UK2Node_VariableGet::GetContextMenuActions(const FGraphNodeContextMenuBuild
 				FSlateIcon(),
 				FUIAction(
 				FExecuteAction::CreateUObject(this, &UK2Node_VariableGet::TogglePurity),
-				FCanExecuteAction::CreateStatic(CanExecutePurityToggle, bCanTogglePurity),
+				FCanExecuteAction::CreateStatic(CanExecutePurityToggle, bCanTogglePurity && !Context.bIsDebugging),
 				FIsActionChecked()
 				)
 				);
@@ -433,7 +433,7 @@ void UK2Node_VariableGet::ValidateNodeDuringCompilation(FCompilerResultsLog& Mes
 				if (PropertyReadableState == FBlueprintEditorUtils::EPropertyReadableState::NotBlueprintVisible)
 				{
 					// DEPRECATED(4.17) ... make this an error
-					MessageLog.Warning(*FText::Format(LOCTEXT("UnableToGet_NotVisible", "{VariableName} is not blueprint visible. Please fix mark up or cease accessing as this will be made an error in a future release. @@"), Args).ToString(), this);
+					MessageLog.Warning(*FText::Format(LOCTEXT("UnableToGet_NotVisible", "{VariableName} is not blueprint visible (BlueprintReadOnly or BlueprintReadWrite). Please fix mark up or cease accessing as this will be made an error in a future release. @@"), Args).ToString(), this);
 				}
 				else if (PropertyReadableState == FBlueprintEditorUtils::EPropertyReadableState::Private)
 				{
diff --git a/Engine/Source/Editor/ClassViewer/Private/SClassViewer.cpp b/Engine/Source/Editor/ClassViewer/Private/SClassViewer.cpp
--- a/Engine/Source/Editor/ClassViewer/Private/SClassViewer.cpp
+++ b/Engine/Source/Editor/ClassViewer/Private/SClassViewer.cpp
@@ -2642,7 +2642,7 @@ bool SClassViewer::IsCurrentDeveloperViewType(EClassViewerDeveloperType ViewType
 	return GetCurrentDeveloperViewType() == ViewType;
 }
 
-void SClassViewer::GetInternalOnlyClasses(TArray<FStringClassReference>& Classes)
+void SClassViewer::GetInternalOnlyClasses(TArray<FSoftClassPath>& Classes)
 {
 	if (!InitOptions.bAllowViewOptions)
 	{
@@ -2992,7 +2992,7 @@ void SClassViewer::Populate()
 
 	bool ShowingInternalClasses = IsShowingInternalClasses();
 
-	TArray<FStringClassReference> InternalClassNames;
+	TArray<FSoftClassPath> InternalClassNames;
 	TArray<UClass*> InternalClasses;
 	TArray<FDirectoryPath> InternalPaths;
 	// If we aren't showing the internal classes, then we need to know what classes to consider Internal Only, so let's gather them up from the settings object.
diff --git a/Engine/Source/Editor/ClassViewer/Private/SClassViewer.h b/Engine/Source/Editor/ClassViewer/Private/SClassViewer.h
--- a/Engine/Source/Editor/ClassViewer/Private/SClassViewer.h
+++ b/Engine/Source/Editor/ClassViewer/Private/SClassViewer.h
@@ -14,7 +14,7 @@
 #include "Widgets/Views/STreeView.h"
 #include "Settings/ClassViewerSettings.h"
 #include "Engine/EngineTypes.h"
-#include "Misc/StringClassReference.h"
+#include "UObject/SoftObjectPath.h"
 
 class FClassViewerNode;
 class FMenuBuilder;
@@ -221,7 +221,7 @@ private:
 	void HandleSettingChanged(FName PropertyName);
 
 	/** Accessor for the classnames that have been marked as internal only in settings */
-	void GetInternalOnlyClasses(TArray<FStringClassReference>& Classes);
+	void GetInternalOnlyClasses(TArray<FSoftClassPath>& Classes);
 	
 	/** Accessor for the class paths that have been marked as internal only in settings */
 	void GetInternalOnlyPaths(TArray<FDirectoryPath>& Paths);
diff --git a/Engine/Source/Editor/ClassViewer/Public/ClassViewerProjectSettings.h b/Engine/Source/Editor/ClassViewer/Public/ClassViewerProjectSettings.h
--- a/Engine/Source/Editor/ClassViewer/Public/ClassViewerProjectSettings.h
+++ b/Engine/Source/Editor/ClassViewer/Public/ClassViewerProjectSettings.h
@@ -6,7 +6,7 @@
 #include "UObject/ObjectMacros.h"
 #include "UObject/Object.h"
 #include "Engine/EngineTypes.h"
-#include "Misc/StringClassReference.h"
+#include "UObject/SoftObjectPath.h"
 #include "ClassViewerProjectSettings.generated.h"
 
 
@@ -25,6 +25,6 @@ class CLASSVIEWER_API UClassViewerProjectSettings : public UObject
 
 	/** The base classes to be considered Internal Only for the class picker.*/
 	UPROPERTY(EditAnywhere, config, Category = ClassVisibilityManagement, meta = (MetaClass = "Object", DisplayName = "List of base classes to consider Internal Only.", AllowAbstract = "true", ShowTreeView, HideViewOptions))
-	TArray<FStringClassReference> InternalOnlyClasses;
+	TArray<FSoftClassPath> InternalOnlyClasses;
 #endif
 };
diff --git a/Engine/Source/Editor/ClothingSystemEditorInterface/Public/ClothingAssetFactoryInterface.h b/Engine/Source/Editor/ClothingSystemEditorInterface/Public/ClothingAssetFactoryInterface.h
--- a/Engine/Source/Editor/ClothingSystemEditorInterface/Public/ClothingAssetFactoryInterface.h
+++ b/Engine/Source/Editor/ClothingSystemEditorInterface/Public/ClothingAssetFactoryInterface.h
@@ -5,7 +5,6 @@
 #include "Object.h"
 #include "ObjectMacros.h"
 #include "Features/IModularFeature.h"
-#include "AssetPtr.h"
 
 #include "ClothingAssetFactoryInterface.generated.h"
 
diff --git a/Engine/Source/Editor/DataTableEditor/Private/SRowEditor.cpp b/Engine/Source/Editor/DataTableEditor/Private/SRowEditor.cpp
--- a/Engine/Source/Editor/DataTableEditor/Private/SRowEditor.cpp
+++ b/Engine/Source/Editor/DataTableEditor/Private/SRowEditor.cpp
@@ -123,7 +123,7 @@ void SRowEditor::PreChange(const UDataTable* Changed, FDataTableEditorUtils::EDa
 
 void SRowEditor::PostChange(const UDataTable* Changed, FDataTableEditorUtils::EDataTableChangeInfo Info)
 {
-	FStringAssetReference::InvalidateTag(); // Should be removed after UE-5615 is fixed
+	FSoftObjectPath::InvalidateTag(); // Should be removed after UE-5615 is fixed
 	if ((Changed == DataTable.Get()) && (FDataTableEditorUtils::EDataTableChangeInfo::RowList == Info))
 	{
 		RefreshNameList();
diff --git a/Engine/Source/Editor/DataTableEditor/Private/SRowEditor.h b/Engine/Source/Editor/DataTableEditor/Private/SRowEditor.h
--- a/Engine/Source/Editor/DataTableEditor/Private/SRowEditor.h
+++ b/Engine/Source/Editor/DataTableEditor/Private/SRowEditor.h
@@ -49,7 +49,7 @@ private:
 
 	TArray<TSharedPtr<FName>> CachedRowNames;
 	TSharedPtr<FStructOnScope> CurrentRow;
-	TAssetPtr<UDataTable> DataTable; // weak obj ptr couldn't handle reimporting
+	TSoftObjectPtr<UDataTable> DataTable; // weak obj ptr couldn't handle reimporting
 	TSharedPtr<class IStructureDetailsView> StructureDetailsView;
 	TSharedPtr<FName> SelectedName;
 	TSharedPtr<SComboBox<TSharedPtr<FName>>> RowComboBox;
diff --git a/Engine/Source/Editor/DetailCustomizations/Private/DetailCustomizations.cpp b/Engine/Source/Editor/DetailCustomizations/Private/DetailCustomizations.cpp
--- a/Engine/Source/Editor/DetailCustomizations/Private/DetailCustomizations.cpp
+++ b/Engine/Source/Editor/DetailCustomizations/Private/DetailCustomizations.cpp
@@ -34,8 +34,8 @@
 #include "MathStructProxyCustomizations.h"
 #include "RangeStructCustomization.h"
 #include "IntervalStructCustomization.h"
-#include "StringAssetReferenceCustomization.h"
-#include "StringClassReferenceCustomization.h"
+#include "SoftObjectPathCustomization.h"
+#include "SoftClassPathCustomization.h"
 #include "AttenuationSettingsCustomizations.h"
 #include "WorldSettingsDetails.h"
 #include "DialogueStructsCustomizations.h"
@@ -169,8 +169,8 @@ void FDetailCustomizationsModule::ShutdownModule()
 
 void FDetailCustomizationsModule::RegisterPropertyTypeCustomizations()
 {
-	RegisterCustomPropertyTypeLayout("StringAssetReference", FOnGetPropertyTypeCustomizationInstance::CreateStatic(&FStringAssetReferenceCustomization::MakeInstance));
-	RegisterCustomPropertyTypeLayout("StringClassReference", FOnGetPropertyTypeCustomizationInstance::CreateStatic(&FStringClassReferenceCustomization::MakeInstance));
+	RegisterCustomPropertyTypeLayout("SoftObjectPath", FOnGetPropertyTypeCustomizationInstance::CreateStatic(&FSoftObjectPathCustomization::MakeInstance));
+	RegisterCustomPropertyTypeLayout("SoftClassPath", FOnGetPropertyTypeCustomizationInstance::CreateStatic(&FSoftClassPathCustomization::MakeInstance));
 	RegisterCustomPropertyTypeLayout("DataTableRowHandle", FOnGetPropertyTypeCustomizationInstance::CreateStatic(&FDataTableCustomizationLayout::MakeInstance));
 	RegisterCustomPropertyTypeLayout("DataTableCategoryHandle", FOnGetPropertyTypeCustomizationInstance::CreateStatic(&FDataTableCategoryCustomizationLayout::MakeInstance));
 	RegisterCustomPropertyTypeLayout("CurveTableRowHandle", FOnGetPropertyTypeCustomizationInstance::CreateStatic(&FCurveTableCustomizationLayout::MakeInstance));
diff --git a/Engine/Source/Editor/DetailCustomizations/Private/SoftObjectPathCustomization.cpp b/Engine/Source/Editor/DetailCustomizations/Private/SoftObjectPathCustomization.cpp
new file mode 100644
--- /dev/null
+++ b/Engine/Source/Editor/DetailCustomizations/Private/SoftObjectPathCustomization.cpp
@@ -0,0 +1,95 @@
+// Copyright 1998-2017 Epic Games, Inc. All Rights Reserved.
+
+#include "SoftObjectPathCustomization.h"
+#include "Widgets/DeclarativeSyntaxSupport.h"
+#include "Engine/GameViewportClient.h"
+#include "AssetData.h"
+#include "PropertyHandle.h"
+#include "PropertyCustomizationHelpers.h"
+
+void FSoftObjectPathCustomization::CustomizeHeader( TSharedRef<IPropertyHandle> InStructPropertyHandle, FDetailWidgetRow& HeaderRow, IPropertyTypeCustomizationUtils& StructCustomizationUtils )
+{
+	StructPropertyHandle = InStructPropertyHandle;
+
+	FString ClassFilterString = StructPropertyHandle->GetMetaData("AllowedClasses");
+	if( !ClassFilterString.IsEmpty() )
+	{
+		TArray<FString> CustomClassFilterNames;
+		ClassFilterString.ParseIntoArray(CustomClassFilterNames, TEXT(","), true);
+
+		for(auto It = CustomClassFilterNames.CreateConstIterator(); It; ++It)
+		{
+			const FString& ClassName = *It;
+
+			UClass* Class = FindObject<UClass>(ANY_PACKAGE, *ClassName);
+			if(!Class)
+			{
+				Class = LoadObject<UClass>(nullptr, *ClassName);
+			}
+			if(Class)
+			{
+				CustomClassFilters.Add(Class);
+			}
+		}
+	}
+
+	bExactClass = StructPropertyHandle->GetBoolMetaData("ExactClass");
+
+	FOnShouldFilterAsset AssetFilter;
+	UClass* ClassFilter = UObject::StaticClass();
+	if(CustomClassFilters.Num() == 1 && !bExactClass)
+	{
+		// If we only have one class to filter on, set it as the class type filter rather than use a filter callback
+		// We can only do this if we don't need an exact match, as the class filter also allows derived types
+		// The class filter is much faster than the callback as we're not performing two different sets of type tests
+		// (one against UObject, one against the actual type)
+		ClassFilter = CustomClassFilters[0];
+	}
+	else if(CustomClassFilters.Num() > 0)
+	{
+		// Only bind the filter if we have classes that need filtering
+		AssetFilter.BindSP(this, &FSoftObjectPathCustomization::OnShouldFilterAsset);
+	}
+
+	// Can the field be cleared
+	const bool bAllowClear = !(StructPropertyHandle->GetMetaDataProperty()->PropertyFlags & CPF_NoClear);
+
+	HeaderRow
+	.NameContent()
+	[
+		InStructPropertyHandle->CreatePropertyNameWidget()
+	]
+	.ValueContent()
+	.MinDesiredWidth(250.0f)
+	.MaxDesiredWidth(0.0f)
+	[
+		// Add an object entry box.  Even though this isn't an object entry, we will simulate one
+		SNew( SObjectPropertyEntryBox )
+		.PropertyHandle( InStructPropertyHandle )
+		.ThumbnailPool( StructCustomizationUtils.GetThumbnailPool() )
+		.AllowedClass( ClassFilter )
+		.OnShouldFilterAsset( AssetFilter )
+		.AllowClear( bAllowClear )
+	];
+}
+
+void FSoftObjectPathCustomization::CustomizeChildren(TSharedRef<IPropertyHandle> InStructPropertyHandle, IDetailChildrenBuilder& StructBuilder, IPropertyTypeCustomizationUtils& StructCustomizationUtils)
+{
+}
+
+bool FSoftObjectPathCustomization::OnShouldFilterAsset( const FAssetData& InAssetData ) const
+{
+	// Only bound if we have classes to filter on, so we don't need to test for an empty array here
+	UClass* const AssetClass = InAssetData.GetClass();
+	for(auto It = CustomClassFilters.CreateConstIterator(); It; ++It)
+	{
+		UClass* const FilterClass = *It;
+		const bool bMatchesFilter = (bExactClass) ? AssetClass == FilterClass : AssetClass->IsChildOf(FilterClass);
+		if(bMatchesFilter)
+		{
+			return false;
+		}
+	}
+
+	return true;
+}
diff --git a/Engine/Source/Editor/DetailCustomizations/Private/SoftObjectPathCustomization.h b/Engine/Source/Editor/DetailCustomizations/Private/SoftObjectPathCustomization.h
new file mode 100644
--- /dev/null
+++ b/Engine/Source/Editor/DetailCustomizations/Private/SoftObjectPathCustomization.h
@@ -0,0 +1,35 @@
+// Copyright 1998-2017 Epic Games, Inc. All Rights Reserved.
+
+#pragma once
+
+#include "CoreMinimal.h"
+#include "IPropertyTypeCustomization.h"
+
+class IPropertyHandle;
+
+/**
+ * Customizes a string asset reference to look like a UObject property
+ */
+class FSoftObjectPathCustomization : public IPropertyTypeCustomization
+{
+public:
+	static TSharedRef<IPropertyTypeCustomization> MakeInstance() 
+	{
+		return MakeShareable( new FSoftObjectPathCustomization );
+	}
+
+	/** IPropertyTypeCustomization interface */
+	virtual void CustomizeHeader( TSharedRef<class IPropertyHandle> InStructPropertyHandle, class FDetailWidgetRow& HeaderRow, IPropertyTypeCustomizationUtils& StructCustomizationUtils ) override;
+	virtual void CustomizeChildren( TSharedRef<class IPropertyHandle> InStructPropertyHandle, class IDetailChildrenBuilder& StructBuilder, IPropertyTypeCustomizationUtils& StructCustomizationUtils ) override;
+
+private:
+	bool OnShouldFilterAsset( const struct FAssetData& InAssetData ) const;
+
+private:
+	/** Handle to the struct property being customized */
+	TSharedPtr<IPropertyHandle> StructPropertyHandle;
+	/** Classes that can be used with this property */
+	TArray<UClass*> CustomClassFilters;
+	/** Whether the classes in the above array must be an exact match, or whether they can also be a derived type; default is false */
+	bool bExactClass;
+};
diff --git a/Engine/Source/Editor/EditorStyle/Private/SlateEditorStyle.cpp b/Engine/Source/Editor/EditorStyle/Private/SlateEditorStyle.cpp
--- a/Engine/Source/Editor/EditorStyle/Private/SlateEditorStyle.cpp
+++ b/Engine/Source/Editor/EditorStyle/Private/SlateEditorStyle.cpp
@@ -3788,6 +3788,8 @@ void FSlateEditorStyle::FStyle::SetupGraphEditorStyles()
 		Set( "Graph.Node.NodeBackground", new BOX_BRUSH( "Old/Graph/NodeBackground", FMargin(12.0f/64) ) );
 
 		Set( "Graph.Node.Body", new BOX_BRUSH( "/Graph/RegularNode_body", FMargin(16.f/64.f, 25.f/64.f, 16.f/64.f, 16.f/64.f) ) );
+		Set( "Graph.Node.DisabledBanner", new IMAGE_BRUSH( "/Graph/GraphPanel_StripesBackground", FVector2D(64, 64), FLinearColor(0.5f, 0.5f, 0.5f, 0.3f), ESlateBrushTileType::Both ) );
+		Set( "Graph.Node.DevelopmentBanner", new IMAGE_BRUSH( "/Graph/GraphPanel_StripesBackground", FVector2D(64, 64), FLinearColor::Yellow * FLinearColor(1.f, 1.f, 1.f, 0.3f), ESlateBrushTileType::Both ) );
 		Set( "Graph.Node.TitleGloss", new BOX_BRUSH( "/Graph/RegularNode_title_gloss", FMargin(12.0f/64.0f) ) );
 		Set( "Graph.Node.ColorSpill", new BOX_BRUSH( "/Graph/RegularNode_color_spill", FMargin(8.0f/64.0f, 3.0f/32.0f, 0, 0) ) );
 		Set( "Graph.Node.TitleHighlight", new BOX_BRUSH( "/Graph/RegularNode_title_highlight", FMargin(16.0f/64.0f, 1.0f, 16.0f/64.0f, 0.0f) ) );
@@ -4942,85 +4944,6 @@ void FSlateEditorStyle::FStyle::SetupLevelEditorStyle()
 	}
 #endif // WITH_EDITOR || IS_PROGRAM
 
-	// Blueprint Profiler
-#if WITH_EDITOR
-	{
-		// Profiler stat window
-		Set( "BlueprintProfiler.ViewToolBar", new BOX_BRUSH("/Docking/Tab_Foreground", 4/16.0f));
-		Set( "BlueprintProfiler.ViewContent", new BOX_BRUSH("/Docking/AppTabContentArea", 4/16.0f));
-		Set( "BlueprintProfiler.BPIcon_Normal", new IMAGE_BRUSH("Icons/AssetIcons/Blueprint_16x", Icon16x16, FLinearColor( 0.46f, 0.54f, 0.81f )));
-		Set( "BlueprintProfiler.BPIcon_Macro", new IMAGE_BRUSH("Icons/AssetIcons/BlueprintMacroLibrary_16x", Icon16x16, FLinearColor( 0.46f, 0.54f, 0.81f )));
-		Set( "BlueprintProfiler.BPIcon_Interface", new IMAGE_BRUSH("Icons/AssetIcons/BlueprintInterface_16x", Icon16x16, FLinearColor( 0.46f, 0.54f, 0.81f )));
-		Set( "BlueprintProfiler.BPIcon_FunctionLibrary", new IMAGE_BRUSH("Icons/AssetIcons/BlueprintFunctionLibrary_16x", Icon16x16, FLinearColor( 0.46f, 0.54f, 0.81f )));
-		Set( "BlueprintProfiler.Actor", new IMAGE_BRUSH("Icons/AssetIcons/Actor_16x", Icon16x16, FLinearColor( 1.f, 1.f, 1.f )));
-		Set( "BlueprintProfiler.BPEventNode", new IMAGE_BRUSH("Icons/icon_Blueprint_Event_16x", Icon16x16, FLinearColor( 0.91f, 0.16f, 0.16f )));
-		Set( "BlueprintProfiler.BPFunctionNode", new IMAGE_BRUSH("Icons/icon_Blueprint_NewFunction_16x", Icon16x16, FLinearColor( 0.46f, 0.54f, 0.81f )));
-		Set( "BlueprintProfiler.BPMacroNode", new IMAGE_BRUSH("Icons/icon_Blueprint_Macro_16x", Icon16x16, FLinearColor( 0.46f, 0.54f, 0.81f )));
-		Set( "BlueprintProfiler.BPBranchNode", new IMAGE_BRUSH("Icons/icon_Blueprint_Sequence_16x", Icon16x16, FLinearColor( 0.46f, 0.54f, 0.81f )));
-		Set( "BlueprintProfiler.BPPinConnected", new IMAGE_BRUSH( "Icons/BPProfiler_ExecPin_Connected", Icon12x16 ) );
-		Set( "BlueprintProfiler.BPPinDisconnected", new IMAGE_BRUSH( "Icons/BPProfiler_ExecPin_Disconnected", Icon12x16 ) );
-		Set( "BlueprintProfiler.PureNode", new IMAGE_BRUSH( "/Icons/pill_16x", Icon16x16 ) );
-		Set( "BlueprintProfiler.BPNode", new IMAGE_BRUSH("Icons/Icon_Profiler_BPNode_16x", Icon16x16));
-		Set( "BlueprintProfiler.RegularNode.HeatDisplay", new BOX_BRUSH( "/Graph/RegularNode_heat_display", FMargin(36.0f/64.0f) ) );
-		Set( "BlueprintProfiler.VarNode.HeatDisplay", new BOX_BRUSH( "/Graph/VarNode_heat_display", FMargin(36.0f/64.0f) ) );
-		Set( "BlueprintProfiler.CompactNode.HeatDisplay", new BOX_BRUSH( "/Graph/MathNode_heat_display", FMargin(26.0f/64.0f) ) );
-
-		Button = FButtonStyle()
-			.SetNormal( FSlateNoResource() )
-			.SetHovered( FSlateNoResource() )
-			.SetPressed( FSlateNoResource() )
-			.SetNormalPadding( FMargin( 2,2,2,2 ) )
-			.SetPressedPadding( FMargin( 2,2,2,2 ) );
-		Set( "BlueprintProfiler.TreeArrowButton", Button );
-
-		Set( "BlueprintProfiler.TreeArrowGlyph", FTextBlockStyle(NormalText)
-			.SetFont(TTF_FONT("Fonts/FontAwesome", 10))
-			.SetShadowOffset(FVector2D::ZeroVector)
-		);
-
-		Set("BlueprintProfiler.ToggleButton.Start", FCheckBoxStyle()
-			.SetCheckBoxType(ESlateCheckBoxType::ToggleButton)
-			.SetUncheckedImage(BOX_BRUSH("Common/SmallRoundedToggleLeft", FMargin(7.f / 16.f), FLinearColor(0.09f, 0.09f, 0.09f)))
-			.SetUncheckedPressedImage(BOX_BRUSH("Common/SmallRoundedToggleLeft", FMargin(7.f / 16.f), SelectionColor_Pressed))
-			.SetUncheckedHoveredImage(BOX_BRUSH("Common/SmallRoundedToggleLeft", FMargin(7.f / 16.f), SelectionColor_Pressed))
-			.SetCheckedHoveredImage(BOX_BRUSH("Common/SmallRoundedToggleLeft", FMargin(7.f / 16.f), SelectionColor))
-			.SetCheckedPressedImage(BOX_BRUSH("Common/SmallRoundedToggleLeft", FMargin(7.f / 16.f), SelectionColor))
-			.SetCheckedImage(BOX_BRUSH("Common/SmallRoundedToggleLeft", FMargin(7.f / 16.f), SelectionColor)));
-
-		Set("BlueprintProfiler.ToggleButton.Middle", FCheckBoxStyle()
-			.SetCheckBoxType(ESlateCheckBoxType::ToggleButton)
-			.SetUncheckedImage(BOX_BRUSH("Common/SmallRoundedToggleCenter", FMargin(7.f / 16.f), FLinearColor(0.09f, 0.09f, 0.09f)))
-			.SetUncheckedPressedImage(BOX_BRUSH("Common/SmallRoundedToggleCenter", FMargin(7.f / 16.f), SelectionColor_Pressed))
-			.SetUncheckedHoveredImage(BOX_BRUSH("Common/SmallRoundedToggleCenter", FMargin(7.f / 16.f), SelectionColor_Pressed))
-			.SetCheckedHoveredImage(BOX_BRUSH("Common/SmallRoundedToggleCenter", FMargin(7.f / 16.f), SelectionColor))
-			.SetCheckedPressedImage(BOX_BRUSH("Common/SmallRoundedToggleCenter", FMargin(7.f / 16.f), SelectionColor))
-			.SetCheckedImage(BOX_BRUSH("Common/SmallRoundedToggleCenter", FMargin(7.f / 16.f), SelectionColor)));
-
-		Set("BlueprintProfiler.ToggleButton.End", FCheckBoxStyle()
-			.SetCheckBoxType(ESlateCheckBoxType::ToggleButton)
-			.SetUncheckedImage(BOX_BRUSH("Common/SmallRoundedToggleRight", FMargin(7.f / 16.f), FLinearColor(0.09f, 0.09f, 0.09f)))
-			.SetUncheckedPressedImage(BOX_BRUSH("Common/SmallRoundedToggleRight", FMargin(7.f / 16.f), SelectionColor_Pressed))
-			.SetUncheckedHoveredImage(BOX_BRUSH("Common/SmallRoundedToggleRight", FMargin(7.f / 16.f), SelectionColor_Pressed))
-			.SetCheckedHoveredImage(BOX_BRUSH("Common/SmallRoundedToggleRight", FMargin(7.f / 16.f), SelectionColor))
-			.SetCheckedPressedImage(BOX_BRUSH("Common/SmallRoundedToggleRight", FMargin(7.f / 16.f), SelectionColor))
-			.SetCheckedImage(BOX_BRUSH("Common/SmallRoundedToggleRight", FMargin(7.f / 16.f), SelectionColor)));
-
-		Set("BlueprintProfiler.TableView.Row", FTableRowStyle(NormalTableRowStyle)
-			.SetActiveBrush(IMAGE_BRUSH("Common/Selection", Icon8x8, FLinearColor(0.25f, 0.25f, 0.25f, 0.35f)))
-			.SetActiveHoveredBrush(IMAGE_BRUSH("Common/Selection", Icon8x8, FLinearColor(0.25f, 0.25f, 0.25f, 0.5f)))
-			.SetInactiveBrush(IMAGE_BRUSH("Common/Selection", Icon8x8, FLinearColor(0.25f, 0.25f, 0.25f, 0.15f)))
-			.SetInactiveHoveredBrush(IMAGE_BRUSH("Common/Selection", Icon8x8, FLinearColor(0.25f, 0.25f, 0.25f, 0.2f))));
-
-		FTextBlockStyle NodeHyperlinkText = FTextBlockStyle(NormalText)
-			.SetColorAndOpacity(FLinearColor(1.0f, 1.0f, 1.0f))
-			.SetHighlightColor(FLinearColor(1.0f, 1.0f, 1.0f));
-
-		FHyperlinkStyle NodeHyperlinkStyle = GetWidgetStyle<FHyperlinkStyle>("HoverOnlyHyperlink");
-		Set("BlueprintProfiler.Node.Hyperlink", NodeHyperlinkStyle
-			.SetTextStyle(NodeHyperlinkText));
-	}
-#endif // WITH_EDITOR
-
 	// Mobility Icons
 	{
 		Set("Mobility.Movable", new IMAGE_BRUSH("/Icons/Mobility/Movable_16x", Icon16x16));
diff --git a/Engine/Source/Editor/EditorWidgets/Private/SObjectNameEditableTextBox.cpp b/Engine/Source/Editor/EditorWidgets/Private/SObjectNameEditableTextBox.cpp
--- a/Engine/Source/Editor/EditorWidgets/Private/SObjectNameEditableTextBox.cpp
+++ b/Engine/Source/Editor/EditorWidgets/Private/SObjectNameEditableTextBox.cpp
@@ -8,6 +8,7 @@
 #include "GameFramework/Actor.h"
 #include "ScopedTransaction.h"
 #include "ActorEditorUtils.h"
+#include "Editor/EditorEngine.h"
 
 const float SObjectNameEditableTextBox::HighlightRectLeftOffset = 0.0f;
 const float SObjectNameEditableTextBox::HighlightRectRightOffset = 0.0f;
@@ -258,7 +259,7 @@ void SObjectNameEditableTextBox::OnNameTextCommitted(const FText& NewText, EText
 
 					if (Actor->IsActorLabelEditable())
 					{
-						Actor->SetActorLabel(TrimmedText.ToString());
+						FActorLabelUtilities::RenameExistingActor(Actor, TrimmedText.ToString());
 						LastCommittedTime = FSlateApplication::Get().GetCurrentTime();
 						RegisterActiveTimer( 0.f, FWidgetActiveTimerDelegate::CreateSP( this, &SObjectNameEditableTextBox::UpdateHighlightSpringState ) );
 					}
@@ -278,7 +279,7 @@ void SObjectNameEditableTextBox::OnNameTextCommitted(const FText& NewText, EText
 						AActor* Actor = (AActor*)Objects[i].Get();
 						if (Actor->IsActorLabelEditable())
 						{
-							Actor->SetActorLabel(TrimmedText.ToString());
+							FActorLabelUtilities::RenameExistingActor(Actor, TrimmedText.ToString());
 							LastCommittedTime = FSlateApplication::Get().GetCurrentTime();
 							RegisterActiveTimer( 0.f, FWidgetActiveTimerDelegate::CreateSP( this, &SObjectNameEditableTextBox::UpdateHighlightSpringState ) );
 						}
diff --git a/Engine/Source/Editor/GraphEditor/Private/BlueprintConnectionDrawingPolicy.cpp b/Engine/Source/Editor/GraphEditor/Private/BlueprintConnectionDrawingPolicy.cpp
--- a/Engine/Source/Editor/GraphEditor/Private/BlueprintConnectionDrawingPolicy.cpp
+++ b/Engine/Source/Editor/GraphEditor/Private/BlueprintConnectionDrawingPolicy.cpp
@@ -508,7 +508,7 @@ void FKismetConnectionDrawingPolicy::DetermineWiringStyle(UEdGraphPin* OutputPin
 						DetermineStyleOfExecWire(/*inout*/ Params.WireThickness, /*inout*/ Params.WireColor, /*inout*/ Params.bDrawBubbles, *ExecTiming);
 					}
 				}
-				
+
 				if (!bExecuted)
 				{
 					// It's not followed, fade it and keep it thin
@@ -530,6 +530,31 @@ void FKismetConnectionDrawingPolicy::DetermineWiringStyle(UEdGraphPin* OutputPin
 				Params.WireThickness = DefaultExecutionWireThickness;
 			}
 		}
+
+
+		// If either end of the connection is not enabled (and not a passthru to something else), draw the wire differently
+		bool bWireIsOnDisabledNodeAndNotPassthru = false;
+		if (OutputNode && !OutputNode->IsNodeEnabled())
+		{
+			if (OutputNode->GetPassThroughPin(OutputPin) == nullptr)
+			{
+				bWireIsOnDisabledNodeAndNotPassthru = true;
+			}
+		}
+		
+		if (InputNode && !InputNode->IsNodeEnabled())
+		{
+			if (InputNode->GetPassThroughPin(InputPin) == nullptr)
+			{
+				bWireIsOnDisabledNodeAndNotPassthru = true;
+			}
+		}
+
+		if (bWireIsOnDisabledNodeAndNotPassthru)
+		{
+			Params.WireColor *= 0.5f;
+			Params.WireThickness = 0.5f;
+		}
 	}
 
 	if (bDeemphasizeUnhoveredPins)
diff --git a/Engine/Source/Editor/GraphEditor/Private/GraphEditorActions.cpp b/Engine/Source/Editor/GraphEditor/Private/GraphEditorActions.cpp
--- a/Engine/Source/Editor/GraphEditor/Private/GraphEditorActions.cpp
+++ b/Engine/Source/Editor/GraphEditor/Private/GraphEditorActions.cpp
@@ -60,13 +60,10 @@ void FGraphEditorCommandsImpl::RegisterCommands()
 
 	UI_COMMAND( SelectReferenceInLevel, "Find Actor in Level", "Select the actor referenced by this node in the level", EUserInterfaceActionType::Button, FInputChord() )
 	UI_COMMAND( AssignReferencedActor, "Assign selected Actor", "Assign the selected actor to be this node's referenced object", EUserInterfaceActionType::Button, FInputChord() )
-	UI_COMMAND( FindReferences, "Find References", "Find references of this item", EUserInterfaceActionType::Button, FInputChord() )
+	UI_COMMAND( FindReferences, "Find References", "Find references of this item", EUserInterfaceActionType::Button, FInputChord(EKeys::F, EModifierKey::Shift | EModifierKey::Alt) )
 	UI_COMMAND( FindAndReplaceReferences, "Find and Replace References", "Brings up a window to help find and replace all instances of this item", EUserInterfaceActionType::Button, FInputChord() )
 	
-	UI_COMMAND( GotoNativeFunctionDefinition, "Goto Code Definition", "Goto the native code definition of this function", EUserInterfaceActionType::Button, FInputChord() )
-	UI_COMMAND( GotoNativeVariableDefinition, "Goto Code Definition", "Goto the native code definition of this variable", EUserInterfaceActionType::Button, FInputChord() )
-
-	UI_COMMAND( GoToDefinition, "Goto Definition", "Jumps to the graph this node is defined in if available.", EUserInterfaceActionType::Button, FInputChord() )
+	UI_COMMAND( GoToDefinition, "Goto Definition", "Jumps to the defintion of the selected node if available, e.g., C++ code for a native function or the graph for a Blueprint function.", EUserInterfaceActionType::Button, FInputChord(EKeys::G, EModifierKey::Alt) )
 
 	UI_COMMAND( BreakPinLinks, "Break Link(s)", "Breaks pin links", EUserInterfaceActionType::Button, FInputChord() )
 	UI_COMMAND( PromoteToVariable, "Promote to Variable", "Promotes something to a variable", EUserInterfaceActionType::Button, FInputChord() )
diff --git a/Engine/Source/Editor/GraphEditor/Private/GraphEditorDragDropAction.cpp b/Engine/Source/Editor/GraphEditor/Private/GraphEditorDragDropAction.cpp
--- a/Engine/Source/Editor/GraphEditor/Private/GraphEditorDragDropAction.cpp
+++ b/Engine/Source/Editor/GraphEditor/Private/GraphEditorDragDropAction.cpp
@@ -195,21 +195,33 @@ EVisibility FGraphEditorDragDropAction::GetErrorIconVisible() const
 
 void FGraphSchemaActionDragDropAction::HoverTargetChanged()
 {
-	if(ActionNode.IsValid())
+	if (SourceAction.IsValid())
 	{
-		const FSlateBrush* StatusSymbol = FEditorStyle::GetBrush(TEXT("Graph.ConnectorFeedback.NewNode"));
+		const FSlateBrush* PrimarySymbol;
+		const FSlateBrush* SecondarySymbol;
+		FSlateColor PrimaryColor;
+		FSlateColor SecondaryColor;
+		GetDefaultStatusSymbol(/*out*/ PrimarySymbol, /*out*/ PrimaryColor, /*out*/ SecondarySymbol, /*out*/ SecondaryColor);
 
 		//Create feedback message with the function name.
-		SetSimpleFeedbackMessage(StatusSymbol, FLinearColor::White, ActionNode->GetMenuDescription());
+		SetSimpleFeedbackMessage(PrimarySymbol, PrimaryColor, SourceAction->GetMenuDescription(), SecondarySymbol, SecondaryColor);
 	}
 }
 
+void FGraphSchemaActionDragDropAction::GetDefaultStatusSymbol(const FSlateBrush*& PrimaryBrushOut, FSlateColor& IconColorOut, FSlateBrush const*& SecondaryBrushOut, FSlateColor& SecondaryColorOut) const
+{
+	PrimaryBrushOut = FEditorStyle::GetBrush(TEXT("Graph.ConnectorFeedback.NewNode"));
+	IconColorOut = FLinearColor::White;
+	SecondaryBrushOut = nullptr;
+	SecondaryColorOut = FLinearColor::White;
+}
+
 FReply FGraphSchemaActionDragDropAction::DroppedOnPanel( const TSharedRef< SWidget >& Panel, FVector2D ScreenPosition, FVector2D GraphPosition, UEdGraph& Graph)
 {
-	if(ActionNode.IsValid())
+	if (SourceAction.IsValid())
 	{
 		TArray<UEdGraphPin*> DummyPins;
-		ActionNode->PerformAction(&Graph, DummyPins, GraphPosition);
+		SourceAction->PerformAction(&Graph, DummyPins, GraphPosition);
 
 		return FReply::Handled();
 	}
@@ -220,7 +232,7 @@ FReply FGraphSchemaActionDragDropAction::DroppedOnPin(FVector2D ScreenPosition,
 {
 	if (UEdGraph* Graph = GetHoveredGraph())
 	{
-		if (ActionNode.IsValid())
+		if (SourceAction.IsValid())
 		{
 			TArray<UEdGraphPin*> DummyPins;
 			if (UEdGraphPin* Pin = GetHoveredPin())
@@ -228,7 +240,7 @@ FReply FGraphSchemaActionDragDropAction::DroppedOnPin(FVector2D ScreenPosition,
 				DummyPins.Add(Pin);
 			}
 
-			ActionNode->PerformAction(Graph, DummyPins, GraphPosition);
+			SourceAction->PerformAction(Graph, DummyPins, GraphPosition);
 
 			return FReply::Handled();
 		}
diff --git a/Engine/Source/Editor/GraphEditor/Private/GraphEditorSettings.cpp b/Engine/Source/Editor/GraphEditor/Private/GraphEditorSettings.cpp
--- a/Engine/Source/Editor/GraphEditor/Private/GraphEditorSettings.cpp
+++ b/Engine/Source/Editor/GraphEditor/Private/GraphEditorSettings.cpp
@@ -40,8 +40,8 @@ UGraphEditorSettings::UGraphEditorSettings( const FObjectInitializer& ObjectInit
 	NamePinTypeColor = FLinearColor(0.607717f, 0.224984f, 1.0f, 1.0f);			// lilac
 	DelegatePinTypeColor = FLinearColor(1.0f, 0.04f, 0.04f, 1.0f);				// bright red
 	ObjectPinTypeColor = FLinearColor(0.0f, 0.4f, 0.910000f, 1.0f);				// sharp blue
-	AssetPinTypeColor = FLinearColor(0.3f, 1.0f, 1.0f, 1.0f);					
-	AssetClassPinTypeColor = FLinearColor(1.0f, 0.3f, 1.0f, 1.0f);				
+	SoftObjectPinTypeColor = FLinearColor(0.3f, 1.0f, 1.0f, 1.0f);					
+	SoftClassPinTypeColor = FLinearColor(1.0f, 0.3f, 1.0f, 1.0f);				
 	InterfacePinTypeColor = FLinearColor(0.8784f, 1.0f, 0.4f, 1.0f);			// pale green
 	StringPinTypeColor = FLinearColor(1.0f, 0.0f, 0.660537f, 1.0f);				// bright pink
 	TextPinTypeColor = FLinearColor(0.8f, 0.2f, 0.4f, 1.0f);					// salmon (light pink)
diff --git a/Engine/Source/Editor/GraphEditor/Private/KismetNodes/SGraphNodeK2Base.cpp b/Engine/Source/Editor/GraphEditor/Private/KismetNodes/SGraphNodeK2Base.cpp
--- a/Engine/Source/Editor/GraphEditor/Private/KismetNodes/SGraphNodeK2Base.cpp
+++ b/Engine/Source/Editor/GraphEditor/Private/KismetNodes/SGraphNodeK2Base.cpp
@@ -26,11 +26,6 @@
 #include "TutorialMetaData.h"
 #include "Widgets/Layout/SBox.h"
 
-// Blueprint Profiler
-#include "Developer/BlueprintProfiler/Public/BlueprintProfilerModule.h"
-#include "Editor/Kismet/Public/Profiler/BlueprintProfilerSettings.h"
-#include "Editor/Kismet/Public/Profiler/EventExecution.h"
-
 #define LOCTEXT_NAMESPACE "SGraphNodeK2Base"
 
 //////////////////////////////////////////////////////////////////////////
@@ -575,53 +570,6 @@ const FSlateBrush* SGraphNodeK2Base::GetShadowBrush(bool bSelected) const
 	}
 }
 
-FLinearColor SGraphNodeK2Base::GetProfilerHeatmapIntensity() const
-{
-	float IntensityValue = 0.0f;
-	IBlueprintProfilerInterface& ProfilerModule = FModuleManager::LoadModuleChecked<IBlueprintProfilerInterface>("BlueprintProfiler");
-	if (ProfilerModule.IsProfilerEnabled() && GraphNode)
-	{
-		UBlueprint* Blueprint = FBlueprintEditorUtils::FindBlueprintForNode(GraphNode);
-		if (Blueprint && Blueprint->GeneratedClass && Blueprint->GeneratedClass->HasInstrumentation())
-		{
-			TSharedPtr<FScriptExecutionNode> ExecNode = ProfilerModule.GetProfilerDataForNode(GraphNode);
-			if (ExecNode.IsValid())
-			{
-				const FScriptPerfData& NodePerfData = ExecNode->GetNodePerfData();
-				switch (GetDefault<UBlueprintProfilerSettings>()->GraphNodeHeatMapDisplayMode)
-				{
-				case EBlueprintProfilerHeatMapDisplayMode::Average:
-					IntensityValue = NodePerfData.GetAverageHeatLevel();
-					break;
-
-				case EBlueprintProfilerHeatMapDisplayMode::Inclusive:
-					IntensityValue = NodePerfData.GetInclusiveHeatLevel();
-					break;
-
-				case EBlueprintProfilerHeatMapDisplayMode::MaxTiming:
-					IntensityValue = NodePerfData.GetMaxTimeHeatLevel();
-					break;
-
-				case EBlueprintProfilerHeatMapDisplayMode::Total:
-					IntensityValue = NodePerfData.GetTotalHeatLevel();
-					break;
-				}
-			}
-		}
-	}
-	const float Value = 1.f - IntensityValue;
-	return FLinearColor(1.f, Value, Value, IntensityValue*IntensityValue);
-}
-
-const FSlateBrush* SGraphNodeK2Base::GetProfilerHeatmapBrush() const
-{
-	const UK2Node* K2Node = CastChecked<UK2Node>(GraphNode);
-	const bool bCompactMode = K2Node->ShouldDrawCompact();
-
-	return bCompactMode ?	FEditorStyle::GetBrush(TEXT("BlueprintProfiler.CompactNode.HeatDisplay")) : 
-							FEditorStyle::GetBrush(TEXT("BlueprintProfiler.RegularNode.HeatDisplay"));
-}
-
 void SGraphNodeK2Base::PerformSecondPassLayout(const TMap< UObject*, TSharedRef<SNode> >& NodeToWidgetLookup) const
 {
 	TSet<UEdGraphNode*> PrevNodes;
diff --git a/Engine/Source/Editor/GraphEditor/Private/KismetNodes/SGraphNodeK2Var.cpp b/Engine/Source/Editor/GraphEditor/Private/KismetNodes/SGraphNodeK2Var.cpp
--- a/Engine/Source/Editor/GraphEditor/Private/KismetNodes/SGraphNodeK2Var.cpp
+++ b/Engine/Source/Editor/GraphEditor/Private/KismetNodes/SGraphNodeK2Var.cpp
@@ -322,8 +322,3 @@ const FSlateBrush* SGraphNodeK2Var::GetShadowBrush(bool bSelected) const
 {
 	return bSelected ? FEditorStyle::GetBrush(TEXT("Graph.VarNode.ShadowSelected")) : FEditorStyle::GetBrush(TEXT("Graph.VarNode.Shadow"));
 }
-
-const FSlateBrush* SGraphNodeK2Var::GetProfilerHeatmapBrush() const
-{
-	return FEditorStyle::GetBrush(TEXT("BlueprintProfiler.VarNode.HeatDisplay"));
-}
diff --git a/Engine/Source/Editor/GraphEditor/Private/KismetNodes/SGraphNodeK2Var.h b/Engine/Source/Editor/GraphEditor/Private/KismetNodes/SGraphNodeK2Var.h
--- a/Engine/Source/Editor/GraphEditor/Private/KismetNodes/SGraphNodeK2Var.h
+++ b/Engine/Source/Editor/GraphEditor/Private/KismetNodes/SGraphNodeK2Var.h
@@ -18,7 +18,6 @@ public:
 	void Construct( const FArguments& InArgs, UK2Node* InNode );
 
 	virtual const FSlateBrush* GetShadowBrush(bool bSelected) const override;
-	virtual const FSlateBrush* GetProfilerHeatmapBrush() const override;
 
 	// SGraphNode interface
 	virtual void UpdateGraphNode() override;
diff --git a/Engine/Source/Editor/GraphEditor/Private/KismetPins/SGraphPinDataTableRowName.cpp b/Engine/Source/Editor/GraphEditor/Private/KismetPins/SGraphPinDataTableRowName.cpp
--- a/Engine/Source/Editor/GraphEditor/Private/KismetPins/SGraphPinDataTableRowName.cpp
+++ b/Engine/Source/Editor/GraphEditor/Private/KismetPins/SGraphPinDataTableRowName.cpp
@@ -23,7 +23,7 @@ void SGraphPinDataTableRowName::PreChange(const UDataTable* Changed, FDataTableE
 
 void SGraphPinDataTableRowName::PostChange(const UDataTable* Changed, FDataTableEditorUtils::EDataTableChangeInfo Info)
 {
-	//FStringAssetReference::InvalidateTag(); // Should be removed after UE-5615 is fixed
+	//FSoftObjectPath::InvalidateTag(); // Should be removed after UE-5615 is fixed
 	if (Changed && (Changed == DataTable.Get()) && (FDataTableEditorUtils::EDataTableChangeInfo::RowList == Info))
 	{
 		RefreshNameList();
diff --git a/Engine/Source/Editor/GraphEditor/Private/KismetPins/SGraphPinKey.cpp b/Engine/Source/Editor/GraphEditor/Private/KismetPins/SGraphPinKey.cpp
--- a/Engine/Source/Editor/GraphEditor/Private/KismetPins/SGraphPinKey.cpp
+++ b/Engine/Source/Editor/GraphEditor/Private/KismetPins/SGraphPinKey.cpp
@@ -12,7 +12,7 @@ void SGraphPinKey::Construct(const FArguments& InArgs, UEdGraphPin* InGraphPinOb
 
 	SelectedKey = FKey(*InGraphPinObj->GetDefaultAsString());
 
-	if (!SelectedKey.IsValid())
+	if (!SelectedKey.IsValid() && InGraphPinObj->Direction == EEdGraphPinDirection::EGPD_Input)
 	{
 		// Ensure first valid key is always set by default
 		SelectedKey = KeyList[0];
diff --git a/Engine/Source/Editor/GraphEditor/Private/NodeFactory.cpp b/Engine/Source/Editor/GraphEditor/Private/NodeFactory.cpp
--- a/Engine/Source/Editor/GraphEditor/Private/NodeFactory.cpp
+++ b/Engine/Source/Editor/GraphEditor/Private/NodeFactory.cpp
@@ -196,7 +196,7 @@ TSharedPtr<SGraphNode> FNodeFactory::CreateNodeWidget(UEdGraphNode* InNode)
 	{
 		return SNew(SGraphNodeDocumentation, DocNode);
 	}
-	else if (InNode->ShouldDrawNodeAsComment())
+	else if (UEdGraphNode_Comment* CommentNode = Cast<UEdGraphNode_Comment>(InNode))
 	{
 		if (UMaterialGraphNode_Comment* MaterialCommentNode = Cast<UMaterialGraphNode_Comment>(InNode))
 		{
@@ -204,7 +204,7 @@ TSharedPtr<SGraphNode> FNodeFactory::CreateNodeWidget(UEdGraphNode* InNode)
 		}
 		else
 		{
-			return SNew(SGraphNodeComment, InNode);
+			return SNew(SGraphNodeComment, CommentNode);
 		}
 	}
 	else
@@ -254,7 +254,7 @@ TSharedPtr<SGraphPin> FNodeFactory::CreatePinWidget(UEdGraphPin* InPin)
 		{
 			return SNew(SGraphPinObject, InPin);
 		}
-		else if (InPin->PinType.PinCategory == K2Schema->PC_Asset)
+		else if (InPin->PinType.PinCategory == K2Schema->PC_SoftObject)
 		{
 			return SNew(SGraphPinObject, InPin);
 		}
@@ -262,7 +262,7 @@ TSharedPtr<SGraphPin> FNodeFactory::CreatePinWidget(UEdGraphPin* InPin)
 		{
 			return SNew(SGraphPinClass, InPin);
 		}
-		else if (InPin->PinType.PinCategory == K2Schema->PC_AssetClass)
+		else if (InPin->PinType.PinCategory == K2Schema->PC_SoftClass)
 		{
 			return SNew(SGraphPinClass, InPin);
 		}
diff --git a/Engine/Source/Editor/GraphEditor/Private/SGraphActionMenu.cpp b/Engine/Source/Editor/GraphEditor/Private/SGraphActionMenu.cpp
--- a/Engine/Source/Editor/GraphEditor/Private/SGraphActionMenu.cpp
+++ b/Engine/Source/Editor/GraphEditor/Private/SGraphActionMenu.cpp
@@ -219,11 +219,6 @@ public:
 
 	virtual void OnDragEnter( const FGeometry& MyGeometry, const FDragDropEvent& DragDropEvent ) override
 	{
-		if (IsReadOnly.Get())
-		{
-			return;
-		}
-
 		TSharedPtr<FGraphEditorDragDropAction> GraphDropOp = DragDropEvent.GetOperationAs<FGraphEditorDragDropAction>();
 		if (GraphDropOp.IsValid())
 		{
@@ -1061,7 +1056,7 @@ TSharedRef<ITableRow> SGraphActionMenu::MakeWidget( TSharedPtr<FGraphActionNode>
 			ReadOnlyArgument.IsReadOnly(this, &SGraphActionMenu::CanRenameNode, WeakItem);
 		}
 
-		TSharedRef<SGraphActionCategoryWidget> CategoryWidget = 
+		TSharedRef<SGraphActionCategoryWidget> CategoryWidget =
 			SNew(SGraphActionCategoryWidget, InItem)
 			.HighlightText(this, &SGraphActionMenu::GetFilterText)
 			.OnTextCommitted(this, &SGraphActionMenu::OnNameTextCommitted, TWeakPtr< FGraphActionNode >(InItem))
diff --git a/Engine/Source/Editor/GraphEditor/Private/SGraphNodeComment.cpp b/Engine/Source/Editor/GraphEditor/Private/SGraphNodeComment.cpp
--- a/Engine/Source/Editor/GraphEditor/Private/SGraphNodeComment.cpp
+++ b/Engine/Source/Editor/GraphEditor/Private/SGraphNodeComment.cpp
@@ -33,7 +33,7 @@ namespace SCommentNodeDefs
 }
 
 
-void SGraphNodeComment::Construct(const FArguments& InArgs, UEdGraphNode* InNode)
+void SGraphNodeComment::Construct(const FArguments& InArgs, UEdGraphNode_Comment* InNode)
 {
 	this->GraphNode = InNode;
 	this->bIsSelected = false;
@@ -73,6 +73,18 @@ void SGraphNodeComment::Tick( const FGeometry& AllottedGeometry, const double In
 	{
 		CachedWidth = CurrentWidth;
 	}
+
+	UEdGraphNode_Comment* CommentNode = CastChecked<UEdGraphNode_Comment>(GraphNode);
+	if (bCachedBubbleVisibility != CommentNode->bCommentBubbleVisible_InDetailsPanel)
+	{
+		CommentBubble->UpdateBubble();
+		bCachedBubbleVisibility = CommentNode->bCommentBubbleVisible_InDetailsPanel;
+	}
+
+	if (CachedFontSize != CommentNode->FontSize)
+	{
+		UpdateGraphNode();
+	}
 }
 
 FReply SGraphNodeComment::OnDrop( const FGeometry& MyGeometry, const FDragDropEvent& DragDropEvent )
@@ -104,20 +116,22 @@ void SGraphNodeComment::UpdateGraphNode()
 	// Avoid standard box model too
 	RightNodeBox.Reset();
 	LeftNodeBox.Reset();
-	
+
+	// Remember if we should be showing the bubble
+	UEdGraphNode_Comment* CommentNode = CastChecked<UEdGraphNode_Comment>(GraphNode);
+	bCachedBubbleVisibility = CommentNode->bCommentBubbleVisible_InDetailsPanel;
+
 	// Setup a tag for this node
 	FString TagName;
-	if (GraphNode != nullptr)
+
+	// We want the name of the blueprint as our name - we can find the node from the GUID
+	UObject* Package = GraphNode->GetOutermost();
+	UObject* LastOuter = GraphNode->GetOuter();
+	while (LastOuter->GetOuter() != Package)
 	{
-		// We want the name of the blueprint as our name - we can find the node from the GUID
-		UObject* Package = GraphNode->GetOutermost();
-		UObject* LastOuter = GraphNode->GetOuter();
-		while (LastOuter->GetOuter() != Package)
-		{
-			LastOuter = LastOuter->GetOuter();
-		}
-		TagName = FString::Printf(TEXT("GraphNode,%s,%s"), *LastOuter->GetFullName(), *GraphNode->NodeGuid.ToString());
+		LastOuter = LastOuter->GetOuter();
 	}
+	TagName = FString::Printf(TEXT("GraphNode,%s,%s"), *LastOuter->GetFullName(), *GraphNode->NodeGuid.ToString());
 
 	SetupErrorReporting();
 
@@ -125,6 +139,11 @@ void SGraphNodeComment::UpdateGraphNode()
 	FGraphNodeMetaData TagMeta(TEXT("Graphnode"));
 	PopulateMetaTag(&TagMeta);
 
+	CommentStyle = FEditorStyle::Get().GetWidgetStyle<FInlineEditableTextBlockStyle>("Graph.CommentBlock.TitleInlineEditableText");
+	CommentStyle.EditableTextBoxStyle.Font.Size = CommentNode->FontSize;
+	CommentStyle.TextStyle.Font.Size = CommentNode->FontSize;
+	CachedFontSize = CommentNode->FontSize;
+
 	bool bIsSet = GraphNode->IsA(UEdGraphNode_Comment::StaticClass());
 	this->ContentScale.Bind( this, &SGraphNode::GetContentScale );
 	this->GetOrAddSlot( ENodeZone::Center )
@@ -153,7 +172,7 @@ void SGraphNodeComment::UpdateGraphNode()
 					.VAlign(VAlign_Center)
 					[
 						SAssignNew(InlineEditableText, SInlineEditableTextBlock)
-						.Style( FEditorStyle::Get(), "Graph.CommentBlock.TitleInlineEditableText" )
+						.Style( &CommentStyle )
 						.Text( this, &SGraphNodeComment::GetEditableNodeTitleAsText )
 						.OnVerifyTextChanged(this, &SGraphNodeComment::OnVerifyNameTextChanged)
 						.OnTextCommitted(this, &SGraphNodeComment::OnNameTextCommited)
@@ -181,10 +200,9 @@ void SGraphNodeComment::UpdateGraphNode()
 				]
 			]			
 		];
-	// Create comment bubble
-	TSharedPtr<SCommentBubble> CommentBubble;
 
-	SAssignNew(CommentBubble, SCommentBubble)
+	// Create comment bubble
+	CommentBubble = SNew(SCommentBubble)
 	.GraphNode(GraphNode)
 	.Text(this, &SGraphNodeComment::GetNodeComment)
 	.OnTextCommitted(this, &SGraphNodeComment::OnNameTextCommited)
diff --git a/Engine/Source/Editor/GraphEditor/Private/SGraphPanel.cpp b/Engine/Source/Editor/GraphEditor/Private/SGraphPanel.cpp
--- a/Engine/Source/Editor/GraphEditor/Private/SGraphPanel.cpp
+++ b/Engine/Source/Editor/GraphEditor/Private/SGraphPanel.cpp
@@ -28,10 +28,6 @@
 #include "KismetNodes/KismetNodeInfoContext.h"
 #include "GraphDiffControl.h"
 
-
-// Blueprint Profiler
-#include "Profiler/BlueprintProfilerSettings.h"
-
 DEFINE_LOG_CATEGORY_STATIC(LogGraphPanel, Log, All);
 
 //////////////////////////////////////////////////////////////////////////
@@ -115,9 +111,6 @@ int32 SGraphPanel::OnPaint( const FPaintArgs& Args, const FGeometry& AllottedGeo
 	// Determine some 'global' settings based on current LOD
 	const bool bDrawShadowsThisFrame = GetCurrentLOD() > EGraphRenderingLOD::LowestDetail;
 
-	// Enable the profiler heatmap displays.
-	const bool bDisplayProfilerHeatmap = GetDefault<UBlueprintProfilerSettings>()->GraphNodeHeatMapDisplayMode != EBlueprintProfilerHeatMapDisplayMode::None;
-
 	// Because we paint multiple children, we must track the maximum layer id that they produced in case one of our parents
 	// wants to an overlay for all of its contents.
 
@@ -190,21 +183,6 @@ int32 SGraphPanel::OnPaint( const FPaintArgs& Args, const FGeometry& AllottedGeo
 					}
 				}
 
-				// Draw the profiler heatmap if active
-				if (bDisplayProfilerHeatmap)
-				{
-					const FSlateBrush* ProfilerBrush = ChildNode->GetProfilerHeatmapBrush();
-					const FLinearColor ProfilerHeatIntensity = ChildNode->GetProfilerHeatmapIntensity();
-					FSlateDrawElement::MakeBox(
-						OutDrawElements,
-						ShadowLayerId,
-						CurWidget.Geometry.ToInflatedPaintGeometry(NodeShadowSize),
-						ProfilerBrush,
-						ESlateDrawEffect::None,
-						ProfilerHeatIntensity
-						);
-				}
-
 				// Draw the node's shadow.
 				if (bDrawShadowsThisFrame || bSelected)
 				{
diff --git a/Engine/Source/Editor/GraphEditor/Public/GraphEditorActions.h b/Engine/Source/Editor/GraphEditor/Public/GraphEditorActions.h
--- a/Engine/Source/Editor/GraphEditor/Public/GraphEditorActions.h
+++ b/Engine/Source/Editor/GraphEditor/Public/GraphEditorActions.h
@@ -89,11 +89,7 @@ public:
 	TSharedPtr< FUICommandInfo > FindReferences;
 	TSharedPtr< FUICommandInfo > FindAndReplaceReferences;
 
-	// Goto native code actions
-	TSharedPtr< FUICommandInfo > GotoNativeFunctionDefinition;
-	TSharedPtr< FUICommandInfo > GotoNativeVariableDefinition;
-
-	// Jumps to a call function node's graph
+	// Jumps to the definition of the selected node (or otherwise focuses something interesting about that node, e.g., the inner graph for a collapsed graph)
 	TSharedPtr< FUICommandInfo > GoToDefinition;
 
 	// Pin-specific actions
diff --git a/Engine/Source/Editor/GraphEditor/Public/GraphEditorDragDropAction.h b/Engine/Source/Editor/GraphEditor/Public/GraphEditorDragDropAction.h
--- a/Engine/Source/Editor/GraphEditor/Public/GraphEditorDragDropAction.h
+++ b/Engine/Source/Editor/GraphEditor/Public/GraphEditorDragDropAction.h
@@ -98,12 +98,14 @@ public:
 	static TSharedRef<FGraphSchemaActionDragDropAction> New(TSharedPtr<FEdGraphSchemaAction> InActionNode )
 	{
 		TSharedRef<FGraphSchemaActionDragDropAction> Operation = MakeShareable(new FGraphSchemaActionDragDropAction);
-		Operation->ActionNode = InActionNode;
+		Operation->SourceAction = InActionNode;
 		Operation->Construct();
 		return Operation;
 	}
 
 protected:
+	virtual void GetDefaultStatusSymbol(const FSlateBrush*& PrimaryBrushOut, FSlateColor& IconColorOut, FSlateBrush const*& SecondaryBrushOut, FSlateColor& SecondaryColorOut) const;
+
 	/** */
-	TSharedPtr<FEdGraphSchemaAction> ActionNode;
+	TSharedPtr<FEdGraphSchemaAction> SourceAction;
 };
diff --git a/Engine/Source/Editor/GraphEditor/Public/GraphEditorSettings.h b/Engine/Source/Editor/GraphEditor/Public/GraphEditorSettings.h
--- a/Engine/Source/Editor/GraphEditor/Public/GraphEditorSettings.h
+++ b/Engine/Source/Editor/GraphEditor/Public/GraphEditorSettings.h
@@ -158,11 +158,11 @@ public:
 
 	/** Asset pin type color */
 	UPROPERTY(EditAnywhere, config, Category = PinColors)
-	FLinearColor AssetPinTypeColor;
+	FLinearColor SoftObjectPinTypeColor;
 
 	/** Asset Class pin type color */
 	UPROPERTY(EditAnywhere, config, Category = PinColors)
-	FLinearColor AssetClassPinTypeColor;
+	FLinearColor SoftClassPinTypeColor;
 
 	/** Delegate pin type color */
 	UPROPERTY(EditAnywhere, config, Category=PinColors)
diff --git a/Engine/Source/Editor/GraphEditor/Public/KismetNodes/SGraphNodeK2Base.h b/Engine/Source/Editor/GraphEditor/Public/KismetNodes/SGraphNodeK2Base.h
--- a/Engine/Source/Editor/GraphEditor/Public/KismetNodes/SGraphNodeK2Base.h
+++ b/Engine/Source/Editor/GraphEditor/Public/KismetNodes/SGraphNodeK2Base.h
@@ -20,8 +20,6 @@ public:
 	virtual void GetOverlayBrushes(bool bSelected, const FVector2D WidgetSize, TArray<FOverlayBrushInfo>& Brushes) const override;
 	virtual void GetNodeInfoPopups(FNodeInfoContext* Context, TArray<FGraphInformationPopupInfo>& Popups) const override;
 	virtual const FSlateBrush* GetShadowBrush(bool bSelected) const override;
-	virtual FLinearColor GetProfilerHeatmapIntensity() const override;
-	virtual const FSlateBrush* GetProfilerHeatmapBrush() const override;
 	void PerformSecondPassLayout(const TMap< UObject*, TSharedRef<SNode> >& NodeToWidgetLookup) const override;
 
 protected :
diff --git a/Engine/Source/Editor/GraphEditor/Public/SGraphNodeComment.h b/Engine/Source/Editor/GraphEditor/Public/SGraphNodeComment.h
--- a/Engine/Source/Editor/GraphEditor/Public/SGraphNodeComment.h
+++ b/Engine/Source/Editor/GraphEditor/Public/SGraphNodeComment.h
@@ -10,6 +10,9 @@
 #include "SNodePanel.h"
 #include "SGraphNodeResizable.h"
 
+class UEdGraphNode_Comment;
+class SCommentBubble;
+
 class SGraphNodeComment : public SGraphNodeResizable
 {
 public:
@@ -42,7 +45,7 @@ public:
 	virtual FSlateColor GetCommentColor() const override { return GetCommentBodyColor(); }
 	//~ End SGraphNode Interface
 
-	void Construct( const FArguments& InArgs, UEdGraphNode* InNode );
+	void Construct( const FArguments& InArgs, UEdGraphNode_Comment* InNode );
 
 	/** return if the node can be selected, by pointing given location */
 	virtual bool CanBeSelected( const FVector2D& MousePositionInNode ) const override;
@@ -89,6 +92,11 @@ private:
 	float GetWrapAt() const;
 
 private:
+	/** The comment bubble widget (used when zoomed out) */
+	TSharedPtr<SCommentBubble> CommentBubble;
+
+	/** Was the bubble desired to be visible last frame? */
+	mutable bool bCachedBubbleVisibility;
 
 	/** The current selection state of the comment */
 	mutable bool bIsSelected;
@@ -101,4 +109,10 @@ private:
 
 	/** cached comment title */
 	int32 CachedWidth;
+
+	/** cached font size */
+	int32 CachedFontSize;
+
+	/** Local copy of the comment style */
+	FInlineEditableTextBlockStyle CommentStyle;
 };
diff --git a/Engine/Source/Editor/GraphEditor/Public/SGraphPinDataTableRowName.h b/Engine/Source/Editor/GraphEditor/Public/SGraphPinDataTableRowName.h
--- a/Engine/Source/Editor/GraphEditor/Public/SGraphPinDataTableRowName.h
+++ b/Engine/Source/Editor/GraphEditor/Public/SGraphPinDataTableRowName.h
@@ -27,5 +27,5 @@ protected:
 
 	void RefreshNameList();
 
-	TAssetPtr<class UDataTable> DataTable;
+	TSoftObjectPtr<class UDataTable> DataTable;
 };
diff --git a/Engine/Source/Editor/GraphEditor/Public/SNodePanel.h b/Engine/Source/Editor/GraphEditor/Public/SNodePanel.h
--- a/Engine/Source/Editor/GraphEditor/Public/SNodePanel.h
+++ b/Engine/Source/Editor/GraphEditor/Public/SNodePanel.h
@@ -507,18 +507,6 @@ public:
 		{
 		}
 
-		/** @return The profiler heat intensity */
-		virtual FLinearColor GetProfilerHeatmapIntensity() const
-		{
-			return FLinearColor(1.f, 1.f, 1.f, 0.f);
-		}
-
-		/** @return The brush to use for drawing the profiler heatmap with the heat intensity */
-		virtual const FSlateBrush* GetProfilerHeatmapBrush() const
-		{
-			return FEditorStyle::GetBrush(TEXT("BlueprintProfiler.RegularNode.HeatDisplay"));
-		}
-
 		/** Populate the widgets array with any overlay widgets to render */
 		virtual TArray<FOverlayWidgetInfo> GetOverlayWidgets(bool bSelected, const FVector2D& WidgetSize) const
 		{
diff --git a/Engine/Source/Editor/IntroTutorials/Private/EditorTutorialSettings.h b/Engine/Source/Editor/IntroTutorials/Private/EditorTutorialSettings.h
--- a/Engine/Source/Editor/IntroTutorials/Private/EditorTutorialSettings.h
+++ b/Engine/Source/Editor/IntroTutorials/Private/EditorTutorialSettings.h
@@ -5,7 +5,7 @@
 #include "CoreMinimal.h"
 #include "UObject/ObjectMacros.h"
 #include "UObject/Object.h"
-#include "Misc/StringClassReference.h"
+#include "UObject/SoftObjectPath.h"
 #include "EditorTutorial.h"
 #include "EditorTutorialSettings.generated.h"
 
@@ -25,11 +25,11 @@ struct FTutorialContext
 
 	/** The tutorial to use in this context to let the user know there is a tutorial available */
 	UPROPERTY(EditAnywhere, Category = "Tutorials", meta = (MetaClass = "EditorTutorial"))
-	FStringClassReference AttractTutorial;
+	FSoftClassPath AttractTutorial;
 
 	/** The tutorial to use in this context when the user chooses to launch */
 	UPROPERTY(EditAnywhere, Category = "Tutorials", meta = (MetaClass = "EditorTutorial"))
-	FStringClassReference LaunchTutorial;
+	FSoftClassPath LaunchTutorial;
 };
 
 /** Editor-wide tutorial settings */
@@ -48,7 +48,7 @@ class UEditorTutorialSettings : public UObject
 
 	/** Tutorial to start on Editor startup */
 	UPROPERTY(Config, EditAnywhere, Category="Tutorials", meta=(MetaClass="EditorTutorial"))
-	FStringClassReference StartupTutorial;
+	FSoftClassPath StartupTutorial;
 
 	/** Tutorials used in various contexts - e.g. the various asset editors */
 	UPROPERTY(Config, EditAnywhere, Category = "Tutorials")
diff --git a/Engine/Source/Editor/IntroTutorials/Private/TutorialSettings.h b/Engine/Source/Editor/IntroTutorials/Private/TutorialSettings.h
--- a/Engine/Source/Editor/IntroTutorials/Private/TutorialSettings.h
+++ b/Engine/Source/Editor/IntroTutorials/Private/TutorialSettings.h
@@ -5,7 +5,7 @@
 #include "CoreMinimal.h"
 #include "UObject/ObjectMacros.h"
 #include "UObject/Object.h"
-#include "Misc/StringClassReference.h"
+#include "UObject/SoftObjectPath.h"
 #include "EditorTutorial.h"
 #include "TutorialSettings.generated.h"
 
@@ -21,5 +21,5 @@ class UTutorialSettings : public UObject
 
 	/** Tutorial to start on project startup */
 	UPROPERTY(Config, EditAnywhere, Category="Tutorials", meta=(MetaClass="EditorTutorial"))
-	FStringClassReference StartupTutorial;
+	FSoftClassPath StartupTutorial;
 };
diff --git a/Engine/Source/Editor/IntroTutorials/Private/TutorialStateSettings.cpp b/Engine/Source/Editor/IntroTutorials/Private/TutorialStateSettings.cpp
--- a/Engine/Source/Editor/IntroTutorials/Private/TutorialStateSettings.cpp
+++ b/Engine/Source/Editor/IntroTutorials/Private/TutorialStateSettings.cpp
@@ -65,7 +65,7 @@ void UTutorialStateSettings::RecordProgress(UEditorTutorial* InTutorial, int32 C
 		else
 		{
 			FTutorialProgress Progress;
-			Progress.Tutorial = FStringClassReference(InTutorial->GetClass());
+			Progress.Tutorial = FSoftClassPath(InTutorial->GetClass());
 			Progress.CurrentStage = CurrentStage;
 			Progress.bUserDismissed = false;
 			Progress.bUserDismissedThisSession = false;
@@ -88,7 +88,7 @@ void UTutorialStateSettings::DismissTutorial(UEditorTutorial* InTutorial, bool b
 		else
 		{
 			FTutorialProgress Progress;
-			Progress.Tutorial = FStringClassReference(InTutorial->GetClass());
+			Progress.Tutorial = FSoftClassPath(InTutorial->GetClass());
 			Progress.CurrentStage = 0;
 			Progress.bUserDismissed = bDismissAcrossSessions;
 			Progress.bUserDismissedThisSession = true;
diff --git a/Engine/Source/Editor/IntroTutorials/Private/TutorialStateSettings.h b/Engine/Source/Editor/IntroTutorials/Private/TutorialStateSettings.h
--- a/Engine/Source/Editor/IntroTutorials/Private/TutorialStateSettings.h
+++ b/Engine/Source/Editor/IntroTutorials/Private/TutorialStateSettings.h
@@ -5,7 +5,7 @@
 #include "CoreMinimal.h"
 #include "UObject/ObjectMacros.h"
 #include "UObject/Object.h"
-#include "Misc/StringClassReference.h"
+#include "UObject/SoftObjectPath.h"
 #include "TutorialStateSettings.generated.h"
 
 class UEditorTutorial;
@@ -22,7 +22,7 @@ struct FTutorialProgress
 	}
 
 	UPROPERTY()
-	FStringClassReference Tutorial;
+	FSoftClassPath Tutorial;
 
 	UPROPERTY()
 	int32 CurrentStage;
diff --git a/Engine/Source/Editor/IntroTutorials/Public/EditorTutorial.h b/Engine/Source/Editor/IntroTutorials/Public/EditorTutorial.h
--- a/Engine/Source/Editor/IntroTutorials/Public/EditorTutorial.h
+++ b/Engine/Source/Editor/IntroTutorials/Public/EditorTutorial.h
@@ -5,10 +5,9 @@
 #include "CoreMinimal.h"
 #include "UObject/ObjectMacros.h"
 #include "UObject/Object.h"
-#include "Misc/StringAssetReference.h"
 #include "UObject/ScriptMacros.h"
 #include "Types/SlateEnums.h"
-#include "Misc/StringClassReference.h"
+#include "UObject/SoftObjectPath.h"
 #include "EditorTutorial.generated.h"
 
 class AActor;
@@ -79,7 +78,7 @@ struct INTROTUTORIALS_API FTutorialCategory
 
 	/** Texture for this tutorial, used when presented to the user in the tutorial browser. */
 	UPROPERTY(EditAnywhere, Category="Tutorial", meta=(AllowedClasses="Texture2D"))
-	FStringAssetReference Texture;
+	FSoftObjectPath Texture;
 };
 
 /** Content wrapper */
@@ -131,7 +130,7 @@ struct INTROTUTORIALS_API FTutorialContentAnchor
 
 	/** If reference is an asset, we use this to resolve it */
 	UPROPERTY(EditAnywhere, Category="Anchor")
-	FStringAssetReference Asset;
+	FSoftObjectPath Asset;
 
 	/** Whether to draw an animated highlight around the widget */
 	UPROPERTY(EditAnywhere, Category="Anchor")
@@ -267,11 +266,11 @@ class INTROTUTORIALS_API UEditorTutorial : public UObject
 
 	/** Tutorial to optionally chain back to if the "back" button is clicked on the first stage */
 	UPROPERTY(EditAnywhere, Category = "Tutorial", meta = (MetaClass = "EditorTutorial"))
-	FStringClassReference PreviousTutorial;
+	FSoftClassPath PreviousTutorial;
 
 	/** Tutorial to optionally chain onto after this tutorial completes */
 	UPROPERTY(EditAnywhere, Category="Tutorial", meta=(MetaClass="EditorTutorial"))
-	FStringClassReference NextTutorial;
+	FSoftClassPath NextTutorial;
 
 	/**
 	* Attempts to find the actor specified by PathToActor in the current editor world
@@ -287,7 +286,7 @@ class INTROTUTORIALS_API UEditorTutorial : public UObject
 	
 	/** Asset to open & attach the tutorial to. Non-widget-bound content will appear in the asset's window */
 	UPROPERTY(EditAnywhere, Category="Tutorial")
-	FStringAssetReference AssetToUse;
+	FSoftObjectPath AssetToUse;
 
 	/** The path this tutorial was imported from, if any. */
 	UPROPERTY()
diff --git a/Engine/Source/Editor/Kismet/Kismet.Build.cs b/Engine/Source/Editor/Kismet/Kismet.Build.cs
--- a/Engine/Source/Editor/Kismet/Kismet.Build.cs
+++ b/Engine/Source/Editor/Kismet/Kismet.Build.cs
@@ -52,7 +52,6 @@ public class Kismet : ModuleRules
 				"DesktopPlatform",
 				"HotReload",
                 "BlueprintNativeCodeGen",
-                "BlueprintProfiler",
                 "UMGEditor",
                 "UMG" // for SBlueprintDiff
 			}
diff --git a/Engine/Source/Editor/Kismet/Private/BPDelegateDragDropAction.h b/Engine/Source/Editor/Kismet/Private/BPDelegateDragDropAction.h
--- a/Engine/Source/Editor/Kismet/Private/BPDelegateDragDropAction.h
+++ b/Engine/Source/Editor/Kismet/Private/BPDelegateDragDropAction.h
@@ -19,7 +19,7 @@ public:
 	DRAG_DROP_OPERATOR_TYPE(FKismetDelegateDragDropAction, FKismetVariableDragDropAction)
 
 	// FGraphEditorDragDropAction interface
-	virtual void HoverTargetChanged() override { FGraphEditorDragDropAction::HoverTargetChanged(); }
+	virtual void HoverTargetChanged() override { FMyBlueprintItemDragDropAction::HoverTargetChanged(); }
 	virtual FReply DroppedOnPin(FVector2D ScreenPosition, FVector2D GraphPosition) override { return FGraphEditorDragDropAction::DroppedOnPin(ScreenPosition, GraphPosition); }
 	virtual FReply DroppedOnNode(FVector2D ScreenPosition, FVector2D GraphPosition) override { return FGraphEditorDragDropAction::DroppedOnNode(ScreenPosition, GraphPosition); }
 
@@ -30,9 +30,10 @@ public:
 
 	bool IsValid() const;
 	
-	static TSharedRef<FKismetDelegateDragDropAction> New( FName InVariableName, UStruct* InSource, FNodeCreationAnalytic AnalyticCallback)
+	static TSharedRef<FKismetDelegateDragDropAction> New(TSharedPtr<FEdGraphSchemaAction> InSourceAction, FName InVariableName, UStruct* InSource, FNodeCreationAnalytic AnalyticCallback)
 	{
 		TSharedRef<FKismetDelegateDragDropAction> Operation = MakeShareable(new FKismetDelegateDragDropAction);
+		Operation->SourceAction = InSourceAction;
 		Operation->VariableName = InVariableName;
 		Operation->VariableSource = InSource;
 		Operation->AnalyticCallback = AnalyticCallback;
diff --git a/Engine/Source/Editor/Kismet/Private/BPFunctionDragDropAction.cpp b/Engine/Source/Editor/Kismet/Private/BPFunctionDragDropAction.cpp
--- a/Engine/Source/Editor/Kismet/Private/BPFunctionDragDropAction.cpp
+++ b/Engine/Source/Editor/Kismet/Private/BPFunctionDragDropAction.cpp
@@ -108,18 +108,17 @@ void FKismetDragDropAction::HoverTargetChanged()
 	if (ActionWillShowExistingNode())
 	{
 		FSlateBrush const* ShowsExistingIcon = FEditorStyle::GetBrush(TEXT("Graph.ConnectorFeedback.ShowNode"));
-		FText DragingText = FText::Format(LOCTEXT("ShowExistingNode", "Show '{0}'"), ActionNode->GetMenuDescription());
+		FText DragingText = FText::Format(LOCTEXT("ShowExistingNode", "Show '{0}'"), SourceAction->GetMenuDescription());
 		SetSimpleFeedbackMessage(ShowsExistingIcon, FLinearColor::White, DragingText);
 	}
 	// it should be obvious that we can't drop on anything but a graph, so no need to point that out
-	else if ((TheHoveredGraph == nullptr) || !CanBeDroppedDelegate.IsBound() || CanBeDroppedDelegate.Execute(ActionNode, TheHoveredGraph, CannotDropReason))
+	else if ((TheHoveredGraph == nullptr) || !CanBeDroppedDelegate.IsBound() || CanBeDroppedDelegate.Execute(SourceAction, TheHoveredGraph, CannotDropReason))
 	{
-		FGraphSchemaActionDragDropAction::HoverTargetChanged();
+		FMyBlueprintItemDragDropAction::HoverTargetChanged();
 	}
 	else 
 	{
-		FSlateBrush const* DropPreventedIcon = FEditorStyle::GetBrush(TEXT("Graph.ConnectorFeedback.Error"));
-		SetSimpleFeedbackMessage(DropPreventedIcon, FLinearColor::White, CannotDropReason);
+		SetFeedbackMessageError(CannotDropReason);
 	}
 }
 
@@ -129,9 +128,9 @@ FReply FKismetDragDropAction::DroppedOnPanel( const TSharedRef< SWidget >& Panel
 	FReply Reply = FReply::Unhandled();
 
 	FText CannotDropReason = FText::GetEmpty();
-	if (!CanBeDroppedDelegate.IsBound() || CanBeDroppedDelegate.Execute(ActionNode, GetHoveredGraph(), CannotDropReason))
+	if (!CanBeDroppedDelegate.IsBound() || CanBeDroppedDelegate.Execute(SourceAction, GetHoveredGraph(), CannotDropReason))
 	{
-		Reply = FGraphSchemaActionDragDropAction::DroppedOnPanel(Panel, ScreenPosition, GraphPosition, Graph);
+		Reply = FMyBlueprintItemDragDropAction::DroppedOnPanel(Panel, ScreenPosition, GraphPosition, Graph);
 	}
 
 	if (Reply.IsEventHandled())
@@ -148,21 +147,21 @@ bool FKismetDragDropAction::ActionWillShowExistingNode() const
 	bool bWillFocusOnExistingNode = false;
 
 	UEdGraph* TheHoveredGraph = GetHoveredGraph();
-	if (ActionNode.IsValid() && (TheHoveredGraph != nullptr))
+	if (SourceAction.IsValid() && (TheHoveredGraph != nullptr))
 	{
-		bWillFocusOnExistingNode = (ActionNode->GetTypeId() == FEdGraphSchemaAction_K2TargetNode::StaticGetTypeId()) ||
-			(ActionNode->GetTypeId() == FEdGraphSchemaAction_K2InputAction::StaticGetTypeId());
+		bWillFocusOnExistingNode = (SourceAction->GetTypeId() == FEdGraphSchemaAction_K2TargetNode::StaticGetTypeId()) ||
+			(SourceAction->GetTypeId() == FEdGraphSchemaAction_K2InputAction::StaticGetTypeId());
 
 		if (!bWillFocusOnExistingNode)
 		{
-			if (ActionNode->GetTypeId() == FEdGraphSchemaAction_K2AddEvent::StaticGetTypeId())
+			if (SourceAction->GetTypeId() == FEdGraphSchemaAction_K2AddEvent::StaticGetTypeId())
 			{
-				FEdGraphSchemaAction_K2AddEvent* AddEventAction = (FEdGraphSchemaAction_K2AddEvent*)ActionNode.Get();
+				FEdGraphSchemaAction_K2AddEvent* AddEventAction = (FEdGraphSchemaAction_K2AddEvent*)SourceAction.Get();
 				bWillFocusOnExistingNode = AddEventAction->EventHasAlreadyBeenPlaced(FBlueprintEditorUtils::FindBlueprintForGraph(TheHoveredGraph));
 			}
-			else if (ActionNode->GetTypeId() == FEdGraphSchemaAction_K2Event::StaticGetTypeId())
+			else if (SourceAction->GetTypeId() == FEdGraphSchemaAction_K2Event::StaticGetTypeId())
 			{
-				FEdGraphSchemaAction_K2Event* FuncAction = (FEdGraphSchemaAction_K2Event*)ActionNode.Get();
+				FEdGraphSchemaAction_K2Event* FuncAction = (FEdGraphSchemaAction_K2Event*)SourceAction.Get();
 				UK2Node_CustomEvent* CustomEvent = Cast<UK2Node_CustomEvent>(FuncAction->NodeTemplate);
 				// Drag and dropping custom event's will place a Call Function and will not focus the existing event
 				if (CustomEvent == nullptr)
@@ -195,7 +194,7 @@ TSharedRef<FKismetFunctionDragDropAction> FKismetFunctionDragDropAction::New(
 	Operation->CallOnMember     = InCallOnMember;
 	Operation->AnalyticCallback = AnalyticCallback;
 	Operation->CanBeDroppedDelegate = CanBeDroppedDelegate;
-	Operation->ActionNode = InActionNode;
+	Operation->SourceAction = InActionNode;
 
 	if (!CanBeDroppedDelegate.IsBound())
 	{
@@ -213,12 +212,6 @@ FKismetFunctionDragDropAction::FKismetFunctionDragDropAction()
 
 }
 
-//------------------------------------------------------------------------------
-void FKismetFunctionDragDropAction::HoverTargetChanged()
-{
-	FKismetDragDropAction::HoverTargetChanged();
-}
-
 //------------------------------------------------------------------------------
 FReply FKismetFunctionDragDropAction::DroppedOnPanel(TSharedRef<SWidget> const& Panel, FVector2D ScreenPosition, FVector2D GraphPosition, UEdGraph& Graph)
 {
@@ -231,7 +224,7 @@ FReply FKismetFunctionDragDropAction::DroppedOnPin(FVector2D ScreenPosition, FVe
 	FReply Reply = FReply::Unhandled();
 
 	UEdGraph* Graph = GetHoveredGraph();
-	check(Graph);
+	check(Graph); 
 
 	// The ActionNode set during construction points to the Graph, this is suitable for displaying the mouse decorator but needs to be more complete based on the current graph
 	UBlueprintFunctionNodeSpawner* FunctionNodeSpawner = GetDropAction(*Graph);
@@ -308,12 +301,14 @@ UBlueprintFunctionNodeSpawner* FKismetFunctionDragDropAction::GetDropAction(UEdG
 
 //------------------------------------------------------------------------------
 TSharedRef<FKismetMacroDragDropAction> FKismetMacroDragDropAction::New(
+	TSharedPtr<FEdGraphSchemaAction> InActionNode,
 	FName                 InMacroName, 
 	UBlueprint*           InBlueprint, 
 	UEdGraph*             InMacro, 
 	FNodeCreationAnalytic AnalyticCallback)
 {
 	TSharedRef<FKismetMacroDragDropAction> Operation = MakeShareable(new FKismetMacroDragDropAction);
+	Operation->SourceAction = InActionNode;
 	Operation->MacroName = InMacroName;
 	Operation->Macro = InMacro;
 	Operation->Blueprint = InBlueprint;
@@ -330,30 +325,12 @@ TSharedRef<FKismetMacroDragDropAction> FKismetMacroDragDropAction::New(
 
 //------------------------------------------------------------------------------
 FKismetMacroDragDropAction::FKismetMacroDragDropAction()
-	: Macro(NULL)
+	: Macro(nullptr)
+	, Blueprint(nullptr)
 {
 
 }
 
-//------------------------------------------------------------------------------
-void FKismetMacroDragDropAction::HoverTargetChanged() 
-{
-	UEdGraph* TheHoveredGraph = GetHoveredGraph();
-
-	FText CannotDropReason = FText::GetEmpty();
-	// it should be obvious that we can't drop on anything but a graph, so no need to point that out
-	if ((TheHoveredGraph == NULL) || !CanBeDroppedDelegate.IsBound() || CanBeDroppedDelegate.Execute(ActionNode, TheHoveredGraph, CannotDropReason))
-	{
-		FSlateBrush const* DropPreventedIcon = FEditorStyle::GetBrush(TEXT("Graph.ConnectorFeedback.NewNode"));
-		SetSimpleFeedbackMessage(DropPreventedIcon, FLinearColor::White, FText::FromName(MacroName));
-	}
-	else 
-	{
-		FSlateBrush const* DropPreventedIcon = FEditorStyle::GetBrush(TEXT("Graph.ConnectorFeedback.Error"));
-		SetSimpleFeedbackMessage(DropPreventedIcon, FLinearColor::White, CannotDropReason);
-	}
-}
-
 //------------------------------------------------------------------------------
 FReply FKismetMacroDragDropAction::DroppedOnPanel(TSharedRef<SWidget> const& Panel, FVector2D ScreenPosition, FVector2D GraphPosition, UEdGraph& Graph)
 {	
diff --git a/Engine/Source/Editor/Kismet/Private/BPFunctionDragDropAction.h b/Engine/Source/Editor/Kismet/Private/BPFunctionDragDropAction.h
--- a/Engine/Source/Editor/Kismet/Private/BPFunctionDragDropAction.h
+++ b/Engine/Source/Editor/Kismet/Private/BPFunctionDragDropAction.h
@@ -9,6 +9,7 @@
 #include "Engine/MemberReference.h"
 #include "BlueprintEditor.h"
 #include "Editor/GraphEditor/Public/GraphEditorDragDropAction.h"
+#include "MyBlueprintItemDragDropAction.h"
 
 class UEdGraph;
 
@@ -16,10 +17,10 @@ class UEdGraph;
 * FKismetDragDropAction
 *******************************************************************************/
 
-class KISMET_API FKismetDragDropAction : public FGraphSchemaActionDragDropAction
+class KISMET_API FKismetDragDropAction : public FMyBlueprintItemDragDropAction
 {
 public:
-	DRAG_DROP_OPERATOR_TYPE(FKismetDragDropAction, FGraphSchemaActionDragDropAction)
+	DRAG_DROP_OPERATOR_TYPE(FKismetDragDropAction, FMyBlueprintItemDragDropAction)
 		
 	// FGraphEditorDragDropAction interface
 	virtual void HoverTargetChanged() override;
@@ -31,7 +32,7 @@ public:
 	static TSharedRef<FKismetDragDropAction> New(TSharedPtr<FEdGraphSchemaAction> InActionNode, FNodeCreationAnalytic AnalyticCallback, FCanBeDroppedDelegate CanBeDroppedDelegate)
 	{
 		TSharedRef<FKismetDragDropAction> Operation = MakeShareable(new FKismetDragDropAction);
-		Operation->ActionNode = InActionNode;
+		Operation->SourceAction = InActionNode;
 		Operation->AnalyticCallback = AnalyticCallback;
 		Operation->CanBeDroppedDelegate = CanBeDroppedDelegate;
 		Operation->Construct();
@@ -41,9 +42,6 @@ public:
 protected:
 	bool ActionWillShowExistingNode() const;
 
-	/** Analytic delegate to track node creation */
-	FNodeCreationAnalytic AnalyticCallback;
-
 	/** */
 	FCanBeDroppedDelegate CanBeDroppedDelegate;
 };
@@ -60,7 +58,6 @@ public:
 	FKismetFunctionDragDropAction();
 
 	// FGraphEditorDragDropAction interface
-	virtual void HoverTargetChanged() override;
 	virtual FReply DroppedOnPanel( const TSharedRef< class SWidget >& Panel, FVector2D ScreenPosition, FVector2D GraphPosition, UEdGraph& Graph) override;
 	virtual FReply DroppedOnPin(FVector2D ScreenPosition, FVector2D GraphPosition) override;
 	// End of FGraphEditorDragDropAction
@@ -94,11 +91,18 @@ public:
 	FKismetMacroDragDropAction();
 
 	// FGraphEditorDragDropAction interface
-	virtual void HoverTargetChanged() override;
 	virtual FReply DroppedOnPanel( const TSharedRef< class SWidget >& Panel, FVector2D ScreenPosition, FVector2D GraphPosition, UEdGraph& Graph) override;
 	// End of FGraphEditorDragDropAction
 
-	static TSharedRef<FKismetMacroDragDropAction> New(FName InMacroName, UBlueprint* InBlueprint, UEdGraph* InMacro, FNodeCreationAnalytic AnalyticCallback);
+	static TSharedRef<FKismetMacroDragDropAction> New(TSharedPtr<FEdGraphSchemaAction> InActionNode, FName InMacroName, UBlueprint* InBlueprint, UEdGraph* InMacro, FNodeCreationAnalytic AnalyticCallback);
+
+protected:
+	// FMyBlueprintItemDragDropAction interface
+	virtual UBlueprint* GetSourceBlueprint() const override
+	{
+		return Blueprint;
+	}
+	// End of FMyBlueprintItemDragDropAction interface
 
 protected:
 	/** Name of macro being dragged */
