Copying //UE4/Dev-Framework to //UE4/Dev-Main (Source: //UE4/Dev-Framework @ 3716594)

diff --git a/Engine/Source/Editor/AnimGraph/Private/AnimGraphNode_BlendListBase.cpp b/Engine/Source/Editor/AnimGraph/Private/AnimGraphNode_BlendListBase.cpp
--- a/Engine/Source/Editor/AnimGraph/Private/AnimGraphNode_BlendListBase.cpp
+++ b/Engine/Source/Editor/AnimGraph/Private/AnimGraphNode_BlendListBase.cpp
@@ -18,7 +18,7 @@ FLinearColor UAnimGraphNode_BlendListBase::GetNodeTitleColor() const
 
 void UAnimGraphNode_BlendListBase::PostEditChangeProperty(struct FPropertyChangedEvent& PropertyChangedEvent)
 {
-	FName PropertyName = (PropertyChangedEvent.Property != NULL) ? PropertyChangedEvent.Property->GetFName() : NAME_None;
+	const FName PropertyName = (PropertyChangedEvent.Property ? PropertyChangedEvent.Property->GetFName() : NAME_None);
 
 	if ((PropertyName == TEXT("Node")))
 	{
@@ -36,7 +36,7 @@ FString UAnimGraphNode_BlendListBase::GetNodeCategory() const
 void UAnimGraphNode_BlendListBase::RemovePinsFromOldPins(TArray<UEdGraphPin*>& OldPins, int32 RemovedArrayIndex)
 {
 	TArray<FString> RemovedPropertyNames;
-	TArray<FString> NewPinNames;
+	TArray<FName> NewPinNames;
 
 	// Store new pin names to compare with old pin names
 	for (int32 NewPinIndx = 0; NewPinIndx < Pins.Num(); NewPinIndx++)
@@ -47,14 +47,15 @@ void UAnimGraphNode_BlendListBase::RemovePinsFromOldPins(TArray<UEdGraphPin*>& O
 	// don't know which pins are removed yet so find removed pins comparing NewPins and OldPins
 	for (int32 OldPinIdx = 0; OldPinIdx < OldPins.Num(); OldPinIdx++)
 	{
-		FString& OldPinName = OldPins[OldPinIdx]->PinName;
+		const FName OldPinName = OldPins[OldPinIdx]->PinName;
 		if (!NewPinNames.Contains(OldPinName))
 		{
-			int32 UnderscoreIndex = OldPinName.Find(TEXT("_"));
+			const FString OldPinNameStr = OldPinName.ToString();
+			const int32 UnderscoreIndex = OldPinNameStr.Find(TEXT("_"), ESearchCase::CaseSensitive);
 			if (UnderscoreIndex != INDEX_NONE)
 			{
-				FString PropertyName = OldPinName.Left(UnderscoreIndex);
-				RemovedPropertyNames.Add(PropertyName);
+				FString PropertyName = OldPinNameStr.Left(UnderscoreIndex);
+				RemovedPropertyNames.Add(MoveTemp(PropertyName));
 			}
 		}
 	}
@@ -62,18 +63,17 @@ void UAnimGraphNode_BlendListBase::RemovePinsFromOldPins(TArray<UEdGraphPin*>& O
 	for (int32 PinIdx = 0; PinIdx < OldPins.Num(); PinIdx++)
 	{
 		// Separate the pin name into property name and index
-		FString PropertyName;
-		int32 ArrayIndex = -1;
-		FString& OldPinName = OldPins[PinIdx]->PinName;
+		const FString OldPinNameStr = OldPins[PinIdx]->PinName.ToString();
+		const int32 UnderscoreIndex = OldPinNameStr.Find(TEXT("_"), ESearchCase::CaseSensitive);
 
-		int32 UnderscoreIndex = OldPinName.Find(TEXT("_"));
 		if (UnderscoreIndex != INDEX_NONE)
 		{
-			PropertyName = OldPinName.Left(UnderscoreIndex);
-			ArrayIndex = FCString::Atoi(*(OldPinName.Mid(UnderscoreIndex + 1)));
+			const FString PropertyName = OldPinNameStr.Left(UnderscoreIndex);
 
 			if (RemovedPropertyNames.Contains(PropertyName))
 			{
+				const int32 ArrayIndex = FCString::Atoi(*(OldPinNameStr.Mid(UnderscoreIndex + 1)));
+
 				// if array index is matched, removes pins 
 				// and if array index is greater than removed index, decrease index
 				if (ArrayIndex == RemovedArrayIndex)
@@ -82,11 +82,10 @@ void UAnimGraphNode_BlendListBase::RemovePinsFromOldPins(TArray<UEdGraphPin*>& O
 					OldPins.RemoveAt(PinIdx);
 					--PinIdx;
 				}
-				else
-					if (ArrayIndex > RemovedArrayIndex)
-					{
-						OldPinName = FString::Printf(TEXT("%s_%d"), *PropertyName, ArrayIndex - 1);
-					}
+				else if (ArrayIndex > RemovedArrayIndex)
+				{
+					OldPins[PinIdx]->PinName = *FString::Printf(TEXT("%s_%d"), *PropertyName, ArrayIndex - 1);
+				}
 			}
 		}
 	}
diff --git a/Engine/Source/Editor/AnimGraph/Private/AnimGraphNode_BlendListByEnum.cpp b/Engine/Source/Editor/AnimGraph/Private/AnimGraphNode_BlendListByEnum.cpp
--- a/Engine/Source/Editor/AnimGraph/Private/AnimGraphNode_BlendListByEnum.cpp
+++ b/Engine/Source/Editor/AnimGraph/Private/AnimGraphNode_BlendListByEnum.cpp
@@ -79,14 +79,14 @@ void UAnimGraphNode_BlendListByEnum::GetMenuActions(FBlueprintActionDatabaseRegi
 
 void UAnimGraphNode_BlendListByEnum::GetContextMenuActions(const FGraphNodeContextMenuBuilder& Context) const
 {
-	if (!Context.bIsDebugging && (BoundEnum != NULL))
+	if (!Context.bIsDebugging && BoundEnum)
 	{
-		if ((Context.Pin != NULL) && (Context.Pin->Direction == EGPD_Input))
+		if (Context.Pin && (Context.Pin->Direction == EGPD_Input))
 		{
 			int32 RawArrayIndex = 0;
 			bool bIsPosePin = false;
 			bool bIsTimePin = false;
-			GetPinInformation(Context.Pin->PinName, /*out*/ RawArrayIndex, /*out*/ bIsPosePin, /*out*/ bIsTimePin);
+			GetPinInformation(Context.Pin->PinName.ToString(), /*out*/ RawArrayIndex, /*out*/ bIsPosePin, /*out*/ bIsTimePin);
 
 			if (bIsPosePin || bIsTimePin)
 			{
@@ -154,7 +154,7 @@ void UAnimGraphNode_BlendListByEnum::RemovePinFromBlendList(UEdGraphPin* Pin)
 	int32 RawArrayIndex = 0;
 	bool bIsPosePin = false;
 	bool bIsTimePin = false;
-	GetPinInformation(Pin->PinName, /*out*/ RawArrayIndex, /*out*/ bIsPosePin, /*out*/ bIsTimePin);
+	GetPinInformation(Pin->PinName.ToString(), /*out*/ RawArrayIndex, /*out*/ bIsPosePin, /*out*/ bIsTimePin);
 
 	const int32 ExposedEnumIndex = (bIsPosePin || bIsTimePin) ? (RawArrayIndex - 1) : INDEX_NONE;
 
@@ -185,14 +185,14 @@ void UAnimGraphNode_BlendListByEnum::RemovePinFromBlendList(UEdGraphPin* Pin)
 
 void UAnimGraphNode_BlendListByEnum::GetPinInformation(const FString& InPinName, int32& Out_PinIndex, bool& Out_bIsPosePin, bool& Out_bIsTimePin)
 {
-	const int32 UnderscoreIndex = InPinName.Find(TEXT("_"));
+	const int32 UnderscoreIndex = InPinName.Find(TEXT("_"), ESearchCase::CaseSensitive);
 	if (UnderscoreIndex != INDEX_NONE)
 	{
 		const FString ArrayName = InPinName.Left(UnderscoreIndex);
 		Out_PinIndex = FCString::Atoi(*(InPinName.Mid(UnderscoreIndex + 1)));
 
-		Out_bIsPosePin = ArrayName == TEXT("BlendPose");
-		Out_bIsTimePin = ArrayName == TEXT("BlendTime");
+		Out_bIsPosePin = (ArrayName == TEXT("BlendPose"));
+		Out_bIsTimePin = (ArrayName == TEXT("BlendTime"));
 	}
 	else
 	{
@@ -208,7 +208,7 @@ void UAnimGraphNode_BlendListByEnum::CustomizePinData(UEdGraphPin* Pin, FName So
 	bool bIsPosePin;
 	bool bIsTimePin;
 	int32 RawArrayIndex;
-	GetPinInformation(Pin->PinName, /*out*/ RawArrayIndex, /*out*/ bIsPosePin, /*out*/ bIsTimePin);
+	GetPinInformation(Pin->PinName.ToString(), /*out*/ RawArrayIndex, /*out*/ bIsPosePin, /*out*/ bIsTimePin);
 	checkSlow(RawArrayIndex == ArrayIndex);
 
 	if (bIsPosePin || bIsTimePin)
diff --git a/Engine/Source/Editor/AnimGraph/Private/AnimGraphNode_Constraint.cpp b/Engine/Source/Editor/AnimGraph/Private/AnimGraphNode_Constraint.cpp
--- a/Engine/Source/Editor/AnimGraph/Private/AnimGraphNode_Constraint.cpp
+++ b/Engine/Source/Editor/AnimGraph/Private/AnimGraphNode_Constraint.cpp
@@ -142,12 +142,12 @@ void UAnimGraphNode_Constraint::PostProcessPinName(const UEdGraphPin* Pin, FStri
 	if (Pin->Direction == EGPD_Input)
 	{
 		const FString ConstraintWeightPrefix = TEXT("ConstraintWeights_");
-		FString PinName = Pin->PinName;
+		const FString PinName = Pin->PinName.ToString();
 		FString IndexString;
 		if (PinName.Split(ConstraintWeightPrefix, nullptr, &IndexString))
 		{
 			// convert index and display better name
-			int32 Index = FCString::Atoi(*IndexString);
+			const int32 Index = FCString::Atoi(*IndexString);
 			if (Node.ConstraintSetup.IsValidIndex(Index))
 			{
 				const FConstraint& Constraint = Node.ConstraintSetup[Index];
diff --git a/Engine/Source/Editor/AnimGraph/Private/AnimGraphNode_LocalToComponentSpace.cpp b/Engine/Source/Editor/AnimGraph/Private/AnimGraphNode_LocalToComponentSpace.cpp
--- a/Engine/Source/Editor/AnimGraph/Private/AnimGraphNode_LocalToComponentSpace.cpp
+++ b/Engine/Source/Editor/AnimGraph/Private/AnimGraphNode_LocalToComponentSpace.cpp
@@ -35,8 +35,7 @@ FString UAnimGraphNode_LocalToComponentSpace::GetNodeCategory() const
 
 void UAnimGraphNode_LocalToComponentSpace::CreateOutputPins()
 {
-	const UAnimationGraphSchema* Schema = GetDefault<UAnimationGraphSchema>();
-	CreatePin(EGPD_Output, Schema->PC_Struct, FString(), FComponentSpacePoseLink::StaticStruct(), TEXT("ComponentPose"));
+	CreatePin(EGPD_Output, UAnimationGraphSchema::PC_Struct, FComponentSpacePoseLink::StaticStruct(), TEXT("ComponentPose"));
 }
 
 void UAnimGraphNode_LocalToComponentSpace::PostProcessPinName(const UEdGraphPin* Pin, FString& DisplayName) const
diff --git a/Engine/Source/Editor/AnimGraph/Private/AnimGraphNode_MeshRefPose.cpp b/Engine/Source/Editor/AnimGraph/Private/AnimGraphNode_MeshRefPose.cpp
--- a/Engine/Source/Editor/AnimGraph/Private/AnimGraphNode_MeshRefPose.cpp
+++ b/Engine/Source/Editor/AnimGraph/Private/AnimGraphNode_MeshRefPose.cpp
@@ -35,8 +35,7 @@ FText UAnimGraphNode_MeshRefPose::GetNodeTitle(ENodeTitleType::Type TitleType) c
 
 void UAnimGraphNode_MeshRefPose::CreateOutputPins()
 {
-	const UAnimationGraphSchema* Schema = GetDefault<UAnimationGraphSchema>();
-	CreatePin(EGPD_Output, Schema->PC_Struct, FString(), FComponentSpacePoseLink::StaticStruct(), TEXT("ComponentPose"));
+	CreatePin(EGPD_Output, UAnimationGraphSchema::PC_Struct, FComponentSpacePoseLink::StaticStruct(), TEXT("ComponentPose"));
 }
 
 #undef LOCTEXT_NAMESPACE
diff --git a/Engine/Source/Editor/AnimGraph/Private/AnimGraphNode_MultiWayBlend.cpp b/Engine/Source/Editor/AnimGraph/Private/AnimGraphNode_MultiWayBlend.cpp
--- a/Engine/Source/Editor/AnimGraph/Private/AnimGraphNode_MultiWayBlend.cpp
+++ b/Engine/Source/Editor/AnimGraph/Private/AnimGraphNode_MultiWayBlend.cpp
@@ -118,7 +118,7 @@ void UAnimGraphNode_MultiWayBlend::ReallocatePinsDuringReconstruction(TArray<UEd
 void UAnimGraphNode_MultiWayBlend::RemovePinsFromOldPins(TArray<UEdGraphPin*>& OldPins, int32 RemovedArrayIndex)
 {
 	TArray<FString> RemovedPropertyNames;
-	TArray<FString> NewPinNames;
+	TArray<FName> NewPinNames;
 
 	// Store new pin names to compare with old pin names
 	for (int32 NewPinIndx = 0; NewPinIndx < Pins.Num(); NewPinIndx++)
@@ -129,14 +129,15 @@ void UAnimGraphNode_MultiWayBlend::RemovePinsFromOldPins(TArray<UEdGraphPin*>& O
 	// don't know which pins are removed yet so find removed pins comparing NewPins and OldPins
 	for (int32 OldPinIdx = 0; OldPinIdx < OldPins.Num(); OldPinIdx++)
 	{
-		FString& OldPinName = OldPins[OldPinIdx]->PinName;
+		const FName OldPinName = OldPins[OldPinIdx]->PinName;
 		if (!NewPinNames.Contains(OldPinName))
 		{
-			int32 UnderscoreIndex = OldPinName.Find(TEXT("_"));
+			const FString OldPinNameStr = OldPinName.ToString();
+			const int32 UnderscoreIndex = OldPinNameStr.Find(TEXT("_"), ESearchCase::CaseSensitive);
 			if (UnderscoreIndex != INDEX_NONE)
 			{
-				FString PropertyName = OldPinName.Left(UnderscoreIndex);
-				RemovedPropertyNames.Add(PropertyName);
+				FString PropertyName = OldPinNameStr.Left(UnderscoreIndex);
+				RemovedPropertyNames.Add(MoveTemp(PropertyName));
 			}
 		}
 	}
@@ -144,18 +145,16 @@ void UAnimGraphNode_MultiWayBlend::RemovePinsFromOldPins(TArray<UEdGraphPin*>& O
 	for (int32 PinIdx = 0; PinIdx < OldPins.Num(); PinIdx++)
 	{
 		// Separate the pin name into property name and index
-		FString PropertyName;
-		int32 ArrayIndex = -1;
-		FString& OldPinName = OldPins[PinIdx]->PinName;
+		const FString OldPinNameStr = OldPins[PinIdx]->PinName.ToString();
 
-		int32 UnderscoreIndex = OldPinName.Find(TEXT("_"));
+		const int32 UnderscoreIndex = OldPinNameStr.Find(TEXT("_"), ESearchCase::CaseSensitive);
 		if (UnderscoreIndex != INDEX_NONE)
 		{
-			PropertyName = OldPinName.Left(UnderscoreIndex);
-			ArrayIndex = FCString::Atoi(*(OldPinName.Mid(UnderscoreIndex + 1)));
+			const FString PropertyName = OldPinNameStr.Left(UnderscoreIndex);
 
 			if (RemovedPropertyNames.Contains(PropertyName))
 			{
+				const int32 ArrayIndex = FCString::Atoi(*(OldPinNameStr.Mid(UnderscoreIndex + 1)));
 				// if array index is matched, removes pins 
 				// and if array index is greater than removed index, decrease index
 				if (ArrayIndex == RemovedArrayIndex)
@@ -164,11 +163,10 @@ void UAnimGraphNode_MultiWayBlend::RemovePinsFromOldPins(TArray<UEdGraphPin*>& O
 					OldPins.RemoveAt(PinIdx);
 					--PinIdx;
 				}
-				else
-					if (ArrayIndex > RemovedArrayIndex)
-					{
-						OldPinName = FString::Printf(TEXT("%s_%d"), *PropertyName, ArrayIndex - 1);
-					}
+				else if (ArrayIndex > RemovedArrayIndex)
+				{
+					OldPins[PinIdx]->PinName = *FString::Printf(TEXT("%s_%d"), *PropertyName, ArrayIndex - 1);
+				}
 			}
 		}
 	}
diff --git a/Engine/Source/Editor/AnimGraph/Private/AnimGraphNode_SkeletalControlBase.cpp b/Engine/Source/Editor/AnimGraph/Private/AnimGraphNode_SkeletalControlBase.cpp
--- a/Engine/Source/Editor/AnimGraph/Private/AnimGraphNode_SkeletalControlBase.cpp
+++ b/Engine/Source/Editor/AnimGraph/Private/AnimGraphNode_SkeletalControlBase.cpp
@@ -72,8 +72,7 @@ FText UAnimGraphNode_SkeletalControlBase::GetTooltipText() const
 
 void UAnimGraphNode_SkeletalControlBase::CreateOutputPins()
 {
-	const UAnimationGraphSchema* Schema = GetDefault<UAnimationGraphSchema>();
-	CreatePin(EGPD_Output, Schema->PC_Struct, FString(), FComponentSpacePoseLink::StaticStruct(), TEXT("Pose"));
+	CreatePin(EGPD_Output, UAnimationGraphSchema::PC_Struct, FComponentSpacePoseLink::StaticStruct(), TEXT("Pose"));
 }
 
 
@@ -289,7 +288,7 @@ FVector UAnimGraphNode_SkeletalControlBase::ConvertWidgetLocation(const USkeleta
 	return WidgetLoc;
 }
 
-void UAnimGraphNode_SkeletalControlBase::GetDefaultValue(const FString& UpdateDefaultValueName, FVector& OutVec)
+void UAnimGraphNode_SkeletalControlBase::GetDefaultValue(const FName UpdateDefaultValueName, FVector& OutVec)
 {
 	for (UEdGraphPin* Pin : Pins)
 	{
@@ -323,7 +322,7 @@ void UAnimGraphNode_SkeletalControlBase::GetDefaultValue(const FString& UpdateDe
 	OutVec = FVector::ZeroVector;
 }
 
-void UAnimGraphNode_SkeletalControlBase::SetDefaultValue(const FString& UpdateDefaultValueName, const FVector& Value)
+void UAnimGraphNode_SkeletalControlBase::SetDefaultValue(const FName UpdateDefaultValueName, const FVector& Value)
 {
 	for (UEdGraphPin* Pin : Pins)
 	{
@@ -334,7 +333,7 @@ void UAnimGraphNode_SkeletalControlBase::SetDefaultValue(const FString& UpdateDe
 				FString Str = FString::Printf(TEXT("%.3f,%.3f,%.3f"), Value.X, Value.Y, Value.Z);
 				if (Pin->DefaultValue != Str)
 				{
-					PreEditChange(NULL);
+					PreEditChange(nullptr);
 					GetSchema()->TrySetDefaultValue(*Pin, Str);
 					PostEditChange();
 					break;
@@ -344,11 +343,11 @@ void UAnimGraphNode_SkeletalControlBase::SetDefaultValue(const FString& UpdateDe
 	}
 }
 
-bool UAnimGraphNode_SkeletalControlBase::IsPinShown(const FString& PinName) const
+bool UAnimGraphNode_SkeletalControlBase::IsPinShown(const FName PinName) const
 {
 	for (const FOptionalPinFromProperty& Pin : ShowPinForProperties)
 	{
-		if (Pin.PropertyName.ToString() == PinName)
+		if (Pin.PropertyName == PinName)
 		{
 			return Pin.bShowPin;
 		}
diff --git a/Engine/Source/Editor/AnimGraph/Private/AnimGraphNode_SubInstance.cpp b/Engine/Source/Editor/AnimGraph/Private/AnimGraphNode_SubInstance.cpp
--- a/Engine/Source/Editor/AnimGraph/Private/AnimGraphNode_SubInstance.cpp
+++ b/Engine/Source/Editor/AnimGraph/Private/AnimGraphNode_SubInstance.cpp
@@ -145,7 +145,7 @@ void UAnimGraphNode_SubInstance::ReallocatePinsDuringReconstruction(TArray<UEdGr
 			FEdGraphPinType PinType;
 			if(Schema->ConvertPropertyToPinType(PoseProperty, PinType))
 			{
-				UEdGraphPin* NewPin = CreatePin(EEdGraphPinDirection::EGPD_Input, PinType, PoseProperty->GetName());
+				UEdGraphPin* NewPin = CreatePin(EEdGraphPinDirection::EGPD_Input, PinType, PoseProperty->GetFName());
 				NewPin->PinFriendlyName = PoseProperty->GetDisplayNameText();
 
 				CustomizePinData(NewPin, PoseProperty->GetFName(), INDEX_NONE);
@@ -178,7 +178,7 @@ void UAnimGraphNode_SubInstance::ReallocatePinsDuringReconstruction(TArray<UEdGr
 
 			verify(Schema->ConvertPropertyToPinType(Property, PinType));
 
-			UEdGraphPin* NewPin = CreatePin(EEdGraphPinDirection::EGPD_Input, PinType, Property->GetName());
+			UEdGraphPin* NewPin = CreatePin(EEdGraphPinDirection::EGPD_Input, PinType, Property->GetFName());
 			NewPin->PinFriendlyName = Property->GetDisplayNameText();
 
 			// Need to grab the default value for the property from the target generated class CDO
@@ -252,7 +252,7 @@ void UAnimGraphNode_SubInstance::GetInstancePinProperty(const UClass* InOwnerIns
 
 FString UAnimGraphNode_SubInstance::GetPinTargetVariableName(const UEdGraphPin* InPin) const
 {
-	return TEXT("__SUBINSTANCE_") + InPin->PinName + TEXT("_") + NodeGuid.ToString();
+	return TEXT("__SUBINSTANCE_") + InPin->PinName.ToString() + TEXT("_") + NodeGuid.ToString();
 }
 
 void UAnimGraphNode_SubInstance::CustomizeDetails(IDetailLayoutBuilder& DetailBuilder)
diff --git a/Engine/Source/Editor/AnimGraph/Private/AnimStateConduitNode.cpp b/Engine/Source/Editor/AnimGraph/Private/AnimStateConduitNode.cpp
--- a/Engine/Source/Editor/AnimGraph/Private/AnimStateConduitNode.cpp
+++ b/Engine/Source/Editor/AnimGraph/Private/AnimStateConduitNode.cpp
@@ -22,8 +22,8 @@ UAnimStateConduitNode::UAnimStateConduitNode(const FObjectInitializer& ObjectIni
 
 void UAnimStateConduitNode::AllocateDefaultPins()
 {
-	CreatePin(EGPD_Input, TEXT("Transition"), FString(), nullptr, TEXT("In"));
-	CreatePin(EGPD_Output, TEXT("Transition"), FString(), nullptr, TEXT("Out"));
+	CreatePin(EGPD_Input, TEXT("Transition"), TEXT("In"));
+	CreatePin(EGPD_Output, TEXT("Transition"), TEXT("Out"));
 }
 
 void UAnimStateConduitNode::AutowireNewNode(UEdGraphPin* FromPin)
diff --git a/Engine/Source/Editor/AnimGraph/Private/AnimStateEntryNode.cpp b/Engine/Source/Editor/AnimGraph/Private/AnimStateEntryNode.cpp
--- a/Engine/Source/Editor/AnimGraph/Private/AnimStateEntryNode.cpp
+++ b/Engine/Source/Editor/AnimGraph/Private/AnimStateEntryNode.cpp
@@ -20,8 +20,7 @@ UAnimStateEntryNode::UAnimStateEntryNode(const FObjectInitializer& ObjectInitial
 
 void UAnimStateEntryNode::AllocateDefaultPins()
 {
-	const UAnimationStateMachineSchema* Schema = GetDefault<UAnimationStateMachineSchema>();
-	UEdGraphPin* Outputs = CreatePin(EGPD_Output, Schema->PC_Exec, FString(), nullptr, TEXT("Entry"));
+	UEdGraphPin* Outputs = CreatePin(EGPD_Output, UAnimationStateMachineSchema::PC_Exec, TEXT("Entry"));
 }
 
 FText UAnimStateEntryNode::GetNodeTitle(ENodeTitleType::Type TitleType) const
diff --git a/Engine/Source/Editor/AnimGraph/Private/AnimStateNode.cpp b/Engine/Source/Editor/AnimGraph/Private/AnimStateNode.cpp
--- a/Engine/Source/Editor/AnimGraph/Private/AnimStateNode.cpp
+++ b/Engine/Source/Editor/AnimGraph/Private/AnimStateNode.cpp
@@ -26,8 +26,8 @@ UAnimStateNode::UAnimStateNode(const FObjectInitializer& ObjectInitializer)
 
 void UAnimStateNode::AllocateDefaultPins()
 {
-	UEdGraphPin* Inputs = CreatePin(EGPD_Input, TEXT("Transition"), FString(), nullptr, TEXT("In"));
-	UEdGraphPin* Outputs = CreatePin(EGPD_Output, TEXT("Transition"), FString(), nullptr, TEXT("Out"));
+	UEdGraphPin* Inputs = CreatePin(EGPD_Input, TEXT("Transition"), TEXT("In"));
+	UEdGraphPin* Outputs = CreatePin(EGPD_Output, TEXT("Transition"), TEXT("Out"));
 }
 
 void UAnimStateNode::AutowireNewNode(UEdGraphPin* FromPin)
diff --git a/Engine/Source/Editor/AnimGraph/Private/AnimStateTransitionNode.cpp b/Engine/Source/Editor/AnimGraph/Private/AnimStateTransitionNode.cpp
--- a/Engine/Source/Editor/AnimGraph/Private/AnimStateTransitionNode.cpp
+++ b/Engine/Source/Editor/AnimGraph/Private/AnimStateTransitionNode.cpp
@@ -84,9 +84,9 @@ UAnimStateTransitionNode::UAnimStateTransitionNode(const FObjectInitializer& Obj
 
 void UAnimStateTransitionNode::AllocateDefaultPins()
 {
-	UEdGraphPin* Inputs = CreatePin(EGPD_Input, TEXT("Transition"), FString(), nullptr, TEXT("In"));
+	UEdGraphPin* Inputs = CreatePin(EGPD_Input, TEXT("Transition"), TEXT("In"));
 	Inputs->bHidden = true;
-	UEdGraphPin* Outputs = CreatePin(EGPD_Output, TEXT("Transition"), FString(), nullptr, TEXT("Out"));
+	UEdGraphPin* Outputs = CreatePin(EGPD_Output, TEXT("Transition"), TEXT("Out"));
 	Outputs->bHidden = true;
 }
 
diff --git a/Engine/Source/Editor/AnimGraph/Private/AnimationGraphSchema.cpp b/Engine/Source/Editor/AnimGraph/Private/AnimationGraphSchema.cpp
--- a/Engine/Source/Editor/AnimGraph/Private/AnimationGraphSchema.cpp
+++ b/Engine/Source/Editor/AnimGraph/Private/AnimationGraphSchema.cpp
@@ -135,18 +135,14 @@ bool UAnimationGraphSchema::IsPosePin(const FEdGraphPinType& PinType)
 
 bool UAnimationGraphSchema::IsLocalSpacePosePin(const FEdGraphPinType& PinType)
 {
-	const UAnimationGraphSchema* Schema = GetDefault<UAnimationGraphSchema>();
-
 	UScriptStruct* PoseLinkStruct = FPoseLink::StaticStruct();
-	return (PinType.PinCategory == Schema->PC_Struct) && (PinType.PinSubCategoryObject == PoseLinkStruct);
+	return (PinType.PinCategory == UAnimationGraphSchema::PC_Struct) && (PinType.PinSubCategoryObject == PoseLinkStruct);
 }
 
 bool UAnimationGraphSchema::IsComponentSpacePosePin(const FEdGraphPinType& PinType)
 {
-	const UAnimationGraphSchema* Schema = GetDefault<UAnimationGraphSchema>();
-
 	UScriptStruct* ComponentSpacePoseLinkStruct = FComponentSpacePoseLink::StaticStruct();
-	return (PinType.PinCategory == Schema->PC_Struct) && (PinType.PinSubCategoryObject == ComponentSpacePoseLinkStruct);
+	return (PinType.PinCategory == UAnimationGraphSchema::PC_Struct) && (PinType.PinSubCategoryObject == ComponentSpacePoseLinkStruct);
 }
 
 bool UAnimationGraphSchema::TryCreateConnection(UEdGraphPin* A, UEdGraphPin* B) const
@@ -428,7 +424,7 @@ void UAnimationGraphSchema::GetAssetsPinHoverMessage(const TArray<FAssetData>& A
 	if (bSkelMatch && bTypeMatch && bDirectionMatch)
 	{
 		OutOkIcon = true;
-		OutTooltipText = FText::Format(LOCTEXT("AssetPinHoverMessage_Success", "Play {0} and feed to {1}"), FText::FromString(Asset->GetName()), FText::FromString(HoverPin->PinName)).ToString();
+		OutTooltipText = FText::Format(LOCTEXT("AssetPinHoverMessage_Success", "Play {0} and feed to {1}"), FText::FromString(Asset->GetName()), FText::FromName(HoverPin->PinName)).ToString();
 	}
 	else
 	{
diff --git a/Engine/Source/Editor/AnimGraph/Private/AnimationStateMachineSchema.cpp b/Engine/Source/Editor/AnimGraph/Private/AnimationStateMachineSchema.cpp
--- a/Engine/Source/Editor/AnimGraph/Private/AnimationStateMachineSchema.cpp
+++ b/Engine/Source/Editor/AnimGraph/Private/AnimationStateMachineSchema.cpp
@@ -113,10 +113,11 @@ UEdGraphNode* FEdGraphSchemaAction_NewStateComment::PerformAction(class UEdGraph
 /////////////////////////////////////////////////////
 // UAnimationStateMachineSchema
 
+const FName UAnimationStateMachineSchema::PC_Exec(TEXT("exec"));
+
 UAnimationStateMachineSchema::UAnimationStateMachineSchema(const FObjectInitializer& ObjectInitializer)
 	: Super(ObjectInitializer)
 {
-	PC_Exec = TEXT("exec");
 }
 
 void UAnimationStateMachineSchema::CreateDefaultNodesForGraph(UEdGraph& Graph) const
@@ -462,7 +463,7 @@ void UAnimationStateMachineSchema::BreakPinLinks(UEdGraphPin& TargetPin, bool bS
 	FBlueprintEditorUtils::MarkBlueprintAsModified(Blueprint);
 }
 
-void UAnimationStateMachineSchema::BreakSinglePinLink(UEdGraphPin* SourcePin, UEdGraphPin* TargetPin)
+void UAnimationStateMachineSchema::BreakSinglePinLink(UEdGraphPin* SourcePin, UEdGraphPin* TargetPin) const
 {
 	const FScopedTransaction Transaction(NSLOCTEXT("UnrealEd", "GraphEd_BreakSinglePinLink", "Break Pin Link"));
 	UBlueprint* Blueprint = FBlueprintEditorUtils::FindBlueprintForNodeChecked(TargetPin->GetOwningNode());
diff --git a/Engine/Source/Editor/AnimGraph/Private/EditModes/ModifyBoneEditMode.cpp b/Engine/Source/Editor/AnimGraph/Private/EditModes/ModifyBoneEditMode.cpp
--- a/Engine/Source/Editor/AnimGraph/Private/EditModes/ModifyBoneEditMode.cpp
+++ b/Engine/Source/Editor/AnimGraph/Private/EditModes/ModifyBoneEditMode.cpp
@@ -66,7 +66,7 @@ FVector FModifyBoneEditMode::GetWidgetLocation() const
 	if (CurWidgetMode == FWidget::WM_Translate)
 	{
 		FCSPose<FCompactHeapPose>& MeshBases = RuntimeNode->ForwardedPose;
-		WidgetLoc = ConvertWidgetLocation(SkelComp, MeshBases, GraphNode->Node.BoneToModify.BoneName, GraphNode->GetNodeValue(FString("Translation"), GraphNode->Node.Translation), GraphNode->Node.TranslationSpace);
+		WidgetLoc = ConvertWidgetLocation(SkelComp, MeshBases, GraphNode->Node.BoneToModify.BoneName, GraphNode->GetNodeValue(TEXT("Translation"), GraphNode->Node.Translation), GraphNode->Node.TranslationSpace);
 
 		if (MeshBases.GetPose().IsValid() && GraphNode->Node.TranslationMode == BMM_Additive)
 		{
diff --git a/Engine/Source/Editor/AnimGraph/Private/K2Node_PlayMontage.cpp b/Engine/Source/Editor/AnimGraph/Private/K2Node_PlayMontage.cpp
--- a/Engine/Source/Editor/AnimGraph/Private/K2Node_PlayMontage.cpp
+++ b/Engine/Source/Editor/AnimGraph/Private/K2Node_PlayMontage.cpp
@@ -32,16 +32,15 @@ void UK2Node_PlayMontage::GetPinHoverText(const UEdGraphPin& Pin, FString& Hover
 {
 	Super::GetPinHoverText(Pin, HoverTextOut);
 
-	static FName NAME_OnNotifyBegin = FName(TEXT("OnNotifyBegin"));
-	static FName NAME_OnNotifyEnd = FName(TEXT("OnNotifyEnd"));
+	static const FName NAME_OnNotifyBegin = FName(TEXT("OnNotifyBegin"));
+	static const FName NAME_OnNotifyEnd = FName(TEXT("OnNotifyEnd"));
 
-	const FName PinName = FName(*Pin.PinName);
-	if (PinName == NAME_OnNotifyBegin)
+	if (Pin.PinName == NAME_OnNotifyBegin)
 	{
 		FText ToolTipText = LOCTEXT("K2Node_PlayMontage_OnNotifyBegin_Tooltip", "Event called when using a PlayMontageNotify or PlayMontageNotifyWindow Notify in a Montage.");
 		HoverTextOut = FString::Printf(TEXT("%s\n%s"), *ToolTipText.ToString(), *HoverTextOut);
 	}
-	else if (PinName == NAME_OnNotifyEnd)
+	else if (Pin.PinName == NAME_OnNotifyEnd)
 	{
 		FText ToolTipText = LOCTEXT("K2Node_PlayMontage_OnNotifyEnd_Tooltip", "Event called when using a PlayMontageNotifyWindow Notify in a Montage.");
 		HoverTextOut = FString::Printf(TEXT("%s\n%s"), *ToolTipText.ToString(), *HoverTextOut);
diff --git a/Engine/Source/Editor/AnimGraph/Private/K2Node_TransitionRuleGetter.cpp b/Engine/Source/Editor/AnimGraph/Private/K2Node_TransitionRuleGetter.cpp
--- a/Engine/Source/Editor/AnimGraph/Private/K2Node_TransitionRuleGetter.cpp
+++ b/Engine/Source/Editor/AnimGraph/Private/K2Node_TransitionRuleGetter.cpp
@@ -26,9 +26,7 @@ UK2Node_TransitionRuleGetter::UK2Node_TransitionRuleGetter(const FObjectInitiali
 
 void UK2Node_TransitionRuleGetter::AllocateDefaultPins()
 {
-	const UEdGraphSchema_K2* Schema = GetDefault<UEdGraphSchema_K2>();
-
-	UEdGraphPin* OutputPin = CreatePin(EGPD_Output, Schema->PC_Float, /*PSC=*/ FString(), /*PSC object=*/ nullptr, TEXT("Output"));
+	UEdGraphPin* OutputPin = CreatePin(EGPD_Output, UEdGraphSchema_K2::PC_Float, TEXT("Output"));
 	OutputPin->PinFriendlyName = GetFriendlyName(GetterType);
 
 	PreloadObject(AssociatedAnimAssetPlayerNode);
diff --git a/Engine/Source/Editor/AnimationBlueprintEditor/AnimationBlueprintEditor.Build.cs b/Engine/Source/Editor/AnimationBlueprintEditor/AnimationBlueprintEditor.Build.cs
--- a/Engine/Source/Editor/AnimationBlueprintEditor/AnimationBlueprintEditor.Build.cs
+++ b/Engine/Source/Editor/AnimationBlueprintEditor/AnimationBlueprintEditor.Build.cs
@@ -38,6 +38,7 @@ public class AnimationBlueprintEditor : ModuleRules
 				"EditorWidgets",
                 "BlueprintGraph",
                 "RHI",
+                "KismetCompiler",
             }
 		);
 
diff --git a/Engine/Source/Editor/AnimationBlueprintEditor/Private/AnimationBlueprintEditorModule.cpp b/Engine/Source/Editor/AnimationBlueprintEditor/Private/AnimationBlueprintEditorModule.cpp
--- a/Engine/Source/Editor/AnimationBlueprintEditor/Private/AnimationBlueprintEditorModule.cpp
+++ b/Engine/Source/Editor/AnimationBlueprintEditor/Private/AnimationBlueprintEditorModule.cpp
@@ -6,7 +6,9 @@
 #include "AnimationBlueprintEditor.h"
 #include "EdGraphSchema_K2.h"
 #include "K2Node_CallFunction.h"
+#include "Animation/AnimBlueprint.h"
 #include "Kismet2/KismetEditorUtilities.h"
+#include "AnimBlueprintCompiler.h"
 
 IMPLEMENT_MODULE( FAnimationBlueprintEditorModule, AnimationBlueprintEditor);
 
@@ -27,10 +29,23 @@ void FAnimationBlueprintEditorModule::StartupModule()
 	FEdGraphUtilities::RegisterVisualPinConnectionFactory(AnimGraphPinConnectionFactory);
 
 	FKismetEditorUtilities::RegisterOnBlueprintCreatedCallback(this, UAnimInstance::StaticClass(), FKismetEditorUtilities::FOnBlueprintCreated::CreateRaw(this, &FAnimationBlueprintEditorModule::OnNewBlueprintCreated));
+
+	// Register widget blueprint compiler we do this no matter what.
+	IKismetCompilerInterface& KismetCompilerModule = FModuleManager::LoadModuleChecked<IKismetCompilerInterface>("KismetCompiler");
+	KismetCompilerModule.GetCompilers().Add(&AnimBlueprintCompiler);
+
+	FBlueprintEditorModule::RegisterCompilerGetter(UAnimBlueprint::StaticClass(), FOnGetBlueprintCompiler::CreateStatic(&FAnimationBlueprintEditorModule::GetCompiler));
 }
 
 void FAnimationBlueprintEditorModule::ShutdownModule()
 {
+	IKismetCompilerInterface* KismetCompilerModule = FModuleManager::GetModulePtr<IKismetCompilerInterface>("KismetCompiler");
+	if (KismetCompilerModule)
+	{
+		KismetCompilerModule->GetCompilers().Remove(&AnimBlueprintCompiler);
+	}
+
+	FBlueprintEditorModule::UnregisterCompilerGetter(UAnimBlueprint::StaticClass());
 	FKismetEditorUtilities::UnregisterAutoBlueprintNodeCreation(this);
 
 	FEdGraphUtilities::UnregisterVisualNodeFactory(AnimGraphNodeFactory);
@@ -80,4 +95,10 @@ void FAnimationBlueprintEditorModule::OnNewBlueprintCreated(UBlueprint* InBluepr
 	}
 }
 
+TSharedPtr<FKismetCompilerContext> FAnimationBlueprintEditorModule::GetCompiler(UBlueprint* BP, FCompilerResultsLog& InMessageLog, const FKismetCompilerOptions& InCompileOptions)
+{
+	return TSharedPtr<FKismetCompilerContext>(new FAnimBlueprintCompilerContext(CastChecked<UAnimBlueprint>(BP), InMessageLog, InCompileOptions, nullptr));
+}
+
+
 #undef LOCTEXT_NAMESPACE
diff --git a/Engine/Source/Editor/AnimationBlueprintEditor/Private/AnimationBlueprintEditorModule.h b/Engine/Source/Editor/AnimationBlueprintEditor/Private/AnimationBlueprintEditorModule.h
--- a/Engine/Source/Editor/AnimationBlueprintEditor/Private/AnimationBlueprintEditorModule.h
+++ b/Engine/Source/Editor/AnimationBlueprintEditor/Private/AnimationBlueprintEditorModule.h
@@ -6,6 +6,7 @@
 #include "Toolkits/IToolkitHost.h"
 #include "Toolkits/AssetEditorToolkit.h"
 #include "IAnimationBlueprintEditorModule.h"
+#include "AnimBlueprintCompiler.h"
 
 struct FAnimationGraphNodeFactory;
 struct FAnimationGraphPinFactory;
@@ -48,4 +49,9 @@ private:
 	TSharedPtr<FAnimationGraphNodeFactory> AnimGraphNodeFactory;
 	TSharedPtr<FAnimationGraphPinFactory> AnimGraphPinFactory;
 	TSharedPtr<FAnimationGraphPinConnectionFactory> AnimGraphPinConnectionFactory;
+
+	static TSharedPtr<FKismetCompilerContext> GetCompiler(UBlueprint* BP, FCompilerResultsLog& InMessageLog, const FKismetCompilerOptions& InCompileOptions);
+
+	/** Compiler customization for animation controllers */
+	FAnimBlueprintCompiler AnimBlueprintCompiler;
 };
diff --git a/Engine/Source/Editor/AnimationBlueprintEditor/Private/AnimationGraphFactory.cpp b/Engine/Source/Editor/AnimationBlueprintEditor/Private/AnimationGraphFactory.cpp
--- a/Engine/Source/Editor/AnimationBlueprintEditor/Private/AnimationGraphFactory.cpp
+++ b/Engine/Source/Editor/AnimationBlueprintEditor/Private/AnimationGraphFactory.cpp
@@ -87,13 +87,9 @@ TSharedPtr<class SGraphPin> FAnimationGraphPinFactory::CreatePin(class UEdGraphP
 		}
 	}
 
-	const UAnimationStateMachineSchema* AnimStateMachineSchema = Cast<const UAnimationStateMachineSchema>(InPin->GetSchema());
-	if (AnimStateMachineSchema != NULL)
+	if (InPin->PinType.PinCategory == UAnimationStateMachineSchema::PC_Exec)
 	{
-		if (InPin->PinType.PinCategory == AnimStateMachineSchema->PC_Exec)
-		{
-			return SNew(SGraphPinExec, InPin);
-		}
+		return SNew(SGraphPinExec, InPin);
 	}
 
 	return nullptr;
diff --git a/Engine/Source/Editor/AnimationEditor/Private/AnimationEditor.h b/Engine/Source/Editor/AnimationEditor/Private/AnimationEditor.h
--- a/Engine/Source/Editor/AnimationEditor/Private/AnimationEditor.h
+++ b/Engine/Source/Editor/AnimationEditor/Private/AnimationEditor.h
@@ -71,7 +71,7 @@ public:
 	/** FTickableEditorObject Interface */
 	virtual void Tick(float DeltaTime) override;
 	virtual TStatId GetStatId() const override;
-	virtual bool IsTickable() const override { return true; }
+	virtual ETickableTickType GetTickableTickType() const override { return ETickableTickType::Always; }
 	
 	/** FEditorUndoClient interface */
 	virtual void PostUndo(bool bSuccess) override;
diff --git a/Engine/Source/Editor/AudioEditor/Classes/SoundClassGraph/SoundClassGraphSchema.h b/Engine/Source/Editor/AudioEditor/Classes/SoundClassGraph/SoundClassGraphSchema.h
--- a/Engine/Source/Editor/AudioEditor/Classes/SoundClassGraph/SoundClassGraphSchema.h
+++ b/Engine/Source/Editor/AudioEditor/Classes/SoundClassGraph/SoundClassGraphSchema.h
@@ -56,7 +56,7 @@ class USoundClassGraphSchema : public UEdGraphSchema
 	virtual FLinearColor GetPinTypeColor(const FEdGraphPinType& PinType) const override;
 	virtual void BreakNodeLinks(UEdGraphNode& TargetNode) const override;
 	virtual void BreakPinLinks(UEdGraphPin& TargetPin, bool bSendsNodeNotifcation) const override;
-	virtual void BreakSinglePinLink(UEdGraphPin* SourcePin, UEdGraphPin* TargetPin) override;
+	virtual void BreakSinglePinLink(UEdGraphPin* SourcePin, UEdGraphPin* TargetPin) const override;
 	virtual void DroppedAssetsOnGraph(const TArray<struct FAssetData>& Assets, const FVector2D& GraphPosition, UEdGraph* Graph) const override;
 	//~ End EdGraphSchema Interface
 };
diff --git a/Engine/Source/Editor/AudioEditor/Private/SoundClassGraphNode.cpp b/Engine/Source/Editor/AudioEditor/Private/SoundClassGraphNode.cpp
--- a/Engine/Source/Editor/AudioEditor/Private/SoundClassGraphNode.cpp
+++ b/Engine/Source/Editor/AudioEditor/Private/SoundClassGraphNode.cpp
@@ -71,8 +71,8 @@ void USoundClassGraphNode::AllocateDefaultPins()
 {
 	check(Pins.Num() == 0);
 
-	ChildPin = CreatePin(EGPD_Output, TEXT("SoundClass"), FString(), nullptr, LOCTEXT("SoundClassChildren", "Children").ToString());
-	ParentPin = CreatePin(EGPD_Input, TEXT("SoundClass"), FString(), nullptr, FString());
+	ChildPin = CreatePin(EGPD_Output, TEXT("SoundClass"), *LOCTEXT("SoundClassChildren", "Children").ToString());
+	ParentPin = CreatePin(EGPD_Input, TEXT("SoundClass"), NAME_None);
 }
 
 void USoundClassGraphNode::AutowireNewNode(UEdGraphPin* FromPin)
diff --git a/Engine/Source/Editor/AudioEditor/Private/SoundClassGraphSchema.cpp b/Engine/Source/Editor/AudioEditor/Private/SoundClassGraphSchema.cpp
--- a/Engine/Source/Editor/AudioEditor/Private/SoundClassGraphSchema.cpp
+++ b/Engine/Source/Editor/AudioEditor/Private/SoundClassGraphSchema.cpp
@@ -46,7 +46,7 @@ void USoundClassGraphSchema::GetBreakLinkToSubMenuActions( class FMenuBuilder& M
 		FText Title = FText::FromString( TitleString );
 		if ( Pin->PinName != TEXT("") )
 		{
-			TitleString = FString::Printf(TEXT("%s (%s)"), *TitleString, *Pin->PinName);
+			TitleString = FString::Printf(TEXT("%s (%s)"), *TitleString, *Pin->PinName.ToString());
 
 			// Add name of connection if possible
 			FFormatNamedArguments Args;
@@ -211,7 +211,7 @@ void USoundClassGraphSchema::BreakPinLinks(UEdGraphPin& TargetPin, bool bSendsNo
 	}
 }
 
-void USoundClassGraphSchema::BreakSinglePinLink(UEdGraphPin* SourcePin, UEdGraphPin* TargetPin)
+void USoundClassGraphSchema::BreakSinglePinLink(UEdGraphPin* SourcePin, UEdGraphPin* TargetPin) const
 {
 	const FScopedTransaction Transaction( NSLOCTEXT("UnrealEd", "GraphEd_BreakSinglePinLink", "Break Pin Link") );
 	Super::BreakSinglePinLink(SourcePin, TargetPin);
diff --git a/Engine/Source/Editor/AudioEditor/Private/SoundCueGraph.cpp b/Engine/Source/Editor/AudioEditor/Private/SoundCueGraph.cpp
--- a/Engine/Source/Editor/AudioEditor/Private/SoundCueGraph.cpp
+++ b/Engine/Source/Editor/AudioEditor/Private/SoundCueGraph.cpp
@@ -159,7 +159,7 @@ public:
 		{
 			if (InputPins[i])
 			{
-				InputPins[i]->PinName = SoundNode->GetInputPinName(i).ToString();
+				InputPins[i]->PinName = *SoundNode->GetInputPinName(i).ToString();
 			}
 		}
 	}
diff --git a/Engine/Source/Editor/AudioEditor/Private/SoundCueGraphNode.cpp b/Engine/Source/Editor/AudioEditor/Private/SoundCueGraphNode.cpp
--- a/Engine/Source/Editor/AudioEditor/Private/SoundCueGraphNode.cpp
+++ b/Engine/Source/Editor/AudioEditor/Private/SoundCueGraphNode.cpp
@@ -36,7 +36,7 @@ void USoundCueGraphNode::PostLoad()
 	for (int32 Index = 0; Index < Pins.Num(); ++Index)
 	{
 		UEdGraphPin* Pin = Pins[Index];
-		if (Pin->PinName.IsEmpty())
+		if (Pin->PinName.IsNone())
 		{
 			// Makes sure pin has a name for lookup purposes but user will never see it
 			if (Pin->Direction == EGPD_Input)
@@ -60,8 +60,8 @@ void USoundCueGraphNode::SetSoundNode(USoundNode* InSoundNode)
 
 void USoundCueGraphNode::CreateInputPin()
 {
-	UEdGraphPin* NewPin = CreatePin(EGPD_Input, TEXT("SoundNode"), FString(), nullptr, SoundNode->GetInputPinName(GetInputCount()).ToString());
-	if (NewPin->PinName.IsEmpty())
+	UEdGraphPin* NewPin = CreatePin(EGPD_Input, TEXT("SoundNode"), *SoundNode->GetInputPinName(GetInputCount()).ToString());
+	if (NewPin->PinName.IsNone())
 	{
 		// Makes sure pin has a name for lookup purposes but user will never see it
 		NewPin->PinName = CreateUniquePinName(TEXT("Input"));
diff --git a/Engine/Source/Editor/AudioEditor/Private/SoundCueGraphNode_Base.cpp b/Engine/Source/Editor/AudioEditor/Private/SoundCueGraphNode_Base.cpp
--- a/Engine/Source/Editor/AudioEditor/Private/SoundCueGraphNode_Base.cpp
+++ b/Engine/Source/Editor/AudioEditor/Private/SoundCueGraphNode_Base.cpp
@@ -124,7 +124,7 @@ void USoundCueGraphNode_Base::AllocateDefaultPins()
 
 	if (!IsRootNode())
 	{
-		CreatePin(EGPD_Output, TEXT("SoundNode"), FString(), nullptr, TEXT("Output"));
+		CreatePin(EGPD_Output, TEXT("SoundNode"), TEXT("Output"));
 	}
 }
 
diff --git a/Engine/Source/Editor/AudioEditor/Private/SoundCueGraphNode_Root.cpp b/Engine/Source/Editor/AudioEditor/Private/SoundCueGraphNode_Root.cpp
--- a/Engine/Source/Editor/AudioEditor/Private/SoundCueGraphNode_Root.cpp
+++ b/Engine/Source/Editor/AudioEditor/Private/SoundCueGraphNode_Root.cpp
@@ -36,7 +36,7 @@ FText USoundCueGraphNode_Root::GetTooltipText() const
 
 void USoundCueGraphNode_Root::CreateInputPins()
 {
-	CreatePin(EGPD_Input, TEXT("SoundNode"), TEXT("Root"), nullptr, FString());
+	CreatePin(EGPD_Input, TEXT("SoundNode"), TEXT("Root"), NAME_None);
 }
 
 void USoundCueGraphNode_Root::GetContextMenuActions(const FGraphNodeContextMenuBuilder& Context) const
diff --git a/Engine/Source/Editor/AudioEditor/Private/SoundSubmixGraphNode.cpp b/Engine/Source/Editor/AudioEditor/Private/SoundSubmixGraphNode.cpp
--- a/Engine/Source/Editor/AudioEditor/Private/SoundSubmixGraphNode.cpp
+++ b/Engine/Source/Editor/AudioEditor/Private/SoundSubmixGraphNode.cpp
@@ -61,8 +61,8 @@ void USoundSubmixGraphNode::AllocateDefaultPins()
 {
 	check(Pins.Num() == 0);
 
-	ChildPin = CreatePin(EGPD_Output, TEXT("SoundSubmix"), FString(), nullptr, LOCTEXT("SoundSubmixChildren", "Children").ToString());
-	ParentPin = CreatePin(EGPD_Input, TEXT("SoundSubmix"), FString(), nullptr, FString());
+	ChildPin = CreatePin(EGPD_Output, TEXT("SoundSubmix"), *LOCTEXT("SoundSubmixChildren", "Children").ToString());
+	ParentPin = CreatePin(EGPD_Input, TEXT("SoundSubmix"), NAME_None);
 }
 
 void USoundSubmixGraphNode::AutowireNewNode(UEdGraphPin* FromPin)
diff --git a/Engine/Source/Editor/AudioEditor/Private/SoundSubmixGraphSchema.cpp b/Engine/Source/Editor/AudioEditor/Private/SoundSubmixGraphSchema.cpp
--- a/Engine/Source/Editor/AudioEditor/Private/SoundSubmixGraphSchema.cpp
+++ b/Engine/Source/Editor/AudioEditor/Private/SoundSubmixGraphSchema.cpp
@@ -45,7 +45,7 @@ void USoundSubmixGraphSchema::GetBreakLinkToSubMenuActions(class FMenuBuilder& M
 		FText Title = FText::FromString( TitleString );
 		if (Pin->PinName != TEXT(""))
 		{
-			TitleString = FString::Printf(TEXT("%s (%s)"), *TitleString, *Pin->PinName);
+			TitleString = FString::Printf(TEXT("%s (%s)"), *TitleString, *Pin->PinName.ToString());
 
 			// Add name of connection if possible
 			FFormatNamedArguments Args;
@@ -210,7 +210,7 @@ void USoundSubmixGraphSchema::BreakPinLinks(UEdGraphPin& TargetPin, bool bSendsN
 	}
 }
 
-void USoundSubmixGraphSchema::BreakSinglePinLink(UEdGraphPin* SourcePin, UEdGraphPin* TargetPin)
+void USoundSubmixGraphSchema::BreakSinglePinLink(UEdGraphPin* SourcePin, UEdGraphPin* TargetPin) const
 {
 	const FScopedTransaction Transaction( NSLOCTEXT("UnrealEd", "GraphEd_BreakSinglePinLink", "Break Pin Link") );
 	Super::BreakSinglePinLink(SourcePin, TargetPin);
diff --git a/Engine/Source/Editor/BehaviorTreeEditor/Classes/BehaviorTreeDecoratorGraphNode_Decorator.h b/Engine/Source/Editor/BehaviorTreeEditor/Classes/BehaviorTreeDecoratorGraphNode_Decorator.h
--- a/Engine/Source/Editor/BehaviorTreeEditor/Classes/BehaviorTreeDecoratorGraphNode_Decorator.h
+++ b/Engine/Source/Editor/BehaviorTreeEditor/Classes/BehaviorTreeDecoratorGraphNode_Decorator.h
@@ -8,6 +8,8 @@
 #include "AIGraphTypes.h"
 #include "BehaviorTreeDecoratorGraphNode_Decorator.generated.h"
 
+struct FGraphNodeClassData;
+
 UCLASS()
 class UBehaviorTreeDecoratorGraphNode_Decorator : public UBehaviorTreeDecoratorGraphNode 
 {
@@ -17,7 +19,7 @@ class UBehaviorTreeDecoratorGraphNode_Decorator : public UBehaviorTreeDecoratorG
 	UObject* NodeInstance;
 
 	UPROPERTY()
-	struct FGraphNodeClassData ClassData;
+	FGraphNodeClassData ClassData;
 
 	virtual void PostPlacedNewNode() override;
 	virtual void AllocateDefaultPins() override;
@@ -26,12 +28,13 @@ class UBehaviorTreeDecoratorGraphNode_Decorator : public UBehaviorTreeDecoratorG
 	virtual EBTDecoratorLogic::Type GetOperationType() const override;
 
 	virtual void PostEditImport() override;
+	virtual void PostEditUndo() override;
 	virtual void PrepareForCopying() override;
 	void PostCopyNode();
 	bool RefreshNodeClass();
 	void UpdateNodeClassData();
 
-private:
-
-	void ResetNodeOwner();
+protected:
+	friend class UBehaviorTreeGraphNode_CompositeDecorator;
+	virtual void ResetNodeOwner();
 };
diff --git a/Engine/Source/Editor/BehaviorTreeEditor/Classes/BehaviorTreeEditorTypes.h b/Engine/Source/Editor/BehaviorTreeEditor/Classes/BehaviorTreeEditorTypes.h
--- a/Engine/Source/Editor/BehaviorTreeEditor/Classes/BehaviorTreeEditorTypes.h
+++ b/Engine/Source/Editor/BehaviorTreeEditor/Classes/BehaviorTreeEditorTypes.h
@@ -59,8 +59,8 @@ class UBehaviorTreeEditorTypes : public UObject
 {
 	GENERATED_UCLASS_BODY()
 
-	static const FString PinCategory_MultipleNodes;
-	static const FString PinCategory_SingleComposite;
-	static const FString PinCategory_SingleTask;
-	static const FString PinCategory_SingleNode;
+	static const FName PinCategory_MultipleNodes;
+	static const FName PinCategory_SingleComposite;
+	static const FName PinCategory_SingleTask;
+	static const FName PinCategory_SingleNode;
 };
diff --git a/Engine/Source/Editor/BehaviorTreeEditor/Classes/BehaviorTreeGraph.h b/Engine/Source/Editor/BehaviorTreeEditor/Classes/BehaviorTreeGraph.h
--- a/Engine/Source/Editor/BehaviorTreeEditor/Classes/BehaviorTreeGraph.h
+++ b/Engine/Source/Editor/BehaviorTreeEditor/Classes/BehaviorTreeGraph.h
@@ -41,7 +41,6 @@ class UBehaviorTreeGraph : public UAIGraph
 	void CreateBTFromGraph(class UBehaviorTreeGraphNode* RootEdNode);
 	void SpawnMissingNodes();
 	void UpdatePinConnectionTypes();
-	void UpdateDeprecatedNodes();
 	bool UpdateInjectedNodes();
 	void UpdateBrokenComposites();
 	class UEdGraphNode* FindInjectedNode(int32 Index);
@@ -63,4 +62,5 @@ protected:
 
 	void UpdateVersion_UnifiedSubNodes();
 	void UpdateVersion_InnerGraphWhitespace();
+	void UpdateVersion_RunBehaviorInSeparateGraph();
 };
diff --git a/Engine/Source/Editor/BehaviorTreeEditor/Classes/BehaviorTreeGraphNode_Composite.h b/Engine/Source/Editor/BehaviorTreeEditor/Classes/BehaviorTreeGraphNode_Composite.h
--- a/Engine/Source/Editor/BehaviorTreeEditor/Classes/BehaviorTreeGraphNode_Composite.h
+++ b/Engine/Source/Editor/BehaviorTreeEditor/Classes/BehaviorTreeGraphNode_Composite.h
@@ -13,6 +13,8 @@ class UBehaviorTreeGraphNode_Composite : public UBehaviorTreeGraphNode
 	GENERATED_UCLASS_BODY()
 	
 	virtual FText GetNodeTitle(ENodeTitleType::Type TitleType) const override;
+	virtual FText GetDescription() const override;
+	virtual FText GetTooltipText() const override;
 	virtual bool RefreshNodeClass() override{ return false; }
 
 	/** Gets a list of actions that can be done to this particular node */
diff --git a/Engine/Source/Editor/BehaviorTreeEditor/Classes/BehaviorTreeGraphNode_CompositeDecorator.h b/Engine/Source/Editor/BehaviorTreeEditor/Classes/BehaviorTreeGraphNode_CompositeDecorator.h
--- a/Engine/Source/Editor/BehaviorTreeEditor/Classes/BehaviorTreeGraphNode_CompositeDecorator.h
+++ b/Engine/Source/Editor/BehaviorTreeEditor/Classes/BehaviorTreeGraphNode_CompositeDecorator.h
@@ -65,6 +65,7 @@ class UBehaviorTreeGraphNode_CompositeDecorator : public UBehaviorTreeGraphNode
 
 protected:
 	void CreateBoundGraph();
+	virtual void ResetNodeOwner() override;
 
 	UPROPERTY()
 	class UBTCompositeNode* ParentNodeInstance;
diff --git a/Engine/Source/Editor/BehaviorTreeEditor/Private/BehaviorTreeDecoratorGraphNode_Decorator.cpp b/Engine/Source/Editor/BehaviorTreeEditor/Private/BehaviorTreeDecoratorGraphNode_Decorator.cpp
--- a/Engine/Source/Editor/BehaviorTreeEditor/Private/BehaviorTreeDecoratorGraphNode_Decorator.cpp
+++ b/Engine/Source/Editor/BehaviorTreeEditor/Private/BehaviorTreeDecoratorGraphNode_Decorator.cpp
@@ -15,7 +15,7 @@ UBehaviorTreeDecoratorGraphNode_Decorator::UBehaviorTreeDecoratorGraphNode_Decor
 
 void UBehaviorTreeDecoratorGraphNode_Decorator::AllocateDefaultPins()
 {
-	CreatePin(EGPD_Output, TEXT("Transition"), FString(), nullptr, TEXT("Out"));
+	CreatePin(EGPD_Output, TEXT("Transition"), TEXT("Out"));
 }
 
 void UBehaviorTreeDecoratorGraphNode_Decorator::PostPlacedNewNode()
@@ -98,6 +98,12 @@ void UBehaviorTreeDecoratorGraphNode_Decorator::PostEditImport()
 	}
 }
 
+void UBehaviorTreeDecoratorGraphNode_Decorator::PostEditUndo()
+{
+	Super::PostEditUndo();
+	ResetNodeOwner();
+}
+
 void UBehaviorTreeDecoratorGraphNode_Decorator::PostCopyNode()
 {
 	ResetNodeOwner();
@@ -111,6 +117,7 @@ void UBehaviorTreeDecoratorGraphNode_Decorator::ResetNodeOwner()
 		UBehaviorTree* BT = OwningNode ? Cast<UBehaviorTree>(OwningNode->GetOuter()->GetOuter()) : NULL;
 
 		NodeInstance->Rename(NULL, BT, REN_DontCreateRedirectors | REN_DoNotDirty);
+		NodeInstance->ClearFlags(RF_Transient);
 	}
 }
 
diff --git a/Engine/Source/Editor/BehaviorTreeEditor/Private/BehaviorTreeDecoratorGraphNode_Logic.cpp b/Engine/Source/Editor/BehaviorTreeEditor/Private/BehaviorTreeDecoratorGraphNode_Logic.cpp
--- a/Engine/Source/Editor/BehaviorTreeEditor/Private/BehaviorTreeDecoratorGraphNode_Logic.cpp
+++ b/Engine/Source/Editor/BehaviorTreeEditor/Private/BehaviorTreeDecoratorGraphNode_Logic.cpp
@@ -23,7 +23,7 @@ void UBehaviorTreeDecoratorGraphNode_Logic::AllocateDefaultPins()
 
 	if (LogicMode != EDecoratorLogicMode::Sink)
 	{
-		CreatePin(EGPD_Output, TEXT("Transition"), FString(), nullptr, TEXT("Out"));
+		CreatePin(EGPD_Output, TEXT("Transition"), TEXT("Out"));
 	}
 }
 
@@ -73,7 +73,7 @@ bool UBehaviorTreeDecoratorGraphNode_Logic::CanRemovePins() const
 
 UEdGraphPin* UBehaviorTreeDecoratorGraphNode_Logic::AddInputPin()
 {
-	return CreatePin(EGPD_Input, TEXT("Transition"), FString(), nullptr, TEXT("In"));
+	return CreatePin(EGPD_Input, TEXT("Transition"), TEXT("In"));
 }
 
 void UBehaviorTreeDecoratorGraphNode_Logic::RemoveInputPin(class UEdGraphPin* Pin)
diff --git a/Engine/Source/Editor/BehaviorTreeEditor/Private/BehaviorTreeEditor.cpp b/Engine/Source/Editor/BehaviorTreeEditor/Private/BehaviorTreeEditor.cpp
--- a/Engine/Source/Editor/BehaviorTreeEditor/Private/BehaviorTreeEditor.cpp
+++ b/Engine/Source/Editor/BehaviorTreeEditor/Private/BehaviorTreeEditor.cpp
@@ -36,7 +36,7 @@
 #include "ScopedTransaction.h"
 #include "BehaviorTreeColors.h"
 
-#include "BehaviorTree/Composites/BTComposite_SimpleParallel.h"
+#include "BehaviorTree/BTCompositeNode.h"
 #include "BehaviorTree/BlackboardData.h"
 #include "BehaviorTree/BehaviorTree.h"
 #include "BehaviorTree/Tasks/BTTask_RunBehavior.h"
@@ -791,7 +791,7 @@ static uint16 GetMaxAllowedRange(const class UBTDecorator* DecoratorOb)
 	UBTCompositeNode* TestParent = DecoratorOb->GetParentNode();
 	while (TestParent)
 	{
-		if (TestParent->IsA(UBTComposite_SimpleParallel::StaticClass()))
+		if (TestParent->IsApplyingDecoratorScope())
 		{
 			MaxRange = TestParent->GetLastExecutionIndex();
 			break;
diff --git a/Engine/Source/Editor/BehaviorTreeEditor/Private/BehaviorTreeEditorTypes.cpp b/Engine/Source/Editor/BehaviorTreeEditor/Private/BehaviorTreeEditorTypes.cpp
--- a/Engine/Source/Editor/BehaviorTreeEditor/Private/BehaviorTreeEditorTypes.cpp
+++ b/Engine/Source/Editor/BehaviorTreeEditor/Private/BehaviorTreeEditorTypes.cpp
@@ -2,10 +2,10 @@
 
 #include "BehaviorTreeEditorTypes.h"
 
-const FString UBehaviorTreeEditorTypes::PinCategory_MultipleNodes("MultipleNodes");
-const FString UBehaviorTreeEditorTypes::PinCategory_SingleComposite("SingleComposite");
-const FString UBehaviorTreeEditorTypes::PinCategory_SingleTask("SingleTask");
-const FString UBehaviorTreeEditorTypes::PinCategory_SingleNode("SingleNode");
+const FName UBehaviorTreeEditorTypes::PinCategory_MultipleNodes("MultipleNodes");
+const FName UBehaviorTreeEditorTypes::PinCategory_SingleComposite("SingleComposite");
+const FName UBehaviorTreeEditorTypes::PinCategory_SingleTask("SingleTask");
+const FName UBehaviorTreeEditorTypes::PinCategory_SingleNode("SingleNode");
 
 UBehaviorTreeEditorTypes::UBehaviorTreeEditorTypes(const FObjectInitializer& ObjectInitializer) : Super(ObjectInitializer)
 {
diff --git a/Engine/Source/Editor/BehaviorTreeEditor/Private/BehaviorTreeGraph.cpp b/Engine/Source/Editor/BehaviorTreeEditor/Private/BehaviorTreeGraph.cpp
--- a/Engine/Source/Editor/BehaviorTreeEditor/Private/BehaviorTreeGraph.cpp
+++ b/Engine/Source/Editor/BehaviorTreeEditor/Private/BehaviorTreeGraph.cpp
@@ -33,8 +33,9 @@ namespace BTGraphVersion
 	const int32 Initial = 0;
 	const int32 UnifiedSubNodes = 1;
 	const int32 InnerGraphWhitespace = 2;
+	const int32 RunBehaviorInSeparateGraph = 3;
 
-	const int32 Latest = InnerGraphWhitespace;
+	const int32 Latest = RunBehaviorInSeparateGraph;
 }
 
 UBehaviorTreeGraph::UBehaviorTreeGraph(const FObjectInitializer& ObjectInitializer) : Super(ObjectInitializer)
@@ -184,7 +185,6 @@ void UBehaviorTreeGraph::OnLoaded()
 {
 	Super::OnLoaded();
 	UpdatePinConnectionTypes();
-	UpdateDeprecatedNodes();
 	RemoveUnknownSubNodes();
 }
 
@@ -212,7 +212,7 @@ void UBehaviorTreeGraph::UpdatePinConnectionTypes()
 
 		for (int32 iPin = 0; iPin < Node->Pins.Num(); iPin++)
 		{
-			FString& PinCategory = Node->Pins[iPin]->PinType.PinCategory;
+			FName& PinCategory = Node->Pins[iPin]->PinType.PinCategory;
 			if (PinCategory == TEXT("Transition"))
 			{
 				PinCategory = bIsCompositeNode ? 
@@ -251,7 +251,7 @@ void UBehaviorTreeGraph::ReplaceNodeConnections(UEdGraphNode* OldNode, UEdGraphN
 	}
 }
 
-void UBehaviorTreeGraph::UpdateDeprecatedNodes()
+void UBehaviorTreeGraph::UpdateVersion_RunBehaviorInSeparateGraph()
 {
 	for (int32 Index = 0; Index < Nodes.Num(); ++Index)
 	{
@@ -1318,6 +1318,11 @@ void UBehaviorTreeGraph::UpdateVersion()
 		UpdateVersion_InnerGraphWhitespace();
 	}
 
+	if (GraphVersion < BTGraphVersion::RunBehaviorInSeparateGraph)
+	{
+		UpdateVersion_RunBehaviorInSeparateGraph();
+	}
+
 	GraphVersion = BTGraphVersion::Latest;
 	Modify();
 }
diff --git a/Engine/Source/Editor/BehaviorTreeEditor/Private/BehaviorTreeGraphNode.cpp b/Engine/Source/Editor/BehaviorTreeEditor/Private/BehaviorTreeGraphNode.cpp
--- a/Engine/Source/Editor/BehaviorTreeEditor/Private/BehaviorTreeGraphNode.cpp
+++ b/Engine/Source/Editor/BehaviorTreeEditor/Private/BehaviorTreeGraphNode.cpp
@@ -44,8 +44,8 @@ UBehaviorTreeGraphNode::UBehaviorTreeGraphNode(const FObjectInitializer& ObjectI
 
 void UBehaviorTreeGraphNode::AllocateDefaultPins()
 {
-	CreatePin(EGPD_Input, UBehaviorTreeEditorTypes::PinCategory_MultipleNodes, FString(), nullptr, TEXT("In"));
-	CreatePin(EGPD_Output, UBehaviorTreeEditorTypes::PinCategory_MultipleNodes, FString(), nullptr, TEXT("Out"));
+	CreatePin(EGPD_Input, UBehaviorTreeEditorTypes::PinCategory_MultipleNodes, TEXT("In"));
+	CreatePin(EGPD_Output, UBehaviorTreeEditorTypes::PinCategory_MultipleNodes, TEXT("Out"));
 }
 
 void UBehaviorTreeGraphNode::InitializeInstance()
diff --git a/Engine/Source/Editor/BehaviorTreeEditor/Private/BehaviorTreeGraphNode_Composite.cpp b/Engine/Source/Editor/BehaviorTreeEditor/Private/BehaviorTreeGraphNode_Composite.cpp
--- a/Engine/Source/Editor/BehaviorTreeEditor/Private/BehaviorTreeGraphNode_Composite.cpp
+++ b/Engine/Source/Editor/BehaviorTreeEditor/Private/BehaviorTreeGraphNode_Composite.cpp
@@ -2,6 +2,7 @@
 
 #include "BehaviorTreeGraphNode_Composite.h"
 #include "BehaviorTree/BTNode.h"
+#include "BehaviorTree/BTCompositeNode.h"
 
 UBehaviorTreeGraphNode_Composite::UBehaviorTreeGraphNode_Composite(const FObjectInitializer& ObjectInitializer)
 	: Super(ObjectInitializer)
@@ -18,6 +19,32 @@ FText UBehaviorTreeGraphNode_Composite::GetNodeTitle(ENodeTitleType::Type TitleT
 	return Super::GetNodeTitle(TitleType);
 }
 
+FText UBehaviorTreeGraphNode_Composite::GetDescription() const
+{
+	const UBTCompositeNode* CompositeNode = Cast<UBTCompositeNode>(NodeInstance);
+	if (CompositeNode && CompositeNode->IsApplyingDecoratorScope())
+	{
+		return FText::Format(FText::FromString(TEXT("{0}\n{1}")),
+			Super::GetDescription(),
+			NSLOCTEXT("BehaviorTreeEditor", "CompositeNodeScopeDesc", "Local scope for observers"));
+	}
+
+	return Super::GetDescription();
+}
+
+FText UBehaviorTreeGraphNode_Composite::GetTooltipText() const
+{
+	const UBTCompositeNode* CompositeNode = Cast<UBTCompositeNode>(NodeInstance);
+	if (CompositeNode && CompositeNode->IsApplyingDecoratorScope())
+	{
+		return FText::Format(FText::FromString(TEXT("{0}\n\n{1}")), 
+			Super::GetDescription(), 
+			NSLOCTEXT("BehaviorTreeEditor", "CompositeNodeScopeTooltip", "This node is a local scope for decorators.\nAll observing decorators (Lower Priority or Both) will be removed when execution flow leaves this branch."));
+	}
+
+	return Super::GetDescription();
+}
+
 void UBehaviorTreeGraphNode_Composite::GetContextMenuActions(const FGraphNodeContextMenuBuilder& Context) const
 {
 	AddContextMenuActionsDecorators(Context);
diff --git a/Engine/Source/Editor/BehaviorTreeEditor/Private/BehaviorTreeGraphNode_CompositeDecorator.cpp b/Engine/Source/Editor/BehaviorTreeEditor/Private/BehaviorTreeGraphNode_CompositeDecorator.cpp
--- a/Engine/Source/Editor/BehaviorTreeEditor/Private/BehaviorTreeGraphNode_CompositeDecorator.cpp
+++ b/Engine/Source/Editor/BehaviorTreeEditor/Private/BehaviorTreeGraphNode_CompositeDecorator.cpp
@@ -99,6 +99,23 @@ void UBehaviorTreeGraphNode_CompositeDecorator::PostCopyNode()
 	}
 }
 
+void UBehaviorTreeGraphNode_CompositeDecorator::ResetNodeOwner()
+{
+	Super::ResetNodeOwner();
+
+	if (BoundGraph)
+	{
+		for (int32 i = 0; i < BoundGraph->Nodes.Num(); i++)
+		{
+			UBehaviorTreeDecoratorGraphNode_Decorator* Node = Cast<UBehaviorTreeDecoratorGraphNode_Decorator>(BoundGraph->Nodes[i]);
+			if (Node)
+			{
+				Node->ResetNodeOwner();
+			}
+		}
+	}
+}
+
 bool UBehaviorTreeGraphNode_CompositeDecorator::RefreshNodeClass()
 {
 	bool bUpdated = false;
diff --git a/Engine/Source/Editor/BehaviorTreeEditor/Private/BehaviorTreeGraphNode_Root.cpp b/Engine/Source/Editor/BehaviorTreeEditor/Private/BehaviorTreeGraphNode_Root.cpp
--- a/Engine/Source/Editor/BehaviorTreeEditor/Private/BehaviorTreeGraphNode_Root.cpp
+++ b/Engine/Source/Editor/BehaviorTreeEditor/Private/BehaviorTreeGraphNode_Root.cpp
@@ -31,7 +31,7 @@ void UBehaviorTreeGraphNode_Root::PostPlacedNewNode()
 
 void UBehaviorTreeGraphNode_Root::AllocateDefaultPins()
 {
-	CreatePin(EGPD_Output, UBehaviorTreeEditorTypes::PinCategory_SingleComposite, FString(), nullptr, TEXT("In"));
+	CreatePin(EGPD_Output, UBehaviorTreeEditorTypes::PinCategory_SingleComposite, TEXT("In"));
 }
 
 FText UBehaviorTreeGraphNode_Root::GetNodeTitle(ENodeTitleType::Type TitleType) const
diff --git a/Engine/Source/Editor/BehaviorTreeEditor/Private/BehaviorTreeGraphNode_SimpleParallel.cpp b/Engine/Source/Editor/BehaviorTreeEditor/Private/BehaviorTreeGraphNode_SimpleParallel.cpp
--- a/Engine/Source/Editor/BehaviorTreeEditor/Private/BehaviorTreeGraphNode_SimpleParallel.cpp
+++ b/Engine/Source/Editor/BehaviorTreeEditor/Private/BehaviorTreeGraphNode_SimpleParallel.cpp
@@ -11,10 +11,10 @@ UBehaviorTreeGraphNode_SimpleParallel::UBehaviorTreeGraphNode_SimpleParallel(con
 
 void UBehaviorTreeGraphNode_SimpleParallel::AllocateDefaultPins()
 {
-	CreatePin(EGPD_Input, UBehaviorTreeEditorTypes::PinCategory_MultipleNodes, FString(), nullptr, TEXT("In"));
+	CreatePin(EGPD_Input, UBehaviorTreeEditorTypes::PinCategory_MultipleNodes, TEXT("In"));
 	
-	CreatePin(EGPD_Output, UBehaviorTreeEditorTypes::PinCategory_SingleTask, FString(), nullptr, TEXT("Task"));
-	CreatePin(EGPD_Output, UBehaviorTreeEditorTypes::PinCategory_SingleNode, FString(), nullptr, TEXT("Out"));
+	CreatePin(EGPD_Output, UBehaviorTreeEditorTypes::PinCategory_SingleTask, TEXT("Task"));
+	CreatePin(EGPD_Output, UBehaviorTreeEditorTypes::PinCategory_SingleNode, TEXT("Out"));
 }
 
 void UBehaviorTreeGraphNode_SimpleParallel::GetPinHoverText(const UEdGraphPin& Pin, FString& HoverTextOut) const
diff --git a/Engine/Source/Editor/BehaviorTreeEditor/Private/BehaviorTreeGraphNode_Task.cpp b/Engine/Source/Editor/BehaviorTreeEditor/Private/BehaviorTreeGraphNode_Task.cpp
--- a/Engine/Source/Editor/BehaviorTreeEditor/Private/BehaviorTreeGraphNode_Task.cpp
+++ b/Engine/Source/Editor/BehaviorTreeEditor/Private/BehaviorTreeGraphNode_Task.cpp
@@ -10,7 +10,7 @@ UBehaviorTreeGraphNode_Task::UBehaviorTreeGraphNode_Task(const FObjectInitialize
 
 void UBehaviorTreeGraphNode_Task::AllocateDefaultPins()
 {
-	CreatePin(EGPD_Input, UBehaviorTreeEditorTypes::PinCategory_SingleComposite, FString(), nullptr, TEXT("In"));
+	CreatePin(EGPD_Input, UBehaviorTreeEditorTypes::PinCategory_SingleComposite, TEXT("In"));
 }
 
 FText UBehaviorTreeGraphNode_Task::GetNodeTitle(ENodeTitleType::Type TitleType) const
diff --git a/Engine/Source/Editor/BehaviorTreeEditor/Private/EdGraphSchema_BehaviorTree.cpp b/Engine/Source/Editor/BehaviorTreeEditor/Private/EdGraphSchema_BehaviorTree.cpp
--- a/Engine/Source/Editor/BehaviorTreeEditor/Private/EdGraphSchema_BehaviorTree.cpp
+++ b/Engine/Source/Editor/BehaviorTreeEditor/Private/EdGraphSchema_BehaviorTree.cpp
@@ -131,7 +131,7 @@ void UEdGraphSchema_BehaviorTree::GetSubNodeClasses(int32 SubNodeFlags, TArray<F
 
 void UEdGraphSchema_BehaviorTree::GetGraphContextActions(FGraphContextMenuBuilder& ContextMenuBuilder) const
 {
-	const FString PinCategory = ContextMenuBuilder.FromPin ?
+	const FName PinCategory = ContextMenuBuilder.FromPin ?
 		ContextMenuBuilder.FromPin->PinType.PinCategory : 
 		UBehaviorTreeEditorTypes::PinCategory_MultipleNodes;
 
diff --git a/Engine/Source/Editor/BehaviorTreeEditor/Private/SGraphNode_BehaviorTree.cpp b/Engine/Source/Editor/BehaviorTreeEditor/Private/SGraphNode_BehaviorTree.cpp
--- a/Engine/Source/Editor/BehaviorTreeEditor/Private/SGraphNode_BehaviorTree.cpp
+++ b/Engine/Source/Editor/BehaviorTreeEditor/Private/SGraphNode_BehaviorTree.cpp
@@ -8,6 +8,7 @@
 #include "Widgets/Layout/SBox.h"
 #include "Widgets/SToolTip.h"
 #include "BehaviorTree/BTNode.h"
+#include "BehaviorTree/BTCompositeNode.h"
 #include "BehaviorTreeGraph.h"
 #include "BehaviorTreeGraphNode.h"
 #include "BehaviorTreeGraphNode_Composite.h"
@@ -256,12 +257,15 @@ FSlateColor SGraphNode_BehaviorTree::GetBackgroundColor() const
 	else if (Cast<UBehaviorTreeGraphNode_Task>(GraphNode))
 	{
 		check(BTGraphNode);
-		const bool bIsSpecialTask = Cast<UBTTask_RunBehavior>(BTGraphNode->NodeInstance) != NULL;
+		const bool bIsSpecialTask = Cast<UBTTask_RunBehavior>(BTGraphNode->NodeInstance) != nullptr;
 		NodeColor = bIsSpecialTask ? BehaviorTreeColors::NodeBody::TaskSpecial : BehaviorTreeColors::NodeBody::Task;
 	}
 	else if (Cast<UBehaviorTreeGraphNode_Composite>(GraphNode))
 	{
-		NodeColor = BehaviorTreeColors::NodeBody::Composite;
+		check(BTGraphNode);
+		UBTCompositeNode* CompositeNodeInstance = Cast<UBTCompositeNode>(BTGraphNode->NodeInstance);
+		const bool bIsScoped = CompositeNodeInstance && CompositeNodeInstance->IsApplyingDecoratorScope();
+		NodeColor = bIsScoped ? BehaviorTreeColors::NodeBody::CompositeScoped : BehaviorTreeColors::NodeBody::Composite;
 	}
 	else if (Cast<UBehaviorTreeGraphNode_Service>(GraphNode))
 	{
diff --git a/Engine/Source/Editor/BlueprintGraph/Classes/EdGraphSchema_K2.h b/Engine/Source/Editor/BlueprintGraph/Classes/EdGraphSchema_K2.h
--- a/Engine/Source/Editor/BlueprintGraph/Classes/EdGraphSchema_K2.h
+++ b/Engine/Source/Editor/BlueprintGraph/Classes/EdGraphSchema_K2.h
@@ -276,52 +276,52 @@ class BLUEPRINTGRAPH_API UEdGraphSchema_K2 : public UEdGraphSchema
 	GENERATED_UCLASS_BODY()
 
 	// Allowable PinType.PinCategory values
-	static const FString PC_Exec;
-	static const FString PC_Boolean;
-	static const FString PC_Byte;
-	static const FString PC_Class;    // SubCategoryObject is the MetaClass of the Class passed thru this pin, or SubCategory can be 'self'. The DefaultValue string should always be empty, use DefaultObject.
-	static const FString PC_SoftClass;
-	static const FString PC_Int;
-	static const FString PC_Float;
-	static const FString PC_Name;
-	static const FString PC_Delegate;    // SubCategoryObject is the UFunction of the delegate signature
-	static const FString PC_MCDelegate;  // SubCategoryObject is the UFunction of the delegate signature
-	static const FString PC_Object;    // SubCategoryObject is the Class of the object passed thru this pin, or SubCategory can be 'self'. The DefaultValue string should always be empty, use DefaultObject.
-	static const FString PC_Interface;	// SubCategoryObject is the Class of the object passed thru this pin.
-	static const FString PC_SoftObject;		// SubCategoryObject is the Class of the SoftObjectPtr passed thru this pin.
-	static const FString PC_String;
-	static const FString PC_Text;
-	static const FString PC_Struct;    // SubCategoryObject is the ScriptStruct of the struct passed thru this pin, 'self' is not a valid SubCategory. DefaultObject should always be empty, the DefaultValue string may be used for supported structs.
-	static const FString PC_Wildcard;    // Special matching rules are imposed by the node itself
-	static const FString PC_Enum;    // SubCategoryObject is the UEnum object passed thru this pin
+	static const FName PC_Exec;
+	static const FName PC_Boolean;
+	static const FName PC_Byte;
+	static const FName PC_Class;    // SubCategoryObject is the MetaClass of the Class passed thru this pin, or SubCategory can be 'self'. The DefaultValue string should always be empty, use DefaultObject.
+	static const FName PC_SoftClass;
+	static const FName PC_Int;
+	static const FName PC_Float;
+	static const FName PC_Name;
+	static const FName PC_Delegate;    // SubCategoryObject is the UFunction of the delegate signature
+	static const FName PC_MCDelegate;  // SubCategoryObject is the UFunction of the delegate signature
+	static const FName PC_Object;    // SubCategoryObject is the Class of the object passed thru this pin, or SubCategory can be 'self'. The DefaultValue string should always be empty, use DefaultObject.
+	static const FName PC_Interface;	// SubCategoryObject is the Class of the object passed thru this pin.
+	static const FName PC_SoftObject;		// SubCategoryObject is the Class of the AssetPtr passed thru this pin.
+	static const FName PC_String;
+	static const FName PC_Text;
+	static const FName PC_Struct;    // SubCategoryObject is the ScriptStruct of the struct passed thru this pin, 'self' is not a valid SubCategory. DefaultObject should always be empty, the DefaultValue string may be used for supported structs.
+	static const FName PC_Wildcard;    // Special matching rules are imposed by the node itself
+	static const FName PC_Enum;    // SubCategoryObject is the UEnum object passed thru this pin
 
 	// Common PinType.PinSubCategory values
-	static const FString PSC_Self;    // Category=PC_Object or PC_Class, indicates the class being compiled
+	static const FName PSC_Self;    // Category=PC_Object or PC_Class, indicates the class being compiled
 
-	static const FString PSC_Index;	// Category=PC_Wildcard, indicates the wildcard will only accept Int, Bool, Byte and Enum pins (used when a pin represents indexing a list)
-	static const FString PSC_Bitmask;	// Category=PC_Byte or PC_Int, indicates that the pin represents a bitmask field. SubCategoryObject is either NULL or the UEnum object to which the bitmap is linked for bitflag name specification.
+	static const FName PSC_Index;	// Category=PC_Wildcard, indicates the wildcard will only accept Int, Bool, Byte and Enum pins (used when a pin represents indexing a list)
+	static const FName PSC_Bitmask;	// Category=PC_Byte or PC_Int, indicates that the pin represents a bitmask field. SubCategoryObject is either NULL or the UEnum object to which the bitmap is linked for bitflag name specification.
 
 	// Pin names that have special meaning and required types in some contexts (depending on the node type)
-	static const FString PN_Execute;    // Category=PC_Exec, singleton, input
-	static const FString PN_Then;    // Category=PC_Exec, singleton, output
-	static const FString PN_Completed;    // Category=PC_Exec, singleton, output
-	static const FString PN_DelegateEntry;    // Category=PC_Exec, singleton, output; entry point for a dynamically bound delegate
-	static const FString PN_EntryPoint;	// entry point to the ubergraph
-	static const FString PN_Self;    // Category=PC_Object, singleton, input
-	static const FString PN_Else;    // Category=PC_Exec, singleton, output
-	static const FString PN_Loop;    // Category=PC_Exec, singleton, output
-	static const FString PN_After;    // Category=PC_Exec, singleton, output
-	static const FString PN_ReturnValue;		// Category=PC_Object, singleton, output
-	static const FString PN_ObjectToCast;    // Category=PC_Object, singleton, input
-	static const FString PN_Condition;    // Category=PC_Boolean, singleton, input
-	static const FString PN_Start;    // Category=PC_Int, singleton, input
-	static const FString PN_Stop;    // Category=PC_Int, singleton, input
-	static const FString PN_Index;    // Category=PC_Int, singleton, output
-	static const FString PN_Item;    // Category=PC_Int, singleton, output
-	static const FString PN_CastSucceeded;    // Category=PC_Exec, singleton, output
-	static const FString PN_CastFailed;    // Category=PC_Exec, singleton, output
+	static const FName PN_Execute;    // Category=PC_Exec, singleton, input
+	static const FName PN_Then;    // Category=PC_Exec, singleton, output
+	static const FName PN_Completed;    // Category=PC_Exec, singleton, output
+	static const FName PN_DelegateEntry;    // Category=PC_Exec, singleton, output; entry point for a dynamically bound delegate
+	static const FName PN_EntryPoint;	// entry point to the ubergraph
+	static const FName PN_Self;    // Category=PC_Object, singleton, input
+	static const FName PN_Else;    // Category=PC_Exec, singleton, output
+	static const FName PN_Loop;    // Category=PC_Exec, singleton, output
+	static const FName PN_After;    // Category=PC_Exec, singleton, output
+	static const FName PN_ReturnValue;		// Category=PC_Object, singleton, output
+	static const FName PN_ObjectToCast;    // Category=PC_Object, singleton, input
+	static const FName PN_Condition;    // Category=PC_Boolean, singleton, input
+	static const FName PN_Start;    // Category=PC_Int, singleton, input
+	static const FName PN_Stop;    // Category=PC_Int, singleton, input
+	static const FName PN_Index;    // Category=PC_Int, singleton, output
+	static const FName PN_Item;    // Category=PC_Int, singleton, output
+	static const FName PN_CastSucceeded;    // Category=PC_Exec, singleton, output
+	static const FName PN_CastFailed;    // Category=PC_Exec, singleton, output
 	static const FString PN_CastedValuePrefix;    // Category=PC_Object, singleton, output; actual pin name varies depending on the type to be casted to, this is just a prefix
-	static const FString PN_MatineeFinished;    // Category=PC_Exec, singleton, output
+	static const FName PN_MatineeFinished;    // Category=PC_Exec, singleton, output
 
 	// construction script function names
 	static const FName FN_UserConstructionScript;
@@ -349,7 +349,7 @@ class BLUEPRINTGRAPH_API UEdGraphSchema_K2 : public UEdGraphSchema
 	static int32 CurrentCacheRefreshID;
 
 	// Pin Selector category for all object types
-	static const FString AllObjectTypes;
+	static const FName AllObjectTypes;
 
 	UPROPERTY(globalconfig)
 	TArray<FBlueprintCallableFunctionRedirect> EditoronlyBPFunctionRedirects;
@@ -386,14 +386,14 @@ public:
 
 	public:
 		const FEdGraphPinType& GetPinType(bool bForceLoadedSubCategoryObject);
-		void SetPinSubTypeCategory(const FString& SubCategory)
+		void SetPinSubTypeCategory(const FName SubCategory)
 		{
 			PinType.PinSubCategory = SubCategory;
 		}
 
-		FPinTypeTreeInfo(const FText& InFriendlyName, const FString& CategoryName, const UEdGraphSchema_K2* Schema, const FText& InTooltip, bool bInReadOnly = false, FTypesDatabase* TypesDatabase = nullptr);
-		FPinTypeTreeInfo(const FString& CategoryName, UObject* SubCategoryObject, const FText& InTooltip, bool bInReadOnly = false, uint8 InPossibleObjectReferenceTypes = 0);
-		FPinTypeTreeInfo(const FText& InFriendlyName, const FString& CategoryName, const FSoftObjectPath& SubCategoryObject, const FText& InTooltip, bool bInReadOnly = false, uint8 InPossibleObjectReferenceTypes = 0);
+		FPinTypeTreeInfo(const FText& InFriendlyName, const FName CategoryName, const UEdGraphSchema_K2* Schema, const FText& InTooltip, bool bInReadOnly = false, FTypesDatabase* TypesDatabase = nullptr);
+		FPinTypeTreeInfo(const FName CategoryName, UObject* SubCategoryObject, const FText& InTooltip, bool bInReadOnly = false, uint8 InPossibleObjectReferenceTypes = 0);
+		FPinTypeTreeInfo(const FText& InFriendlyName, const FName CategoryName, const FSoftObjectPath& SubCategoryObject, const FText& InTooltip, bool bInReadOnly = false, uint8 InPossibleObjectReferenceTypes = 0);
 
 		FPinTypeTreeInfo(TSharedPtr<FPinTypeTreeInfo> InInfo)
 		{
@@ -436,7 +436,7 @@ public:
 			, bReadOnly(false)
 		{}
 
-		void Init(const FText& FriendlyCategoryName, const FString& CategoryName, const UEdGraphSchema_K2* Schema, const FText& InTooltip, bool bInReadOnly, FTypesDatabase* TypesDatabase);
+		void Init(const FText& FriendlyCategoryName, const FName CategoryName, const UEdGraphSchema_K2* Schema, const FText& InTooltip, bool bInReadOnly, FTypesDatabase* TypesDatabase);
 
 		FText GenerateDescription();
 	};
@@ -468,7 +468,7 @@ public:
 	virtual bool IsTitleBarPin(const UEdGraphPin& Pin) const override;
 	virtual void BreakNodeLinks(UEdGraphNode& TargetNode) const override;
 	virtual void BreakPinLinks(UEdGraphPin& TargetPin, bool bSendsNodeNotifcation) const override;
-	virtual void BreakSinglePinLink(UEdGraphPin* SourcePin, UEdGraphPin* TargetPin) override;
+	virtual void BreakSinglePinLink(UEdGraphPin* SourcePin, UEdGraphPin* TargetPin) const override;
 	virtual void ReconstructNode(UEdGraphNode& TargetNode, bool bIsBatchRequest=false) const override;
 	virtual bool CanEncapuslateNode(UEdGraphNode const& TestNode) const override;
 	virtual void HandleGraphBeingDeleted(UEdGraph& GraphBeingRemoved) const override;
@@ -490,7 +490,7 @@ public:
 	virtual void SplitPin(UEdGraphPin* Pin, bool bNotify = true) const override;
 	virtual void RecombinePin(UEdGraphPin* Pin) const override;
 	virtual void OnPinConnectionDoubleCicked(UEdGraphPin* PinA, UEdGraphPin* PinB, const FVector2D& GraphPosition) const override;
-	virtual UEdGraphPin* DropPinOnNode(UEdGraphNode* InTargetNode, const FString& InSourcePinName, const FEdGraphPinType& InSourcePinType, EEdGraphPinDirection InSourcePinDirection) const override;
+	virtual UEdGraphPin* DropPinOnNode(UEdGraphNode* InTargetNode, const FName& InSourcePinName, const FEdGraphPinType& InSourcePinType, EEdGraphPinDirection InSourcePinDirection) const override;
 	virtual bool SupportsDropPinOnNode(UEdGraphNode* InTargetNode, const FEdGraphPinType& InSourcePinType, EEdGraphPinDirection InSourcePinDirection, FText& OutErrorMessage) const override;
 	virtual bool IsCacheVisualizationOutOfDate(int32 InVisualizationCacheID) const override;
 	virtual int32 GetCurrentVisualizationCacheID() const override;
@@ -552,7 +552,7 @@ public:
 	virtual void GetPinDefaultValuesFromString(const FEdGraphPinType& PinType, UObject* OwningObject, const FString& NewValue, FString& UseDefaultValue, UObject*& UseDefaultObject, FText& UseDefaultText) const;
 
 	/** Do validation, that doesn't require a knowledge about actual pin */
-	virtual bool DefaultValueSimpleValidation(const FEdGraphPinType& PinType, const FString& PinName, const FString& NewDefaultValue, UObject* NewDefaultObject, const FText& InText, FString* OutMsg = nullptr) const;
+	virtual bool DefaultValueSimpleValidation(const FEdGraphPinType& PinType, const FName PinName, const FString& NewDefaultValue, UObject* NewDefaultObject, const FText& InText, FString* OutMsg = nullptr) const;
 
 	/** Returns true if the owning node is a function with AutoCreateRefTerm meta data */
 	bool IsAutoCreateRefTerm(const UEdGraphPin* Pin) const;
@@ -624,7 +624,7 @@ public:
 	}
 
 	/** Is given string a delegate category name ? */
-	virtual bool IsDelegateCategory(const FString& Category) const override;
+	virtual bool IsDelegateCategory(const FName Category) const override;
 
 	/** Returns whether a pin category is compatible with an Index Wildcard (PC_Wildcard and PSC_Index) */
 	inline bool IsIndexWildcardCompatible(const FEdGraphPinType& PinType) const
@@ -685,7 +685,7 @@ public:
 	}
 
 	/** Can Pin be promoted to a variable? */
-	bool CanPromotePinToVariable (const UEdGraphPin& Pin) const;
+	bool CanPromotePinToVariable (const UEdGraphPin& Pin, bool bInToMemberVariable) const;
 
 	/** Can Pin be split in to its component elements */
 	bool CanSplitStructPin(const UEdGraphPin& Pin) const;
@@ -698,7 +698,7 @@ public:
 	 * 
 	 * @return	true on success, false if the property is unsupported or invalid.
 	 */
-	static bool GetPropertyCategoryInfo(const UProperty* TestProperty, FString& OutCategory, FString& OutSubCategory, UObject*& OutSubCategoryObject, bool& bOutIsWeakPointer);
+	static bool GetPropertyCategoryInfo(const UProperty* TestProperty, FName& OutCategory, FName& OutSubCategory, UObject*& OutSubCategoryObject, bool& bOutIsWeakPointer);
 
 	/**
 	 * Convert the type of a UProperty to the corresponding pin type.
@@ -863,7 +863,7 @@ public:
 	*
 	* @return	The converted type text.
 	*/
-	static FText TerminalTypeToText(const FString& Category, const FString& SubCategory, UObject* SubCategoryObject, bool bIsWeakPtr);
+	static FText TerminalTypeToText(const FName Category, const FName SubCategory, UObject* SubCategoryObject, bool bIsWeakPtr);
 
 	/**
 	 * Converts a pin type into a fully qualified FText (e.g., object'ObjectName').
@@ -882,7 +882,7 @@ public:
 	 *
 	 * @return	The text to display for the category.
 	 */
-	static FText GetCategoryText(const FString& Category, const bool bForMenu = false);
+	static FText GetCategoryText(const FName Category, const bool bForMenu = false);
 
 	/**
 	 * Get the type tree for all of the property types valid for this schema
@@ -897,7 +897,7 @@ public:
 	 *
 	 * @param	Type	The type to check for subtypes
 	 */
-	bool DoesTypeHaveSubtypes( const FString& FriendlyTypeName ) const;
+	bool DoesTypeHaveSubtypes( const FName Category ) const;
 
 	/**
 	 * Returns true if the types and directions of two pins are schema compatible. Handles
@@ -1052,10 +1052,10 @@ public:
 	void LinkDataPinFromOutputToInput(UEdGraphNode* InOutputNode, UEdGraphNode* InInputNode) const;
 
 	/** Moves all connections from the old node to the new one. Returns true and destroys OldNode on success. Fails if it cannot find a mapping from an old pin. */
-	bool ReplaceOldNodeWithNew(UK2Node* OldNode, UK2Node* NewNode, const TMap<FString, FString>& OldPinToNewPinMap) const;
+	bool ReplaceOldNodeWithNew(UK2Node* OldNode, UK2Node* NewNode, const TMap<FName, FName>& OldPinToNewPinMap) const;
 
 	/** Convert a deprecated node into a function call node, called from per-node ConvertDeprecatedNode */
-	UK2Node* ConvertDeprecatedNodeToFunctionCall(UK2Node* OldNode, UFunction* NewFunction, TMap<FString, FString>& OldPinToNewPinMap, UEdGraph* Graph) const;
+	UK2Node* ConvertDeprecatedNodeToFunctionCall(UK2Node* OldNode, UFunction* NewFunction, TMap<FName, FName>& OldPinToNewPinMap, UEdGraph* Graph) const;
 
 	/** some inherited schemas don't want anim-notify actions listed, so this is an easy way to check that */
 	virtual bool DoesSupportAnimNotifyActions() const { return true; }
@@ -1070,7 +1070,7 @@ public:
 	static void OnCreateNonExistentLocalVariable(class UK2Node_Variable* Variable, UBlueprint* OwnerBlueprint);
 
 	/** Replace the variable that a variable node refers to when the variable it refers to does not exist */
-	static void OnReplaceVariableForVariableNode(class UK2Node_Variable* Variable, UBlueprint* OwnerBlueprint, FString VariableName, bool bIsSelfMember);
+	static void OnReplaceVariableForVariableNode(class UK2Node_Variable* Variable, UBlueprint* OwnerBlueprint, FName VariableName, bool bIsSelfMember);
 
 	/** Create sub menu that shows all possible variables that can be used to replace the existing variable reference */
 	static void GetReplaceVariableMenu(class FMenuBuilder& MenuBuilder, class UK2Node_Variable* Variable, UBlueprint* OwnerBlueprint, bool bReplaceExistingVariable = false);
diff --git a/Engine/Source/Editor/BlueprintGraph/Classes/K2Node.h b/Engine/Source/Editor/BlueprintGraph/Classes/K2Node.h
--- a/Engine/Source/Editor/BlueprintGraph/Classes/K2Node.h
+++ b/Engine/Source/Editor/BlueprintGraph/Classes/K2Node.h
@@ -308,7 +308,7 @@ class UK2Node : public UEdGraphNode
 	BLUEPRINTGRAPH_API virtual bool CanCreateUnderSpecifiedSchema(const UEdGraphSchema* DesiredSchema) const override;
 
 	// Renames an existing pin on the node.
-	BLUEPRINTGRAPH_API virtual ERenamePinResult RenameUserDefinedPin(const FString& OldName, const FString& NewName, bool bTest = false);
+	BLUEPRINTGRAPH_API virtual ERenamePinResult RenameUserDefinedPin(const FName OldName, const FName NewName, bool bTest = false);
 
 	// Returns which dynamic binding class (if any) to use for this node
 	BLUEPRINTGRAPH_API virtual UClass* GetDynamicBindingClass() const { return NULL; }
@@ -420,7 +420,7 @@ protected:
 	 * 
 	 * returns the redirect type
 	 */
-	BLUEPRINTGRAPH_API ERedirectType ShouldRedirectParam(const TArray<FString>& OldPinNames, FName& NewPinName, const UK2Node * NewPinNode) const;
+	BLUEPRINTGRAPH_API ERedirectType ShouldRedirectParam(const TArray<FString>& OldPinNames, FName& NewPinName, const UK2Node* NewPinNode) const;
 
 	// Helper function to restore Split Pins after ReallocatePinsDuringReconstruction, call after recreating all pins to restore split pin state
 	BLUEPRINTGRAPH_API void RestoreSplitPins(TArray<UEdGraphPin*>& OldPins);
diff --git a/Engine/Source/Editor/BlueprintGraph/Classes/K2Node_AssignmentStatement.h b/Engine/Source/Editor/BlueprintGraph/Classes/K2Node_AssignmentStatement.h
--- a/Engine/Source/Editor/BlueprintGraph/Classes/K2Node_AssignmentStatement.h
+++ b/Engine/Source/Editor/BlueprintGraph/Classes/K2Node_AssignmentStatement.h
@@ -18,9 +18,9 @@ class UK2Node_AssignmentStatement : public UK2Node
 
 
 	// Name of the Variable pin for this node
-	static FString VariablePinName;
+	static FName VariablePinName;
 	// Name of the Value pin for this node
-	static FString ValuePinName;
+	static FName ValuePinName;
 
 	//~ Begin UEdGraphNode Interface
 	virtual void AllocateDefaultPins() override;
diff --git a/Engine/Source/Editor/BlueprintGraph/Classes/K2Node_BaseAsyncTask.h b/Engine/Source/Editor/BlueprintGraph/Classes/K2Node_BaseAsyncTask.h
--- a/Engine/Source/Editor/BlueprintGraph/Classes/K2Node_BaseAsyncTask.h
+++ b/Engine/Source/Editor/BlueprintGraph/Classes/K2Node_BaseAsyncTask.h
@@ -25,7 +25,7 @@ class UK2Node_TemporaryVariable;
 
 struct FAsyncTaskPinRedirectMapInfo
 {
-	TMap<FString, TArray<UClass*> > OldPinToProxyClassMap;
+	TMap<FName, TArray<UClass*> > OldPinToProxyClassMap;
 };
 
 /** !!! The proxy object should have RF_StrongRefOnFrame flag. !!! */
@@ -84,7 +84,7 @@ protected:
 			FOutputPinAndLocalVariable(UEdGraphPin* Pin, UK2Node_TemporaryVariable* Var) : OutputPin(Pin), TempVar(Var) {}
 		};
 
-		static bool ValidDataPin(const UEdGraphPin* Pin, EEdGraphPinDirection Direction, const UEdGraphSchema_K2* Schema);
+		static bool ValidDataPin(const UEdGraphPin* Pin, EEdGraphPinDirection Direction);
 		static bool CreateDelegateForNewFunction(UEdGraphPin* DelegateInputPin, FName FunctionName, UK2Node* CurrentNode, UEdGraph* SourceGraph, FKismetCompilerContext& CompilerContext);
 		static bool CopyEventSignature(UK2Node_CustomEvent* CENode, UFunction* Function, const UEdGraphSchema_K2* Schema);
 		static bool HandleDelegateImplementation(
@@ -92,10 +92,10 @@ protected:
 			UEdGraphPin* ProxyObjectPin, UEdGraphPin*& InOutLastThenPin,
 			UK2Node* CurrentNode, UEdGraph* SourceGraph, FKismetCompilerContext& CompilerContext);
 
-		static const FString& GetAsyncTaskProxyName();
+		static const FName GetAsyncTaskProxyName();
 	};
 
 	// Pin Redirector support
-	static TMap<FString, FAsyncTaskPinRedirectMapInfo> AsyncTaskPinRedirectMap;
+	static TMap<FName, FAsyncTaskPinRedirectMapInfo> AsyncTaskPinRedirectMap;
 	static bool bAsyncTaskPinRedirectMapInitialized;
 };
diff --git a/Engine/Source/Editor/BlueprintGraph/Classes/K2Node_BitmaskLiteral.h b/Engine/Source/Editor/BlueprintGraph/Classes/K2Node_BitmaskLiteral.h
--- a/Engine/Source/Editor/BlueprintGraph/Classes/K2Node_BitmaskLiteral.h
+++ b/Engine/Source/Editor/BlueprintGraph/Classes/K2Node_BitmaskLiteral.h
@@ -42,7 +42,7 @@ class UK2Node_BitmaskLiteral : public UK2Node, public INodeDependingOnEnumInterf
 	virtual bool ShouldBeReconstructedAfterEnumChanged() const override { return true; }
 	//~ End INodeDependingOnEnumInterface
 
-	BLUEPRINTGRAPH_API static const FString& GetBitmaskInputPinName();
+	BLUEPRINTGRAPH_API static const FName& GetBitmaskInputPinName();
 
 protected:
 	/** Internal helper method used to validate the current enum type */
diff --git a/Engine/Source/Editor/BlueprintGraph/Classes/K2Node_CallFunction.h b/Engine/Source/Editor/BlueprintGraph/Classes/K2Node_CallFunction.h
--- a/Engine/Source/Editor/BlueprintGraph/Classes/K2Node_CallFunction.h
+++ b/Engine/Source/Editor/BlueprintGraph/Classes/K2Node_CallFunction.h
@@ -188,7 +188,7 @@ public:
 	UEdGraph* GetFunctionGraph(const UEdGraphNode*& OutGraphNode) const;
 
 	/** Checks if the property is marked as "CustomStructureParam" */
-	static bool IsStructureWildcardProperty(const UFunction* InFunction, const FString& PropertyName);
+	static bool IsStructureWildcardProperty(const UFunction* InFunction, const FName PropertyName);
 
 	/** returns true if InProperty should be treated as a wildcard (e.g. due to SetParam markup) */
 	static bool IsWildcardProperty(const UFunction* InFunction, const UProperty* InProperty);
diff --git a/Engine/Source/Editor/BlueprintGraph/Classes/K2Node_CastByteToEnum.h b/Engine/Source/Editor/BlueprintGraph/Classes/K2Node_CastByteToEnum.h
--- a/Engine/Source/Editor/BlueprintGraph/Classes/K2Node_CastByteToEnum.h
+++ b/Engine/Source/Editor/BlueprintGraph/Classes/K2Node_CastByteToEnum.h
@@ -26,7 +26,7 @@ class UK2Node_CastByteToEnum : public UK2Node
 	UPROPERTY()
 	bool bSafe;
 
-	static const FString ByteInputPinName;
+	static const FName ByteInputPinName;
 
 	//~ Begin UEdGraphNode Interface
 	virtual void AllocateDefaultPins() override;
diff --git a/Engine/Source/Editor/BlueprintGraph/Classes/K2Node_CommutativeAssociativeBinaryOperator.h b/Engine/Source/Editor/BlueprintGraph/Classes/K2Node_CommutativeAssociativeBinaryOperator.h
--- a/Engine/Source/Editor/BlueprintGraph/Classes/K2Node_CommutativeAssociativeBinaryOperator.h
+++ b/Engine/Source/Editor/BlueprintGraph/Classes/K2Node_CommutativeAssociativeBinaryOperator.h
@@ -24,7 +24,7 @@ private:
 	const static int32 BinaryOperatorInputsNum = 2;
 
 	static int32 GetMaxInputPinsNum();
-	static FString GetNameForPin(int32 PinIndex);
+	static FName GetNameForPin(int32 PinIndex);
 
 	FEdGraphPinType GetType() const;
 
diff --git a/Engine/Source/Editor/BlueprintGraph/Classes/K2Node_EaseFunction.h b/Engine/Source/Editor/BlueprintGraph/Classes/K2Node_EaseFunction.h
--- a/Engine/Source/Editor/BlueprintGraph/Classes/K2Node_EaseFunction.h
+++ b/Engine/Source/Editor/BlueprintGraph/Classes/K2Node_EaseFunction.h
@@ -22,11 +22,11 @@ struct FCustomPin
 
 	/** Name of the pin within the current node */
 	UPROPERTY()
-	FString PinName;
+	FName PinName;
 
 	/** Name of the pin of the call function */
 	UPROPERTY()
-	FString CallFuncPinName;
+	FName CallFuncPinName;
 
 	/** If TRUE this is a custom pin that does come from the type of wildcard pin connected */
 	UPROPERTY()
@@ -99,7 +99,7 @@ private:
 
 	/** Name of the kismet ease function to be called */
 	UPROPERTY()
-	FString EaseFunctionName;
+	FName EaseFunctionName;
 
 	/** Tooltip text for this node. */
 	FText NodeTooltip;
diff --git a/Engine/Source/Editor/BlueprintGraph/Classes/K2Node_EditablePinBase.h b/Engine/Source/Editor/BlueprintGraph/Classes/K2Node_EditablePinBase.h
--- a/Engine/Source/Editor/BlueprintGraph/Classes/K2Node_EditablePinBase.h
+++ b/Engine/Source/Editor/BlueprintGraph/Classes/K2Node_EditablePinBase.h
@@ -19,7 +19,7 @@ struct FUserPinInfo
 
 	/** The name of the pin, as defined by the user */
 	UPROPERTY()
-	FString PinName;
+	FName PinName;
 
 	/** Type info for the pin */
 	UPROPERTY()
@@ -120,7 +120,7 @@ class UK2Node_EditablePinBase : public UK2Node
 	 * @param	InPinType				The type info for the pin to create
 	 * @param	InDesiredDirection		Desired direction of the pin, will auto-correct if the direction is not allowed on the pin.
 	 */
-	BLUEPRINTGRAPH_API UEdGraphPin* CreateUserDefinedPin(const FString& InPinName, const FEdGraphPinType& InPinType, EEdGraphPinDirection InDesiredDirection, bool bUseUniqueName = true);
+	BLUEPRINTGRAPH_API UEdGraphPin* CreateUserDefinedPin(const FName InPinName, const FEdGraphPinType& InPinType, EEdGraphPinDirection InDesiredDirection, bool bUseUniqueName = true);
 
 	/**
 	 * Removes a pin from the user-defined array, and removes the pin with the same name from the Pins array
@@ -134,7 +134,7 @@ class UK2Node_EditablePinBase : public UK2Node
 	 *
 	 * @param	PinName name of pin to remove
 	 */
-	BLUEPRINTGRAPH_API void RemoveUserDefinedPinByName(const FString& PinName);
+	BLUEPRINTGRAPH_API void RemoveUserDefinedPinByName(const FName PinName);
 
 	/**
 	 * Creates a new pin on the node from the specified user pin info.
@@ -142,7 +142,7 @@ class UK2Node_EditablePinBase : public UK2Node
 	 * 
 	 * @param	NewPinInfo		Shared pointer to the struct containing the info for this pin
 	 */
-	virtual UEdGraphPin* CreatePinFromUserDefinition(const TSharedPtr<FUserPinInfo> NewPinInfo) { return NULL; }
+	virtual UEdGraphPin* CreatePinFromUserDefinition(const TSharedPtr<FUserPinInfo> NewPinInfo) { return nullptr; }
 
 	// Modifies the default value of an existing pin on the node.
 	BLUEPRINTGRAPH_API virtual bool ModifyUserDefinedPinDefaultValue(TSharedPtr<FUserPinInfo> PinInfo, const FString& NewDefaultValue);
diff --git a/Engine/Source/Editor/BlueprintGraph/Classes/K2Node_EnumLiteral.h b/Engine/Source/Editor/BlueprintGraph/Classes/K2Node_EnumLiteral.h
--- a/Engine/Source/Editor/BlueprintGraph/Classes/K2Node_EnumLiteral.h
+++ b/Engine/Source/Editor/BlueprintGraph/Classes/K2Node_EnumLiteral.h
@@ -22,7 +22,7 @@ class UK2Node_EnumLiteral : public UK2Node, public INodeDependingOnEnumInterface
 	UPROPERTY()
 	UEnum* Enum;
 
-	static BLUEPRINTGRAPH_API const FString& GetEnumInputPinName();
+	static BLUEPRINTGRAPH_API const FName GetEnumInputPinName();
 
 	//~ Begin UEdGraphNode Interface
 	virtual void AllocateDefaultPins() override;
diff --git a/Engine/Source/Editor/BlueprintGraph/Classes/K2Node_Event.h b/Engine/Source/Editor/BlueprintGraph/Classes/K2Node_Event.h
--- a/Engine/Source/Editor/BlueprintGraph/Classes/K2Node_Event.h
+++ b/Engine/Source/Editor/BlueprintGraph/Classes/K2Node_Event.h
@@ -20,7 +20,7 @@ UCLASS(MinimalAPI)
 class UK2Node_Event : public UK2Node_EditablePinBase
 {
 	GENERATED_UCLASS_BODY()
-	BLUEPRINTGRAPH_API static const FString DelegateOutputName;
+	BLUEPRINTGRAPH_API static const FName DelegateOutputName;
 
 	/** Name of function signature that this event implements */
 	UPROPERTY()
diff --git a/Engine/Source/Editor/BlueprintGraph/Classes/K2Node_ExecutionSequence.h b/Engine/Source/Editor/BlueprintGraph/Classes/K2Node_ExecutionSequence.h
--- a/Engine/Source/Editor/BlueprintGraph/Classes/K2Node_ExecutionSequence.h
+++ b/Engine/Source/Editor/BlueprintGraph/Classes/K2Node_ExecutionSequence.h
@@ -41,7 +41,7 @@ class UK2Node_ExecutionSequence : public UK2Node, public IK2Node_AddPinInterface
 	//~ Begin K2Node_ExecutionSequence Interface
 
 	/** Gets a unique pin name, the next in the sequence */
-	FString GetUniquePinName();
+	FName GetUniquePinName();
 
 	/**
 	 * Adds a new execution pin to an execution node
@@ -64,6 +64,6 @@ class UK2Node_ExecutionSequence : public UK2Node, public IK2Node_AddPinInterface
 
 private:
 	// Returns the exec output pin name for a given 0-based index
-	virtual FString GetPinNameGivenIndex(int32 Index) const;
+	virtual FName GetPinNameGivenIndex(int32 Index) const;
 };
 
diff --git a/Engine/Source/Editor/BlueprintGraph/Classes/K2Node_ForEachElementInEnum.h b/Engine/Source/Editor/BlueprintGraph/Classes/K2Node_ForEachElementInEnum.h
--- a/Engine/Source/Editor/BlueprintGraph/Classes/K2Node_ForEachElementInEnum.h
+++ b/Engine/Source/Editor/BlueprintGraph/Classes/K2Node_ForEachElementInEnum.h
@@ -20,8 +20,8 @@ class UK2Node_ForEachElementInEnum : public UK2Node
 	UPROPERTY()
 	UEnum* Enum;
 
-	BLUEPRINTGRAPH_API static const FString InsideLoopPinName;
-	BLUEPRINTGRAPH_API static const FString EnumOuputPinName;
+	BLUEPRINTGRAPH_API static const FName InsideLoopPinName;
+	BLUEPRINTGRAPH_API static const FName EnumOuputPinName;
 
 	//~ Begin UEdGraphNode Interface
 	virtual void AllocateDefaultPins() override;
diff --git a/Engine/Source/Editor/BlueprintGraph/Classes/K2Node_FormatText.h b/Engine/Source/Editor/BlueprintGraph/Classes/K2Node_FormatText.h
--- a/Engine/Source/Editor/BlueprintGraph/Classes/K2Node_FormatText.h
+++ b/Engine/Source/Editor/BlueprintGraph/Classes/K2Node_FormatText.h
@@ -66,7 +66,7 @@ public:
 	 * @param InIndex		The argument's index to find the name for
 	 * @param InName		Name to set the argument to
 	 */
-	BLUEPRINTGRAPH_API void SetArgumentName(int32 InIndex, FString InName);
+	BLUEPRINTGRAPH_API void SetArgumentName(int32 InIndex, FName InName);
 
 	/** Swaps two arguments by index */
 	BLUEPRINTGRAPH_API void SwapArguments(int32 InIndexA, int32 InIndexB);
@@ -83,19 +83,19 @@ public:
 	 * @param InPinName		The pin name to check for
 	 * @return				NULL if the pin was not found, otherwise the found pin.
 	 */
-	BLUEPRINTGRAPH_API UEdGraphPin* FindArgumentPin(const FString& InPinName) const;
+	BLUEPRINTGRAPH_API UEdGraphPin* FindArgumentPin(const FName InPinName) const;
 
 private:
 	/** Synchronize the type of the given argument pin with the type its connected to, or reset it to a wildcard pin if there's no connection */
 	void SynchronizeArgumentPinType(UEdGraphPin* Pin);
 
 	/** Returns a unique pin name to use for a pin */
-	FString GetUniquePinName();
+	FName GetUniquePinName();
 
 private:
 	/** When adding arguments to the node, their names are placed here and are generated as pins during construction */
 	UPROPERTY()
-	TArray<FString> PinNames;
+	TArray<FName> PinNames;
 
 	/** The "Format" input pin, always available on the node */
 	UEdGraphPin* CachedFormatPin;
diff --git a/Engine/Source/Editor/BlueprintGraph/Classes/K2Node_FunctionTerminator.h b/Engine/Source/Editor/BlueprintGraph/Classes/K2Node_FunctionTerminator.h
--- a/Engine/Source/Editor/BlueprintGraph/Classes/K2Node_FunctionTerminator.h
+++ b/Engine/Source/Editor/BlueprintGraph/Classes/K2Node_FunctionTerminator.h
@@ -29,7 +29,7 @@ class UK2Node_FunctionTerminator : public UK2Node_EditablePinBase
 	//~ Begin UEdGraphNode Interface
 	virtual bool CanDuplicateNode() const override { return false; }
 	virtual FLinearColor GetNodeTitleColor() const override;
-	virtual FString CreateUniquePinName(FString SourcePinName) const override;
+	virtual FName CreateUniquePinName(FName SourcePinName) const override;
 	virtual void ValidateNodeDuringCompilation(class FCompilerResultsLog& MessageLog) const override;
 	//~ End UEdGraphNode Interface
 
diff --git a/Engine/Source/Editor/BlueprintGraph/Classes/K2Node_GetClassDefaults.h b/Engine/Source/Editor/BlueprintGraph/Classes/K2Node_GetClassDefaults.h
--- a/Engine/Source/Editor/BlueprintGraph/Classes/K2Node_GetClassDefaults.h
+++ b/Engine/Source/Editor/BlueprintGraph/Classes/K2Node_GetClassDefaults.h
@@ -84,7 +84,7 @@ protected:
 
 private:
 	/** Class pin name */
-	static FString ClassPinName;
+	static FName ClassPinName;
 
 	/** Blueprint that we subscribed OnBlueprintChangedDelegate and OnBlueprintCompiledDelegate to */
 	UPROPERTY()
diff --git a/Engine/Source/Editor/BlueprintGraph/Classes/K2Node_GetEnumeratorName.h b/Engine/Source/Editor/BlueprintGraph/Classes/K2Node_GetEnumeratorName.h
--- a/Engine/Source/Editor/BlueprintGraph/Classes/K2Node_GetEnumeratorName.h
+++ b/Engine/Source/Editor/BlueprintGraph/Classes/K2Node_GetEnumeratorName.h
@@ -17,7 +17,7 @@ class UK2Node_GetEnumeratorName : public UK2Node
 {
 	GENERATED_UCLASS_BODY()
 
-	static FString EnumeratorPinName;
+	static FName EnumeratorPinName;
 
 	//~ Begin UEdGraphNode Interface
 	virtual void AllocateDefaultPins() override;
diff --git a/Engine/Source/Editor/BlueprintGraph/Classes/K2Node_LoadAsset.h b/Engine/Source/Editor/BlueprintGraph/Classes/K2Node_LoadAsset.h
--- a/Engine/Source/Editor/BlueprintGraph/Classes/K2Node_LoadAsset.h
+++ b/Engine/Source/Editor/BlueprintGraph/Classes/K2Node_LoadAsset.h
@@ -36,11 +36,11 @@ public:
 protected:
 	virtual FName NativeFunctionName() const;
 
-	virtual const FString& GetInputCategory() const;
-	virtual const FString& GetOutputCategory() const;
+	virtual const FName& GetInputCategory() const;
+	virtual const FName& GetOutputCategory() const;
 
-	virtual const FString& GetInputPinName() const;
-	virtual const FString& GetOutputPinName() const;
+	virtual const FName& GetInputPinName() const;
+	virtual const FName& GetOutputPinName() const;
 };
 
 UCLASS(MinimalAPI)
@@ -56,9 +56,9 @@ public:
 protected:
 	virtual FName NativeFunctionName() const override;
 
-	virtual const FString& GetInputCategory() const override;
-	virtual const FString& GetOutputCategory() const override;
+	virtual const FName& GetInputCategory() const override;
+	virtual const FName& GetOutputCategory() const override;
 
-	virtual const FString& GetInputPinName() const override;
-	virtual const FString& GetOutputPinName() const override;
+	virtual const FName& GetInputPinName() const override;
+	virtual const FName& GetOutputPinName() const override;
 };
diff --git a/Engine/Source/Editor/BlueprintGraph/Classes/K2Node_MakeArray.h b/Engine/Source/Editor/BlueprintGraph/Classes/K2Node_MakeArray.h
--- a/Engine/Source/Editor/BlueprintGraph/Classes/K2Node_MakeArray.h
+++ b/Engine/Source/Editor/BlueprintGraph/Classes/K2Node_MakeArray.h
@@ -29,6 +29,6 @@ public:
 	// End of UK2Node interface
 
 	// UK2Node_MakeContainer interface
-	virtual FString GetOutputPinName() const override;
+	virtual FName GetOutputPinName() const override;
 	// UK2Node_MakeContainer interface
 };
diff --git a/Engine/Source/Editor/BlueprintGraph/Classes/K2Node_MakeContainer.h b/Engine/Source/Editor/BlueprintGraph/Classes/K2Node_MakeContainer.h
--- a/Engine/Source/Editor/BlueprintGraph/Classes/K2Node_MakeContainer.h
+++ b/Engine/Source/Editor/BlueprintGraph/Classes/K2Node_MakeContainer.h
@@ -45,8 +45,8 @@ public:
 	BLUEPRINTGRAPH_API UEdGraphPin* GetOutputPin() const;
 
 	/** returns a reference to the output array pin of this node, which is responsible for defining the type */
-	virtual FString GetOutputPinName() const PURE_VIRTUAL(UK2Node_MakeContainer::GetOutputPinName, return FString(););
-	virtual FString GetPinName(int32 PinIndex) const;
+	virtual FName GetOutputPinName() const PURE_VIRTUAL(UK2Node_MakeContainer::GetOutputPinName, return NAME_None;);
+	virtual FName GetPinName(int32 PinIndex) const;
 	virtual void GetKeyAndValuePins(TArray<UEdGraphPin*>& KeyPins, TArray<UEdGraphPin*>& ValuePins) const;
 
 public:
diff --git a/Engine/Source/Editor/BlueprintGraph/Classes/K2Node_MakeMap.h b/Engine/Source/Editor/BlueprintGraph/Classes/K2Node_MakeMap.h
--- a/Engine/Source/Editor/BlueprintGraph/Classes/K2Node_MakeMap.h
+++ b/Engine/Source/Editor/BlueprintGraph/Classes/K2Node_MakeMap.h
@@ -29,8 +29,8 @@ public:
 	// End of UK2Node interface
 
 	// UK2Node_MakeContainer interface
-	virtual FString GetOutputPinName() const override;
-	virtual FString GetPinName(int32 PinIndex) const override;
+	virtual FName GetOutputPinName() const override;
+	virtual FName GetPinName(int32 PinIndex) const override;
 	virtual void GetKeyAndValuePins(TArray<UEdGraphPin*>& KeyPins, TArray<UEdGraphPin*>& ValuePins) const override;
 	// UK2Node_MakeContainer interface
 
diff --git a/Engine/Source/Editor/BlueprintGraph/Classes/K2Node_MakeSet.h b/Engine/Source/Editor/BlueprintGraph/Classes/K2Node_MakeSet.h
--- a/Engine/Source/Editor/BlueprintGraph/Classes/K2Node_MakeSet.h
+++ b/Engine/Source/Editor/BlueprintGraph/Classes/K2Node_MakeSet.h
@@ -29,6 +29,6 @@ public:
 	// End of UK2Node interface
 
 	// UK2Node_MakeContainer interface
-	virtual FString GetOutputPinName() const override;
+	virtual FName GetOutputPinName() const override;
 	// UK2Node_MakeContainer interface
 };
diff --git a/Engine/Source/Editor/BlueprintGraph/Classes/K2Node_MultiGate.h b/Engine/Source/Editor/BlueprintGraph/Classes/K2Node_MultiGate.h
--- a/Engine/Source/Editor/BlueprintGraph/Classes/K2Node_MultiGate.h
+++ b/Engine/Source/Editor/BlueprintGraph/Classes/K2Node_MultiGate.h
@@ -60,6 +60,6 @@ class UK2Node_MultiGate : public UK2Node_ExecutionSequence
 
 private:
 	// Returns the exec output pin name for a given 0-based index
- 	virtual FString GetPinNameGivenIndex(int32 Index) const override;
+ 	virtual FName GetPinNameGivenIndex(int32 Index) const override;
 };
 
diff --git a/Engine/Source/Editor/BlueprintGraph/Classes/K2Node_PureAssignmentStatement.h b/Engine/Source/Editor/BlueprintGraph/Classes/K2Node_PureAssignmentStatement.h
--- a/Engine/Source/Editor/BlueprintGraph/Classes/K2Node_PureAssignmentStatement.h
+++ b/Engine/Source/Editor/BlueprintGraph/Classes/K2Node_PureAssignmentStatement.h
@@ -15,11 +15,11 @@ class UK2Node_PureAssignmentStatement : public UK2Node
 	GENERATED_UCLASS_BODY()
 
 	// Name of the Variable pin for this node
-	static FString VariablePinName;
+	static FName VariablePinName;
 	// Name of the Value pin for this node
-	static FString ValuePinName;
+	static FName ValuePinName;
 	// Name of the output pin for this node
-	static FString OutputPinName;
+	static FName OutputPinName;
 
 	//~ Begin UEdGraphNode Interface
 	virtual void AllocateDefaultPins() override;
