Copying //UE4/Dev-Framework to //UE4/Dev-Main (Source: //UE4/Dev-Framework @ 3716594)

diff --git a/Engine/Source/Editor/BlueprintGraph/Classes/K2Node_StructOperation.h b/Engine/Source/Editor/BlueprintGraph/Classes/K2Node_StructOperation.h
--- a/Engine/Source/Editor/BlueprintGraph/Classes/K2Node_StructOperation.h
+++ b/Engine/Source/Editor/BlueprintGraph/Classes/K2Node_StructOperation.h
@@ -19,7 +19,7 @@ class UK2Node_StructOperation : public UK2Node_Variable
 	UScriptStruct* StructType;
 
 	//~ Begin UEdGraphNode Interface
-	virtual FString GetPinMetaData(FString InPinName, FName InKey) override;
+	virtual FString GetPinMetaData(FName InPinName, FName InKey) override;
 	//~ End UEdGraphNode Interface
 
 	//~ Begin UK2Node Interface
diff --git a/Engine/Source/Editor/BlueprintGraph/Classes/K2Node_Switch.h b/Engine/Source/Editor/BlueprintGraph/Classes/K2Node_Switch.h
--- a/Engine/Source/Editor/BlueprintGraph/Classes/K2Node_Switch.h
+++ b/Engine/Source/Editor/BlueprintGraph/Classes/K2Node_Switch.h
@@ -48,7 +48,7 @@ class UK2Node_Switch : public UK2Node
 	// UK2Node_Switch interface
 
 	/** Gets a unique pin name, the next in the sequence */
-	virtual FString GetUniquePinName() { return FString(); }
+	virtual FName GetUniquePinName() { return NAME_None; }
 
 	/** Gets the pin type from the schema for the subclass */
 	virtual FEdGraphPinType GetPinType() const { check(false); return FEdGraphPinType(); }
@@ -75,9 +75,9 @@ class UK2Node_Switch : public UK2Node
 	BLUEPRINTGRAPH_API UEdGraphPin* GetDefaultPin() const;
 	BLUEPRINTGRAPH_API UEdGraphPin* GetFunctionPin() const;
 
-	BLUEPRINTGRAPH_API static FString	GetSelectionPinName();
+	BLUEPRINTGRAPH_API static FName GetSelectionPinName();
 
-	virtual FString GetPinNameGivenIndex(int32 Index);
+	virtual FName GetPinNameGivenIndex(int32 Index) const;
 
 protected:
 	virtual void CreateSelectionPin() {}
diff --git a/Engine/Source/Editor/BlueprintGraph/Classes/K2Node_SwitchEnum.h b/Engine/Source/Editor/BlueprintGraph/Classes/K2Node_SwitchEnum.h
--- a/Engine/Source/Editor/BlueprintGraph/Classes/K2Node_SwitchEnum.h
+++ b/Engine/Source/Editor/BlueprintGraph/Classes/K2Node_SwitchEnum.h
@@ -49,7 +49,6 @@ class UK2Node_SwitchEnum : public UK2Node_Switch, public INodeDependingOnEnumInt
 	// End of UK2Node interface
 
 	// UK2Node_Switch Interface
-	virtual FString GetUniquePinName() override;
 	virtual FEdGraphPinType GetPinType() const override;
 	virtual void AddPinToSwitchNode() override;
 	virtual void RemovePinFromSwitchNode(UEdGraphPin* TargetPin) override;
diff --git a/Engine/Source/Editor/BlueprintGraph/Classes/K2Node_SwitchInteger.h b/Engine/Source/Editor/BlueprintGraph/Classes/K2Node_SwitchInteger.h
--- a/Engine/Source/Editor/BlueprintGraph/Classes/K2Node_SwitchInteger.h
+++ b/Engine/Source/Editor/BlueprintGraph/Classes/K2Node_SwitchInteger.h
@@ -35,8 +35,8 @@ class UK2Node_SwitchInteger : public UK2Node_Switch
 	// End of UK2Node interface
 
 	// UK2Node_Switch Interface
-	virtual FString GetUniquePinName() override;
-	virtual FString GetPinNameGivenIndex(int32 Index) override;
+	virtual FName GetUniquePinName() override;
+	virtual FName GetPinNameGivenIndex(int32 Index) const override;
 	virtual void ReallocatePinsDuringReconstruction(TArray<UEdGraphPin*>& OldPins) override;
 	virtual FEdGraphPinType GetPinType() const override;
 	virtual bool CanRemoveExecutionPin(UEdGraphPin* TargetPin) const override;
diff --git a/Engine/Source/Editor/BlueprintGraph/Classes/K2Node_SwitchString.h b/Engine/Source/Editor/BlueprintGraph/Classes/K2Node_SwitchString.h
--- a/Engine/Source/Editor/BlueprintGraph/Classes/K2Node_SwitchString.h
+++ b/Engine/Source/Editor/BlueprintGraph/Classes/K2Node_SwitchString.h
@@ -17,7 +17,7 @@ class UK2Node_SwitchString : public UK2Node_Switch
 	GENERATED_UCLASS_BODY()
 
 	UPROPERTY(EditAnywhere, Category=PinOptions)
-	TArray<FString> PinNames;
+	TArray<FName> PinNames;
 
 	UPROPERTY(EditAnywhere, Category=PinOptions)
 	uint32 bIsCaseSensitive : 1;
@@ -39,11 +39,11 @@ class UK2Node_SwitchString : public UK2Node_Switch
 
 	// UK2Node_Switch Interface
 	BLUEPRINTGRAPH_API virtual void AddPinToSwitchNode() override;
-	virtual FString GetUniquePinName() override;
+	virtual FName GetUniquePinName() override;
 	virtual FEdGraphPinType GetPinType() const override;
 	// End of UK2Node_Switch Interface
 
-	virtual FString GetPinNameGivenIndex(int32 Index) override;
+	virtual FName GetPinNameGivenIndex(int32 Index) const override;
 
 protected:
 	virtual void CreateSelectionPin() override;
diff --git a/Engine/Source/Editor/BlueprintGraph/Classes/K2Node_Tunnel.h b/Engine/Source/Editor/BlueprintGraph/Classes/K2Node_Tunnel.h
--- a/Engine/Source/Editor/BlueprintGraph/Classes/K2Node_Tunnel.h
+++ b/Engine/Source/Editor/BlueprintGraph/Classes/K2Node_Tunnel.h
@@ -44,7 +44,7 @@ class BLUEPRINTGRAPH_API UK2Node_Tunnel : public UK2Node_EditablePinBase
 	virtual bool CanUserDeleteNode() const override;
 	virtual bool CanDuplicateNode() const override;
 	virtual UObject* GetJumpTargetForDoubleClick() const override;
-	virtual FString CreateUniquePinName(FString SourcePinName) const override;
+	virtual FName CreateUniquePinName(FName SourcePinName) const override;
 	virtual void ReallocatePinsDuringReconstruction(TArray<UEdGraphPin*>& OldPins) override;
 	//~ End UEdGraphNode Interface
 
@@ -58,7 +58,7 @@ class BLUEPRINTGRAPH_API UK2Node_Tunnel : public UK2Node_EditablePinBase
 	//~ Begin UK2Node_EditablePinBase Interface.
 	virtual UEdGraphPin* CreatePinFromUserDefinition(const TSharedPtr<FUserPinInfo> NewPinInfo) override;
 	virtual bool CanModifyExecutionWires() override;
-	virtual ERenamePinResult RenameUserDefinedPin(const FString& OldName, const FString& NewName, bool bTest = false) override;
+	virtual ERenamePinResult RenameUserDefinedPin(const FName OldName, const FName NewName, bool bTest = false) override;
 	virtual bool CanUseRefParams() const override { return true; }
 	virtual bool CanCreateUserDefinedPin(const FEdGraphPinType& InPinType, EEdGraphPinDirection InDesiredDirection, FText& OutErrorMessage) override;
 	virtual bool ModifyUserDefinedPinDefaultValue(TSharedPtr<FUserPinInfo> PinInfo, const FString& NewDefaultValue) override;
diff --git a/Engine/Source/Editor/BlueprintGraph/Classes/K2Node_Variable.h b/Engine/Source/Editor/BlueprintGraph/Classes/K2Node_Variable.h
--- a/Engine/Source/Editor/BlueprintGraph/Classes/K2Node_Variable.h
+++ b/Engine/Source/Editor/BlueprintGraph/Classes/K2Node_Variable.h
@@ -118,7 +118,7 @@ public:
 	 *
 	 * @return	true if it succeeds, false if it fails.
 	 */
-	bool CreatePinForVariable(EEdGraphPinDirection Direction, FString InPinName = FString());
+	bool CreatePinForVariable(EEdGraphPinDirection Direction, FName InPinName = NAME_None);
 
 	/** Creates 'self' pin */
 	void CreatePinForSelf();
@@ -132,7 +132,7 @@ public:
 	 *
 	 * @return	true if it succeeds, false if it fails.
 	 */
-	bool RecreatePinForVariable(EEdGraphPinDirection Direction, TArray<UEdGraphPin*>& OldPins, FString InPinName = FString());
+	bool RecreatePinForVariable(EEdGraphPinDirection Direction, TArray<UEdGraphPin*>& OldPins, FName InPinName = NAME_None);
 
 	/** Get the class to look for this variable in */
 	UClass* GetVariableSourceClass() const;
diff --git a/Engine/Source/Editor/BlueprintGraph/Classes/K2Node_VariableSet.h b/Engine/Source/Editor/BlueprintGraph/Classes/K2Node_VariableSet.h
--- a/Engine/Source/Editor/BlueprintGraph/Classes/K2Node_VariableSet.h
+++ b/Engine/Source/Editor/BlueprintGraph/Classes/K2Node_VariableSet.h
@@ -35,7 +35,7 @@ class UK2Node_VariableSet : public UK2Node_Variable
 	//~ End K2Node Interface
 
 	/** Retrieves the output pin name for the node */
-	FString GetVariableOutputPinName() const;
+	FName GetVariableOutputPinName() const;
 
 	BLUEPRINTGRAPH_API bool HasLocalRepNotify() const;
 	BLUEPRINTGRAPH_API FName GetRepNotifyName() const;
diff --git a/Engine/Source/Editor/BlueprintGraph/Private/BlueprintActionFilter.cpp b/Engine/Source/Editor/BlueprintGraph/Private/BlueprintActionFilter.cpp
--- a/Engine/Source/Editor/BlueprintGraph/Private/BlueprintActionFilter.cpp
+++ b/Engine/Source/Editor/BlueprintGraph/Private/BlueprintActionFilter.cpp
@@ -30,6 +30,7 @@
 #include "BlueprintEventNodeSpawner.h"
 #include "BlueprintBoundEventNodeSpawner.h"
 #include "BlueprintBoundNodeSpawner.h"
+#include "Algo/Transform.h"
 // "impure" node types (utilized in BlueprintActionFilterImpl::IsImpure)
 #include "K2Node_MultiGate.h"
 #include "K2Node_Message.h"
@@ -1378,7 +1379,7 @@ static bool BlueprintActionFilterImpl::ArrayFunctionHasParamOfType(const UFuncti
 {
 	UEdGraphSchema_K2 const* K2Schema = GetDefault<UEdGraphSchema_K2>();
 
-	TSet<FString> HiddenPins;
+	TSet<FName> HiddenPins;
 	FBlueprintEditorUtils::GetHiddenPinsForFunction(InGraph, ArrayFunction, HiddenPins);
 
 	FName ParamTag = FBlueprintMetadata::MD_ArrayDependentParam;
@@ -1388,13 +1389,14 @@ static bool BlueprintActionFilterImpl::ArrayFunctionHasParamOfType(const UFuncti
 	}
 	const FString& FlaggedParamMetaData = ArrayFunction->GetMetaData(ParamTag);
 
-	TArray<FString> WildcardPinNames;
-	FlaggedParamMetaData.ParseIntoArray(WildcardPinNames, TEXT(","), /*CullEmpty =*/true);
+	TArray<FString> WildcardPinNameStrs;
+	TArray<FName> WildcardPinNames;
+	FlaggedParamMetaData.ParseIntoArray(WildcardPinNameStrs, TEXT(","), /*CullEmpty =*/true);
+	Algo::Transform(WildcardPinNameStrs, WildcardPinNames, [](const FString& NameStr) { return FName(*NameStr); });
 
 	for (TFieldIterator<UProperty> PropIt(ArrayFunction); PropIt && (PropIt->PropertyFlags & CPF_Parm); ++PropIt)
 	{
 		UProperty* FuncParam = *PropIt;
-		const FString ParamName = FuncParam->GetName();
 
 		const bool bIsFunctionInput = !FuncParam->HasAnyPropertyFlags(CPF_OutParm) || FuncParam->HasAnyPropertyFlags(CPF_ReferenceParm);
 		if (bWantOutput == bIsFunctionInput)
@@ -1402,7 +1404,9 @@ static bool BlueprintActionFilterImpl::ArrayFunctionHasParamOfType(const UFuncti
 			continue;
 		}
 
-		if (!WildcardPinNames.Contains(ParamName) || HiddenPins.Contains(ParamName))
+		const FName ParamName = FuncParam->GetFName();
+
+		if (HiddenPins.Contains(ParamName) || !WildcardPinNames.Contains(ParamName))
 		{
 			continue;
 		}
@@ -1420,7 +1424,6 @@ static bool BlueprintActionFilterImpl::ArrayFunctionHasParamOfType(const UFuncti
 				return true;
 			}
 		}
-
 	}
 
 	return false;
@@ -1439,7 +1442,6 @@ static bool BlueprintActionFilterImpl::IsMissmatchedPropertyType(FBlueprintActio
 			bool const bIsGetter   = BlueprintAction.GetNodeClass()->IsChildOf<UK2Node_VariableGet>();
 			bool const bIsSetter   = BlueprintAction.GetNodeClass()->IsChildOf<UK2Node_VariableSet>();
 
-			//K2Schema->ConvertPropertyToPinType(VariableProperty, /*out*/ VariablePin->PinType);
 			for (int32 PinIndex = 0; !bIsFilteredOut && PinIndex < ContextPins.Num(); ++PinIndex)
 			{
 				UEdGraphPin const* ContextPin = ContextPins[PinIndex];
@@ -1457,7 +1459,7 @@ static bool BlueprintActionFilterImpl::IsMissmatchedPropertyType(FBlueprintActio
 					// just iterate over all the pins
 					bIsFilteredOut = !HasMatchingPin(BlueprintAction, ContextPin);
 				}
-				else if (ContextPinType.PinCategory == K2Schema->PC_Exec)
+				else if (ContextPinType.PinCategory == UEdGraphSchema_K2::PC_Exec)
 				{
 					// setters are impure, and therefore should have exec pins
 					bIsFilteredOut = bIsGetter;
diff --git a/Engine/Source/Editor/BlueprintGraph/Private/BlueprintFunctionNodeSpawner.cpp b/Engine/Source/Editor/BlueprintGraph/Private/BlueprintFunctionNodeSpawner.cpp
--- a/Engine/Source/Editor/BlueprintGraph/Private/BlueprintFunctionNodeSpawner.cpp
+++ b/Engine/Source/Editor/BlueprintGraph/Private/BlueprintFunctionNodeSpawner.cpp
@@ -122,7 +122,7 @@ static bool BlueprintFunctionNodeSpawnerImpl::BindFunctionNode(UK2Node_CallFunct
 	BindingOffset.Y += UEdGraphSchema_K2::EstimateNodeHeight(BindingNode);
 
 	UEdGraphPin* LiteralOutput = BindingNode->GetValuePin();
-	UEdGraphPin* CallSelfInput = NewNode->FindPin(GetDefault<UEdGraphSchema_K2>()->PN_Self);
+	UEdGraphPin* CallSelfInput = NewNode->FindPin(UEdGraphSchema_K2::PN_Self);
 	// connect the new "get-var" node with the spawned function node
 	if ((LiteralOutput != nullptr) && (CallSelfInput != nullptr))
 	{
diff --git a/Engine/Source/Editor/BlueprintGraph/Private/BlueprintGraphPanelPinFactory.cpp b/Engine/Source/Editor/BlueprintGraph/Private/BlueprintGraphPanelPinFactory.cpp
--- a/Engine/Source/Editor/BlueprintGraph/Private/BlueprintGraphPanelPinFactory.cpp
+++ b/Engine/Source/Editor/BlueprintGraph/Private/BlueprintGraphPanelPinFactory.cpp
@@ -12,8 +12,7 @@
 	
 TSharedPtr<class SGraphPin> FBlueprintGraphPanelPinFactory::CreatePin(class UEdGraphPin* InPin) const
 {
-	const UEdGraphSchema_K2* K2Schema = GetDefault<UEdGraphSchema_K2>();
-	if (InPin->PinType.PinCategory == K2Schema->PC_Name)
+	if (InPin->PinType.PinCategory == UEdGraphSchema_K2::PC_Name)
 	{
 		UObject* Outer = InPin->GetOuter();
 
diff --git a/Engine/Source/Editor/BlueprintGraph/Private/CallFunctionHandler.cpp b/Engine/Source/Editor/BlueprintGraph/Private/CallFunctionHandler.cpp
--- a/Engine/Source/Editor/BlueprintGraph/Private/CallFunctionHandler.cpp
+++ b/Engine/Source/Editor/BlueprintGraph/Private/CallFunctionHandler.cpp
@@ -109,7 +109,7 @@ void FKCHandler_CallFunction::CreateFunctionCallStatement(FKismetFunctionContext
 						
 							int32 LatentUUID = CompilerContext.MessageLog.CalculateStableIdentifierForLatentActionManager(LatentInfoPin->GetOwningNode());
 
-							const FString ExecutionFunctionName = CompilerContext.GetSchema()->FN_ExecuteUbergraphBase.ToString() + TEXT("_") + Context.Blueprint->GetName();
+							const FString ExecutionFunctionName = UEdGraphSchema_K2::FN_ExecuteUbergraphBase.ToString() + TEXT("_") + Context.Blueprint->GetName();
 							(*Term)->Name = FString::Printf(TEXT("(Linkage=%s,UUID=%s,ExecutionFunction=%s,CallbackTarget=None)"), *FString::FromInt(INDEX_NONE), *FString::FromInt(LatentUUID), *ExecutionFunctionName);
 
 							// Record the UUID in the debugging information
@@ -148,7 +148,7 @@ void FKCHandler_CallFunction::CreateFunctionCallStatement(FKismetFunctionContext
 			for (int32 i = 0; !bFoundParam && (i < RemainingPins.Num()); ++i)
 			{
 				UEdGraphPin* PinMatch = RemainingPins[i];
-				if (FCString::Stricmp(*Property->GetName(), *PinMatch->PinName) == 0)
+				if (Property->GetFName() == PinMatch->PinName)
 				{
 					// Found a corresponding pin, does it match in type and direction?
 					if (FKismetCompilerUtilities::IsTypeCompatibleWithProperty(PinMatch, Property, CompilerContext.MessageLog, CompilerContext.GetSchema(), Context.NewClass))
@@ -182,7 +182,7 @@ void FKCHandler_CallFunction::CreateFunctionCallStatement(FKismetFunctionContext
 
 											// Import the literal text to a dummy struct to verify it's well-formed
 											FImportTextErrorContext ErrorPipe(CompilerContext.MessageLog, Node);
-											StructProperty->ImportText(*((*Term)->Name), StructData, 0, NULL, &ErrorPipe);
+											StructProperty->ImportText(*((*Term)->Name), StructData, 0, nullptr, &ErrorPipe);
 											if(ErrorPipe.NumErrors > 0)
 											{
 												bMatchedAllParams = false;
@@ -211,7 +211,7 @@ void FKCHandler_CallFunction::CreateFunctionCallStatement(FKismetFunctionContext
 									ClassTerm->bIsLiteral = true;
 									ClassTerm->Source     = Node;
 									ClassTerm->ObjectLiteral = InterfaceClass;
-									ClassTerm->Type.PinCategory = CompilerContext.GetSchema()->PC_Class;
+									ClassTerm->Type.PinCategory = UEdGraphSchema_K2::PC_Class;
 
 									// insert a cast op before a call to the function (and replace
 									// the param with the result from the cast)
@@ -285,7 +285,7 @@ void FKCHandler_CallFunction::CreateFunctionCallStatement(FKismetFunctionContext
 		// At this point, we should have consumed all pins.  If not, there are extras that need to be removed.
 		for (int32 i = 0; i < RemainingPins.Num(); ++i)
 		{
-			CompilerContext.MessageLog.Error(*FText::Format(LOCTEXT("PinMismatchParameter_ErrorFmt", "Pin @@ named {0} doesn't match any parameters of function {1}"), FText::FromString(RemainingPins[i]->PinName), FText::FromString(Function->GetName())).ToString(), RemainingPins[i]);
+			CompilerContext.MessageLog.Error(*FText::Format(LOCTEXT("PinMismatchParameter_ErrorFmt", "Pin @@ named {0} doesn't match any parameters of function {1}"), FText::FromName(RemainingPins[i]->PinName), FText::FromString(Function->GetName())).ToString(), RemainingPins[i]);
 		}
 
 		if (NumErrorsAtStart == CompilerContext.MessageLog.NumErrors)
@@ -475,14 +475,12 @@ UClass* FKCHandler_CallFunction::GetTrueCallingClass(FKismetFunctionContext& Con
 {
 	if (SelfPin)
 	{
-		UEdGraphSchema_K2 const* K2Schema = CompilerContext.GetSchema();
-
 		// TODO: here FBlueprintCompiledStatement::GetScopeFromPinType should be called, but since FEdGraphPinType::PinSubCategory is not always initialized properly that function works wrong
 		// return Cast<UClass>(Context.GetScopeFromPinType(SelfPin->PinType, Context.NewClass));
 		FEdGraphPinType& Type = SelfPin->PinType;
-		if ((Type.PinCategory == K2Schema->PC_Object) || (Type.PinCategory == K2Schema->PC_Class) || (Type.PinCategory == K2Schema->PC_Interface))
+		if ((Type.PinCategory == UEdGraphSchema_K2::PC_Object) || (Type.PinCategory == UEdGraphSchema_K2::PC_Class) || (Type.PinCategory == UEdGraphSchema_K2::PC_Interface))
 		{
-			if (!Type.PinSubCategory.IsEmpty() && (Type.PinSubCategory != K2Schema->PSC_Self))
+			if (!Type.PinSubCategory.IsNone() && (Type.PinSubCategory != UEdGraphSchema_K2::PSC_Self))
 			{
 				return Cast<UClass>(Type.PinSubCategoryObject.Get());
 			}
@@ -493,24 +491,23 @@ UClass* FKCHandler_CallFunction::GetTrueCallingClass(FKismetFunctionContext& Con
 
 void FKCHandler_CallFunction::RegisterNets(FKismetFunctionContext& Context, UEdGraphNode* Node)
 {
-	UEdGraphSchema_K2 const* K2Schema = CompilerContext.GetSchema();
-
 	if (UFunction* Function = FindFunction(Context, Node))
 	{
-		TArray<FString> DefaultToSelfParamNames;
-		TArray<FString> RequiresSetValue;
+		TArray<FName> DefaultToSelfParamNames;
+		TArray<FName> RequiresSetValue;
 
 		if (Function->HasMetaData(FBlueprintMetadata::MD_DefaultToSelf))
 		{
-			const FString& DefaltToSelfPinName = Function->GetMetaData(FBlueprintMetadata::MD_DefaultToSelf);
+			const FName DefaltToSelfPinName = *Function->GetMetaData(FBlueprintMetadata::MD_DefaultToSelf);
 
 			DefaultToSelfParamNames.Add(DefaltToSelfPinName);
 		}
 		if (Function->HasMetaData(FBlueprintMetadata::MD_WorldContext))
 		{
+			UEdGraphSchema_K2 const* K2Schema = CompilerContext.GetSchema();
 			const bool bHasIntrinsicWorldContext = !K2Schema->IsStaticFunctionGraph(Context.SourceGraph) && FBlueprintEditorUtils::ImplentsGetWorld(Context.Blueprint);
 
-			const FString& WorldContextPinName = Function->GetMetaData(FBlueprintMetadata::MD_WorldContext);
+			const FName WorldContextPinName = *Function->GetMetaData(FBlueprintMetadata::MD_WorldContext);
 
 			if (bHasIntrinsicWorldContext)
 			{
@@ -532,10 +529,10 @@ void FKCHandler_CallFunction::RegisterNets(FKismetFunctionContext& Context, UEdG
 				if (DefaultToSelfParamNames.Contains(Pin->PinName) && FKismetCompilerUtilities::ValidateSelfCompatibility(Pin, Context))
 				{
 					ensure(Pin->PinType.PinSubCategoryObject != nullptr);
-					ensure((Pin->PinType.PinCategory == K2Schema->PC_Object) || (Pin->PinType.PinCategory == K2Schema->PC_Interface));
+					ensure((Pin->PinType.PinCategory == UEdGraphSchema_K2::PC_Object) || (Pin->PinType.PinCategory == UEdGraphSchema_K2::PC_Interface));
 
 					FBPTerminal* Term = Context.RegisterLiteral(Pin);
-					Term->Type.PinSubCategory = CompilerContext.GetSchema()->PN_Self;
+					Term->Type.PinSubCategory = UEdGraphSchema_K2::PN_Self;
 					Context.NetMap.Add(Pin, Term);
 				}
 				else if (RequiresSetValue.Contains(Pin->PinName))
@@ -556,8 +553,8 @@ void FKCHandler_CallFunction::RegisterNets(FKismetFunctionContext& Context, UEdG
 		// if we have an object plugged into an interface pin, let's create a 
 		// term that'll be used as an intermediate, holding the result of a cast 
 		// from object to interface
-		if (((Pin->PinType.PinCategory == K2Schema->PC_Interface) && (Pin->LinkedTo[0]->PinType.PinCategory == K2Schema->PC_Object)) ||
-			((Pin->PinType.PinCategory == K2Schema->PC_Object) && (Pin->LinkedTo[0]->PinType.PinCategory == K2Schema->PC_Interface)))
+		if (((Pin->PinType.PinCategory == UEdGraphSchema_K2::PC_Interface) && (Pin->LinkedTo[0]->PinType.PinCategory == UEdGraphSchema_K2::PC_Object)) ||
+			((Pin->PinType.PinCategory == UEdGraphSchema_K2::PC_Object) && (Pin->LinkedTo[0]->PinType.PinCategory == UEdGraphSchema_K2::PC_Interface)))
 		{
 			FBPTerminal* InterfaceTerm = Context.CreateLocalTerminal();
 			InterfaceTerm->CopyFromPin(Pin, Context.NetNameMap->MakeValidName(Pin) + TEXT("_CastInput"));
@@ -594,7 +591,6 @@ UFunction* FKCHandler_CallFunction::FindFunction(FKismetFunctionContext& Context
 void FKCHandler_CallFunction::Transform(FKismetFunctionContext& Context, UEdGraphNode* Node)
 {
 	// Add an object reference pin for this call
-	//UEdGraphPin* OperatingOn = Node->CreatePin(EGPD_Input, Schema->PC_Object, TEXT(""), TEXT("OperatingContext"));
 
 	if (IsCalledFunctionPure(Node))
 	{
@@ -681,8 +677,8 @@ void FKCHandler_CallFunction::Compile(FKismetFunctionContext& Context, UEdGraphN
 	if (SelfPin && (SelfPin->LinkedTo.Num() == 0))
 	{
 		FEdGraphPinType SelfType;
-		SelfType.PinCategory = CompilerContext.GetSchema()->PC_Object;
-		SelfType.PinSubCategory = CompilerContext.GetSchema()->PSC_Self;
+		SelfType.PinCategory = UEdGraphSchema_K2::PC_Object;
+		SelfType.PinSubCategory = UEdGraphSchema_K2::PSC_Self;
 
 		if (!CompilerContext.GetSchema()->ArePinTypesCompatible(SelfType, SelfPin->PinType, Context.NewClass) && (SelfPin->DefaultObject == NULL))
 		{
@@ -699,7 +695,7 @@ void FKCHandler_CallFunction::CheckIfFunctionIsCallable(UFunction* Function, FKi
 	// Verify that the function is a Blueprint callable function (in case a BlueprintCallable specifier got removed)
 	if (!Function->HasAnyFunctionFlags(FUNC_BlueprintCallable) && (Function->GetOuter() != Context.NewClass))
 	{
-		if (!IsCalledFunctionFinal(Node) && Function->GetName().Find(CompilerContext.GetSchema()->FN_ExecuteUbergraphBase.ToString()))
+		if (!IsCalledFunctionFinal(Node) && Function->GetName().Find(UEdGraphSchema_K2::FN_ExecuteUbergraphBase.ToString()))
 		{
 			CompilerContext.MessageLog.Error(*FText::Format(NSLOCTEXT("KismetCompiler", "ShouldNotCallFromBlueprint_ErrorFmt", "Function '{0}' called from @@ should not be called from a Blueprint"), FText::FromString(Function->GetName())).ToString(), Node);
 		}
diff --git a/Engine/Source/Editor/BlueprintGraph/Private/DelegateNodeHandlers.cpp b/Engine/Source/Editor/BlueprintGraph/Private/DelegateNodeHandlers.cpp
--- a/Engine/Source/Editor/BlueprintGraph/Private/DelegateNodeHandlers.cpp
+++ b/Engine/Source/Editor/BlueprintGraph/Private/DelegateNodeHandlers.cpp
@@ -192,7 +192,7 @@ void FKCHandler_AddRemoveDelegate::Compile(FKismetFunctionContext& Context, UEdG
 		AddStatement.RHS.Add(*DelegateInputTerm);
 	}
 
-	GenerateSimpleThenGoto(Context, *DelegateNode, DelegateNode->FindPin(CompilerContext.GetSchema()->PN_Then));
+	GenerateSimpleThenGoto(Context, *DelegateNode, DelegateNode->FindPin(UEdGraphSchema_K2::PN_Then));
 	FNodeHandlingFunctor::Compile(Context, DelegateNode);
 }
 
@@ -229,7 +229,7 @@ void FKCHandler_CreateDelegate::RegisterNets(FKismetFunctionContext& Context, UE
 			{
 				InputObjTerm = Context.CreateLocalTerminal(ETerminalSpecification::TS_Literal);
 				InputObjTerm->Name = Context.NetNameMap->MakeValidName(Net);
-				InputObjTerm->Type.PinSubCategory = CompilerContext.GetSchema()->PN_Self;
+				InputObjTerm->Type.PinSubCategory = UEdGraphSchema_K2::PN_Self;
 			}
 			else
 			{
@@ -289,7 +289,7 @@ void FKCHandler_CreateDelegate::Compile(FKismetFunctionContext& Context, UEdGrap
 
 	{
 		FBPTerminal* DelegateNameTerm = Context.CreateLocalTerminal(ETerminalSpecification::TS_Literal);
-		DelegateNameTerm->Type.PinCategory = CompilerContext.GetSchema()->PC_Name;
+		DelegateNameTerm->Type.PinCategory = UEdGraphSchema_K2::PC_Name;
 		DelegateNameTerm->Name = DelegateNode->GetFunctionName().ToString();
 		DelegateNameTerm->bIsLiteral = true;
 		Statement.RHS.Add(DelegateNameTerm);
@@ -339,7 +339,7 @@ void FKCHandler_ClearDelegate::Compile(FKismetFunctionContext& Context, UEdGraph
 		AddStatement.LHS = *VarDelegate;
 	}
 
-	GenerateSimpleThenGoto(Context, *DelegateNode, DelegateNode->FindPin(CompilerContext.GetSchema()->PN_Then));
+	GenerateSimpleThenGoto(Context, *DelegateNode, DelegateNode->FindPin(UEdGraphSchema_K2::PN_Then));
 	FNodeHandlingFunctor::Compile(Context, DelegateNode);
 }
 
diff --git a/Engine/Source/Editor/BlueprintGraph/Private/DynamicCastHandler.cpp b/Engine/Source/Editor/BlueprintGraph/Private/DynamicCastHandler.cpp
--- a/Engine/Source/Editor/BlueprintGraph/Private/DynamicCastHandler.cpp
+++ b/Engine/Source/Editor/BlueprintGraph/Private/DynamicCastHandler.cpp
@@ -23,7 +23,7 @@ void FKCHandler_DynamicCast::RegisterNets(FKismetFunctionContext& Context, UEdGr
 		{
 			// Create a term to determine if the cast was successful or not
 			FBPTerminal* BoolTerm = Context.CreateLocalTerminal();
-			BoolTerm->Type.PinCategory = CompilerContext.GetSchema()->PC_Boolean;
+			BoolTerm->Type.PinCategory = UEdGraphSchema_K2::PC_Boolean;
 			BoolTerm->Source = Node;
 			BoolTerm->Name = Context.NetNameMap->MakeValidName(Node) + TEXT("_CastSuccess");
 			BoolTermMap.Add(Node, BoolTerm);
@@ -84,7 +84,7 @@ void FKCHandler_DynamicCast::Compile(FKismetFunctionContext& Context, UEdGraphNo
 	ClassTerm->bIsLiteral = true;
 	ClassTerm->Source = Node;
 	ClassTerm->ObjectLiteral = DynamicCastNode->TargetType;
-	ClassTerm->Type.PinCategory = CompilerContext.GetSchema()->PC_Class;
+	ClassTerm->Type.PinCategory = UEdGraphSchema_K2::PC_Class;
 
 	UClass const* const InputObjClass  = Cast<UClass>((*ObjectToCast)->Type.PinSubCategoryObject.Get());
 	UClass const* const OutputObjClass = Cast<UClass>((*CastResultTerm)->Type.PinSubCategoryObject.Get());
diff --git a/Engine/Source/Editor/BlueprintGraph/Private/EdGraphSchema_K2.cpp b/Engine/Source/Editor/BlueprintGraph/Private/EdGraphSchema_K2.cpp
--- a/Engine/Source/Editor/BlueprintGraph/Private/EdGraphSchema_K2.cpp
+++ b/Engine/Source/Editor/BlueprintGraph/Private/EdGraphSchema_K2.cpp
@@ -181,7 +181,7 @@ const FName FBlueprintMetadata::MD_UseEnumValuesAsMaskValuesInEditor(TEXT("UseEn
 
 #define LOCTEXT_NAMESPACE "KismetSchema"
 
-UEdGraphSchema_K2::FPinTypeTreeInfo::FPinTypeTreeInfo(const FText& InFriendlyName, const FString& CategoryName, const UEdGraphSchema_K2* Schema, const FText& InTooltip, bool bInReadOnly/*=false*/, FTypesDatabase* TypesDatabase /*=nullptr*/)
+UEdGraphSchema_K2::FPinTypeTreeInfo::FPinTypeTreeInfo(const FText& InFriendlyName, const FName CategoryName, const UEdGraphSchema_K2* Schema, const FText& InTooltip, bool bInReadOnly/*=false*/, FTypesDatabase* TypesDatabase /*=nullptr*/)
 	: PossibleObjectReferenceTypes(0)
 {
 	Init(InFriendlyName, CategoryName, Schema, InTooltip, bInReadOnly, TypesDatabase);
@@ -233,10 +233,10 @@ struct FLoadedAssetData
 struct FTypesDatabase
 {
 	typedef TSharedPtr<TArray<FLoadedAssetData>> FLoadedTypesList;
-	TMap<FString, FLoadedTypesList> LoadedTypesMap;
+	TMap<FName, FLoadedTypesList> LoadedTypesMap;
 
 	typedef TSharedPtr<TArray<FUnloadedAssetData>> FUnLoadedTypesList;
-	TMap<FString, FUnLoadedTypesList> UnLoadedTypesMap;
+	TMap<FName, FUnLoadedTypesList> UnLoadedTypesMap;
 };
 
 /** Helper class to gather variable types */
@@ -364,7 +364,6 @@ public:
 				TypesDatabase.UnLoadedTypesMap.Add(UEdGraphSchema_K2::PC_Struct, UnLoadedTypesList);
 			}
 
-			//else if (Schema->PC_Object == CategoryName || Schema->PC_Class == CategoryName || Schema->PC_Interface == CategoryName || Schema->PC_SoftObject == CategoryName || Schema->PC_SoftClass == CategoryName)
 			{
 				TArray<FAssetData> AssetData;
 				AssetRegistryModule.Get().GetAssetsByClass(UBlueprint::StaticClass()->GetFName(), AssetData);
@@ -415,12 +414,12 @@ public:
 	 * @param TypesDatabase		Types database
 	 * @param OutChildren		All the gathered children
 	 */
-	static void Gather(const FText& FriendlyName, const FString& CategoryName, FTypesDatabase& TypesDatabase, TArray<FPinTypeTreeInfoPtr>& OutChildren)
+	static void Gather(const FText& FriendlyName, const FName CategoryName, FTypesDatabase& TypesDatabase, TArray<FPinTypeTreeInfoPtr>& OutChildren)
 	{
 		FEdGraphPinType LoadedPinSubtype;
 		LoadedPinSubtype.PinCategory = (CategoryName == UEdGraphSchema_K2::PC_Enum ? UEdGraphSchema_K2::PC_Byte : CategoryName);
-		LoadedPinSubtype.PinSubCategory = TEXT("");
-		LoadedPinSubtype.PinSubCategoryObject = NULL;
+		LoadedPinSubtype.PinSubCategory = NAME_None;
+		LoadedPinSubtype.PinSubCategoryObject = nullptr;
 
 		FTypesDatabase::FLoadedTypesList* LoadedSubTypesPtr = TypesDatabase.LoadedTypesMap.Find(CategoryName);
 		if (LoadedSubTypesPtr && LoadedSubTypesPtr->IsValid())
@@ -515,16 +514,16 @@ const FEdGraphPinType& UEdGraphSchema_K2::FPinTypeTreeInfo::GetPinType(bool bFor
 	return PinType;
 }
 
-void UEdGraphSchema_K2::FPinTypeTreeInfo::Init(const FText& InFriendlyName, const FString& CategoryName, const UEdGraphSchema_K2* Schema, const FText& InTooltip, bool bInReadOnly, FTypesDatabase* TypesDatabase)
+void UEdGraphSchema_K2::FPinTypeTreeInfo::Init(const FText& InFriendlyName, const FName CategoryName, const UEdGraphSchema_K2* Schema, const FText& InTooltip, bool bInReadOnly, FTypesDatabase* TypesDatabase)
 {
-	check( !CategoryName.IsEmpty() );
+	check( !CategoryName.IsNone() );
 	check( Schema );
 
 	FriendlyName = InFriendlyName;
 	Tooltip = InTooltip;
 	PinType.PinCategory = (CategoryName == PC_Enum? PC_Byte : CategoryName);
-	PinType.PinSubCategory = TEXT("");
-	PinType.PinSubCategoryObject = NULL;
+	PinType.PinSubCategory = NAME_None;
+	PinType.PinSubCategoryObject = nullptr;
 
 	bReadOnly = bInReadOnly;
 
@@ -539,10 +538,10 @@ void UEdGraphSchema_K2::FPinTypeTreeInfo::Init(const FText& InFriendlyName, cons
 	}
 }
 
-UEdGraphSchema_K2::FPinTypeTreeInfo::FPinTypeTreeInfo(const FString& CategoryName, UObject* SubCategoryObject, const FText& InTooltip, bool bInReadOnly/*=false*/, uint8 InPossibleObjectReferenceTypes)
+UEdGraphSchema_K2::FPinTypeTreeInfo::FPinTypeTreeInfo(const FName CategoryName, UObject* SubCategoryObject, const FText& InTooltip, bool bInReadOnly/*=false*/, uint8 InPossibleObjectReferenceTypes)
 	: PossibleObjectReferenceTypes(InPossibleObjectReferenceTypes)
 {
-	check( !CategoryName.IsEmpty() );
+	check( !CategoryName.IsNone() );
 	check( SubCategoryObject );
 
 	Tooltip = InTooltip;
@@ -553,12 +552,12 @@ UEdGraphSchema_K2::FPinTypeTreeInfo::FPinTypeTreeInfo(const FString& CategoryNam
 	CachedDescription = GenerateDescription();
 }
 
-UEdGraphSchema_K2::FPinTypeTreeInfo::FPinTypeTreeInfo(const FText& InFriendlyName, const FString& CategoryName, const FSoftObjectPath& SubCategoryObject, const FText& InTooltip, bool bInReadOnly, uint8 InPossibleObjectReferenceTypes)
+UEdGraphSchema_K2::FPinTypeTreeInfo::FPinTypeTreeInfo(const FText& InFriendlyName, const FName CategoryName, const FSoftObjectPath& SubCategoryObject, const FText& InTooltip, bool bInReadOnly, uint8 InPossibleObjectReferenceTypes)
 	: PossibleObjectReferenceTypes(InPossibleObjectReferenceTypes)
 {
 	FriendlyName = InFriendlyName;
 
-	check(!CategoryName.IsEmpty());
+	check(!CategoryName.IsNone());
 	check(SubCategoryObject.IsValid());
 
 	Tooltip = InTooltip;
@@ -602,47 +601,47 @@ FText UEdGraphSchema_K2::FPinTypeTreeInfo::GetDescription() const
 	return CachedDescription;
 }
 
-const FString UEdGraphSchema_K2::PC_Exec(TEXT("exec"));
-const FString UEdGraphSchema_K2::PC_Boolean(TEXT("bool"));
-const FString UEdGraphSchema_K2::PC_Byte(TEXT("byte"));
-const FString UEdGraphSchema_K2::PC_Class(TEXT("class"));
-const FString UEdGraphSchema_K2::PC_Int(TEXT("int"));
-const FString UEdGraphSchema_K2::PC_Float(TEXT("float"));
-const FString UEdGraphSchema_K2::PC_Name(TEXT("name"));
-const FString UEdGraphSchema_K2::PC_Delegate(TEXT("delegate"));
-const FString UEdGraphSchema_K2::PC_MCDelegate(TEXT("mcdelegate"));
-const FString UEdGraphSchema_K2::PC_Object(TEXT("object"));
-const FString UEdGraphSchema_K2::PC_Interface(TEXT("interface"));
-const FString UEdGraphSchema_K2::PC_String(TEXT("string"));
-const FString UEdGraphSchema_K2::PC_Text(TEXT("text"));
-const FString UEdGraphSchema_K2::PC_Struct(TEXT("struct"));
-const FString UEdGraphSchema_K2::PC_Wildcard(TEXT("wildcard"));
-const FString UEdGraphSchema_K2::PC_Enum(TEXT("enum"));
-const FString UEdGraphSchema_K2::PC_SoftObject(TEXT("softobject"));
-const FString UEdGraphSchema_K2::PC_SoftClass(TEXT("softclass"));
-const FString UEdGraphSchema_K2::PSC_Self(TEXT("self"));
-const FString UEdGraphSchema_K2::PSC_Index(TEXT("index"));
-const FString UEdGraphSchema_K2::PSC_Bitmask(TEXT("bitmask"));
-const FString UEdGraphSchema_K2::PN_Execute(TEXT("execute"));
-const FString UEdGraphSchema_K2::PN_Then(TEXT("then"));
-const FString UEdGraphSchema_K2::PN_Completed(TEXT("Completed"));
-const FString UEdGraphSchema_K2::PN_DelegateEntry(TEXT("delegate"));
-const FString UEdGraphSchema_K2::PN_EntryPoint(TEXT("EntryPoint"));
-const FString UEdGraphSchema_K2::PN_Self(TEXT("self"));
-const FString UEdGraphSchema_K2::PN_Else(TEXT("else"));
-const FString UEdGraphSchema_K2::PN_Loop(TEXT("loop"));
-const FString UEdGraphSchema_K2::PN_After(TEXT("after"));
-const FString UEdGraphSchema_K2::PN_ReturnValue(TEXT("ReturnValue"));
-const FString UEdGraphSchema_K2::PN_ObjectToCast(TEXT("Object"));
-const FString UEdGraphSchema_K2::PN_Condition(TEXT("Condition"));
-const FString UEdGraphSchema_K2::PN_Start(TEXT("Start"));
-const FString UEdGraphSchema_K2::PN_Stop(TEXT("Stop"));
-const FString UEdGraphSchema_K2::PN_Index(TEXT("Index"));
-const FString UEdGraphSchema_K2::PN_Item(TEXT("Item"));
-const FString UEdGraphSchema_K2::PN_CastSucceeded(TEXT("then"));
-const FString UEdGraphSchema_K2::PN_CastFailed(TEXT("CastFailed"));
+const FName UEdGraphSchema_K2::PC_Exec(TEXT("exec"));
+const FName UEdGraphSchema_K2::PC_Boolean(TEXT("bool"));
+const FName UEdGraphSchema_K2::PC_Byte(TEXT("byte"));
+const FName UEdGraphSchema_K2::PC_Class(TEXT("class"));
+const FName UEdGraphSchema_K2::PC_Int(TEXT("int"));
+const FName UEdGraphSchema_K2::PC_Float(TEXT("float"));
+const FName UEdGraphSchema_K2::PC_Name(TEXT("name"));
+const FName UEdGraphSchema_K2::PC_Delegate(TEXT("delegate"));
+const FName UEdGraphSchema_K2::PC_MCDelegate(TEXT("mcdelegate"));
+const FName UEdGraphSchema_K2::PC_Object(TEXT("object"));
+const FName UEdGraphSchema_K2::PC_Interface(TEXT("interface"));
+const FName UEdGraphSchema_K2::PC_String(TEXT("string"));
+const FName UEdGraphSchema_K2::PC_Text(TEXT("text"));
+const FName UEdGraphSchema_K2::PC_Struct(TEXT("struct"));
+const FName UEdGraphSchema_K2::PC_Wildcard(TEXT("wildcard"));
+const FName UEdGraphSchema_K2::PC_Enum(TEXT("enum"));
+const FName UEdGraphSchema_K2::PC_SoftObject(TEXT("softobject"));
+const FName UEdGraphSchema_K2::PC_SoftClass(TEXT("softclass"));
+const FName UEdGraphSchema_K2::PSC_Self(TEXT("self"));
+const FName UEdGraphSchema_K2::PSC_Index(TEXT("index"));
+const FName UEdGraphSchema_K2::PSC_Bitmask(TEXT("bitmask"));
+const FName UEdGraphSchema_K2::PN_Execute(TEXT("execute"));
+const FName UEdGraphSchema_K2::PN_Then(TEXT("then"));
+const FName UEdGraphSchema_K2::PN_Completed(TEXT("Completed"));
+const FName UEdGraphSchema_K2::PN_DelegateEntry(TEXT("delegate"));
+const FName UEdGraphSchema_K2::PN_EntryPoint(TEXT("EntryPoint"));
+const FName UEdGraphSchema_K2::PN_Self(TEXT("self"));
+const FName UEdGraphSchema_K2::PN_Else(TEXT("else"));
+const FName UEdGraphSchema_K2::PN_Loop(TEXT("loop"));
+const FName UEdGraphSchema_K2::PN_After(TEXT("after"));
+const FName UEdGraphSchema_K2::PN_ReturnValue(TEXT("ReturnValue"));
+const FName UEdGraphSchema_K2::PN_ObjectToCast(TEXT("Object"));
+const FName UEdGraphSchema_K2::PN_Condition(TEXT("Condition"));
+const FName UEdGraphSchema_K2::PN_Start(TEXT("Start"));
+const FName UEdGraphSchema_K2::PN_Stop(TEXT("Stop"));
+const FName UEdGraphSchema_K2::PN_Index(TEXT("Index"));
+const FName UEdGraphSchema_K2::PN_Item(TEXT("Item"));
+const FName UEdGraphSchema_K2::PN_CastSucceeded(TEXT("then"));
+const FName UEdGraphSchema_K2::PN_CastFailed(TEXT("CastFailed"));
 const FString UEdGraphSchema_K2::PN_CastedValuePrefix(TEXT("As"));
-const FString UEdGraphSchema_K2::PN_MatineeFinished(TEXT("Finished"));
+const FName UEdGraphSchema_K2::PN_MatineeFinished(TEXT("Finished"));
 
 const FName UEdGraphSchema_K2::FN_UserConstructionScript(TEXT("UserConstructionScript"));
 const FName UEdGraphSchema_K2::FN_ExecuteUbergraphBase(TEXT("ExecuteUbergraph"));
@@ -661,7 +660,7 @@ const UScriptStruct* UEdGraphSchema_K2::ColorStruct = nullptr;
 bool UEdGraphSchema_K2::bGeneratingDocumentation = false;
 int32 UEdGraphSchema_K2::CurrentCacheRefreshID = 0;
 
-const FString UEdGraphSchema_K2::AllObjectTypes(TEXT("ObjectTypes"));
+const FName UEdGraphSchema_K2::AllObjectTypes(TEXT("ObjectTypes"));
 
 UEdGraphSchema_K2::UEdGraphSchema_K2(const FObjectInitializer& ObjectInitializer)
 	: Super(ObjectInitializer)
@@ -984,7 +983,7 @@ void UEdGraphSchema_K2::GetAutoEmitTermParameters(const UFunction* Function, TAr
 
 bool UEdGraphSchema_K2::FunctionHasParamOfType(const UFunction* InFunction, UEdGraph const* InGraph, const FEdGraphPinType& DesiredPinType, bool bWantOutput) const
 {
-	TSet<FString> HiddenPins;
+	TSet<FName> HiddenPins;
 	FBlueprintEditorUtils::GetHiddenPinsForFunction(InGraph, InFunction, HiddenPins);
 
 	// Iterate over all params of function
@@ -993,7 +992,7 @@ bool UEdGraphSchema_K2::FunctionHasParamOfType(const UFunction* InFunction, UEdG
 		UProperty* FuncParam = *PropIt;
 
 		// Ensure that this isn't a hidden parameter
-		if (!HiddenPins.Contains(FuncParam->GetName()))
+		if (!HiddenPins.Contains(FuncParam->GetFName()))
 		{
 			// See if this is the direction we want (input or output)
 			const bool bIsFunctionInput = !FuncParam->HasAnyPropertyFlags(CPF_OutParm) || FuncParam->HasAnyPropertyFlags(CPF_ReferenceParm);
@@ -1733,7 +1732,7 @@ void UEdGraphSchema_K2::GetContextMenuActions(const UEdGraph* CurrentGraph, cons
 
 void UEdGraphSchema_K2::OnCreateNonExistentVariable( UK2Node_Variable* Variable,  UBlueprint* OwnerBlueprint)
 {
-	if (UEdGraphPin* Pin = Variable->FindPin(Variable->GetVarNameString()))
+	if (UEdGraphPin* Pin = Variable->FindPin(Variable->GetVarName()))
 	{
 		const FScopedTransaction Transaction( LOCTEXT("CreateMissingVariable", "Create Missing Variable") );
 
@@ -1747,7 +1746,7 @@ void UEdGraphSchema_K2::OnCreateNonExistentVariable( UK2Node_Variable* Variable,
 
 void UEdGraphSchema_K2::OnCreateNonExistentLocalVariable( UK2Node_Variable* Variable,  UBlueprint* OwnerBlueprint)
 {
-	if (UEdGraphPin* Pin = Variable->FindPin(Variable->GetVarNameString()))
+	if (UEdGraphPin* Pin = Variable->FindPin(Variable->GetVarName()))
 	{
 		const FScopedTransaction Transaction( LOCTEXT("CreateMissingLocalVariable", "Create Missing Local Variable") );
 
@@ -1774,9 +1773,9 @@ void UEdGraphSchema_K2::OnCreateNonExistentLocalVariable( UK2Node_Variable* Vari
 	}	
 }
 
-void UEdGraphSchema_K2::OnReplaceVariableForVariableNode( UK2Node_Variable* Variable, UBlueprint* OwnerBlueprint, FString VariableName, bool bIsSelfMember)
+void UEdGraphSchema_K2::OnReplaceVariableForVariableNode( UK2Node_Variable* Variable, UBlueprint* OwnerBlueprint, FName VariableName, bool bIsSelfMember)
 {
-	if(UEdGraphPin* Pin = Variable->FindPin(Variable->GetVarNameString()))
+	if(UEdGraphPin* Pin = Variable->FindPin(Variable->GetVarName()))
 	{
 		const FScopedTransaction Transaction( NSLOCTEXT("UnrealEd", "GraphEd_ReplaceVariable", "Replace Variable") );
 		Variable->Modify();
@@ -1784,14 +1783,13 @@ void UEdGraphSchema_K2::OnReplaceVariableForVariableNode( UK2Node_Variable* Vari
 
 		if (bIsSelfMember)
 		{
-			FName VarName = FName(*VariableName);
-			FGuid Guid = FBlueprintEditorUtils::FindMemberVariableGuidByName(OwnerBlueprint, VarName);
-			Variable->VariableReference.SetSelfMember( VarName, Guid );
+			FGuid Guid = FBlueprintEditorUtils::FindMemberVariableGuidByName(OwnerBlueprint, VariableName);
+			Variable->VariableReference.SetSelfMember( VariableName, Guid );
 		}
 		else
 		{
 			UEdGraph* FunctionGraph = FBlueprintEditorUtils::GetTopLevelGraph(Variable->GetGraph());
-			Variable->VariableReference.SetLocalMember( FName(*VariableName), FunctionGraph->GetName(), FBlueprintEditorUtils::FindLocalVariableGuidByName(OwnerBlueprint, FunctionGraph, *VariableName));
+			Variable->VariableReference.SetLocalMember( VariableName, FunctionGraph->GetName(), FBlueprintEditorUtils::FindLocalVariableGuidByName(OwnerBlueprint, FunctionGraph, VariableName));
 		}
 		Pin->PinName = VariableName;
 		Variable->ReconstructNode();
@@ -1802,7 +1800,7 @@ void UEdGraphSchema_K2::OnReplaceVariableForVariableNode( UK2Node_Variable* Vari
 
 void UEdGraphSchema_K2::GetReplaceVariableMenu(FMenuBuilder& MenuBuilder, UK2Node_Variable* Variable,  UBlueprint* OwnerBlueprint, bool bReplaceExistingVariable/*=false*/)
 {
-	if (UEdGraphPin* Pin = Variable->FindPin(Variable->GetVarNameString()))
+	if (UEdGraphPin* Pin = Variable->FindPin(Variable->GetVarName()))
 	{
 		FName ExistingVariableName = bReplaceExistingVariable? Variable->GetVarName() : NAME_None;
 
@@ -1832,7 +1830,7 @@ void UEdGraphSchema_K2::GetReplaceVariableMenu(FMenuBuilder& MenuBuilder, UK2Nod
 				const FText Desc = FText::Format(ReplaceVariableWithTooltipFormat, TooltipArgs);
 
 				MenuBuilder.AddMenuEntry( AlternativeVar, Desc, FSlateIcon(), FUIAction(
-					FExecuteAction::CreateStatic(&UEdGraphSchema_K2::OnReplaceVariableForVariableNode, const_cast<UK2Node_Variable* >(Variable),OwnerBlueprint, (*VarIt).ToString(), /*bIsSelfMember=*/true ) ) );
+					FExecuteAction::CreateStatic(&UEdGraphSchema_K2::OnReplaceVariableForVariableNode, const_cast<UK2Node_Variable* >(Variable), OwnerBlueprint, *VarIt, /*bIsSelfMember=*/true ) ) );
 			}
 		}
 		MenuBuilder.EndSection();
@@ -1863,7 +1861,7 @@ void UEdGraphSchema_K2::GetReplaceVariableMenu(FMenuBuilder& MenuBuilder, UK2Nod
 				const FText Desc = FText::Format( ReplaceLocalVariableWithTooltipFormat, TooltipArgs );
 
 				MenuBuilder.AddMenuEntry( AlternativeVar, Desc, FSlateIcon(), FUIAction(
-					FExecuteAction::CreateStatic(&UEdGraphSchema_K2::OnReplaceVariableForVariableNode, const_cast<UK2Node_Variable* >(Variable),OwnerBlueprint, (*VarIt).ToString(), /*bIsSelfMember=*/false ) ) );
+					FExecuteAction::CreateStatic(&UEdGraphSchema_K2::OnReplaceVariableForVariableNode, const_cast<UK2Node_Variable* >(Variable),OwnerBlueprint, *VarIt, /*bIsSelfMember=*/false ) ) );
 			}
 		}
 		MenuBuilder.EndSection();
@@ -1913,7 +1911,7 @@ void UEdGraphSchema_K2::GetReplaceVariableMenu(const UEdGraphNode* InGraphNode,
 	if (const UK2Node_Variable* Variable = Cast<const UK2Node_Variable>(InGraphNode))
 	{
 		// replace with matching variables
-		if (UEdGraphPin* Pin = Variable->FindPin(Variable->GetVarNameString()))
+		if (UEdGraphPin* Pin = Variable->FindPin(Variable->GetVarName()))
 		{
 			FName ExistingVariableName = bInReplaceExistingVariable? Variable->GetVarName() : NAME_None;
 
@@ -2345,13 +2343,13 @@ private:
 	static FString GenerateTypeData(const FEdGraphPinType& PinType)
 	{
 		UObject* Obj = PinType.PinSubCategoryObject.Get();
-		FString PinSubCategory = PinType.PinSubCategory;
-		if (PinSubCategory.StartsWith(UEdGraphSchema_K2::PSC_Bitmask))
+		FString PinSubCategory = PinType.PinSubCategory.ToString();
+		if (PinSubCategory.StartsWith(UEdGraphSchema_K2::PSC_Bitmask.ToString()))
 		{
 			// Exclude the bitmask subcategory string from integral types so that autocast will work.
-			PinSubCategory = TEXT("");
+			PinSubCategory.Reset();
 		}
-		return FString::Printf(TEXT("%s;%s;%s"), *PinType.PinCategory, *PinSubCategory, Obj ? *Obj->GetPathName() : TEXT(""));
+		return FString::Printf(TEXT("%s;%s;%s"), *PinType.PinCategory.ToString(), *PinSubCategory, Obj ? *Obj->GetPathName() : TEXT(""));
 	}
 
 	static FString GenerateCastData(const FEdGraphPinType& InputPinType, const FEdGraphPinType& OutputPinType)
@@ -2884,38 +2882,38 @@ void UEdGraphSchema_K2::ClearPinWatch(UEdGraphPin const* Pin) const
 
 FLinearColor UEdGraphSchema_K2::GetPinTypeColor(const FEdGraphPinType& PinType) const
 {
-	const FString& TypeString = PinType.PinCategory;
+	const FName TypeName = PinType.PinCategory;
 	const UGraphEditorSettings* Settings = GetDefault<UGraphEditorSettings>();
 
-	if (TypeString == PC_Exec)
+	if (TypeName == PC_Exec)
 	{
 		return Settings->ExecutionPinTypeColor;
 	}
-	else if (TypeString == PC_Object)
+	else if (TypeName == PC_Object)
 	{
 		return Settings->ObjectPinTypeColor;
 	}
-	else if (TypeString == PC_Interface)
+	else if (TypeName == PC_Interface)
 	{
 		return Settings->InterfacePinTypeColor;
 	}
-	else if (TypeString == PC_Float)
+	else if (TypeName == PC_Float)
 	{
 		return Settings->FloatPinTypeColor;
 	}
-	else if (TypeString == PC_Boolean)
+	else if (TypeName == PC_Boolean)
 	{
 		return Settings->BooleanPinTypeColor;
 	}
-	else if (TypeString == PC_Byte)
+	else if (TypeName == PC_Byte)
 	{
 		return Settings->BytePinTypeColor;
 	}
-	else if (TypeString == PC_Int)
+	else if (TypeName == PC_Int)
 	{
 		return Settings->IntPinTypeColor;
 	}
-	else if (TypeString == PC_Struct)
+	else if (TypeName == PC_Struct)
 	{
 		if (PinType.PinSubCategoryObject == VectorStruct)
 		{
@@ -2937,15 +2935,15 @@ FLinearColor UEdGraphSchema_K2::GetPinTypeColor(const FEdGraphPinType& PinType)
 			return Settings->StructPinTypeColor;
 		}
 	}
-	else if (TypeString == PC_String)
+	else if (TypeName == PC_String)
 	{
 		return Settings->StringPinTypeColor;
 	}
-	else if (TypeString == PC_Text)
+	else if (TypeName == PC_Text)
 	{
 		return Settings->TextPinTypeColor;
 	}
-	else if (TypeString == PC_Wildcard)
+	else if (TypeName == PC_Wildcard)
 	{
 		if (PinType.PinSubCategory == PSC_Index)
 		{
@@ -2956,23 +2954,23 @@ FLinearColor UEdGraphSchema_K2::GetPinTypeColor(const FEdGraphPinType& PinType)
 			return Settings->WildcardPinTypeColor;
 		}
 	}
-	else if (TypeString == PC_Name)
+	else if (TypeName == PC_Name)
 	{
 		return Settings->NamePinTypeColor;
 	}
-	else if (TypeString == PC_SoftObject)
+	else if (TypeName == PC_SoftObject)
 	{
 		return Settings->SoftObjectPinTypeColor;
 	}
-	else if (TypeString == PC_SoftClass)
+	else if (TypeName == PC_SoftClass)
 	{
 		return Settings->SoftClassPinTypeColor;
 	}
-	else if (TypeString == PC_Delegate)
+	else if (TypeName == PC_Delegate)
 	{
 		return Settings->DelegatePinTypeColor;
 	}
-	else if (TypeString == PC_Class)
+	else if (TypeName == PC_Class)
 	{
 		return Settings->ClassPinTypeColor;
 	}
@@ -2994,7 +2992,6 @@ FLinearColor UEdGraphSchema_K2::GetSecondaryPinTypeColor(const FEdGraphPinType&
 	}
 	else
 	{
-		const FString& TypeString = PinType.PinCategory;
 		const UGraphEditorSettings* Settings = GetDefault<UGraphEditorSettings>();
 		return Settings->WildcardPinTypeColor;
 	}
@@ -3004,7 +3001,7 @@ FText UEdGraphSchema_K2::GetPinDisplayName(const UEdGraphPin* Pin) const
 {
 	FText DisplayName = FText::GetEmpty();
 
-	if (Pin != NULL)
+	if (Pin)
 	{
 		UEdGraphNode* Node = Pin->GetOwningNode();
 		if (Node->ShouldOverridePinNames())
@@ -3016,12 +3013,15 @@ FText UEdGraphSchema_K2::GetPinDisplayName(const UEdGraphPin* Pin) const
 			DisplayName = Super::GetPinDisplayName(Pin);
 	
 			// bit of a hack to hide 'execute' and 'then' pin names
-			if ((Pin->PinType.PinCategory == PC_Exec) && 
-				((DisplayName.ToString() == PN_Execute) || (DisplayName.ToString() == PN_Then)))
+			if (Pin->PinType.PinCategory == PC_Exec)
+			{
+				FName DisplayFName(*DisplayName.ToString(), FNAME_Find);
+				if ((DisplayFName == PN_Execute) || (DisplayFName == PN_Then))
 			{
 				DisplayName = FText::GetEmpty();
 			}
 		}
+	}
 
 		if( GEditor && GetDefault<UEditorStyleSettings>()->bShowFriendlyNames )
 		{
@@ -3295,7 +3295,7 @@ void UEdGraphSchema_K2::CreateFunctionGraphTerminators(UEdGraph& Graph, UFunctio
 	}
 }
 
-bool UEdGraphSchema_K2::GetPropertyCategoryInfo(const UProperty* TestProperty, FString& OutCategory, FString& OutSubCategory, UObject*& OutSubCategoryObject, bool& bOutIsWeakPointer)
+bool UEdGraphSchema_K2::GetPropertyCategoryInfo(const UProperty* TestProperty, FName& OutCategory, FName& OutSubCategory, UObject*& OutSubCategoryObject, bool& bOutIsWeakPointer)
 {
 	if (const UInterfaceProperty* InterfaceProperty = Cast<const UInterfaceProperty>(TestProperty))
 	{
@@ -3401,13 +3401,13 @@ bool UEdGraphSchema_K2::GetPropertyCategoryInfo(const UProperty* TestProperty, F
 
 bool UEdGraphSchema_K2::ConvertPropertyToPinType(const UProperty* Property, /*out*/ FEdGraphPinType& TypeOut) const
 {
-	if (Property == NULL)
+	if (Property == nullptr)
 	{
 		TypeOut.PinCategory = TEXT("bad_type");
 		return false;
 	}
 
-	TypeOut.PinSubCategory = TEXT("");
+	TypeOut.PinSubCategory = NAME_None;
 	
 	// Handle whether or not this is an array property
 	const UMapProperty* MapProperty = Cast<const UMapProperty>(Property);
@@ -3496,9 +3496,9 @@ bool UEdGraphSchema_K2::IsWildcardProperty(const UProperty* Property)
 	UFunction* Function = Cast<UFunction>(Property->GetOuter());
 
 	return Function && ( UK2Node_CallArrayFunction::IsWildcardProperty(Function, Property)
-		|| UK2Node_CallFunction::IsStructureWildcardProperty(Function, Property->GetName())
+		|| UK2Node_CallFunction::IsStructureWildcardProperty(Function, Property->GetFName())
 		|| UK2Node_CallFunction::IsWildcardProperty(Function, Property)
-		|| FEdGraphUtilities::IsArrayDependentParam(Function, Property->GetName()) );
+		|| FEdGraphUtilities::IsArrayDependentParam(Function, Property->GetFName()) );
 }
 
 FText UEdGraphSchema_K2::TypeToText(UProperty* const Property)
@@ -3578,14 +3578,14 @@ FText UEdGraphSchema_K2::TypeToText(UProperty* const Property)
 	return FText::FromString(Property->GetClass()->GetName());
 }
 
-FText UEdGraphSchema_K2::GetCategoryText(const FString& Category, const bool bForMenu)
+FText UEdGraphSchema_K2::GetCategoryText(const FName Category, const bool bForMenu)
 {
-	if (Category.IsEmpty())
+	if (Category.IsNone())
 	{
 		return FText::GetEmpty();
 	}
 
-	static TMap<FString, FText> CategoryDescriptions;
+	static TMap<FName, FText> CategoryDescriptions;
 	if (CategoryDescriptions.Num() == 0)
 	{
 		CategoryDescriptions.Add(PC_Exec, LOCTEXT("Exec", "Exec"));
@@ -3615,18 +3615,17 @@ FText UEdGraphSchema_K2::GetCategoryText(const FString& Category, const bool bFo
 	}
 	else
 	{
-		return FText::FromString(Category);
+		return FText::FromName(Category);
 	}
 }
 
-FText UEdGraphSchema_K2::TerminalTypeToText(const FString& Category, const FString& SubCategory, UObject* SubCategoryObject, bool bIsWeakPtr)
+FText UEdGraphSchema_K2::TerminalTypeToText(const FName Category, const FName SubCategory, UObject* SubCategoryObject, bool bIsWeakPtr)
 {
 	FText PropertyText;
 
-	const UEdGraphSchema_K2* Schema = GetDefault<UEdGraphSchema_K2>();
-	if (SubCategory != Schema->PSC_Bitmask && SubCategoryObject != NULL)
+	if (SubCategory != UEdGraphSchema_K2::PSC_Bitmask && SubCategoryObject != nullptr)
 	{
-		if (Category == Schema->PC_Byte)
+		if (Category == UEdGraphSchema_K2::PC_Byte)
 		{
 			FFormatNamedArguments Args;
 			Args.Add(TEXT("EnumName"), FText::FromString(SubCategoryObject->GetName()));
@@ -3646,23 +3645,23 @@ FText UEdGraphSchema_K2::TerminalTypeToText(const FString& Category, const FStri
 				const bool bIsInterface = PSCOAsClass && PSCOAsClass->HasAnyClassFlags(CLASS_Interface);
 
 				FFormatNamedArguments Args;
+				Args.Add(TEXT("ObjectName"), FText::FromString(FName::NameToDisplayString(SubCategoryObjName, /*bIsBool =*/false)));
+
 				// Don't display the category for "well-known" struct types
 				if (Category == UEdGraphSchema_K2::PC_Struct && (SubCategoryObject == UEdGraphSchema_K2::VectorStruct || SubCategoryObject == UEdGraphSchema_K2::RotatorStruct || SubCategoryObject == UEdGraphSchema_K2::TransformStruct))
 				{
-					Args.Add(TEXT("Category"), FText::GetEmpty());
+					PropertyText = FText::Format(LOCTEXT("ObjectAsTextWithoutCategory", "{ObjectName}"), Args);
 				}
 				else
 				{
 					Args.Add(TEXT("Category"), (!bIsInterface ? UEdGraphSchema_K2::GetCategoryText(Category) : UEdGraphSchema_K2::GetCategoryText(PC_Interface)));
+					PropertyText = FText::Format(LOCTEXT("ObjectAsText", "{ObjectName} {Category}"), Args);
 				}
-
-				Args.Add(TEXT("ObjectName"), FText::FromString(FName::NameToDisplayString(SubCategoryObjName, /*bIsBool =*/false)));
-				PropertyText = FText::Format(LOCTEXT("ObjectAsText", "{ObjectName} {Category}"), Args);
 			}
 			else
 			{
 				FFormatNamedArguments Args;
-				Args.Add(TEXT("Category"), FText::FromString(Category));
+				Args.Add(TEXT("Category"), FText::FromName(Category));
 				Args.Add(TEXT("ObjectName"), FText::FromString(SubCategoryObjName));
 				PropertyText = FText::Format(LOCTEXT("WeakPtrAsText", "{ObjectName} Weak {Category}"), Args);
 			}
@@ -3672,7 +3671,7 @@ FText UEdGraphSchema_K2::TerminalTypeToText(const FString& Category, const FStri
 	{
 		FFormatNamedArguments Args;
 		Args.Add(TEXT("Category"), UEdGraphSchema_K2::GetCategoryText(Category));
-		Args.Add(TEXT("ObjectName"), FText::FromString(FName::NameToDisplayString(SubCategory, false)));
+		Args.Add(TEXT("ObjectName"), FText::FromString(FName::NameToDisplayString(SubCategory.ToString(), false)));
 		PropertyText = FText::Format(LOCTEXT("ObjectAsText", "{ObjectName} {Category}"), Args);
 	}
 	else
@@ -3813,7 +3812,7 @@ void UEdGraphSchema_K2::GetVariableTypeTree(TArray< TSharedPtr<FPinTypeTreeInfo>
 #undef SCHEMA_K2_GETVARIABLETYPETREE_LOG_TIME
 }
 
-bool UEdGraphSchema_K2::DoesTypeHaveSubtypes(const FString& Category) const
+bool UEdGraphSchema_K2::DoesTypeHaveSubtypes(const FName Category) const
 {
 	return (Category == PC_Struct) || (Category == PC_Object) || (Category == PC_SoftObject) || (Category == PC_SoftClass) || (Category == PC_Interface) || (Category == PC_Class) || (Category == PC_Enum) || (Category == AllObjectTypes);
 }
@@ -3834,7 +3833,7 @@ struct FWildcardArrayPinHelper
 			return true;
 		}
 
-		const bool bCheckInputPin = (InputPin->PinType.PinCategory == GetDefault<UEdGraphSchema_K2>()->PC_Wildcard) && !InputPin->PinType.IsArray();
+		const bool bCheckInputPin = (InputPin->PinType.PinCategory == UEdGraphSchema_K2::PC_Wildcard) && !InputPin->PinType.IsArray();
 		const bool bArrayOutputPin = OutputPin && OutputPin->PinType.IsArray();
 		return !(bCheckInputPin && bArrayOutputPin);
 	}
@@ -3854,8 +3853,8 @@ bool UEdGraphSchema_K2::ArePinsCompatible(const UEdGraphPin* PinA, const UEdGrap
 	}
 	else
 	{
-	return false;
-}
+		return false;
+	}
 }
 
 namespace
@@ -3950,15 +3949,15 @@ namespace
 };
 
 
-bool UEdGraphSchema_K2::DefaultValueSimpleValidation(const FEdGraphPinType& PinType, const FString& PinName, const FString& NewDefaultValue, UObject* NewDefaultObject, const FText& InNewDefaultText, FString* OutMsg /*= NULL*/) const
+bool UEdGraphSchema_K2::DefaultValueSimpleValidation(const FEdGraphPinType& PinType, const FName PinName, const FString& NewDefaultValue, UObject* NewDefaultObject, const FText& InNewDefaultText, FString* OutMsg /*= NULL*/) const
 {
 #ifdef DVSV_RETURN_MSG
 	static_assert(false, "Macro redefinition.");
 #endif
-#define DVSV_RETURN_MSG(Str) if(NULL != OutMsg) { *OutMsg = Str; } return false;
+#define DVSV_RETURN_MSG(Str) if (OutMsg) { *OutMsg = Str; } return false;
 
-	const FString& PinCategory = PinType.PinCategory;
-	const FString& PinSubCategory = PinType.PinSubCategory;
+	const FName PinCategory = PinType.PinCategory;
+	const FName PinSubCategory = PinType.PinSubCategory;
 	const UObject* PinSubCategoryObject = PinType.PinSubCategoryObject.Get();
 
 	if (PinType.IsContainer())
@@ -4000,10 +3999,10 @@ bool UEdGraphSchema_K2::DefaultValueSimpleValidation(const FEdGraphPinType& PinT
 		// Should have an object set but no string
 		if (!NewDefaultValue.IsEmpty())
 		{
-			DVSV_RETURN_MSG(FString::Printf(TEXT("String NewDefaultValue '%s' specified on class pin '%s'"), *NewDefaultValue, *(PinName)));
+			DVSV_RETURN_MSG(FString::Printf(TEXT("String NewDefaultValue '%s' specified on class pin '%s'"), *NewDefaultValue, *PinName.ToString()));
 		}
 
-		if (NewDefaultObject == NULL)
+		if (NewDefaultObject == nullptr)
 		{
 			// Valid self-reference or empty reference
 		}
@@ -4011,24 +4010,24 @@ bool UEdGraphSchema_K2::DefaultValueSimpleValidation(const FEdGraphPinType& PinT
 		{
 			// Otherwise, we expect to be able to resolve the type at least
 			const UClass* DefaultClassType = Cast<const UClass>(NewDefaultObject);
-			if (DefaultClassType == NULL)
+			if (DefaultClassType == nullptr)
 			{
-				DVSV_RETURN_MSG(FString::Printf(TEXT("Literal on pin %s is not a class."), *(PinName)));
+				DVSV_RETURN_MSG(FString::Printf(TEXT("Literal on pin %s is not a class."), *PinName.ToString()));
 			}
 			else
 			{
 				// @TODO support PinSubCategory == 'self'
 				const UClass* PinClassType = Cast<const UClass>(PinSubCategoryObject);
-				if (PinClassType == NULL)
+				if (PinClassType == nullptr)
 				{
-					DVSV_RETURN_MSG(FString::Printf(TEXT("Failed to find class for pin %s"), *(PinName)));
+					DVSV_RETURN_MSG(FString::Printf(TEXT("Failed to find class for pin %s"), *PinName.ToString()));
 				}
 				else
 				{
 					// Have both types, make sure the specified type is a valid subtype
 					if (!IsAuthoritativeChildOf(DefaultClassType, PinClassType))
 					{
-						DVSV_RETURN_MSG(FString::Printf(TEXT("%s isn't a valid subclass of %s (specified on pin %s)"), *NewDefaultObject->GetPathName(), *PinClassType->GetName(), *(PinName)));
+						DVSV_RETURN_MSG(FString::Printf(TEXT("%s isn't a valid subclass of %s (specified on pin %s)"), *NewDefaultObject->GetPathName(), *PinClassType->GetName(), *PinName.ToString()));
 					}
 				}
 			}
@@ -4060,20 +4059,20 @@ bool UEdGraphSchema_K2::DefaultValueSimpleValidation(const FEdGraphPinType& PinT
 	}
 	else if ((PinCategory == PC_Object) || (PinCategory == PC_Interface))
 	{
-		if (PinSubCategoryObject == NULL && (PinSubCategory != PSC_Self))
+		if (PinSubCategoryObject == nullptr && (PinSubCategory != PSC_Self))
 		{
-			DVSV_RETURN_MSG(FString::Printf(TEXT("PinSubCategoryObject on pin '%s' is NULL and PinSubCategory is '%s' not 'self'"), *(PinName), *PinSubCategory));
+			DVSV_RETURN_MSG(FString::Printf(TEXT("PinSubCategoryObject on pin '%s' is NULL and PinSubCategory is '%s' not 'self'"), *PinName.ToString(), *PinSubCategory.ToString()));
 		}
 
-		if (PinSubCategoryObject != NULL && PinSubCategory != TEXT(""))
+		if (PinSubCategoryObject != nullptr && !PinSubCategory.IsNone())
 		{
-			DVSV_RETURN_MSG(FString::Printf(TEXT("PinSubCategoryObject on pin '%s' is non-NULL but PinSubCategory is '%s', should be empty"), *(PinName), *PinSubCategory));
+			DVSV_RETURN_MSG(FString::Printf(TEXT("PinSubCategoryObject on pin '%s' is non-NULL but PinSubCategory is '%s', should be empty"), *PinName.ToString(), *PinSubCategory.ToString()));
 		}
 
 		// Should have an object set but no string - 'self' is not a valid NewDefaultValue for PC_Object pins
 		if (!NewDefaultValue.IsEmpty())
 		{
-			DVSV_RETURN_MSG(FString::Printf(TEXT("String NewDefaultValue '%s' specified on object pin '%s'"), *NewDefaultValue, *(PinName)));
+			DVSV_RETURN_MSG(FString::Printf(TEXT("String NewDefaultValue '%s' specified on object pin '%s'"), *NewDefaultValue, *PinName.ToString()));
 		}
 
 		// Check that the object that is set is of the correct class
@@ -4082,9 +4081,9 @@ bool UEdGraphSchema_K2::DefaultValueSimpleValidation(const FEdGraphPinType& PinT
 		{
 			ObjectClass = ObjectClass->GetAuthoritativeClass();
 		}
-		if (NewDefaultObject != NULL && ObjectClass != NULL && !NewDefaultObject->GetClass()->GetAuthoritativeClass()->IsChildOf(ObjectClass))
+		if (NewDefaultObject != nullptr && ObjectClass != nullptr && !NewDefaultObject->GetClass()->GetAuthoritativeClass()->IsChildOf(ObjectClass))
 		{
-			DVSV_RETURN_MSG(FString::Printf(TEXT("%s isn't a %s (specified on pin %s)"), *NewDefaultObject->GetPathName(), *ObjectClass->GetName(), *(PinName)));
+			DVSV_RETURN_MSG(FString::Printf(TEXT("%s isn't a %s (specified on pin %s)"), *NewDefaultObject->GetPathName(), *ObjectClass->GetName(), *PinName.ToString()));
 		}
 	}
 	else if ((PinCategory == PC_SoftObject) || (PinCategory == PC_SoftClass))
@@ -4092,7 +4091,7 @@ bool UEdGraphSchema_K2::DefaultValueSimpleValidation(const FEdGraphPinType& PinT
 		// Should not have an object set, should be converted to string before getting here
 		if (NewDefaultObject)
 		{
-			DVSV_RETURN_MSG(FString::Printf(TEXT("NewDefaultObject '%s' specified on object pin '%s'"), *NewDefaultObject->GetPathName(), *(PinName)));
+			DVSV_RETURN_MSG(FString::Printf(TEXT("NewDefaultObject '%s' specified on object pin '%s'"), *NewDefaultObject->GetPathName(), *PinName.ToString()));
 		}
 
 		if (!NewDefaultValue.IsEmpty())
@@ -4101,7 +4100,7 @@ bool UEdGraphSchema_K2::DefaultValueSimpleValidation(const FEdGraphPinType& PinT
 
 			if (!FPackageName::IsValidObjectPath(NewDefaultValue, &PathReason))
 			{
-				DVSV_RETURN_MSG(FString::Printf(TEXT("Soft Reference '%s' is invalid format for object pin '%s':"), *NewDefaultValue, *(PinName), *PathReason.ToString()));
+				DVSV_RETURN_MSG(FString::Printf(TEXT("Soft Reference '%s' is invalid format for object pin '%s':"), *NewDefaultValue, *PinName.ToString(), *PathReason.ToString()));
 			}
 
 			// Class and IsAsset validation is not foolproof for soft references, skip
@@ -4121,17 +4120,17 @@ bool UEdGraphSchema_K2::DefaultValueSimpleValidation(const FEdGraphPinType& PinT
 	}
 	else if (PinCategory == PC_Struct)
 	{
-		if (PinSubCategory != TEXT(""))
+		if (!PinSubCategory.IsNone())
 		{
-			DVSV_RETURN_MSG(FString::Printf(TEXT("Invalid PinSubCategory value '%s' (it should be empty)"), *PinSubCategory));
+			DVSV_RETURN_MSG(FString::Printf(TEXT("Invalid PinSubCategory value '%s' (it should be empty)"), *PinSubCategory.ToString()));
 		}
 
 		// Only FRotator and FVector properties are currently allowed to have a valid default value
 		const UScriptStruct* StructType = Cast<const UScriptStruct>(PinSubCategoryObject);
-		if (StructType == NULL)
+		if (StructType == nullptr)
 		{
 			//@TODO: MessageLog.Error(*FString::Printf(TEXT("Failed to find struct named %s (passed thru @@)"), *PinSubCategory), SourceObject);
-			DVSV_RETURN_MSG(FString::Printf(TEXT("No struct specified for pin '%s'"), *(PinName)));
+			DVSV_RETURN_MSG(FString::Printf(TEXT("No struct specified for pin '%s'"), *PinName.ToString()));
 		}
 		else if (!NewDefaultValue.IsEmpty())
 		{
@@ -4181,7 +4180,7 @@ bool UEdGraphSchema_K2::DefaultValueSimpleValidation(const FEdGraphPinType& PinT
 	else
 	{
 		//@TODO: MessageLog.Error(*FString::Printf(TEXT("Unsupported type %s on @@"), *UEdGraphSchema_K2::TypeToText(Type).ToString()), SourceObject);
-		DVSV_RETURN_MSG(FString::Printf(TEXT("Unsupported type %s on pin %s"), *UEdGraphSchema_K2::TypeToText(PinType).ToString(), *(PinName)));
+		DVSV_RETURN_MSG(FString::Printf(TEXT("Unsupported type %s on pin %s"), *UEdGraphSchema_K2::TypeToText(PinType).ToString(), *PinName.ToString()));
 	}
 
 #undef DVSV_RETURN_MSG
@@ -4286,7 +4285,8 @@ bool UEdGraphSchema_K2::ArePinTypesCompatible(const FEdGraphPinType& Output, con
 		else if (PC_Byte == Output.PinCategory || PC_Int == Output.PinCategory)
 		{
 			// Bitmask integral types are compatible with non-bitmask integral types (of the same word size).
-			return Output.PinSubCategory.StartsWith(PSC_Bitmask) || Input.PinSubCategory.StartsWith(PSC_Bitmask);
+			const FString PSC_Bitmask_Str = PSC_Bitmask.ToString();
+			return Output.PinSubCategory.ToString().StartsWith(PSC_Bitmask_Str) || Input.PinSubCategory.ToString().StartsWith(PSC_Bitmask_Str);
 		}
 		else if (PC_Delegate == Output.PinCategory || PC_MCDelegate == Output.PinCategory)
 		{
@@ -4379,7 +4379,7 @@ void UEdGraphSchema_K2::BreakPinLinks(UEdGraphPin& TargetPin, bool bSendsNodeNot
 	FBlueprintEditorUtils::MarkBlueprintAsModified(Blueprint);
 }
 
-void UEdGraphSchema_K2::BreakSinglePinLink(UEdGraphPin* SourcePin, UEdGraphPin* TargetPin)
+void UEdGraphSchema_K2::BreakSinglePinLink(UEdGraphPin* SourcePin, UEdGraphPin* TargetPin) const
 {
 	const FScopedTransaction Transaction( NSLOCTEXT("UnrealEd", "GraphEd_BreakSinglePinLink", "Break Pin Link") );
 
@@ -4692,11 +4692,11 @@ bool UEdGraphSchema_K2::IsAutoCreateRefTerm(const UEdGraphPin* Pin) const
 	if (FuncNode)
 	{
 		UFunction* TargetFunction = FuncNode->GetTargetFunction();
-		if (TargetFunction && !Pin->PinName.IsEmpty())
+		if (TargetFunction && !Pin->PinName.IsNone())
 		{
 			TArray<FString> AutoCreateParameterNames;
 			GetAutoEmitTermParameters(TargetFunction, AutoCreateParameterNames);
-			bIsAutoCreateRefTerm = AutoCreateParameterNames.Contains(Pin->PinName);
+			bIsAutoCreateRefTerm = AutoCreateParameterNames.Contains(Pin->PinName.ToString());
 		}
 	}
 
@@ -5056,7 +5056,7 @@ UFunction* UEdGraphSchema_K2::FindSetVariableByNameFunction(const FEdGraphPinTyp
 	return Function;
 }
 
-bool UEdGraphSchema_K2::CanPromotePinToVariable( const UEdGraphPin& Pin ) const
+bool UEdGraphSchema_K2::CanPromotePinToVariable( const UEdGraphPin& Pin, const bool bInToMemberVariable ) const
 {
 	const FEdGraphPinType& PinType = Pin.PinType;
 	bool bCanPromote = (PinType.PinCategory != PC_Wildcard && PinType.PinCategory != PC_Exec ) ? true : false;
@@ -5068,7 +5068,7 @@ bool UEdGraphSchema_K2::CanPromotePinToVariable( const UEdGraphPin& Pin ) const
 	{
 		bCanPromote = false;
 	}
-	else if (!OwningBlueprint || (OwningBlueprint->BlueprintType == BPTYPE_MacroLibrary) || (OwningBlueprint->BlueprintType == BPTYPE_FunctionLibrary) || IsStaticFunctionGraph(Node->GetGraph()))
+	else if (!OwningBlueprint || (OwningBlueprint->BlueprintType == BPTYPE_MacroLibrary) || (bInToMemberVariable && (OwningBlueprint->BlueprintType == BPTYPE_FunctionLibrary || IsStaticFunctionGraph(Node->GetGraph()))))
 	{
 		// Never allow promotion in macros, because there's not a scope to define them in
 		bCanPromote = false;
@@ -5255,7 +5255,7 @@ bool UEdGraphSchema_K2::IsSelfPin(const UEdGraphPin& Pin) const
 	return (Pin.PinName == PN_Self);
 }
 
-bool UEdGraphSchema_K2::IsDelegateCategory(const FString& Category) const
+bool UEdGraphSchema_K2::IsDelegateCategory(const FName Category) const
 {
 	return (Category == PC_Delegate);
 }
@@ -5658,7 +5658,7 @@ struct FBackwardCompatibilityConversionHelper
 				{
 					UE_LOG(LogBlueprint, Warning, TEXT("BackwardCompatibilityNodeConversion Error 'cannot connect' in blueprint: %s, pin: %s"),
 						Blueprint ? *Blueprint->GetName() : TEXT("Unknown"),
-						*ExecCastPin->PinName);
+						*ExecCastPin->PinName.ToString());
 					return false;
 				}
 
@@ -5668,7 +5668,7 @@ struct FBackwardCompatibilityConversionHelper
 				{
 					UE_LOG(LogBlueprint, Warning, TEXT("BackwardCompatibilityNodeConversion Error 'cannot connect' in blueprint: %s, pin: %s"),
 						Blueprint ? *Blueprint->GetName() : TEXT("Unknown"),
-						*ValidCastPin->PinName);
+						*ValidCastPin->PinName.ToString());
 					return false;
 				}
 
@@ -5678,7 +5678,7 @@ struct FBackwardCompatibilityConversionHelper
 				{
 					UE_LOG(LogBlueprint, Warning, TEXT("BackwardCompatibilityNodeConversion Error 'cannot connect' in blueprint: %s, pin: %s"),
 						Blueprint ? *Blueprint->GetName() : TEXT("Unknown"),
-						*InValidCastPin->PinName);
+						*InValidCastPin->PinName.ToString());
 					return false;
 				}
 
@@ -5688,7 +5688,7 @@ struct FBackwardCompatibilityConversionHelper
 				{
 					UE_LOG(LogBlueprint, Warning, TEXT("BackwardCompatibilityNodeConversion Error 'cannot connect' in blueprint: %s, pin: %s"),
 						Blueprint ? *Blueprint->GetName() : TEXT("Unknown"),
-						*CastSourcePin->PinName);
+						*CastSourcePin->PinName.ToString());
 					return false;
 				}
 
@@ -5698,7 +5698,7 @@ struct FBackwardCompatibilityConversionHelper
 				{
 					UE_LOG(LogBlueprint, Warning, TEXT("BackwardCompatibilityNodeConversion Error 'cannot connect' in blueprint: %s, pin: %s"),
 						Blueprint ? *Blueprint->GetName() : TEXT("Unknown"),
-						*CastResultPin->PinName);
+						*CastResultPin->PinName.ToString());
 					return false;
 				}
 			}
@@ -5719,14 +5719,14 @@ struct FBackwardCompatibilityConversionHelper
 						{
 							UE_LOG(LogBlueprint, Warning, TEXT("BackwardCompatibilityNodeConversion Error 'cannot connect' in blueprint: %s, pin: %s"),
 								Blueprint ? *Blueprint->GetName() : TEXT("Unknown"),
-								*Pin->PinName);
+								*Pin->PinName.ToString());
 						}
 					}
 					else
 					{
 						UE_LOG(LogBlueprint, Warning, TEXT("BackwardCompatibilityNodeConversion Error 'missing old pin' in blueprint: %s"),
 							Blueprint ? *Blueprint->GetName() : TEXT("Unknown"),
-							Pin ? *Pin->PinName : TEXT("Unknown"));
+							Pin ? *Pin->PinName.ToString() : TEXT("Unknown"));
 					}
 				}
 			}
@@ -5737,7 +5737,7 @@ struct FBackwardCompatibilityConversionHelper
 				{
 					UE_LOG(LogBlueprint, Warning, TEXT("BackwardCompatibilityNodeConversion Error 'missing new pin' in blueprint: %s"),
 						Blueprint ? *Blueprint->GetName() : TEXT("Unknown"),
-						Pin ? *Pin->PinName : TEXT("Unknown"));
+						Pin ? *Pin->PinName.ToString() : TEXT("Unknown"));
 				}
 			}
 			Graph->RemoveNode(OldNode);
@@ -5795,7 +5795,7 @@ struct FBackwardCompatibilityConversionHelper
 	}
 };
 
-bool UEdGraphSchema_K2::ReplaceOldNodeWithNew(UK2Node* OldNode, UK2Node* NewNode, const TMap<FString, FString>& OldPinToNewPinMap) const
+bool UEdGraphSchema_K2::ReplaceOldNodeWithNew(UK2Node* OldNode, UK2Node* NewNode, const TMap<FName, FName>& OldPinToNewPinMap) const
 {
 	if (!ensure(NewNode->GetGraph() == OldNode->GetGraph()))
 	{
@@ -5815,8 +5815,8 @@ bool UEdGraphSchema_K2::ReplaceOldNodeWithNew(UK2Node* OldNode, UK2Node* NewNode
 		UEdGraphPin* OldPin = OldNode->Pins[PinIdx];
 		UEdGraphPin* NewPin = nullptr;
 
-		const FString* NewPinNamePtr = OldPinToNewPinMap.Find(OldPin->PinName);
-		if (NewPinNamePtr && NewPinNamePtr->IsEmpty())
+		const FName* NewPinNamePtr = OldPinToNewPinMap.Find(OldPin->PinName);
+		if (NewPinNamePtr && NewPinNamePtr->IsNone())
 		{
 			// if they added an remapping for this pin, but left it empty, then it's assumed that they didn't want us to port any of the connections
 			NewPinArray.Add(nullptr);
@@ -5824,8 +5824,8 @@ bool UEdGraphSchema_K2::ReplaceOldNodeWithNew(UK2Node* OldNode, UK2Node* NewNode
 		}
 		else
 		{
-			const FString NewPinName = NewPinNamePtr ? *NewPinNamePtr : OldPin->PinName;
-			NewPin = NewNode->FindPin(*NewPinName);
+			const FName NewPinName = NewPinNamePtr ? *NewPinNamePtr : OldPin->PinName;
+			NewPin = NewNode->FindPin(NewPinName);
 
 			if (!NewPin && OldPin->ParentPin)
 			{
@@ -5846,18 +5846,23 @@ bool UEdGraphSchema_K2::ReplaceOldNodeWithNew(UK2Node* OldNode, UK2Node* NewNode
 							SplitPin(NewParent);
 						}
 
-						FString OldPinName = OldPin->PinName;
-						OldPinName.RemoveFromStart(OldParent->PinName);
-
-						for (UEdGraphPin* SubPin : NewParent->SubPins)
+						if (NewParent->SubPins.Num() > 0)
 						{
-							FString SubPinName = SubPin->PinName;
-							SubPinName.RemoveFromStart(NewParent->PinName);
+							FString OldPinName = OldPin->PinName.ToString();
+							OldPinName.RemoveFromStart(OldParent->PinName.ToString());
+
+							const FString NewParentNameStr = NewParent->PinName.ToString();
 
-							if (SubPinName == OldPinName)
+							for (UEdGraphPin* SubPin : NewParent->SubPins)
 							{
-								NewPin = SubPin;
-								break;
+								FString SubPinName = SubPin->PinName.ToString();
+								SubPinName.RemoveFromStart(NewParentNameStr);
+
+								if (SubPinName == OldPinName)
+								{
+									NewPin = SubPin;
+									break;
+								}
 							}
 						}
 					}
@@ -5870,7 +5875,7 @@ bool UEdGraphSchema_K2::ReplaceOldNodeWithNew(UK2Node* OldNode, UK2Node* NewNode
 			bFailedToFindPin = true;
 
 			UE_LOG(LogBlueprint, Warning, TEXT("BackwardCompatibilityNodeConversion Error 'cannot find pin %s in node %s' in blueprint: %s"),
-				*OldPin->PinName,
+				*OldPin->PinName.ToString(),
 				*NewNode->GetNodeTitle(ENodeTitleType::FullTitle).ToString(),
 				Blueprint ? *Blueprint->GetName() : TEXT("Unknown"));
 
@@ -5897,8 +5902,8 @@ bool UEdGraphSchema_K2::ReplaceOldNodeWithNew(UK2Node* OldNode, UK2Node* NewNode
 			else if (!Schema->MovePinLinks(*OldPin, *NewPin).CanSafeConnect())
 			{
 				UE_LOG(LogBlueprint, Warning, TEXT("BackwardCompatibilityNodeConversion Error 'cannot safely move pin %s to %s' in blueprint: %s"),
-					*OldPin->PinName,
-					*NewPin->PinName,
+					*OldPin->PinName.ToString(),
+					*NewPin->PinName.ToString(),
 					Blueprint ? *Blueprint->GetName() : TEXT("Unknown"));
 			}
 			else
@@ -5917,7 +5922,7 @@ bool UEdGraphSchema_K2::ReplaceOldNodeWithNew(UK2Node* OldNode, UK2Node* NewNode
 	return !bFailedToFindPin;
 }
 
-UK2Node* UEdGraphSchema_K2::ConvertDeprecatedNodeToFunctionCall(UK2Node* OldNode, UFunction* NewFunction, TMap<FString, FString>& OldPinToNewPinMap, UEdGraph* Graph) const
+UK2Node* UEdGraphSchema_K2::ConvertDeprecatedNodeToFunctionCall(UK2Node* OldNode, UFunction* NewFunction, TMap<FName, FName>& OldPinToNewPinMap, UEdGraph* Graph) const
 {
 	UK2Node_CallFunction* CallFunctionNode = NewObject<UK2Node_CallFunction>(Graph);
 	check(CallFunctionNode);
@@ -6454,7 +6459,7 @@ bool UEdGraphSchema_K2::CollapseGatewayNode(UK2Node* InNode, UEdGraphNode* InEnt
 					// We had an input/output with a connection that wasn't twinned
 					bSuccessful = false;
 					OwningBP->Message_Warn( FString::Printf(*NSLOCTEXT("K2Node", "PinOnBoundryNode_Warning", "Warning: Pin '%s' on boundary node '%s' could not be found in the composite node '%s'").ToString(),
-						*(BoundaryPin->PinName),
+						*(BoundaryPin->PinName.ToString()),
 						(GatewayNode ? *(GatewayNode->GetName()) : TEXT("(null)")),
 						*(GetName()))					
 						);
@@ -6462,7 +6467,7 @@ bool UEdGraphSchema_K2::CollapseGatewayNode(UK2Node* InNode, UEdGraphNode* InEnt
 				else
 				{
 					UE_LOG(LogBlueprint, Warning, TEXT("%s"), *FString::Printf(*NSLOCTEXT("K2Node", "PinOnBoundryNode_Warning", "Warning: Pin '%s' on boundary node '%s' could not be found in the composite node '%s'").ToString(),
-						*(BoundaryPin->PinName),
+						*(BoundaryPin->PinName.ToString()),
 						(GatewayNode ? *(GatewayNode->GetName()) : TEXT("(null)")),
 						*(GetName()))					
 						);
@@ -6688,9 +6693,12 @@ void UEdGraphSchema_K2::SplitPin(UEdGraphPin* Pin, const bool bNotify) const
 	{
 		if (ProtoPin->Direction == Pin->Direction && !ProtoPin->bHidden)
 		{
-			const FString PinName = FString::Printf(TEXT("%s_%s"), *Pin->PinName, *ProtoPin->PinName);
+			const FName PinName = *FString::Printf(TEXT("%s_%s"), *Pin->PinName.ToString(), *ProtoPin->PinName.ToString());
 			const FEdGraphPinType& ProtoPinType = ProtoPin->PinType;
-			UEdGraphPin* SubPin = GraphNode->CreatePin(Pin->Direction, ProtoPinType.PinCategory, ProtoPinType.PinSubCategory, ProtoPinType.PinSubCategoryObject.Get(), PinName, ProtoPinType.ContainerType, false, false, INDEX_NONE, ProtoPinType.PinValueType);
+			UEdGraphNode::FCreatePinParams PinParams;
+			PinParams.ContainerType = ProtoPinType.ContainerType;
+			PinParams.ValueTerminalType = ProtoPinType.PinValueType;
+			UEdGraphPin* SubPin = GraphNode->CreatePin(Pin->Direction, ProtoPinType.PinCategory, ProtoPinType.PinSubCategory, ProtoPinType.PinSubCategoryObject.Get(), PinName, PinParams);
 
 			if (K2Node != nullptr && K2Node->ShouldDrawCompact() && !Pin->ParentPin)
 			{
@@ -6704,8 +6712,8 @@ void UEdGraphSchema_K2::SplitPin(UEdGraphPin* Pin, const bool bNotify) const
 				SubPin->PinFriendlyName = FText::Format(LOCTEXT("SplitPinFriendlyNameFormat", "{PinDisplayName} {ProtoPinDisplayName}"), Arguments);
 			}
 
-			SubPin->DefaultValue = ProtoPin->DefaultValue;
-			SubPin->AutogeneratedDefaultValue = ProtoPin->AutogeneratedDefaultValue;
+			SubPin->DefaultValue = MoveTemp(ProtoPin->DefaultValue);
+			SubPin->AutogeneratedDefaultValue = MoveTemp(ProtoPin->AutogeneratedDefaultValue);
 
 			SubPin->ParentPin = Pin;
 
@@ -6938,7 +6946,7 @@ UK2Node_VariableSet* UEdGraphSchema_K2::SpawnVariableSetNode(const FVector2D Gra
 	return FEdGraphSchemaAction_K2NewNode::SpawnNodeFromTemplate<UK2Node_VariableSet>(ParentGraph, NodeTemplate, GraphPosition);
 }
 
-UEdGraphPin* UEdGraphSchema_K2::DropPinOnNode(UEdGraphNode* InTargetNode, const FString& InSourcePinName, const FEdGraphPinType& InSourcePinType, EEdGraphPinDirection InSourcePinDirection) const
+UEdGraphPin* UEdGraphSchema_K2::DropPinOnNode(UEdGraphNode* InTargetNode, const FName& InSourcePinName, const FEdGraphPinType& InSourcePinType, EEdGraphPinDirection InSourcePinDirection) const
 {
 	UEdGraphPin* ResultPin = nullptr;
 	if (UK2Node_EditablePinBase* EditablePinNode = Cast<UK2Node_EditablePinBase>(InTargetNode))
@@ -6985,11 +6993,10 @@ UEdGraphPin* UEdGraphSchema_K2::DropPinOnNode(UEdGraphNode* InTargetNode, const
 			}
 		}
 
-		FString NewPinName = InSourcePinName;
 		for (UK2Node_EditablePinBase* CurrentEditablePinNode : EditablePinNodes)
 		{
 			CurrentEditablePinNode->Modify();
-			UEdGraphPin* CreatedPin = CurrentEditablePinNode->CreateUserDefinedPin(NewPinName, InSourcePinType, (InSourcePinDirection == EGPD_Input) ? EGPD_Output : EGPD_Input);
+			UEdGraphPin* CreatedPin = CurrentEditablePinNode->CreateUserDefinedPin(InSourcePinName, InSourcePinType, (InSourcePinDirection == EGPD_Input) ? EGPD_Output : EGPD_Input);
 
 			// The final ResultPin is from the node the user dragged and dropped to
 			if (EditablePinNode == CurrentEditablePinNode)
diff --git a/Engine/Source/Editor/BlueprintGraph/Private/EdGraphSchema_K2_Actions.cpp b/Engine/Source/Editor/BlueprintGraph/Private/EdGraphSchema_K2_Actions.cpp
--- a/Engine/Source/Editor/BlueprintGraph/Private/EdGraphSchema_K2_Actions.cpp
+++ b/Engine/Source/Editor/BlueprintGraph/Private/EdGraphSchema_K2_Actions.cpp
@@ -555,8 +555,6 @@ UEdGraphNode* FEdGraphSchemaAction_K2AddCustomEvent::PerformAction(class UEdGrap
 
 UEdGraphNode* FEdGraphSchemaAction_K2AddCallOnActor::PerformAction(class UEdGraph* ParentGraph, UEdGraphPin* FromPin, const FVector2D Location, bool bSelectNewNode/* = true*/)
 {
-	const UEdGraphSchema_K2* K2Schema = GetDefault<UEdGraphSchema_K2>();
-
 	// Snap the node placement location to the grid, ensures calculations later match up better
 	FVector2D LocalLocation;
 	LocalLocation.X = FMath::GridSnap( Location.X, SNAP_GRID );
@@ -598,7 +596,7 @@ UEdGraphNode* FEdGraphSchemaAction_K2AddCallOnActor::PerformAction(class UEdGrap
 
 			// Connect the literal out to the self of the call
 			UEdGraphPin* LiteralOutput = LiteralNode->GetValuePin();
-			UEdGraphPin* CallSelfInput = CallNode->FindPin(K2Schema->PN_Self);
+			UEdGraphPin* CallSelfInput = CallNode->FindPin(UEdGraphSchema_K2::PN_Self);
 			if(LiteralOutput != nullptr && CallSelfInput != nullptr)
 			{
 				LiteralOutput->MakeLinkTo(CallSelfInput);
diff --git a/Engine/Source/Editor/BlueprintGraph/Private/K2Node.cpp b/Engine/Source/Editor/BlueprintGraph/Private/K2Node.cpp
--- a/Engine/Source/Editor/BlueprintGraph/Private/K2Node.cpp
+++ b/Engine/Source/Editor/BlueprintGraph/Private/K2Node.cpp
@@ -92,7 +92,7 @@ void UK2Node::Serialize(FArchive& Ar)
 
 						// Import the literal text to a dummy struct and then serialize that. Hard object references will not properly import, this is only useful for soft references!
 						FOutputDeviceNull NullOutput;
-						Struct->ImportText(*Pin->DefaultValue, StructData->GetStructMemory(), nullptr, PPF_SerializedAsImportText, &NullOutput, Pin->PinName);						
+						Struct->ImportText(*Pin->DefaultValue, StructData->GetStructMemory(), nullptr, PPF_SerializedAsImportText, &NullOutput, Pin->PinName.ToString());						
 						Struct->SerializeItem(Ar, StructData->GetStructMemory(), nullptr);
 					}
 				}
@@ -166,7 +166,7 @@ void UK2Node::FixupPinDefaultValues()
 		{
 			UEdGraphPin* Pin = Pins[i];
 
-			if (Pin->PinType.PinCategory == K2Schema->PC_SoftObject || Pin->PinType.PinCategory == K2Schema->PC_SoftClass)
+			if (Pin->PinType.PinCategory == UEdGraphSchema_K2::PC_SoftObject || Pin->PinType.PinCategory == UEdGraphSchema_K2::PC_SoftClass)
 			{
 				if (Pin->DefaultObject && Pin->DefaultValue.IsEmpty())
 				{
@@ -226,7 +226,7 @@ bool UK2Node::CreatePinsForFunctionEntryExit(const UFunction* Function, bool bFo
 		{
 			const EEdGraphPinDirection Direction = bForFunctionEntry ? EGPD_Output : EGPD_Input;
 
-			UEdGraphPin* Pin = CreatePin(Direction, FString(), FString(), nullptr, Param->GetName());
+			UEdGraphPin* Pin = CreatePin(Direction, NAME_None, Param->GetFName());
 			const bool bPinGood = K2Schema->ConvertPropertyToPinType(Param, /*out*/ Pin->PinType);
 			K2Schema->SetPinAutogeneratedDefaultValueBasedOnType(Pin);
 			
@@ -264,7 +264,7 @@ void UK2Node::AutowireNewNode(UEdGraphPin* FromPin)
 			}
 
 			ECanCreateConnectionResponse ConnectResponse = K2Schema->CanCreateConnection(FromPin, Pin).Response;
-			if ((FromPin->PinType.PinCategory == K2Schema->PC_Exec) && (ConnectResponse == ECanCreateConnectionResponse::CONNECT_RESPONSE_BREAK_OTHERS_A))
+			if ((FromPin->PinType.PinCategory == UEdGraphSchema_K2::PC_Exec) && (ConnectResponse == ECanCreateConnectionResponse::CONNECT_RESPONSE_BREAK_OTHERS_A))
 			{
 				InsertNewNode(FromPin, Pin, NodeList);
 
@@ -305,12 +305,12 @@ void UK2Node::AutowireNewNode(UEdGraphPin* FromPin)
 		}
 
 		// If we were not dragging an exec pin, but it was an output pin, try and connect the Then and Execute pins
-		if ((FromPin->PinType.PinCategory != K2Schema->PC_Exec  && FromPin->Direction == EGPD_Output))
+		if ((FromPin->PinType.PinCategory != UEdGraphSchema_K2::PC_Exec  && FromPin->Direction == EGPD_Output))
 		{
 			UEdGraphNode* FromPinNode = FromPin->GetOwningNode();
-			UEdGraphPin* FromThenPin = FromPinNode->FindPin(K2Schema->PN_Then);
+			UEdGraphPin* FromThenPin = FromPinNode->FindPin(UEdGraphSchema_K2::PN_Then);
 
-			UEdGraphPin* ToExecutePin = FindPin(K2Schema->PN_Execute);
+			UEdGraphPin* ToExecutePin = FindPin(UEdGraphSchema_K2::PN_Execute);
 
 			if ((FromThenPin != NULL) && (FromThenPin->LinkedTo.Num() == 0) && (ToExecutePin != NULL) && K2Schema->ArePinsCompatible(FromThenPin, ToExecutePin, NULL))
 			{
@@ -492,14 +492,13 @@ void UK2Node::PostReconstructNode()
 		UEdGraph* OuterGraph = GetGraph();
 		if( OuterGraph && OuterGraph->Schema )
 		{
-			const UEdGraphSchema_K2* Schema = Cast<const UEdGraphSchema_K2>(GetSchema());
 			// fix up any pin data if it needs to 
 			for (UEdGraphPin* CurrentPin : Pins)
 			{
-				const FString& PinCategory = CurrentPin->PinType.PinCategory;
+				const FName& PinCategory = CurrentPin->PinType.PinCategory;
 
 				// fix up enum names if it exists in enum redirects
-				if (PinCategory == Schema->PC_Byte)
+				if (PinCategory == UEdGraphSchema_K2::PC_Byte)
 				{
 					UEnum* EnumPtr = Cast<UEnum>(CurrentPin->PinType.PinSubCategoryObject.Get());
 					if (EnumPtr)
@@ -522,12 +521,12 @@ void UK2Node::PostReconstructNode()
 						}
 					}
 				}
-				else if (PinCategory == Schema->PC_Object)
+				else if (PinCategory == UEdGraphSchema_K2::PC_Object)
 				{
 					UClass const* PinClass = Cast<UClass const>(CurrentPin->PinType.PinSubCategoryObject.Get());
 					if ((PinClass != nullptr) && PinClass->IsChildOf(UInterface::StaticClass()))
 					{
-						CurrentPin->PinType.PinCategory = Schema->PC_Interface;
+						CurrentPin->PinType.PinCategory = UEdGraphSchema_K2::PC_Interface;
 					}
 				}
 			}
@@ -592,7 +591,7 @@ void UK2Node::ReconstructNode()
 
 void UK2Node::GetRedirectPinNames(const UEdGraphPin& Pin, TArray<FString>& RedirectPinNames) const
 {
-	RedirectPinNames.Add(Pin.PinName);
+	RedirectPinNames.Add(Pin.PinName.ToString());
 }
 
 UK2Node::ERedirectType UK2Node::ShouldRedirectParam(const TArray<FString>& OldPinNames, FName& NewPinName, const UK2Node * NewPinNode) const 
@@ -655,7 +654,7 @@ UK2Node::ERedirectType UK2Node::DoPinsMatchForReconstruction(const UEdGraphPin*
 	ERedirectType RedirectType = ERedirectType_None;
 
 	// if the pin names do match
-	if (FCString::Stricmp(*(NewPin->PinName), *(OldPin->PinName)) == 0)
+	if (NewPin->PinName == OldPin->PinName)
 	{
 		// Make sure we're not dealing with a menu node
 		UEdGraph* OuterGraph = GetGraph();
@@ -677,7 +676,7 @@ UK2Node::ERedirectType UK2Node::DoPinsMatchForReconstruction(const UEdGraphPin*
 		// try looking for a redirect if it's a K2 node
 		if (UK2Node* Node = Cast<UK2Node>(NewPin->GetOwningNode()))
 		{	
-			if (OldPin->ParentPin == NULL)
+			if (OldPin->ParentPin == nullptr)
 			{
 				// if you don't have matching pin, now check if there is any redirect param set
 				TArray<FString> OldPinNames;
@@ -687,7 +686,7 @@ UK2Node::ERedirectType UK2Node::DoPinsMatchForReconstruction(const UEdGraphPin*
 				RedirectType = ShouldRedirectParam(OldPinNames, /*out*/ NewPinName, Node);
 
 				// make sure they match
-				if ((RedirectType != ERedirectType_None) && FCString::Stricmp(*(NewPin->PinName), *(NewPinName.ToString())) != 0)
+				if ((RedirectType != ERedirectType_None) && (NewPin->PinName != NewPinName))
 				{
 					RedirectType = ERedirectType_None;
 				}
@@ -711,7 +710,7 @@ UK2Node::ERedirectType UK2Node::DoPinsMatchForReconstruction(const UEdGraphPin*
 					const UEdGraphPin* CurPin = OldPin;
 					do 
 					{
-						ParentHierarchy.Add(FPropertyDetails(CurPin, CurPin->PinName.RightChop(CurPin->ParentPin->PinName.Len() + 1)));
+						ParentHierarchy.Add(FPropertyDetails(CurPin, CurPin->PinName.ToString().RightChop(CurPin->ParentPin->PinName.ToString().Len() + 1)));
 						CurPin = CurPin->ParentPin;
 					} while (CurPin->ParentPin);
 
@@ -722,7 +721,7 @@ UK2Node::ERedirectType UK2Node::DoPinsMatchForReconstruction(const UEdGraphPin*
 					FName NewPinName;
 					RedirectType = ShouldRedirectParam(OldPinNames, /*out*/ NewPinName, Node);
 
-					NewPinNameStr = (RedirectType == ERedirectType_None ? CurPin->PinName : NewPinName.ToString());
+					NewPinNameStr = (RedirectType == ERedirectType_None ? CurPin->PinName.ToString() : NewPinName.ToString());
 				}
 
 				for (int32 ParentIndex = ParentHierarchy.Num() - 1; ParentIndex >= 0; --ParentIndex)
@@ -735,16 +734,16 @@ UK2Node::ERedirectType UK2Node::DoPinsMatchForReconstruction(const UEdGraphPin*
 
 					if (RedirectedPinName != NAME_Name)
 					{
-						NewPinNameStr += FString("_") + RedirectedPinName.ToString();
+						NewPinNameStr += FString(TEXT("_")) + RedirectedPinName.ToString();
 					}
 					else
 					{
-						NewPinNameStr += FString("_") + ParentHierarchy[ParentIndex].PropertyName;
+						NewPinNameStr += FString(TEXT("_")) + ParentHierarchy[ParentIndex].PropertyName;
 					}
 				}
 
 				// make sure they match
-				RedirectType = ((FCString::Stricmp(*(NewPin->PinName), *NewPinNameStr) != 0) ? ERedirectType_None : ERedirectType_Name);
+				RedirectType = ((NewPin->PinName.ToString() != NewPinNameStr) ? ERedirectType_None : ERedirectType_Name);
 			}
 		}
 	}
@@ -895,7 +894,7 @@ void UK2Node::RewireOldPinsToNewPins(TArray<UEdGraphPin*>& InOldPins, TArray<UEd
 		}
 
 		// Orphaned pins are those that existed in the OldPins array but do not in the NewPins.
-		// We will save these pins and add the to the NewPins array if they are linked to other pins or have non-default value unless:
+		// We will save these pins and add them to the NewPins array if they are linked to other pins or have non-default value unless:
 		// * The node has been flagged to not save orphaned pins
 		// * The pin has been flagged not be saved if orphaned
 		// * The pin is hidden and not a split pin
@@ -1111,16 +1110,16 @@ FLinearColor UK2Node::GetNodeTitleColor() const
 	return GetDefault<UGraphEditorSettings>()->FunctionCallNodeTitleColor;
 }
 
-ERenamePinResult UK2Node::RenameUserDefinedPin(const FString& OldName, const FString& NewName, bool bTest)
+ERenamePinResult UK2Node::RenameUserDefinedPin(const FName OldName, const FName NewName, bool bTest)
 {
-	UEdGraphPin* Pin = NULL;
-	for (int32 PinIdx=0; PinIdx<Pins.Num(); PinIdx++)
+	UEdGraphPin* Pin = nullptr;
+	for (UEdGraphPin* CurrentPin : Pins)
 	{
-		if (OldName == Pins[PinIdx]->PinName)
+		if (OldName == CurrentPin->PinName)
 		{
-			Pin = Pins[PinIdx];
+			Pin = CurrentPin;
 		}
-		else if(NewName == Pins[PinIdx]->PinName)
+		else if (NewName == CurrentPin->PinName)
 		{
 			return ERenamePinResult::ERenamePinResult_NameCollision;
 		}
@@ -1152,8 +1151,12 @@ ERenamePinResult UK2Node::RenameUserDefinedPin(const FString& OldName, const FSt
 					PinsToUpdate.Append(PinToRename->SubPins);
 				}
 				PinToRename->Modify();
-				PinToRename->PinName = NewName + PinToRename->PinName.RightChop(OldName.Len());
-				PinToRename->PinFriendlyName = FText::FromString(NewName + PinToRename->PinFriendlyName.ToString().RightChop(OldName.Len()));
+
+				const int32 OldNameLength = OldName.ToString().Len();
+				FString NewNameStr = NewName.ToString();
+
+				PinToRename->PinName = *(NewNameStr + PinToRename->PinName.ToString().RightChop(OldNameLength));
+				PinToRename->PinFriendlyName = FText::FromString(MoveTemp(NewNameStr) + PinToRename->PinFriendlyName.ToString().RightChop(OldNameLength));
 			}
 		}
 	}
@@ -1287,10 +1290,10 @@ void FOptionalPinManager::CreateVisiblePins(TArray<FOptionalPinFromProperty>& Pr
 						UEdGraphPin* NewPin = NULL;
 						if (PropertyEntry.bShowPin)
 						{
-							const FString PinName = FString::Printf(TEXT("%s_%d"), *PropertyEntry.PropertyName.ToString(), Index);
+							const FName PinName = *FString::Printf(TEXT("%s_%d"), *PropertyEntry.PropertyName.ToString(), Index);
 
 							FFormatNamedArguments Args;
-							Args.Add(TEXT("PinName"), FText::FromString(PropertyEntry.PropertyFriendlyName.IsEmpty() ? PinName : PropertyEntry.PropertyFriendlyName));
+							Args.Add(TEXT("PinName"), FText::FromString(PropertyEntry.PropertyFriendlyName.IsEmpty() ? PinName.ToString() : PropertyEntry.PropertyFriendlyName));
 							Args.Add(TEXT("Index"), Index);
 							const FText PinFriendlyName = FText::Format(LOCTEXT("PinFriendlyNameWithIndex", "{PinName} {Index}"), Args);
 
@@ -1339,9 +1342,9 @@ void FOptionalPinManager::CreateVisiblePins(TArray<FOptionalPinFromProperty>& Pr
 					UEdGraphPin* NewPin = nullptr;
 					if (PropertyEntry.bShowPin)
 					{
-						const FString PinName = PropertyEntry.PropertyName.ToString();
+						const FName PinName = PropertyEntry.PropertyName;
 						NewPin = TargetNode->CreatePin(Direction, PinType, PinName);
-						NewPin->PinFriendlyName = FText::FromString(PropertyEntry.PropertyFriendlyName.IsEmpty() ? PinName : PropertyEntry.PropertyFriendlyName);
+						NewPin->PinFriendlyName = FText::FromString(PropertyEntry.PropertyFriendlyName.IsEmpty() ? PinName.ToString() : PropertyEntry.PropertyFriendlyName);
 						NewPin->bNotConnectable = !PropertyEntry.bIsSetValuePinVisible;
 						NewPin->bDefaultValueIsIgnored = !PropertyEntry.bIsSetValuePinVisible;
 						Schema->ConstructBasicPinTooltip(*NewPin, PropertyEntry.PropertyTooltip, NewPin->PinToolTip);
@@ -1387,7 +1390,7 @@ void FOptionalPinManager::EvaluateOldShownPins(const TArray<FOptionalPinFromProp
 	{
 		if (ShowPinForProperty.bShowPin == false && OldShownPins.Contains(ShowPinForProperty.PropertyName))
 		{
-			if (UEdGraphPin* Pin = Node->FindPin(ShowPinForProperty.PropertyName.ToString()))
+			if (UEdGraphPin* Pin = Node->FindPin(ShowPinForProperty.PropertyName))
 			{
 				Pin->bSavePinIfOrphaned = false;
 			}
@@ -1398,10 +1401,8 @@ void FOptionalPinManager::EvaluateOldShownPins(const TArray<FOptionalPinFromProp
 
 UEdGraphPin* UK2Node::GetExecPin() const
 {
-	const UEdGraphSchema_K2* K2Schema = GetDefault<UEdGraphSchema_K2>();
-
-	UEdGraphPin* Pin = FindPin(K2Schema->PN_Execute);
-	check(Pin == NULL || Pin->Direction == EGPD_Input); // If pin exists, it must be input
+	UEdGraphPin* Pin = FindPin(UEdGraphSchema_K2::PN_Execute);
+	check(Pin == nullptr || Pin->Direction == EGPD_Input); // If pin exists, it must be input
 	return Pin;
 }
 
diff --git a/Engine/Source/Editor/BlueprintGraph/Private/K2Node_ActorBoundEvent.cpp b/Engine/Source/Editor/BlueprintGraph/Private/K2Node_ActorBoundEvent.cpp
--- a/Engine/Source/Editor/BlueprintGraph/Private/K2Node_ActorBoundEvent.cpp
+++ b/Engine/Source/Editor/BlueprintGraph/Private/K2Node_ActorBoundEvent.cpp
@@ -245,7 +245,7 @@ FMulticastScriptDelegate* UK2Node_ActorBoundEvent::GetTargetDelegate() const
 	if( EventOwner )
 	{
 		UMulticastDelegateProperty* TargetDelegateProp = GetTargetDelegateProperty();
-		if (TargetDelegateProp)
+		if (TargetDelegateProp && ensure(EventOwner->IsA(DelegateOwnerClass)))
 		{
 			return TargetDelegateProp->GetPropertyValuePtr_InContainer(EventOwner);
 		}
diff --git a/Engine/Source/Editor/BlueprintGraph/Private/K2Node_AddComponent.cpp b/Engine/Source/Editor/BlueprintGraph/Private/K2Node_AddComponent.cpp
--- a/Engine/Source/Editor/BlueprintGraph/Private/K2Node_AddComponent.cpp
+++ b/Engine/Source/Editor/BlueprintGraph/Private/K2Node_AddComponent.cpp
@@ -105,10 +105,10 @@ void UK2Node_AddComponent::AllocatePinsForExposedVariables()
 			{
 				FEdGraphPinType PinType;
 				K2Schema->ConvertPropertyToPinType(Property, /*out*/ PinType);	
-				const bool bIsUnique = (NULL == FindPin(Property->GetName()));
+				const bool bIsUnique = (nullptr == FindPin(Property->GetFName()));
 				if (K2Schema->FindSetVariableByNameFunction(PinType) && bIsUnique)
 				{
-					UEdGraphPin* Pin = CreatePin(EGPD_Input, FString(), FString(), nullptr, Property->GetName());
+					UEdGraphPin* Pin = CreatePin(EGPD_Input, NAME_None, Property->GetFName());
 					Pin->PinType = PinType;
 					bHasExposedVariable = true;
 
@@ -445,8 +445,7 @@ void UK2Node_AddComponent::ExpandNode(class FKismetCompilerContext& CompilerCont
 		NewNode->AllocateDefaultPinsWithoutExposedVariables();
 
 		// function parameters
-		const UEdGraphSchema_K2* K2Schema = GetDefault<UEdGraphSchema_K2>();
-		CompilerContext.MovePinLinksToIntermediate(*FindPin(K2Schema->PN_Self), *NewNode->FindPin(K2Schema->PN_Self));
+		CompilerContext.MovePinLinksToIntermediate(*FindPin(UEdGraphSchema_K2::PN_Self), *NewNode->FindPin(UEdGraphSchema_K2::PN_Self));
 		CompilerContext.MovePinLinksToIntermediate(*GetTemplateNamePinChecked(), *NewNode->GetTemplateNamePinChecked());
 		CompilerContext.MovePinLinksToIntermediate(*GetRelativeTransformPin(), *NewNode->GetRelativeTransformPin());
 		CompilerContext.MovePinLinksToIntermediate(*GetManualAttachmentPin(), *NewNode->GetManualAttachmentPin());
diff --git a/Engine/Source/Editor/BlueprintGraph/Private/K2Node_AssignmentStatement.cpp b/Engine/Source/Editor/BlueprintGraph/Private/K2Node_AssignmentStatement.cpp
--- a/Engine/Source/Editor/BlueprintGraph/Private/K2Node_AssignmentStatement.cpp
+++ b/Engine/Source/Editor/BlueprintGraph/Private/K2Node_AssignmentStatement.cpp
@@ -55,8 +55,8 @@ public:
 };
 
 
-FString UK2Node_AssignmentStatement::VariablePinName = FString(TEXT("Variable"));
-FString UK2Node_AssignmentStatement::ValuePinName = FString(TEXT("Value"));
+FName UK2Node_AssignmentStatement::VariablePinName(TEXT("Variable"));
+FName UK2Node_AssignmentStatement::ValuePinName(TEXT("Value"));
 
 
 UK2Node_AssignmentStatement::UK2Node_AssignmentStatement(const FObjectInitializer& ObjectInitializer)
@@ -66,13 +66,11 @@ UK2Node_AssignmentStatement::UK2Node_AssignmentStatement(const FObjectInitialize
 
 void UK2Node_AssignmentStatement::AllocateDefaultPins()
 {
-	const UEdGraphSchema_K2* Schema = GetDefault<UEdGraphSchema_K2>();
+	CreatePin(EGPD_Input, UEdGraphSchema_K2::PC_Exec, UEdGraphSchema_K2::PN_Execute);
+	CreatePin(EGPD_Output, UEdGraphSchema_K2::PC_Exec, UEdGraphSchema_K2::PN_Then);
 
-	CreatePin(EGPD_Input, Schema->PC_Exec, FString(), nullptr, Schema->PN_Execute);
-	CreatePin(EGPD_Output, Schema->PC_Exec, FString(), nullptr, Schema->PN_Then);
-
-	UEdGraphPin* VariablePin = CreatePin(EGPD_Input, Schema->PC_Wildcard, FString(), nullptr, VariablePinName);
-	UEdGraphPin* ValuePin = CreatePin(EGPD_Input, Schema->PC_Wildcard, FString(), nullptr, ValuePinName);
+	UEdGraphPin* VariablePin = CreatePin(EGPD_Input, UEdGraphSchema_K2::PC_Wildcard, VariablePinName);
+	UEdGraphPin* ValuePin = CreatePin(EGPD_Input, UEdGraphSchema_K2::PC_Wildcard, ValuePinName);
 
 	Super::AllocateDefaultPins();
 }
@@ -111,18 +109,16 @@ void UK2Node_AssignmentStatement::NotifyPinConnectionListChanged(UEdGraphPin* Pi
 {
 	Super::NotifyPinConnectionListChanged(Pin);
 
-	const UEdGraphSchema_K2* Schema = GetDefault<UEdGraphSchema_K2>();
-
 	UEdGraphPin* VariablePin = FindPin(TEXT("Variable"));
 	UEdGraphPin* ValuePin = FindPin(TEXT("Value"));
 
 	if ((VariablePin->LinkedTo.Num() == 0) && (ValuePin->LinkedTo.Num() == 0))
 	{
 		// Restore the wildcard status
-		VariablePin->PinType.PinCategory = Schema->PC_Wildcard;
+		VariablePin->PinType.PinCategory = UEdGraphSchema_K2::PC_Wildcard;
 		VariablePin->PinType.PinSubCategory = TEXT("");
 		VariablePin->PinType.PinSubCategoryObject = NULL;
-		ValuePin->PinType.PinCategory = Schema->PC_Wildcard;
+		ValuePin->PinType.PinCategory = UEdGraphSchema_K2::PC_Wildcard;
 		ValuePin->PinType.PinSubCategory = TEXT("");
 		ValuePin->PinType.PinSubCategoryObject = NULL;
 	}
@@ -157,17 +153,13 @@ void UK2Node_AssignmentStatement::PostReconstructNode()
 
 UEdGraphPin* UK2Node_AssignmentStatement::GetThenPin() const
 {
-	const UEdGraphSchema_K2* K2Schema = GetDefault<UEdGraphSchema_K2>();
-
-	UEdGraphPin* Pin = FindPin(K2Schema->PN_Then);
+	UEdGraphPin* Pin = FindPin(UEdGraphSchema_K2::PN_Then);
 	check(Pin != NULL);
 	return Pin;
 }
 
 UEdGraphPin* UK2Node_AssignmentStatement::GetVariablePin() const
 {
-	const UEdGraphSchema_K2* K2Schema = GetDefault<UEdGraphSchema_K2>();
-
 	UEdGraphPin* Pin = FindPin(VariablePinName);
 	check(Pin != NULL);
 	return Pin;
@@ -175,8 +167,6 @@ UEdGraphPin* UK2Node_AssignmentStatement::GetVariablePin() const
 
 UEdGraphPin* UK2Node_AssignmentStatement::GetValuePin() const
 {
-	const UEdGraphSchema_K2* K2Schema = GetDefault<UEdGraphSchema_K2>();
-
 	UEdGraphPin* Pin = FindPin(ValuePinName);
 	check(Pin != NULL);
 	return Pin;
