Copying //UE4/Dev-Framework to //UE4/Dev-Main (Source: //UE4/Dev-Framework @ 3716594)

diff --git a/Engine/Source/Editor/BlueprintGraph/Private/K2Node_BaseAsyncTask.cpp b/Engine/Source/Editor/BlueprintGraph/Private/K2Node_BaseAsyncTask.cpp
--- a/Engine/Source/Editor/BlueprintGraph/Private/K2Node_BaseAsyncTask.cpp
+++ b/Engine/Source/Editor/BlueprintGraph/Private/K2Node_BaseAsyncTask.cpp
@@ -65,7 +65,7 @@ void UK2Node_BaseAsyncTask::AllocateDefaultPins()
 {
 	const UEdGraphSchema_K2* K2Schema = GetDefault<UEdGraphSchema_K2>();
 
-	CreatePin(EGPD_Input, K2Schema->PC_Exec, FString(), nullptr, K2Schema->PN_Execute);
+	CreatePin(EGPD_Input, UEdGraphSchema_K2::PC_Exec, UEdGraphSchema_K2::PN_Execute);
 
 	bool bExposeProxy = false;
 	bool bHideThen = false;
@@ -77,20 +77,20 @@ void UK2Node_BaseAsyncTask::AllocateDefaultPins()
 
 	if (!bHideThen)
 	{
-		CreatePin(EGPD_Output, K2Schema->PC_Exec, FString(), nullptr, K2Schema->PN_Then);
+		CreatePin(EGPD_Output, UEdGraphSchema_K2::PC_Exec, UEdGraphSchema_K2::PN_Then);
 	}
 
 	if (bExposeProxy)
 	{
-		CreatePin(EGPD_Output, K2Schema->PC_Object, FString(), ProxyClass, FBaseAsyncTaskHelper::GetAsyncTaskProxyName());
+		CreatePin(EGPD_Output, UEdGraphSchema_K2::PC_Object, ProxyClass, FBaseAsyncTaskHelper::GetAsyncTaskProxyName());
 	}
 
-	UFunction* DelegateSignatureFunction = NULL;
+	UFunction* DelegateSignatureFunction = nullptr;
 	for (TFieldIterator<UProperty> PropertyIt(ProxyClass); PropertyIt; ++PropertyIt)
 	{
 		if (UMulticastDelegateProperty* Property = Cast<UMulticastDelegateProperty>(*PropertyIt))
 		{
-			CreatePin(EGPD_Output, K2Schema->PC_Exec, FString(), nullptr, *Property->GetName());
+			CreatePin(EGPD_Output, UEdGraphSchema_K2::PC_Exec, Property->GetFName());
 			if (!DelegateSignatureFunction)
 			{
 				DelegateSignatureFunction = Property->SignatureFunction;
@@ -106,7 +106,7 @@ void UK2Node_BaseAsyncTask::AllocateDefaultPins()
 			const bool bIsFunctionInput = !Param->HasAnyPropertyFlags(CPF_OutParm) || Param->HasAnyPropertyFlags(CPF_ReferenceParm);
 			if (bIsFunctionInput)
 			{
-				UEdGraphPin* Pin = CreatePin(EGPD_Output, FString(), FString(), nullptr, Param->GetName());
+				UEdGraphPin* Pin = CreatePin(EGPD_Output, NAME_None, Param->GetFName());
 				K2Schema->ConvertPropertyToPinType(Param, /*out*/ Pin->PinType);
 			}
 		}
@@ -116,7 +116,7 @@ void UK2Node_BaseAsyncTask::AllocateDefaultPins()
 	UFunction* Function = ProxyFactoryClass ? ProxyFactoryClass->FindFunctionByName(ProxyFactoryFunctionName) : nullptr;
 	if (Function)
 	{
-		TSet<FString> PinsToHide;
+		TSet<FName> PinsToHide;
 		FBlueprintEditorUtils::GetHiddenPinsForFunction(GetGraph(), Function, PinsToHide);
 		for (TFieldIterator<UProperty> PropIt(Function); PropIt && (PropIt->PropertyFlags & CPF_Parm); ++PropIt)
 		{
@@ -128,9 +128,10 @@ void UK2Node_BaseAsyncTask::AllocateDefaultPins()
 				continue;
 			}
 
-			const bool bIsRefParam = Param->HasAnyPropertyFlags(CPF_ReferenceParm) && bIsFunctionInput;
-			UEdGraphPin* Pin = CreatePin(EGPD_Input, FString(), FString(), nullptr, Param->GetName(), EPinContainerType::None, bIsRefParam);
-			const bool bPinGood = (Pin != NULL) && K2Schema->ConvertPropertyToPinType(Param, /*out*/ Pin->PinType);
+			UEdGraphNode::FCreatePinParams PinParams;
+			PinParams.bIsReference = Param->HasAnyPropertyFlags(CPF_ReferenceParm) && bIsFunctionInput;
+			UEdGraphPin* Pin = CreatePin(EGPD_Input, NAME_None, Param->GetFName(), PinParams);
+			const bool bPinGood = (Pin && K2Schema->ConvertPropertyToPinType(Param, /*out*/ Pin->PinType));
 
 			if (bPinGood)
 			{
@@ -168,18 +169,17 @@ void UK2Node_BaseAsyncTask::AllocateDefaultPins()
 	Super::AllocateDefaultPins();
 }
 
-const FString& UK2Node_BaseAsyncTask::FBaseAsyncTaskHelper::GetAsyncTaskProxyName()
+const FName UK2Node_BaseAsyncTask::FBaseAsyncTaskHelper::GetAsyncTaskProxyName()
 {
-	static const FString Name(TEXT("AsyncTaskProxy"));
+	static const FName Name(TEXT("AsyncTaskProxy"));
 	return Name;
 }
 
-bool UK2Node_BaseAsyncTask::FBaseAsyncTaskHelper::ValidDataPin(const UEdGraphPin* Pin, EEdGraphPinDirection Direction, const UEdGraphSchema_K2* Schema)
+bool UK2Node_BaseAsyncTask::FBaseAsyncTaskHelper::ValidDataPin(const UEdGraphPin* Pin, EEdGraphPinDirection Direction)
 {
-	check(Schema);
 	const bool bValidDataPin = Pin
 		&& !Pin->bOrphanedPin
-		&& (Pin->PinType.PinCategory != Schema->PC_Exec);
+		&& (Pin->PinType.PinCategory != UEdGraphSchema_K2::PC_Exec);
 
 	const bool bProperDirection = Pin && (Pin->Direction == Direction);
 
@@ -200,7 +200,7 @@ bool UK2Node_BaseAsyncTask::FBaseAsyncTaskHelper::CreateDelegateForNewFunction(U
 	UK2Node_CreateDelegate* CreateDelegateNode = CompilerContext.SpawnIntermediateNode<UK2Node_CreateDelegate>(CurrentNode, SourceGraph);
 	CreateDelegateNode->AllocateDefaultPins();
 	bResult &= Schema->TryCreateConnection(DelegateInputPin, CreateDelegateNode->GetDelegateOutPin());
-	bResult &= Schema->TryCreateConnection(SelfNode->FindPinChecked(Schema->PN_Self), CreateDelegateNode->GetObjectInPin());
+	bResult &= Schema->TryCreateConnection(SelfNode->FindPinChecked(UEdGraphSchema_K2::PN_Self), CreateDelegateNode->GetObjectInPin());
 	CreateDelegateNode->SetFunction(FunctionName);
 
 	return bResult;
@@ -218,7 +218,7 @@ bool UK2Node_BaseAsyncTask::FBaseAsyncTaskHelper::CopyEventSignature(UK2Node_Cus
 		{
 			FEdGraphPinType PinType;
 			bResult &= Schema->ConvertPropertyToPinType(Param, /*out*/ PinType);
-			bResult &= (NULL != CENode->CreateUserDefinedPin(Param->GetName(), PinType, EGPD_Output));
+			bResult &= (nullptr != CENode->CreateUserDefinedPin(Param->GetFName(), PinType, EGPD_Output));
 		}
 	}
 	return bResult;
@@ -233,8 +233,8 @@ bool UK2Node_BaseAsyncTask::FBaseAsyncTaskHelper::HandleDelegateImplementation(
 	const UEdGraphSchema_K2* Schema = CompilerContext.GetSchema();
 	check(CurrentProperty && ProxyObjectPin && InOutLastThenPin && CurrentNode && SourceGraph && Schema);
 
-	UEdGraphPin* PinForCurrentDelegateProperty = CurrentNode->FindPin(CurrentProperty->GetName());
-	if (!PinForCurrentDelegateProperty || (Schema->PC_Exec != PinForCurrentDelegateProperty->PinType.PinCategory))
+	UEdGraphPin* PinForCurrentDelegateProperty = CurrentNode->FindPin(CurrentProperty->GetFName());
+	if (!PinForCurrentDelegateProperty || (UEdGraphSchema_K2::PC_Exec != PinForCurrentDelegateProperty->PinType.PinCategory))
 	{
 		FText ErrorMessage = FText::Format(LOCTEXT("WrongDelegateProperty", "BaseAsyncTask: Cannot find execution pin for delegate "), FText::FromString(CurrentProperty->GetName()));
 		CompilerContext.MessageLog.Error(*ErrorMessage.ToString(), CurrentNode);
@@ -246,9 +246,9 @@ bool UK2Node_BaseAsyncTask::FBaseAsyncTaskHelper::HandleDelegateImplementation(
 	UK2Node_AddDelegate* AddDelegateNode = CompilerContext.SpawnIntermediateNode<UK2Node_AddDelegate>(CurrentNode, SourceGraph);
 	AddDelegateNode->SetFromProperty(CurrentProperty, false);
 	AddDelegateNode->AllocateDefaultPins();
-		bIsErrorFree &= Schema->TryCreateConnection(AddDelegateNode->FindPinChecked(Schema->PN_Self), ProxyObjectPin);
-		bIsErrorFree &= Schema->TryCreateConnection(InOutLastThenPin, AddDelegateNode->FindPinChecked(Schema->PN_Execute));
-		InOutLastThenPin = AddDelegateNode->FindPinChecked(Schema->PN_Then);
+		bIsErrorFree &= Schema->TryCreateConnection(AddDelegateNode->FindPinChecked(UEdGraphSchema_K2::PN_Self), ProxyObjectPin);
+		bIsErrorFree &= Schema->TryCreateConnection(InOutLastThenPin, AddDelegateNode->FindPinChecked(UEdGraphSchema_K2::PN_Execute));
+		InOutLastThenPin = AddDelegateNode->FindPinChecked(UEdGraphSchema_K2::PN_Then);
 		CurrentCENode->CustomFunctionName = *FString::Printf(TEXT("%s_%s"), *CurrentProperty->GetName(), *CompilerContext.GetGuid(CurrentNode));
 		CurrentCENode->AllocateDefaultPins();
 
@@ -256,7 +256,7 @@ bool UK2Node_BaseAsyncTask::FBaseAsyncTaskHelper::HandleDelegateImplementation(
 		bIsErrorFree &= FBaseAsyncTaskHelper::CopyEventSignature(CurrentCENode, AddDelegateNode->GetDelegateSignature(), Schema);
 	}
 
-	UEdGraphPin* LastActivatedNodeThen = CurrentCENode->FindPinChecked(Schema->PN_Then);
+	UEdGraphPin* LastActivatedNodeThen = CurrentCENode->FindPinChecked(UEdGraphSchema_K2::PN_Then);
 	for (const FBaseAsyncTaskHelper::FOutputPinAndLocalVariable& OutputPair : VariableOutputs) // CREATE CHAIN OF ASSIGMENTS
 	{
 		UEdGraphPin* PinWithData = CurrentCENode->FindPin(OutputPair.OutputPin->PinName);
@@ -305,11 +305,11 @@ void UK2Node_BaseAsyncTask::ExpandNode(class FKismetCompilerContext& CompilerCon
 		return;
 	}
 
-	bIsErrorFree &= CompilerContext.MovePinLinksToIntermediate(*FindPinChecked(Schema->PN_Execute), *CallCreateProxyObjectNode->FindPinChecked(Schema->PN_Execute)).CanSafeConnect();
+	bIsErrorFree &= CompilerContext.MovePinLinksToIntermediate(*FindPinChecked(UEdGraphSchema_K2::PN_Execute), *CallCreateProxyObjectNode->FindPinChecked(UEdGraphSchema_K2::PN_Execute)).CanSafeConnect();
 
 	for (UEdGraphPin* CurrentPin : Pins)
 	{
-		if (FBaseAsyncTaskHelper::ValidDataPin(CurrentPin, EGPD_Input, Schema))
+		if (FBaseAsyncTaskHelper::ValidDataPin(CurrentPin, EGPD_Input))
 		{
 			UEdGraphPin* DestPin = CallCreateProxyObjectNode->FindPin(CurrentPin->PinName); // match function inputs, to pass data to function from CallFunction node
 			bIsErrorFree &= DestPin && CompilerContext.MovePinLinksToIntermediate(*CurrentPin, *DestPin).CanSafeConnect();
@@ -325,7 +325,7 @@ void UK2Node_BaseAsyncTask::ExpandNode(class FKismetCompilerContext& CompilerCon
 	TArray<FBaseAsyncTaskHelper::FOutputPinAndLocalVariable> VariableOutputs;
 	for (UEdGraphPin* CurrentPin : Pins)
 	{
-		if ((OutputAsyncTaskProxy != CurrentPin) && FBaseAsyncTaskHelper::ValidDataPin(CurrentPin, EGPD_Output, Schema))
+		if ((OutputAsyncTaskProxy != CurrentPin) && FBaseAsyncTaskHelper::ValidDataPin(CurrentPin, EGPD_Output))
 		{
 			const FEdGraphPinType& PinType = CurrentPin->PinType;
 			UK2Node_TemporaryVariable* TempVarOutput = CompilerContext.SpawnInternalVariable(
@@ -336,7 +336,7 @@ void UK2Node_BaseAsyncTask::ExpandNode(class FKismetCompilerContext& CompilerCon
 	}
 
 	// FOR EACH DELEGATE DEFINE EVENT, CONNECT IT TO DELEGATE AND IMPLEMENT A CHAIN OF ASSIGMENTS
-	UEdGraphPin* LastThenPin = CallCreateProxyObjectNode->FindPinChecked(Schema->PN_Then);
+	UEdGraphPin* LastThenPin = CallCreateProxyObjectNode->FindPinChecked(UEdGraphSchema_K2::PN_Then);
 
 	UK2Node_CallFunction* IsValidFuncNode = CompilerContext.SpawnIntermediateNode<UK2Node_CallFunction>(this, SourceGraph);
 	const FName IsValidFuncName = GET_FUNCTION_NAME_CHECKED(UKismetSystemLibrary, IsValid);
@@ -358,7 +358,7 @@ void UK2Node_BaseAsyncTask::ExpandNode(class FKismetCompilerContext& CompilerCon
 		bIsErrorFree &= FBaseAsyncTaskHelper::HandleDelegateImplementation(*PropertyIt, VariableOutputs, ProxyObjectPin, LastThenPin, this, SourceGraph, CompilerContext);
 	}
 
-	if (CallCreateProxyObjectNode->FindPinChecked(Schema->PN_Then) == LastThenPin)
+	if (CallCreateProxyObjectNode->FindPinChecked(UEdGraphSchema_K2::PN_Then) == LastThenPin)
 	{
 		CompilerContext.MessageLog.Error(*LOCTEXT("MissingDelegateProperties", "BaseAsyncTask: Proxy has no delegates defined. @@").ToString(), this);
 		return;
@@ -378,8 +378,8 @@ void UK2Node_BaseAsyncTask::ExpandNode(class FKismetCompilerContext& CompilerCon
 		bIsErrorFree &= Schema->TryCreateConnection(ProxyObjectPin, ActivateCallSelfPin);
 
 		// Hook the activate node up in the exec chain
-		UEdGraphPin* ActivateExecPin = CallActivateProxyObjectNode->FindPinChecked(Schema->PN_Execute);
-		UEdGraphPin* ActivateThenPin = CallActivateProxyObjectNode->FindPinChecked(Schema->PN_Then);
+		UEdGraphPin* ActivateExecPin = CallActivateProxyObjectNode->FindPinChecked(UEdGraphSchema_K2::PN_Execute);
+		UEdGraphPin* ActivateThenPin = CallActivateProxyObjectNode->FindPinChecked(UEdGraphSchema_K2::PN_Then);
 
 		bIsErrorFree &= Schema->TryCreateConnection(LastThenPin, ActivateExecPin);
 
@@ -388,7 +388,7 @@ void UK2Node_BaseAsyncTask::ExpandNode(class FKismetCompilerContext& CompilerCon
 
 	// Move the connections from the original node then pin to the last internal then pin
 
-	UEdGraphPin* OriginalThenPin = FindPin(Schema->PN_Then);
+	UEdGraphPin* OriginalThenPin = FindPin(UEdGraphSchema_K2::PN_Then);
 
 	if (OriginalThenPin)
 	{
@@ -479,15 +479,13 @@ void UK2Node_BaseAsyncTask::ValidateNodeDuringCompilation(class FCompilerResults
 {
 	Super::ValidateNodeDuringCompilation(MessageLog);
 
-	UEdGraphSchema_K2 const* K2Schema = GetDefault<UEdGraphSchema_K2>();
-
 	if(UObject const* SourceObject = MessageLog.FindSourceObject(this))
 	{
 		// Lets check if it's a result of macro expansion, to give a helpful error
 		if(UK2Node_MacroInstance const* MacroInstance = Cast<UK2Node_MacroInstance>(SourceObject))
 		{
 			// Since it's not possible to check the graph's type, just check if this is a ubergraph using the schema's name for it
-			if(!(GetGraph()->HasAnyFlags(RF_Transient) && GetGraph()->GetName().StartsWith(K2Schema->FN_ExecuteUbergraphBase.ToString())))
+			if(!(GetGraph()->HasAnyFlags(RF_Transient) && GetGraph()->GetName().StartsWith(UEdGraphSchema_K2::FN_ExecuteUbergraphBase.ToString())))
 			{
 				MessageLog.Error(*LOCTEXT("AsyncTaskInFunctionFromMacro", "@@ is being used in Function '@@' resulting from expansion of Macro '@@'").ToString(), this, GetGraph(), MacroInstance);
 			}
@@ -495,7 +493,7 @@ void UK2Node_BaseAsyncTask::ValidateNodeDuringCompilation(class FCompilerResults
 	}
 }
 
-TMap<FString, FAsyncTaskPinRedirectMapInfo> UK2Node_BaseAsyncTask::AsyncTaskPinRedirectMap;
+TMap<FName, FAsyncTaskPinRedirectMapInfo> UK2Node_BaseAsyncTask::AsyncTaskPinRedirectMap;
 bool UK2Node_BaseAsyncTask::bAsyncTaskPinRedirectMapInitialized = false;
 
 UK2Node::ERedirectType UK2Node_BaseAsyncTask::DoPinsMatchForReconstruction(const UEdGraphPin* NewPin, int32 NewPinIndex, const UEdGraphPin* OldPin, int32 OldPinIndex) const
@@ -522,8 +520,8 @@ UK2Node::ERedirectType UK2Node_BaseAsyncTask::DoPinsMatchForReconstruction(const
 					UClass* RedirectProxyClass = FindObject<UClass>(ANY_PACKAGE, *ProxyClassString);
 					if (RedirectProxyClass)
 					{
-						FAsyncTaskPinRedirectMapInfo& PinRedirectInfo = AsyncTaskPinRedirectMap.FindOrAdd(OldPinString);
-						TArray<UClass*>& ProxyClassArray = PinRedirectInfo.OldPinToProxyClassMap.FindOrAdd(NewPinString);
+						FAsyncTaskPinRedirectMapInfo& PinRedirectInfo = AsyncTaskPinRedirectMap.FindOrAdd(*OldPinString);
+						TArray<UClass*>& ProxyClassArray = PinRedirectInfo.OldPinToProxyClassMap.FindOrAdd(*NewPinString);
 						ProxyClassArray.AddUnique(RedirectProxyClass);
 					}
 				}
diff --git a/Engine/Source/Editor/BlueprintGraph/Private/K2Node_BitmaskLiteral.cpp b/Engine/Source/Editor/BlueprintGraph/Private/K2Node_BitmaskLiteral.cpp
--- a/Engine/Source/Editor/BlueprintGraph/Private/K2Node_BitmaskLiteral.cpp
+++ b/Engine/Source/Editor/BlueprintGraph/Private/K2Node_BitmaskLiteral.cpp
@@ -14,9 +14,9 @@
 
 #define LOCTEXT_NAMESPACE "UK2Node_BitmaskLiteral"
 
-const FString& UK2Node_BitmaskLiteral::GetBitmaskInputPinName()
+const FName& UK2Node_BitmaskLiteral::GetBitmaskInputPinName()
 {
-	static const FString PinName(TEXT("Bitmask"));
+	static const FName PinName(TEXT("Bitmask"));
 	return PinName;
 }
 
@@ -99,10 +99,10 @@ void UK2Node_BitmaskLiteral::AllocateDefaultPins()
 {
 	const UEdGraphSchema_K2* Schema = GetDefault<UEdGraphSchema_K2>();
 
-	UEdGraphPin* InputPin = CreatePin(EGPD_Input, Schema->PC_Int, Schema->PSC_Bitmask, BitflagsEnum, GetBitmaskInputPinName());
+	UEdGraphPin* InputPin = CreatePin(EGPD_Input, UEdGraphSchema_K2::PC_Int, UEdGraphSchema_K2::PSC_Bitmask, BitflagsEnum, GetBitmaskInputPinName());
 	Schema->SetPinAutogeneratedDefaultValueBasedOnType(InputPin);
 
-	CreatePin(EGPD_Output, Schema->PC_Int, Schema->PSC_Bitmask, BitflagsEnum, Schema->PN_ReturnValue);
+	CreatePin(EGPD_Output, UEdGraphSchema_K2::PC_Int, UEdGraphSchema_K2::PSC_Bitmask, BitflagsEnum, UEdGraphSchema_K2::PN_ReturnValue);
 
 	Super::AllocateDefaultPins();
 }
@@ -117,9 +117,6 @@ void UK2Node_BitmaskLiteral::ReconstructNode()
 
 void UK2Node_BitmaskLiteral::ExpandNode(class FKismetCompilerContext& CompilerContext, UEdGraph* SourceGraph)
 {
-	const UEdGraphSchema_K2* Schema = CompilerContext.GetSchema();
-	check(nullptr != Schema);
-
 	const FName FunctionName = GET_FUNCTION_NAME_CHECKED(UKismetSystemLibrary, MakeLiteralInt);
 	UK2Node_CallFunction* MakeLiteralInt = CompilerContext.SpawnIntermediateNode<UK2Node_CallFunction>(this, SourceGraph);
 	MakeLiteralInt->SetFromFunction(UKismetSystemLibrary::StaticClass()->FindFunctionByName(FunctionName));
@@ -130,7 +127,7 @@ void UK2Node_BitmaskLiteral::ExpandNode(class FKismetCompilerContext& CompilerCo
 	check(nullptr != NewInputPin);
 	CompilerContext.MovePinLinksToIntermediate(*OrgInputPin, *NewInputPin);
 
-	UEdGraphPin* OrgReturnPin = FindPinChecked(Schema->PN_ReturnValue);
+	UEdGraphPin* OrgReturnPin = FindPinChecked(UEdGraphSchema_K2::PN_ReturnValue);
 	UEdGraphPin* NewReturnPin = MakeLiteralInt->GetReturnValuePin();
 	check(nullptr != NewReturnPin);
 	CompilerContext.MovePinLinksToIntermediate(*OrgReturnPin, *NewReturnPin);
diff --git a/Engine/Source/Editor/BlueprintGraph/Private/K2Node_BreakStruct.cpp b/Engine/Source/Editor/BlueprintGraph/Private/K2Node_BreakStruct.cpp
--- a/Engine/Source/Editor/BlueprintGraph/Private/K2Node_BreakStruct.cpp
+++ b/Engine/Source/Editor/BlueprintGraph/Private/K2Node_BreakStruct.cpp
@@ -84,7 +84,7 @@ public:
 
 	void RegisterOutputTerm(FKismetFunctionContext& Context, UScriptStruct* StructType, UEdGraphPin* Net, FBPTerminal* ContextTerm)
 	{
-		if (UProperty* BoundProperty = FindField<UProperty>(StructType, *(Net->PinName)))
+		if (UProperty* BoundProperty = FindField<UProperty>(StructType, Net->PinName))
 		{
 			if (BoundProperty->HasAnyPropertyFlags(CPF_Deprecated) && Net->LinkedTo.Num())
 			{
@@ -95,7 +95,7 @@ public:
 			}
 
 			UBlueprintEditorSettings* Settings = GetMutableDefault<UBlueprintEditorSettings>();
-			FBPTerminal* Term = Context.CreateLocalTerminalFromPinAutoChooseScope(Net, Net->PinName);
+			FBPTerminal* Term = Context.CreateLocalTerminalFromPinAutoChooseScope(Net, Net->PinName.ToString());
 			Term->bPassedByReference = ContextTerm->bPassedByReference;
 			Term->AssociatedVarProperty = BoundProperty;
 			Context.NetMap.Add(Net, Term);
@@ -169,11 +169,14 @@ bool UK2Node_BreakStruct::CanBeBroken(const UScriptStruct* Struct, const bool bF
 
 void UK2Node_BreakStruct::AllocateDefaultPins()
 {
-	const UEdGraphSchema_K2* Schema = GetDefault<UEdGraphSchema_K2>();
-	if(Schema && StructType)
+	if (StructType)
 	{
+		UEdGraphNode::FCreatePinParams PinParams;
+		PinParams.bIsConst = true;
+		PinParams.bIsReference = true;
+
 		PreloadObject(StructType);
-		CreatePin(EGPD_Input, Schema->PC_Struct, FString(), StructType, StructType->GetName(), EPinContainerType::None, true, true);
+		CreatePin(EGPD_Input, UEdGraphSchema_K2::PC_Struct, StructType, StructType->GetFName(), PinParams);
 		
 		struct FBreakStructPinManager : public FStructOperationOptionalPinManager
 		{
@@ -267,7 +270,7 @@ void UK2Node_BreakStruct::ValidateNodeDuringCompilation(class FCompilerResultsLo
 				const bool bIsBlueprintVisible = Property->HasAnyPropertyFlags(CPF_BlueprintVisible) || (Property->GetOwnerStruct() && Property->GetOwnerStruct()->IsA<UUserDefinedStruct>());
 				bHasAnyBlueprintVisibleProperty |= bIsBlueprintVisible;
 
-				const UEdGraphPin* Pin = FindPin(Property->GetName());
+				const UEdGraphPin* Pin = FindPin(Property->GetFName());
 				const bool bIsLinked = Pin && Pin->LinkedTo.Num();
 
 				if (!bIsBlueprintVisible && bIsLinked)
@@ -305,7 +308,7 @@ FLinearColor UK2Node_BreakStruct::GetNodeTitleColor() const
 	if(const UEdGraphSchema_K2* K2Schema = GetDefault<UEdGraphSchema_K2>())
 	{
 		FEdGraphPinType PinType;
-		PinType.PinCategory = K2Schema->PC_Struct;
+		PinType.PinCategory = UEdGraphSchema_K2::PC_Struct;
 		PinType.PinSubCategoryObject = StructType;
 		return K2Schema->GetPinTypeColor(PinType);
 	}
@@ -331,11 +334,11 @@ UK2Node::ERedirectType UK2Node_BreakStruct::DoPinsMatchForReconstruction(const U
 		}
 		else if ((EGPD_Output == NewPin->Direction) && (EGPD_Output == OldPin->Direction))
 		{
-			FName RedirectedPinName = UProperty::FindRedirectedPropertyName(StructType, FName(*OldPin->PinName));
+			FName RedirectedPinName = UProperty::FindRedirectedPropertyName(StructType, OldPin->PinName);
 
 			if (RedirectedPinName != NAME_Name)
 			{
-				Result = ((FCString::Stricmp(*RedirectedPinName.ToString(), *NewPin->PinName) != 0) ? ERedirectType_None : ERedirectType_Name);
+				Result = (RedirectedPinName != NewPin->PinName ? ERedirectType_None : ERedirectType_Name);
 			}
 		}
 	}
@@ -448,7 +451,7 @@ void UK2Node_BreakStruct::ConvertDeprecatedNode(UEdGraph* Graph, bool bOnlySafeC
 		UFunction* BreakNodeFunction = nullptr;
 
 		// If any pins need to change their names during the conversion, add them to the map.
-		TMap<FString, FString> OldPinToNewPinMap;
+		TMap<FName, FName> OldPinToNewPinMap;
 
 		if (StructType == TBaseStructure<FRotator>::Get())
 		{
@@ -476,7 +479,7 @@ void UK2Node_BreakStruct::ConvertDeprecatedNode(UEdGraph* Graph, bool bOnlySafeC
 				{
 					if (FieldIterator->PropertyFlags & CPF_Parm && !(FieldIterator->PropertyFlags & CPF_ReturnParm))
 					{
-						OldPinToNewPinMap.Add(*StructType->GetName(), *FieldIterator->GetName());
+						OldPinToNewPinMap.Add(StructType->GetFName(), FieldIterator->GetFName());
 						break;
 					}
 				}
diff --git a/Engine/Source/Editor/BlueprintGraph/Private/K2Node_CallArrayFunction.cpp b/Engine/Source/Editor/BlueprintGraph/Private/K2Node_CallArrayFunction.cpp
--- a/Engine/Source/Editor/BlueprintGraph/Private/K2Node_CallArrayFunction.cpp
+++ b/Engine/Source/Editor/BlueprintGraph/Private/K2Node_CallArrayFunction.cpp
@@ -6,6 +6,7 @@
 #include "K2Node_GetArrayItem.h"
 #include "BlueprintsObjectVersion.h"
 #include "Kismet/KismetArrayLibrary.h" // for Array_Get()
+#include "Algo/Transform.h"
 
 UK2Node_CallArrayFunction::UK2Node_CallArrayFunction(const FObjectInitializer& ObjectInitializer)
 	: Super(ObjectInitializer)
@@ -16,15 +17,13 @@ void UK2Node_CallArrayFunction::AllocateDefaultPins()
 {
 	Super::AllocateDefaultPins();
 
-	const UEdGraphSchema_K2* Schema = GetDefault<UEdGraphSchema_K2>();
-
 	UEdGraphPin* TargetArrayPin = GetTargetArrayPin();
 	if (ensure(TargetArrayPin))
 	{
 		TargetArrayPin->PinType.ContainerType = EPinContainerType::Array;
 		TargetArrayPin->PinType.bIsReference = true;
-		TargetArrayPin->PinType.PinCategory = Schema->PC_Wildcard;
-		TargetArrayPin->PinType.PinSubCategory.Reset();
+		TargetArrayPin->PinType.PinCategory = UEdGraphSchema_K2::PC_Wildcard;
+		TargetArrayPin->PinType.PinSubCategory = NAME_None;
 		TargetArrayPin->PinType.PinSubCategoryObject = nullptr;
 	}
 
@@ -79,12 +78,11 @@ void UK2Node_CallArrayFunction::NotifyPinConnectionListChanged(UEdGraphPin* Pin)
 
 	if (PinsToCheck.Contains(Pin))
 	{
-		const UEdGraphSchema_K2* Schema = GetDefault<UEdGraphSchema_K2>();
 		bool bNeedToPropagate = false;
 
 		if( Pin->LinkedTo.Num() > 0 )
 		{
-			if (Pin->PinType.PinCategory == Schema->PC_Wildcard)
+			if (Pin->PinType.PinCategory == UEdGraphSchema_K2::PC_Wildcard)
 			{
 				UEdGraphPin* LinkedTo = Pin->LinkedTo[0];
 				check(LinkedTo);
@@ -112,9 +110,9 @@ void UK2Node_CallArrayFunction::NotifyPinConnectionListChanged(UEdGraphPin* Pin)
 
 			if (bNeedToPropagate)
 			{
-				Pin->PinType.PinCategory = Schema->PC_Wildcard;
-				Pin->PinType.PinSubCategory = TEXT("");
-				Pin->PinType.PinSubCategoryObject = NULL;
+				Pin->PinType.PinCategory = UEdGraphSchema_K2::PC_Wildcard;
+				Pin->PinType.PinSubCategory = NAME_None;
+				Pin->PinType.PinSubCategoryObject = nullptr;
 			}
 		}
 
@@ -149,7 +147,7 @@ void UK2Node_CallArrayFunction::ConvertDeprecatedNode(UEdGraph* Graph, bool bOnl
 			const UEdGraphSchema_K2* K2Schema = Cast<UEdGraphSchema_K2>(Graph->GetSchema());
 			if (K2Schema && GetItemNode)
 			{
-				TMap<FString, FString> OldToNewPinMap;
+				TMap<FName, FName> OldToNewPinMap;
 				for (UEdGraphPin* Pin : Pins)
 				{
 					if (Pin->ParentPin != nullptr)
@@ -160,7 +158,7 @@ void UK2Node_CallArrayFunction::ConvertDeprecatedNode(UEdGraph* Graph, bool bOnl
 					else if (Pin->PinName == UEdGraphSchema_K2::PN_Self)
 					{
 						// there's no analogous pin, signal that we're expecting this
-						OldToNewPinMap.Add(Pin->PinName, FString());
+						OldToNewPinMap.Add(Pin->PinName, NAME_None);
 					}
 					else if (Pin->PinType.IsArray())
 					{
@@ -258,8 +256,11 @@ void UK2Node_CallArrayFunction::GetArrayTypeDependentPins(TArray<UEdGraphPin*>&
 	if (ensure(TargetFunction))
 	{
 		const FString& DependentPinMetaData = TargetFunction->GetMetaData(FBlueprintMetadata::MD_ArrayDependentParam);
-		TArray<FString> TypeDependentPinNames;
-		DependentPinMetaData.ParseIntoArray(TypeDependentPinNames, TEXT(","), true);
+		TArray<FString> TypeDependentPinNameStrs;
+		DependentPinMetaData.ParseIntoArray(TypeDependentPinNameStrs, TEXT(","), true);
+
+		TArray<FName> TypeDependentPinNames;
+		Algo::Transform(TypeDependentPinNameStrs, TypeDependentPinNames, [](const FString& PinNameStr) { return FName(*PinNameStr); });
 
 		for (TArray<UEdGraphPin*>::TConstIterator it(Pins); it; ++it)
 		{
diff --git a/Engine/Source/Editor/BlueprintGraph/Private/K2Node_CallDataTableFunction.cpp b/Engine/Source/Editor/BlueprintGraph/Private/K2Node_CallDataTableFunction.cpp
--- a/Engine/Source/Editor/BlueprintGraph/Private/K2Node_CallDataTableFunction.cpp
+++ b/Engine/Source/Editor/BlueprintGraph/Private/K2Node_CallDataTableFunction.cpp
@@ -14,7 +14,7 @@ void UK2Node_CallDataTableFunction::PinDefaultValueChanged(UEdGraphPin* Pin)
 	Super::PinDefaultValueChanged(Pin);
 
 	const FString& DataTablePinName = GetTargetFunction()->GetMetaData(FBlueprintMetadata::MD_DataTablePin);
-	if (Pin->PinName == DataTablePinName)
+	if (Pin->PinName.ToString() == DataTablePinName)
 	{
 		// When the DataTable pin gets a new value assigned, we need to update the Slate UI so that the RowName drop down gets updated
 		GetGraph()->NotifyGraphChanged();
@@ -26,7 +26,7 @@ void UK2Node_CallDataTableFunction::NotifyPinConnectionListChanged(UEdGraphPin*
 	Super::NotifyPinConnectionListChanged(Pin);
 
 	const FString& DataTablePinName = GetTargetFunction()->GetMetaData(FBlueprintMetadata::MD_DataTablePin);
-	if (Pin->PinName == DataTablePinName)
+	if (Pin->PinName.ToString() == DataTablePinName)
 	{
 		// When the DataTable pin gets a new connection assigned, we need to update the Slate UI so that the RowName drop down gets updated
 		GetGraph()->NotifyGraphChanged();
diff --git a/Engine/Source/Editor/BlueprintGraph/Private/K2Node_CallFunction.cpp b/Engine/Source/Editor/BlueprintGraph/Private/K2Node_CallFunction.cpp
--- a/Engine/Source/Editor/BlueprintGraph/Private/K2Node_CallFunction.cpp
+++ b/Engine/Source/Editor/BlueprintGraph/Private/K2Node_CallFunction.cpp
@@ -50,7 +50,7 @@ struct FCustomStructureParamHelper
 
 	static void FillCustomStructureParameterNames(const UFunction* Function, TArray<FString>& OutNames)
 	{
-		OutNames.Empty();
+		OutNames.Reset();
 		if (Function)
 		{
 			const FString& MetaDataValue = Function->GetMetaData(GetCustomStructureParamName());
@@ -76,13 +76,13 @@ struct FCustomStructureParamHelper
 			else
 			{
 				Pin->PinType.PinCategory = UEdGraphSchema_K2::PC_Wildcard;
-				Pin->PinType.PinSubCategory.Reset();
+				Pin->PinType.PinSubCategory = NAME_None;
 				Pin->PinType.PinSubCategoryObject = nullptr;
 			}
 		}
 	}
 
-	static void UpdateCustomStructurePins(const UFunction* Function, UK2Node* Node, UEdGraphPin* SinglePin = NULL)
+	static void UpdateCustomStructurePins(const UFunction* Function, UK2Node* Node, UEdGraphPin* SinglePin = nullptr)
 	{
 		if (Function && Node)
 		{
@@ -90,7 +90,7 @@ struct FCustomStructureParamHelper
 			FCustomStructureParamHelper::FillCustomStructureParameterNames(Function, Names);
 			if (SinglePin)
 			{
-				if (Names.Contains(SinglePin->PinName))
+				if (Names.Contains(SinglePin->PinName.ToString()))
 				{
 					HandleSinglePin(SinglePin);
 				}
@@ -356,7 +356,7 @@ bool FDynamicOutputHelper::IsTypePickerPin(UEdGraphPin* Pin) const
 		const FString& TypeDeterminingPinName = Function->GetMetaData(FBlueprintMetadata::MD_DynamicOutputType);
 		if (!TypeDeterminingPinName.IsEmpty())
 		{
-			bIsTypeDeterminingPin = (Pin->PinName == TypeDeterminingPinName);
+			bIsTypeDeterminingPin = (Pin->PinName.ToString() == TypeDeterminingPinName);
 		}
 	}
 
@@ -402,7 +402,7 @@ UEdGraphPin* FDynamicOutputHelper::GetDynamicOutPin(const UK2Node_CallFunction*
 	UEdGraphPin* DynamicOutPin = nullptr;
 	if (TaggedOutputParam != nullptr)
 	{
-		DynamicOutPin = FuncNode->FindPin(TaggedOutputParam->GetName());
+		DynamicOutPin = FuncNode->FindPin(TaggedOutputParam->GetFName());
 		if (DynamicOutPin && (DynamicOutPin->Direction != EGPD_Output))
 		{
 			DynamicOutPin = nullptr;
@@ -418,7 +418,7 @@ bool FDynamicOutputHelper::CanConformPinType(const UK2Node_CallFunction* FuncNod
 	{
 		UClass* BasePickerClass = CastChecked<UClass>(TypePickerPin->PinType.PinSubCategoryObject.Get());
 
-		const FString& PinCategory = TypeToTest.PinCategory;
+		const FName PinCategory = TypeToTest.PinCategory;
 		if ((PinCategory == UEdGraphSchema_K2::PC_Object) ||
 			(PinCategory == UEdGraphSchema_K2::PC_Interface) ||
 			(PinCategory == UEdGraphSchema_K2::PC_Class))
@@ -709,12 +709,12 @@ UEdGraphPin* UK2Node_CallFunction::CreateSelfPin(const UFunction* Function)
 	}
 	else if (FunctionClass->IsChildOf(UInterface::StaticClass()))
 	{
-		SelfPin = CreatePin(EGPD_Input, UEdGraphSchema_K2::PC_Interface, FString(), FunctionClass, UEdGraphSchema_K2::PN_Self);
+		SelfPin = CreatePin(EGPD_Input, UEdGraphSchema_K2::PC_Interface, FunctionClass, UEdGraphSchema_K2::PN_Self);
 	}
 	else
 	{
 		// This means that the function is declared in an external class, and should reference that class
-		SelfPin = CreatePin(EGPD_Input, UEdGraphSchema_K2::PC_Object, FString(), FunctionClass, UEdGraphSchema_K2::PN_Self);
+		SelfPin = CreatePin(EGPD_Input, UEdGraphSchema_K2::PC_Object, FunctionClass, UEdGraphSchema_K2::PN_Self);
 	}
 	check(SelfPin != nullptr);
 
@@ -762,8 +762,9 @@ void UK2Node_CallFunction::CreateExecPinsForFunctionCall(const UFunction* Functi
 					bool const bShouldBeHidden = Enum->HasMetaData(TEXT("Hidden"), ExecIdx) || Enum->HasMetaData(TEXT("Spacer"), ExecIdx);
 					if (!bShouldBeHidden)
 					{
-						FString ExecName = Enum->GetNameStringByIndex(ExecIdx);
-						CreatePin(Direction, UEdGraphSchema_K2::PC_Exec, FString(), nullptr, ExecName);
+						// Can't use Enum->GetNameByIndex here because it doesn't do namespace mangling
+						const FName ExecName = *Enum->GetNameStringByIndex(ExecIdx);
+						CreatePin(Direction, UEdGraphSchema_K2::PC_Exec, ExecName);
 					}
 				}
 				
@@ -783,16 +784,16 @@ void UK2Node_CallFunction::CreateExecPinsForFunctionCall(const UFunction* Functi
 		if (bCreateSingleExecInputPin)
 		{
 			// Single input exec pin
-			CreatePin(EGPD_Input, UEdGraphSchema_K2::PC_Exec, FString(), nullptr, UEdGraphSchema_K2::PN_Execute);
+			CreatePin(EGPD_Input, UEdGraphSchema_K2::PC_Exec, UEdGraphSchema_K2::PN_Execute);
 		}
 
 		if (bCreateThenPin)
 		{
-			UEdGraphPin* OutputExecPin = CreatePin(EGPD_Output, UEdGraphSchema_K2::PC_Exec, FString(), nullptr, UEdGraphSchema_K2::PN_Then);
+			UEdGraphPin* OutputExecPin = CreatePin(EGPD_Output, UEdGraphSchema_K2::PC_Exec, UEdGraphSchema_K2::PN_Then);
 			// Use 'completed' name for output pins on latent functions
 			if (Function->HasMetaData(FBlueprintMetadata::MD_Latent))
 			{
-				OutputExecPin->PinFriendlyName = FText::FromString(UEdGraphSchema_K2::PN_Completed);
+				OutputExecPin->PinFriendlyName = FText::FromName(UEdGraphSchema_K2::PN_Completed);
 			}
 		}
 	}
@@ -906,8 +907,8 @@ bool UK2Node_CallFunction::CreatePinsForFunctionCall(const UFunction* Function)
 	}
 
 	// Build a list of the pins that should be hidden for this function (ones that are automagically filled in by the K2 compiler)
-	TSet<FString> PinsToHide;
-	TSet<FString> InternalPins;
+	TSet<FName> PinsToHide;
+	TSet<FName> InternalPins;
 	FBlueprintEditorUtils::GetHiddenPinsForFunction(Graph, Function, PinsToHide, &InternalPins);
 
 	const bool bShowWorldContextPin = ((PinsToHide.Num() > 0) && BP && BP->ParentClass && BP->ParentClass->HasMetaDataHierarchical(FBlueprintMetadata::MD_ShowWorldContextPin));
@@ -922,8 +923,10 @@ bool UK2Node_CallFunction::CreatePinsForFunctionCall(const UFunction* Function)
 
 		const EEdGraphPinDirection Direction = bIsFunctionInput ? EGPD_Input : EGPD_Output;
 
-		UEdGraphPin* Pin = CreatePin(Direction, FString(), FString(), nullptr, Param->GetName(), EPinContainerType::None, bIsRefParam);
-		const bool bPinGood = (Pin != NULL) && K2Schema->ConvertPropertyToPinType(Param, /*out*/ Pin->PinType);
+		UEdGraphNode::FCreatePinParams PinParams;
+		PinParams.bIsReference = bIsRefParam;
+		UEdGraphPin* Pin = CreatePin(Direction, NAME_None, Param->GetFName(), PinParams);
+		const bool bPinGood = (Pin && K2Schema->ConvertPropertyToPinType(Param, /*out*/ Pin->PinType));
 
 		if (bPinGood)
 		{
@@ -956,9 +959,10 @@ bool UK2Node_CallFunction::CreatePinsForFunctionCall(const UFunction* Function)
 			
 			if (PinsToHide.Contains(Pin->PinName))
 			{
+				const FString PinNameStr = Pin->PinName.ToString();
 				const FString& DefaultToSelfMetaValue = Function->GetMetaData(FBlueprintMetadata::MD_DefaultToSelf);
 				const FString& WorldContextMetaValue  = Function->GetMetaData(FBlueprintMetadata::MD_WorldContext);
-				bool bIsSelfPin = ((Pin->PinName == DefaultToSelfMetaValue) || (Pin->PinName == WorldContextMetaValue));
+				bool bIsSelfPin = ((PinNameStr == DefaultToSelfMetaValue) || (PinNameStr == WorldContextMetaValue));
 
 				if (!bShowWorldContextPin || !bIsSelfPin)
 				{
@@ -1298,7 +1302,7 @@ void UK2Node_CallFunction::GeneratePinTooltipFromFunction(UEdGraphPin& Pin, cons
 	}
 	else
 	{
-		ParamName = Pin.PinName.ToLower();
+		ParamName = Pin.PinName.ToString();
 	}
 
 	// grab the the function's comment block for us to parse
@@ -1739,29 +1743,33 @@ void UK2Node_CallFunction::PostPasteNode()
 
 	if (UFunction* Function = GetTargetFunction())
 	{
-		// After pasting we need to go through and ensure the hidden the self pins is correct in case the source blueprint had different metadata
-		TSet<FString> PinsToHide;
-		FBlueprintEditorUtils::GetHiddenPinsForFunction(GetGraph(), Function, PinsToHide);
+		if (Pins.Num() > 0)
+		{
+			// After pasting we need to go through and ensure the hidden the self pins is correct in case the source blueprint had different metadata
+			TSet<FName> PinsToHide;
+			FBlueprintEditorUtils::GetHiddenPinsForFunction(GetGraph(), Function, PinsToHide);
 
-		const bool bShowWorldContextPin = ((PinsToHide.Num() > 0) && GetBlueprint()->ParentClass->HasMetaDataHierarchical(FBlueprintMetadata::MD_ShowWorldContextPin));
+			const bool bShowWorldContextPin = ((PinsToHide.Num() > 0) && GetBlueprint()->ParentClass->HasMetaDataHierarchical(FBlueprintMetadata::MD_ShowWorldContextPin));
 
-		const FString& DefaultToSelfMetaValue = Function->GetMetaData(FBlueprintMetadata::MD_DefaultToSelf);
-		const FString& WorldContextMetaValue  = Function->GetMetaData(FBlueprintMetadata::MD_WorldContext);
+			const FString& DefaultToSelfMetaValue = Function->GetMetaData(FBlueprintMetadata::MD_DefaultToSelf);
+			const FString& WorldContextMetaValue = Function->GetMetaData(FBlueprintMetadata::MD_WorldContext);
 
-		const UEdGraphSchema_K2* K2Schema = GetDefault<UEdGraphSchema_K2>();
-		for (int32 PinIndex = 0; PinIndex < Pins.Num(); ++PinIndex)
-		{
-			UEdGraphPin* Pin = Pins[PinIndex];
+			const UEdGraphSchema_K2* K2Schema = GetDefault<UEdGraphSchema_K2>();
+			for (int32 PinIndex = 0; PinIndex < Pins.Num(); ++PinIndex)
+			{
+				UEdGraphPin* Pin = Pins[PinIndex];
+					const FString PinNameStr = Pin->PinName.ToString();
 
-			const bool bIsSelfPin = ((Pin->PinName == DefaultToSelfMetaValue) || (Pin->PinName == WorldContextMetaValue));
-			const bool bPinShouldBeHidden = ((Pin->SubPins.Num() > 0) || (PinsToHide.Contains(Pin->PinName) && (!bShowWorldContextPin || !bIsSelfPin)));
+					const bool bIsSelfPin = ((PinNameStr == DefaultToSelfMetaValue) || (PinNameStr == WorldContextMetaValue));
+					const bool bPinShouldBeHidden = ((Pin->SubPins.Num() > 0) || (PinsToHide.Contains(Pin->PinName) && (!bShowWorldContextPin || !bIsSelfPin)));
 
-			if (bPinShouldBeHidden && !Pin->bHidden)
-			{
-				Pin->BreakAllPinLinks();
-				K2Schema->SetPinAutogeneratedDefaultValueBasedOnType(Pin);
+				if (bPinShouldBeHidden && !Pin->bHidden)
+				{
+					Pin->BreakAllPinLinks();
+					K2Schema->SetPinAutogeneratedDefaultValueBasedOnType(Pin);
+				}
+				Pin->bHidden = bPinShouldBeHidden;
 			}
-			Pin->bHidden = bPinShouldBeHidden;
 		}
 	}
 }
@@ -2009,11 +2017,11 @@ void UK2Node_CallFunction::ExpandNode(class FKismetCompilerContext& CompilerCont
 				if (EnumParamPin->Direction == EGPD_Input)
 				{
 					// Create normal exec input
-					UEdGraphPin* ExecutePin = CreatePin(EGPD_Input, Schema->PC_Exec, FString(), nullptr, Schema->PN_Execute);
+					UEdGraphPin* ExecutePin = CreatePin(EGPD_Input, UEdGraphSchema_K2::PC_Exec, UEdGraphSchema_K2::PN_Execute);
 
 					// Create temp enum variable
 					UK2Node_TemporaryVariable* TempEnumVarNode = CompilerContext.SpawnIntermediateNode<UK2Node_TemporaryVariable>(this, SourceGraph);
-					TempEnumVarNode->VariableType.PinCategory = Schema->PC_Byte;
+					TempEnumVarNode->VariableType.PinCategory = UEdGraphSchema_K2::PC_Byte;
 					TempEnumVarNode->VariableType.PinSubCategoryObject = Enum;
 					TempEnumVarNode->AllocateDefaultPins();
 					// Get the output pin
@@ -2029,7 +2037,7 @@ void UK2Node_CallFunction::ExpandNode(class FKismetCompilerContext& CompilerCont
 						if( Pin != NULL && 
 							Pin != ExecutePin &&
 							Pin->Direction == EGPD_Input && 
-							Pin->PinType.PinCategory == Schema->PC_Exec )
+							Pin->PinType.PinCategory == UEdGraphSchema_K2::PC_Exec )
 						{
 							// Create node to set the temp enum var
 							UK2Node_AssignmentStatement* AssignNode = CompilerContext.SpawnIntermediateNode<UK2Node_AssignmentStatement>(this, SourceGraph);
@@ -2045,7 +2053,7 @@ void UK2Node_CallFunction::ExpandNode(class FKismetCompilerContext& CompilerCont
 							Schema->TryCreateConnection(AssignNode->GetThenPin(), ExecutePin);
 
 							// set the literal enum value to set to
-							AssignNode->GetValuePin()->DefaultValue = Pin->PinName;
+							AssignNode->GetValuePin()->DefaultValue = Pin->PinName.ToString();
 
 							// Finally remove this 'cosmetic' exec pin
 							Pins[PinIdx]->MarkPendingKill();
@@ -2057,7 +2065,7 @@ void UK2Node_CallFunction::ExpandNode(class FKismetCompilerContext& CompilerCont
 				else if (EnumParamPin->Direction == EGPD_Output)
 				{
 					// Create normal exec output
-					UEdGraphPin* ExecutePin = CreatePin(EGPD_Output, Schema->PC_Exec, FString(), nullptr, Schema->PN_Execute);
+					UEdGraphPin* ExecutePin = CreatePin(EGPD_Output, UEdGraphSchema_K2::PC_Exec, UEdGraphSchema_K2::PN_Execute);
 						
 					// Create a SwitchEnum node to switch on the output enum
 					UK2Node_SwitchEnum* SwitchEnumNode = CompilerContext.SpawnIntermediateNode<UK2Node_SwitchEnum>(this, SourceGraph);
@@ -2077,7 +2085,7 @@ void UK2Node_CallFunction::ExpandNode(class FKismetCompilerContext& CompilerCont
 						if( Pin != NULL &&
 							Pin != ExecutePin &&
 							Pin->Direction == EGPD_Output &&
-							Pin->PinType.PinCategory == Schema->PC_Exec )
+							Pin->PinType.PinCategory == UEdGraphSchema_K2::PC_Exec )
 						{
 							// Move connections from fake 'enum exec' pin to this switch node
 							CompilerContext.MovePinLinksToIntermediate(*Pin, *SwitchEnumNode->FindPinChecked(Pin->PinName));
@@ -2113,7 +2121,7 @@ void UK2Node_CallFunction::ExpandNode(class FKismetCompilerContext& CompilerCont
 
 				const bool bHasConnections = Pin->LinkedTo.Num() > 0;
 				const bool bCreateDefaultValRefTerm = bHasAutoCreateRefTerms && 
-					!bHasConnections && AutoCreateRefTermPinNames.Contains(Pin->PinName);
+					!bHasConnections && AutoCreateRefTermPinNames.Contains(Pin->PinName.ToString());
 
 				if (bCreateDefaultValRefTerm)
 				{
@@ -2159,7 +2167,7 @@ void UK2Node_CallFunction::ExpandNode(class FKismetCompilerContext& CompilerCont
 	{
 		const bool bProperInputToExpandForEach = 
 			(1 == MultiSelf->LinkedTo.Num()) && 
-			(NULL != MultiSelf->LinkedTo[0]) && 
+			(nullptr != MultiSelf->LinkedTo[0]) && 
 			(MultiSelf->LinkedTo[0]->PinType.IsArray());
 		if(bProperInputToExpandForEach)
 		{
@@ -2216,11 +2224,11 @@ void UK2Node_CallFunction::CallForEachElementInArrayExpansion(UK2Node* Node, UEd
 		(MultiSelf->LinkedTo[0]->PinType.IsArray());
 	ensure(bProperInputToExpandForEach);
 
-	UEdGraphPin* ThenPin = Node->FindPinChecked(Schema->PN_Then);
+	UEdGraphPin* ThenPin = Node->FindPinChecked(UEdGraphSchema_K2::PN_Then);
 
 	// Create int Iterator
 	UK2Node_TemporaryVariable* IteratorVar = CompilerContext.SpawnIntermediateNode<UK2Node_TemporaryVariable>(Node, SourceGraph);
-	IteratorVar->VariableType.PinCategory = Schema->PC_Int;
+	IteratorVar->VariableType.PinCategory = UEdGraphSchema_K2::PC_Int;
 	IteratorVar->AllocateDefaultPins();
 
 	// Initialize iterator
@@ -2414,9 +2422,9 @@ void UK2Node_CallFunction::ConformContainerPins()
 			{
 				// reset to wildcard:
 				Pin->PinType.PinCategory = UEdGraphSchema_K2::PC_Wildcard;
-				Pin->PinType.PinSubCategory.Empty();
+				Pin->PinType.PinSubCategory = NAME_None;
 				Pin->PinType.PinSubCategoryObject = nullptr;
-				Pin->DefaultValue = TEXT("");
+				Pin->DefaultValue = FString();
 			}
 		}
 	};
@@ -2434,7 +2442,7 @@ void UK2Node_CallFunction::ConformContainerPins()
 			else
 			{
 				Pin->PinType.PinValueType.TerminalCategory = UEdGraphSchema_K2::PC_Wildcard;
-				Pin->PinType.PinValueType.TerminalSubCategory.Empty();
+				Pin->PinType.PinValueType.TerminalSubCategory = NAME_None;
 				Pin->PinType.PinValueType.TerminalSubCategoryObject = nullptr;
 			}
 		}
@@ -2674,13 +2682,13 @@ UEdGraph* UK2Node_CallFunction::GetFunctionGraph(const UEdGraphNode*& OutGraphNo
 	return NULL;
 }
 
-bool UK2Node_CallFunction::IsStructureWildcardProperty(const UFunction* Function, const FString& PropertyName)
+bool UK2Node_CallFunction::IsStructureWildcardProperty(const UFunction* Function, const FName PropertyName)
 {
-	if (Function && !PropertyName.IsEmpty())
+	if (Function && !PropertyName.IsNone())
 	{
 		TArray<FString> Names;
 		FCustomStructureParamHelper::FillCustomStructureParameterNames(Function, Names);
-		if (Names.Contains(PropertyName))
+		if (Names.Contains(PropertyName.ToString()))
 		{
 			return true;
 		}
@@ -2692,7 +2700,7 @@ bool UK2Node_CallFunction::IsWildcardProperty(const UFunction* InFunction, const
 {
 	if (InProperty)
 	{
-		return FEdGraphUtilities::IsSetParam(InFunction, InProperty->GetName()) || FEdGraphUtilities::IsMapParam(InFunction, InProperty->GetName());
+		return FEdGraphUtilities::IsSetParam(InFunction, InProperty->GetFName()) || FEdGraphUtilities::IsMapParam(InFunction, InProperty->GetFName());
 	}
 	return false;
 }
diff --git a/Engine/Source/Editor/BlueprintGraph/Private/K2Node_CallFunctionOnMember.cpp b/Engine/Source/Editor/BlueprintGraph/Private/K2Node_CallFunctionOnMember.cpp
--- a/Engine/Source/Editor/BlueprintGraph/Private/K2Node_CallFunctionOnMember.cpp
+++ b/Engine/Source/Editor/BlueprintGraph/Private/K2Node_CallFunctionOnMember.cpp
@@ -18,18 +18,16 @@ UK2Node_CallFunctionOnMember::UK2Node_CallFunctionOnMember(const FObjectInitiali
 
 UEdGraphPin* UK2Node_CallFunctionOnMember::CreateSelfPin(const UFunction* Function) 
 {
-	const UEdGraphSchema_K2* K2Schema = GetDefault<UEdGraphSchema_K2>();
-
 	UEdGraphPin* SelfPin = nullptr;
 	if (MemberVariableToCallOn.IsSelfContext())
 	{
 		// This means the function is defined within the blueprint, so the pin should be a true "self" pin
-		SelfPin = CreatePin(EGPD_Input, K2Schema->PC_Object, K2Schema->PSC_Self, nullptr, K2Schema->PN_Self);
+		SelfPin = CreatePin(EGPD_Input, UEdGraphSchema_K2::PC_Object, UEdGraphSchema_K2::PSC_Self, UEdGraphSchema_K2::PN_Self);
 	}
 	else
 	{
 		// This means that the function is declared in an external class, and should reference that class
-		SelfPin = CreatePin(EGPD_Input, K2Schema->PC_Object, FString(), MemberVariableToCallOn.GetMemberParentClass(GetBlueprintClassFromNode()), K2Schema->PN_Self);
+		SelfPin = CreatePin(EGPD_Input, UEdGraphSchema_K2::PC_Object, MemberVariableToCallOn.GetMemberParentClass(GetBlueprintClassFromNode()), UEdGraphSchema_K2::PN_Self);
 	}
 	check(SelfPin);
 
diff --git a/Engine/Source/Editor/BlueprintGraph/Private/K2Node_CastByteToEnum.cpp b/Engine/Source/Editor/BlueprintGraph/Private/K2Node_CastByteToEnum.cpp
--- a/Engine/Source/Editor/BlueprintGraph/Private/K2Node_CastByteToEnum.cpp
+++ b/Engine/Source/Editor/BlueprintGraph/Private/K2Node_CastByteToEnum.cpp
@@ -12,7 +12,7 @@
 #include "EditorCategoryUtils.h"
 #include "BlueprintActionDatabaseRegistrar.h"
 
-const FString UK2Node_CastByteToEnum::ByteInputPinName = TEXT("Byte");
+const FName UK2Node_CastByteToEnum::ByteInputPinName(TEXT("Byte"));
 
 UK2Node_CastByteToEnum::UK2Node_CastByteToEnum(const FObjectInitializer& ObjectInitializer)
 	: Super(ObjectInitializer)
@@ -30,10 +30,8 @@ void UK2Node_CastByteToEnum::ValidateNodeDuringCompilation(class FCompilerResult
 
 void UK2Node_CastByteToEnum::AllocateDefaultPins()
 {
-	const UEdGraphSchema_K2* Schema = GetDefault<UEdGraphSchema_K2>();
-	
-	CreatePin(EGPD_Input, Schema->PC_Byte, FString(), nullptr, ByteInputPinName);
-	CreatePin(EGPD_Output, Schema->PC_Byte, FString(), Enum, Schema->PN_ReturnValue);
+	CreatePin(EGPD_Input, UEdGraphSchema_K2::PC_Byte, ByteInputPinName);
+	CreatePin(EGPD_Output, UEdGraphSchema_K2::PC_Byte, Enum, UEdGraphSchema_K2::PN_ReturnValue);
 }
 
 FText UK2Node_CastByteToEnum::GetTooltipText() const
@@ -100,7 +98,7 @@ void UK2Node_CastByteToEnum::ExpandNode(class FKismetCompilerContext& CompilerCo
 		// FUNCTION INPUT BYTE PIN
 		UEdGraphPin* OrgInputPin = FindPinChecked(ByteInputPinName);
 		UEdGraphPin* FunctionIndexPin = CallValidation->FindPinChecked(TEXT("EnumeratorValue"));
-		check(EGPD_Input == FunctionIndexPin->Direction && Schema->PC_Byte == FunctionIndexPin->PinType.PinCategory);
+		check(EGPD_Input == FunctionIndexPin->Direction && UEdGraphSchema_K2::PC_Byte == FunctionIndexPin->PinType.PinCategory);
 		CompilerContext.MovePinLinksToIntermediate(*OrgInputPin, *FunctionIndexPin);
 
 		// UNSAFE CAST NODE
@@ -116,8 +114,8 @@ void UK2Node_CastByteToEnum::ExpandNode(class FKismetCompilerContext& CompilerCo
 		Schema->TryCreateConnection(CastInputPin, FunctionReturnPin);
 
 		// OPUTPUT PIN
-		UEdGraphPin* OrgReturnPin = FindPinChecked(Schema->PN_ReturnValue);
-		UEdGraphPin* NewReturnPin = UsafeCast->FindPinChecked(Schema->PN_ReturnValue);
+		UEdGraphPin* OrgReturnPin = FindPinChecked(UEdGraphSchema_K2::PN_ReturnValue);
+		UEdGraphPin* NewReturnPin = UsafeCast->FindPinChecked(UEdGraphSchema_K2::PN_ReturnValue);
 		CompilerContext.MovePinLinksToIntermediate(*OrgReturnPin, *NewReturnPin);
 
 		BreakAllNodeLinks();
@@ -146,9 +144,8 @@ public:
 
 		FBPTerminal** ValueSource = Context.NetMap.Find(Net);
 		check(ValueSource && *ValueSource);
-		const UEdGraphSchema_K2* Schema = GetDefault<UEdGraphSchema_K2>();
-		UEdGraphPin* OutPin = Node->FindPinChecked(Schema->PN_ReturnValue);
-		if (ensure(Context.NetMap.Find(OutPin) == NULL))
+		UEdGraphPin* OutPin = Node->FindPinChecked(UEdGraphSchema_K2::PN_ReturnValue);
+		if (ensure(Context.NetMap.Find(OutPin) == nullptr))
 		{
 			// We need to copy here to avoid passing in a reference to an element inside the map. The array
 			// that owns the map members could be reallocated, causing the reference to become stale.
diff --git a/Engine/Source/Editor/BlueprintGraph/Private/K2Node_ClassDynamicCast.cpp b/Engine/Source/Editor/BlueprintGraph/Private/K2Node_ClassDynamicCast.cpp
--- a/Engine/Source/Editor/BlueprintGraph/Private/K2Node_ClassDynamicCast.cpp
+++ b/Engine/Source/Editor/BlueprintGraph/Private/K2Node_ClassDynamicCast.cpp
@@ -10,16 +10,13 @@
 
 struct FClassDynamicCastHelper
 {
-	static const FString CastSuccessPinName;
-
-	static const FString& GetClassToCastName()
-	{
-		static FString ClassToCastName(TEXT("Class"));
-		return ClassToCastName;
-	}
+	static const FName CastSuccessPinName;
+	static const FName ClassToCastName;;
 };
 
-const FString FClassDynamicCastHelper::CastSuccessPinName = TEXT("bSuccess");
+const FName FClassDynamicCastHelper::CastSuccessPinName(TEXT("bSuccess"));
+const FName FClassDynamicCastHelper::ClassToCastName(TEXT("Class"));
+
 
 UK2Node_ClassDynamicCast::UK2Node_ClassDynamicCast(const FObjectInitializer& ObjectInitializer)
 	: Super(ObjectInitializer)
@@ -30,7 +27,7 @@ void UK2Node_ClassDynamicCast::AllocateDefaultPins()
 {
 	// Check to track down possible BP comms corruption
 	//@TODO: Move this somewhere more sensible
-	ensure((TargetType == NULL) || (!TargetType->HasAnyClassFlags(CLASS_NewerVersionExists)));
+	ensure((TargetType == nullptr) || (!TargetType->HasAnyClassFlags(CLASS_NewerVersionExists)));
 
 	const UEdGraphSchema_K2* K2Schema = Cast<UEdGraphSchema_K2>(GetSchema());
 	check(K2Schema != nullptr);
@@ -42,24 +39,24 @@ void UK2Node_ClassDynamicCast::AllocateDefaultPins()
 	if (!bIsPureCast)
 	{
 		// Input - Execution Pin
-		CreatePin(EGPD_Input, K2Schema->PC_Exec, FString(), nullptr, K2Schema->PN_Execute);
+		CreatePin(EGPD_Input, UEdGraphSchema_K2::PC_Exec, UEdGraphSchema_K2::PN_Execute);
 
 		// Output - Execution Pins
-		CreatePin(EGPD_Output, K2Schema->PC_Exec, FString(), nullptr, K2Schema->PN_CastSucceeded);
-		CreatePin(EGPD_Output, K2Schema->PC_Exec, FString(), nullptr, K2Schema->PN_CastFailed);
+		CreatePin(EGPD_Output, UEdGraphSchema_K2::PC_Exec, UEdGraphSchema_K2::PN_CastSucceeded);
+		CreatePin(EGPD_Output, UEdGraphSchema_K2::PC_Exec, UEdGraphSchema_K2::PN_CastFailed);
 	}
 
 	// Input - Source type Pin
-	CreatePin(EGPD_Input, K2Schema->PC_Class, FString(), UObject::StaticClass(), FClassDynamicCastHelper::GetClassToCastName());
+	CreatePin(EGPD_Input, UEdGraphSchema_K2::PC_Class, UObject::StaticClass(), FClassDynamicCastHelper::ClassToCastName);
 
 	// Output - Data Pin
-	if (TargetType != NULL)
+	if (TargetType)
 	{
-		const FString CastResultPinName = K2Schema->PN_CastedValuePrefix + TargetType->GetDisplayNameText().ToString();
-		CreatePin(EGPD_Output, K2Schema->PC_Class, FString(), *TargetType, CastResultPinName);
+		const FString CastResultPinName = UEdGraphSchema_K2::PN_CastedValuePrefix + TargetType->GetDisplayNameText().ToString();
+		CreatePin(EGPD_Output, UEdGraphSchema_K2::PC_Class, *TargetType, *CastResultPinName);
 	}
 
-	UEdGraphPin* BoolSuccessPin = CreatePin(EGPD_Output, K2Schema->PC_Boolean, FString(), nullptr, FClassDynamicCastHelper::CastSuccessPinName);
+	UEdGraphPin* BoolSuccessPin = CreatePin(EGPD_Output, UEdGraphSchema_K2::PC_Boolean, FClassDynamicCastHelper::CastSuccessPinName);
 	BoolSuccessPin->bHidden = !bIsPureCast;
 
 	UK2Node::AllocateDefaultPins();
@@ -81,7 +78,7 @@ FText UK2Node_ClassDynamicCast::GetNodeTitle(ENodeTitleType::Type TitleType) con
 
 UEdGraphPin* UK2Node_ClassDynamicCast::GetCastSourcePin() const
 {
-	UEdGraphPin* Pin = FindPinChecked(FClassDynamicCastHelper::GetClassToCastName());
+	UEdGraphPin* Pin = FindPinChecked(FClassDynamicCastHelper::ClassToCastName);
 	check(Pin->Direction == EGPD_Input);
 	return Pin;
 }
diff --git a/Engine/Source/Editor/BlueprintGraph/Private/K2Node_CommutativeAssociativeBinaryOperator.cpp b/Engine/Source/Editor/BlueprintGraph/Private/K2Node_CommutativeAssociativeBinaryOperator.cpp
--- a/Engine/Source/Editor/BlueprintGraph/Private/K2Node_CommutativeAssociativeBinaryOperator.cpp
+++ b/Engine/Source/Editor/BlueprintGraph/Private/K2Node_CommutativeAssociativeBinaryOperator.cpp
@@ -17,11 +17,10 @@ int32 UK2Node_CommutativeAssociativeBinaryOperator::GetMaxInputPinsNum()
 	return (TCHAR('Z') - TCHAR('A'));
 }
 
-FString UK2Node_CommutativeAssociativeBinaryOperator::GetNameForPin(int32 PinIndex)
+FName UK2Node_CommutativeAssociativeBinaryOperator::GetNameForPin(int32 PinIndex)
 {
 	check(PinIndex < GetMaxInputPinsNum());
-	FString Name;
-	Name.AppendChar(TCHAR('A') + PinIndex);
+	const FName Name(*FString::Chr(TCHAR('A') + PinIndex));
 	return Name;
 }
 
@@ -45,10 +44,9 @@ UEdGraphPin* UK2Node_CommutativeAssociativeBinaryOperator::FindOutPin() const
 
 UEdGraphPin* UK2Node_CommutativeAssociativeBinaryOperator::FindSelfPin() const
 {
-	const UEdGraphSchema_K2* K2Schema = GetDefault<UEdGraphSchema_K2>();
 	for(int32 PinIdx=0; PinIdx<Pins.Num(); PinIdx++)
 	{
-		if(Pins[PinIdx]->PinName == K2Schema->PN_Self)
+		if(Pins[PinIdx]->PinName == UEdGraphSchema_K2::PN_Self)
 		{
 			return Pins[PinIdx];
 		}
@@ -153,16 +151,16 @@ void UK2Node_CommutativeAssociativeBinaryOperator::AllocateDefaultPins()
 void UK2Node_CommutativeAssociativeBinaryOperator::AddInputPinInner(int32 AdditionalPinIndex)
 {
 	const FEdGraphPinType InputType = GetType();
+	UEdGraphNode::FCreatePinParams PinParams;
+	PinParams.ContainerType = InputType.ContainerType;
+	PinParams.bIsReference = InputType.bIsReference;
+	PinParams.ValueTerminalType = InputType.PinValueType;
 	CreatePin(EGPD_Input, 
 		InputType.PinCategory, 
 		InputType.PinSubCategory, 
 		InputType.PinSubCategoryObject.Get(), 
-		*GetNameForPin(AdditionalPinIndex + BinaryOperatorInputsNum),
-		InputType.ContainerType, 
-		InputType.bIsReference, 
-		false,
-		INDEX_NONE,
-		InputType.PinValueType
+		GetNameForPin(AdditionalPinIndex + BinaryOperatorInputsNum),
+		PinParams
 	);
 }
 
@@ -199,8 +197,8 @@ void UK2Node_CommutativeAssociativeBinaryOperator::RemoveInputPin(UEdGraphPin* P
 				UEdGraphPin* LocalPin = Pins[PinIndex];
 				if(LocalPin && (LocalPin != OutPin) && (LocalPin != SelfPin))
 				{
-					const FString PinName = GetNameForPin(NameIndex);
-					if(PinName != LocalPin->PinName)
+					const FName PinName = GetNameForPin(NameIndex);
+					if (PinName != LocalPin->PinName)
 					{
 						LocalPin->Modify();
 						LocalPin->PinName = PinName;
diff --git a/Engine/Source/Editor/BlueprintGraph/Private/K2Node_ConstructObjectFromClass.cpp b/Engine/Source/Editor/BlueprintGraph/Private/K2Node_ConstructObjectFromClass.cpp
--- a/Engine/Source/Editor/BlueprintGraph/Private/K2Node_ConstructObjectFromClass.cpp
+++ b/Engine/Source/Editor/BlueprintGraph/Private/K2Node_ConstructObjectFromClass.cpp
@@ -10,14 +10,14 @@
 
 struct FK2Node_ConstructObjectFromClassHelper
 {
-	static FString WorldContextPinName;
-	static FString ClassPinName;
-	static FString OuterPinName;
+	static FName WorldContextPinName;
+	static FName ClassPinName;
+	static FName OuterPinName;
 };
 
-FString FK2Node_ConstructObjectFromClassHelper::WorldContextPinName(TEXT("WorldContextObject"));
-FString FK2Node_ConstructObjectFromClassHelper::ClassPinName(TEXT("Class"));
-FString FK2Node_ConstructObjectFromClassHelper::OuterPinName(TEXT("Outer"));
+FName FK2Node_ConstructObjectFromClassHelper::WorldContextPinName(TEXT("WorldContextObject"));
+FName FK2Node_ConstructObjectFromClassHelper::ClassPinName(TEXT("Class"));
+FName FK2Node_ConstructObjectFromClassHelper::OuterPinName(TEXT("Outer"));
 
 #define LOCTEXT_NAMESPACE "K2Node_ConstructObjectFromClass"
 
@@ -42,24 +42,24 @@ bool UK2Node_ConstructObjectFromClass::UseWorldContext() const
 void UK2Node_ConstructObjectFromClass::AllocateDefaultPins()
 {
 	// Add execution pins
-	CreatePin(EGPD_Input, UEdGraphSchema_K2::PC_Exec, FString(), nullptr, UEdGraphSchema_K2::PN_Execute);
-	CreatePin(EGPD_Output, UEdGraphSchema_K2::PC_Exec, FString(), nullptr, UEdGraphSchema_K2::PN_Then);
+	CreatePin(EGPD_Input, UEdGraphSchema_K2::PC_Exec, UEdGraphSchema_K2::PN_Execute);
+	CreatePin(EGPD_Output, UEdGraphSchema_K2::PC_Exec, UEdGraphSchema_K2::PN_Then);
 
 	// If required add the world context pin
 	if (UseWorldContext())
 	{
-		CreatePin(EGPD_Input, UEdGraphSchema_K2::PC_Object, FString(), UObject::StaticClass(), FK2Node_ConstructObjectFromClassHelper::WorldContextPinName);
+		CreatePin(EGPD_Input, UEdGraphSchema_K2::PC_Object, UObject::StaticClass(), FK2Node_ConstructObjectFromClassHelper::WorldContextPinName);
 	}
 
 	// Add blueprint pin
-	UEdGraphPin* ClassPin = CreatePin(EGPD_Input, UEdGraphSchema_K2::PC_Class, FString(), GetClassPinBaseClass(), FK2Node_ConstructObjectFromClassHelper::ClassPinName);
+	UEdGraphPin* ClassPin = CreatePin(EGPD_Input, UEdGraphSchema_K2::PC_Class, GetClassPinBaseClass(), FK2Node_ConstructObjectFromClassHelper::ClassPinName);
 	
 	// Result pin
-	UEdGraphPin* ResultPin = CreatePin(EGPD_Output, UEdGraphSchema_K2::PC_Object, FString(), GetClassPinBaseClass(), UEdGraphSchema_K2::PN_ReturnValue);
+	UEdGraphPin* ResultPin = CreatePin(EGPD_Output, UEdGraphSchema_K2::PC_Object, GetClassPinBaseClass(), UEdGraphSchema_K2::PN_ReturnValue);
 	
 	if (UseOuter())
 	{
-		UEdGraphPin* OuterPin = CreatePin(EGPD_Input, UEdGraphSchema_K2::PC_Object, FString(), UObject::StaticClass(), FK2Node_ConstructObjectFromClassHelper::OuterPinName);
+		UEdGraphPin* OuterPin = CreatePin(EGPD_Input, UEdGraphSchema_K2::PC_Object, UObject::StaticClass(), FK2Node_ConstructObjectFromClassHelper::OuterPinName);
 	}
 
 	Super::AllocateDefaultPins();
@@ -107,10 +107,10 @@ void UK2Node_ConstructObjectFromClass::CreatePinsForClass(UClass* InClass, TArra
 			bIsSettableExternally &&
 			Property->HasAllPropertyFlags(CPF_BlueprintVisible) &&
 			!bIsDelegate &&
-			(nullptr == FindPin(Property->GetName()) ) &&
+			(nullptr == FindPin(Property->GetFName()) ) &&
 			FBlueprintEditorUtils::PropertyStillExists(Property))
 		{
-			if (UEdGraphPin* Pin = CreatePin(EGPD_Input, FString(), FString(), nullptr, Property->GetName()))
+			if (UEdGraphPin* Pin = CreatePin(EGPD_Input, NAME_None, Property->GetFName()))
 			{
 				K2Schema->ConvertPropertyToPinType(Property, /*out*/ Pin->PinType);
 				if (OutClassPins)
diff --git a/Engine/Source/Editor/BlueprintGraph/Private/K2Node_ConvertAsset.cpp b/Engine/Source/Editor/BlueprintGraph/Private/K2Node_ConvertAsset.cpp
--- a/Engine/Source/Editor/BlueprintGraph/Private/K2Node_ConvertAsset.cpp
+++ b/Engine/Source/Editor/BlueprintGraph/Private/K2Node_ConvertAsset.cpp
@@ -17,8 +17,8 @@
 
 namespace UK2Node_ConvertAssetImpl
 {
-	static const FString InputPinName("Input");
-	static const FString OutputPinName("Output");
+	static const FName InputPinName("Input");
+	static const FName OutputPinName("Output");
 }
 
 UClass* UK2Node_ConvertAsset::GetTargetClass() const
@@ -76,8 +76,8 @@ void UK2Node_ConvertAsset::RefreshPinTypes()
 		const bool bIsAssetClass = bIsConnected ? IsAssetClassType() : false;
 		const bool bIsConvertToAsset = bIsConnected ? IsConvertToAsset() : false;
 
-		FString InputCategory = UEdGraphSchema_K2::PC_Wildcard;
-		FString OutputCategory = UEdGraphSchema_K2::PC_Wildcard;
+		FName InputCategory = UEdGraphSchema_K2::PC_Wildcard;
+		FName OutputCategory = UEdGraphSchema_K2::PC_Wildcard;
 		if (bIsConnected)
 		{
 			if (bIsConvertToAsset)
@@ -92,8 +92,8 @@ void UK2Node_ConvertAsset::RefreshPinTypes()
 			}
 		}
 			
-		InputPin->PinType = FEdGraphPinType(InputCategory, FString(), TargetType, EPinContainerType::None, false, FEdGraphTerminalType() );
-		OutputPin->PinType = FEdGraphPinType(OutputCategory, FString(), TargetType, EPinContainerType::None, false, FEdGraphTerminalType() );
+		InputPin->PinType = FEdGraphPinType(InputCategory, NAME_None, TargetType, EPinContainerType::None, false, FEdGraphTerminalType() );
+		OutputPin->PinType = FEdGraphPinType(OutputCategory, NAME_None, TargetType, EPinContainerType::None, false, FEdGraphTerminalType() );
 
 		PinTypeChanged(InputPin);
 		PinTypeChanged(OutputPin);
@@ -102,7 +102,7 @@ void UK2Node_ConvertAsset::RefreshPinTypes()
 		{
 			TArray<UEdGraphPin*> PinsToUnlink = OutputPin->LinkedTo;
 
-			UClass const* CallingContext = NULL;
+			UClass const* CallingContext = nullptr;
 			if (UBlueprint const* Blueprint = GetBlueprint())
 			{
 				CallingContext = Blueprint->GeneratedClass;
@@ -143,8 +143,8 @@ void UK2Node_ConvertAsset::NotifyPinConnectionListChanged(UEdGraphPin* Pin)
 
 void UK2Node_ConvertAsset::AllocateDefaultPins()
 {
-	CreatePin(EGPD_Input, UEdGraphSchema_K2::PC_Wildcard, FString(), nullptr, UK2Node_ConvertAssetImpl::InputPinName);
-	CreatePin(EGPD_Output, UEdGraphSchema_K2::PC_Wildcard, FString(), nullptr, UK2Node_ConvertAssetImpl::OutputPinName);
+	CreatePin(EGPD_Input, UEdGraphSchema_K2::PC_Wildcard, UK2Node_ConvertAssetImpl::InputPinName);
+	CreatePin(EGPD_Output, UEdGraphSchema_K2::PC_Wildcard, UK2Node_ConvertAssetImpl::OutputPinName);
 }
 
 UK2Node::ERedirectType UK2Node_ConvertAsset::DoPinsMatchForReconstruction(const UEdGraphPin* NewPin, int32 NewPinIndex, const UEdGraphPin* OldPin, int32 OldPinIndex) const
@@ -202,9 +202,7 @@ void UK2Node_ConvertAsset::ExpandNode(class FKismetCompilerContext& CompilerCont
 
 			//Connect input to convert
 			UEdGraphPin* InputPin = FindPin(UK2Node_ConvertAssetImpl::InputPinName);
-			const FString ConvertInputName = bIsAssetClass
-				? FString(TEXT("Class"))
-				: FString(TEXT("Object"));
+			const FName ConvertInputName = bIsAssetClass ? FName(TEXT("Class")) : FName(TEXT("Object"));
 			UEdGraphPin* ConvertInput = ConvertToObjectFunc->FindPin(ConvertInputName);
 			bIsErrorFree = InputPin && ConvertInput && CompilerContext.MovePinLinksToIntermediate(*InputPin, *ConvertInput).CanSafeConnect();
 
@@ -234,9 +232,7 @@ void UK2Node_ConvertAsset::ExpandNode(class FKismetCompilerContext& CompilerCont
 
 			//Connect input to convert
 			UEdGraphPin* InputPin = FindPin(UK2Node_ConvertAssetImpl::InputPinName);
-			const FString ConvertInputName = bIsAssetClass
-				? FString(TEXT("SoftClass"))
-				: FString(TEXT("SoftObject"));
+			const FName ConvertInputName = bIsAssetClass ? FName(TEXT("SoftClass")) : FName(TEXT("SoftObject"));
 			UEdGraphPin* ConvertInput = ConvertToObjectFunc->FindPin(ConvertInputName);
 			bIsErrorFree = InputPin && ConvertInput && CompilerContext.MovePinLinksToIntermediate(*InputPin, *ConvertInput).CanSafeConnect();
 
diff --git a/Engine/Source/Editor/BlueprintGraph/Private/K2Node_Copy.cpp b/Engine/Source/Editor/BlueprintGraph/Private/K2Node_Copy.cpp
--- a/Engine/Source/Editor/BlueprintGraph/Private/K2Node_Copy.cpp
+++ b/Engine/Source/Editor/BlueprintGraph/Private/K2Node_Copy.cpp
@@ -66,8 +66,8 @@ void UK2Node_Copy::AllocateDefaultPins()
 {
 	const UEdGraphSchema_K2* Schema = GetDefault<UEdGraphSchema_K2>();
 
-	CreatePin(EGPD_Input, Schema->PC_Wildcard, FString(), nullptr, Schema->PN_Item);
-	CreatePin(EGPD_Output, Schema->PC_Wildcard, FString(), nullptr, Schema->PN_ReturnValue);
+	CreatePin(EGPD_Input, UEdGraphSchema_K2::PC_Wildcard, UEdGraphSchema_K2::PN_Item);
+	CreatePin(EGPD_Output, UEdGraphSchema_K2::PC_Wildcard, UEdGraphSchema_K2::PN_ReturnValue);
 
 	Super::AllocateDefaultPins();
 }
@@ -127,19 +127,15 @@ void UK2Node_Copy::NotifyPinConnectionListChanged(UEdGraphPin* Pin)
 
 UEdGraphPin* UK2Node_Copy::GetInputReferencePin() const
 {
-	const UEdGraphSchema_K2* K2Schema = GetDefault<UEdGraphSchema_K2>();
-
-	UEdGraphPin* Pin = FindPin(K2Schema->PN_Item);
-	check(Pin != NULL);
+	UEdGraphPin* Pin = FindPin(UEdGraphSchema_K2::PN_Item);
+	check(Pin);
 	return Pin;
 }
 
 UEdGraphPin* UK2Node_Copy::GetCopyResultPin() const
 {
-	const UEdGraphSchema_K2* K2Schema = GetDefault<UEdGraphSchema_K2>();
-
-	UEdGraphPin* Pin = FindPin(K2Schema->PN_ReturnValue);
-	check(Pin != NULL);
+	UEdGraphPin* Pin = FindPin(UEdGraphSchema_K2::PN_ReturnValue);
+	check(Pin);
 	return Pin;
 }
 
@@ -175,19 +171,18 @@ void UK2Node_Copy::PinTypeChanged(UEdGraphPin* Pin)
 
 bool UK2Node_Copy::IsConnectionDisallowed(const UEdGraphPin* MyPin, const UEdGraphPin* OtherPin, FString& OutReason) const
 {
-	const UEdGraphSchema_K2* K2Schema = GetDefault<UEdGraphSchema_K2>();
 	if (OtherPin)
 	{
-		if (OtherPin->PinType.PinCategory == K2Schema->PC_Exec)
+		if (OtherPin->PinType.PinCategory == UEdGraphSchema_K2::PC_Exec)
 		{
 			OutReason = LOCTEXT("ExecConnectionDisallowed", "Cannot connect with Exec pin.").ToString();
 			return true;
 		}
-		else if (OtherPin->PinType.PinCategory == K2Schema->PC_Object || OtherPin->PinType.PinCategory == K2Schema->PC_Class
-			|| OtherPin->PinType.PinCategory == K2Schema->PC_SoftObject || OtherPin->PinType.PinCategory == K2Schema->PC_SoftClass
-			|| OtherPin->PinType.PinCategory == K2Schema->PC_Interface)
+		else if (OtherPin->PinType.PinCategory == UEdGraphSchema_K2::PC_Object || OtherPin->PinType.PinCategory == UEdGraphSchema_K2::PC_Class
+			|| OtherPin->PinType.PinCategory == UEdGraphSchema_K2::PC_SoftObject || OtherPin->PinType.PinCategory == UEdGraphSchema_K2::PC_SoftClass
+			|| OtherPin->PinType.PinCategory == UEdGraphSchema_K2::PC_Interface)
 		{
-			OutReason = FText::Format(LOCTEXT("ObjectConnectionDisallowed", "Cannot connect with {0} pin."), FText::FromString(OtherPin->PinType.PinCategory)).ToString();
+			OutReason = FText::Format(LOCTEXT("ObjectConnectionDisallowed", "Cannot connect with {0} pin."), FText::FromName(OtherPin->PinType.PinCategory)).ToString();
 			return true;
 		}
 		else if (OtherPin->PinType.IsContainer())
diff --git a/Engine/Source/Editor/BlueprintGraph/Private/K2Node_CreateDelegate.cpp b/Engine/Source/Editor/BlueprintGraph/Private/K2Node_CreateDelegate.cpp
--- a/Engine/Source/Editor/BlueprintGraph/Private/K2Node_CreateDelegate.cpp
+++ b/Engine/Source/Editor/BlueprintGraph/Private/K2Node_CreateDelegate.cpp
@@ -17,11 +17,11 @@
 
 struct FK2Node_CreateDelegate_Helper
 {
-	static FString DelegateOutputName;
-	static FString InputObjectName; // Deprecated, for fixup
+	static FName DelegateOutputName;
+	static FName InputObjectName; // Deprecated, for fixup
 };
-FString FK2Node_CreateDelegate_Helper::DelegateOutputName(TEXT("OutputDelegate"));
-FString FK2Node_CreateDelegate_Helper::InputObjectName(TEXT("InputObject"));
+FName FK2Node_CreateDelegate_Helper::DelegateOutputName(TEXT("OutputDelegate"));
+FName FK2Node_CreateDelegate_Helper::InputObjectName(TEXT("InputObject"));
 
 UK2Node_CreateDelegate::UK2Node_CreateDelegate(const FObjectInitializer& ObjectInitializer)
 	: Super(ObjectInitializer)
@@ -30,14 +30,12 @@ UK2Node_CreateDelegate::UK2Node_CreateDelegate(const FObjectInitializer& ObjectI
 
 void UK2Node_CreateDelegate::AllocateDefaultPins()
 {
-	const UEdGraphSchema_K2* K2Schema = GetDefault<UEdGraphSchema_K2>();
-
-	if (UEdGraphPin* ObjPin = CreatePin(EGPD_Input, K2Schema->PC_Object, FString(), UObject::StaticClass(), K2Schema->PN_Self))
+	if (UEdGraphPin* ObjPin = CreatePin(EGPD_Input, UEdGraphSchema_K2::PC_Object, UObject::StaticClass(), UEdGraphSchema_K2::PN_Self))
 	{
 		ObjPin->PinFriendlyName = NSLOCTEXT("K2Node", "CreateDelegate_ObjectInputName", "Object");
 	}
 
-	if(UEdGraphPin* DelegatePin = CreatePin(EGPD_Output, K2Schema->PC_Delegate, FString(), nullptr, FK2Node_CreateDelegate_Helper::DelegateOutputName))
+	if(UEdGraphPin* DelegatePin = CreatePin(EGPD_Output, UEdGraphSchema_K2::PC_Delegate, FK2Node_CreateDelegate_Helper::DelegateOutputName))
 	{
 		DelegatePin->PinFriendlyName = NSLOCTEXT("K2Node", "CreateDelegate_DelegateOutName", "Event");
 	}
@@ -47,10 +45,8 @@ void UK2Node_CreateDelegate::AllocateDefaultPins()
 
 UK2Node::ERedirectType UK2Node_CreateDelegate::DoPinsMatchForReconstruction(const UEdGraphPin* NewPin, int32 NewPinIndex, const UEdGraphPin* OldPin, int32 OldPinIndex) const
 {
-	const UEdGraphSchema_K2* K2Schema = GetDefault<UEdGraphSchema_K2>();
-
 	// Handles remap of InputObject to Self, from 4.10 time frame
-	if (OldPin->PinName == FK2Node_CreateDelegate_Helper::InputObjectName && NewPin->PinName == K2Schema->PN_Self)
+	if (OldPin->PinName == FK2Node_CreateDelegate_Helper::InputObjectName && NewPin->PinName == UEdGraphSchema_K2::PN_Self)
 	{
 		return ERedirectType_Name;
 	}
@@ -120,10 +116,14 @@ bool UK2Node_CreateDelegate::IsValid(FString* OutMsg, bool bDontUseSkeletalClass
 	{
 		if (OutMsg)
 		{
-			FString SelfPinName = UEdGraphSchema_K2::PN_Self;
+			FString SelfPinName;
 			if (UEdGraphPin* SelfPin = GetObjectInPin())
 			{
-				SelfPinName = SelfPin->PinFriendlyName.IsEmpty() ? SelfPin->PinFriendlyName.ToString() : SelfPin->PinName;
+				SelfPinName = SelfPin->PinFriendlyName.IsEmpty() ? SelfPin->PinFriendlyName.ToString() : SelfPin->PinName.ToString();
+			}
+			else
+			{
+				SelfPinName = UEdGraphSchema_K2::PN_Self.ToString();
 			}
 
 			*OutMsg = FString::Printf(*NSLOCTEXT("K2Node", "Class_not_found", "Unable to determine context for the selected function/event: '%s' - make sure the target '%s' pin is properly set up.").ToString(),
@@ -334,8 +334,7 @@ UFunction* UK2Node_CreateDelegate::GetDelegateSignature() const
 
 UClass* UK2Node_CreateDelegate::GetScopeClass(bool bDontUseSkeletalClassForSelf/* = false*/) const
 {
-	const UEdGraphSchema_K2* K2Schema = GetDefault<UEdGraphSchema_K2>();
-	UEdGraphPin* Pin = FindPin(K2Schema->PN_Self);
+	UEdGraphPin* Pin = FindPin(UEdGraphSchema_K2::PN_Self);
 	if (Pin == nullptr)
 	{
 		// The BlueprintNodeTemplateCache creates nodes but doesn't call allocate default pins.
@@ -353,8 +352,8 @@ UClass* UK2Node_CreateDelegate::GetScopeClass(bool bDontUseSkeletalClassForSelf/
 	{
 		if(UEdGraphPin* ResultPin = Pin->LinkedTo[0])
 		{
-			ensure(K2Schema->PC_Object == ResultPin->PinType.PinCategory);
-			if (K2Schema->PN_Self == ResultPin->PinType.PinSubCategory)
+			ensure(UEdGraphSchema_K2::PC_Object == ResultPin->PinType.PinCategory);
+			if (UEdGraphSchema_K2::PN_Self == ResultPin->PinType.PinSubCategory)
 			{
 				bUseSelf = true;
 			}
diff --git a/Engine/Source/Editor/BlueprintGraph/Private/K2Node_CustomEvent.cpp b/Engine/Source/Editor/BlueprintGraph/Private/K2Node_CustomEvent.cpp
--- a/Engine/Source/Editor/BlueprintGraph/Private/K2Node_CustomEvent.cpp
+++ b/Engine/Source/Editor/BlueprintGraph/Private/K2Node_CustomEvent.cpp
@@ -189,7 +189,7 @@ bool UK2Node_CustomEvent::CanCreateUserDefinedPin(const FEdGraphPinType& InPinTy
 		OutErrorMessage = NSLOCTEXT("K2Node", "AddInputPinError", "Cannot add input pins to custom event node!");
 		return false;
 	}
-	else if (InPinType.PinCategory == Schema->PC_Exec && !CanModifyExecutionWires())
+	else if (InPinType.PinCategory == UEdGraphSchema_K2::PC_Exec && !CanModifyExecutionWires())
 	{
 		OutErrorMessage = LOCTEXT("MultipleExecPinError", "Cannot support more exec pins!");
 		return false;
@@ -393,13 +393,13 @@ void UK2Node_CustomEvent::ReconstructNode()
 		{
 			DelegateSignature = OtherNode->GetDelegateSignature();
 		}
-		else if ( LinkedPin->PinType.PinCategory == K2Schema->PC_Delegate )
+		else if ( LinkedPin->PinType.PinCategory == UEdGraphSchema_K2::PC_Delegate )
 		{
 			DelegateSignature = FMemberReference::ResolveSimpleMemberReference<UFunction>(LinkedPin->PinType.PinSubCategoryMemberReference);
 		}
 	}
 	
-	const bool bUseDelegateSignature = (NULL == FindEventSignatureFunction()) && DelegateSignature;
+	const bool bUseDelegateSignature = (nullptr == FindEventSignatureFunction()) && DelegateSignature;
 
 	if (bUseDelegateSignature)
 	{
@@ -412,12 +412,12 @@ void UK2Node_CustomEvent::ReconstructNode()
 				FEdGraphPinType PinType;
 				K2Schema->ConvertPropertyToPinType(Param, /*out*/ PinType);
 
-				FString NewPinName = Param->GetName();
+				FName NewPinName = Param->GetFName();
 				int32 Index = 1;
-				while ((DelegateOutputName == NewPinName) || (K2Schema->PN_Then == NewPinName))
+				while ((DelegateOutputName == NewPinName) || (UEdGraphSchema_K2::PN_Then == NewPinName))
 				{
 					++Index;
-					NewPinName += FString::FromInt(Index);
+					NewPinName = *FString::Printf(TEXT("%s%d"), *NewPinName.ToString(), Index);
 				}
 				TSharedPtr<FUserPinInfo> NewPinInfo = MakeShareable( new FUserPinInfo() );
 				NewPinInfo->PinName = NewPinName;
@@ -453,7 +453,7 @@ UK2Node_CustomEvent* UK2Node_CustomEvent::CreateFromFunction(FVector2D GraphPosi
 			{
 				FEdGraphPinType PinType;
 				K2Schema->ConvertPropertyToPinType(Param, /*out*/ PinType);
-				CustomEventNode->CreateUserDefinedPin(Param->GetName(), PinType, EGPD_Output);
+				CustomEventNode->CreateUserDefinedPin(Param->GetFName(), PinType, EGPD_Output);
 			}
 		}
 
diff --git a/Engine/Source/Editor/BlueprintGraph/Private/K2Node_DelegateSet.cpp b/Engine/Source/Editor/BlueprintGraph/Private/K2Node_DelegateSet.cpp
--- a/Engine/Source/Editor/BlueprintGraph/Private/K2Node_DelegateSet.cpp
+++ b/Engine/Source/Editor/BlueprintGraph/Private/K2Node_DelegateSet.cpp
@@ -35,7 +35,7 @@ public:
 
 			// Create a term to store the locally created delegate that we'll use to add to the MC delegate
 			FBPTerminal* DelegateTerm = Context.CreateLocalTerminal();
-			DelegateTerm->Type.PinCategory = CompilerContext.GetSchema()->PC_Delegate;
+			DelegateTerm->Type.PinCategory = UEdGraphSchema_K2::PC_Delegate;
 			FMemberReference::FillSimpleMemberReference<UFunction>(DelegateNode->GetDelegateSignature(), DelegateTerm->Type.PinSubCategoryMemberReference);
 			DelegateTerm->Source = Node;
 			DelegateTerm->Name = Context.NetNameMap->MakeValidName(Node) + TEXT("_TempBindingDelegate");
@@ -106,7 +106,7 @@ public:
 
 		// Create a delegate name term
 		FBPTerminal* DelegateNameTerm = Context.CreateLocalTerminal(ETerminalSpecification::TS_Literal);
-		DelegateNameTerm->Type.PinCategory = CompilerContext.GetSchema()->PC_Name;
+		DelegateNameTerm->Type.PinCategory = UEdGraphSchema_K2::PC_Name;
 		DelegateNameTerm->Name = DelegateNode->GetDelegateTargetEntryPointName().ToString();
 		DelegateNameTerm->bIsLiteral = true;
 
@@ -123,7 +123,7 @@ public:
 		AddStatement.LHS = *pDelegateOwnerTerm;
 		AddStatement.RHS.Add(LocalDelegate);
 
-		GenerateSimpleThenGoto(Context, *Node, DelegateNode->FindPin(CompilerContext.GetSchema()->PN_Then));
+		GenerateSimpleThenGoto(Context, *Node, DelegateNode->FindPin(UEdGraphSchema_K2::PN_Then));
 
 		FNodeHandlingFunctor::Compile(Context, Node);
 	}
@@ -136,16 +136,14 @@ UK2Node_DelegateSet::UK2Node_DelegateSet(const FObjectInitializer& ObjectInitial
 
 void UK2Node_DelegateSet::AllocateDefaultPins()
 {
-	const UEdGraphSchema_K2* K2Schema = GetDefault<UEdGraphSchema_K2>();
-
 	// Cache off the delegate signature, which will update the DelegatePropertyName as well, if it's been redirected
 	UFunction* DelegateSignature = GetDelegateSignature();
 
-	CreatePin(EGPD_Input, K2Schema->PC_Exec, FString(), nullptr, K2Schema->PN_Execute);
-	CreatePin(EGPD_Input, K2Schema->PC_Object, FString(), DelegatePropertyClass, DelegatePropertyName.ToString());
+	CreatePin(EGPD_Input, UEdGraphSchema_K2::PC_Exec, UEdGraphSchema_K2::PN_Execute);
+	CreatePin(EGPD_Input, UEdGraphSchema_K2::PC_Object, DelegatePropertyClass, DelegatePropertyName);
 
-	CreatePin(EGPD_Output, K2Schema->PC_Exec, FString(), nullptr, K2Schema->PN_Then);
-	CreatePin(EGPD_Output, K2Schema->PC_Exec, FString(), nullptr, K2Schema->PN_DelegateEntry);
+	CreatePin(EGPD_Output, UEdGraphSchema_K2::PC_Exec, UEdGraphSchema_K2::PN_Then);
+	CreatePin(EGPD_Output, UEdGraphSchema_K2::PC_Exec, UEdGraphSchema_K2::PN_DelegateEntry);
 	
 	CreatePinsForFunctionEntryExit(DelegateSignature, true);
 
@@ -185,10 +183,8 @@ FText UK2Node_DelegateSet::GetNodeTitle(ENodeTitleType::Type TitleType) const
 
 UEdGraphPin* UK2Node_DelegateSet::GetDelegateOwner() const
 {
-	const UEdGraphSchema_K2* K2Schema = GetDefault<UEdGraphSchema_K2>();
-
-	UEdGraphPin* Pin = FindPin(DelegatePropertyName.ToString());
-	check(Pin != NULL);
+	UEdGraphPin* Pin = FindPin(DelegatePropertyName);
+	check(Pin);
 	check(Pin->Direction == EGPD_Input);
 	return Pin;
 }
@@ -246,11 +242,10 @@ FNodeHandlingFunctor* UK2Node_DelegateSet::CreateNodeHandler(FKismetCompilerCont
 UK2Node::ERedirectType UK2Node_DelegateSet::DoPinsMatchForReconstruction(const UEdGraphPin* NewPin, int32 NewPinIndex, const UEdGraphPin* OldPin, int32 OldPinIndex) const
 {
 	ERedirectType OrginalResult = Super::DoPinsMatchForReconstruction(NewPin, NewPinIndex, OldPin, OldPinIndex);
-	const UEdGraphSchema_K2* K2Schema = Cast<const UEdGraphSchema_K2>(GetSchema());
-	if ((ERedirectType::ERedirectType_None == OrginalResult) && K2Schema && NewPin && OldPin)
+	if ((ERedirectType::ERedirectType_None == OrginalResult) && NewPin && OldPin)
 	{
-		bool const bOldPinIsObj = (OldPin->PinType.PinCategory == K2Schema->PC_Object) || (OldPin->PinType.PinCategory == K2Schema->PC_Interface);
-		bool const bNewPinIsObj = (OldPin->PinType.PinCategory == K2Schema->PC_Object) || (OldPin->PinType.PinCategory == K2Schema->PC_Interface);
+		bool const bOldPinIsObj = (OldPin->PinType.PinCategory == UEdGraphSchema_K2::PC_Object) || (OldPin->PinType.PinCategory == UEdGraphSchema_K2::PC_Interface);
+		bool const bNewPinIsObj = (OldPin->PinType.PinCategory == UEdGraphSchema_K2::PC_Object) || (OldPin->PinType.PinCategory == UEdGraphSchema_K2::PC_Interface);
 
 		if ((NewPin->Direction == EGPD_Input && OldPin->Direction == EGPD_Input) &&
 			bOldPinIsObj && bNewPinIsObj)
@@ -271,11 +266,8 @@ void UK2Node_DelegateSet::ExpandNode(class FKismetCompilerContext& CompilerConte
 	}
 	else
 	{
-		UFunction* TargetFunction = GetDelegateSignature();
-		if(TargetFunction != NULL)
+		if (UFunction* TargetFunction = GetDelegateSignature())
 		{
-			const UEdGraphSchema_K2* Schema = CompilerContext.GetSchema();
-
 			// First, create an event node matching the delegate signature
 			UK2Node_Event* DelegateEvent = CompilerContext.SpawnIntermediateEventNode<UK2Node_Event>(this, nullptr, SourceGraph);
 			DelegateEvent->EventReference.SetFromField<UFunction>(TargetFunction, false);
@@ -291,10 +283,10 @@ void UK2Node_DelegateSet::ExpandNode(class FKismetCompilerContext& CompilerConte
 
 				if( CurrentPin->Direction == EGPD_Output )
 				{
-					if( CurrentPin->PinType.PinCategory == Schema->PC_Exec )
+					if( CurrentPin->PinType.PinCategory == UEdGraphSchema_K2::PC_Exec )
 					{
 						// Hook up the exec pin specially, since it has a different name on the dynamic delegate node
-						UEdGraphPin* OldExecPin = FindPin(Schema->PN_DelegateEntry);
+						UEdGraphPin* OldExecPin = FindPin(UEdGraphSchema_K2::PN_DelegateEntry);
 						check(OldExecPin);
 						CompilerContext.MovePinLinksToIntermediate(*OldExecPin, *CurrentPin);
 					}
diff --git a/Engine/Source/Editor/BlueprintGraph/Private/K2Node_DoOnceMultiInput.cpp b/Engine/Source/Editor/BlueprintGraph/Private/K2Node_DoOnceMultiInput.cpp
--- a/Engine/Source/Editor/BlueprintGraph/Private/K2Node_DoOnceMultiInput.cpp
+++ b/Engine/Source/Editor/BlueprintGraph/Private/K2Node_DoOnceMultiInput.cpp
@@ -22,13 +22,13 @@ UK2Node::ERedirectType UK2Node_DoOnceMultiInput::DoPinsMatchForReconstruction(co
 {
 	// Temp work around: remove whitespaces from pin names before doing string comparison.
 
-	FString NewName = NewPin->PinName;
-	FString OldName = OldPin->PinName;
+	FString NewName = NewPin->PinName.ToString();
+	FString OldName = OldPin->PinName.ToString();
 
 	NewName.ReplaceInline(TEXT(" "), TEXT(""));
 	OldName.ReplaceInline(TEXT(" "), TEXT(""));
 	
-	if (NewName.Equals(OldName, ESearchCase::IgnoreCase))
+	if (NewName == OldName)
 	{
 		// Make sure we're not dealing with a menu node
 		UEdGraph* OuterGraph = GetGraph();
@@ -88,20 +88,19 @@ UEdGraphPin* UK2Node_DoOnceMultiInput::FindOutPin() const
 
 UEdGraphPin* UK2Node_DoOnceMultiInput::FindSelfPin() const
 {
-	const UEdGraphSchema_K2* K2Schema = GetDefault<UEdGraphSchema_K2>();
-	for(int32 PinIdx=0; PinIdx<Pins.Num(); PinIdx++)
+	for (UEdGraphPin* Pin : Pins)
 	{
-		if(Pins[PinIdx]->PinName == K2Schema->PN_Self)
+		if (Pin->PinName == UEdGraphSchema_K2::PN_Self)
 		{
-			return Pins[PinIdx];
+			return Pin;
 		}
 	}
-	return NULL;
+	return nullptr;
 }
 
 bool UK2Node_DoOnceMultiInput::CanAddPin() const
 {
-	return (NumAdditionalInputs) < GetMaxInputPinsNum();
+	return (NumAdditionalInputs < GetMaxInputPinsNum());
 }
 
 bool UK2Node_DoOnceMultiInput::CanRemovePin(const UEdGraphPin* Pin) const
@@ -182,20 +181,18 @@ void UK2Node_DoOnceMultiInput::AllocateDefaultPins()
 {
 	Super::AllocateDefaultPins();
 
-	const UEdGraphSchema_K2* K2Schema = GetDefault<UEdGraphSchema_K2>();
-
 	FText InputPinAName = GetNameForPin(0, true);
-	UEdGraphPin* InputPinA = CreatePin(EGPD_Input, K2Schema->PC_Exec, FString(), nullptr, InputPinAName.BuildSourceString());
+	UEdGraphPin* InputPinA = CreatePin(EGPD_Input, UEdGraphSchema_K2::PC_Exec, NAME_None, nullptr, *InputPinAName.BuildSourceString());
 	InputPinA->PinFriendlyName = InputPinAName;
 
 	FText OutputPinAName = GetNameForPin(0, false);
-	UEdGraphPin* OutputPinA = CreatePin(EGPD_Output, K2Schema->PC_Exec, FString(), nullptr, OutputPinAName.BuildSourceString());
+	UEdGraphPin* OutputPinA = CreatePin(EGPD_Output, UEdGraphSchema_K2::PC_Exec, NAME_None, nullptr, *OutputPinAName.BuildSourceString());
 	OutputPinA->PinFriendlyName = OutputPinAName;
 
-	UEdGraphPin* DoOnceResetIn = CreatePin(EGPD_Input, K2Schema->PC_Exec, FString(), nullptr, TEXT("Reset In"));
+	UEdGraphPin* DoOnceResetIn = CreatePin(EGPD_Input, UEdGraphSchema_K2::PC_Exec, NAME_None, nullptr, TEXT("Reset In"));
 	DoOnceResetIn->PinFriendlyName = LOCTEXT("DoOnceResetIn", "Reset In");
 
-	UEdGraphPin* DoOnceResetOut = CreatePin(EGPD_Output, K2Schema->PC_Exec, FString(), nullptr, TEXT("Reset Out"));
+	UEdGraphPin* DoOnceResetOut = CreatePin(EGPD_Output, UEdGraphSchema_K2::PC_Exec, NAME_None, nullptr, TEXT("Reset Out"));
 	DoOnceResetOut->PinFriendlyName = LOCTEXT("DoOnceResetOut", "Reset Out");
 
 	for (int32 i = 0; i < NumAdditionalInputs; ++i)
@@ -209,14 +206,14 @@ void UK2Node_DoOnceMultiInput::AddPinsInner(int32 AdditionalPinIndex)
 	{
 		const FEdGraphPinType InputType = GetInType();
 		FText InputPinName = GetNameForPin(AdditionalPinIndex, true);
-		UEdGraphPin* InputPin = CreatePin(EGPD_Input, InputType, InputPinName.BuildSourceString());
+		UEdGraphPin* InputPin = CreatePin(EGPD_Input, InputType, *InputPinName.BuildSourceString());
 		InputPin->PinFriendlyName = InputPinName;
 	}
 
 	{
 		const FEdGraphPinType OutputType = GetOutType();
 		FText OutputPinName = GetNameForPin(AdditionalPinIndex, false);
-		UEdGraphPin* OutputPin = CreatePin(EGPD_Output, OutputType, OutputPinName.BuildSourceString());
+		UEdGraphPin* OutputPin = CreatePin(EGPD_Output, OutputType, *OutputPinName.BuildSourceString());
 		OutputPin->PinFriendlyName = OutputPinName;
 	}
 }
@@ -257,7 +254,7 @@ void UK2Node_DoOnceMultiInput::RemoveInputPin(UEdGraphPin* Pin)
 				UEdGraphPin* LocalPin = Pins[PinIndex];
 				if(LocalPin && (LocalPin != OutPin) && (LocalPin != SelfPin))
 				{
-					const FString PinName = GetNameForPin(NameIndex + NumBaseInputs, true).BuildSourceString();  // FIXME
+					const FName PinName = *GetNameForPin(NameIndex + NumBaseInputs, true).BuildSourceString();  // FIXME
 					if(PinName != LocalPin->PinName)
 					{
 						LocalPin->Modify();
@@ -323,7 +320,7 @@ void UK2Node_DoOnceMultiInput::ExpandNode(FKismetCompilerContext& CompilerContex
 
 	// Create the node
 	UK2Node_TemporaryVariable* TempVarNode = SourceGraph->CreateIntermediateNode<UK2Node_TemporaryVariable>();
-	TempVarNode->VariableType.PinCategory = Schema->PC_Boolean;
+	TempVarNode->VariableType.PinCategory = UEdGraphSchema_K2::PC_Boolean;
 	TempVarNode->AllocateDefaultPins();
 	CompilerContext.MessageLog.NotifyIntermediateObjectCreation(TempVarNode, this);
 	// Give a reference of the variable node to the multi gate node
@@ -347,7 +344,7 @@ void UK2Node_DoOnceMultiInput::ExpandNode(FKismetCompilerContext& CompilerContex
 		AssignmentNode->GetVariablePin()->MakeLinkTo(TempVarNode->GetVariablePin());
 		AssignmentNode->GetValuePin()->PinType = TempVarNode->GetVariablePin()->PinType;
 			
-		if (!ExecPin->PinName.Contains(TEXT("Reset"))) // Fixme this wont work for localization
+		if (!ExecPin->PinName.ToString().Contains(TEXT("Reset"))) // Fixme this wont work for localization
 		{
 			// BranchNode
 			UK2Node_IfThenElse* BranchNode = SourceGraph->CreateIntermediateNode<UK2Node_IfThenElse>();
diff --git a/Engine/Source/Editor/BlueprintGraph/Private/K2Node_DynamicCast.cpp b/Engine/Source/Editor/BlueprintGraph/Private/K2Node_DynamicCast.cpp
--- a/Engine/Source/Editor/BlueprintGraph/Private/K2Node_DynamicCast.cpp
+++ b/Engine/Source/Editor/BlueprintGraph/Private/K2Node_DynamicCast.cpp
@@ -19,7 +19,7 @@
 
 namespace UK2Node_DynamicCastImpl
 {
-	static const FString CastSuccessPinName("bSuccess");
+	static const FName CastSuccessPinName("bSuccess");
 }
 
 UK2Node_DynamicCast::UK2Node_DynamicCast(const FObjectInitializer& ObjectInitializer)
@@ -47,31 +47,31 @@ void UK2Node_DynamicCast::AllocateDefaultPins()
 	if (!bIsPureCast)
 	{
 		// Input - Execution Pin
-		CreatePin(EGPD_Input, K2Schema->PC_Exec, FString(), nullptr, K2Schema->PN_Execute);
+		CreatePin(EGPD_Input, UEdGraphSchema_K2::PC_Exec, UEdGraphSchema_K2::PN_Execute);
 
 		// Output - Execution Pins
-		CreatePin(EGPD_Output, K2Schema->PC_Exec, FString(), nullptr, K2Schema->PN_CastSucceeded);
-		CreatePin(EGPD_Output, K2Schema->PC_Exec, FString(), nullptr, K2Schema->PN_CastFailed);
+		CreatePin(EGPD_Output, UEdGraphSchema_K2::PC_Exec, UEdGraphSchema_K2::PN_CastSucceeded);
+		CreatePin(EGPD_Output, UEdGraphSchema_K2::PC_Exec, UEdGraphSchema_K2::PN_CastFailed);
 	}
 
 	// Input - Source type Pin
-	CreatePin(EGPD_Input, K2Schema->PC_Wildcard, FString(), UObject::StaticClass(), K2Schema->PN_ObjectToCast);
+	CreatePin(EGPD_Input, UEdGraphSchema_K2::PC_Wildcard, UObject::StaticClass(), UEdGraphSchema_K2::PN_ObjectToCast);
 
 	// Output - Data Pin
-	if (TargetType != NULL)
+	if (TargetType)
 	{
-		FString CastResultPinName = K2Schema->PN_CastedValuePrefix + TargetType->GetDisplayNameText().ToString();
+		const FString CastResultPinName = UEdGraphSchema_K2::PN_CastedValuePrefix + TargetType->GetDisplayNameText().ToString();
 		if (TargetType->IsChildOf(UInterface::StaticClass()))
 		{
-			CreatePin(EGPD_Output, K2Schema->PC_Interface, FString(), *TargetType, CastResultPinName);
+			CreatePin(EGPD_Output, UEdGraphSchema_K2::PC_Interface, *TargetType, *CastResultPinName);
 		}
 		else 
 		{
-			CreatePin(EGPD_Output, K2Schema->PC_Object, FString(), *TargetType, CastResultPinName);
+			CreatePin(EGPD_Output, UEdGraphSchema_K2::PC_Object, *TargetType, *CastResultPinName);
 		}
 	}
 
-	UEdGraphPin* BoolSuccessPin = CreatePin(EGPD_Output, K2Schema->PC_Boolean, FString(), nullptr, UK2Node_DynamicCastImpl::CastSuccessPinName);
+	UEdGraphPin* BoolSuccessPin = CreatePin(EGPD_Output, UEdGraphSchema_K2::PC_Boolean, UK2Node_DynamicCastImpl::CastSuccessPinName);
 	BoolSuccessPin->bHidden = !bIsPureCast;
 
 	Super::AllocateDefaultPins();
@@ -181,9 +181,7 @@ void UK2Node_DynamicCast::PostPlacedNewNode()
 
 UEdGraphPin* UK2Node_DynamicCast::GetValidCastPin() const
 {
-	const UEdGraphSchema_K2* K2Schema = GetDefault<UEdGraphSchema_K2>();
-
-	UEdGraphPin* Pin = FindPin(K2Schema->PN_CastSucceeded);
+	UEdGraphPin* Pin = FindPin(UEdGraphSchema_K2::PN_CastSucceeded);
 	check((Pin != nullptr) || bIsPureCast);
 	check((Pin == nullptr) || (Pin->Direction == EGPD_Output));
 	return Pin;
@@ -191,9 +189,7 @@ UEdGraphPin* UK2Node_DynamicCast::GetValidCastPin() const
 
 UEdGraphPin* UK2Node_DynamicCast::GetInvalidCastPin() const
 {
-	const UEdGraphSchema_K2* K2Schema = GetDefault<UEdGraphSchema_K2>();
-
-	UEdGraphPin* Pin = FindPin(K2Schema->PN_CastFailed);
+	UEdGraphPin* Pin = FindPin(UEdGraphSchema_K2::PN_CastFailed);
 	check((Pin != nullptr) || bIsPureCast);
 	check((Pin == nullptr) || (Pin->Direction == EGPD_Output));
 	return Pin;
@@ -203,13 +199,11 @@ UEdGraphPin* UK2Node_DynamicCast::GetCastResultPin() const
 {
 	if(TargetType != nullptr)
 	{
-		const UEdGraphSchema_K2* K2Schema = GetDefault<UEdGraphSchema_K2>();
-
 		for (int32 PinIdx = 0; PinIdx < Pins.Num(); PinIdx++)
 		{
 			if (Pins[PinIdx]->PinType.PinSubCategoryObject == *TargetType
 				&& Pins[PinIdx]->Direction == EGPD_Output
-				&& Pins[PinIdx]->PinName.StartsWith(K2Schema->PN_CastedValuePrefix))
+				&& Pins[PinIdx]->PinName.ToString().StartsWith(UEdGraphSchema_K2::PN_CastedValuePrefix))
 			{
 				return Pins[PinIdx];
 			}
@@ -221,10 +215,8 @@ UEdGraphPin* UK2Node_DynamicCast::GetCastResultPin() const
 
 UEdGraphPin* UK2Node_DynamicCast::GetCastSourcePin() const
 {
-	const UEdGraphSchema_K2* K2Schema = GetDefault<UEdGraphSchema_K2>();
-
-	UEdGraphPin* Pin = FindPin(K2Schema->PN_ObjectToCast);
-	check(Pin != NULL);
+	UEdGraphPin* Pin = FindPin(UEdGraphSchema_K2::PN_ObjectToCast);
+	check(Pin);
 	check(Pin->Direction == EGPD_Input);
 	return Pin;
 }
@@ -264,11 +256,9 @@ UK2Node::ERedirectType UK2Node_DynamicCast::DoPinsMatchForReconstruction(const U
 	ERedirectType RedirectType = Super::DoPinsMatchForReconstruction(NewPin, NewPinIndex, OldPin, OldPinIndex);
 	if((ERedirectType_None == RedirectType) && (NULL != NewPin) && (NULL != OldPin))
 	{
-		const UEdGraphSchema_K2* K2Schema = GetDefault<UEdGraphSchema_K2>();
-
 		const bool bProperPrefix = 
-			NewPin->PinName.StartsWith(K2Schema->PN_CastedValuePrefix, ESearchCase::CaseSensitive) && 
-			OldPin->PinName.StartsWith(K2Schema->PN_CastedValuePrefix, ESearchCase::CaseSensitive);
+			NewPin->PinName.ToString().StartsWith(UEdGraphSchema_K2::PN_CastedValuePrefix, ESearchCase::CaseSensitive) && 
+			OldPin->PinName.ToString().StartsWith(UEdGraphSchema_K2::PN_CastedValuePrefix, ESearchCase::CaseSensitive);
 
 		const bool bClassMatch = NewPin->PinType.PinSubCategoryObject.IsValid() &&
 			(NewPin->PinType.PinSubCategoryObject == OldPin->PinType.PinSubCategoryObject);
@@ -325,7 +315,7 @@ bool UK2Node_DynamicCast::IsConnectionDisallowed(const UEdGraphPin* MyPin, const
 	if (MyPin == GetCastSourcePin())
 	{
 		const FEdGraphPinType& OtherPinType = OtherPin->PinType;
-		const FText OtherPinName = OtherPin->PinFriendlyName.IsEmpty() ? FText::FromString(OtherPin->PinName) : OtherPin->PinFriendlyName;
+		const FText OtherPinName = OtherPin->PinFriendlyName.IsEmpty() ? FText::FromName(OtherPin->PinName) : OtherPin->PinFriendlyName;
 
 		if (OtherPinType.IsContainer())
 		{
@@ -366,7 +356,7 @@ void UK2Node_DynamicCast::NotifyPinConnectionListChanged(UEdGraphPin* Pin)
 		if (Pin->LinkedTo.Num() == 0)
 		{
 			InputPinType.PinCategory = UEdGraphSchema_K2::PC_Wildcard;
-			InputPinType.PinSubCategory.Empty();
+			InputPinType.PinSubCategory = NAME_None;
 			InputPinType.PinSubCategoryObject = nullptr;
 		}
 		else
@@ -429,7 +419,7 @@ void UK2Node_DynamicCast::ValidateNodeDuringCompilation(FCompilerResultsLog& Mes
 
 			if (SourceClass == nullptr)
 			{
-				const FString SourcePinName = CastInput->PinFriendlyName.IsEmpty() ? CastInput->PinName : CastInput->PinFriendlyName.ToString();
+				const FString SourcePinName = CastInput->PinFriendlyName.IsEmpty() ? CastInput->PinName.ToString() : CastInput->PinFriendlyName.ToString();
 
 				FText const ErrorFormat = LOCTEXT("BadCastInput", "'%s' does not have a clear object type (invalid input into @@).");
 				MessageLog.Error( *FString::Printf(*ErrorFormat.ToString(), *SourcePinName), this );
@@ -440,14 +430,14 @@ void UK2Node_DynamicCast::ValidateNodeDuringCompilation(FCompilerResultsLog& Mes
 
 			if (SourceClass == SourceType)
 			{
-				const FString SourcePinName = CastInput->PinFriendlyName.IsEmpty() ? CastInput->PinName : CastInput->PinFriendlyName.ToString();
+				const FString SourcePinName = CastInput->PinFriendlyName.IsEmpty() ? CastInput->PinName.ToString() : CastInput->PinFriendlyName.ToString();
 
 				FText const WarningFormat = LOCTEXT("EqualObjectCast", "'%s' is already a '%s', you don't need @@.");
 				MessageLog.Note( *FString::Printf(*WarningFormat.ToString(), *SourcePinName, *TargetType->GetDisplayNameText().ToString()), this );
 			}
 			else if (SourceClass->IsChildOf(SourceType))
 			{
-				const FString SourcePinName = CastInput->PinFriendlyName.IsEmpty() ? CastInput->PinName : CastInput->PinFriendlyName.ToString();
+				const FString SourcePinName = CastInput->PinFriendlyName.IsEmpty() ? CastInput->PinName.ToString() : CastInput->PinFriendlyName.ToString();
 
 				FText const WarningFormat = LOCTEXT("UnneededObjectCast", "'%s' is already a '%s' (which inherits from '%s'), so you don't need @@.");
 				MessageLog.Note( *FString::Printf(*WarningFormat.ToString(), *SourcePinName, *SourceClass->GetDisplayNameText().ToString(), *TargetType->GetDisplayNameText().ToString()), this );
@@ -467,11 +457,10 @@ bool UK2Node_DynamicCast::ReconnectPureExecPins(TArray<UEdGraphPin*>& OldPins)
 	if (bIsPureCast)
 	{
 		// look for an old exec pin
-		const UEdGraphSchema_K2* K2Schema = GetDefault<UEdGraphSchema_K2>();
 		UEdGraphPin* PinExec = nullptr;
 		for (UEdGraphPin* Pin : OldPins)
 		{
-			if (Pin->PinName == K2Schema->PN_Execute)
+			if (Pin->PinName == UEdGraphSchema_K2::PN_Execute)
 			{
 				PinExec = Pin;
 				break;
@@ -483,7 +472,7 @@ bool UK2Node_DynamicCast::ReconnectPureExecPins(TArray<UEdGraphPin*>& OldPins)
 			UEdGraphPin* PinThen = nullptr;
 			for (UEdGraphPin* Pin : OldPins)
 			{
-				if (Pin->PinName == K2Schema->PN_Then)
+				if (Pin->PinName == UEdGraphSchema_K2::PN_Then)
 				{
 					PinThen = Pin;
 					break;
