Copying //UE4/Dev-Framework to //UE4/Dev-Main (Source: //UE4/Dev-Framework @ 3716594)

diff --git a/Engine/Source/Editor/BlueprintGraph/Private/K2Node_EaseFunction.cpp b/Engine/Source/Editor/BlueprintGraph/Private/K2Node_EaseFunction.cpp
--- a/Engine/Source/Editor/BlueprintGraph/Private/K2Node_EaseFunction.cpp
+++ b/Engine/Source/Editor/BlueprintGraph/Private/K2Node_EaseFunction.cpp
@@ -19,62 +19,71 @@
 
 struct FEaseFunctionNodeHelper
 {
-	static const FString& GetEaseFuncPinName()
+	static const FName EaseFuncPinName;
+	static const FName AlphaPinName;
+	static const FName APinFloatName;
+	static const FName BPinFloatName;
+	static const FName ResultPinName;
+	static const FName BlendExpPinName;
+	static const FName StepsPinName;
+	static const FName ShortestPathPinName;
+
+	static const FName GetEaseFuncPinName()
 	{
-		static const FString EaseFuncPinName(TEXT("Function"));
 		return EaseFuncPinName;
 	}
 
-	static const FString& GetAlphaPinName()
+	static const FName GetAlphaPinName()
 	{
-		static const FString AlphaPinName(TEXT("Alpha"));
 		return AlphaPinName;
 	}
 
-	static const FString& GetAPinName()
+	static const FName GetAPinName()
 	{
-		static const FString APinFloatName(TEXT("A"));
 		return APinFloatName;
 	}
 
-	static const FString& GetBPinName()
+	static const FName GetBPinName()
 	{
-		static const FString BPinFloatName(TEXT("B"));
 		return BPinFloatName;
 	}
 
-	static const FString& GetResultPinName()
+	static const FName GetResultPinName()
 	{
-		static const FString ResultPinName(TEXT("Result"));
 		return ResultPinName;
 	}
 
-	static const FString& GetBlendExpPinName()
+	static const FName GetBlendExpPinName()
 	{
-		static const FString BlendExpPinName(TEXT("BlendExp"));
 		return BlendExpPinName;
 	}
 
-	static const FString& GetStepsPinName()
+	static const FName GetStepsPinName()
 	{
-		static const FString StepsPinName(TEXT("Steps"));
 		return StepsPinName;
 	}
 
-	static const FString& GetShortestPathPinName()
+	static const FName GetShortestPathPinName()
 	{
-		static const FString ShortestPathPinName(TEXT("ShortestPath"));
 		return ShortestPathPinName;
 	}
-
 };
 
+const FName FEaseFunctionNodeHelper::EaseFuncPinName(TEXT("Function"));
+const FName FEaseFunctionNodeHelper::AlphaPinName(TEXT("Alpha"));
+const FName FEaseFunctionNodeHelper::APinFloatName(TEXT("A"));
+const FName FEaseFunctionNodeHelper::BPinFloatName(TEXT("B"));
+const FName FEaseFunctionNodeHelper::ResultPinName(TEXT("Result"));
+const FName FEaseFunctionNodeHelper::BlendExpPinName(TEXT("BlendExp"));
+const FName FEaseFunctionNodeHelper::StepsPinName(TEXT("Steps"));
+const FName FEaseFunctionNodeHelper::ShortestPathPinName(TEXT("ShortestPath"));
+
 UK2Node_EaseFunction::UK2Node_EaseFunction(const FObjectInitializer& ObjectInitializer)
 	: Super(ObjectInitializer)
-	, CachedEaseFuncPin(NULL)
+	, CachedEaseFuncPin(nullptr)
 {
 	NodeTooltip = LOCTEXT("NodeTooltip", "Interpolates from value A to value B using a user specified easing function");
-	EaseFunctionName = TEXT("");
+	EaseFunctionName = NAME_None;
 }
 
 void UK2Node_EaseFunction::AllocateDefaultPins()
@@ -84,36 +93,36 @@ void UK2Node_EaseFunction::AllocateDefaultPins()
 	const UEdGraphSchema_K2* K2Schema = GetDefault<UEdGraphSchema_K2>();
 
 	// Add the first pin representing all available easing functions. If EEasingFunc changes its name this will fail a runtime check!
-	UEnum * EaseFuncEnum = FindObject<UEnum>(ANY_PACKAGE, TEXT("EEasingFunc"), true);
-	check(EaseFuncEnum != NULL);
-	CachedEaseFuncPin = CreatePin(EGPD_Input, K2Schema->PC_Byte, FString(), EaseFuncEnum, FEaseFunctionNodeHelper::GetEaseFuncPinName());
+	UEnum* EaseFuncEnum = FindObject<UEnum>(ANY_PACKAGE, TEXT("EEasingFunc"), true);
+	check(EaseFuncEnum);
+	CachedEaseFuncPin = CreatePin(EGPD_Input, UEdGraphSchema_K2::PC_Byte, EaseFuncEnum, FEaseFunctionNodeHelper::GetEaseFuncPinName());
 	SetPinToolTip(*CachedEaseFuncPin, LOCTEXT("EaseFunsPinDescription", "Specifies the desired ease function to be applied. If connected no customization is possible."));
 
 	// Make sure that the default value is set correctly if none has been set
 	K2Schema->SetPinAutogeneratedDefaultValueBasedOnType(CachedEaseFuncPin);
 
-	UEdGraphPin* AlphaPin = CreatePin(EGPD_Input, K2Schema->PC_Float, FString(), nullptr, FEaseFunctionNodeHelper::GetAlphaPinName());
+	UEdGraphPin* AlphaPin = CreatePin(EGPD_Input, UEdGraphSchema_K2::PC_Float, FEaseFunctionNodeHelper::GetAlphaPinName());
 	SetPinToolTip(*AlphaPin, LOCTEXT("AlphaPinTooltip", "Alpha value used to specify the easing in time."));
 
 	// Add wildcard pins for A, B and the return Pin
-	UEdGraphPin* APin = CreatePin(EGPD_Input, K2Schema->PC_Wildcard, FString(), nullptr, FEaseFunctionNodeHelper::GetAPinName());
+	UEdGraphPin* APin = CreatePin(EGPD_Input, UEdGraphSchema_K2::PC_Wildcard, FEaseFunctionNodeHelper::GetAPinName());
 	SetPinToolTip(*APin, LOCTEXT("APinDescription", "Easing start value"));
 
-	UEdGraphPin* BPin = CreatePin(EGPD_Input, K2Schema->PC_Wildcard, FString(), nullptr, FEaseFunctionNodeHelper::GetBPinName());
+	UEdGraphPin* BPin = CreatePin(EGPD_Input, UEdGraphSchema_K2::PC_Wildcard, FEaseFunctionNodeHelper::GetBPinName());
 	SetPinToolTip(*BPin, LOCTEXT("BPinDescription", "Easing end value"));
 
-	UEdGraphPin* ResultPin = CreatePin(EGPD_Output, K2Schema->PC_Wildcard, FString(), nullptr, FEaseFunctionNodeHelper::GetResultPinName());
+	UEdGraphPin* ResultPin = CreatePin(EGPD_Output, UEdGraphSchema_K2::PC_Wildcard, FEaseFunctionNodeHelper::GetResultPinName());
 	SetPinToolTip(*ResultPin, LOCTEXT("ResultPinDescription", "Easing result value"));
 
-	UEdGraphPin* ShortestPathPin = CreatePin(EGPD_Input, K2Schema->PC_Boolean, FString(), nullptr, FEaseFunctionNodeHelper::GetShortestPathPinName());
+	UEdGraphPin* ShortestPathPin = CreatePin(EGPD_Input, UEdGraphSchema_K2::PC_Boolean, FEaseFunctionNodeHelper::GetShortestPathPinName());
 	SetPinToolTip(*ShortestPathPin, LOCTEXT("ShortestPathPinTooltip", "When interpolating the shortest path should be taken."));
 	K2Schema->SetPinAutogeneratedDefaultValue(ShortestPathPin, TEXT("true"));
 
-	UEdGraphPin* BlendExpPin = CreatePin(EGPD_Input, K2Schema->PC_Float, FString(), nullptr, FEaseFunctionNodeHelper::GetBlendExpPinName());
+	UEdGraphPin* BlendExpPin = CreatePin(EGPD_Input, UEdGraphSchema_K2::PC_Float, FEaseFunctionNodeHelper::GetBlendExpPinName());
 	SetPinToolTip(*BlendExpPin, LOCTEXT("BlendExpPinDescription", "Blend Exponent for basic ease functions"));
 	K2Schema->SetPinAutogeneratedDefaultValue(BlendExpPin, TEXT("2.0"));
 
-	UEdGraphPin* StepsPin = CreatePin(EGPD_Input, K2Schema->PC_Int, FString(), nullptr, FEaseFunctionNodeHelper::GetStepsPinName());
+	UEdGraphPin* StepsPin = CreatePin(EGPD_Input, UEdGraphSchema_K2::PC_Int, FEaseFunctionNodeHelper::GetStepsPinName());
 	SetPinToolTip(*StepsPin, LOCTEXT("StepsPinDescription", "Number of steps required to go from A to B"));
 	K2Schema->SetPinAutogeneratedDefaultValue(StepsPin, TEXT("2"));
 
@@ -208,17 +217,15 @@ void UK2Node_EaseFunction::PostReconstructNode()
 
 bool UK2Node_EaseFunction::IsConnectionDisallowed(const UEdGraphPin* MyPin, const UEdGraphPin* OtherPin, FString& OutReason) const
 {
-	const UEdGraphSchema_K2* K2Schema = GetDefault<UEdGraphSchema_K2>();
-
 	// Check the pin name and see if it matches on of our three base pins
 	if (MyPin->PinName == FEaseFunctionNodeHelper::GetAPinName() || 
 		MyPin->PinName == FEaseFunctionNodeHelper::GetBPinName() || 
 		MyPin->PinName == FEaseFunctionNodeHelper::GetResultPinName())
 	{
 		const bool bConnectionOk = (
-			OtherPin->PinType.PinCategory == K2Schema->PC_Float ||
+			OtherPin->PinType.PinCategory == UEdGraphSchema_K2::PC_Float ||
 			(
-				OtherPin->PinType.PinCategory == K2Schema->PC_Struct &&
+				OtherPin->PinType.PinCategory == UEdGraphSchema_K2::PC_Struct &&
 				OtherPin->PinType.PinSubCategoryObject.IsValid() &&
 				(
 					OtherPin->PinType.PinSubCategoryObject.Get()->GetName() == TEXT("Vector") ||
@@ -275,7 +282,6 @@ void UK2Node_EaseFunction::ChangePinType(UEdGraphPin* Pin)
 
 void UK2Node_EaseFunction::PinTypeChanged(UEdGraphPin* Pin)
 {
-	const UEdGraphSchema_K2* Schema = GetDefault<UEdGraphSchema_K2>();
 	bool bChanged = false;
 
 	if (Pin->PinName == FEaseFunctionNodeHelper::GetAPinName() ||
@@ -299,14 +305,14 @@ void UK2Node_EaseFunction::PinTypeChanged(UEdGraphPin* Pin)
 			if (bChanged)
 			{
 				// Just in case we switch to an invalid function clean it first
-				EaseFunctionName.Reset();
+				EaseFunctionName = NAME_None;
 
 				// Generate the right function name
-				if (InstigatorPin->PinType.PinCategory == Schema->PC_Float)
+				if (InstigatorPin->PinType.PinCategory == UEdGraphSchema_K2::PC_Float)
 				{
 					EaseFunctionName = TEXT("Ease");
 				}
-				else if (InstigatorPin->PinType.PinCategory == Schema->PC_Struct)
+				else if (InstigatorPin->PinType.PinCategory == UEdGraphSchema_K2::PC_Struct)
 				{
 					if (InstigatorPin->PinType.PinSubCategoryObject.Get()->GetName() == TEXT("Vector"))
 					{
@@ -330,9 +336,9 @@ void UK2Node_EaseFunction::PinTypeChanged(UEdGraphPin* Pin)
 				ResultPin->LinkedTo.Num() == 0)
 			{
 				// Restore wild card pin
-				APin->PinType.PinCategory = Schema->PC_Wildcard;
-				APin->PinType.PinSubCategory = TEXT("");
-				APin->PinType.PinSubCategoryObject = NULL;
+				APin->PinType.PinCategory = UEdGraphSchema_K2::PC_Wildcard;
+				APin->PinType.PinSubCategory = NAME_None;
+				APin->PinType.PinSubCategoryObject = nullptr;
 
 				// Propagate change
 				UpdatePin(BPin, APin);
@@ -389,7 +395,7 @@ void UK2Node_EaseFunction::ExpandNode(class FKismetCompilerContext& CompilerCont
 		the other nodes into the Blueprint.
 	*/
 
-	UFunction* Function = UKismetMathLibrary::StaticClass()->FindFunctionByName(*EaseFunctionName);
+	UFunction* Function = UKismetMathLibrary::StaticClass()->FindFunctionByName(EaseFunctionName);
 	if (Function == NULL)
 	{
 		CompilerContext.MessageLog.Error(*LOCTEXT("InvalidFunctionName", "BaseAsyncTask: Type not supported or not initialized. @@").ToString(), this);
@@ -491,11 +497,9 @@ void UK2Node_EaseFunction::RefreshPinVisibility()
 
 void UK2Node_EaseFunction::GenerateExtraPins()
 {
-	const UEdGraphSchema_K2* K2Schema = Cast<const UEdGraphSchema_K2>(GetSchema());
-
 	// Add pins based on the pin type
 	const UEdGraphPin* APin = FindPinChecked(FEaseFunctionNodeHelper::GetAPinName());
-	const bool bIsRotator = APin->PinType.PinCategory == K2Schema->PC_Struct && APin->PinType.PinSubCategoryObject.Get()->GetName() == TEXT("Rotator");
+	const bool bIsRotator = APin->PinType.PinCategory == UEdGraphSchema_K2::PC_Struct && APin->PinType.PinSubCategoryObject.Get()->GetName() == TEXT("Rotator");
 	UEdGraphPin* ShortestPathPin = FindPinChecked(FEaseFunctionNodeHelper::GetShortestPathPinName());
 	if (bIsRotator)
 	{
diff --git a/Engine/Source/Editor/BlueprintGraph/Private/K2Node_EditablePinBase.cpp b/Engine/Source/Editor/BlueprintGraph/Private/K2Node_EditablePinBase.cpp
--- a/Engine/Source/Editor/BlueprintGraph/Private/K2Node_EditablePinBase.cpp
+++ b/Engine/Source/Editor/BlueprintGraph/Private/K2Node_EditablePinBase.cpp
@@ -2,13 +2,25 @@
 
 #include "K2Node_EditablePinBase.h"
 #include "UObject/UnrealType.h"
+#include "UObject/FrameworkObjectVersion.h"
 #include "Misc/FeedbackContext.h"
 #include "EdGraphSchema_K2.h"
 #include "Kismet2/KismetDebugUtilities.h"
 
 FArchive& operator<<(FArchive& Ar, FUserPinInfo& Info)
 {
+	Ar.UsingCustomVersion(FFrameworkObjectVersion::GUID);
+
+	if (Ar.CustomVer(FFrameworkObjectVersion::GUID) >= FFrameworkObjectVersion::PinsStoreFName)
+	{
 	Ar << Info.PinName;
+	}
+	else
+	{
+		FString PinNameStr;
+		Ar << PinNameStr;
+		Info.PinName = *PinNameStr;
+	}
 
 	if (Ar.UE4Ver() >= VER_UE4_SERIALIZE_PINTYPE_CONST)
 	{
@@ -17,20 +29,25 @@ FArchive& operator<<(FArchive& Ar, FUserPinInfo& Info)
 	}
 	else
 	{
+		check(Ar.IsLoading());
+
 		bool bIsArray = (Info.PinType.ContainerType == EPinContainerType::Array);
 		Ar << bIsArray;
 
 		bool bIsReference = Info.PinType.bIsReference;
 		Ar << bIsReference;
 
-		if (Ar.IsLoading())
-		{
 			Info.PinType.ContainerType = (bIsArray ? EPinContainerType::Array : EPinContainerType::None);
 			Info.PinType.bIsReference = bIsReference;
-		}
 
-		Ar << Info.PinType.PinCategory;
-		Ar << Info.PinType.PinSubCategory;
+		FString PinCategoryStr;
+		FString PinSubCategoryStr;
+		
+		Ar << PinCategoryStr;
+		Ar << PinSubCategoryStr;
+
+		Info.PinType.PinCategory = *PinCategoryStr;
+		Info.PinType.PinSubCategory = *PinSubCategoryStr;
 
 		Ar << Info.PinType.PinSubCategoryObject;
 	}
@@ -61,10 +78,10 @@ void UK2Node_EditablePinBase::AllocateDefaultPins()
 	}
 }
 
-UEdGraphPin* UK2Node_EditablePinBase::CreateUserDefinedPin(const FString& InPinName, const FEdGraphPinType& InPinType, EEdGraphPinDirection InDesiredDirection, bool bUseUniqueName)
+UEdGraphPin* UK2Node_EditablePinBase::CreateUserDefinedPin(const FName InPinName, const FEdGraphPinType& InPinType, EEdGraphPinDirection InDesiredDirection, bool bUseUniqueName)
 {
 	// Sanitize the name, if needed
-	const FString NewPinName = bUseUniqueName ? CreateUniquePinName(InPinName) : InPinName;
+	const FName NewPinName = bUseUniqueName ? CreateUniquePinName(InPinName) : InPinName;
 
 	// First, add this pin to the user-defined pins
 	TSharedPtr<FUserPinInfo> NewPinInfo = MakeShareable( new FUserPinInfo() );
@@ -84,7 +101,7 @@ void UK2Node_EditablePinBase::RemoveUserDefinedPin(TSharedPtr<FUserPinInfo> PinT
 	RemoveUserDefinedPinByName(PinToRemove->PinName);
 }
 
-void UK2Node_EditablePinBase::RemoveUserDefinedPinByName(const FString& PinName)
+void UK2Node_EditablePinBase::RemoveUserDefinedPinByName(const FName PinName)
 {
 	for (UEdGraphPin* Pin : Pins)
 	{
@@ -115,53 +132,16 @@ void UK2Node_EditablePinBase::ExportCustomProperties(FOutputDevice& Out, uint32
 {
 	Super::ExportCustomProperties(Out, Indent);
 
+	const FUserPinInfo DefaultPinInfo;
+
 	for (int32 PinIndex = 0; PinIndex < UserDefinedPins.Num(); ++PinIndex)
 	{
 		const FUserPinInfo& PinInfo = *UserDefinedPins[PinIndex].Get();
-
-		Out.Logf( TEXT("%sCustomProperties UserDefinedPin "), FCString::Spc(Indent));
-		Out.Logf( TEXT("Name=\"%s\" "), *PinInfo.PinName);
-		Out.Logf( TEXT("IsReference=%s "), (PinInfo.PinType.bIsReference ? TEXT("1") : TEXT("0")));
-
-		if (UEnum* PinContainerType = FindObject<UEnum>(ANY_PACKAGE, TEXT("EPinContainerType")))
-		{ 
-			const FString ValueName = PinContainerType->GetNameStringByValue((uint8)PinInfo.PinType.ContainerType);
-			if (!ValueName.IsEmpty())
-			{
-				Out.Logf(TEXT("PinContainerType=\"%s\" "), *ValueName);
-			}
-		}
-
-		if (UEnum* PinDirEnum = FindObject<UEnum>(ANY_PACKAGE, TEXT("EEdGraphPinDirection")))
-		{ 
-			const FString ValueName = PinDirEnum->GetNameStringByValue(PinInfo.DesiredPinDirection);
-			if (!ValueName.IsEmpty())
-			{
-				Out.Logf(TEXT("PinDir=\"%s\" "), *ValueName);
-			}
-		}
 		
-		if (PinInfo.PinType.PinCategory.Len() > 0)
-		{
-			Out.Logf( TEXT("Category=%s "), *PinInfo.PinType.PinCategory);
-		}
+		FString PinInfoStr;
+		FUserPinInfo::StaticStruct()->ExportText(PinInfoStr, &PinInfo, &DefaultPinInfo, this, 0, nullptr, false);
 
-		if (PinInfo.PinType.PinSubCategory.Len() > 0)
-		{
-			Out.Logf( TEXT("SubCategory=%s "), *PinInfo.PinType.PinSubCategory);
-		}
-
-		if (PinInfo.PinType.PinSubCategoryObject.IsValid())
-		{
-			Out.Logf( TEXT("SubCategoryObject=%s "), *PinInfo.PinType.PinSubCategoryObject.Get()->GetPathName());
-		}
-
-		if (PinInfo.PinDefaultValue.Len() > 0)
-		{
-			Out.Logf( TEXT("DefaultValue=%s "), *PinInfo.PinDefaultValue);
-		}
-
-		Out.Logf( TEXT("\r\n"));
+		Out.Logf( TEXT("%sCustomProperties UserDefinedPin %s\r\n"), FCString::Spc(Indent), *PinInfoStr);
 	}
 }
 
@@ -169,68 +149,12 @@ void UK2Node_EditablePinBase::ImportCustomProperties(const TCHAR* SourceText, FF
 {
 	if (FParse::Command(&SourceText, TEXT("UserDefinedPin")))
 	{
-		TSharedPtr<FUserPinInfo> PinInfo = MakeShareable( new FUserPinInfo() );
-
-		if (!FParse::Value(SourceText, TEXT("Name="), PinInfo->PinName))
-		{
-			Warn->Logf( *NSLOCTEXT( "Core", "SyntaxError", "Syntax Error" ).ToString() );
-			return;
-		}
-
-		int32 BoolAsInt = 0;
-		if (FParse::Value(SourceText, TEXT("IsArray="), BoolAsInt))
-		{
-			PinInfo->PinType.ContainerType = (BoolAsInt != 0 ? EPinContainerType::Array : EPinContainerType::None);
-		}
-
-		if (UEnum* PinContainerType = FindObject<UEnum>(ANY_PACKAGE, TEXT("EPinContainerType")))
-		{
-			FString DesiredContainerType;
-			if (FParse::Value(SourceText, TEXT("PinContainerType="), DesiredContainerType))
-			{
-				const int32 DesiredContainerTypeVal = PinContainerType->GetValueByName(*DesiredContainerType);
-				if (DesiredContainerTypeVal != INDEX_NONE)
-				{
-					PinInfo->PinType.ContainerType = (EPinContainerType)DesiredContainerTypeVal;
-				}
-			}
-		}
-
-		if (FParse::Value(SourceText, TEXT("IsReference="), BoolAsInt))
-		{
-			PinInfo->PinType.bIsReference = (BoolAsInt != 0);
-		}
-
-		if (UEnum* PinDirEnum = FindObject<UEnum>(ANY_PACKAGE, TEXT("EEdGraphPinDirection")))
-		{
-			FString DesiredDirection;
-			if (FParse::Value(SourceText, TEXT("PinDir="), DesiredDirection))
-			{
-				const int32 DesiredDirectionVal = PinDirEnum->GetValueByName(*DesiredDirection);
-				if (DesiredDirectionVal != INDEX_NONE)
-				{
-					PinInfo->DesiredPinDirection = (EEdGraphPinDirection)DesiredDirectionVal;
-				}
-			}
-		}
-
-		FParse::Value(SourceText, TEXT("Category="), PinInfo->PinType.PinCategory);
-		FParse::Value(SourceText, TEXT("SubCategory="), PinInfo->PinType.PinSubCategory);
-
-		FString ObjectPathName;
-		if (FParse::Value(SourceText, TEXT("SubCategoryObject="), ObjectPathName))
-		{
-			PinInfo->PinType.PinSubCategoryObject = FindObject<UObject>(ANY_PACKAGE, *ObjectPathName);
-			if (!PinInfo->PinType.PinSubCategoryObject.IsValid())
-			{
-				Warn->Logf( *NSLOCTEXT( "Core", "UnableToFindObject", "Unable to find object" ).ToString() );
-				return;
-			}
-		}
+		TSharedPtr<FUserPinInfo> SharedPinInfo = MakeShareable(new FUserPinInfo());
+		FUserPinInfo* PinInfo = SharedPinInfo.Get();
 
-		FParse::Value(SourceText, TEXT("DefaultValue="), PinInfo->PinDefaultValue);
+		FUserPinInfo::StaticStruct()->ImportText(SourceText, PinInfo, this, 0, Warn, TEXT("PinInfo"), false);
 
-		UserDefinedPins.Add(PinInfo);
+		UserDefinedPins.Add(SharedPinInfo);
 	}
 	else
 	{
@@ -380,7 +304,7 @@ bool UK2Node_EditablePinBase::CreateUserDefinedPinsForFunctionEntryExit(const UF
 			FEdGraphPinType PinType;
 			K2Schema->ConvertPropertyToPinType(Param, /*out*/ PinType);
 
-			const bool bPinGood = CreateUserDefinedPin(Param->GetName(), PinType, Direction) != NULL;
+			const bool bPinGood = (CreateUserDefinedPin(Param->GetFName(), PinType, Direction) != nullptr);
 
 			bAllPinsGood = bAllPinsGood && bPinGood;
 		}
diff --git a/Engine/Source/Editor/BlueprintGraph/Private/K2Node_EnumEquality.cpp b/Engine/Source/Editor/BlueprintGraph/Private/K2Node_EnumEquality.cpp
--- a/Engine/Source/Editor/BlueprintGraph/Private/K2Node_EnumEquality.cpp
+++ b/Engine/Source/Editor/BlueprintGraph/Private/K2Node_EnumEquality.cpp
@@ -20,14 +20,12 @@ UK2Node_EnumEquality::UK2Node_EnumEquality(const FObjectInitializer& ObjectIniti
 
 void UK2Node_EnumEquality::AllocateDefaultPins()
 {
-	const UEdGraphSchema_K2* Schema = GetDefault<UEdGraphSchema_K2>();
-
 	// Create the return value pin
-	CreatePin(EGPD_Output, Schema->PC_Boolean, FString(), nullptr, Schema->PN_ReturnValue);
+	CreatePin(EGPD_Output, UEdGraphSchema_K2::PC_Boolean, UEdGraphSchema_K2::PN_ReturnValue);
 
 	// Create the input pins
-	CreatePin(EGPD_Input, Schema->PC_Wildcard, FString(), nullptr, TEXT("A"));
-	CreatePin(EGPD_Input, Schema->PC_Wildcard, FString(), nullptr, TEXT("B"));
+	CreatePin(EGPD_Input, UEdGraphSchema_K2::PC_Wildcard, TEXT("A"));
+	CreatePin(EGPD_Input, UEdGraphSchema_K2::PC_Wildcard, TEXT("B"));
 
 	Super::AllocateDefaultPins();
 }
@@ -39,7 +37,7 @@ FText UK2Node_EnumEquality::GetTooltipText() const
 
 FText UK2Node_EnumEquality::GetNodeTitle(ENodeTitleType::Type TitleType) const
 {
-	return NSLOCTEXT("K2Node_EnumEquality", "EqualEnum", "Equal (Enum)");
+	return LOCTEXT("EqualEnum", "Equal (Enum)");
 }
 
 void UK2Node_EnumEquality::PostReconstructNode()
@@ -49,21 +47,6 @@ void UK2Node_EnumEquality::PostReconstructNode()
 	PinConnectionListChanged(GetInput2Pin());
 
 	Super::PostReconstructNode();
-// 	const UEdGraphSchema_K2* Schema = GetDefault<UEdGraphSchema_K2>();
-// 
-// 	UEdGraphPin* Input1Pin = GetInput1Pin();
-// 	UEdGraphPin* Input2Pin = GetInput2Pin();
-// 
-// 	if ((Input1Pin->LinkedTo.Num() == 0) && (Input2Pin->LinkedTo.Num() == 0))
-// 	{
-// 		// Restore the wildcard status
-// 		Input1Pin->PinType.PinCategory = Schema->PC_Wildcard;
-// 		Input1Pin->PinType.PinSubCategory.Reset();
-// 		Input1Pin->PinType.PinSubCategoryObject = nullptr;
-// 		Input2Pin->PinType.PinCategory = Schema->PC_Wildcard;
-// 		Input2Pin->PinType.PinSubCategory.Reset();
-// 		Input2Pin->PinType.PinSubCategoryObject = nullptr;
-// 	}
 }
 
 /** Determine if any pins are connected, if so make all the other pins the same type, if not, make sure pins are switched back to wildcards */
@@ -81,11 +64,11 @@ void UK2Node_EnumEquality::NotifyPinConnectionListChanged(UEdGraphPin* Pin)
 		if ((Input1Pin->LinkedTo.Num() == 0) && (Input2Pin->LinkedTo.Num() == 0))
 		{
 			// Restore the wildcard status
-			Input1Pin->PinType.PinCategory = Schema->PC_Wildcard;
-			Input1Pin->PinType.PinSubCategory.Reset();
+			Input1Pin->PinType.PinCategory = UEdGraphSchema_K2::PC_Wildcard;
+			Input1Pin->PinType.PinSubCategory= NAME_None;
 			Input1Pin->PinType.PinSubCategoryObject = nullptr;
-			Input2Pin->PinType.PinCategory = Schema->PC_Wildcard;
-			Input2Pin->PinType.PinSubCategory.Reset();
+			Input2Pin->PinType.PinCategory = UEdGraphSchema_K2::PC_Wildcard;
+			Input2Pin->PinType.PinSubCategory = NAME_None;
 			Input2Pin->PinType.PinSubCategoryObject = nullptr;
 			Schema->SetPinAutogeneratedDefaultValueBasedOnType(Input1Pin);
 			Schema->SetPinAutogeneratedDefaultValueBasedOnType(Input2Pin);
@@ -95,7 +78,7 @@ void UK2Node_EnumEquality::NotifyPinConnectionListChanged(UEdGraphPin* Pin)
 		else if (Pin->LinkedTo.Num() > 0)
 		{
 			// Make sure the pin is a valid enum
-			if (Pin->LinkedTo[0]->PinType.PinCategory == Schema->PC_Byte &&
+			if (Pin->LinkedTo[0]->PinType.PinCategory == UEdGraphSchema_K2::PC_Byte &&
 				Pin->LinkedTo[0]->PinType.PinSubCategoryObject.IsValid() &&
 				Pin->LinkedTo[0]->PinType.PinSubCategoryObject.Get()->IsA(UEnum::StaticClass()))
 			{
@@ -128,9 +111,7 @@ void UK2Node_EnumEquality::NotifyPinConnectionListChanged(UEdGraphPin* Pin)
 
 bool UK2Node_EnumEquality::IsConnectionDisallowed(const UEdGraphPin* MyPin, const UEdGraphPin* OtherPin, FString& OutReason) const
 {
-	UEdGraphSchema_K2 const* K2Schema = GetDefault<UEdGraphSchema_K2>();
-
-	bool const bOtherPinIsEnum = (OtherPin->PinType.PinCategory == K2Schema->PC_Byte) &&
+	bool const bOtherPinIsEnum = (OtherPin->PinType.PinCategory == UEdGraphSchema_K2::PC_Byte) &&
 		(OtherPin->PinType.PinSubCategoryObject.IsValid()) &&
 		(OtherPin->PinType.PinSubCategoryObject.Get()->IsA(UEnum::StaticClass()));
 
@@ -146,28 +127,22 @@ bool UK2Node_EnumEquality::IsConnectionDisallowed(const UEdGraphPin* MyPin, cons
 
 UEdGraphPin* UK2Node_EnumEquality::GetReturnValuePin() const
 {
-	const UEdGraphSchema_K2* K2Schema = GetDefault<UEdGraphSchema_K2>();
-
-	UEdGraphPin* Pin = FindPin(K2Schema->PN_ReturnValue);
-	check(Pin != NULL);
+	UEdGraphPin* Pin = FindPin(UEdGraphSchema_K2::PN_ReturnValue);
+	check(Pin);
 	return Pin;
 }
 
 UEdGraphPin* UK2Node_EnumEquality::GetInput1Pin() const
 {
-	const UEdGraphSchema_K2* K2Schema = GetDefault<UEdGraphSchema_K2>();
-
-	UEdGraphPin* Pin = FindPin("A");
-	check(Pin != NULL);
+	UEdGraphPin* Pin = FindPin(TEXT("A"));
+	check(Pin);
 	return Pin;
 }
 
 UEdGraphPin* UK2Node_EnumEquality::GetInput2Pin() const
 {
-	const UEdGraphSchema_K2* K2Schema = GetDefault<UEdGraphSchema_K2>();
-
-	UEdGraphPin* Pin = FindPin("B");
-	check(Pin != NULL);
+	UEdGraphPin* Pin = FindPin(TEXT("B"));
+	check(Pin);
 	return Pin;
 }
 
@@ -186,8 +161,6 @@ void UK2Node_EnumEquality::ExpandNode(class FKismetCompilerContext& CompilerCont
 {
 	Super::ExpandNode(CompilerContext, SourceGraph);
 
-	const UEdGraphSchema_K2* Schema = CompilerContext.GetSchema();
-
 	// Get the enum equality node and the KismetMathLibrary function info for use when we build those nodes
 	FName ConditionalFunctionName = "";
 	UClass* ConditionalFunctionClass = NULL;
@@ -202,7 +175,7 @@ void UK2Node_EnumEquality::ExpandNode(class FKismetCompilerContext& CompilerCont
 	// Rewire the enum pins to the new conditional node
 	UEdGraphPin* LeftSideConditionalPin = ConditionalNode->FindPinChecked(TEXT("A"));
 	UEdGraphPin* RightSideConditionalPin = ConditionalNode->FindPinChecked(TEXT("B"));
-	UEdGraphPin* ReturnConditionalPin = ConditionalNode->FindPinChecked(Schema->PN_ReturnValue);
+	UEdGraphPin* ReturnConditionalPin = ConditionalNode->FindPinChecked(UEdGraphSchema_K2::PN_ReturnValue);
 	UEdGraphPin* Input1Pin = GetInput1Pin();
 	UEdGraphPin* Input2Pin = GetInput2Pin();
 	LeftSideConditionalPin->PinType = Input1Pin->PinType;
diff --git a/Engine/Source/Editor/BlueprintGraph/Private/K2Node_EnumLiteral.cpp b/Engine/Source/Editor/BlueprintGraph/Private/K2Node_EnumLiteral.cpp
--- a/Engine/Source/Editor/BlueprintGraph/Private/K2Node_EnumLiteral.cpp
+++ b/Engine/Source/Editor/BlueprintGraph/Private/K2Node_EnumLiteral.cpp
@@ -10,9 +10,9 @@
 #include "EditorCategoryUtils.h"
 #include "BlueprintActionDatabaseRegistrar.h"
 
-const FString& UK2Node_EnumLiteral::GetEnumInputPinName()
+const FName UK2Node_EnumLiteral::GetEnumInputPinName()
 {
-	static const FString Name(TEXT("Enum"));
+	static const FName Name(TEXT("Enum"));
 	return Name;
 }
 
@@ -34,10 +34,10 @@ void UK2Node_EnumLiteral::AllocateDefaultPins()
 {
 	const UEdGraphSchema_K2* Schema = GetDefault<UEdGraphSchema_K2>();
 
-	auto InputPin = CreatePin(EGPD_Input, Schema->PC_Byte, FString(), Enum, GetEnumInputPinName());
+	UEdGraphPin* InputPin = CreatePin(EGPD_Input, UEdGraphSchema_K2::PC_Byte, Enum, GetEnumInputPinName());
 	Schema->SetPinAutogeneratedDefaultValueBasedOnType(InputPin);
 
-	CreatePin(EGPD_Output, Schema->PC_Byte, FString(), Enum, Schema->PN_ReturnValue);
+	CreatePin(EGPD_Output, UEdGraphSchema_K2::PC_Byte, Enum, UEdGraphSchema_K2::PN_ReturnValue);
 
 	Super::AllocateDefaultPins();
 }
@@ -95,7 +95,7 @@ public:
 
 		FBPTerminal** ValueSource = Context.NetMap.Find(Net);
 		check(ValueSource && *ValueSource);
-		UEdGraphPin* OutPin = Node->FindPinChecked(Context.Schema->PN_ReturnValue);
+		UEdGraphPin* OutPin = Node->FindPinChecked(UEdGraphSchema_K2::PN_ReturnValue);
 		if (ensure(Context.NetMap.Find(OutPin) == NULL))
 		{
 			FBPTerminal* TerminalPtr = *ValueSource; //necessary because of CheckAddress in Map::Add
diff --git a/Engine/Source/Editor/BlueprintGraph/Private/K2Node_Event.cpp b/Engine/Source/Editor/BlueprintGraph/Private/K2Node_Event.cpp
--- a/Engine/Source/Editor/BlueprintGraph/Private/K2Node_Event.cpp
+++ b/Engine/Source/Editor/BlueprintGraph/Private/K2Node_Event.cpp
@@ -16,7 +16,7 @@
 #include "KismetCompiler.h"
 #include "EventEntryHandler.h"
 
-const FString UK2Node_Event::DelegateOutputName(TEXT("OutputDelegate"));
+const FName UK2Node_Event::DelegateOutputName(TEXT("OutputDelegate"));
 
 bool UK2Node_Event::IsCosmeticTickEvent() const
 {
@@ -28,7 +28,7 @@ bool UK2Node_Event::IsCosmeticTickEvent() const
 		if (Blueprint)
 		{
 			UClass* BPClass = Blueprint->GeneratedClass;
-			const AActor* DefaultActor = BPClass ? Cast<const AActor>(BPClass->GetDefaultObject()) : NULL;
+			const AActor* DefaultActor = BPClass ? Cast<const AActor>(BPClass->GetDefaultObject()) : nullptr;
 			if (DefaultActor && !DefaultActor->AllowReceiveTickEventOnDedicatedServer())
 			{
 				return true;
@@ -75,8 +75,7 @@ void UK2Node_Event::PostLoad()
 	// Fix up legacy nodes that may not yet have a delegate pin
 	if (!FindPin(DelegateOutputName))
 	{
-		const UEdGraphSchema_K2* K2Schema = GetDefault<UEdGraphSchema_K2>();
-		CreatePin(EGPD_Output, K2Schema->PC_Delegate, FString(), nullptr, DelegateOutputName);
+		CreatePin(EGPD_Output, UEdGraphSchema_K2::PC_Delegate, DelegateOutputName);
 	}
 }
 
@@ -306,10 +305,8 @@ UFunction* UK2Node_Event::FindEventSignatureFunction()
 
 void UK2Node_Event::AllocateDefaultPins()
 {
-	const UEdGraphSchema_K2* K2Schema = GetDefault<UEdGraphSchema_K2>();
-
-	CreatePin(EGPD_Output, K2Schema->PC_Delegate, FString(), nullptr, DelegateOutputName);
-	CreatePin(EGPD_Output, K2Schema->PC_Exec, FString(), nullptr, K2Schema->PN_Then);
+	CreatePin(EGPD_Output, UEdGraphSchema_K2::PC_Delegate, DelegateOutputName);
+	CreatePin(EGPD_Output, UEdGraphSchema_K2::PC_Exec, UEdGraphSchema_K2::PN_Then);
 
 	const UFunction* Function = FindEventSignatureFunction();
 	if (Function)
@@ -385,8 +382,8 @@ bool UK2Node_Event::IsFunctionEntryCompatible(const UK2Node_FunctionEntry* Entry
 	for(int32 i = 0; i < EventPins.Num(); i++)
 	{
 		const UEdGraphPin* CurPin = EventPins[i];
-		if( CurPin->PinType.PinCategory == K2Schema->PC_Exec 
-			|| CurPin->PinType.PinSubCategory == K2Schema->PSC_Self
+		if( CurPin->PinType.PinCategory == UEdGraphSchema_K2::PC_Exec 
+			|| CurPin->PinType.PinSubCategory == UEdGraphSchema_K2::PSC_Self
 			|| CurPin->PinName == DelegateOutputName
 			|| CurPin->Direction == EGPD_Input
 			|| CurPin->ParentPin != nullptr )
@@ -399,8 +396,8 @@ bool UK2Node_Event::IsFunctionEntryCompatible(const UK2Node_FunctionEntry* Entry
 	for(int32 i = 0; i < EntryPins.Num(); i++)
 	{
 		const UEdGraphPin* CurPin = EntryPins[i];
-		if( CurPin->PinType.PinCategory == K2Schema->PC_Exec 
-			|| CurPin->PinType.PinSubCategory == K2Schema->PSC_Self
+		if( CurPin->PinType.PinCategory == UEdGraphSchema_K2::PC_Exec 
+			|| CurPin->PinType.PinSubCategory == UEdGraphSchema_K2::PSC_Self
 			|| CurPin->PinName == DelegateOutputName
 			|| CurPin->Direction == EGPD_Input
 			|| CurPin->ParentPin != nullptr )
@@ -685,7 +682,7 @@ void UK2Node_Event::ExpandNode(class FKismetCompilerContext& CompilerContext, UE
 		UK2Node_CreateDelegate* CreateDelegateNode = CompilerContext.SpawnIntermediateNode<UK2Node_CreateDelegate>(this, SourceGraph);
 		CreateDelegateNode->AllocateDefaultPins();
 		CompilerContext.MovePinLinksToIntermediate(*OrgDelegatePin, *CreateDelegateNode->GetDelegateOutPin());
-		Schema->TryCreateConnection(SelfNode->FindPinChecked(Schema->PN_Self), CreateDelegateNode->GetObjectInPin());
+		Schema->TryCreateConnection(SelfNode->FindPinChecked(UEdGraphSchema_K2::PN_Self), CreateDelegateNode->GetObjectInPin());
 		// When called UFunction is defined in the same class, it wasn't created yet (previously the Skeletal class was checked). So no "CreateDelegateNode->HandleAnyChangeWithoutNotifying();" is called.
 		CreateDelegateNode->SetFunction(FunctionName);
 	}
diff --git a/Engine/Source/Editor/BlueprintGraph/Private/K2Node_ExecutionSequence.cpp b/Engine/Source/Editor/BlueprintGraph/Private/K2Node_ExecutionSequence.cpp
--- a/Engine/Source/Editor/BlueprintGraph/Private/K2Node_ExecutionSequence.cpp
+++ b/Engine/Source/Editor/BlueprintGraph/Private/K2Node_ExecutionSequence.cpp
@@ -30,7 +30,7 @@ public:
 		ExpectedPinType.PinCategory = UEdGraphSchema_K2::PC_Exec;
 
 		UEdGraphPin* ExecTriggeringPin = Context.FindRequiredPinByName(Node, UEdGraphSchema_K2::PN_Execute, EGPD_Input);
-		if ((ExecTriggeringPin == NULL) || !Context.ValidatePinType(ExecTriggeringPin, ExpectedPinType))
+		if ((ExecTriggeringPin == nullptr) || !Context.ValidatePinType(ExecTriggeringPin, ExpectedPinType))
 		{
 			CompilerContext.MessageLog.Error(*FString::Printf(*LOCTEXT("NoValidExecutionPinForExecSeq_Error", "@@ must have a valid execution pin @@").ToString()), Node, ExecTriggeringPin);
 			return;
@@ -43,10 +43,9 @@ public:
 
 		// Find the valid, connected output pins, and add them to the processing list
 		TArray<UEdGraphPin*> OutputPins;
-		for (int32 i = 0; i < Node->Pins.Num(); i++)
+		for (UEdGraphPin* CurrentPin : Node->Pins)
 		{
-			UEdGraphPin* CurrentPin = Node->Pins[i];
-			if ((CurrentPin->Direction == EGPD_Output) && (CurrentPin->PinName.StartsWith(CompilerContext.GetSchema()->PN_Then)) && (CurrentPin->LinkedTo.Num() > 0))
+			if ((CurrentPin->Direction == EGPD_Output) && (CurrentPin->LinkedTo.Num() > 0) && (CurrentPin->PinName.ToString().StartsWith(UEdGraphSchema_K2::PN_Then.ToString())))
 			{
 				OutputPins.Add(CurrentPin);
 			}
@@ -140,13 +139,11 @@ UK2Node_ExecutionSequence::UK2Node_ExecutionSequence(const FObjectInitializer& O
 
 void UK2Node_ExecutionSequence::AllocateDefaultPins()
 {
-	const UEdGraphSchema_K2* K2Schema = GetDefault<UEdGraphSchema_K2>();
-
-	CreatePin(EGPD_Input, K2Schema->PC_Exec, FString(), nullptr, K2Schema->PN_Execute);
+	CreatePin(EGPD_Input, UEdGraphSchema_K2::PC_Exec, UEdGraphSchema_K2::PN_Execute);
 
 	// Add two default pins
-	CreatePin(EGPD_Output, K2Schema->PC_Exec, FString(), nullptr, GetPinNameGivenIndex(0));
-	CreatePin(EGPD_Output, K2Schema->PC_Exec, FString(), nullptr, GetPinNameGivenIndex(1));
+	CreatePin(EGPD_Output, UEdGraphSchema_K2::PC_Exec, GetPinNameGivenIndex(0));
+	CreatePin(EGPD_Output, UEdGraphSchema_K2::PC_Exec, GetPinNameGivenIndex(1));
 
 	Super::AllocateDefaultPins();
 }
@@ -172,11 +169,9 @@ FText UK2Node_ExecutionSequence::GetTooltipText() const
 	return NSLOCTEXT("K2Node", "ExecutePinInOrder_Tooltip", "Executes a series of pins in order");
 }
 
-FString UK2Node_ExecutionSequence::GetUniquePinName()
+FName UK2Node_ExecutionSequence::GetUniquePinName()
 {
-	const UEdGraphSchema_K2* K2Schema = GetDefault<UEdGraphSchema_K2>();
-
-	FString NewPinName;
+	FName NewPinName;
 	int32 i = 0;
 	while (true)
 	{
@@ -193,8 +188,7 @@ FString UK2Node_ExecutionSequence::GetUniquePinName()
 void UK2Node_ExecutionSequence::AddInputPin()
 {
 	Modify();
-	const UEdGraphSchema_K2* K2Schema = GetDefault<UEdGraphSchema_K2>();
-	CreatePin(EGPD_Output, K2Schema->PC_Exec, FString(), nullptr, GetUniquePinName());
+	CreatePin(EGPD_Output, UEdGraphSchema_K2::PC_Exec, GetUniquePinName());
 }
 
 void UK2Node_ExecutionSequence::RemovePinFromExecutionNode(UEdGraphPin* TargetPin) 
@@ -206,13 +200,11 @@ void UK2Node_ExecutionSequence::RemovePinFromExecutionNode(UEdGraphPin* TargetPi
 		TargetPin->MarkPendingKill();
 
 		// Renumber the pins so the numbering is compact
-		const UEdGraphSchema_K2* K2Schema = GetDefault<UEdGraphSchema_K2>();
-
 		int32 ThenIndex = 0;
 		for (int32 i = 0; i < OwningSeq->Pins.Num(); ++i)
 		{
 			UEdGraphPin* PotentialPin = OwningSeq->Pins[i];
-			if (K2Schema->IsExecPin(*PotentialPin) && (PotentialPin->Direction == EGPD_Output))
+			if (UEdGraphSchema_K2::IsExecPin(*PotentialPin) && (PotentialPin->Direction == EGPD_Output))
 			{
 				PotentialPin->PinName = GetPinNameGivenIndex(ThenIndex);
 				++ThenIndex;
@@ -225,11 +217,10 @@ bool UK2Node_ExecutionSequence::CanRemoveExecutionPin() const
 {
 	int32 NumOutPins = 0;
 
-	const UEdGraphSchema_K2* K2Schema = GetDefault<UEdGraphSchema_K2>();
 	for (int32 i = 0; i < Pins.Num(); ++i)
 	{
 		UEdGraphPin* PotentialPin = Pins[i];
-		if (K2Schema->IsExecPin(*PotentialPin) && (PotentialPin->Direction == EGPD_Output))
+		if (UEdGraphSchema_K2::IsExecPin(*PotentialPin) && (PotentialPin->Direction == EGPD_Output))
 		{
 			NumOutPins++;
 		}
@@ -238,10 +229,9 @@ bool UK2Node_ExecutionSequence::CanRemoveExecutionPin() const
 	return (NumOutPins > 2);
 }
 
-FString UK2Node_ExecutionSequence::GetPinNameGivenIndex(int32 Index) const
+FName UK2Node_ExecutionSequence::GetPinNameGivenIndex(int32 Index) const
 {
-	const UEdGraphSchema_K2* K2Schema = GetDefault<UEdGraphSchema_K2>();
-	return K2Schema->PN_Then + FString::Printf(TEXT("_%d"), Index);
+	return *FString::Printf(TEXT("%s_%d"), *UEdGraphSchema_K2::PN_Then.ToString(), Index);
 }
 
 void UK2Node_ExecutionSequence::ReallocatePinsDuringReconstruction(TArray<UEdGraphPin*>& OldPins)
@@ -249,32 +239,30 @@ void UK2Node_ExecutionSequence::ReallocatePinsDuringReconstruction(TArray<UEdGra
 	Super::AllocateDefaultPins();
 
 	// Create the execution input pin
-	const UEdGraphSchema_K2* K2Schema = GetDefault<UEdGraphSchema_K2>();
-	CreatePin(EGPD_Input, K2Schema->PC_Exec, FString(), nullptr, K2Schema->PN_Execute);
+	CreatePin(EGPD_Input, UEdGraphSchema_K2::PC_Exec, UEdGraphSchema_K2::PN_Execute);
 
 	// Create a new pin for each old execution output pin, and coerce the names to match on both sides
 	int32 ExecOutPinCount = 0;
 	for (int32 i = 0; i < OldPins.Num(); ++i)
 	{
 		UEdGraphPin* TestPin = OldPins[i];
-		if (K2Schema->IsExecPin(*TestPin) && (TestPin->Direction == EGPD_Output))
+		if (UEdGraphSchema_K2::IsExecPin(*TestPin) && (TestPin->Direction == EGPD_Output))
 		{
-			FString NewPinName = GetPinNameGivenIndex(ExecOutPinCount);
+			const FName NewPinName(GetPinNameGivenIndex(ExecOutPinCount));
 			ExecOutPinCount++;
 
 			// Make sure the old pin and new pin names match
 			TestPin->PinName = NewPinName;
 
 			// Create the new output pin to match
-			CreatePin(EGPD_Output, K2Schema->PC_Exec, FString(), nullptr, NewPinName);
+			CreatePin(EGPD_Output, UEdGraphSchema_K2::PC_Exec, NewPinName);
 		}
 	}
 }
 
-UEdGraphPin * UK2Node_ExecutionSequence::GetThenPinGivenIndex(int32 Index)
+UEdGraphPin* UK2Node_ExecutionSequence::GetThenPinGivenIndex(const int32 Index) 
 {
-	FString PinName = GetPinNameGivenIndex(Index);
-	return FindPin(PinName);
+	return FindPin(GetPinNameGivenIndex(Index));
 }
 
 FNodeHandlingFunctor* UK2Node_ExecutionSequence::CreateNodeHandler(FKismetCompilerContext& CompilerContext) const
diff --git a/Engine/Source/Editor/BlueprintGraph/Private/K2Node_ForEachElementInEnum.cpp b/Engine/Source/Editor/BlueprintGraph/Private/K2Node_ForEachElementInEnum.cpp
--- a/Engine/Source/Editor/BlueprintGraph/Private/K2Node_ForEachElementInEnum.cpp
+++ b/Engine/Source/Editor/BlueprintGraph/Private/K2Node_ForEachElementInEnum.cpp
@@ -48,7 +48,7 @@ struct FForExpandNodeHelper
 
 		// Create int Loop Counter
 		UK2Node_TemporaryVariable* LoopCounterNode = CompilerContext.SpawnIntermediateNode<UK2Node_TemporaryVariable>(Node, SourceGraph);
-		LoopCounterNode->VariableType.PinCategory = Schema->PC_Int;
+		LoopCounterNode->VariableType.PinCategory = UEdGraphSchema_K2::PC_Int;
 		LoopCounterNode->AllocateDefaultPins();
 		LoopCounterOutPin = LoopCounterNode->GetVariablePin();
 		check(LoopCounterOutPin);
@@ -63,7 +63,7 @@ struct FForExpandNodeHelper
 
 		// Create int Array Index
 		UK2Node_TemporaryVariable* ArrayIndexNode = CompilerContext.SpawnIntermediateNode<UK2Node_TemporaryVariable>(Node, SourceGraph);
-		ArrayIndexNode->VariableType.PinCategory = Schema->PC_Int;
+		ArrayIndexNode->VariableType.PinCategory = UEdGraphSchema_K2::PC_Int;
 		ArrayIndexNode->AllocateDefaultPins();
 		ArrayIndexOutPin = ArrayIndexNode->GetVariablePin();
 		check(ArrayIndexOutPin);
@@ -136,23 +136,20 @@ UK2Node_ForEachElementInEnum::UK2Node_ForEachElementInEnum(const FObjectInitiali
 {
 }
 
-const FString UK2Node_ForEachElementInEnum::InsideLoopPinName(TEXT("LoopBody"));
-const FString UK2Node_ForEachElementInEnum::EnumOuputPinName(TEXT("EnumValue"));
+const FName UK2Node_ForEachElementInEnum::InsideLoopPinName(TEXT("LoopBody"));
+const FName UK2Node_ForEachElementInEnum::EnumOuputPinName(TEXT("EnumValue"));
 
 void UK2Node_ForEachElementInEnum::AllocateDefaultPins()
 {
-	const UEdGraphSchema_K2* K2Schema = GetDefault<UEdGraphSchema_K2>();
-	check(K2Schema);
-
-	CreatePin(EGPD_Input, K2Schema->PC_Exec, FString(), nullptr, K2Schema->PN_Execute);
+	CreatePin(EGPD_Input, UEdGraphSchema_K2::PC_Exec, UEdGraphSchema_K2::PN_Execute);
 
 	if (Enum)
 	{
-		CreatePin(EGPD_Output, K2Schema->PC_Exec, FString(), nullptr, InsideLoopPinName);
-		CreatePin(EGPD_Output, K2Schema->PC_Byte, FString(), Enum, EnumOuputPinName);
+		CreatePin(EGPD_Output, UEdGraphSchema_K2::PC_Exec, InsideLoopPinName);
+		CreatePin(EGPD_Output, UEdGraphSchema_K2::PC_Byte, Enum, EnumOuputPinName);
 	}
 
-	if (UEdGraphPin* CompletedPin = CreatePin(EGPD_Output, K2Schema->PC_Exec, FString(), nullptr, K2Schema->PN_Then))
+	if (UEdGraphPin* CompletedPin = CreatePin(EGPD_Output, UEdGraphSchema_K2::PC_Exec, UEdGraphSchema_K2::PN_Then))
 	{
 		CompletedPin->PinFriendlyName = LOCTEXT("Completed", "Completed");
 	}
@@ -214,13 +211,13 @@ void UK2Node_ForEachElementInEnum::ExpandNode(class FKismetCompilerContext& Comp
 	const UEdGraphSchema_K2* Schema = CompilerContext.GetSchema();
 
 	CompilerContext.MovePinLinksToIntermediate(*GetExecPin(), *ForLoop.StartLoopExecInPin);
-	CompilerContext.MovePinLinksToIntermediate(*FindPinChecked(Schema->PN_Then), *ForLoop.LoopCompleteOutExecPin);
+	CompilerContext.MovePinLinksToIntermediate(*FindPinChecked(UEdGraphSchema_K2::PN_Then), *ForLoop.LoopCompleteOutExecPin);
 	CompilerContext.MovePinLinksToIntermediate(*FindPinChecked(InsideLoopPinName), *ForLoop.InsideLoopExecOutPin);
 
 	UK2Node_GetNumEnumEntries* GetNumEnumEntries = CompilerContext.SpawnIntermediateNode<UK2Node_GetNumEnumEntries>(this, SourceGraph);
 	GetNumEnumEntries->Enum = Enum;
 	GetNumEnumEntries->AllocateDefaultPins();
-	bool bResult = Schema->TryCreateConnection(GetNumEnumEntries->FindPinChecked(Schema->PN_ReturnValue), ForLoop.LoopCounterLimitInPin);
+	bool bResult = Schema->TryCreateConnection(GetNumEnumEntries->FindPinChecked(UEdGraphSchema_K2::PN_ReturnValue), ForLoop.LoopCounterLimitInPin);
 
 	UK2Node_CallFunction* Conv_Func = CompilerContext.SpawnIntermediateNode<UK2Node_CallFunction>(this, SourceGraph);
 	FName Conv_Func_Name = GET_FUNCTION_NAME_CHECKED(UKismetMathLibrary, Conv_IntToByte);
@@ -232,8 +229,8 @@ void UK2Node_ForEachElementInEnum::ExpandNode(class FKismetCompilerContext& Comp
 	CastByteToEnum->Enum = Enum;
 	CastByteToEnum->bSafe = true;
 	CastByteToEnum->AllocateDefaultPins();
-	bResult &= Schema->TryCreateConnection(Conv_Func->FindPinChecked(Schema->PN_ReturnValue), CastByteToEnum->FindPinChecked(UK2Node_CastByteToEnum::ByteInputPinName));
-	CompilerContext.MovePinLinksToIntermediate(*FindPinChecked(EnumOuputPinName), *CastByteToEnum->FindPinChecked(Schema->PN_ReturnValue));
+	bResult &= Schema->TryCreateConnection(Conv_Func->FindPinChecked(UEdGraphSchema_K2::PN_ReturnValue), CastByteToEnum->FindPinChecked(UK2Node_CastByteToEnum::ByteInputPinName));
+	CompilerContext.MovePinLinksToIntermediate(*FindPinChecked(EnumOuputPinName), *CastByteToEnum->FindPinChecked(UEdGraphSchema_K2::PN_ReturnValue));
 
 	if (!bResult)
 	{
diff --git a/Engine/Source/Editor/BlueprintGraph/Private/K2Node_FormatText.cpp b/Engine/Source/Editor/BlueprintGraph/Private/K2Node_FormatText.cpp
--- a/Engine/Source/Editor/BlueprintGraph/Private/K2Node_FormatText.cpp
+++ b/Engine/Source/Editor/BlueprintGraph/Private/K2Node_FormatText.cpp
@@ -25,13 +25,16 @@
 
 struct FFormatTextNodeHelper
 {
-	static const FString& GetFormatPinName()
+	static const FName FormatPinName;
+
+	static const FName GetFormatPinName()
 	{
-		static const FString FormatPinName(TEXT("Format"));
 		return FormatPinName;
 	}
 };
 
+const FName FFormatTextNodeHelper::FormatPinName(TEXT("Format"));
+
 UK2Node_FormatText::UK2Node_FormatText(const FObjectInitializer& ObjectInitializer)
 	: Super(ObjectInitializer)
 	, CachedFormatPin(NULL)
@@ -43,13 +46,12 @@ void UK2Node_FormatText::AllocateDefaultPins()
 {
 	Super::AllocateDefaultPins();
 
-	const UEdGraphSchema_K2* K2Schema = GetDefault<UEdGraphSchema_K2>();
-	CachedFormatPin = CreatePin(EGPD_Input, K2Schema->PC_Text, FString(), nullptr, FFormatTextNodeHelper::GetFormatPinName());
-	CreatePin(EGPD_Output, K2Schema->PC_Text, FString(), nullptr, TEXT("Result"));
+	CachedFormatPin = CreatePin(EGPD_Input, UEdGraphSchema_K2::PC_Text, FFormatTextNodeHelper::GetFormatPinName());
+	CreatePin(EGPD_Output, UEdGraphSchema_K2::PC_Text, TEXT("Result"));
 
-	for (const FString& PinName : PinNames)
+	for (const FName& PinName : PinNames)
 	{
-		CreatePin(EGPD_Input, K2Schema->PC_Wildcard, FString(), nullptr, PinName);
+		CreatePin(EGPD_Input, UEdGraphSchema_K2::PC_Wildcard, PinName);
 	}
 }
 
@@ -63,7 +65,7 @@ void UK2Node_FormatText::SynchronizeArgumentPinType(UEdGraphPin* Pin)
 		bool bPinTypeChanged = false;
 		if (Pin->LinkedTo.Num() == 0)
 		{
-			static const FEdGraphPinType WildcardPinType = FEdGraphPinType(K2Schema->PC_Wildcard, FString(), nullptr, EPinContainerType::None, false, FEdGraphTerminalType());
+			static const FEdGraphPinType WildcardPinType = FEdGraphPinType(UEdGraphSchema_K2::PC_Wildcard, NAME_None, nullptr, EPinContainerType::None, false, FEdGraphTerminalType());
 
 			// Ensure wildcard
 			if (Pin->PinType != WildcardPinType)
@@ -106,16 +108,16 @@ FText UK2Node_FormatText::GetNodeTitle(ENodeTitleType::Type TitleType) const
 
 FText UK2Node_FormatText::GetPinDisplayName(const UEdGraphPin* Pin) const
 {
-	return FText::FromString(Pin->PinName);
+	return FText::FromName(Pin->PinName);
 }
 
-FString UK2Node_FormatText::GetUniquePinName()
+FName UK2Node_FormatText::GetUniquePinName()
 {
-	FString NewPinName;
+	FName NewPinName;
 	int32 i = 0;
 	while (true)
 	{
-		NewPinName = FString::FromInt(i++);
+		NewPinName = *FString::FromInt(i++);
 		if (!FindPin(NewPinName))
 		{
 			break;
@@ -171,33 +173,31 @@ void UK2Node_FormatText::PinDefaultValueChanged(UEdGraphPin* Pin)
 	const UEdGraphPin* FormatPin = GetFormatPin();
 	if(Pin == FormatPin && FormatPin->LinkedTo.Num() == 0)
 	{
-		const UEdGraphSchema_K2* K2Schema = GetDefault<UEdGraphSchema_K2>();
-		
 		TArray< FString > ArgumentParams;
 		FText::GetFormatPatternParameters(FormatPin->DefaultTextValue, ArgumentParams);
 
-		PinNames.Empty();
+		PinNames.Reset();
 
 		for (const FString& Param : ArgumentParams)
 		{
-			if(!FindArgumentPin(Param))
+			const FName ParamName(*Param);
+			if (!FindArgumentPin(ParamName))
 			{
-				CreatePin(EGPD_Input, K2Schema->PC_Wildcard, FString(), nullptr, Param);
+				CreatePin(EGPD_Input, UEdGraphSchema_K2::PC_Wildcard, ParamName);
 			}
-			PinNames.Add(Param);
+			PinNames.Add(ParamName);
 		}
 
-		for(auto It = Pins.CreateConstIterator(); It; ++It)
+		for (auto It = Pins.CreateIterator(); It; ++It)
 		{
 			UEdGraphPin* CheckPin = *It;
-			if(CheckPin != FormatPin && CheckPin->Direction == EGPD_Input)
+			if (CheckPin != FormatPin && CheckPin->Direction == EGPD_Input)
 			{
 				int Index = 0;
-				if(!ArgumentParams.Find(CheckPin->PinName, Index))
+				if (!ArgumentParams.Find(CheckPin->PinName.ToString(), Index))
 				{
 					CheckPin->MarkPendingKill();
-					Pins.Remove(CheckPin);
-					--It;
+					It.RemoveCurrent();
 				}
 			}
 		}
@@ -334,32 +334,31 @@ void UK2Node_FormatText::ExpandNode(class FKismetCompilerContext& CompilerContex
 		CompilerContext.MessageLog.NotifyIntermediateObjectCreation(MakeFormatArgumentDataStruct, this);
 
 		// Set the struct's "ArgumentName" pin literal to be the argument pin's name.
-		MakeFormatArgumentDataStruct->GetSchema()->TrySetDefaultValue(*MakeFormatArgumentDataStruct->FindPinChecked(GET_MEMBER_NAME_STRING_CHECKED(FFormatArgumentData, ArgumentName)), ArgumentPin->PinName);
+		MakeFormatArgumentDataStruct->GetSchema()->TrySetDefaultValue(*MakeFormatArgumentDataStruct->FindPinChecked(GET_MEMBER_NAME_STRING_CHECKED(FFormatArgumentData, ArgumentName)), ArgumentPin->PinName.ToString());
 
 		UEdGraphPin* ArgumentTypePin = MakeFormatArgumentDataStruct->FindPinChecked(GET_MEMBER_NAME_STRING_CHECKED(FFormatArgumentData, ArgumentValueType));
 
 		// Move the connection of the argument pin to the correct argument value pin, and also set the correct argument type based on the pin that was hooked up.
 		if (ArgumentPin->LinkedTo.Num() > 0)
 		{
-			const FString& ArgumentPinCategory = ArgumentPin->PinType.PinCategory;
+			const FName& ArgumentPinCategory = ArgumentPin->PinType.PinCategory;
 
-			const UEdGraphSchema_K2* K2Schema = Cast<const UEdGraphSchema_K2>(GetSchema());
-			if (ArgumentPinCategory == K2Schema->PC_Int)
+			if (ArgumentPinCategory == UEdGraphSchema_K2::PC_Int)
 			{
 				MakeFormatArgumentDataStruct->GetSchema()->TrySetDefaultValue(*ArgumentTypePin, TEXT("Int"));
 				CompilerContext.MovePinLinksToIntermediate(*ArgumentPin, *MakeFormatArgumentDataStruct->FindPinChecked(GET_MEMBER_NAME_STRING_CHECKED(FFormatArgumentData, ArgumentValueInt)));
 			}
-			else if (ArgumentPinCategory == K2Schema->PC_Float)
+			else if (ArgumentPinCategory == UEdGraphSchema_K2::PC_Float)
 			{
 				MakeFormatArgumentDataStruct->GetSchema()->TrySetDefaultValue(*ArgumentTypePin, TEXT("Float"));
 				CompilerContext.MovePinLinksToIntermediate(*ArgumentPin, *MakeFormatArgumentDataStruct->FindPinChecked(GET_MEMBER_NAME_STRING_CHECKED(FFormatArgumentData, ArgumentValueFloat)));
 			}
-			else if (ArgumentPinCategory == K2Schema->PC_Text)
+			else if (ArgumentPinCategory == UEdGraphSchema_K2::PC_Text)
 			{
 				MakeFormatArgumentDataStruct->GetSchema()->TrySetDefaultValue(*ArgumentTypePin, TEXT("Text"));
 				CompilerContext.MovePinLinksToIntermediate(*ArgumentPin, *MakeFormatArgumentDataStruct->FindPinChecked(GET_MEMBER_NAME_STRING_CHECKED(FFormatArgumentData, ArgumentValue)));
 			}
-			else if (ArgumentPinCategory == K2Schema->PC_Byte && !ArgumentPin->PinType.PinSubCategoryObject.IsValid())
+			else if (ArgumentPinCategory == UEdGraphSchema_K2::PC_Byte && !ArgumentPin->PinType.PinSubCategoryObject.IsValid())
 			{
 				MakeFormatArgumentDataStruct->GetSchema()->TrySetDefaultValue(*ArgumentTypePin, TEXT("Int"));
 
@@ -375,7 +374,7 @@ void UK2Node_FormatText::ExpandNode(class FKismetCompilerContext& CompilerContex
 				// Connect the int output pin to the argument value
 				CallByteToIntFunction->FindPinChecked(TEXT("ReturnValue"))->MakeLinkTo(MakeFormatArgumentDataStruct->FindPinChecked(GET_MEMBER_NAME_STRING_CHECKED(FFormatArgumentData, ArgumentValueInt)));
 			}
-			else if (ArgumentPinCategory == K2Schema->PC_Byte || ArgumentPinCategory == K2Schema->PC_Enum)
+			else if (ArgumentPinCategory == UEdGraphSchema_K2::PC_Byte || ArgumentPinCategory == UEdGraphSchema_K2::PC_Enum)
 			{
 				static UEnum* TextGenderEnum = FindObjectChecked<UEnum>(ANY_PACKAGE, TEXT("ETextGender"), /*ExactClass*/true);
 				if (ArgumentPin->PinType.PinSubCategoryObject == TextGenderEnum)
@@ -387,7 +386,7 @@ void UK2Node_FormatText::ExpandNode(class FKismetCompilerContext& CompilerContex
 			else
 			{
 				// Unexpected pin type!
-				CompilerContext.MessageLog.Error(*FText::Format(LOCTEXT("Error_UnexpectedPinType", "Pin '{0}' has an unexpected type: {1}"), FText::FromString(PinNames[ArgIdx]), FText::FromString(ArgumentPinCategory)).ToString());
+				CompilerContext.MessageLog.Error(*FText::Format(LOCTEXT("Error_UnexpectedPinType", "Pin '{0}' has an unexpected type: {1}"), FText::FromName(PinNames[ArgIdx]), FText::FromName(ArgumentPinCategory)).ToString());
 			}
 		}
 		else
@@ -404,7 +403,7 @@ void UK2Node_FormatText::ExpandNode(class FKismetCompilerContext& CompilerContex
 		}
 
 		// Find the input pin on the "Make Array" node by index.
-		FString PinName = FString::Printf(TEXT("[%d]"), ArgIdx);
+		const FString PinName = FString::Printf(TEXT("[%d]"), ArgIdx);
 		UEdGraphPin* InputPin = MakeArrayNode->FindPinChecked(PinName);
 
 		// Find the output for the pin's "Make Struct" node and link it to the corresponding pin on the "Make Array" node.
@@ -419,18 +418,18 @@ void UK2Node_FormatText::ExpandNode(class FKismetCompilerContext& CompilerContex
 	BreakAllNodeLinks();
 }
 
-UEdGraphPin* UK2Node_FormatText::FindArgumentPin(const FString& InPinName) const
+UEdGraphPin* UK2Node_FormatText::FindArgumentPin(const FName InPinName) const
 {
 	const UEdGraphPin* FormatPin = GetFormatPin();
-	for(int32 PinIdx=0; PinIdx<Pins.Num(); PinIdx++)
+	for (UEdGraphPin* Pin : Pins)
 	{
-		if( Pins[PinIdx] != FormatPin && Pins[PinIdx]->Direction != EGPD_Output && Pins[PinIdx]->PinName.Equals(InPinName) )
+		if( Pin != FormatPin && Pin->Direction != EGPD_Output && Pin->PinName == InPinName )
 		{
-			return Pins[PinIdx];
+			return Pin;
 		}
 	}
 
-	return NULL;
+	return nullptr;
 }
 
 UK2Node::ERedirectType UK2Node_FormatText::DoPinsMatchForReconstruction(const UEdGraphPin* NewPin, int32 NewPinIndex, const UEdGraphPin* OldPin, int32 OldPinIndex) const
@@ -438,7 +437,7 @@ UK2Node::ERedirectType UK2Node_FormatText::DoPinsMatchForReconstruction(const UE
 	ERedirectType RedirectType = ERedirectType_None;
 
 	// if the pin names do match
-	if (FCString::Strcmp(*(NewPin->PinName), *(OldPin->PinName)) == 0)
+	if (NewPin->PinName == OldPin->PinName)
 	{
 		// Make sure we're not dealing with a menu node
 		UEdGraph* OuterGraph = GetGraph();
@@ -468,7 +467,7 @@ UK2Node::ERedirectType UK2Node_FormatText::DoPinsMatchForReconstruction(const UE
 			RedirectType = ShouldRedirectParam(OldPinNames, /*out*/ NewPinName, Node);
 
 			// make sure they match
-			if ((RedirectType != ERedirectType_None) && FCString::Stricmp(*(NewPin->PinName), *(NewPinName.ToString())) != 0)
+			if ((RedirectType != ERedirectType_None) && (NewPin->PinName != NewPinName))
 			{
 				RedirectType = ERedirectType_None;
 			}
@@ -485,14 +484,14 @@ bool UK2Node_FormatText::IsConnectionDisallowed(const UEdGraphPin* MyPin, const
 	if (MyPin != FormatPin && MyPin->Direction == EGPD_Input)
 	{
 		const UEdGraphSchema_K2* K2Schema = Cast<const UEdGraphSchema_K2>(GetSchema());
-		const FString& OtherPinCategory = OtherPin->PinType.PinCategory;
+		const FName& OtherPinCategory = OtherPin->PinType.PinCategory;
 
 		bool bIsValidType = false;
-		if (OtherPinCategory == K2Schema->PC_Int || OtherPinCategory == K2Schema->PC_Float || OtherPinCategory == K2Schema->PC_Text || (OtherPinCategory == K2Schema->PC_Byte && !OtherPin->PinType.PinSubCategoryObject.IsValid()))
+		if (OtherPinCategory == UEdGraphSchema_K2::PC_Int || OtherPinCategory == UEdGraphSchema_K2::PC_Float || OtherPinCategory == UEdGraphSchema_K2::PC_Text || (OtherPinCategory == UEdGraphSchema_K2::PC_Byte && !OtherPin->PinType.PinSubCategoryObject.IsValid()))
 		{
 			bIsValidType = true;
 		}
-		else if (OtherPinCategory == K2Schema->PC_Byte || OtherPinCategory == K2Schema->PC_Enum)
+		else if (OtherPinCategory == UEdGraphSchema_K2::PC_Byte || OtherPinCategory == UEdGraphSchema_K2::PC_Enum)
 		{
 			static UEnum* TextGenderEnum = FindObjectChecked<UEnum>(ANY_PACKAGE, TEXT("ETextGender"), /*ExactClass*/true);
 			if (OtherPin->PinType.PinSubCategoryObject == TextGenderEnum)
@@ -513,9 +512,9 @@ bool UK2Node_FormatText::IsConnectionDisallowed(const UEdGraphPin* MyPin, const
 
 FText UK2Node_FormatText::GetArgumentName(int32 InIndex) const
 {
-	if(InIndex < PinNames.Num())
+	if (InIndex < PinNames.Num())
 	{
-		return FText::FromString(PinNames[InIndex]);
+		return FText::FromName(PinNames[InIndex]);
 	}
 	return FText::GetEmpty();
 }
@@ -525,9 +524,8 @@ void UK2Node_FormatText::AddArgumentPin()
 	const FScopedTransaction Transaction( NSLOCTEXT("Kismet", "AddArgumentPin", "Add Argument Pin") );
 	Modify();
 
-	const UEdGraphSchema_K2* K2Schema = Cast<const UEdGraphSchema_K2>(GetSchema());
-	FString PinName = GetUniquePinName();
-	CreatePin(EGPD_Input, K2Schema->PC_Wildcard, FString(), nullptr, PinName);
+	const FName PinName(GetUniquePinName());
+	CreatePin(EGPD_Input, UEdGraphSchema_K2::PC_Wildcard, PinName);
 	PinNames.Add(PinName);
 
 	FBlueprintEditorUtils::MarkBlueprintAsStructurallyModified(GetBlueprint());
@@ -550,7 +548,7 @@ void UK2Node_FormatText::RemoveArgument(int32 InIndex)
 	GetGraph()->NotifyGraphChanged();
 }
 
-void UK2Node_FormatText::SetArgumentName(int32 InIndex, FString InName)
+void UK2Node_FormatText::SetArgumentName(int32 InIndex, FName InName)
 {
 	PinNames[InIndex] = InName;
 
diff --git a/Engine/Source/Editor/BlueprintGraph/Private/K2Node_FunctionEntry.cpp b/Engine/Source/Editor/BlueprintGraph/Private/K2Node_FunctionEntry.cpp
--- a/Engine/Source/Editor/BlueprintGraph/Private/K2Node_FunctionEntry.cpp
+++ b/Engine/Source/Editor/BlueprintGraph/Private/K2Node_FunctionEntry.cpp
@@ -45,7 +45,7 @@ public:
 		//@TODO: Still doesn't handle/allow users to declare new pass by reference, this only helps inherited functions
 		if( Function )
 		{
-			if (UProperty* ParentProperty = FindField<UProperty>(Function, FName(*(Net->PinName))))
+			if (UProperty* ParentProperty = FindField<UProperty>(Function, Net->PinName))
 			{
 				if (ParentProperty->HasAnyPropertyFlags(CPF_ReferenceParm))
 				{
@@ -142,9 +142,9 @@ public:
 	{
 		UK2Node_FunctionEntry* EntryNode = CastChecked<UK2Node_FunctionEntry>(Node);
 		//check(EntryNode->SignatureName != NAME_None);
-		if (EntryNode->SignatureName == CompilerContext.GetSchema()->FN_ExecuteUbergraphBase)
+		if (EntryNode->SignatureName == UEdGraphSchema_K2::FN_ExecuteUbergraphBase)
 		{
-			UEdGraphPin* EntryPointPin = Node->FindPin(CompilerContext.GetSchema()->PN_EntryPoint);
+			UEdGraphPin* EntryPointPin = Node->FindPin(UEdGraphSchema_K2::PN_EntryPoint);
 			FBPTerminal** pTerm = Context.NetMap.Find(EntryPointPin);
 			if ((EntryPointPin != NULL) && (pTerm != NULL))
 			{
@@ -172,9 +172,9 @@ public:
 
 struct FFunctionEntryHelper
 {
-	static const FString& GetWorldContextPinName()
+	static const FName& GetWorldContextPinName()
 	{
-		static const FString WorldContextPinName(TEXT("__WorldContext"));
+		static const FName WorldContextPinName(TEXT("__WorldContext"));
 		return WorldContextPinName;
 	}
 
@@ -288,7 +288,7 @@ void UK2Node_FunctionEntry::Serialize(FArchive& Ar)
 					K2Schema->GetPinDefaultValuesFromString(LocalVar.VarType, this, LocalVar.DefaultValue, UseDefaultValue, UseDefaultObject, UseDefaultText);
 					FString ErrorMessage;
 
-					if (!K2Schema->DefaultValueSimpleValidation(LocalVar.VarType, LocalVar.VarName.ToString(), UseDefaultValue, UseDefaultObject, UseDefaultText, &ErrorMessage))
+					if (!K2Schema->DefaultValueSimpleValidation(LocalVar.VarType, LocalVar.VarName, UseDefaultValue, UseDefaultObject, UseDefaultText, &ErrorMessage))
 					{
 						const UBlueprint* Blueprint = GetBlueprint();
 						UE_LOG(LogBlueprint, Log, TEXT("Clearing invalid default value for local variable %s on blueprint %s: %s"), *LocalVar.VarName.ToString(), Blueprint ? *Blueprint->GetName() : TEXT("Unknown"), *ErrorMessage);
@@ -312,8 +312,7 @@ FText UK2Node_FunctionEntry::GetNodeTitle(ENodeTitleType::Type TitleType) const
 
 void UK2Node_FunctionEntry::AllocateDefaultPins()
 {
-	const UEdGraphSchema_K2* K2Schema = GetDefault<UEdGraphSchema_K2>();
-	CreatePin(EGPD_Output, K2Schema->PC_Exec, FString(), nullptr, K2Schema->PN_Then);
+	CreatePin(EGPD_Output, UEdGraphSchema_K2::PC_Exec, UEdGraphSchema_K2::PN_Then);
 
 	UFunction* Function = FindField<UFunction>(SignatureClass, SignatureName);
 
@@ -338,8 +337,7 @@ void UK2Node_FunctionEntry::AllocateDefaultPins()
 	{
 		UEdGraphPin* WorldContextPin = CreatePin(
 			EGPD_Output,
-			K2Schema->PC_Object,
-			FString(),
+			UEdGraphSchema_K2::PC_Object,
 			UObject::StaticClass(),
 			FFunctionEntryHelper::GetWorldContextPinName());
 		WorldContextPin->bHidden = true;
@@ -379,7 +377,7 @@ UEdGraphPin* UK2Node_FunctionEntry::CreatePinFromUserDefinition(const TSharedPtr
 {
 	// Make sure that if this is an exec node we are allowed one.
 	const UEdGraphSchema_K2* Schema = GetDefault<UEdGraphSchema_K2>();
-	if (NewPinInfo->PinType.PinCategory == Schema->PC_Exec && !CanModifyExecutionWires())
+	if (NewPinInfo->PinType.PinCategory == UEdGraphSchema_K2::PC_Exec && !CanModifyExecutionWires())
 	{
 		return nullptr;
 	}
@@ -517,7 +515,7 @@ void UK2Node_FunctionEntry::ExpandNode(class FKismetCompilerContext& CompilerCon
 						Schema->ConfigureVarNode(VariableSetNode, LocalVar.VarName, Function, CompilerContext.Blueprint);
 						VariableSetNode->AllocateDefaultPins();
 
-						if(UEdGraphPin* SetPin = VariableSetNode->FindPin(Property->GetName()))
+						if(UEdGraphPin* SetPin = VariableSetNode->FindPin(Property->GetFName()))
 						{
 							if(LocalVar.VarType.IsArray())
 							{
@@ -565,7 +563,7 @@ void UK2Node_FunctionEntry::ExpandNode(class FKismetCompilerContext& CompilerCon
 								{
 									// When regenerating on load, we want to force load assets referenced by local variables.
 									// This functionality is already handled when generating Terms in the Kismet Compiler for Arrays and Structs, so we do not have to worry about them.
-									if (LocalVar.VarType.PinCategory == Schema->PC_Object || LocalVar.VarType.PinCategory == Schema->PC_Class || LocalVar.VarType.PinCategory == Schema->PC_Interface)
+									if (LocalVar.VarType.PinCategory == UEdGraphSchema_K2::PC_Object || LocalVar.VarType.PinCategory == UEdGraphSchema_K2::PC_Class || LocalVar.VarType.PinCategory == UEdGraphSchema_K2::PC_Interface)
 									{
 										FBlueprintEditorUtils::PropertyValueFromString(Property, LocalVar.DefaultValue, LocalVarData->GetStructMemory());
 									}
diff --git a/Engine/Source/Editor/BlueprintGraph/Private/K2Node_FunctionResult.cpp b/Engine/Source/Editor/BlueprintGraph/Private/K2Node_FunctionResult.cpp
--- a/Engine/Source/Editor/BlueprintGraph/Private/K2Node_FunctionResult.cpp
+++ b/Engine/Source/Editor/BlueprintGraph/Private/K2Node_FunctionResult.cpp
@@ -18,7 +18,7 @@ private:
 	{
 		if(K2Schema && Pin)
 		{
-			const bool bValuePin = (Pin->PinType.PinCategory != K2Schema->PC_Exec);
+			const bool bValuePin = (Pin->PinType.PinCategory != UEdGraphSchema_K2::PC_Exec);
 			const bool bNotConnected = (Pin->Direction == EEdGraphPinDirection::EGPD_Input) && (0 == Pin->LinkedTo.Num());
 			const bool bNeedToResetDefaultValue = (Pin->DefaultValue.IsEmpty() && Pin->DefaultObject == nullptr && Pin->DefaultTextValue.IsEmpty()) || !(K2Schema->IsPinDefaultValid(Pin, Pin->DefaultValue, Pin->DefaultObject, Pin->DefaultTextValue).IsEmpty());
 			if (bValuePin && bNotConnected && bNeedToResetDefaultValue)
@@ -70,16 +70,17 @@ public:
 		// Do not register as a default any Pin that comes from being Split
 		if (Net->ParentPin == nullptr)
 		{
-			for (auto& ResultTerm : Context.Results)
+			FString NetPinName = Net->PinName.ToString();
+			for (FBPTerminal& ResultTerm : Context.Results)
 			{
-				if ((ResultTerm.Name == Net->PinName) && (ResultTerm.Type == Net->PinType))
+				if ((ResultTerm.Name == NetPinName) && (ResultTerm.Type == Net->PinType))
 				{
 					Context.NetMap.Add(Net, &ResultTerm);
 					return;
 				}
 			}
 			FBPTerminal* Term = new (Context.Results) FBPTerminal();
-			Term->CopyFromPin(Net, Net->PinName);
+			Term->CopyFromPin(Net, MoveTemp(NetPinName));
 			Context.NetMap.Add(Net, Term);
 		}
 	}
@@ -132,11 +133,9 @@ FText UK2Node_FunctionResult::GetNodeTitle(ENodeTitleType::Type TitleType) const
 
 void UK2Node_FunctionResult::AllocateDefaultPins()
 {
-	const UEdGraphSchema_K2* K2Schema = GetDefault<UEdGraphSchema_K2>();
-	CreatePin(EGPD_Input, K2Schema->PC_Exec, FString(), nullptr, K2Schema->PN_Execute);
+	CreatePin(EGPD_Input, UEdGraphSchema_K2::PC_Exec, UEdGraphSchema_K2::PN_Execute);
 
-	UFunction* Function = FindField<UFunction>(SignatureClass, SignatureName);
-	if (Function != NULL)
+	if (UFunction* Function = FindField<UFunction>(SignatureClass, SignatureName))
 	{
 		CreatePinsForFunctionEntryExit(Function, /*bIsFunctionEntry=*/ false);
 	}
diff --git a/Engine/Source/Editor/BlueprintGraph/Private/K2Node_FunctionTerminator.cpp b/Engine/Source/Editor/BlueprintGraph/Private/K2Node_FunctionTerminator.cpp
--- a/Engine/Source/Editor/BlueprintGraph/Private/K2Node_FunctionTerminator.cpp
+++ b/Engine/Source/Editor/BlueprintGraph/Private/K2Node_FunctionTerminator.cpp
@@ -20,16 +20,16 @@ FLinearColor UK2Node_FunctionTerminator::GetNodeTitleColor() const
 	return GetDefault<UGraphEditorSettings>()->FunctionTerminatorNodeTitleColor;
 }
 
-FString UK2Node_FunctionTerminator::CreateUniquePinName(FString InSourcePinName) const
+FName UK2Node_FunctionTerminator::CreateUniquePinName(FName InSourcePinName) const
 {
 	const UFunction* FoundFunction = FFunctionFromNodeHelper::FunctionFromNode(this);
 
-	FString ResultName = InSourcePinName;
+	FName ResultName = InSourcePinName;
 	int UniqueNum = 0;
 	// Prevent the unique name from being the same as another of the UFunction's properties
-	while(FindPin(ResultName) || FindField<const UProperty>(FoundFunction, *ResultName) != NULL)
+	while(FindPin(ResultName) || FindField<const UProperty>(FoundFunction, ResultName) != nullptr)
 	{
-		ResultName = FString::Printf(TEXT("%s%d"), *InSourcePinName, ++UniqueNum);
+		ResultName = *FString::Printf(TEXT("%s%d"), *InSourcePinName.ToString(), ++UniqueNum);
 	}
 	return ResultName;
 }
@@ -39,8 +39,7 @@ bool UK2Node_FunctionTerminator::CanCreateUserDefinedPin(const FEdGraphPinType&
 	const bool bIsNodeEditable = IsEditable();
 
 	// Make sure that if this is an exec node we are allowed one.
-	const UEdGraphSchema_K2* Schema = GetDefault<UEdGraphSchema_K2>();
-	if (bIsNodeEditable && InPinType.PinCategory == Schema->PC_Exec && !CanModifyExecutionWires())
+	if (bIsNodeEditable && InPinType.PinCategory == UEdGraphSchema_K2::PC_Exec && !CanModifyExecutionWires())
 	{
 		OutErrorMessage = LOCTEXT("MultipleExecPinError", "Cannot support more exec pins!");
 		return false;
diff --git a/Engine/Source/Editor/BlueprintGraph/Private/K2Node_GetArrayItem.cpp b/Engine/Source/Editor/BlueprintGraph/Private/K2Node_GetArrayItem.cpp
--- a/Engine/Source/Editor/BlueprintGraph/Private/K2Node_GetArrayItem.cpp
+++ b/Engine/Source/Editor/BlueprintGraph/Private/K2Node_GetArrayItem.cpp
@@ -118,13 +118,17 @@ UK2Node_GetArrayItem::UK2Node_GetArrayItem(const FObjectInitializer& ObjectIniti
 
 void UK2Node_GetArrayItem::AllocateDefaultPins()
 {
-	// Create the output pin
-	CreatePin(EGPD_Input, UEdGraphSchema_K2::PC_Wildcard, FString(), nullptr, TEXT("Array"), EPinContainerType::Array);
-	UEdGraphPin* IndexPin = CreatePin(EGPD_Input, UEdGraphSchema_K2::PC_Int, FString(), nullptr, TEXT("Dimension 1"));
+	// Create the input pins to create the arrays from
+	UEdGraphNode::FCreatePinParams ArrayPinParams;
+	ArrayPinParams.ContainerType = EPinContainerType::Array;
+	CreatePin(EGPD_Input, UEdGraphSchema_K2::PC_Wildcard, TEXT("Array"), ArrayPinParams);
+	UEdGraphPin* IndexPin = CreatePin(EGPD_Input, UEdGraphSchema_K2::PC_Int, TEXT("Dimension 1"));
 	GetDefault<UEdGraphSchema_K2>()->SetPinAutogeneratedDefaultValueBasedOnType(IndexPin);
 
-	// Create the input pins to create the arrays from
-	CreatePin(EGPD_Output, UEdGraphSchema_K2::PC_Wildcard, FString(), nullptr, TEXT("Output"), EPinContainerType::None, bReturnByRefDesired);
+	// Create the output pin
+	UEdGraphNode::FCreatePinParams OutputPinParams;
+	OutputPinParams.bIsReference = bReturnByRefDesired;
+	CreatePin(EGPD_Output, UEdGraphSchema_K2::PC_Wildcard, TEXT("Output"), OutputPinParams);
 }
 
 void UK2Node_GetArrayItem::PostReconstructNode()
@@ -280,11 +284,11 @@ void UK2Node_GetArrayItem::NotifyPinConnectionListChanged(UEdGraphPin* Pin)
 			if (OtherPin->LinkedTo.Num() == 0)
 			{
 				ArrayPin->PinType.PinCategory = UEdGraphSchema_K2::PC_Wildcard;
-				ArrayPin->PinType.PinSubCategory.Empty();
+				ArrayPin->PinType.PinSubCategory = NAME_None;
 				ArrayPin->PinType.PinSubCategoryObject = nullptr;
 
 				ResultPin->PinType.PinCategory = UEdGraphSchema_K2::PC_Wildcard;
-				ResultPin->PinType.PinSubCategory.Empty();
+				ResultPin->PinType.PinSubCategory = NAME_None;
 				ResultPin->PinType.PinSubCategoryObject = nullptr;
 				ResultPin->PinType.bIsReference = bReturnByRefDesired;
 
diff --git a/Engine/Source/Editor/BlueprintGraph/Private/K2Node_GetClassDefaults.cpp b/Engine/Source/Editor/BlueprintGraph/Private/K2Node_GetClassDefaults.cpp
--- a/Engine/Source/Editor/BlueprintGraph/Private/K2Node_GetClassDefaults.cpp
+++ b/Engine/Source/Editor/BlueprintGraph/Private/K2Node_GetClassDefaults.cpp
@@ -16,7 +16,7 @@
 
 #define LOCTEXT_NAMESPACE "UK2Node_GetClassDefaults"
 
-FString UK2Node_GetClassDefaults::ClassPinName(TEXT("Class"));
+FName UK2Node_GetClassDefaults::ClassPinName(TEXT("Class"));
 
 namespace
 {
@@ -148,10 +148,10 @@ namespace
 								UEdGraphPin* Pin = Node->Pins[PinIndex];
 								if(Pin != nullptr && Pin->Direction == EGPD_Output)
 								{
-									UProperty* BoundProperty = FindField<UProperty>(ClassType, *(Pin->PinName));
+									UProperty* BoundProperty = FindField<UProperty>(ClassType, Pin->PinName);
 									if(BoundProperty != nullptr)
 									{
-										FBPTerminal* OutputTerm = Context.CreateLocalTerminalFromPinAutoChooseScope(Pin, Pin->PinName);
+										FBPTerminal* OutputTerm = Context.CreateLocalTerminalFromPinAutoChooseScope(Pin, Pin->PinName.ToString());
 										check(OutputTerm != nullptr);
 
 										// Set as a variable within the class context
@@ -223,7 +223,7 @@ void UK2Node_GetClassDefaults::AllocateDefaultPins()
 	const UEdGraphSchema_K2* K2Schema = GetDefault<UEdGraphSchema_K2>();
 
 	// Create the class input type selector pin
-	UEdGraphPin* ClassPin = CreatePin(EGPD_Input, K2Schema->PC_Class, FString(), UObject::StaticClass(), ClassPinName);
+	UEdGraphPin* ClassPin = CreatePin(EGPD_Input, UEdGraphSchema_K2::PC_Class, UObject::StaticClass(), ClassPinName);
 	K2Schema->ConstructBasicPinTooltip(*ClassPin, LOCTEXT("ClassPinDescription", "The class from which to access one or more default values."), ClassPin->PinToolTip);
 }
 
@@ -278,7 +278,7 @@ void UK2Node_GetClassDefaults::ValidateNodeDuringCompilation(class FCompilerResu
 			{
 				// Even though container property defaults are copied, the copy could still contain a reference to a non-class object that belongs to the CDO, which would potentially be unsafe to modify.
 				bool bEmitWarning = false;
-				const UProperty* TestProperty = SourceClass->FindPropertyByName(FName(*Pin->PinName));
+				const UProperty* TestProperty = SourceClass->FindPropertyByName(Pin->PinName);
 				if (const UArrayProperty* ArrayProperty = Cast<UArrayProperty>(TestProperty))
 				{
 					bEmitWarning = ArrayProperty->Inner && ArrayProperty->Inner->IsA<UObjectProperty>() && !ArrayProperty->Inner->IsA<UClassProperty>();
@@ -347,7 +347,7 @@ void UK2Node_GetClassDefaults::ExpandNode(class FKismetCompilerContext& Compiler
 	{
 		if(OutputPin != nullptr && OutputPin->Direction == EGPD_Output && OutputPin->LinkedTo.Num() > 0)
 		{
-			UProperty* BoundProperty = FindField<UProperty>(ClassType, *(OutputPin->PinName));
+			UProperty* BoundProperty = FindField<UProperty>(ClassType, OutputPin->PinName);
 			if(BoundProperty != nullptr && (BoundProperty->IsA<UArrayProperty>() || BoundProperty->IsA<USetProperty>() || BoundProperty->IsA<UMapProperty>()))
 			{
 				UK2Node_TemporaryVariable* LocalVariable = CompilerContext.SpawnIntermediateNode<UK2Node_TemporaryVariable>(this, SourceGraph);
diff --git a/Engine/Source/Editor/BlueprintGraph/Private/K2Node_GetDataTableRow.cpp b/Engine/Source/Editor/BlueprintGraph/Private/K2Node_GetDataTableRow.cpp
--- a/Engine/Source/Editor/BlueprintGraph/Private/K2Node_GetDataTableRow.cpp
+++ b/Engine/Source/Editor/BlueprintGraph/Private/K2Node_GetDataTableRow.cpp
@@ -17,14 +17,14 @@
 
 struct UK2Node_GetDataTableRowHelper
 {
-	static FString DataTablePinName;
-	static FString RowNamePinName;
-	static FString RowNotFoundPinName;
+	static FName DataTablePinName;
+	static FName RowNamePinName;
+	static FName RowNotFoundPinName;
 };
 
-FString UK2Node_GetDataTableRowHelper::DataTablePinName(LOCTEXT("DataTablePinName","DataTable").ToString());
-FString UK2Node_GetDataTableRowHelper::RowNotFoundPinName(LOCTEXT("RowNotFoundPinName","RowNotFound").ToString());
-FString UK2Node_GetDataTableRowHelper::RowNamePinName(LOCTEXT("RowNamePinName","RowName").ToString());
+FName UK2Node_GetDataTableRowHelper::DataTablePinName(*LOCTEXT("DataTablePinName","DataTable").ToString());
+FName UK2Node_GetDataTableRowHelper::RowNotFoundPinName(*LOCTEXT("RowNotFoundPinName","RowNotFound").ToString());
+FName UK2Node_GetDataTableRowHelper::RowNamePinName(*LOCTEXT("RowNamePinName","RowName").ToString());
 
 UK2Node_GetDataTableRow::UK2Node_GetDataTableRow(const FObjectInitializer& ObjectInitializer)
 	: Super(ObjectInitializer)
@@ -37,21 +37,21 @@ void UK2Node_GetDataTableRow::AllocateDefaultPins()
 	const UEdGraphSchema_K2* K2Schema = GetDefault<UEdGraphSchema_K2>();
 
 	// Add execution pins
-	CreatePin(EGPD_Input, K2Schema->PC_Exec, FString(), nullptr, K2Schema->PN_Execute);
-	UEdGraphPin* RowFoundPin = CreatePin(EGPD_Output, K2Schema->PC_Exec, FString(), nullptr, K2Schema->PN_Then);
+	CreatePin(EGPD_Input, UEdGraphSchema_K2::PC_Exec, UEdGraphSchema_K2::PN_Execute);
+	UEdGraphPin* RowFoundPin = CreatePin(EGPD_Output, UEdGraphSchema_K2::PC_Exec, UEdGraphSchema_K2::PN_Then);
 	RowFoundPin->PinFriendlyName = LOCTEXT("GetDataTableRow Row Found Exec pin", "Row Found");
-	CreatePin(EGPD_Output, K2Schema->PC_Exec, FString(), nullptr, UK2Node_GetDataTableRowHelper::RowNotFoundPinName);
+	CreatePin(EGPD_Output, UEdGraphSchema_K2::PC_Exec, UK2Node_GetDataTableRowHelper::RowNotFoundPinName);
 
 	// Add DataTable pin
-	UEdGraphPin* DataTablePin = CreatePin(EGPD_Input, K2Schema->PC_Object, FString(), UDataTable::StaticClass(), UK2Node_GetDataTableRowHelper::DataTablePinName);
+	UEdGraphPin* DataTablePin = CreatePin(EGPD_Input, UEdGraphSchema_K2::PC_Object, UDataTable::StaticClass(), UK2Node_GetDataTableRowHelper::DataTablePinName);
 	SetPinToolTip(*DataTablePin, LOCTEXT("DataTablePinDescription", "The DataTable you want to retreive a row from"));
 
 	// Row Name pin
-	UEdGraphPin* RowNamePin = CreatePin(EGPD_Input, K2Schema->PC_Name, FString(), nullptr, UK2Node_GetDataTableRowHelper::RowNamePinName);
+	UEdGraphPin* RowNamePin = CreatePin(EGPD_Input, UEdGraphSchema_K2::PC_Name, UK2Node_GetDataTableRowHelper::RowNamePinName);
 	SetPinToolTip(*RowNamePin, LOCTEXT("RowNamePinDescription", "The name of the row to retrieve from the DataTable"));
 
 	// Result pin
-	UEdGraphPin* ResultPin = CreatePin(EGPD_Output, K2Schema->PC_Wildcard, FString(), nullptr, K2Schema->PN_ReturnValue);
+	UEdGraphPin* ResultPin = CreatePin(EGPD_Output, UEdGraphSchema_K2::PC_Wildcard, UEdGraphSchema_K2::PN_ReturnValue);
 	ResultPin->PinFriendlyName = LOCTEXT("GetDataTableRow Output Row", "Out Row");
 	SetPinToolTip(*ResultPin, LOCTEXT("ResultPinDescription", "The returned TableRow, if found"));
 
@@ -260,7 +260,7 @@ UEdGraphPin* UK2Node_GetDataTableRow::GetThenPin()const
 {
 	const UEdGraphSchema_K2* K2Schema = GetDefault<UEdGraphSchema_K2>();
 
-	UEdGraphPin* Pin = FindPinChecked(K2Schema->PN_Then);
+	UEdGraphPin* Pin = FindPinChecked(UEdGraphSchema_K2::PN_Then);
 	check(Pin->Direction == EGPD_Output);
 	return Pin;
 }
@@ -300,7 +300,7 @@ UEdGraphPin* UK2Node_GetDataTableRow::GetResultPin() const
 {
 	const UEdGraphSchema_K2* K2Schema = GetDefault<UEdGraphSchema_K2>();
 
-	UEdGraphPin* Pin = FindPinChecked(K2Schema->PN_ReturnValue);
+	UEdGraphPin* Pin = FindPinChecked(UEdGraphSchema_K2::PN_ReturnValue);
 	check(Pin->Direction == EGPD_Output);
 	return Pin;
 }
@@ -342,8 +342,6 @@ void UK2Node_GetDataTableRow::ExpandNode(class FKismetCompilerContext& CompilerC
 {
     Super::ExpandNode(CompilerContext, SourceGraph);
     
-	const UEdGraphSchema_K2* Schema = CompilerContext.GetSchema();
-        
     UEdGraphPin* OriginalDataTableInPin = GetDataTablePin();
     UDataTable* Table = (OriginalDataTableInPin != NULL) ? Cast<UDataTable>(OriginalDataTableInPin->DefaultObject) : NULL;
     if((nullptr == OriginalDataTableInPin) || (0 == OriginalDataTableInPin->LinkedTo.Num() && nullptr == Table))
@@ -377,9 +375,9 @@ void UK2Node_GetDataTableRow::ExpandNode(class FKismetCompilerContext& CompilerC
 	CompilerContext.MovePinLinksToIntermediate(*GetRowNamePin(), *RowNameInPin);
 
 	// Get some pins to work with
-	UEdGraphPin* OriginalOutRowPin = FindPinChecked(Schema->PN_ReturnValue);
+	UEdGraphPin* OriginalOutRowPin = FindPinChecked(UEdGraphSchema_K2::PN_ReturnValue);
 	UEdGraphPin* FunctionOutRowPin = GetDataTableRowFunction->FindPinChecked(TEXT("OutRow"));
-    UEdGraphPin* FunctionReturnPin = GetDataTableRowFunction->FindPinChecked(Schema->PN_ReturnValue);
+    UEdGraphPin* FunctionReturnPin = GetDataTableRowFunction->FindPinChecked(UEdGraphSchema_K2::PN_ReturnValue);
     UEdGraphPin* FunctionThenPin = GetDataTableRowFunction->GetThenPin();
         
     // Set the type of the OutRow pin on this expanded mode to match original
diff --git a/Engine/Source/Editor/BlueprintGraph/Private/K2Node_GetEnumeratorName.cpp b/Engine/Source/Editor/BlueprintGraph/Private/K2Node_GetEnumeratorName.cpp
--- a/Engine/Source/Editor/BlueprintGraph/Private/K2Node_GetEnumeratorName.cpp
+++ b/Engine/Source/Editor/BlueprintGraph/Private/K2Node_GetEnumeratorName.cpp
@@ -11,7 +11,7 @@
 #include "EditorCategoryUtils.h"
 #include "BlueprintActionDatabaseRegistrar.h"
 
-FString UK2Node_GetEnumeratorName::EnumeratorPinName = TEXT("Enumerator");
+FName UK2Node_GetEnumeratorName::EnumeratorPinName(TEXT("Enumerator"));
 
 UK2Node_GetEnumeratorName::UK2Node_GetEnumeratorName(const FObjectInitializer& ObjectInitializer)
 	: Super(ObjectInitializer)
@@ -22,8 +22,8 @@ void UK2Node_GetEnumeratorName::AllocateDefaultPins()
 {
 	const UEdGraphSchema_K2* Schema = GetDefault<UEdGraphSchema_K2>();
 
-	CreatePin(EGPD_Input, Schema->PC_Byte, FString(), nullptr, EnumeratorPinName);
-	CreatePin(EGPD_Output, Schema->PC_Name, FString(), nullptr, Schema->PN_ReturnValue);
+	CreatePin(EGPD_Input, UEdGraphSchema_K2::PC_Byte, EnumeratorPinName);
+	CreatePin(EGPD_Output, UEdGraphSchema_K2::PC_Name, UEdGraphSchema_K2::PN_ReturnValue);
 }
 
 FText UK2Node_GetEnumeratorName::GetTooltipText() const
@@ -53,7 +53,7 @@ void UK2Node_GetEnumeratorName::ValidateNodeDuringCompilation(class FCompilerRes
 	Super::ValidateNodeDuringCompilation(MessageLog);
 
 	const UEdGraphSchema_K2* Schema = GetDefault<UEdGraphSchema_K2>();
-	const UEdGraphPin* OutputPin = FindPinChecked(Schema->PN_ReturnValue); 
+	const UEdGraphPin* OutputPin = FindPinChecked(UEdGraphSchema_K2::PN_ReturnValue); 
 	/*Don't validate isolated nodes */
 	if (0 != OutputPin->LinkedTo.Num())
 	{
@@ -83,7 +83,7 @@ bool UK2Node_GetEnumeratorName::IsConnectionDisallowed(const UEdGraphPin* MyPin,
 	const UEdGraphSchema_K2* Schema = GetDefault<UEdGraphSchema_K2>();
 
 	const UEdGraphPin* InputPin = FindPinChecked(EnumeratorPinName);
-	if((InputPin == MyPin) && OtherPin && (OtherPin->PinType.PinCategory == Schema->PC_Byte))
+	if((InputPin == MyPin) && OtherPin && (OtherPin->PinType.PinCategory == UEdGraphSchema_K2::PC_Byte))
 	{
 		if(NULL == Cast<UEnum>(OtherPin->PinType.PinSubCategoryObject.Get()))
 		{
@@ -150,9 +150,9 @@ void UK2Node_GetEnumeratorName::ExpandNode(class FKismetCompilerContext& Compile
 	check(CallGetName->IsNodePure());
 		
 	//OPUTPUT PIN
-	UEdGraphPin* OrgReturnPin = FindPinChecked(Schema->PN_ReturnValue);
+	UEdGraphPin* OrgReturnPin = FindPinChecked(UEdGraphSchema_K2::PN_ReturnValue);
 	UEdGraphPin* NewReturnPin = CallGetName->GetReturnValuePin();
-	check(NULL != NewReturnPin);
+	check(NewReturnPin);
 	CompilerContext.MovePinLinksToIntermediate(*OrgReturnPin, *NewReturnPin);
 
 	//ENUM PIN
@@ -163,7 +163,7 @@ void UK2Node_GetEnumeratorName::ExpandNode(class FKismetCompilerContext& Compile
 	//VALUE PIN
 	UEdGraphPin* OrgInputPin = FindPinChecked(EnumeratorPinName);
 	UEdGraphPin* IndexPin = CallGetName->FindPinChecked(TEXT("EnumeratorValue"));
-	check(EGPD_Input == IndexPin->Direction && Schema->PC_Byte == IndexPin->PinType.PinCategory);
+	check(EGPD_Input == IndexPin->Direction && UEdGraphSchema_K2::PC_Byte == IndexPin->PinType.PinCategory);
 	CompilerContext.MovePinLinksToIntermediate(*OrgInputPin, *IndexPin);
 
 	if (!IndexPin->LinkedTo.Num())
@@ -184,7 +184,7 @@ void UK2Node_GetEnumeratorName::ExpandNode(class FKismetCompilerContext& Compile
 		MakeLiteralByte->SetFromFunction(UKismetSystemLibrary::StaticClass()->FindFunctionByName(FunctionName));
 		MakeLiteralByte->AllocateDefaultPins();
 
-		UEdGraphPin* MakeLiteralByteReturnPin = MakeLiteralByte->FindPinChecked(Schema->PN_ReturnValue);
+		UEdGraphPin* MakeLiteralByteReturnPin = MakeLiteralByte->FindPinChecked(UEdGraphSchema_K2::PN_ReturnValue);
 		Schema->TryCreateConnection(MakeLiteralByteReturnPin, IndexPin);
 
 		UEdGraphPin* MakeLiteralByteInputPin = MakeLiteralByte->FindPinChecked(TEXT("Value"));
diff --git a/Engine/Source/Editor/BlueprintGraph/Private/K2Node_GetEnumeratorNameAsString.cpp b/Engine/Source/Editor/BlueprintGraph/Private/K2Node_GetEnumeratorNameAsString.cpp
--- a/Engine/Source/Editor/BlueprintGraph/Private/K2Node_GetEnumeratorNameAsString.cpp
+++ b/Engine/Source/Editor/BlueprintGraph/Private/K2Node_GetEnumeratorNameAsString.cpp
@@ -17,8 +17,8 @@ void UK2Node_GetEnumeratorNameAsString::AllocateDefaultPins()
 {
 	const UEdGraphSchema_K2* Schema = GetDefault<UEdGraphSchema_K2>();
 
-	CreatePin(EGPD_Input, Schema->PC_Byte, FString(), nullptr, EnumeratorPinName);
-	CreatePin(EGPD_Output, Schema->PC_String, FString(), nullptr, Schema->PN_ReturnValue);
+	CreatePin(EGPD_Input, UEdGraphSchema_K2::PC_Byte, EnumeratorPinName);
+	CreatePin(EGPD_Output, UEdGraphSchema_K2::PC_String, UEdGraphSchema_K2::PN_ReturnValue);
 }
 
 FText UK2Node_GetEnumeratorNameAsString::GetTooltipText() const
diff --git a/Engine/Source/Editor/BlueprintGraph/Private/K2Node_GetNumEnumEntries.cpp b/Engine/Source/Editor/BlueprintGraph/Private/K2Node_GetNumEnumEntries.cpp
--- a/Engine/Source/Editor/BlueprintGraph/Private/K2Node_GetNumEnumEntries.cpp
+++ b/Engine/Source/Editor/BlueprintGraph/Private/K2Node_GetNumEnumEntries.cpp
@@ -17,10 +17,8 @@ UK2Node_GetNumEnumEntries::UK2Node_GetNumEnumEntries(const FObjectInitializer& O
 
 void UK2Node_GetNumEnumEntries::AllocateDefaultPins()
 {
-	const UEdGraphSchema_K2* Schema = GetDefault<UEdGraphSchema_K2>();
-
 	// Create the return value pin
-	CreatePin(EGPD_Output, Schema->PC_Int, FString(), nullptr, Schema->PN_ReturnValue);
+	CreatePin(EGPD_Output, UEdGraphSchema_K2::PC_Int, UEdGraphSchema_K2::PN_ReturnValue);
 
 	Super::AllocateDefaultPins();
 }
@@ -76,9 +74,6 @@ void UK2Node_GetNumEnumEntries::ExpandNode(class FKismetCompilerContext& Compile
 		Enum->GetLinker()->Preload(Enum);
 	}
 
-	const UEdGraphSchema_K2* Schema = CompilerContext.GetSchema();
-	check(NULL != Schema);
-
 	//MAKE LITERAL
 	const FName FunctionName = GET_FUNCTION_NAME_CHECKED(UKismetSystemLibrary, MakeLiteralInt);
 	UK2Node_CallFunction* MakeLiteralInt = CompilerContext.SpawnIntermediateNode<UK2Node_CallFunction>(this, SourceGraph); 
@@ -86,7 +81,7 @@ void UK2Node_GetNumEnumEntries::ExpandNode(class FKismetCompilerContext& Compile
 	MakeLiteralInt->AllocateDefaultPins();
 
 	//OPUTPUT PIN
-	UEdGraphPin* OrgReturnPin = FindPinChecked(Schema->PN_ReturnValue);
+	UEdGraphPin* OrgReturnPin = FindPinChecked(UEdGraphSchema_K2::PN_ReturnValue);
 	UEdGraphPin* NewReturnPin = MakeLiteralInt->GetReturnValuePin();
 	check(NULL != NewReturnPin);
 	CompilerContext.MovePinLinksToIntermediate(*OrgReturnPin, *NewReturnPin);
diff --git a/Engine/Source/Editor/BlueprintGraph/Private/K2Node_IfThenElse.cpp b/Engine/Source/Editor/BlueprintGraph/Private/K2Node_IfThenElse.cpp
--- a/Engine/Source/Editor/BlueprintGraph/Private/K2Node_IfThenElse.cpp
+++ b/Engine/Source/Editor/BlueprintGraph/Private/K2Node_IfThenElse.cpp
@@ -48,9 +48,9 @@ public:
 		}
 
 		// Generate the output impulse from this node
-		UEdGraphPin* CondPin = Context.FindRequiredPinByName(Node, CompilerContext.GetSchema()->PN_Condition, EGPD_Input);
-		UEdGraphPin* ThenPin = Context.FindRequiredPinByName(Node, CompilerContext.GetSchema()->PN_Then, EGPD_Output);
-		UEdGraphPin* ElsePin = Context.FindRequiredPinByName(Node, CompilerContext.GetSchema()->PN_Else, EGPD_Output);
+		UEdGraphPin* CondPin = Context.FindRequiredPinByName(Node, UEdGraphSchema_K2::PN_Condition, EGPD_Input);
+		UEdGraphPin* ThenPin = Context.FindRequiredPinByName(Node, UEdGraphSchema_K2::PN_Then, EGPD_Output);
+		UEdGraphPin* ElsePin = Context.FindRequiredPinByName(Node, UEdGraphSchema_K2::PN_Else, EGPD_Output);
 		if (Context.ValidatePinType(ThenPin, ExpectedExecPinType) &&
 			Context.ValidatePinType(ElsePin, ExpectedExecPinType) &&
 			Context.ValidatePinType(CondPin, ExpectedBoolPinType))
@@ -93,14 +93,14 @@ void UK2Node_IfThenElse::AllocateDefaultPins()
 {
 	const UEdGraphSchema_K2* K2Schema = GetDefault<UEdGraphSchema_K2>();
 
-	CreatePin(EGPD_Input, K2Schema->PC_Exec, FString(), nullptr, K2Schema->PN_Execute);
-	UEdGraphPin* ConditionPin = CreatePin(EGPD_Input, K2Schema->PC_Boolean, FString(), nullptr, K2Schema->PN_Condition);
+	CreatePin(EGPD_Input, UEdGraphSchema_K2::PC_Exec, UEdGraphSchema_K2::PN_Execute);
+	UEdGraphPin* ConditionPin = CreatePin(EGPD_Input, UEdGraphSchema_K2::PC_Boolean, UEdGraphSchema_K2::PN_Condition);
 	K2Schema->SetPinAutogeneratedDefaultValue(ConditionPin, TEXT("true"));
 
-	UEdGraphPin* TruePin = CreatePin(EGPD_Output, K2Schema->PC_Exec, FString(), nullptr, K2Schema->PN_Then);
+	UEdGraphPin* TruePin = CreatePin(EGPD_Output, UEdGraphSchema_K2::PC_Exec, UEdGraphSchema_K2::PN_Then);
 	TruePin->PinFriendlyName =LOCTEXT("true", "true");
 
-	UEdGraphPin* FalsePin = CreatePin(EGPD_Output, K2Schema->PC_Exec, FString(), nullptr, K2Schema->PN_Else);
+	UEdGraphPin* FalsePin = CreatePin(EGPD_Output, UEdGraphSchema_K2::PC_Exec, UEdGraphSchema_K2::PN_Else);
 	FalsePin->PinFriendlyName = LOCTEXT("false", "false");
 
 	Super::AllocateDefaultPins();
@@ -129,28 +129,22 @@ FText UK2Node_IfThenElse::GetTooltipText() const
 
 UEdGraphPin* UK2Node_IfThenElse::GetThenPin() const
 {
-	const UEdGraphSchema_K2* K2Schema = GetDefault<UEdGraphSchema_K2>();
-
-	UEdGraphPin* Pin = FindPin(K2Schema->PN_Then);
-	check(Pin != NULL);
+	UEdGraphPin* Pin = FindPin(UEdGraphSchema_K2::PN_Then);
+	check(Pin);
 	return Pin;
 }
 
 UEdGraphPin* UK2Node_IfThenElse::GetElsePin() const
 {
-	const UEdGraphSchema_K2* K2Schema = GetDefault<UEdGraphSchema_K2>();
-
-	UEdGraphPin* Pin = FindPin(K2Schema->PN_Else);
-	check(Pin != NULL);
+	UEdGraphPin* Pin = FindPin(UEdGraphSchema_K2::PN_Else);
+	check(Pin);
 	return Pin;
 }
 
 UEdGraphPin* UK2Node_IfThenElse::GetConditionPin() const
 {
-	const UEdGraphSchema_K2* K2Schema = GetDefault<UEdGraphSchema_K2>();
-
-	UEdGraphPin* Pin = FindPin(K2Schema->PN_Condition);
-	check(Pin != NULL);
+	UEdGraphPin* Pin = FindPin(UEdGraphSchema_K2::PN_Condition);
+	check(Pin);
 	return Pin;
 }
 
diff --git a/Engine/Source/Editor/BlueprintGraph/Private/K2Node_InputAction.cpp b/Engine/Source/Editor/BlueprintGraph/Private/K2Node_InputAction.cpp
--- a/Engine/Source/Editor/BlueprintGraph/Private/K2Node_InputAction.cpp
+++ b/Engine/Source/Editor/BlueprintGraph/Private/K2Node_InputAction.cpp
@@ -38,11 +38,9 @@ void UK2Node_InputAction::PostLoad()
 
 void UK2Node_InputAction::AllocateDefaultPins()
 {
-	const UEdGraphSchema_K2* K2Schema = GetDefault<UEdGraphSchema_K2>();
-
-	CreatePin(EGPD_Output, K2Schema->PC_Exec, FString(), nullptr, TEXT("Pressed"));
-	CreatePin(EGPD_Output, K2Schema->PC_Exec, FString(), nullptr, TEXT("Released"));
-	CreatePin(EGPD_Output, K2Schema->PC_Struct, FString(), FKey::StaticStruct(), TEXT("Key"));
+	CreatePin(EGPD_Output, UEdGraphSchema_K2::PC_Exec, TEXT("Pressed"));
+	CreatePin(EGPD_Output, UEdGraphSchema_K2::PC_Exec, TEXT("Released"));
+	CreatePin(EGPD_Output, UEdGraphSchema_K2::PC_Struct, FKey::StaticStruct(), TEXT("Key"));
 
 	Super::AllocateDefaultPins();
 }
@@ -98,7 +96,7 @@ bool UK2Node_InputAction::IsCompatibleWithGraph(UEdGraph const* Graph) const
 		UBlueprint* Blueprint = FBlueprintEditorUtils::FindBlueprintForGraph(Graph);
 
 		UEdGraphSchema_K2 const* K2Schema = Cast<UEdGraphSchema_K2>(Graph->GetSchema());
-		bool const bIsConstructionScript = (K2Schema != nullptr) ? K2Schema->IsConstructionScript(Graph) : false;
+		bool const bIsConstructionScript = (K2Schema != nullptr) ? UEdGraphSchema_K2::IsConstructionScript(Graph) : false;
 
 		bIsCompatible = (Blueprint != nullptr) && Blueprint->SupportsInputEvents() && !bIsConstructionScript && Super::IsCompatibleWithGraph(Graph);
 	}
@@ -159,7 +157,7 @@ void UK2Node_InputAction::ExpandNode(FKismetCompilerContext& CompilerContext, UE
 		// Create a temporary variable to copy Key in to
 		static UScriptStruct* KeyStruct = FKey::StaticStruct();
 		UK2Node_TemporaryVariable* ActionKeyVar = CompilerContext.SpawnIntermediateNode<UK2Node_TemporaryVariable>(this, SourceGraph);
-		ActionKeyVar->VariableType.PinCategory = Schema->PC_Struct;
+		ActionKeyVar->VariableType.PinCategory = UEdGraphSchema_K2::PC_Struct;
 		ActionKeyVar->VariableType.PinSubCategoryObject = KeyStruct;
 		ActionKeyVar->AllocateDefaultPins();
 
diff --git a/Engine/Source/Editor/BlueprintGraph/Private/K2Node_InputKey.cpp b/Engine/Source/Editor/BlueprintGraph/Private/K2Node_InputKey.cpp
--- a/Engine/Source/Editor/BlueprintGraph/Private/K2Node_InputKey.cpp
+++ b/Engine/Source/Editor/BlueprintGraph/Private/K2Node_InputKey.cpp
@@ -53,11 +53,9 @@ void UK2Node_InputKey::PostEditChangeProperty(FPropertyChangedEvent& PropertyCha
 
 void UK2Node_InputKey::AllocateDefaultPins()
 {
-	const UEdGraphSchema_K2* K2Schema = GetDefault<UEdGraphSchema_K2>();
-
-	CreatePin(EGPD_Output, K2Schema->PC_Exec, FString(), nullptr, TEXT("Pressed"));
-	CreatePin(EGPD_Output, K2Schema->PC_Exec, FString(), nullptr, TEXT("Released"));
-	CreatePin(EGPD_Output, K2Schema->PC_Struct, FString(), FKey::StaticStruct(), TEXT("Key"));
+	CreatePin(EGPD_Output, UEdGraphSchema_K2::PC_Exec, TEXT("Pressed"));
+	CreatePin(EGPD_Output, UEdGraphSchema_K2::PC_Exec, TEXT("Released"));
+	CreatePin(EGPD_Output, UEdGraphSchema_K2::PC_Struct, FKey::StaticStruct(), TEXT("Key"));
 
 	Super::AllocateDefaultPins();
 }
@@ -209,7 +207,7 @@ bool UK2Node_InputKey::IsCompatibleWithGraph(UEdGraph const* Graph) const
 		UBlueprint* Blueprint = FBlueprintEditorUtils::FindBlueprintForGraph(Graph);
 
 		UEdGraphSchema_K2 const* K2Schema = Cast<UEdGraphSchema_K2>(Graph->GetSchema());
-		bool const bIsConstructionScript = (K2Schema != nullptr) ? K2Schema->IsConstructionScript(Graph) : false;
+		bool const bIsConstructionScript = (K2Schema != nullptr) ? UEdGraphSchema_K2::IsConstructionScript(Graph) : false;
 
 		bIsCompatible = (Blueprint != nullptr) && Blueprint->SupportsInputEvents() && !bIsConstructionScript && Super::IsCompatibleWithGraph(Graph);
 	}
@@ -276,7 +274,7 @@ void UK2Node_InputKey::ExpandNode(FKismetCompilerContext& CompilerContext, UEdGr
 		// Create a temporary variable to copy Key in to
 		static UScriptStruct* KeyStruct = FKey::StaticStruct();
 		UK2Node_TemporaryVariable* KeyVar = CompilerContext.SpawnIntermediateNode<UK2Node_TemporaryVariable>(this, SourceGraph);
-		KeyVar->VariableType.PinCategory = Schema->PC_Struct;
+		KeyVar->VariableType.PinCategory = UEdGraphSchema_K2::PC_Struct;
 		KeyVar->VariableType.PinSubCategoryObject = KeyStruct;
 		KeyVar->AllocateDefaultPins();
 
diff --git a/Engine/Source/Editor/BlueprintGraph/Private/K2Node_InputTouch.cpp b/Engine/Source/Editor/BlueprintGraph/Private/K2Node_InputTouch.cpp
--- a/Engine/Source/Editor/BlueprintGraph/Private/K2Node_InputTouch.cpp
+++ b/Engine/Source/Editor/BlueprintGraph/Private/K2Node_InputTouch.cpp
@@ -44,16 +44,14 @@ UEnum* UK2Node_InputTouch::GetTouchIndexEnum()
 
 void UK2Node_InputTouch::AllocateDefaultPins()
 {
-	const UEdGraphSchema_K2* K2Schema = GetDefault<UEdGraphSchema_K2>();
-
-	CreatePin(EGPD_Output, K2Schema->PC_Exec, FString(), nullptr, TEXT("Pressed"));
-	CreatePin(EGPD_Output, K2Schema->PC_Exec, FString(), nullptr, TEXT("Released"));
-	CreatePin(EGPD_Output, K2Schema->PC_Exec, FString(), nullptr, TEXT("Moved"));
+	CreatePin(EGPD_Output, UEdGraphSchema_K2::PC_Exec, TEXT("Pressed"));
+	CreatePin(EGPD_Output, UEdGraphSchema_K2::PC_Exec, TEXT("Released"));
+	CreatePin(EGPD_Output, UEdGraphSchema_K2::PC_Exec, TEXT("Moved"));
 	
 	UScriptStruct* VectorStruct = TBaseStructure<FVector>::Get();
-	CreatePin(EGPD_Output, K2Schema->PC_Struct, FString(), VectorStruct, TEXT("Location"));
+	CreatePin(EGPD_Output, UEdGraphSchema_K2::PC_Struct, VectorStruct, TEXT("Location"));
 
-	CreatePin(EGPD_Output, K2Schema->PC_Byte, FString(), GetTouchIndexEnum(), TEXT("FingerIndex"));
+	CreatePin(EGPD_Output, UEdGraphSchema_K2::PC_Byte, GetTouchIndexEnum(), TEXT("FingerIndex"));
 
 	Super::AllocateDefaultPins();
 }
@@ -189,13 +187,13 @@ void UK2Node_InputTouch::ExpandNode(FKismetCompilerContext& CompilerContext, UEd
 		// Create a temporary variable to copy location in to
 		static UScriptStruct* VectorStruct = TBaseStructure<FVector>::Get();
 		UK2Node_TemporaryVariable* TouchLocationVar = CompilerContext.SpawnIntermediateNode<UK2Node_TemporaryVariable>(this, SourceGraph);
-		TouchLocationVar->VariableType.PinCategory = Schema->PC_Struct;
+		TouchLocationVar->VariableType.PinCategory = UEdGraphSchema_K2::PC_Struct;
 		TouchLocationVar->VariableType.PinSubCategoryObject = VectorStruct;
 		TouchLocationVar->AllocateDefaultPins();
 
 		// Create a temporary variable to copy finger index in to
 		UK2Node_TemporaryVariable* TouchFingerVar = CompilerContext.SpawnIntermediateNode<UK2Node_TemporaryVariable>(this, SourceGraph);
-		TouchFingerVar->VariableType.PinCategory = Schema->PC_Byte;
+		TouchFingerVar->VariableType.PinCategory = UEdGraphSchema_K2::PC_Byte;
 		TouchFingerVar->VariableType.PinSubCategoryObject = UK2Node_InputTouch::GetTouchIndexEnum();
 		TouchFingerVar->AllocateDefaultPins();
 
diff --git a/Engine/Source/Editor/BlueprintGraph/Private/K2Node_Knot.cpp b/Engine/Source/Editor/BlueprintGraph/Private/K2Node_Knot.cpp
--- a/Engine/Source/Editor/BlueprintGraph/Private/K2Node_Knot.cpp
+++ b/Engine/Source/Editor/BlueprintGraph/Private/K2Node_Knot.cpp
@@ -19,15 +19,13 @@ UK2Node_Knot::UK2Node_Knot(const FObjectInitializer& ObjectInitializer)
 
 void UK2Node_Knot::AllocateDefaultPins()
 {
-	const UEdGraphSchema_K2* Schema = GetDefault<UEdGraphSchema_K2>();
+	const FName InputPinName(TEXT("InputPin"));
+	const FName OutputPinName(TEXT("OutputPin"));
 
-	const FString InputPinName(TEXT("InputPin"));
-	const FString OutputPinName(TEXT("OutputPin"));
-
-	UEdGraphPin* MyInputPin = CreatePin(EGPD_Input, Schema->PC_Wildcard, FString(), nullptr, InputPinName);
+	UEdGraphPin* MyInputPin = CreatePin(EGPD_Input, UEdGraphSchema_K2::PC_Wildcard, InputPinName);
 	MyInputPin->bDefaultValueIsIgnored = true;
 
-	UEdGraphPin* MyOutputPin = CreatePin(EGPD_Output, Schema->PC_Wildcard, FString(), nullptr, OutputPinName);
+	CreatePin(EGPD_Output, UEdGraphSchema_K2::PC_Wildcard, OutputPinName);
 }
 
 FText UK2Node_Knot::GetTooltipText() const
@@ -122,15 +120,13 @@ void UK2Node_Knot::PropagatePinType()
 	else
 	{
 		// Revert to wildcard
-		const UEdGraphSchema_K2* K2Schema = GetDefault<UEdGraphSchema_K2>();
-
 		MyInputPin->BreakAllPinLinks();
 		MyInputPin->PinType.ResetToDefaults();
-		MyInputPin->PinType.PinCategory = K2Schema->PC_Wildcard;
+		MyInputPin->PinType.PinCategory = UEdGraphSchema_K2::PC_Wildcard;
 
 		MyOutputPin->BreakAllPinLinks();
 		MyOutputPin->PinType.ResetToDefaults();
-		MyOutputPin->PinType.PinCategory = K2Schema->PC_Wildcard;
+		MyOutputPin->PinType.PinCategory = UEdGraphSchema_K2::PC_Wildcard;
 	}
 }
 
diff --git a/Engine/Source/Editor/BlueprintGraph/Private/K2Node_Literal.cpp b/Engine/Source/Editor/BlueprintGraph/Private/K2Node_Literal.cpp
--- a/Engine/Source/Editor/BlueprintGraph/Private/K2Node_Literal.cpp
+++ b/Engine/Source/Editor/BlueprintGraph/Private/K2Node_Literal.cpp
@@ -71,13 +71,12 @@ UK2Node_Literal::UK2Node_Literal(const FObjectInitializer& ObjectInitializer)
 {
 }
 
-static FString ValuePinName(TEXT("Value"));
+static const FName ValuePinName(TEXT("Value"));
 
 void UK2Node_Literal::AllocateDefaultPins()
 {
 	// The literal node only has one pin:  an output of the desired value, on a wildcard pin type
-	const UEdGraphSchema_K2* Schema = GetDefault<UEdGraphSchema_K2>();
-	CreatePin(EGPD_Output, Schema->PC_Object, FString(), nullptr, *ValuePinName);
+	CreatePin(EGPD_Output, UEdGraphSchema_K2::PC_Object, ValuePinName);
 
 	// After allocating the pins, try to coerce pin type
 	SetObjectRef( ObjectRef );
@@ -85,14 +84,13 @@ void UK2Node_Literal::AllocateDefaultPins()
 
 void UK2Node_Literal::PostReconstructNode()
 {
-	const UEdGraphSchema_K2* Schema = GetDefault<UEdGraphSchema_K2>();
 	UEdGraphPin* ValuePin = GetValuePin();
 
 	// See if we need to fix up the value pin to have a valid type after reconstruction.  Could be invalid with a stale object ref
 	if( ValuePin && !ObjectRef && (ValuePin->LinkedTo.Num() > 0) )
 	{
 		// Loop over the node, and figure out the most derived class connected to this pin, and use that, so everyone is happy
-		UClass* PinSubtype = NULL;
+		UClass* PinSubtype = nullptr;
 		for( TArray<UEdGraphPin*>::TIterator PinIt(ValuePin->LinkedTo); PinIt; ++PinIt )
 		{
 			UClass* TestType = Cast<UClass>((*PinIt)->PinType.PinSubCategoryObject.Get());
@@ -255,7 +253,6 @@ UEdGraphPin* UK2Node_Literal::GetValuePin() const
 
 void UK2Node_Literal::SetObjectRef(UObject* NewValue)
 {
-	const UEdGraphSchema_K2* Schema = GetDefault<UEdGraphSchema_K2>();
 	UEdGraphPin* ValuePin = GetValuePin();
 
 	// First, see if this is an object
@@ -267,8 +264,8 @@ void UK2Node_Literal::SetObjectRef(UObject* NewValue)
 		if( ValuePin )
 		{
 			ValuePin->Modify();
-			ValuePin->PinType.PinCategory = Schema->PC_Object;
-			ValuePin->PinType.PinSubCategory.Reset();
+			ValuePin->PinType.PinCategory = UEdGraphSchema_K2::PC_Object;
+			ValuePin->PinType.PinSubCategory = NAME_None;
 			ValuePin->PinType.PinSubCategoryObject = ObjectRef->GetClass();
 		}
 	}
@@ -276,7 +273,7 @@ void UK2Node_Literal::SetObjectRef(UObject* NewValue)
 	if( ValuePin )
 	{
 		ValuePin->PinFriendlyName = GetNodeTitle(ENodeTitleType::FullTitle);
-		ValuePin->PinName = ValuePin->PinFriendlyName.BuildSourceString();
+		ValuePin->PinName = *ValuePin->PinFriendlyName.BuildSourceString();
 	}
 }
 
