Copying //UE4/Dev-Framework to //UE4/Dev-Main (Source: //UE4/Dev-Framework @ 3716594)

diff --git a/Engine/Source/Editor/BlueprintGraph/Private/K2Node_LoadAsset.cpp b/Engine/Source/Editor/BlueprintGraph/Private/K2Node_LoadAsset.cpp
--- a/Engine/Source/Editor/BlueprintGraph/Private/K2Node_LoadAsset.cpp
+++ b/Engine/Source/Editor/BlueprintGraph/Private/K2Node_LoadAsset.cpp
@@ -19,16 +19,16 @@
 
 void UK2Node_LoadAsset::AllocateDefaultPins()
 {
-	CreatePin(EGPD_Input, UEdGraphSchema_K2::PC_Exec, FString(), nullptr, UEdGraphSchema_K2::PN_Execute);
+	CreatePin(EGPD_Input, UEdGraphSchema_K2::PC_Exec, UEdGraphSchema_K2::PN_Execute);
 
 	// The immediate continue pin
-	CreatePin(EGPD_Output, UEdGraphSchema_K2::PC_Exec, FString(), nullptr, UEdGraphSchema_K2::PN_Then);
+	CreatePin(EGPD_Output, UEdGraphSchema_K2::PC_Exec, UEdGraphSchema_K2::PN_Then);
 
 	// The delayed completed pin, this used to be called Then
-	CreatePin(EGPD_Output, UEdGraphSchema_K2::PC_Exec, FString(), nullptr, UEdGraphSchema_K2::PN_Completed);
+	CreatePin(EGPD_Output, UEdGraphSchema_K2::PC_Exec, UEdGraphSchema_K2::PN_Completed);
 
-	CreatePin(EGPD_Input, GetInputCategory(), FString(), UObject::StaticClass(), GetInputPinName());
-	CreatePin(EGPD_Output, GetOutputCategory(), FString(), UObject::StaticClass(), GetOutputPinName());
+	CreatePin(EGPD_Input, GetInputCategory(), UObject::StaticClass(), GetInputPinName());
+	CreatePin(EGPD_Output, GetOutputCategory(), UObject::StaticClass(), GetOutputPinName());
 }
 
 void UK2Node_LoadAsset::ReallocatePinsDuringReconstruction(TArray<UEdGraphPin*>& OldPins)
@@ -89,13 +89,13 @@ void UK2Node_LoadAsset::ExpandNode(class FKismetCompilerContext& CompilerContext
 
 	// connect then to second sequence pin
 	{
-		UEdGraphPin* OutputThenPin = FindPin(Schema->PN_Then);
+		UEdGraphPin* OutputThenPin = FindPin(UEdGraphSchema_K2::PN_Then);
 		UEdGraphPin* SequenceSecondExePin = SequenceNode->GetThenPinGivenIndex(1);
 		bIsErrorFree &= OutputThenPin && SequenceSecondExePin && CompilerContext.MovePinLinksToIntermediate(*OutputThenPin, *SequenceSecondExePin).CanSafeConnect();
 	}
 
 	// Create Local Variable
-	UK2Node_TemporaryVariable* TempVarOutput = CompilerContext.SpawnInternalVariable(this, GetOutputCategory(), FString(), UObject::StaticClass());
+	UK2Node_TemporaryVariable* TempVarOutput = CompilerContext.SpawnInternalVariable(this, GetOutputCategory(), NAME_None, UObject::StaticClass());
 
 	// Create assign node
 	UK2Node_AssignmentStatement* AssignNode = CompilerContext.SpawnIntermediateNode<UK2Node_AssignmentStatement>(this, SourceGraph);
@@ -117,14 +117,14 @@ void UK2Node_LoadAsset::ExpandNode(class FKismetCompilerContext& CompilerContext
 
 	// connect assign exec input to function output
 	{
-		UEdGraphPin* CallFunctionOutputExePin = CallLoadAssetNode->FindPin(Schema->PN_Then);
+		UEdGraphPin* CallFunctionOutputExePin = CallLoadAssetNode->FindPin(UEdGraphSchema_K2::PN_Then);
 		UEdGraphPin* AssignInputExePin = AssignNode->GetExecPin();
 		bIsErrorFree &= AssignInputExePin && CallFunctionOutputExePin && Schema->TryCreateConnection(AssignInputExePin, CallFunctionOutputExePin);
 	}
 
 	// connect assign exec output to output
 	{
-		UEdGraphPin* OutputCompletedPin = FindPin(Schema->PN_Completed);
+		UEdGraphPin* OutputCompletedPin = FindPin(UEdGraphSchema_K2::PN_Completed);
 		UEdGraphPin* AssignOutputExePin = AssignNode->GetThenPin();
 		bIsErrorFree &= OutputCompletedPin && AssignOutputExePin && CompilerContext.MovePinLinksToIntermediate(*OutputCompletedPin, *AssignOutputExePin).CanSafeConnect();
 	}
@@ -154,13 +154,13 @@ void UK2Node_LoadAsset::ExpandNode(class FKismetCompilerContext& CompilerContext
 	}
 
 	// Create OnLoadEvent
-	const FString DelegateOnLoadedParamName(TEXT("OnLoaded"));
+	const FName DelegateOnLoadedParamName(TEXT("OnLoaded"));
 	UK2Node_CustomEvent* OnLoadEventNode = CompilerContext.SpawnIntermediateEventNode<UK2Node_CustomEvent>(this, CallFunctionAssetPin, SourceGraph);
 	OnLoadEventNode->CustomFunctionName = *FString::Printf(TEXT("OnLoaded_%s"), *CompilerContext.GetGuid(this));
 	OnLoadEventNode->AllocateDefaultPins();
 	{
 		UFunction* LoadAssetFunction = CallLoadAssetNode->GetTargetFunction();
-		UDelegateProperty* OnLoadDelegateProperty = LoadAssetFunction ? FindField<UDelegateProperty>(LoadAssetFunction, *DelegateOnLoadedParamName) : nullptr;
+		UDelegateProperty* OnLoadDelegateProperty = LoadAssetFunction ? FindField<UDelegateProperty>(LoadAssetFunction, DelegateOnLoadedParamName) : nullptr;
 		UFunction* OnLoadedSignature = OnLoadDelegateProperty ? OnLoadDelegateProperty->SignatureFunction : nullptr;
 		ensure(OnLoadedSignature);
 		for (TFieldIterator<UProperty> PropIt(OnLoadedSignature); PropIt && (PropIt->PropertyFlags & CPF_Parm); ++PropIt)
@@ -170,7 +170,7 @@ void UK2Node_LoadAsset::ExpandNode(class FKismetCompilerContext& CompilerContext
 			{
 				FEdGraphPinType PinType;
 				bIsErrorFree &= Schema->ConvertPropertyToPinType(Param, /*out*/ PinType);
-				bIsErrorFree &= (NULL != OnLoadEventNode->CreateUserDefinedPin(Param->GetName(), PinType, EGPD_Output));
+				bIsErrorFree &= (nullptr != OnLoadEventNode->CreateUserDefinedPin(Param->GetFName(), PinType, EGPD_Output));
 			}
 		}
 	}
@@ -251,25 +251,25 @@ FText UK2Node_LoadAsset::GetMenuCategory() const
 	return FText(LOCTEXT("UK2Node_LoadAssetGetMenuCategory", "Utilities"));
 }
 
-const FString& UK2Node_LoadAsset::GetInputCategory() const
+const FName& UK2Node_LoadAsset::GetInputCategory() const
 {
 	return UEdGraphSchema_K2::PC_SoftObject;
 }
 
-const FString& UK2Node_LoadAsset::GetOutputCategory() const
+const FName& UK2Node_LoadAsset::GetOutputCategory() const
 {
 	return UEdGraphSchema_K2::PC_Object;
 }
 
-const FString& UK2Node_LoadAsset::GetInputPinName() const
+const FName& UK2Node_LoadAsset::GetInputPinName() const
 {
-	static const FString InputAssetPinName("Asset");
+	static const FName InputAssetPinName("Asset");
 	return InputAssetPinName;
 }
 
-const FString& UK2Node_LoadAsset::GetOutputPinName() const
+const FName& UK2Node_LoadAsset::GetOutputPinName() const
 {
-	static const FString OutputObjectPinName("Object");
+	static const FName OutputObjectPinName("Object");
 	return OutputObjectPinName;
 }
 
@@ -296,25 +296,25 @@ FText UK2Node_LoadAssetClass::GetNodeTitle(ENodeTitleType::Type TitleType) const
 	return FText(LOCTEXT("UK2Node_LoadAssetClassGetNodeTitle", "Async Load Class Asset"));
 }
 
-const FString& UK2Node_LoadAssetClass::GetInputCategory() const
+const FName& UK2Node_LoadAssetClass::GetInputCategory() const
 {
 	return UEdGraphSchema_K2::PC_SoftClass;
 }
 
-const FString& UK2Node_LoadAssetClass::GetOutputCategory() const
+const FName& UK2Node_LoadAssetClass::GetOutputCategory() const
 {
 	return UEdGraphSchema_K2::PC_Class;
 }
 
-const FString& UK2Node_LoadAssetClass::GetInputPinName() const
+const FName& UK2Node_LoadAssetClass::GetInputPinName() const
 {
-	static const FString InputAssetPinName("AssetClass");
+	static const FName InputAssetPinName("AssetClass");
 	return InputAssetPinName;
 }
 
-const FString& UK2Node_LoadAssetClass::GetOutputPinName() const
+const FName& UK2Node_LoadAssetClass::GetOutputPinName() const
 {
-	static const FString OutputObjectPinName("Class");
+	static const FName OutputObjectPinName("Class");
 	return OutputObjectPinName;
 }
 
diff --git a/Engine/Source/Editor/BlueprintGraph/Private/K2Node_LocalVariable.cpp b/Engine/Source/Editor/BlueprintGraph/Private/K2Node_LocalVariable.cpp
--- a/Engine/Source/Editor/BlueprintGraph/Private/K2Node_LocalVariable.cpp
+++ b/Engine/Source/Editor/BlueprintGraph/Private/K2Node_LocalVariable.cpp
@@ -112,8 +112,8 @@ void UDEPRECATED_K2Node_LocalVariable::ReconstructNode()
 		NewVar.VarGuid = FGuid::NewGuid();
 		NewVar.VarType = GetVariablePin()->PinType;
 		NewVar.PropertyFlags |= CPF_BlueprintVisible;
-		NewVar.FriendlyName = FName::NameToDisplayString( NewVar.VarName.ToString(), (NewVar.VarType.PinCategory == K2Schema->PC_Boolean) ? true : false );
-		NewVar.Category = K2Schema->VR_DefaultCategory;
+		NewVar.FriendlyName = FName::NameToDisplayString( NewVar.VarName.ToString(), (NewVar.VarType.PinCategory == UEdGraphSchema_K2::PC_Boolean) ? true : false );
+		NewVar.Category = UEdGraphSchema_K2::VR_DefaultCategory;
 
 		// Find the function entry node in the graph
 		TArray<UK2Node_FunctionEntry*> FunctionEntryNodes;
@@ -127,9 +127,9 @@ void UDEPRECATED_K2Node_LocalVariable::ReconstructNode()
 		TArray< UEdGraphPin* > VariableLinkedPins = GetVariablePin()->LinkedTo;
 
 		// Search through all linked pins for AssignmentStatement nodes, those need to be replaced with a VariableSet node and hooked up to the new variable
-		for( auto LinkedPin : VariableLinkedPins )
+		for (UEdGraphPin* LinkedPin : VariableLinkedPins)
 		{
-			if(LinkedPin->GetOwningNode()->IsA(UK2Node_AssignmentStatement::StaticClass()))
+			if (LinkedPin->GetOwningNode()->IsA<UK2Node_AssignmentStatement>())
 			{
 				UK2Node_AssignmentStatement* AssignementNode = Cast<UK2Node_AssignmentStatement>(LinkedPin->GetOwningNode());
 
@@ -143,14 +143,14 @@ void UDEPRECATED_K2Node_LocalVariable::ReconstructNode()
 					SetNode->PostPlacedNewNode();
 
 					// Locally re-construct the pins, can't allow the node to do it because the property does not yet exist
-					UEdGraphPin* ExecPin = SetNode->CreatePin(EGPD_Input, K2Schema->PC_Exec, FString(), nullptr, K2Schema->PN_Execute);
-					UEdGraphPin* ThenPin = SetNode->CreatePin(EGPD_Output, K2Schema->PC_Exec, FString(), nullptr, K2Schema->PN_Then);
-					UEdGraphPin* ValuePin = SetNode->CreatePin(EGPD_Input, FString(), FString(), nullptr, NewVar.VarName.ToString());
+					UEdGraphPin* ExecPin = SetNode->CreatePin(EGPD_Input, UEdGraphSchema_K2::PC_Exec, UEdGraphSchema_K2::PN_Execute);
+					UEdGraphPin* ThenPin = SetNode->CreatePin(EGPD_Output, UEdGraphSchema_K2::PC_Exec, UEdGraphSchema_K2::PN_Then);
+					UEdGraphPin* ValuePin = SetNode->CreatePin(EGPD_Input, NAME_None, NewVar.VarName);
 					ValuePin->PinType = NewVar.VarType;
 
 					// Move the pin links over to the new node
-					K2Schema->MovePinLinks(*AssignementNode->FindPin(K2Schema->PN_Execute), *ExecPin);
-					K2Schema->MovePinLinks(*AssignementNode->FindPin(K2Schema->PN_Then), *ThenPin);
+					K2Schema->MovePinLinks(*AssignementNode->FindPin(UEdGraphSchema_K2::PN_Execute), *ExecPin);
+					K2Schema->MovePinLinks(*AssignementNode->FindPin(UEdGraphSchema_K2::PN_Then), *ThenPin);
 					K2Schema->MovePinLinks(*AssignementNode->GetValuePin(), *ValuePin);
 
 					// Position the new node
@@ -174,7 +174,7 @@ void UDEPRECATED_K2Node_LocalVariable::ReconstructNode()
 			GetNode->PostPlacedNewNode();
 
 			// Locally re-construct the pins, can't allow the node to do it because the property does not yet exist
-			UEdGraphPin* VariablePin = GetNode->CreatePin(EGPD_Output, FString(), FString(), nullptr, NewVar.VarName.ToString());
+			UEdGraphPin* VariablePin = GetNode->CreatePin(EGPD_Output, NAME_None, NewVar.VarName);
 			VariablePin->PinType = NewVar.VarType;
 			K2Schema->SetPinAutogeneratedDefaultValueBasedOnType(VariablePin);
 
diff --git a/Engine/Source/Editor/BlueprintGraph/Private/K2Node_MCDelegate.cpp b/Engine/Source/Editor/BlueprintGraph/Private/K2Node_MCDelegate.cpp
--- a/Engine/Source/Editor/BlueprintGraph/Private/K2Node_MCDelegate.cpp
+++ b/Engine/Source/Editor/BlueprintGraph/Private/K2Node_MCDelegate.cpp
@@ -23,9 +23,9 @@
 
 struct FK2Node_BaseMCDelegateHelper
 {
-	static FString DelegatePinName;
+	static FName DelegatePinName;
 };
-FString FK2Node_BaseMCDelegateHelper::DelegatePinName(TEXT("Delegate"));
+FName FK2Node_BaseMCDelegateHelper::DelegatePinName(TEXT("Delegate"));
 
 /////// UK2Node_BaseMCDelegate ///////////
 
@@ -58,12 +58,11 @@ bool UK2Node_BaseMCDelegate::IsCompatibleWithGraph(const UEdGraph* TargetGraph)
 UK2Node::ERedirectType UK2Node_BaseMCDelegate::DoPinsMatchForReconstruction(const UEdGraphPin* NewPin, int32 NewPinIndex, const UEdGraphPin* OldPin, int32 OldPinIndex) const
 {
 	ERedirectType OrginalResult = Super::DoPinsMatchForReconstruction(NewPin, NewPinIndex, OldPin, OldPinIndex);
-	const UEdGraphSchema_K2* K2Schema = GetDefault<UEdGraphSchema_K2>();
-	if ((ERedirectType::ERedirectType_None == OrginalResult) && K2Schema && NewPin && OldPin)
+	if ((ERedirectType::ERedirectType_None == OrginalResult) && NewPin && OldPin)
 	{
-		if ((OldPin->PinType.PinCategory == K2Schema->PC_Delegate) &&
-			(NewPin->PinType.PinCategory == K2Schema->PC_Delegate) &&
-			(FCString::Stricmp(*(NewPin->PinName), *(OldPin->PinName)) == 0))
+		if ((OldPin->PinType.PinCategory == UEdGraphSchema_K2::PC_Delegate) &&
+			(NewPin->PinType.PinCategory == UEdGraphSchema_K2::PC_Delegate) &&
+			(NewPin->PinName == OldPin->PinName))
 		{
 			return ERedirectType_Name;
 		}
@@ -75,28 +74,26 @@ void UK2Node_BaseMCDelegate::AllocateDefaultPins()
 {
 	Super::AllocateDefaultPins();
 
-	const UEdGraphSchema_K2* K2Schema = GetDefault<UEdGraphSchema_K2>();
-
-	CreatePin(EGPD_Input, K2Schema->PC_Exec, FString(), nullptr, K2Schema->PN_Execute);
-	CreatePin(EGPD_Output, K2Schema->PC_Exec, FString(), nullptr, K2Schema->PN_Then);
+	CreatePin(EGPD_Input, UEdGraphSchema_K2::PC_Exec, UEdGraphSchema_K2::PN_Execute);
+	CreatePin(EGPD_Output, UEdGraphSchema_K2::PC_Exec, UEdGraphSchema_K2::PN_Then);
 
 	UClass* PropertyOwnerClass = DelegateReference.GetMemberParentClass(GetBlueprintClassFromNode());
 	if (PropertyOwnerClass != nullptr)
 	{
 		PropertyOwnerClass = PropertyOwnerClass->GetAuthoritativeClass();
 	}
-	const auto Blueprint = GetBlueprint();
+	const UBlueprint* Blueprint = GetBlueprint();
 	
 	const bool bUseSelf = Blueprint && (PropertyOwnerClass == Blueprint->GeneratedClass);
 
 	UEdGraphPin* SelfPin = NULL;
 	if (bUseSelf)
 	{
-		SelfPin = CreatePin(EGPD_Input, K2Schema->PC_Object, K2Schema->PSC_Self, nullptr, K2Schema->PN_Self);
+		SelfPin = CreatePin(EGPD_Input, UEdGraphSchema_K2::PC_Object, UEdGraphSchema_K2::PSC_Self, UEdGraphSchema_K2::PN_Self);
 	}
 	else
 	{
-		SelfPin = CreatePin(EGPD_Input, K2Schema->PC_Object, FString(), PropertyOwnerClass, K2Schema->PN_Self);
+		SelfPin = CreatePin(EGPD_Input, UEdGraphSchema_K2::PC_Object, PropertyOwnerClass, UEdGraphSchema_K2::PN_Self);
 	}
 
 	if(SelfPin)
@@ -256,8 +253,8 @@ void UK2Node_BaseMCDelegate::AutowireNewNode(UEdGraphPin* FromPin)
 					{
 						// If we get here, then the property delegate is also available on the FromPin's class.
 						// Fix up the type by propagating it from the FromPin to our target pin
-						UEdGraphPin* TargetPin = FindPin(K2Schema->PN_Self);
-						TargetPin->PinType.PinSubCategory.Empty();
+						UEdGraphPin* TargetPin = FindPin(UEdGraphSchema_K2::PN_Self);
+						TargetPin->PinType.PinSubCategory = NAME_None;
 						TargetPin->PinType.PinSubCategoryObject = DelegateOwner;
 
 						// And finally, try to establish the connection
@@ -294,9 +291,10 @@ void UK2Node_AddDelegate::AllocateDefaultPins()
 {
 	Super::AllocateDefaultPins();
 
-	const UEdGraphSchema_K2* K2Schema = GetDefault<UEdGraphSchema_K2>();
-
-	if(UEdGraphPin* DelegatePin = CreatePin(EGPD_Input, K2Schema->PC_Delegate, FString(), nullptr, FK2Node_BaseMCDelegateHelper::DelegatePinName, EPinContainerType::None, true, true))
+	UEdGraphNode::FCreatePinParams PinParams;
+	PinParams.bIsReference = true;
+	PinParams.bIsConst = true;
+	if(UEdGraphPin* DelegatePin = CreatePin(EGPD_Input, UEdGraphSchema_K2::PC_Delegate, FK2Node_BaseMCDelegateHelper::DelegatePinName, PinParams))
 	{
 		FMemberReference::FillSimpleMemberReference<UFunction>(GetDelegateSignature(), DelegatePin->PinType.PinSubCategoryMemberReference);
 		DelegatePin->PinFriendlyName = NSLOCTEXT("K2Node", "PinFriendlyDelegatetName", "Event");
@@ -362,9 +360,10 @@ void UK2Node_RemoveDelegate::AllocateDefaultPins()
 {
 	Super::AllocateDefaultPins();
 
-	const UEdGraphSchema_K2* K2Schema = GetDefault<UEdGraphSchema_K2>();
-
-	if(UEdGraphPin* DelegatePin = CreatePin(EGPD_Input, K2Schema->PC_Delegate, FString(), nullptr, FK2Node_BaseMCDelegateHelper::DelegatePinName, EPinContainerType::None, true, true))
+	UEdGraphNode::FCreatePinParams PinParams;
+	PinParams.bIsReference = true;
+	PinParams.bIsConst = true;
+	if(UEdGraphPin* DelegatePin = CreatePin(EGPD_Input, UEdGraphSchema_K2::PC_Delegate, FK2Node_BaseMCDelegateHelper::DelegatePinName, PinParams))
 	{
 		FMemberReference::FillSimpleMemberReference<UFunction>(GetDelegateSignature(), DelegatePin->PinType.PinSubCategoryMemberReference);
 		DelegatePin->PinFriendlyName = NSLOCTEXT("K2Node", "PinFriendlyDelegatetName", "Event");
@@ -407,7 +406,7 @@ bool UK2Node_CallDelegate::CreatePinsForFunctionInputs(const UFunction* Function
 		const bool bIsFunctionInput = !Param->HasAnyPropertyFlags(CPF_OutParm) || Param->HasAnyPropertyFlags(CPF_ReferenceParm);
 		if (bIsFunctionInput)
 		{
-			UEdGraphPin* Pin = CreatePin(EGPD_Input, FString(), FString(), nullptr, Param->GetName());
+			UEdGraphPin* Pin = CreatePin(EGPD_Input, NAME_None, Param->GetFName());
 			const bool bPinGood = K2Schema->ConvertPropertyToPinType(Param, /*out*/ Pin->PinType);
 
 			bAllPinsGood = bAllPinsGood && bPinGood;
diff --git a/Engine/Source/Editor/BlueprintGraph/Private/K2Node_MacroInstance.cpp b/Engine/Source/Editor/BlueprintGraph/Private/K2Node_MacroInstance.cpp
--- a/Engine/Source/Editor/BlueprintGraph/Private/K2Node_MacroInstance.cpp
+++ b/Engine/Source/Editor/BlueprintGraph/Private/K2Node_MacroInstance.cpp
@@ -250,13 +250,11 @@ void UK2Node_MacroInstance::NotifyPinConnectionListChanged(UEdGraphPin* ChangedP
 {
 	Super::NotifyPinConnectionListChanged(ChangedPin);
 
-	const UEdGraphSchema_K2* const Schema = GetDefault<UEdGraphSchema_K2>();
-
 	// added a link?
 	if (ChangedPin->LinkedTo.Num() > 0)
 	{
 		// ... to a wildcard pin?
-		bool const bIsWildcardPin = ChangedPin->PinType.PinCategory == Schema->PC_Wildcard;
+		bool const bIsWildcardPin = ChangedPin->PinType.PinCategory == UEdGraphSchema_K2::PC_Wildcard;
 		if (bIsWildcardPin)
 		{
 			// get type of pin we just got linked to
@@ -270,7 +268,7 @@ void UK2Node_MacroInstance::NotifyPinConnectionListChanged(UEdGraphPin* ChangedP
 				UEdGraphPin* const TmpPin = Pins[PinIdx];
 				if (TmpPin)
 				{
-					if (TmpPin->PinType.PinCategory == Schema->PC_Wildcard)
+					if (TmpPin->PinType.PinCategory == UEdGraphSchema_K2::PC_Wildcard)
 					{
 						// only copy the category stuff to preserve array and ref status
 						TmpPin->PinType.PinCategory = LinkedPinType.PinCategory;
diff --git a/Engine/Source/Editor/BlueprintGraph/Private/K2Node_MakeArray.cpp b/Engine/Source/Editor/BlueprintGraph/Private/K2Node_MakeArray.cpp
--- a/Engine/Source/Editor/BlueprintGraph/Private/K2Node_MakeArray.cpp
+++ b/Engine/Source/Editor/BlueprintGraph/Private/K2Node_MakeArray.cpp
@@ -20,7 +20,7 @@
 
 namespace MakeArrayLiterals
 {
-	static const FString OutputPinName = FString(TEXT("Array"));
+	static const FName OutputPinName(TEXT("Array"));
 };
 
 #define LOCTEXT_NAMESPACE "MakeArrayNode"
@@ -56,7 +56,7 @@ FText UK2Node_MakeArray::GetNodeTitle(ENodeTitleType::Type TitleType) const
 	return LOCTEXT("NodeTitle", "Make Array");
 }
 
-FString UK2Node_MakeArray::GetOutputPinName() const
+FName UK2Node_MakeArray::GetOutputPinName() const
 {
 	return MakeArrayLiterals::OutputPinName;
 }
diff --git a/Engine/Source/Editor/BlueprintGraph/Private/K2Node_MakeContainer.cpp b/Engine/Source/Editor/BlueprintGraph/Private/K2Node_MakeContainer.cpp
--- a/Engine/Source/Editor/BlueprintGraph/Private/K2Node_MakeContainer.cpp
+++ b/Engine/Source/Editor/BlueprintGraph/Private/K2Node_MakeContainer.cpp
@@ -101,12 +101,14 @@ void UK2Node_MakeContainer::ReallocatePinsDuringReconstruction(TArray<UEdGraphPi
 void UK2Node_MakeContainer::AllocateDefaultPins()
 {
 	// Create the output pin
-	CreatePin(EGPD_Output, UEdGraphSchema_K2::PC_Wildcard, FString(), nullptr, *GetOutputPinName(), ContainerType);
+	UEdGraphNode::FCreatePinParams PinParams;
+	PinParams.ContainerType = ContainerType;
+	CreatePin(EGPD_Output, UEdGraphSchema_K2::PC_Wildcard, GetOutputPinName(), PinParams);
 
 	// Create the input pins to create the container from
 	for (int32 i = 0; i < NumInputs; ++i)
 	{
-		CreatePin(EGPD_Input, UEdGraphSchema_K2::PC_Wildcard, FString(), nullptr, *FString::Printf(TEXT("[%d]"), i));
+		CreatePin(EGPD_Input, UEdGraphSchema_K2::PC_Wildcard, GetPinName(i));
 	}
 }
 
@@ -145,13 +147,13 @@ void UK2Node_MakeContainer::ClearPinTypeToWildcard()
 	{
 		UEdGraphPin* OutputPin = GetOutputPin();
 		OutputPin->PinType.PinCategory = UEdGraphSchema_K2::PC_Wildcard;
-		OutputPin->PinType.PinSubCategory.Reset();
+		OutputPin->PinType.PinSubCategory = NAME_None;
 		OutputPin->PinType.PinSubCategoryObject = nullptr;
 
 		if (ContainerType == EPinContainerType::Map)
 		{
 			OutputPin->PinType.PinValueType.TerminalCategory = UEdGraphSchema_K2::PC_Wildcard;
-			OutputPin->PinType.PinValueType.TerminalSubCategory.Reset();
+			OutputPin->PinType.PinValueType.TerminalSubCategory = NAME_None;
 			OutputPin->PinType.PinValueType.TerminalSubCategoryObject = nullptr;
 		}
 
@@ -255,7 +257,7 @@ void UK2Node_MakeContainer::NotifyPinConnectionListChanged(UEdGraphPin* Pin)
 		if (bResetOutputPinPrimary)
 		{
 			OutputPin->PinType.PinCategory = UEdGraphSchema_K2::PC_Wildcard;
-			OutputPin->PinType.PinSubCategory.Reset();
+			OutputPin->PinType.PinSubCategory = NAME_None;
 			OutputPin->PinType.PinSubCategoryObject = nullptr;
 
 			bNotifyGraphChanged = true;
@@ -263,7 +265,7 @@ void UK2Node_MakeContainer::NotifyPinConnectionListChanged(UEdGraphPin* Pin)
 		if (bResetOutputPinSecondary && ContainerType == EPinContainerType::Map)
 		{
 			OutputPin->PinType.PinValueType.TerminalCategory = UEdGraphSchema_K2::PC_Wildcard;
-			OutputPin->PinType.PinValueType.TerminalSubCategory.Reset();
+			OutputPin->PinType.PinValueType.TerminalSubCategory = NAME_None;
 			OutputPin->PinType.PinValueType.TerminalSubCategoryObject = nullptr;
 
 			bNotifyGraphChanged = true;
@@ -453,7 +455,7 @@ void UK2Node_MakeContainer::PostReconstructNode()
 	else
 	{
 		OutputPin->PinType.PinCategory = UEdGraphSchema_K2::PC_Wildcard;
-		OutputPin->PinType.PinSubCategory.Reset();
+		OutputPin->PinType.PinSubCategory = NAME_None;
 		OutputPin->PinType.PinSubCategoryObject = nullptr;
 	}
 
@@ -466,7 +468,7 @@ void UK2Node_MakeContainer::PostReconstructNode()
 		else
 		{
 			OutputPin->PinType.PinValueType.TerminalCategory = UEdGraphSchema_K2::PC_Wildcard;
-			OutputPin->PinType.PinValueType.TerminalSubCategory.Reset();
+			OutputPin->PinType.PinValueType.TerminalSubCategory = NAME_None;
 			OutputPin->PinType.PinValueType.TerminalSubCategoryObject = nullptr;
 		}
 	}
@@ -489,7 +491,7 @@ void UK2Node_MakeContainer::AddInputPin()
 
 	++NumInputs;
 	const FEdGraphPinType& OutputPinType = GetOutputPin()->PinType;
-	UEdGraphPin* Pin = CreatePin(EGPD_Input, OutputPinType.PinCategory, OutputPinType.PinSubCategory, OutputPinType.PinSubCategoryObject.Get(), *GetPinName(NumInputs-1));
+	UEdGraphPin* Pin = CreatePin(EGPD_Input, OutputPinType.PinCategory, OutputPinType.PinSubCategory, OutputPinType.PinSubCategoryObject.Get(), GetPinName(NumInputs-1));
 	GetDefault<UEdGraphSchema_K2>()->SetPinAutogeneratedDefaultValueBasedOnType(Pin);
 
 	const bool bIsCompiling = GetBlueprint()->bBeingCompiled;
@@ -499,9 +501,9 @@ void UK2Node_MakeContainer::AddInputPin()
 	}
 }
 
-FString UK2Node_MakeContainer::GetPinName(const int32 PinIndex) const
+FName UK2Node_MakeContainer::GetPinName(const int32 PinIndex) const
 {
-	return FString::Printf(TEXT("[%d]"), PinIndex);
+	return *FString::Printf(TEXT("[%d]"), PinIndex);
 }
 
 void UK2Node_MakeContainer::SyncPinNames()
@@ -513,16 +515,18 @@ void UK2Node_MakeContainer::SyncPinNames()
 		if (CurrentPin->Direction == EGPD_Input &&
 			CurrentPin->ParentPin == nullptr)
 		{
-			const FString OldName = CurrentPin->PinName;
-			const FString ElementName = GetPinName(CurrentNumParentPins++);
+			const FName OldName = CurrentPin->PinName;
+			const FName ElementName = GetPinName(CurrentNumParentPins++);
 
 			CurrentPin->Modify();
 			CurrentPin->PinName = ElementName;
 
 			if (CurrentPin->SubPins.Num() > 0)
 			{
-				FString OldFriendlyName = OldName;
-				FString ElementFriendlyName = ElementName;
+				const FString OldNameStr = OldName.ToString();
+				const FString ElementNameStr = ElementName.ToString();
+				FString OldFriendlyName = OldNameStr;
+				FString ElementFriendlyName = ElementNameStr;
 
 				// SubPin Friendly Name has an extra space in it so we need to account for that
 				OldFriendlyName.InsertAt(1, " ");
@@ -531,10 +535,10 @@ void UK2Node_MakeContainer::SyncPinNames()
 				for (UEdGraphPin* SubPin : CurrentPin->SubPins)
 				{
 					FString SubPinFriendlyName = SubPin->PinFriendlyName.ToString();
-					SubPinFriendlyName = SubPinFriendlyName.Replace(*OldFriendlyName, *ElementFriendlyName);
+					SubPinFriendlyName.ReplaceInline(*OldFriendlyName, *ElementFriendlyName);
 
 					SubPin->Modify();
-					SubPin->PinName = SubPin->PinName.Replace(*OldName, *ElementName);
+					SubPin->PinName = *SubPin->PinName.ToString().Replace(*OldNameStr, *ElementNameStr);
 					SubPin->PinFriendlyName = FText::FromString(SubPinFriendlyName);
 				}
 			}
diff --git a/Engine/Source/Editor/BlueprintGraph/Private/K2Node_MakeMap.cpp b/Engine/Source/Editor/BlueprintGraph/Private/K2Node_MakeMap.cpp
--- a/Engine/Source/Editor/BlueprintGraph/Private/K2Node_MakeMap.cpp
+++ b/Engine/Source/Editor/BlueprintGraph/Private/K2Node_MakeMap.cpp
@@ -20,7 +20,7 @@
 
 namespace MakeMapLiterals
 {
-	static const FString OutputPinName = FString(TEXT("Map"));
+	static const FName OutputPinName(TEXT("Map"));
 };
 
 #define LOCTEXT_NAMESPACE "MakeMapNode"
@@ -54,15 +54,17 @@ FNodeHandlingFunctor* UK2Node_MakeMap::CreateNodeHandler(FKismetCompilerContext&
 void UK2Node_MakeMap::AllocateDefaultPins()
 {
 	// Create the output pin
-	FEdGraphTerminalType TerminalType;
-	TerminalType.TerminalCategory = UEdGraphSchema_K2::PC_Wildcard;
-	CreatePin(EGPD_Output, UEdGraphSchema_K2::PC_Wildcard, FString(), nullptr, *GetOutputPinName(), ContainerType, false, false, INDEX_NONE, TerminalType);
+	UEdGraphNode::FCreatePinParams PinParams;
+	PinParams.ContainerType = ContainerType;
+	PinParams.ValueTerminalType.TerminalCategory = UEdGraphSchema_K2::PC_Wildcard;
+
+	CreatePin(EGPD_Output, UEdGraphSchema_K2::PC_Wildcard, GetOutputPinName(), PinParams);
 
 	// Create the input pins to create the container from
 	for (int32 i = 0; i < NumInputs; ++i)
 	{
-		CreatePin(EGPD_Input, UEdGraphSchema_K2::PC_Wildcard, FString(), nullptr, *GetPinName(i * 2));
-		CreatePin(EGPD_Input, UEdGraphSchema_K2::PC_Wildcard, FString(), nullptr, *GetPinName((i * 2) + 1));
+		CreatePin(EGPD_Input, UEdGraphSchema_K2::PC_Wildcard, GetPinName(i * 2));
+		CreatePin(EGPD_Input, UEdGraphSchema_K2::PC_Wildcard, GetPinName((i * 2) + 1));
 	}
 }
 
@@ -94,11 +96,11 @@ void UK2Node_MakeMap::AddInputPin()
 
 	++NumInputs;
 	const FEdGraphPinType& OutputPinType = GetOutputPin()->PinType;
-	UEdGraphPin* Pin = CreatePin(EGPD_Input, OutputPinType.PinCategory, OutputPinType.PinSubCategory, OutputPinType.PinSubCategoryObject.Get(), *GetPinName((NumInputs - 1) * 2));
+	UEdGraphPin* Pin = CreatePin(EGPD_Input, OutputPinType.PinCategory, OutputPinType.PinSubCategory, OutputPinType.PinSubCategoryObject.Get(), GetPinName((NumInputs - 1) * 2));
 	Schema->SetPinAutogeneratedDefaultValueBasedOnType(Pin);
 
 	const FEdGraphPinType ValuePinType = FEdGraphPinType::GetPinTypeForTerminalType(OutputPinType.PinValueType);
-	Pin = CreatePin(EGPD_Input, ValuePinType.PinCategory, ValuePinType.PinSubCategory, ValuePinType.PinSubCategoryObject.Get(), *GetPinName(((NumInputs - 1) * 2) + 1));
+	Pin = CreatePin(EGPD_Input, ValuePinType.PinCategory, ValuePinType.PinSubCategory, ValuePinType.PinSubCategoryObject.Get(), GetPinName(((NumInputs - 1) * 2) + 1));
 	Schema->SetPinAutogeneratedDefaultValueBasedOnType(Pin);
 
 	const bool bIsCompiling = GetBlueprint()->bBeingCompiled;
@@ -113,20 +115,20 @@ FText UK2Node_MakeMap::GetNodeTitle(ENodeTitleType::Type TitleType) const
 	return LOCTEXT("NodeTitle", "Make Map");
 }
 
-FString UK2Node_MakeMap::GetPinName(const int32 PinIndex) const
+FName UK2Node_MakeMap::GetPinName(const int32 PinIndex) const
 {
 	const int32 PairIndex = PinIndex / 2;
 	if (PinIndex % 2 == 0)
 	{
-		return FString::Printf(TEXT("Key %d"), PairIndex);
+		return *FString::Printf(TEXT("Key %d"), PairIndex);
 	}
 	else
 	{
-		return FString::Printf(TEXT("Value %d"), PairIndex);
+		return *FString::Printf(TEXT("Value %d"), PairIndex);
 	}
 }
 
-FString UK2Node_MakeMap::GetOutputPinName() const
+FName UK2Node_MakeMap::GetOutputPinName() const
 {
 	return MakeMapLiterals::OutputPinName;
 }
diff --git a/Engine/Source/Editor/BlueprintGraph/Private/K2Node_MakeSet.cpp b/Engine/Source/Editor/BlueprintGraph/Private/K2Node_MakeSet.cpp
--- a/Engine/Source/Editor/BlueprintGraph/Private/K2Node_MakeSet.cpp
+++ b/Engine/Source/Editor/BlueprintGraph/Private/K2Node_MakeSet.cpp
@@ -20,7 +20,7 @@
 
 namespace MakeSetLiterals
 {
-	static const FString OutputPinName = FString(TEXT("Set"));
+	static const FName OutputPinName(TEXT("Set"));
 };
 
 #define LOCTEXT_NAMESPACE "MakeSetNode"
@@ -56,7 +56,7 @@ FText UK2Node_MakeSet::GetNodeTitle(ENodeTitleType::Type TitleType) const
 	return LOCTEXT("NodeTitle", "Make Set");
 }
 
-FString UK2Node_MakeSet::GetOutputPinName() const
+FName UK2Node_MakeSet::GetOutputPinName() const
 {
 	return MakeSetLiterals::OutputPinName;
 }
diff --git a/Engine/Source/Editor/BlueprintGraph/Private/K2Node_MakeStruct.cpp b/Engine/Source/Editor/BlueprintGraph/Private/K2Node_MakeStruct.cpp
--- a/Engine/Source/Editor/BlueprintGraph/Private/K2Node_MakeStruct.cpp
+++ b/Engine/Source/Editor/BlueprintGraph/Private/K2Node_MakeStruct.cpp
@@ -47,11 +47,11 @@ void UK2Node_MakeStruct::FMakeStructPinManager::CustomizePinData(UEdGraphPin* Pi
 
 		// Should pin default value be filled as FText?
 		const bool bIsText = Property->IsA<UTextProperty>();
-		checkSlow(bIsText == ((Schema->PC_Text == Pin->PinType.PinCategory) && !Pin->PinType.IsContainer()));
+		checkSlow(bIsText == ((UEdGraphSchema_K2::PC_Text == Pin->PinType.PinCategory) && !Pin->PinType.IsContainer()));
 
 		const bool bIsObject = Property->IsA<UObjectPropertyBase>();
-		checkSlow(bIsObject == ((Schema->PC_Object == Pin->PinType.PinCategory || Schema->PC_Class == Pin->PinType.PinCategory || 
-			Schema->PC_SoftObject == Pin->PinType.PinCategory || Schema->PC_SoftClass == Pin->PinType.PinCategory) && !Pin->PinType.IsContainer()));
+		checkSlow(bIsObject == ((UEdGraphSchema_K2::PC_Object == Pin->PinType.PinCategory || UEdGraphSchema_K2::PC_Class == Pin->PinType.PinCategory || 
+			UEdGraphSchema_K2::PC_SoftObject == Pin->PinType.PinCategory || UEdGraphSchema_K2::PC_SoftClass == Pin->PinType.PinCategory) && !Pin->PinType.IsContainer()));
 
 		if (Property->HasAnyPropertyFlags(CPF_AdvancedDisplay))
 		{
@@ -118,11 +118,10 @@ UK2Node_MakeStruct::UK2Node_MakeStruct(const FObjectInitializer& ObjectInitializ
 
 void UK2Node_MakeStruct::AllocateDefaultPins()
 {
-	const UEdGraphSchema_K2* Schema = GetDefault<UEdGraphSchema_K2>();
-	if(Schema && StructType)
+	if (StructType)
 	{
 		PreloadObject(StructType);
-		CreatePin(EGPD_Output, Schema->PC_Struct, FString(), StructType, StructType->GetName());
+		CreatePin(EGPD_Output, UEdGraphSchema_K2::PC_Struct, StructType, StructType->GetFName());
 		
 		bool bHasAdvancedPins = false;
 		{
@@ -177,7 +176,7 @@ void UK2Node_MakeStruct::ValidateNodeDuringCompilation(class FCompilerResultsLog
 			{
 				if (Property->ArrayDim > 1)
 				{
-					const UEdGraphPin* Pin = FindPin(Property->GetName());
+					const UEdGraphPin* Pin = FindPin(Property->GetFName());
 					MessageLog.Warning(*LOCTEXT("StaticArray_Warning", "@@ - the native property is a static array, which is not supported by blueprints").ToString(), Pin);
 				}
 			}
@@ -235,7 +234,7 @@ FLinearColor UK2Node_MakeStruct::GetNodeTitleColor() const
 	if(const UEdGraphSchema_K2* K2Schema = GetDefault<UEdGraphSchema_K2>())
 	{
 		FEdGraphPinType PinType;
-		PinType.PinCategory = K2Schema->PC_Struct;
+		PinType.PinCategory = UEdGraphSchema_K2::PC_Struct;
 		PinType.PinSubCategoryObject = StructType;
 		return K2Schema->GetPinTypeColor(PinType);
 	}
@@ -286,11 +285,11 @@ UK2Node::ERedirectType UK2Node_MakeStruct::DoPinsMatchForReconstruction(const UE
 		}
 		else if ((EGPD_Input == NewPin->Direction) && (EGPD_Input == OldPin->Direction))
 		{
-			FName RedirectedPinName = UProperty::FindRedirectedPropertyName(StructType, FName(*OldPin->PinName));
+			FName RedirectedPinName = UProperty::FindRedirectedPropertyName(StructType, OldPin->PinName);
 
 			if (RedirectedPinName != NAME_Name)
 			{
-				Result = ((FCString::Stricmp(*RedirectedPinName.ToString(), *NewPin->PinName) != 0) ? ERedirectType_None : ERedirectType_Name);
+				Result = ((RedirectedPinName != NewPin->PinName) ? ERedirectType_None : ERedirectType_Name);
 			}
 		}
 	}
@@ -385,16 +384,15 @@ void UK2Node_MakeStruct::Serialize(FArchive& Ar)
 						}
 
 						bMadeAfterOverridePinRemoval = false;
-						UEdGraphPin* Pin = FindPin(Property->GetName());
+						UEdGraphPin* Pin = FindPin(Property->GetFName());
 
 						if (bHadOverridePropertySeparation)
 						{
-							UEdGraphPin* OverridePin = FindPin(OverrideProperty->GetName());
-							const UEdGraphSchema_K2* Schema = GetDefault<UEdGraphSchema_K2>();
+							UEdGraphPin* OverridePin = FindPin(OverrideProperty->GetFName());
 							if (OverridePin)
 							{
 								// Override pins are always booleans
-								check(OverridePin->PinType.PinCategory == Schema->PC_Boolean);
+								check(OverridePin->PinType.PinCategory == UEdGraphSchema_K2::PC_Boolean);
 								// If the old override pin's default value was true, then the override should be marked as enabled
 								PropertyEntry.bIsOverrideEnabled = OverridePin->DefaultValue.ToBool();
 								// It had an override pin, so conceptually the override pin is visible
@@ -450,7 +448,7 @@ void UK2Node_MakeStruct::ConvertDeprecatedNode(UEdGraph* Graph, bool bOnlySafeCh
 		UFunction* MakeNodeFunction = nullptr;
 
 		// If any pins need to change their names during the conversion, add them to the map.
-		TMap<FString, FString> OldPinToNewPinMap;
+		TMap<FName, FName> OldPinToNewPinMap;
 
 		if (StructType == TBaseStructure<FRotator>::Get())
 		{
diff --git a/Engine/Source/Editor/BlueprintGraph/Private/K2Node_MathExpression.cpp b/Engine/Source/Editor/BlueprintGraph/Private/K2Node_MathExpression.cpp
--- a/Engine/Source/Editor/BlueprintGraph/Private/K2Node_MathExpression.cpp
+++ b/Engine/Source/Editor/BlueprintGraph/Private/K2Node_MathExpression.cpp
@@ -58,11 +58,10 @@ static void DeleteGeneratedNodesInGraph(UEdGraph* Graph)
  */
 static bool PromoteByteToInt(FEdGraphPinType& InOutType)
 {
-	const UEdGraphSchema_K2* Schema = GetDefault<UEdGraphSchema_K2>();
-	if (InOutType.PinCategory == Schema->PC_Byte)
+	if (InOutType.PinCategory == UEdGraphSchema_K2::PC_Byte)
 	{
-		InOutType.PinCategory = Schema->PC_Int;
-		InOutType.PinSubCategoryObject = NULL;
+		InOutType.PinCategory = UEdGraphSchema_K2::PC_Int;
+		InOutType.PinSubCategoryObject = nullptr;
 		return true;
 	}
 	return false;
@@ -77,11 +76,10 @@ static bool PromoteByteToInt(FEdGraphPinType& InOutType)
 */
 static bool PromoteIntToFloat(FEdGraphPinType& InOutType)
 {
-	const UEdGraphSchema_K2* Schema = GetDefault<UEdGraphSchema_K2>();
-	if (InOutType.PinCategory == Schema->PC_Int)
+	if (InOutType.PinCategory == UEdGraphSchema_K2::PC_Int)
 	{
-		InOutType.PinCategory = Schema->PC_Float;
-		InOutType.PinSubCategoryObject = NULL;
+		InOutType.PinCategory = UEdGraphSchema_K2::PC_Float;
+		InOutType.PinSubCategoryObject = nullptr;
 		return true;
 	}
 	return false;
@@ -1138,7 +1136,7 @@ public:
 	virtual bool ConnectToInput(UEdGraphPin* InputPin, FCompilerResultsLog& MessageLog) override
 	{
 		bool bSuccess = false;
-		if (UEdGraphPin* VariablePin = GeneratedNode->FindPin(GeneratedNode->VariableReference.GetMemberName().ToString()))
+		if (UEdGraphPin* VariablePin = GeneratedNode->FindPin(GeneratedNode->VariableReference.GetMemberName()))
 		{
 			bSuccess = SafeConnectPins(VariablePin, InputPin, MessageLog);
 		}
@@ -1220,8 +1218,8 @@ public:
 		}
 		else 
 		{
-			FText ErrorText = FText::Format(LOCTEXT("LiteralNotCompatible", "Literal type ({0}) is incompatible with pin: '@@'"),
-				FText::FromString(GetOutputType().PinCategory));
+			const FText ErrorText = FText::Format(LOCTEXT("LiteralNotCompatible", "Literal type ({0}) is incompatible with pin: '@@'"), 
+				FText::FromName(GetOutputType().PinCategory));
 			MessageLog.Error(*ErrorText.ToString(), InputPin);
 		}
 		return bSuccess;
@@ -1369,47 +1367,44 @@ public:
 			// first we try to match up variables with existing variable properties on the blueprint
 
 			FMemberReference VariableReference;
-			FString VariableName;
+			FName VariableName;
 			FGuid VariableGuid;
 
 			if (ExpressionNode.Token.TokenType == FBasicToken::TOKEN_Guid && FGuid::Parse(VariableIdentifier, VariableGuid))
 			{
 				// First look the variable up as a Member variable
-				FName VariableFName;
-				VariableFName = FBlueprintEditorUtils::FindMemberVariableNameByGuid(TargetBlueprint, VariableGuid);
+				VariableName = FBlueprintEditorUtils::FindMemberVariableNameByGuid(TargetBlueprint, VariableGuid);
 
 				// If the variable was not found, look it up as a local variable
-				if (VariableFName.IsNone())
+				if (VariableName.IsNone())
 				{
-					VariableFName = FBlueprintEditorUtils::FindLocalVariableNameByGuid(TargetBlueprint, VariableGuid);
-					VariableReference.SetLocalMember(VariableFName, CompilingNode->GetGraph()->GetName(), VariableGuid);
+					VariableName = FBlueprintEditorUtils::FindLocalVariableNameByGuid(TargetBlueprint, VariableGuid);
+					VariableReference.SetLocalMember(VariableName, CompilingNode->GetGraph()->GetName(), VariableGuid);
 				}
 				else
 				{
-					VariableReference.SetSelfMember(VariableFName);
+					VariableReference.SetSelfMember(VariableName);
 				}
-
-				VariableName = VariableFName.ToString();
 			}
 			else 
 			{
-				VariableName = VariableIdentifier;
+				VariableName = *VariableIdentifier;
 
 				// First look the variable up as a Member variable
-				VariableGuid = FBlueprintEditorUtils::FindMemberVariableGuidByName(TargetBlueprint, FName(*VariableName));
+				VariableGuid = FBlueprintEditorUtils::FindMemberVariableGuidByName(TargetBlueprint, VariableName);
 
 				// If the variable was not found, look it up as a local variable
 				if (!VariableGuid.IsValid())
 				{
-					VariableGuid = FBlueprintEditorUtils::FindLocalVariableGuidByName(TargetBlueprint, CompilingNode->GetGraph(), FName(*VariableName));
+					VariableGuid = FBlueprintEditorUtils::FindLocalVariableGuidByName(TargetBlueprint, CompilingNode->GetGraph(), VariableName);
 					if (VariableGuid.IsValid())
 					{
-						VariableReference.SetLocalMember(FName(*VariableName), CompilingNode->GetGraph()->GetName(), VariableGuid);
+						VariableReference.SetLocalMember(VariableName, CompilingNode->GetGraph()->GetName(), VariableGuid);
 					}
 				}
 				else
 				{
-					VariableReference.SetSelfMember(FName(*VariableName));
+					VariableReference.SetSelfMember(VariableName);
 				}
 
 				// If we found a valid Guid, change the expression's identifier to be the Guid
@@ -1431,7 +1426,7 @@ public:
 			// if a variable-get couldn't be created for it, it needs to be an input to the math node
 			else if(ExpressionNode.Token.TokenType != FBasicToken::TOKEN_Guid)
 			{
-				CompiledFragments.Add(&ExpressionNode, GenerateInputPinFragment(VariableIdentifier));
+				CompiledFragments.Add(&ExpressionNode, GenerateInputPinFragment(*VariableIdentifier));
 			}
 			
 		}
@@ -1619,7 +1614,7 @@ private:
      * @param  VariableIdentifier   The name of the pin to generate a fragment for.
      * @return A new input pin fragment (associated with a pin on the math expression's entry node).
      */
-	TSharedPtr<FCodeGenFragment_InputPin> GenerateInputPinFragment(FString const VariableIdentifier)
+	TSharedPtr<FCodeGenFragment_InputPin> GenerateInputPinFragment(const FName VariableIdentifier)
 	{
 		TSharedPtr<FCodeGenFragment_InputPin> InputPinFragment;
 		UEdGraphSchema_K2 const* K2Schema = GetDefault<UEdGraphSchema_K2>();
@@ -1636,7 +1631,7 @@ private:
 			// Create an input pin (using the default guessed type)
 			FEdGraphPinType DefaultType;
 			// currently, generated expressions ALWAYS take a float (it is the most versatile type)
-			DefaultType.PinCategory = K2Schema->PC_Float;
+			DefaultType.PinCategory = UEdGraphSchema_K2::PC_Float;
 			
 			UEdGraphPin* NewInputPin = EntryNode->CreateUserDefinedPin(VariableIdentifier, DefaultType, EGPD_Output);
 			InputPinFragment = MakeShareable(new FCodeGenFragment_InputPin(NewInputPin));
@@ -1711,22 +1706,21 @@ private:
 	TSharedPtr<FCodeGenFragment_Literal> GenerateLiteralFragment(FBasicToken const& Token, FCompilerResultsLog& MessageLog)
 	{
 		check(Token.TokenType == FBasicToken::TOKEN_Const);
-		UEdGraphSchema_K2 const* K2Schema = GetDefault<UEdGraphSchema_K2>();
 		
 		FEdGraphPinType LiteralType;
 		switch (Token.ConstantType)
 		{
 			case CPT_Bool:
-				LiteralType.PinCategory = K2Schema->PC_Boolean;
+				LiteralType.PinCategory = UEdGraphSchema_K2::PC_Boolean;
 				break;
 			case CPT_Float:
-				LiteralType.PinCategory = K2Schema->PC_Float;
+				LiteralType.PinCategory = UEdGraphSchema_K2::PC_Float;
 				break;
 			case CPT_Int:
-				LiteralType.PinCategory = K2Schema->PC_Int;
+				LiteralType.PinCategory = UEdGraphSchema_K2::PC_Int;
 				break;
 			case CPT_String:
-				LiteralType.PinCategory = K2Schema->PC_String;
+				LiteralType.PinCategory = UEdGraphSchema_K2::PC_String;
 				break;
 			default:
 				MessageLog.Error(*FText::Format(LOCTEXT("UnhandledLiteralType", "Unknown literal type in expression: '@@'"),
@@ -1944,7 +1938,7 @@ private:
 	 * After the code generation, we want to clear any old pins that 
 	 * weren't reused, so here we track the ones in use.
 	 */
-	TArray<FString> InputPinNames;
+	TArray<FName> InputPinNames;
 };
 
 /*******************************************************************************
diff --git a/Engine/Source/Editor/BlueprintGraph/Private/K2Node_MatineeController.cpp b/Engine/Source/Editor/BlueprintGraph/Private/K2Node_MatineeController.cpp
--- a/Engine/Source/Editor/BlueprintGraph/Private/K2Node_MatineeController.cpp
+++ b/Engine/Source/Editor/BlueprintGraph/Private/K2Node_MatineeController.cpp
@@ -36,8 +36,6 @@ void UK2Node_MatineeController::BeginDestroy()
 
 void UK2Node_MatineeController::AllocateDefaultPins()
 {
-	const UEdGraphSchema_K2* K2Schema = GetDefault<UEdGraphSchema_K2>();
-
 	// Preload the matinee data, if needed, so that we can have all the event tracks we need
 	if (MatineeActor)
 	{
@@ -46,7 +44,7 @@ void UK2Node_MatineeController::AllocateDefaultPins()
 	}
 
 	// Create the "finished" playing pin
-	CreatePin(EGPD_Output, K2Schema->PC_Exec, FString(), nullptr, K2Schema->PN_MatineeFinished);
+	CreatePin(EGPD_Output, UEdGraphSchema_K2::PC_Exec, UEdGraphSchema_K2::PN_MatineeFinished);
 
 	// Create pins for each event
 	if (MatineeActor && MatineeActor->MatineeData)
@@ -54,10 +52,9 @@ void UK2Node_MatineeController::AllocateDefaultPins()
 		TArray<FName> EventNames;
 		MatineeActor->MatineeData->GetAllEventNames(EventNames);
 
-		for(int32 i=0; i<EventNames.Num(); i++)
+		for (const FName EventName : EventNames)
 		{
-			FName EventName = EventNames[i];
-			CreatePin(EGPD_Output, K2Schema->PC_Exec, FString(), nullptr, EventName.ToString());			
+			CreatePin(EGPD_Output, UEdGraphSchema_K2::PC_Exec, EventName);			
 		}
 	}
 
@@ -99,8 +96,7 @@ AActor* UK2Node_MatineeController::GetReferencedLevelActor() const
 
 UEdGraphPin* UK2Node_MatineeController::GetFinishedPin() const
 {
-	const UEdGraphSchema_K2* K2Schema = GetDefault<UEdGraphSchema_K2>();
-	return FindPin(K2Schema->PN_MatineeFinished);
+	return FindPin(UEdGraphSchema_K2::PN_MatineeFinished);
 }
 
 void UK2Node_MatineeController::ExpandNode(FKismetCompilerContext& CompilerContext, UEdGraph* SourceGraph)
@@ -124,9 +120,9 @@ void UK2Node_MatineeController::ExpandNode(FKismetCompilerContext& CompilerConte
 		for(int32 PinIdx=0; PinIdx<Pins.Num(); PinIdx++)
 		{
 			UEdGraphPin* MatineePin = Pins[PinIdx];
-			if(MatineePin->Direction == EGPD_Output && MatineePin->PinType.PinCategory == Schema->PC_Exec)
+			if(MatineePin->Direction == EGPD_Output && MatineePin->PinType.PinCategory == UEdGraphSchema_K2::PC_Exec)
 			{
-				FName EventFuncName = MatineeActor->GetFunctionNameForEvent( FName(*(MatineePin->PinName)) );
+				FName EventFuncName = MatineeActor->GetFunctionNameForEvent(MatineePin->PinName);
 
 				UK2Node_Event* MatineeEventNode = CompilerContext.SpawnIntermediateEventNode<UK2Node_Event>(this, MatineePin, SourceGraph);
 				MatineeEventNode->EventReference.SetFromField<UFunction>(MatineeEventSig, false);
@@ -217,12 +213,12 @@ void UK2Node_MatineeController::OnEventKeyframeAdded(const AMatineeActor* InMati
 	if ( MatineeActor == InMatineeActor )
 	{
 		// Only add unique event names to the controller node
-		if (!FindPin(InPinName.ToString()))
+		if (!FindPin(InPinName))
 		{
-			const UEdGraphSchema_K2* K2Schema = GetDefault<UEdGraphSchema_K2>();
-
 			// Add one to the index as we insert "finished" at index 0
-			CreatePin(EGPD_Output, K2Schema->PC_Exec, FString(), nullptr, InPinName.ToString(), EPinContainerType::None, false, false, InIndex + 1);
+			UEdGraphNode::FCreatePinParams PinParams;
+			PinParams.Index = InIndex + 1;
+			CreatePin(EGPD_Output, UEdGraphSchema_K2::PC_Exec, InPinName, PinParams);
 
 			// Update and refresh the blueprint
 			FBlueprintEditorUtils::MarkBlueprintAsStructurallyModified(GetBlueprint());
@@ -234,10 +230,10 @@ void UK2Node_MatineeController::OnEventKeyframeRenamed(const AMatineeActor* InMa
 {
 	if ( MatineeActor == InMatineeActor )
 	{
-		if( UEdGraphPin* OldPin = FindPin(InOldPinName.ToString()) )
+		if( UEdGraphPin* OldPin = FindPin(InOldPinName) )
 		{
 			OldPin->Modify();
-			OldPin->PinName = InNewPinName.ToString();
+			OldPin->PinName = InNewPinName;
 
 			GetGraph()->NotifyGraphChanged();
 		}
@@ -251,7 +247,7 @@ void UK2Node_MatineeController::OnEventKeyframeRemoved(const AMatineeActor* InMa
 		bool bNeedsRefresh = false;
 		for(int32 PinIdx=0; PinIdx<InPinNames.Num(); PinIdx++)
 		{
-			if(UEdGraphPin* Pin = FindPin(InPinNames[PinIdx].ToString()))
+			if(UEdGraphPin* Pin = FindPin(InPinNames[PinIdx]))
 			{
 				RemovePin(Pin);
 				bNeedsRefresh = true;
diff --git a/Engine/Source/Editor/BlueprintGraph/Private/K2Node_Message.cpp b/Engine/Source/Editor/BlueprintGraph/Private/K2Node_Message.cpp
--- a/Engine/Source/Editor/BlueprintGraph/Private/K2Node_Message.cpp
+++ b/Engine/Source/Editor/BlueprintGraph/Private/K2Node_Message.cpp
@@ -62,9 +62,9 @@ void UK2Node_Message::AllocateDefaultPins()
 	if (MessageNodeFunction && MessageNodeFunction->HasAnyFunctionFlags(FUNC_BlueprintPure))
 	{
 		// Input - Execution Pin
-		CreatePin(EGPD_Input,  UEdGraphSchema_K2::PC_Exec, FString(), nullptr, UEdGraphSchema_K2::PN_Execute);
+		CreatePin(EGPD_Input,  UEdGraphSchema_K2::PC_Exec, UEdGraphSchema_K2::PN_Execute);
 		// Output - Execution Pin
-		CreatePin(EGPD_Output, UEdGraphSchema_K2::PC_Exec, FString(), nullptr, UEdGraphSchema_K2::PN_Then);
+		CreatePin(EGPD_Output, UEdGraphSchema_K2::PC_Exec, UEdGraphSchema_K2::PN_Then);
 	}
 
 	Super::AllocateDefaultPins();
@@ -72,8 +72,7 @@ void UK2Node_Message::AllocateDefaultPins()
 
 UEdGraphPin* UK2Node_Message::CreateSelfPin(const UFunction* Function)
 {
-	const UEdGraphSchema_K2* K2Schema = GetDefault<UEdGraphSchema_K2>();
-	UEdGraphPin* SelfPin = CreatePin(EGPD_Input, K2Schema->PC_Object, FString(), UObject::StaticClass(), K2Schema->PN_Self);
+	UEdGraphPin* SelfPin = CreatePin(EGPD_Input, UEdGraphSchema_K2::PC_Object, UObject::StaticClass(), UEdGraphSchema_K2::PN_Self);
 	SelfPin->bDefaultValueIsIgnored = true;
 	return SelfPin;
 }
@@ -256,7 +255,7 @@ void UK2Node_Message::ExpandNode(class FKismetCompilerContext& CompilerContext,
 		for( int32 i = 0; i < Pins.Num(); i++ )
 		{
 			UEdGraphPin* CurrentPin = Pins[i];
-			if( CurrentPin && (CurrentPin->PinType.PinCategory != Schema->PC_Exec) && (CurrentPin->PinName != Schema->PN_Self) )
+			if( CurrentPin && (CurrentPin->PinType.PinCategory != UEdGraphSchema_K2::PC_Exec) && (CurrentPin->PinName != UEdGraphSchema_K2::PN_Self) )
 			{
 				// Try to find a match for the pin on the function call node
 				UEdGraphPin* FunctionCallPin = FunctionCallNode->FindPin(CurrentPin->PinName);
diff --git a/Engine/Source/Editor/BlueprintGraph/Private/K2Node_MultiGate.cpp b/Engine/Source/Editor/BlueprintGraph/Private/K2Node_MultiGate.cpp
--- a/Engine/Source/Editor/BlueprintGraph/Private/K2Node_MultiGate.cpp
+++ b/Engine/Source/Editor/BlueprintGraph/Private/K2Node_MultiGate.cpp
@@ -55,7 +55,7 @@ public:
 
 		// Create a term to store a bool that determines if we're in the first execution of the node or not
 		FBPTerminal* FirstRunTerm = Context.CreateLocalTerminal();
-		FirstRunTerm->Type.PinCategory = CompilerContext.GetSchema()->PC_Boolean;
+		FirstRunTerm->Type.PinCategory = UEdGraphSchema_K2::PC_Boolean;
 		FirstRunTerm->Source = Node;
 		FirstRunTerm->Name = BaseNetName + TEXT("_FirstRun");
 		FirstRunTermMap.Add(Node, FirstRunTerm);
@@ -65,7 +65,7 @@ public:
 		if (!GateNode || !GateNode->DataNode)
 		{
 			FBPTerminal* DataTerm = Context.CreateLocalTerminal();
-			DataTerm->Type.PinCategory = CompilerContext.GetSchema()->PC_Int;
+			DataTerm->Type.PinCategory = UEdGraphSchema_K2::PC_Int;
 			DataTerm->Source = Node;
 			DataTerm->Name = BaseNetName + TEXT("_Data");
 			DataTermMap.Add(Node, DataTerm);
@@ -77,7 +77,7 @@ public:
 		if (!FuncLocals.GenericBoolTerm)
 		{
 			FuncLocals.GenericBoolTerm = Context.CreateLocalTerminal();
-			FuncLocals.GenericBoolTerm->Type.PinCategory = CompilerContext.GetSchema()->PC_Boolean;
+			FuncLocals.GenericBoolTerm->Type.PinCategory = UEdGraphSchema_K2::PC_Boolean;
 			FuncLocals.GenericBoolTerm->Source = Node;
 			FuncLocals.GenericBoolTerm->Name = BaseNetName + TEXT("_ScratchBool");
 		}
@@ -86,7 +86,7 @@ public:
 		if (!FuncLocals.IndexTerm)
 		{
 			FuncLocals.IndexTerm = Context.CreateLocalTerminal();
-			FuncLocals.IndexTerm->Type.PinCategory = CompilerContext.GetSchema()->PC_Int;
+			FuncLocals.IndexTerm->Type.PinCategory = UEdGraphSchema_K2::PC_Int;
 			FuncLocals.IndexTerm->Source = Node;
 			FuncLocals.IndexTerm->Name = BaseNetName + TEXT("_ScratchIndex");
 		}
@@ -178,12 +178,12 @@ public:
 		// Create a literal pin that represents a -1 value
 		FBPTerminal* InvalidIndexTerm = Context.CreateLocalTerminal(ETerminalSpecification::TS_Literal);
 		InvalidIndexTerm->bIsLiteral = true;
-		InvalidIndexTerm->Type.PinCategory = CompilerContext.GetSchema()->PC_Int;
+		InvalidIndexTerm->Type.PinCategory = UEdGraphSchema_K2::PC_Int;
 		InvalidIndexTerm->Name = TEXT("-1");
 
 		// Create a literal pin that represents a true value
 		FBPTerminal* TrueBoolTerm = Context.CreateLocalTerminal(ETerminalSpecification::TS_Literal);
-		TrueBoolTerm->Type.PinCategory = CompilerContext.GetSchema()->PC_Boolean;
+		TrueBoolTerm->Type.PinCategory = UEdGraphSchema_K2::PC_Boolean;
 		TrueBoolTerm->bIsLiteral = true;
 		TrueBoolTerm->Name = TEXT("true");
 
@@ -191,7 +191,7 @@ public:
 		TArray<UEdGraphPin*> OutPins;
 		GateNode->GetOutPins(OutPins);
 		FBPTerminal* NumOutsTerm = Context.CreateLocalTerminal(ETerminalSpecification::TS_Literal);
-		NumOutsTerm->Type.PinCategory = CompilerContext.GetSchema()->PC_Int;
+		NumOutsTerm->Type.PinCategory = UEdGraphSchema_K2::PC_Int;
 		NumOutsTerm->bIsLiteral = true;
 		NumOutsTerm->Name = FString::Printf(TEXT("%d"), OutPins.Num());
 
@@ -416,7 +416,7 @@ public:
 			// LiteralIndexTerm will be the right side of the == statemnt
 			FBPTerminal* LiteralIndexTerm = Context.CreateLocalTerminal(ETerminalSpecification::TS_Literal);
 			LiteralIndexTerm->bIsLiteral = true;
-			LiteralIndexTerm->Type.PinCategory = CompilerContext.GetSchema()->PC_Int;
+			LiteralIndexTerm->Type.PinCategory = UEdGraphSchema_K2::PC_Int;
 			LiteralIndexTerm->Name = FString::Printf(TEXT("%d"), OutIdx);
 			// Set the params
 			IndexEqualityStatement.LHS = FuncLocals.GenericBoolTerm;
@@ -465,7 +465,7 @@ public:
 		// Create a local int for use in the equality call function below (LiteralTerm = the right hand side of the EqualEqual_IntInt or NotEqual_BoolBool statement)
 		FBPTerminal* LiteralStringTerm = Context.CreateLocalTerminal(ETerminalSpecification::TS_Literal);
 		LiteralStringTerm->bIsLiteral = true;
-		LiteralStringTerm->Type.PinCategory = CompilerContext.GetSchema()->PC_String;
+		LiteralStringTerm->Type.PinCategory = UEdGraphSchema_K2::PC_String;
 		LiteralStringTerm->Name = FString::Printf(*LOCTEXT("MultiGateNode IndexWarning", "MultiGate Node failed! Out of bounds indexing of the out pins. There are only %d outs available.").ToString(), OutPins.Num());
 		PrintStatement.RHS.Add(LiteralStringTerm);
 		// Hook the IfNot statement's jump target to this statement
@@ -498,10 +498,10 @@ void UK2Node_MultiGate::AllocateDefaultPins()
 	Super::AllocateDefaultPins();
 
 	const UEdGraphSchema_K2* K2Schema = GetDefault<UEdGraphSchema_K2>();
-	CreatePin(EGPD_Input, K2Schema->PC_Exec, FString(), nullptr, TEXT("Reset"));
-	CreatePin(EGPD_Input, K2Schema->PC_Boolean, FString(), nullptr, TEXT("IsRandom"));
-	CreatePin(EGPD_Input, K2Schema->PC_Boolean, FString(), nullptr, TEXT("Loop"));
-	UEdGraphPin* IndexPin = CreatePin(EGPD_Input, K2Schema->PC_Int, FString(), nullptr, TEXT("StartIndex"));
+	CreatePin(EGPD_Input, UEdGraphSchema_K2::PC_Exec, TEXT("Reset"));
+	CreatePin(EGPD_Input, UEdGraphSchema_K2::PC_Boolean, TEXT("IsRandom"));
+	CreatePin(EGPD_Input, UEdGraphSchema_K2::PC_Boolean, TEXT("Loop"));
+	UEdGraphPin* IndexPin = CreatePin(EGPD_Input, UEdGraphSchema_K2::PC_Int, TEXT("StartIndex"));
 	K2Schema->SetPinAutogeneratedDefaultValue(IndexPin, TEXT("-1"));
 }
 
@@ -510,10 +510,10 @@ void UK2Node_MultiGate::ReallocatePinsDuringReconstruction(TArray<UEdGraphPin*>&
 	Super::ReallocatePinsDuringReconstruction(OldPins);
 
 	const UEdGraphSchema_K2* K2Schema = GetDefault<UEdGraphSchema_K2>();
-	CreatePin(EGPD_Input, K2Schema->PC_Exec, FString(), nullptr, TEXT("Reset"));
-	CreatePin(EGPD_Input, K2Schema->PC_Boolean, FString(), nullptr, TEXT("IsRandom"));
-	CreatePin(EGPD_Input, K2Schema->PC_Boolean, FString(), nullptr, TEXT("Loop"));
-	UEdGraphPin* IndexPin = CreatePin(EGPD_Input, K2Schema->PC_Int, FString(), nullptr, TEXT("StartIndex"));
+	CreatePin(EGPD_Input, UEdGraphSchema_K2::PC_Exec, TEXT("Reset"));
+	CreatePin(EGPD_Input, UEdGraphSchema_K2::PC_Boolean, TEXT("IsRandom"));
+	CreatePin(EGPD_Input, UEdGraphSchema_K2::PC_Boolean, TEXT("Loop"));
+	UEdGraphPin* IndexPin = CreatePin(EGPD_Input, UEdGraphSchema_K2::PC_Int, TEXT("StartIndex"));
 	K2Schema->SetPinAutogeneratedDefaultValue(IndexPin, TEXT("-1"));
 }
 
@@ -547,14 +547,12 @@ UEdGraphPin* UK2Node_MultiGate::GetStartIndexPin() const
 
 void UK2Node_MultiGate::GetOutPins(TArray<UEdGraphPin*>& OutPins) const
 {
-	const UEdGraphSchema_K2* K2Schema = GetDefault<UEdGraphSchema_K2>();
-
-	OutPins.Empty();
+	OutPins.Reset();
 
-	for (auto It = Pins.CreateConstIterator(); It; It++)
+	const TCHAR* OutStr(TEXT("Out"));
+	for (UEdGraphPin* Pin : Pins)
 	{
-		UEdGraphPin* Pin = (*It);
-		if (Pin->PinName.Left(3) == "Out")
+		if (Pin->PinName.ToString().StartsWith(OutStr))
 		{
 			OutPins.Add(Pin);
 		}
@@ -612,9 +610,9 @@ void UK2Node_MultiGate::GetClearAllBitsFunction(FName& FunctionName, UClass** Fu
 	*FunctionClass = UKismetNodeHelperLibrary::StaticClass();
 }
 
-FString UK2Node_MultiGate::GetPinNameGivenIndex(int32 Index) const
+FName UK2Node_MultiGate::GetPinNameGivenIndex(int32 Index) const
 {
-	return FString::Printf(TEXT("Out %d"), Index);
+	return *FString::Printf(TEXT("Out %d"), Index);
 }
 
 FNodeHandlingFunctor* UK2Node_MultiGate::CreateNodeHandler(FKismetCompilerContext& CompilerContext) const
@@ -634,15 +632,13 @@ void UK2Node_MultiGate::ExpandNode(class FKismetCompilerContext& CompilerContext
 	UEdGraphPin* ResetPin = GetResetPin();
 	if (ResetPin->LinkedTo.Num() > 0)
 	{
-		const UEdGraphSchema_K2* Schema = CompilerContext.GetSchema();
-
 		/////////////////////////////
 		// Temporary Variable node
 		/////////////////////////////
 
 		// Create the node
 		UK2Node_TemporaryVariable* TempVarNode = SourceGraph->CreateIntermediateNode<UK2Node_TemporaryVariable>();
-		TempVarNode->VariableType.PinCategory = Schema->PC_Int;
+		TempVarNode->VariableType.PinCategory = UEdGraphSchema_K2::PC_Int;
 		TempVarNode->AllocateDefaultPins();
 		CompilerContext.MessageLog.NotifyIntermediateObjectCreation(TempVarNode, this);
 		// Give a reference of the variable node to the multi gate node
diff --git a/Engine/Source/Editor/BlueprintGraph/Private/K2Node_PureAssignmentStatement.cpp b/Engine/Source/Editor/BlueprintGraph/Private/K2Node_PureAssignmentStatement.cpp
--- a/Engine/Source/Editor/BlueprintGraph/Private/K2Node_PureAssignmentStatement.cpp
+++ b/Engine/Source/Editor/BlueprintGraph/Private/K2Node_PureAssignmentStatement.cpp
@@ -65,9 +65,9 @@ public:
 
 // The node is (was originally) needed in expansion step for handling CreateAutoRefTerm parameters in pure functions 
 
-FString UK2Node_PureAssignmentStatement::VariablePinName = FString(TEXT("Variable"));
-FString UK2Node_PureAssignmentStatement::ValuePinName = FString(TEXT("Value"));
-FString UK2Node_PureAssignmentStatement::OutputPinName = FString(TEXT("ReturnValue"));
+FName UK2Node_PureAssignmentStatement::VariablePinName = TEXT("Variable");
+FName UK2Node_PureAssignmentStatement::ValuePinName = TEXT("Value");
+FName UK2Node_PureAssignmentStatement::OutputPinName = TEXT("ReturnValue");
 
 
 UK2Node_PureAssignmentStatement::UK2Node_PureAssignmentStatement(const FObjectInitializer& ObjectInitializer)
@@ -77,11 +77,9 @@ UK2Node_PureAssignmentStatement::UK2Node_PureAssignmentStatement(const FObjectIn
 
 void UK2Node_PureAssignmentStatement::AllocateDefaultPins()
 {
-	const UEdGraphSchema_K2* Schema = GetDefault<UEdGraphSchema_K2>();
-
-	CreatePin(EGPD_Input, Schema->PC_Wildcard, FString(), nullptr, VariablePinName);
-	CreatePin(EGPD_Input, Schema->PC_Wildcard, FString(), nullptr, ValuePinName);
-	CreatePin(EGPD_Output, Schema->PC_Wildcard, FString(), nullptr, OutputPinName);
+	CreatePin(EGPD_Input, UEdGraphSchema_K2::PC_Wildcard, VariablePinName);
+	CreatePin(EGPD_Input, UEdGraphSchema_K2::PC_Wildcard, ValuePinName);
+	CreatePin(EGPD_Output, UEdGraphSchema_K2::PC_Wildcard, OutputPinName);
 
 	Super::AllocateDefaultPins();
 }
@@ -92,17 +90,17 @@ void UK2Node_PureAssignmentStatement::NotifyPinConnectionListChanged(UEdGraphPin
 
 	if (Pin->LinkedTo.Num() > 0)
 	{
-		const auto PinType = Pin->LinkedTo[0]->PinType;
+		const FEdGraphPinType PinType = Pin->LinkedTo[0]->PinType;
 
-		auto VariablePin = GetVariablePin();
+		UEdGraphPin* VariablePin = GetVariablePin();
 		VariablePin->PinType = PinType;
 		UEdGraphSchema_K2::ValidateExistingConnections(VariablePin);
 
-		auto OutputPin = GetOutputPin();
+		UEdGraphPin* OutputPin = GetOutputPin();
 		OutputPin->PinType = PinType;
 		UEdGraphSchema_K2::ValidateExistingConnections(OutputPin);
 
-		auto ValuePin = GetValuePin();
+		UEdGraphPin* ValuePin = GetValuePin();
 		ValuePin->PinType = PinType;
 		ValuePin->PinType.bIsReference = false;
 		UEdGraphSchema_K2::ValidateExistingConnections(ValuePin);
diff --git a/Engine/Source/Editor/BlueprintGraph/Private/K2Node_Select.cpp b/Engine/Source/Editor/BlueprintGraph/Private/K2Node_Select.cpp
--- a/Engine/Source/Editor/BlueprintGraph/Private/K2Node_Select.cpp
+++ b/Engine/Source/Editor/BlueprintGraph/Private/K2Node_Select.cpp
@@ -117,9 +117,9 @@ public:
 				LiteralTerm->Type = IndexTerm->Type;
 				LiteralTerm->bIsLiteral = true;
 				const UEnum* NodeEnum = SelectNode->GetEnum();
-				LiteralTerm->Name = NodeEnum ? OptionPins[OptionIdx]->PinName : FString::Printf(TEXT("%d"), OptionIdx); //-V595
+				LiteralTerm->Name = NodeEnum ? OptionPins[OptionIdx]->PinName.ToString() : FString::Printf(TEXT("%d"), OptionIdx); //-V595
 
-				if (!CompilerContext.GetSchema()->DefaultValueSimpleValidation(LiteralTerm->Type, LiteralTerm->Name, LiteralTerm->Name, nullptr, FText()))
+				if (!CompilerContext.GetSchema()->DefaultValueSimpleValidation(LiteralTerm->Type, *LiteralTerm->Name, LiteralTerm->Name, nullptr, FText()))
 				{
 					Context.MessageLog.Error(*FString::Printf(*LOCTEXT("Error_InvalidOptionValue", "Invalid option value '%s' in @@").ToString(), *LiteralTerm->Name), Node);
 					return;
@@ -226,7 +226,7 @@ public:
 
 				if (UEnum* NodeEnum = SelectNode->GetEnum())
 				{
-					int32 EnumValue = NodeEnum->GetValueByName(FName(*OptionPins[OptionIdx]->PinName));
+					int32 EnumValue = NodeEnum->GetValueByName(OptionPins[OptionIdx]->PinName);
 					LiteralTerm->Name = FString::Printf(TEXT("%d"), EnumValue);
 				}
 				else
@@ -285,21 +285,23 @@ public:
 					LiteralStringTerm->bIsLiteral = true;
 					LiteralStringTerm->Type.PinCategory = UEdGraphSchema_K2::PC_String;
 
-					FString SelectionNodeType(TEXT("NONE"));
+					FString SelectionNodeType;
 					if (IndexPin)
 					{
-						UEnum* EnumObject = Cast<UEnum>(IndexPin->PinType.PinSubCategoryObject.Get());
-						if (EnumObject != NULL)
+						if (UEnum* EnumObject = Cast<UEnum>(IndexPin->PinType.PinSubCategoryObject.Get()))
 						{
 							SelectionNodeType = EnumObject->GetName();
 						}
 						else
 						{
 							// Not an enum, so just use the basic type
-							SelectionNodeType = IndexPin->PinType.PinCategory;
+							SelectionNodeType = IndexPin->PinType.PinCategory.ToString();
 						}
 					}
-
+					else
+					{
+						SelectionNodeType = TEXT("NONE");
+					}
 					const UEdGraph* OwningGraph = Context.MessageLog.FindSourceObjectTypeChecked<UEdGraph>( SelectNode->GetGraph() );
 					LiteralStringTerm->Name =
 						FString::Printf(*LOCTEXT("SelectNodeIndexWarning", "Graph %s: Selection Node of type %s failed! Out of bounds indexing of the options. There are only %d options available.").ToString(),
@@ -363,17 +365,17 @@ void UK2Node_Select::AllocateDefaultPins()
 
 		if (Enum)
 		{
-			const FString PinName = EnumEntries[Idx].ToString();
+			const FName PinName = EnumEntries[Idx];
 			UEdGraphPin* TempPin = FindPin(PinName);
 			if (!TempPin)
 			{
-				NewPin = CreatePin(EGPD_Input, UEdGraphSchema_K2::PC_Wildcard, FString(), nullptr, PinName);
+				NewPin = CreatePin(EGPD_Input, UEdGraphSchema_K2::PC_Wildcard, PinName);
 			}
 		}
 		else
 		{
-			const FString PinName = FString::Printf(TEXT("Option %d"), Idx);
-			NewPin = CreatePin(EGPD_Input, UEdGraphSchema_K2::PC_Wildcard, FString(), nullptr, PinName);
+			const FName PinName = *FString::Printf(TEXT("Option %d"), Idx);
+			NewPin = CreatePin(EGPD_Input, UEdGraphSchema_K2::PC_Wildcard, PinName);
 		}
 
 		if (NewPin)
@@ -394,7 +396,7 @@ void UK2Node_Select::AllocateDefaultPins()
 	CreatePin(EGPD_Input, IndexPinType.PinCategory, IndexPinType.PinSubCategory, IndexPinType.PinSubCategoryObject.Get(), TEXT("Index"));
 
 	// Create the return value
-	UEdGraphPin* ReturnPin = CreatePin(EGPD_Output, UEdGraphSchema_K2::PC_Wildcard, FString(), nullptr, UEdGraphSchema_K2::PN_ReturnValue);
+	UEdGraphPin* ReturnPin = CreatePin(EGPD_Output, UEdGraphSchema_K2::PC_Wildcard, UEdGraphSchema_K2::PN_ReturnValue);
 	ReturnPin->bDisplayAsMutableRef = UseSelectRef;
 
 	Super::AllocateDefaultPins();
@@ -436,16 +438,17 @@ FText UK2Node_Select::GetNodeTitle(ENodeTitleType::Type TitleType) const
 
 UK2Node::ERedirectType UK2Node_Select::DoPinsMatchForReconstruction(const UEdGraphPin* NewPin, int32 NewPinIndex, const UEdGraphPin* OldPin, int32 OldPinIndex) const
 {
-	// Check to see if the new pin name matches the old pin name (case insensitive - since the base uses Stricmp() to compare pin names, we also ignore case here).
-	if(Enum != nullptr && NewPinIndex < NumOptionPins && !NewPin->PinName.Equals(OldPin->PinName, ESearchCase::IgnoreCase))
+	// Check to see if the new pin name matches the old pin name.
+	if (Enum && (NewPinIndex < NumOptionPins) && (NewPin->PinName != OldPin->PinName))
 	{
 		// The names don't match, so check for an enum redirect from the old pin name.
-		int32 EnumIndex = Enum->GetIndexByNameString(OldPin->PinName);
+		const int32 EnumIndex = Enum->GetIndexByName(OldPin->PinName);
 		if(EnumIndex != INDEX_NONE)
 		{
 			// Found a redirect. Attempt to match it to the new pin name.
-			FString NewPinName = Enum->GetNameStringByIndex(EnumIndex);
-			if(NewPinName.Equals(NewPin->PinName, ESearchCase::IgnoreCase))
+			// Can't use Enum->GetNameByIndex here because it doesn't do namespace mangling
+			const FString NewPinName = Enum->GetNameStringByIndex(EnumIndex);
+			if (NewPinName == NewPin->PinName.ToString())
 			{
 				// The redirect is a match, so we can reconstruct this pin using the old pin's state.
 				return UK2Node::ERedirectType_Name;
@@ -477,7 +480,7 @@ void UK2Node_Select::ReallocatePinsDuringReconstruction(TArray<UEdGraphPin*>& Ol
 		{
 			OldIndexPin = OldPin;
 		}
-		else if (OldPin->PinName == Schema->PN_ReturnValue)
+		else if (OldPin->PinName == UEdGraphSchema_K2::PN_ReturnValue)
 		{
 			OldReturnPin = OldPin;
 		}
@@ -486,7 +489,7 @@ void UK2Node_Select::ReallocatePinsDuringReconstruction(TArray<UEdGraphPin*>& Ol
 	UEdGraphPin* ReturnPin = GetReturnValuePin();
 	check(ReturnPin);
 
-	if (OldReturnPin && (ReturnPin->PinType.PinCategory == Schema->PC_Wildcard))
+	if (OldReturnPin && (ReturnPin->PinType.PinCategory == UEdGraphSchema_K2::PC_Wildcard))
 	{
 		// Always copy type from node prior, if pins have changed those will error at compilation time
 		ReturnPin->PinType = OldReturnPin->PinType;
@@ -499,8 +502,8 @@ void UK2Node_Select::ReallocatePinsDuringReconstruction(TArray<UEdGraphPin*>& Ol
 	if (OldConditionPin)
 	{
 		// Set the index pin type
-		IndexPinType.PinCategory = Schema->PC_Boolean;
-		IndexPinType.PinSubCategory.Reset();
+		IndexPinType.PinCategory = UEdGraphSchema_K2::PC_Boolean;
+		IndexPinType.PinSubCategory = NAME_None;
 		IndexPinType.PinSubCategoryObject = nullptr;
 
 		// Set the pin type and Copy the pin
@@ -512,8 +515,8 @@ void UK2Node_Select::ReallocatePinsDuringReconstruction(TArray<UEdGraphPin*>& Ol
 			PinConnectionListChanged(IndexPin);
 		}
 
-		UEdGraphPin* OptionPin0 = FindPin("Option 0");
-		UEdGraphPin* OptionPin1 = FindPin("Option 1");
+		UEdGraphPin* OptionPin0 = FindPin(TEXT("Option 0"));
+		UEdGraphPin* OptionPin1 = FindPin(TEXT("Option 1"));
 
 		for (UEdGraphPin* OldPin : OldPins)
 		{
@@ -530,11 +533,11 @@ void UK2Node_Select::ReallocatePinsDuringReconstruction(TArray<UEdGraphPin*>& Ol
 
 	// If the index pin has links or a default value but is a wildcard, this is an old int pin so convert it
 	if (OldIndexPin &&
-		IndexPinType.PinCategory == Schema->PC_Wildcard &&
+		IndexPinType.PinCategory == UEdGraphSchema_K2::PC_Wildcard &&
 		(OldIndexPin->LinkedTo.Num() > 0 || !OldIndexPin->DefaultValue.IsEmpty()))
 	{
-		IndexPinType.PinCategory = Schema->PC_Int;
-		IndexPinType.PinSubCategory.Reset();
+		IndexPinType.PinCategory = UEdGraphSchema_K2::PC_Int;
+		IndexPinType.PinSubCategory = NAME_None;
 		IndexPinType.PinSubCategoryObject = nullptr;
 		IndexPin->PinType = IndexPinType;
 	}
@@ -542,12 +545,12 @@ void UK2Node_Select::ReallocatePinsDuringReconstruction(TArray<UEdGraphPin*>& Ol
 	// Set up default values for index and option pins now that the information is available
 	Schema->SetPinAutogeneratedDefaultValueBasedOnType(IndexPin);
 
-	const bool bFillTypeFromReturn = ReturnPin->PinType.PinCategory != Schema->PC_Wildcard;
+	const bool bFillTypeFromReturn = ReturnPin->PinType.PinCategory != UEdGraphSchema_K2::PC_Wildcard;
 	TArray<UEdGraphPin*> OptionPins;
 	GetOptionPins(OptionPins);
 	for (UEdGraphPin* Pin : OptionPins)
 	{
-		const bool bTypeShouldBeFilled = Pin && (Pin->PinType.PinCategory == Schema->PC_Wildcard);
+		const bool bTypeShouldBeFilled = Pin && (Pin->PinType.PinCategory == UEdGraphSchema_K2::PC_Wildcard);
 		if (bTypeShouldBeFilled && bFillTypeFromReturn)
 		{
 			Pin->PinType = ReturnPin->PinType;			
@@ -598,13 +601,11 @@ void UK2Node_Select::NotifyPinConnectionListChanged(UEdGraphPin* Pin)
 {
 	Super::NotifyPinConnectionListChanged(Pin);
 
-	const UEdGraphSchema_K2* Schema = GetDefault<UEdGraphSchema_K2>();
-
 	// If this is the Enum pin we need to set the enum and reconstruct the node
 	if (Pin == GetIndexPin())
 	{
 		// If the index pin was just linked to another pin
-		if (Pin->LinkedTo.Num() > 0 && Pin->PinType.PinCategory == Schema->PC_Wildcard)
+		if (Pin->LinkedTo.Num() > 0 && Pin->PinType.PinCategory == UEdGraphSchema_K2::PC_Wildcard)
 		{
 			UEdGraphPin* LinkPin = Pin->LinkedTo[0];
 
@@ -621,10 +622,10 @@ void UK2Node_Select::NotifyPinConnectionListChanged(UEdGraphPin* Pin)
 		// Grab references to all option pins and the return pin
 		TArray<UEdGraphPin*> OptionPins;
 		GetOptionPins(OptionPins);
-		UEdGraphPin* ReturnPin = FindPin(Schema->PN_ReturnValue);
+		UEdGraphPin* ReturnPin = FindPin(UEdGraphSchema_K2::PN_ReturnValue);
 
 		// See if this pin is one of the wildcard pins
-		bool bIsWildcardPin = (Pin == ReturnPin || OptionPins.Find(Pin) != INDEX_NONE) && Pin->PinType.PinCategory == Schema->PC_Wildcard;
+		bool bIsWildcardPin = (Pin == ReturnPin || OptionPins.Find(Pin) != INDEX_NONE) && Pin->PinType.PinCategory == UEdGraphSchema_K2::PC_Wildcard;
 
 		// If the pin was one of the wildcards we have to handle it specially
 		if (bIsWildcardPin)
@@ -647,10 +648,8 @@ void UK2Node_Select::NotifyPinConnectionListChanged(UEdGraphPin* Pin)
 
 UEdGraphPin* UK2Node_Select::GetReturnValuePin() const
 {
-	const UEdGraphSchema_K2* K2Schema = GetDefault<UEdGraphSchema_K2>();
-
-	UEdGraphPin* Pin = FindPin(K2Schema->PN_ReturnValue);
-	check(Pin != NULL);
+	UEdGraphPin* Pin = FindPin(UEdGraphSchema_K2::PN_ReturnValue);
+	check(Pin);
 	return Pin;
 }
 
@@ -663,24 +662,22 @@ UEdGraphPin* UK2Node_Select::GetIndexPin() const
 
 UEdGraphPin* UK2Node_Select::GetIndexPinUnchecked() const
 {
-	return FindPin("Index");
+	return FindPin(TEXT("Index"));
 }
 
 void UK2Node_Select::GetOptionPins(TArray<UEdGraphPin*>& OptionPins) const
 {
-	const UEdGraphSchema_K2* K2Schema = GetDefault<UEdGraphSchema_K2>();
-
-	OptionPins.Empty();
+	OptionPins.Reset();
 
 	// If the select node is currently dealing with an enum
-	if (IndexPinType.PinCategory == K2Schema->PC_Byte &&
-		IndexPinType.PinSubCategory.IsEmpty() &&
+	if (IndexPinType.PinCategory == UEdGraphSchema_K2::PC_Byte &&
+		IndexPinType.PinSubCategory.IsNone() &&
 		IndexPinType.PinSubCategoryObject != nullptr &&
 		IndexPinType.PinSubCategoryObject->IsA(UEnum::StaticClass()))
 	{
 		for (UEdGraphPin* Pin : Pins)
 		{
-			if (EnumEntries.Contains(FName(*Pin->PinName)))
+			if (EnumEntries.Contains(Pin->PinName))
 			{
 				OptionPins.Add(Pin);
 			}
@@ -688,9 +685,10 @@ void UK2Node_Select::GetOptionPins(TArray<UEdGraphPin*>& OptionPins) const
 	}
 	else
 	{
+		const TCHAR* OptionStr = TEXT("Option");
 		for (UEdGraphPin* Pin : Pins)
 		{
-			if (Pin->PinName.Left(6) == "Option")
+			if (Pin->PinName.ToString().StartsWith(OptionStr))
 			{
 				OptionPins.Add(Pin);
 			}
@@ -702,15 +700,15 @@ void UK2Node_Select::GetConditionalFunction(FName& FunctionName, UClass** Functi
 {
 	const UEdGraphSchema_K2* K2Schema = GetDefault<UEdGraphSchema_K2>();
 
-	if (IndexPinType.PinCategory == K2Schema->PC_Boolean)
+	if (IndexPinType.PinCategory == UEdGraphSchema_K2::PC_Boolean)
 	{
 		FunctionName = GET_FUNCTION_NAME_CHECKED(UKismetMathLibrary, EqualEqual_BoolBool);
 	}
-	else if (IndexPinType.PinCategory == K2Schema->PC_Byte)
+	else if (IndexPinType.PinCategory == UEdGraphSchema_K2::PC_Byte)
 	{
 		FunctionName = GET_FUNCTION_NAME_CHECKED(UKismetMathLibrary, EqualEqual_ByteByte);
 	}
-	else if (IndexPinType.PinCategory == K2Schema->PC_Int)
+	else if (IndexPinType.PinCategory == UEdGraphSchema_K2::PC_Int)
 	{
 		FunctionName = GET_FUNCTION_NAME_CHECKED(UKismetMathLibrary, EqualEqual_IntInt);
 	}
@@ -733,9 +731,9 @@ void UK2Node_Select::AddOptionPinToNode()
 	// We guarantee at least 2 options by default and since we just increased the count
 	// to more than 2, we need to make sure we're now dealing with an index for selection
 	// instead of the default boolean check
-	if (IndexPinType.PinCategory == K2Schema->PC_Boolean)
+	if (IndexPinType.PinCategory == UEdGraphSchema_K2::PC_Boolean)
 	{
-		IndexPinType.PinCategory = K2Schema->PC_Int;
+		IndexPinType.PinCategory = UEdGraphSchema_K2::PC_Int;
 		GetIndexPin()->BreakAllPinLinks();
 	}
 	// We will let the AllocateDefaultPins call handle the actual addition via ReconstructNode
@@ -801,15 +799,13 @@ void UK2Node_Select::NodeConnectionListChanged()
 
 bool UK2Node_Select::CanAddOptionPinToNode() const
 {
-	const UEdGraphSchema_K2* Schema = GetDefault<UEdGraphSchema_K2>();
-
-	if (IndexPinType.PinCategory == Schema->PC_Byte &&
+	if (IndexPinType.PinCategory == UEdGraphSchema_K2::PC_Byte &&
 		IndexPinType.PinSubCategoryObject.IsValid() &&
 		IndexPinType.PinSubCategoryObject.Get()->IsA(UEnum::StaticClass()))
 	{
 		return false;
 	}
-	else if (IndexPinType.PinCategory == Schema->PC_Boolean)
+	else if (IndexPinType.PinCategory == UEdGraphSchema_K2::PC_Boolean)
 	{
 		return false;
 	}
@@ -819,14 +815,12 @@ bool UK2Node_Select::CanAddOptionPinToNode() const
 
 bool UK2Node_Select::CanRemoveOptionPinToNode() const
 {
-	const UEdGraphSchema_K2* Schema = GetDefault<UEdGraphSchema_K2>();
-
-	if (IndexPinType.PinCategory == Schema->PC_Byte &&
+	if (IndexPinType.PinCategory == UEdGraphSchema_K2::PC_Byte &&
 		(NULL != Cast<UEnum>(IndexPinType.PinSubCategoryObject.Get())))
 	{
 		return false;
 	}
-	else if (IndexPinType.PinCategory == Schema->PC_Boolean)
+	else if (IndexPinType.PinCategory == UEdGraphSchema_K2::PC_Boolean)
 	{
 		return false;
 	}
diff --git a/Engine/Source/Editor/BlueprintGraph/Private/K2Node_Self.cpp b/Engine/Source/Editor/BlueprintGraph/Private/K2Node_Self.cpp
--- a/Engine/Source/Editor/BlueprintGraph/Private/K2Node_Self.cpp
+++ b/Engine/Source/Editor/BlueprintGraph/Private/K2Node_Self.cpp
@@ -23,9 +23,8 @@ public:
 	virtual void RegisterNets(FKismetFunctionContext& Context, UEdGraphNode* Node) override
 	{
 		UK2Node_Self* SelfNode = CastChecked<UK2Node_Self>(Node);
-		const UEdGraphSchema_K2* Schema = GetDefault<UEdGraphSchema_K2>();
 
-		UEdGraphPin* VarPin = SelfNode->FindPin(Schema->PN_Self);
+		UEdGraphPin* VarPin = SelfNode->FindPin(UEdGraphSchema_K2::PN_Self);
 		check( VarPin );
 
 		FBPTerminal* Term = new (Context.Literals) FBPTerminal();
@@ -42,8 +41,7 @@ UK2Node_Self::UK2Node_Self(const FObjectInitializer& ObjectInitializer)
 
 void UK2Node_Self::AllocateDefaultPins()
 {
-	const UEdGraphSchema_K2* K2Schema = GetDefault<UEdGraphSchema_K2>();
-	CreatePin(EGPD_Output, K2Schema->PC_Object, K2Schema->PSC_Self, nullptr, K2Schema->PN_Self);
+	CreatePin(EGPD_Output, UEdGraphSchema_K2::PC_Object, UEdGraphSchema_K2::PSC_Self, nullptr, UEdGraphSchema_K2::PN_Self);
 
 	Super::AllocateDefaultPins();
 }
diff --git a/Engine/Source/Editor/BlueprintGraph/Private/K2Node_SetFieldsInStruct.cpp b/Engine/Source/Editor/BlueprintGraph/Private/K2Node_SetFieldsInStruct.cpp
--- a/Engine/Source/Editor/BlueprintGraph/Private/K2Node_SetFieldsInStruct.cpp
+++ b/Engine/Source/Editor/BlueprintGraph/Private/K2Node_SetFieldsInStruct.cpp
@@ -123,15 +123,16 @@ UK2Node_SetFieldsInStruct::UK2Node_SetFieldsInStruct(const FObjectInitializer& O
 
 void UK2Node_SetFieldsInStruct::AllocateDefaultPins()
 {
-	const UEdGraphSchema_K2* Schema = GetDefault<UEdGraphSchema_K2>();
-	if (Schema && StructType)
+	if (StructType)
 	{
-		CreatePin(EGPD_Input, Schema->PC_Exec, FString(), nullptr, Schema->PN_Execute);
-		CreatePin(EGPD_Output, Schema->PC_Exec, FString(), nullptr, Schema->PN_Then);
+		CreatePin(EGPD_Input, UEdGraphSchema_K2::PC_Exec, UEdGraphSchema_K2::PN_Execute);
+		CreatePin(EGPD_Output, UEdGraphSchema_K2::PC_Exec, UEdGraphSchema_K2::PN_Then);
 
-		UEdGraphPin* InPin = CreatePin(EGPD_Input, Schema->PC_Struct, FString(), StructType, SetFieldsInStructHelper::StructRefPinName(), EPinContainerType::None, true);
+		UEdGraphNode::FCreatePinParams PinParams;
+		PinParams.bIsReference = true;
 
-		UEdGraphPin* OutPin = CreatePin(EGPD_Output, Schema->PC_Struct, FString(), StructType, SetFieldsInStructHelper::StructOutPinName(), EPinContainerType::None, true);
+		UEdGraphPin* InPin = CreatePin(EGPD_Input, UEdGraphSchema_K2::PC_Struct, StructType, SetFieldsInStructHelper::StructRefPinName(), PinParams);
+		UEdGraphPin* OutPin = CreatePin(EGPD_Output, UEdGraphSchema_K2::PC_Struct, StructType, SetFieldsInStructHelper::StructOutPinName(), PinParams);
 
 		// Input pin will forward the ref to the output, if the input value is not a reference connection, a copy is made and modified instead and provided as a reference until the function is called again.
 		InPin->AssignByRefPassThroughConnection(OutPin);
@@ -228,7 +229,7 @@ void UK2Node_SetFieldsInStruct::RemoveFieldPins(UEdGraphPin* Pin, EPinsToRemove
 		bool bWasChanged = false;
 		for (FOptionalPinFromProperty& OptionalProperty : ShowPinForProperties)
 		{
-			const bool bSelected = (Pin->PinName == OptionalProperty.PropertyName.ToString());
+			const bool bSelected = (Pin->PinName == OptionalProperty.PropertyName);
 			const bool bHide = (bSelected && bHideSelected) || (!bSelected && bHideNotSelected);
 			if (OptionalProperty.bShowPin && bHide)
 			{
diff --git a/Engine/Source/Editor/BlueprintGraph/Private/K2Node_SetVariableOnPersistentFrame.cpp b/Engine/Source/Editor/BlueprintGraph/Private/K2Node_SetVariableOnPersistentFrame.cpp
--- a/Engine/Source/Editor/BlueprintGraph/Private/K2Node_SetVariableOnPersistentFrame.cpp
+++ b/Engine/Source/Editor/BlueprintGraph/Private/K2Node_SetVariableOnPersistentFrame.cpp
@@ -82,10 +82,8 @@ public:
 
 void UK2Node_SetVariableOnPersistentFrame::AllocateDefaultPins()
 {
-	const UEdGraphSchema_K2* K2Schema = GetDefault<UEdGraphSchema_K2>();
-
-	CreatePin(EGPD_Input, K2Schema->PC_Exec, FString(), nullptr, K2Schema->PN_Execute);
-	CreatePin(EGPD_Output, K2Schema->PC_Exec, FString(), nullptr, K2Schema->PN_Then);
+	CreatePin(EGPD_Input, UEdGraphSchema_K2::PC_Exec, UEdGraphSchema_K2::PN_Execute);
+	CreatePin(EGPD_Output, UEdGraphSchema_K2::PC_Exec, UEdGraphSchema_K2::PN_Then);
 
 	Super::AllocateDefaultPins();
 }
diff --git a/Engine/Source/Editor/BlueprintGraph/Private/K2Node_SpawnActor.cpp b/Engine/Source/Editor/BlueprintGraph/Private/K2Node_SpawnActor.cpp
--- a/Engine/Source/Editor/BlueprintGraph/Private/K2Node_SpawnActor.cpp
+++ b/Engine/Source/Editor/BlueprintGraph/Private/K2Node_SpawnActor.cpp
@@ -12,10 +12,10 @@
 #include "KismetCompilerMisc.h"
 #include "KismetCompiler.h"
 
-static FString WorldContextPinName(TEXT("WorldContextObject"));
-static FString BlueprintPinName(TEXT("Blueprint"));
-static FString SpawnTransformPinName(TEXT("SpawnTransform"));
-static FString NoCollisionFailPinName(TEXT("SpawnEvenIfColliding"));
+static const FName WorldContextPinName(TEXT("WorldContextObject"));
+static const FName BlueprintPinName(TEXT("Blueprint"));
+static const FName SpawnTransformPinName(TEXT("SpawnTransform"));
+static const FName NoCollisionFailPinName(TEXT("SpawnEvenIfColliding"));
 
 
 #define LOCTEXT_NAMESPACE "K2Node_SpawnActor"
@@ -31,30 +31,30 @@ void UK2Node_SpawnActor::AllocateDefaultPins()
 	UEdGraphSchema_K2 const* K2Schema = GetDefault<UEdGraphSchema_K2>();
 
 	// Add execution pins
-	CreatePin(EGPD_Input, K2Schema->PC_Exec, FString(), nullptr, K2Schema->PN_Execute);
-	CreatePin(EGPD_Output, K2Schema->PC_Exec, FString(), nullptr, K2Schema->PN_Then);
+	CreatePin(EGPD_Input, UEdGraphSchema_K2::PC_Exec, UEdGraphSchema_K2::PN_Execute);
+	CreatePin(EGPD_Output, UEdGraphSchema_K2::PC_Exec, UEdGraphSchema_K2::PN_Then);
 
 	// If required add the world context pin
 	if (GetBlueprint()->ParentClass->HasMetaDataHierarchical(FBlueprintMetadata::MD_ShowWorldContextPin))
 	{
-		CreatePin(EGPD_Input, K2Schema->PC_Object, FString(), UObject::StaticClass(), WorldContextPinName);
+		CreatePin(EGPD_Input, UEdGraphSchema_K2::PC_Object, UObject::StaticClass(), WorldContextPinName);
 	}
 
 	// Add blueprint pin
-	UEdGraphPin* BlueprintPin = CreatePin(EGPD_Input, K2Schema->PC_Object, FString(), UBlueprint::StaticClass(), BlueprintPinName);
+	UEdGraphPin* BlueprintPin = CreatePin(EGPD_Input, UEdGraphSchema_K2::PC_Object, UBlueprint::StaticClass(), BlueprintPinName);
 	K2Schema->ConstructBasicPinTooltip(*BlueprintPin, LOCTEXT("BlueprintPinDescription", "The blueprint Actor you want to spawn"), BlueprintPin->PinToolTip);
 
 	// Transform pin
 	UScriptStruct* TransformStruct = TBaseStructure<FTransform>::Get();
-	UEdGraphPin* TransformPin = CreatePin(EGPD_Input, K2Schema->PC_Struct, FString(), TransformStruct, SpawnTransformPinName);
+	UEdGraphPin* TransformPin = CreatePin(EGPD_Input, UEdGraphSchema_K2::PC_Struct, TransformStruct, SpawnTransformPinName);
 	K2Schema->ConstructBasicPinTooltip(*TransformPin, LOCTEXT("TransformPinDescription", "The transform to spawn the Actor with"), TransformPin->PinToolTip);
 
 	// bNoCollisionFail pin
-	UEdGraphPin* NoCollisionFailPin = CreatePin(EGPD_Input, K2Schema->PC_Boolean, FString(), nullptr, NoCollisionFailPinName);
+	UEdGraphPin* NoCollisionFailPin = CreatePin(EGPD_Input, UEdGraphSchema_K2::PC_Boolean, NoCollisionFailPinName);
 	K2Schema->ConstructBasicPinTooltip(*NoCollisionFailPin, LOCTEXT("NoCollisionFailPinDescription", "Determines if the Actor should be spawned when the location is blocked by a collision"), NoCollisionFailPin->PinToolTip);
 
 	// Result pin
-	UEdGraphPin* ResultPin = CreatePin(EGPD_Output, K2Schema->PC_Object, FString(), AActor::StaticClass(), K2Schema->PN_ReturnValue);
+	UEdGraphPin* ResultPin = CreatePin(EGPD_Output, UEdGraphSchema_K2::PC_Object, AActor::StaticClass(), UEdGraphSchema_K2::PN_ReturnValue);
 	K2Schema->ConstructBasicPinTooltip(*ResultPin, LOCTEXT("ResultPinDescription", "The spawned Actor"), ResultPin->PinToolTip);
 
 	Super::AllocateDefaultPins();
@@ -80,7 +80,7 @@ void UK2Node_SpawnActor::CreatePinsForClass(UClass* InClass)
 			Property->HasAllPropertyFlags(CPF_BlueprintVisible) &&
 			!bIsDelegate )
 		{
-			UEdGraphPin* Pin = CreatePin(EGPD_Input, FString(), FString(), nullptr, Property->GetName());
+			UEdGraphPin* Pin = CreatePin(EGPD_Input, NAME_None, Property->GetFName());
 			const bool bPinGood = (Pin != nullptr) && K2Schema->ConvertPropertyToPinType(Property, /*out*/ Pin->PinType);	
 
 			// Copy tooltip from the property.
@@ -125,11 +125,9 @@ void UK2Node_SpawnActor::ReallocatePinsDuringReconstruction(TArray<UEdGraphPin*>
 
 bool UK2Node_SpawnActor::IsSpawnVarPin(UEdGraphPin* Pin)
 {
-	const UEdGraphSchema_K2* K2Schema = GetDefault<UEdGraphSchema_K2>();
-
-	return(	Pin->PinName != K2Schema->PN_Execute &&
-			Pin->PinName != K2Schema->PN_Then &&
-			Pin->PinName != K2Schema->PN_ReturnValue &&
+	return(	Pin->PinName != UEdGraphSchema_K2::PN_Execute &&
+			Pin->PinName != UEdGraphSchema_K2::PN_Then &&
+			Pin->PinName != UEdGraphSchema_K2::PN_ReturnValue &&
 			Pin->PinName != BlueprintPinName &&
 			Pin->PinName != WorldContextPinName &&
 			Pin->PinName != NoCollisionFailPinName &&
@@ -185,7 +183,7 @@ UEdGraphPin* UK2Node_SpawnActor::GetThenPin()const
 {
 	const UEdGraphSchema_K2* K2Schema = GetDefault<UEdGraphSchema_K2>();
 
-	UEdGraphPin* Pin = FindPinChecked(K2Schema->PN_Then);
+	UEdGraphPin* Pin = FindPinChecked(UEdGraphSchema_K2::PN_Then);
 	check(Pin->Direction == EGPD_Output);
 	return Pin;
 }
@@ -194,17 +192,16 @@ UEdGraphPin* UK2Node_SpawnActor::GetBlueprintPin(const TArray<UEdGraphPin*>* InP
 {
 	const TArray<UEdGraphPin*>* PinsToSearch = InPinsToSearch ? InPinsToSearch : &Pins;
 
-	UEdGraphPin* Pin = NULL;
-	for( auto PinIt = PinsToSearch->CreateConstIterator(); PinIt; ++PinIt )
+	UEdGraphPin* Pin = nullptr;
+	for (UEdGraphPin* TestPin : *PinsToSearch)
 	{
-		UEdGraphPin* TestPin = *PinIt;
 		if( TestPin && TestPin->PinName == BlueprintPinName )
 		{
 			Pin = TestPin;
 			break;
 		}
 	}
-	check(Pin == NULL || Pin->Direction == EGPD_Input);
+	check(Pin == nullptr || Pin->Direction == EGPD_Input);
 	return Pin;
 }
 
@@ -231,9 +228,7 @@ UEdGraphPin* UK2Node_SpawnActor::GetWorldContextPin() const
 
 UEdGraphPin* UK2Node_SpawnActor::GetResultPin() const
 {
-	const UEdGraphSchema_K2* K2Schema = GetDefault<UEdGraphSchema_K2>();
-
-	UEdGraphPin* Pin = FindPinChecked(K2Schema->PN_ReturnValue);
+	UEdGraphPin* Pin = FindPinChecked(UEdGraphSchema_K2::PN_ReturnValue);
 	check(Pin->Direction == EGPD_Output);
 	return Pin;
 }
@@ -282,17 +277,17 @@ void UK2Node_SpawnActor::ExpandNode(class FKismetCompilerContext& CompilerContex
 	Super::ExpandNode(CompilerContext, SourceGraph);
 
 	static FName BeginSpawningBlueprintFuncName = GET_FUNCTION_NAME_CHECKED(UGameplayStatics, BeginSpawningActorFromBlueprint);
-	static FString BlueprintParamName = FString(TEXT("Blueprint"));
-	static FString WorldContextParamName = FString(TEXT("WorldContextObject"));
+	static FName BlueprintParamName(TEXT("Blueprint"));
+	static FName WorldContextParamName(TEXT("WorldContextObject"));
 
 	static FName FinishSpawningFuncName = GET_FUNCTION_NAME_CHECKED(UGameplayStatics, FinishSpawningActor);
-	static FString ActorParamName = FString(TEXT("Actor"));
-	static FString TransformParamName = FString(TEXT("SpawnTransform"));
-	static FString NoCollisionFailParamName = FString(TEXT("bNoCollisionFail"));
+	static FName ActorParamName(TEXT("Actor"));
+	static FName TransformParamName(TEXT("SpawnTransform"));
+	static FName NoCollisionFailParamName(TEXT("bNoCollisionFail"));
 
-	static FString ObjectParamName = FString(TEXT("Object"));
-	static FString ValueParamName = FString(TEXT("Value"));
-	static FString PropertyNameParamName = FString(TEXT("PropertyName"));
+	static FName ObjectParamName(TEXT("Object"));
+	static FName ValueParamName(TEXT("Value"));
+	static FName PropertyNameParamName(TEXT("PropertyName"));
 
 	const UEdGraphSchema_K2* Schema = CompilerContext.GetSchema();
 
@@ -447,7 +442,7 @@ void UK2Node_SpawnActor::ExpandNode(class FKismetCompilerContext& CompilerContex
 
 				// Fill in literal for 'property name' pin - name of pin is property name
 				UEdGraphPin* PropertyNamePin = SetVarNode->FindPinChecked(PropertyNameParamName);
-				PropertyNamePin->DefaultValue = SpawnVarPin->PinName;
+				PropertyNamePin->DefaultValue = SpawnVarPin->PinName.ToString();
 
 				// Move connection from the variable pin on the spawn node to the 'value' pin
 				UEdGraphPin* ValuePin = SetVarNode->FindPinChecked(ValueParamName);
diff --git a/Engine/Source/Editor/BlueprintGraph/Private/K2Node_SpawnActorFromClass.cpp b/Engine/Source/Editor/BlueprintGraph/Private/K2Node_SpawnActorFromClass.cpp
--- a/Engine/Source/Editor/BlueprintGraph/Private/K2Node_SpawnActorFromClass.cpp
+++ b/Engine/Source/Editor/BlueprintGraph/Private/K2Node_SpawnActorFromClass.cpp
@@ -19,18 +19,18 @@
 
 struct FK2Node_SpawnActorFromClassHelper
 {
-	static FString SpawnTransformPinName;
-	static FString SpawnEvenIfCollidingPinName;
-	static FString NoCollisionFailPinName;
-	static FString CollisionHandlingOverridePinName;
-	static FString OwnerPinName;
+	static const FName SpawnTransformPinName;
+	static const FName SpawnEvenIfCollidingPinName;
+	static const FName NoCollisionFailPinName;
+	static const FName CollisionHandlingOverridePinName;
+	static const FName OwnerPinName;
 };
 
-FString FK2Node_SpawnActorFromClassHelper::SpawnTransformPinName(TEXT("SpawnTransform"));
-FString FK2Node_SpawnActorFromClassHelper::SpawnEvenIfCollidingPinName(TEXT("SpawnEvenIfColliding"));		// deprecated pin, name kept for backwards compat
-FString FK2Node_SpawnActorFromClassHelper::NoCollisionFailPinName(TEXT("bNoCollisionFail"));		// deprecated pin, name kept for backwards compat
-FString FK2Node_SpawnActorFromClassHelper::CollisionHandlingOverridePinName(TEXT("CollisionHandlingOverride"));
-FString FK2Node_SpawnActorFromClassHelper::OwnerPinName(TEXT("Owner"));
+const FName FK2Node_SpawnActorFromClassHelper::SpawnTransformPinName(TEXT("SpawnTransform"));
+const FName FK2Node_SpawnActorFromClassHelper::SpawnEvenIfCollidingPinName(TEXT("SpawnEvenIfColliding"));		// deprecated pin, name kept for backwards compat
+const FName FK2Node_SpawnActorFromClassHelper::NoCollisionFailPinName(TEXT("bNoCollisionFail"));		// deprecated pin, name kept for backwards compat
+const FName FK2Node_SpawnActorFromClassHelper::CollisionHandlingOverridePinName(TEXT("CollisionHandlingOverride"));
+const FName FK2Node_SpawnActorFromClassHelper::OwnerPinName(TEXT("Owner"));
 
 #define LOCTEXT_NAMESPACE "K2Node_SpawnActorFromClass"
 
@@ -49,18 +49,16 @@ void UK2Node_SpawnActorFromClass::AllocateDefaultPins()
 {
 	Super::AllocateDefaultPins();
 
-	const UEdGraphSchema_K2* K2Schema = GetDefault<UEdGraphSchema_K2>();
-
 	// Transform pin
 	UScriptStruct* TransformStruct = TBaseStructure<FTransform>::Get();
-	UEdGraphPin* TransformPin = CreatePin(EGPD_Input, UEdGraphSchema_K2::PC_Struct, FString(), TransformStruct, FK2Node_SpawnActorFromClassHelper::SpawnTransformPinName);
+	UEdGraphPin* TransformPin = CreatePin(EGPD_Input, UEdGraphSchema_K2::PC_Struct, TransformStruct, FK2Node_SpawnActorFromClassHelper::SpawnTransformPinName);
 
 	// Collision handling method pin
 	UEnum* const MethodEnum = FindObjectChecked<UEnum>(ANY_PACKAGE, TEXT("ESpawnActorCollisionHandlingMethod"), true);
-	UEdGraphPin* const CollisionHandlingOverridePin = CreatePin(EGPD_Input, UEdGraphSchema_K2::PC_Byte, FString(), MethodEnum, FK2Node_SpawnActorFromClassHelper::CollisionHandlingOverridePinName);
+	UEdGraphPin* const CollisionHandlingOverridePin = CreatePin(EGPD_Input, UEdGraphSchema_K2::PC_Byte, MethodEnum, FK2Node_SpawnActorFromClassHelper::CollisionHandlingOverridePinName);
 	CollisionHandlingOverridePin->DefaultValue = MethodEnum->GetNameStringByValue(static_cast<int>(ESpawnActorCollisionHandlingMethod::Undefined));
 
-	UEdGraphPin* OwnerPin = CreatePin(EGPD_Input, UEdGraphSchema_K2::PC_Object, FString(), AActor::StaticClass(), FK2Node_SpawnActorFromClassHelper::OwnerPinName);
+	UEdGraphPin* OwnerPin = CreatePin(EGPD_Input, UEdGraphSchema_K2::PC_Object, AActor::StaticClass(), FK2Node_SpawnActorFromClassHelper::OwnerPinName);
 	OwnerPin->bAdvancedView = true;
 	if (ENodeAdvancedPins::NoPins == AdvancedPinDisplay)
 	{
@@ -330,18 +328,18 @@ void UK2Node_SpawnActorFromClass::ExpandNode(class FKismetCompilerContext& Compi
 	Super::ExpandNode(CompilerContext, SourceGraph);
 
 	static const FName BeginSpawningBlueprintFuncName = GET_FUNCTION_NAME_CHECKED(UGameplayStatics, BeginDeferredActorSpawnFromClass);
-	static const FString ActorClassParamName = FString(TEXT("ActorClass"));
-	static const FString WorldContextParamName = FString(TEXT("WorldContextObject"));
+	static const FName ActorClassParamName(TEXT("ActorClass"));
+	static const FName WorldContextParamName(TEXT("WorldContextObject"));
 
 	static const FName FinishSpawningFuncName = GET_FUNCTION_NAME_CHECKED(UGameplayStatics, FinishSpawningActor);
-	static const FString ActorParamName = FString(TEXT("Actor"));
-	static const FString TransformParamName = FString(TEXT("SpawnTransform"));
-	static const FString CollisionHandlingOverrideParamName = FString(TEXT("CollisionHandlingOverride"));
-	static const FString OwnerParamName = FString(TEXT("Owner"));
-
-	static const FString ObjectParamName = FString(TEXT("Object"));
-	static const FString ValueParamName = FString(TEXT("Value"));
-	static const FString PropertyNameParamName = FString(TEXT("PropertyName"));
+	static const FName ActorParamName(TEXT("Actor"));
+	static const FName TransformParamName(TEXT("SpawnTransform"));
+	static const FName CollisionHandlingOverrideParamName(TEXT("CollisionHandlingOverride"));
+	static const FName OwnerParamName(TEXT("Owner"));
+
+	static const FName ObjectParamName(TEXT("Object"));
+	static const FName ValueParamName(TEXT("Value"));
+	static const FName PropertyNameParamName(TEXT("PropertyName"));
 
 	UK2Node_SpawnActorFromClass* SpawnNode = this;
 	UEdGraphPin* SpawnNodeExec = SpawnNode->GetExecPin();
diff --git a/Engine/Source/Editor/BlueprintGraph/Private/K2Node_StructMemberSet.cpp b/Engine/Source/Editor/BlueprintGraph/Private/K2Node_StructMemberSet.cpp
--- a/Engine/Source/Editor/BlueprintGraph/Private/K2Node_StructMemberSet.cpp
+++ b/Engine/Source/Editor/BlueprintGraph/Private/K2Node_StructMemberSet.cpp
@@ -39,11 +39,9 @@ void UK2Node_StructMemberSet::PostEditChangeProperty(FPropertyChangedEvent& Prop
 
 void UK2Node_StructMemberSet::AllocateDefaultPins()
 {
-	const UEdGraphSchema_K2* Schema = GetDefault<UEdGraphSchema_K2>();
-
 	// Add the execution sequencing pin
-	CreatePin(EGPD_Input, Schema->PC_Exec, FString(), nullptr, Schema->PN_Execute);
-	CreatePin(EGPD_Output, Schema->PC_Exec, FString(), nullptr, Schema->PN_Then);
+	CreatePin(EGPD_Input, UEdGraphSchema_K2::PC_Exec, UEdGraphSchema_K2::PN_Execute);
+	CreatePin(EGPD_Output, UEdGraphSchema_K2::PC_Exec, UEdGraphSchema_K2::PN_Then);
 
 	// Display any currently visible optional pins
 	{
diff --git a/Engine/Source/Editor/BlueprintGraph/Private/K2Node_StructOperation.cpp b/Engine/Source/Editor/BlueprintGraph/Private/K2Node_StructOperation.cpp
--- a/Engine/Source/Editor/BlueprintGraph/Private/K2Node_StructOperation.cpp
+++ b/Engine/Source/Editor/BlueprintGraph/Private/K2Node_StructOperation.cpp
@@ -79,14 +79,14 @@ bool UK2Node_StructOperation::DoRenamedPinsMatch(const UEdGraphPin* NewPin, cons
 	return bResult;
 }
 
-FString UK2Node_StructOperation::GetPinMetaData(FString InPinName, FName InKey)
+FString UK2Node_StructOperation::GetPinMetaData(FName InPinName, FName InKey)
 {
 	FString ReturnValue;
 
 	for (TFieldIterator<UProperty> It(StructType); It; ++It)
 	{
 		const UProperty* Property = *It;
-		if(Property && Property->GetName() == InPinName)
+		if(Property && Property->GetFName() == InPinName)
 		{
 			ReturnValue = Property->GetMetaData(InKey);
 			break;
diff --git a/Engine/Source/Editor/BlueprintGraph/Private/K2Node_Switch.cpp b/Engine/Source/Editor/BlueprintGraph/Private/K2Node_Switch.cpp
--- a/Engine/Source/Editor/BlueprintGraph/Private/K2Node_Switch.cpp
+++ b/Engine/Source/Editor/BlueprintGraph/Private/K2Node_Switch.cpp
@@ -15,8 +15,8 @@
 
 namespace 
 {
-	static FString DefaultPinName(TEXT("Default"));
-	static FString SelectionPinName(TEXT("Selection"));
+	static FName DefaultPinName(TEXT("Default"));
+	static FName SelectionPinName(TEXT("Selection"));
 }
 
 //////////////////////////////////////////////////////////////////////////
@@ -42,7 +42,7 @@ public:
 		// Create a term to determine if the compare was successful or not
 		//@TODO: Ideally we just create one ever, not one per switch
 		FBPTerminal* BoolTerm = Context.CreateLocalTerminal();
-		BoolTerm->Type.PinCategory = CompilerContext.GetSchema()->PC_Boolean;
+		BoolTerm->Type.PinCategory = UEdGraphSchema_K2::PC_Boolean;
 		BoolTerm->Source = Node;
 		BoolTerm->Name = Context.NetNameMap->MakeValidName(Node) + TEXT("_CmpSuccess");
 		BoolTermMap.Add(Node, BoolTerm);
@@ -87,7 +87,7 @@ public:
 
 			// Pull out function to use
 			UClass* FuncClass = Cast<UClass>(FuncPin->PinType.PinSubCategoryObject.Get());
-			UFunction* FunctionPtr = FindField<UFunction>(FuncClass, *FuncPin->PinName);
+			UFunction* FunctionPtr = FindField<UFunction>(FuncClass, FuncPin->PinName);
 			check(FunctionPtr);
 
 			// Run thru all the output pins except for the default label
@@ -99,7 +99,7 @@ public:
 				{
 					// Create a term for the switch case value
 					FBPTerminal* CaseValueTerm = new (Context.Literals) FBPTerminal();
-					CaseValueTerm->Name = Pin->PinName;
+					CaseValueTerm->Name = Pin->PinName.ToString();
 					CaseValueTerm->Type = SwitchNode->GetInnerCaseType();
 					CaseValueTerm->SourcePin = Pin;
 					CaseValueTerm->bIsLiteral = true;
@@ -172,23 +172,21 @@ void UK2Node_Switch::PostEditChangeProperty(struct FPropertyChangedEvent& Proper
 	Super::PostEditChangeProperty(PropertyChangedEvent);
 }
 
-FString UK2Node_Switch::GetSelectionPinName()
+FName UK2Node_Switch::GetSelectionPinName()
 {
 	return SelectionPinName;
 }
 
 void UK2Node_Switch::AllocateDefaultPins()
 {
-	const UEdGraphSchema_K2* K2Schema = GetDefault<UEdGraphSchema_K2>();
-
 	// Add default pin
 	if (bHasDefaultPin)
 	{
-		CreatePin(EGPD_Output, K2Schema->PC_Exec, FString(), nullptr, DefaultPinName);
+		CreatePin(EGPD_Output, UEdGraphSchema_K2::PC_Exec, DefaultPinName);
 	}
 
 	// Add exec input pin
-	CreatePin(EGPD_Input, K2Schema->PC_Exec, FString(), nullptr, K2Schema->PN_Execute);
+	CreatePin(EGPD_Input, UEdGraphSchema_K2::PC_Exec, UEdGraphSchema_K2::PN_Execute);
 
 	// Create selection pin based on type
 	CreateSelectionPin();
@@ -211,7 +209,7 @@ UK2Node::ERedirectType UK2Node_Switch::DoPinsMatchForReconstruction(const UEdGra
 			return ERedirectType_Name;
 		}
 	}
-	else if (FCString::Strcmp(*(NewPin->PinName), *(OldPin->PinName)) == 0)
+	else if (NewPin->PinName == OldPin->PinName)
 	{
 		// Compare the names, case-sensitively
 		return ERedirectType_Name;
@@ -233,11 +231,10 @@ FSlateIcon UK2Node_Switch::GetIconAndTint(FLinearColor& OutColor) const
 
 void UK2Node_Switch::AddPinToSwitchNode()
 {
-	const UEdGraphSchema_K2* K2Schema = GetDefault<UEdGraphSchema_K2>();
-	FString NewPinName = GetUniquePinName();
-	if (NewPinName.Len() > 0)
+	const FName NewPinName = GetUniquePinName();
+	if (!NewPinName.IsNone())
 	{
-		CreatePin(EGPD_Output, K2Schema->PC_Exec, FString(), nullptr, NewPinName);
+		CreatePin(EGPD_Output, UEdGraphSchema_K2::PC_Exec, NewPinName);
 	}
 }
 
@@ -284,17 +281,15 @@ bool UK2Node_Switch::CanRemoveExecutionPin(UEdGraphPin* TargetPin) const
 }
 
 // Returns the exec output pin name for a given 0-based index
-FString UK2Node_Switch::GetPinNameGivenIndex(int32 Index)
+FName UK2Node_Switch::GetPinNameGivenIndex(int32 Index) const
 {
-	const UEdGraphSchema_K2* K2Schema = GetDefault<UEdGraphSchema_K2>();
-	return FString::Printf(TEXT("%d"), Index);
+	return *FString::Printf(TEXT("%d"), Index);
 }
 
 void UK2Node_Switch::CreateFunctionPin()
 {
 	// Set properties on the function pin
-	const UEdGraphSchema_K2* K2Schema = GetDefault<UEdGraphSchema_K2>();
-	UEdGraphPin* FunctionPin = CreatePin(EGPD_Input, K2Schema->PC_Object, FString(), FunctionClass, FunctionName.ToString());
+	UEdGraphPin* FunctionPin = CreatePin(EGPD_Input, UEdGraphSchema_K2::PC_Object, FunctionClass, FunctionName);
 	FunctionPin->bDefaultValueIsReadOnly = true;
 	FunctionPin->bNotConnectable = true;
 	FunctionPin->bHidden = true;
@@ -318,7 +313,7 @@ void UK2Node_Switch::CreateFunctionPin()
 UEdGraphPin* UK2Node_Switch::GetFunctionPin() const
 {
 	//@TODO: Should probably use a specific index, though FindPin starts at 0, so this won't *currently* conflict with user created pins
-	return FindPin(FunctionName.ToString());
+	return FindPin(FunctionName);
 }
 
 UEdGraphPin* UK2Node_Switch::GetSelectionPin() const
