Copying //UE4/Dev-Framework to //UE4/Dev-Main (Source: //UE4/Dev-Framework @ 3716594)

diff --git a/Engine/Source/Editor/BlueprintGraph/Private/K2Node_SwitchEnum.cpp b/Engine/Source/Editor/BlueprintGraph/Private/K2Node_SwitchEnum.cpp
--- a/Engine/Source/Editor/BlueprintGraph/Private/K2Node_SwitchEnum.cpp
+++ b/Engine/Source/Editor/BlueprintGraph/Private/K2Node_SwitchEnum.cpp
@@ -112,9 +112,8 @@ void UK2Node_SwitchEnum::GetMenuActions(FBlueprintActionDatabaseRegistrar& Actio
 
 void UK2Node_SwitchEnum::CreateSelectionPin()
 {
-	const UEdGraphSchema_K2* K2Schema = GetDefault<UEdGraphSchema_K2>();
-	UEdGraphPin* Pin = CreatePin(EGPD_Input, K2Schema->PC_Byte, FString(), Enum, TEXT("Selection"));
-	K2Schema->SetPinAutogeneratedDefaultValueBasedOnType(Pin);
+	UEdGraphPin* Pin = CreatePin(EGPD_Input, UEdGraphSchema_K2::PC_Byte, Enum, TEXT("Selection"));
+	GetDefault<UEdGraphSchema_K2>()->SetPinAutogeneratedDefaultValueBasedOnType(Pin);
 }
 
 FEdGraphPinType UK2Node_SwitchEnum::GetPinType() const
@@ -125,23 +124,9 @@ FEdGraphPinType UK2Node_SwitchEnum::GetPinType() const
 	return EnumPinType; 
 }
 
-FString UK2Node_SwitchEnum::GetUniquePinName()
-{
-	FString NewPinName;
-	for (auto EnumIt = EnumEntries.CreateConstIterator(); EnumIt; ++EnumIt)
-	{
-		FName EnumEntry = *EnumIt;
-		if ( !FindPin(EnumEntry.ToString()) )
-		{
-			break;
-		}
-	}
-	return NewPinName;
-}
-
 void UK2Node_SwitchEnum::CreateCasePins()
 {
-	if(NULL != Enum)
+	if (Enum)
 	{
 		SetEnum(Enum);
 	}
@@ -152,18 +137,15 @@ void UK2Node_SwitchEnum::CreateCasePins()
 		AdvancedPinDisplay = ENodeAdvancedPins::Hidden;
 	}
 
-	const UEdGraphSchema_K2* K2Schema = GetDefault<UEdGraphSchema_K2>();
-	for (auto EnumIt = EnumEntries.CreateConstIterator(); EnumIt; ++EnumIt)
+	for (int32 Index = 0; Index < EnumEntries.Num(); ++Index)
 	{
-		FName EnumEntry = *EnumIt;
-		UEdGraphPin * NewPin = CreatePin(EGPD_Output, K2Schema->PC_Exec, FString(), nullptr, EnumEntry.ToString());
-		int32 Index = EnumIt.GetIndex();
+		UEdGraphPin* NewPin = CreatePin(EGPD_Output, UEdGraphSchema_K2::PC_Exec, EnumEntries[Index]);
 		if (EnumFriendlyNames.IsValidIndex(Index))
 		{
 			NewPin->PinFriendlyName = EnumFriendlyNames[Index];
 		}
 		
-		if(bShouldUseAdvancedView && (EnumIt.GetIndex() > 2))
+		if(bShouldUseAdvancedView && (Index > 2))
 		{
 			NewPin->bAdvancedView = true;
 		}
@@ -175,8 +157,8 @@ UK2Node::ERedirectType UK2Node_SwitchEnum::DoPinsMatchForReconstruction(const UE
 	UK2Node::ERedirectType ReturnValue = Super::DoPinsMatchForReconstruction(NewPin, NewPinIndex, OldPin, OldPinIndex);
 	if (ReturnValue == UK2Node::ERedirectType_None && Enum && OldPinIndex > 2 && NewPinIndex > 2)
 	{
-		int32 OldValue = Enum->GetValueByName(FName(*OldPin->PinName));
-		int32 NewValue = Enum->GetValueByName(FName(*NewPin->PinName));
+		const int32 OldValue = Enum->GetValueByName(OldPin->PinName);
+		const int32 NewValue = Enum->GetValueByName(NewPin->PinName);
 		// This handles redirects properly
 		if (OldValue == NewValue && OldValue != INDEX_NONE)
 		{
diff --git a/Engine/Source/Editor/BlueprintGraph/Private/K2Node_SwitchInteger.cpp b/Engine/Source/Editor/BlueprintGraph/Private/K2Node_SwitchInteger.cpp
--- a/Engine/Source/Editor/BlueprintGraph/Private/K2Node_SwitchInteger.cpp
+++ b/Engine/Source/Editor/BlueprintGraph/Private/K2Node_SwitchInteger.cpp
@@ -47,9 +47,8 @@ void UK2Node_SwitchInteger::GetMenuActions(FBlueprintActionDatabaseRegistrar& Ac
 
 void UK2Node_SwitchInteger::CreateSelectionPin()
 {
-	const UEdGraphSchema_K2* K2Schema = GetDefault<UEdGraphSchema_K2>();
-	UEdGraphPin* Pin = CreatePin(EGPD_Input, K2Schema->PC_Int, FString(), nullptr, TEXT("Selection"));
-	K2Schema->SetPinAutogeneratedDefaultValueBasedOnType(Pin);
+	UEdGraphPin* Pin = CreatePin(EGPD_Input, UEdGraphSchema_K2::PC_Int, TEXT("Selection"));
+	GetDefault<UEdGraphSchema_K2>()->SetPinAutogeneratedDefaultValueBasedOnType(Pin);
 }
 
 FEdGraphPinType UK2Node_SwitchInteger::GetPinType() const 
@@ -59,15 +58,14 @@ FEdGraphPinType UK2Node_SwitchInteger::GetPinType() const
 	return PinType;
 }
 
-FString UK2Node_SwitchInteger::GetPinNameGivenIndex(int32 Index)
+FName UK2Node_SwitchInteger::GetPinNameGivenIndex(int32 Index) const
 {
-	const UEdGraphSchema_K2* K2Schema = GetDefault<UEdGraphSchema_K2>();
-	return FString::Printf(TEXT("%d"), Index);
+	return *FString::Printf(TEXT("%d"), Index);
 }
 
-FString UK2Node_SwitchInteger::GetUniquePinName()
+FName UK2Node_SwitchInteger::GetUniquePinName()
 {
-	FString NewPinName;
+	FName NewPinName;
 	int32 i = StartIndex;
 	while (true)
 	{
@@ -90,7 +88,6 @@ void UK2Node_SwitchInteger::ReallocatePinsDuringReconstruction(TArray<UEdGraphPi
 {
 	Super::ReallocatePinsDuringReconstruction(OldPins);
 
-	const UEdGraphSchema_K2* K2Schema = GetDefault<UEdGraphSchema_K2>();
 	int32 ExecOutPinCount = StartIndex;
 
 	UEdGraphPin* DefaultPin = GetDefaultPin();
@@ -98,19 +95,19 @@ void UK2Node_SwitchInteger::ReallocatePinsDuringReconstruction(TArray<UEdGraphPi
 	for (int32 i = bHasDefaultPin ? 1 : 0; i < OldPins.Num(); ++i)
 	{
 		UEdGraphPin* TestPin = OldPins[i];
-		if (K2Schema->IsExecPin(*TestPin) && (TestPin->Direction == EGPD_Output))
+		if (UEdGraphSchema_K2::IsExecPin(*TestPin) && (TestPin->Direction == EGPD_Output))
 		{
 			// Skip the default pin to avoid creating an extra output pin in the case where the default pin has been toggled off
-			if(TestPin->PinName != TEXT("Default"))
+			if (TestPin->PinName != TEXT("Default"))
 			{
-				FString NewPinName = GetPinNameGivenIndex(ExecOutPinCount);
+				const FName NewPinName = GetPinNameGivenIndex(ExecOutPinCount);
 				ExecOutPinCount++;
 
 				// Make sure the old pin and new pin names match
 				TestPin->PinName = NewPinName;
 
 				// Create the new output pin to match
-				CreatePin(EGPD_Output, K2Schema->PC_Exec, FString(), nullptr, NewPinName);
+				CreatePin(EGPD_Output, UEdGraphSchema_K2::PC_Exec, NewPinName);
 			}
 		}
 	}
diff --git a/Engine/Source/Editor/BlueprintGraph/Private/K2Node_SwitchName.cpp b/Engine/Source/Editor/BlueprintGraph/Private/K2Node_SwitchName.cpp
--- a/Engine/Source/Editor/BlueprintGraph/Private/K2Node_SwitchName.cpp
+++ b/Engine/Source/Editor/BlueprintGraph/Private/K2Node_SwitchName.cpp
@@ -62,9 +62,8 @@ void UK2Node_SwitchName::GetMenuActions(FBlueprintActionDatabaseRegistrar& Actio
 
 void UK2Node_SwitchName::CreateSelectionPin()
 {
-	const UEdGraphSchema_K2* K2Schema = GetDefault<UEdGraphSchema_K2>();
-	UEdGraphPin* Pin = CreatePin(EGPD_Input, K2Schema->PC_Name, FString(), nullptr, TEXT("Selection"));
-	K2Schema->SetPinAutogeneratedDefaultValueBasedOnType(Pin);
+	UEdGraphPin* Pin = CreatePin(EGPD_Input, UEdGraphSchema_K2::PC_Name, TEXT("Selection"));
+	GetDefault<UEdGraphSchema_K2>()->SetPinAutogeneratedDefaultValueBasedOnType(Pin);
 }
 
 FEdGraphPinType UK2Node_SwitchName::GetPinType() const 
@@ -74,28 +73,27 @@ FEdGraphPinType UK2Node_SwitchName::GetPinType() const
 	return PinType;
 }
 
-FString UK2Node_SwitchName::GetPinNameGivenIndex(int32 Index)
+FName UK2Node_SwitchName::GetPinNameGivenIndex(int32 Index) const
 {
 	check(Index);
-	return PinNames[Index].ToString();
+	return PinNames[Index];
 }
 
 void UK2Node_SwitchName::CreateCasePins()
 {
-	const UEdGraphSchema_K2* K2Schema = GetDefault<UEdGraphSchema_K2>();
-	for ( FName PinName : PinNames)
+	for (const FName& PinName : PinNames)
 	{
-		CreatePin(EGPD_Output, K2Schema->PC_Exec, FString(), nullptr, PinName.ToString());
+		CreatePin(EGPD_Output, UEdGraphSchema_K2::PC_Exec, PinName);
 	}
 }
 
-FString UK2Node_SwitchName::GetUniquePinName()
+FName UK2Node_SwitchName::GetUniquePinName()
 {
-	FString NewPinName;
+	FName NewPinName;
 	int32 Index = 0;
 	while (true)
 	{
-		NewPinName =  FString::Printf(TEXT("Case_%d"), Index++);
+		NewPinName =  *FString::Printf(TEXT("Case_%d"), Index++);
 		if (!FindPin(NewPinName))
 		{
 			break;
@@ -106,11 +104,10 @@ FString UK2Node_SwitchName::GetUniquePinName()
 
 void UK2Node_SwitchName::AddPinToSwitchNode()
 {
-	FString PinName = GetUniquePinName();
-	PinNames.Add(FName(*PinName));
+	const FName PinName = GetUniquePinName();
+	PinNames.Add(PinName);
 
-	const UEdGraphSchema_K2* K2Schema = GetDefault<UEdGraphSchema_K2>();
-	CreatePin(EGPD_Output, K2Schema->PC_Exec, FString(), nullptr, PinName);
+	CreatePin(EGPD_Output, UEdGraphSchema_K2::PC_Exec, PinName);
 }
 
 void UK2Node_SwitchName::RemovePin(UEdGraphPin* TargetPin)
@@ -118,5 +115,5 @@ void UK2Node_SwitchName::RemovePin(UEdGraphPin* TargetPin)
 	checkSlow(TargetPin);
 
 	// Clean-up pin name array
-	PinNames.Remove(FName(*TargetPin->PinName));
+	PinNames.Remove(TargetPin->PinName);
 }
diff --git a/Engine/Source/Editor/BlueprintGraph/Private/K2Node_SwitchString.cpp b/Engine/Source/Editor/BlueprintGraph/Private/K2Node_SwitchString.cpp
--- a/Engine/Source/Editor/BlueprintGraph/Private/K2Node_SwitchString.cpp
+++ b/Engine/Source/Editor/BlueprintGraph/Private/K2Node_SwitchString.cpp
@@ -79,9 +79,8 @@ void UK2Node_SwitchString::GetMenuActions(FBlueprintActionDatabaseRegistrar& Act
 
 void UK2Node_SwitchString::CreateSelectionPin()
 {
-	const UEdGraphSchema_K2* K2Schema = GetDefault<UEdGraphSchema_K2>();
-	UEdGraphPin* Pin = CreatePin(EGPD_Input, K2Schema->PC_String, FString(), nullptr, TEXT("Selection"));
-	K2Schema->SetPinAutogeneratedDefaultValueBasedOnType(Pin);
+	UEdGraphPin* Pin = CreatePin(EGPD_Input, UEdGraphSchema_K2::PC_String, TEXT("Selection"));
+	GetDefault<UEdGraphSchema_K2>()->SetPinAutogeneratedDefaultValueBasedOnType(Pin);
 }
 
 FEdGraphPinType UK2Node_SwitchString::GetPinType() const 
@@ -91,29 +90,27 @@ FEdGraphPinType UK2Node_SwitchString::GetPinType() const
 	return PinType;
 }
 
-FString UK2Node_SwitchString::GetPinNameGivenIndex(int32 Index)
+FName UK2Node_SwitchString::GetPinNameGivenIndex(int32 Index) const
 {
 	check(Index);
 	return PinNames[Index];
 }
 
-
 void UK2Node_SwitchString::CreateCasePins()
 {
-	const UEdGraphSchema_K2* K2Schema = GetDefault<UEdGraphSchema_K2>();
-	for( const FString& PinName : PinNames)
+	for( const FName& PinName : PinNames)
 	{
-		CreatePin(EGPD_Output, K2Schema->PC_Exec, FString(), nullptr, PinName);
+		CreatePin(EGPD_Output, UEdGraphSchema_K2::PC_Exec, PinName);
 	}
 }
 
-FString UK2Node_SwitchString::GetUniquePinName()
+FName UK2Node_SwitchString::GetUniquePinName()
 {
-	FString NewPinName;
+	FName NewPinName;
 	int32 Index = 0;
 	while (true)
 	{
-		NewPinName =  FString::Printf(TEXT("Case_%d"), Index++);
+		NewPinName = *FString::Printf(TEXT("Case_%d"), Index++);
 		if (!FindPin(NewPinName))
 		{
 			break;
@@ -124,11 +121,10 @@ FString UK2Node_SwitchString::GetUniquePinName()
 
 void UK2Node_SwitchString::AddPinToSwitchNode()
 {
-	FString PinName = GetUniquePinName();
+	const FName PinName = GetUniquePinName();
 	PinNames.Add(PinName);
 
-	const UEdGraphSchema_K2* K2Schema = GetDefault<UEdGraphSchema_K2>();
-	CreatePin(EGPD_Output, K2Schema->PC_Exec, FString(), nullptr, PinName);
+	CreatePin(EGPD_Output, UEdGraphSchema_K2::PC_Exec, PinName);
 }
 
 void UK2Node_SwitchString::RemovePin(UEdGraphPin* TargetPin)
@@ -143,7 +139,7 @@ void UK2Node_SwitchString::ValidateNodeDuringCompilation(FCompilerResultsLog& Me
 {
 	Super::ValidateNodeDuringCompilation(MessageLog);
 
-	TSet<FString> UniquePinNames(PinNames);
+	TSet<FName> UniquePinNames(PinNames);
 
 	if (UniquePinNames.Num() != PinNames.Num())
 	{
diff --git a/Engine/Source/Editor/BlueprintGraph/Private/K2Node_TemporaryVariable.cpp b/Engine/Source/Editor/BlueprintGraph/Private/K2Node_TemporaryVariable.cpp
--- a/Engine/Source/Editor/BlueprintGraph/Private/K2Node_TemporaryVariable.cpp
+++ b/Engine/Source/Editor/BlueprintGraph/Private/K2Node_TemporaryVariable.cpp
@@ -41,9 +41,7 @@ UK2Node_TemporaryVariable::UK2Node_TemporaryVariable(const FObjectInitializer& O
 
 void UK2Node_TemporaryVariable::AllocateDefaultPins()
 {
-	const UEdGraphSchema_K2* K2Schema = GetDefault<UEdGraphSchema_K2>();
-
-	UEdGraphPin* VariablePin = CreatePin(EGPD_Output, FString(), FString(), nullptr, TEXT("Variable"));
+	UEdGraphPin* VariablePin = CreatePin(EGPD_Output, NAME_None, TEXT("Variable"));
 	VariablePin->PinType = VariableType;
 
 	Super::AllocateDefaultPins();
@@ -83,7 +81,7 @@ bool UK2Node_TemporaryVariable::IsNodePure() const
 
 FString UK2Node_TemporaryVariable::GetDescriptiveCompiledName() const
 {
-	FString Result = NSLOCTEXT("K2Node", "TempPinCategory", "Temp_").ToString() + VariableType.PinCategory;
+	FString Result = NSLOCTEXT("K2Node", "TempPinCategory", "Temp_").ToString() + VariableType.PinCategory.ToString();
 		
 	if (!NodeComment.IsEmpty())
 	{
@@ -169,41 +167,39 @@ void UK2Node_TemporaryVariable::GetMenuActions(FBlueprintActionDatabaseRegistrar
 		return NodeSpawner;
 	};
 
-	UEdGraphSchema_K2 const* K2Schema = GetDefault<UEdGraphSchema_K2>();
-
-	ActionRegistrar.AddBlueprintAction(ActionKey, MakeTempVarNodeSpawner(FEdGraphPinType(K2Schema->PC_Int, FString(), nullptr, EPinContainerType::None, /*bIsReference =*/ false, /*InValueTerminalType=*/FEdGraphTerminalType()), /*bIsPersistent =*/false));
-	ActionRegistrar.AddBlueprintAction(ActionKey, MakeTempVarNodeSpawner(FEdGraphPinType(K2Schema->PC_Int, FString(), nullptr, EPinContainerType::Array, /*bIsReference =*/ false, /*InValueTerminalType=*/FEdGraphTerminalType()), /*bIsPersistent =*/false));
-	ActionRegistrar.AddBlueprintAction(ActionKey, MakeTempVarNodeSpawner(FEdGraphPinType(K2Schema->PC_Float, FString(), nullptr, EPinContainerType::None, /*bIsReference =*/ false, /*InValueTerminalType=*/FEdGraphTerminalType()), /*bIsPersistent =*/false));
-	ActionRegistrar.AddBlueprintAction(ActionKey, MakeTempVarNodeSpawner(FEdGraphPinType(K2Schema->PC_Float, FString(), nullptr, EPinContainerType::Array, /*bIsReference =*/ false, /*InValueTerminalType=*/FEdGraphTerminalType()), /*bIsPersistent =*/false));
-	ActionRegistrar.AddBlueprintAction(ActionKey, MakeTempVarNodeSpawner(FEdGraphPinType(K2Schema->PC_Boolean, FString(), nullptr, EPinContainerType::None, /*bIsReference =*/ false, /*InValueTerminalType=*/FEdGraphTerminalType()), /*bIsPersistent =*/false));
-	ActionRegistrar.AddBlueprintAction(ActionKey, MakeTempVarNodeSpawner(FEdGraphPinType(K2Schema->PC_Boolean, FString(), nullptr, EPinContainerType::Array, /*bIsReference =*/ false, /*InValueTerminalType=*/FEdGraphTerminalType()), /*bIsPersistent =*/false));
-	ActionRegistrar.AddBlueprintAction(ActionKey, MakeTempVarNodeSpawner(FEdGraphPinType(K2Schema->PC_String, FString(), nullptr, EPinContainerType::None, /*bIsReference =*/ false, /*InValueTerminalType=*/FEdGraphTerminalType()), /*bIsPersistent =*/false));
-	ActionRegistrar.AddBlueprintAction(ActionKey, MakeTempVarNodeSpawner(FEdGraphPinType(K2Schema->PC_String, FString(), nullptr, EPinContainerType::Array, /*bIsReference =*/ false, /*InValueTerminalType=*/FEdGraphTerminalType()), /*bIsPersistent =*/false));
-	ActionRegistrar.AddBlueprintAction(ActionKey, MakeTempVarNodeSpawner(FEdGraphPinType(K2Schema->PC_Text, FString(), nullptr, EPinContainerType::None, /*bIsReference =*/ false, /*InValueTerminalType=*/FEdGraphTerminalType()), /*bIsPersistent =*/false));
-	ActionRegistrar.AddBlueprintAction(ActionKey, MakeTempVarNodeSpawner(FEdGraphPinType(K2Schema->PC_Text, FString(), nullptr, EPinContainerType::Array, /*bIsReference =*/ false, /*InValueTerminalType=*/FEdGraphTerminalType()), /*bIsPersistent =*/false));
+	ActionRegistrar.AddBlueprintAction(ActionKey, MakeTempVarNodeSpawner(FEdGraphPinType(UEdGraphSchema_K2::PC_Int, NAME_None, nullptr, EPinContainerType::None, /*bIsReference =*/ false, /*InValueTerminalType=*/FEdGraphTerminalType()), /*bIsPersistent =*/false));
+	ActionRegistrar.AddBlueprintAction(ActionKey, MakeTempVarNodeSpawner(FEdGraphPinType(UEdGraphSchema_K2::PC_Int, NAME_None, nullptr, EPinContainerType::Array, /*bIsReference =*/ false, /*InValueTerminalType=*/FEdGraphTerminalType()), /*bIsPersistent =*/false));
+	ActionRegistrar.AddBlueprintAction(ActionKey, MakeTempVarNodeSpawner(FEdGraphPinType(UEdGraphSchema_K2::PC_Float, NAME_None, nullptr, EPinContainerType::None, /*bIsReference =*/ false, /*InValueTerminalType=*/FEdGraphTerminalType()), /*bIsPersistent =*/false));
+	ActionRegistrar.AddBlueprintAction(ActionKey, MakeTempVarNodeSpawner(FEdGraphPinType(UEdGraphSchema_K2::PC_Float, NAME_None, nullptr, EPinContainerType::Array, /*bIsReference =*/ false, /*InValueTerminalType=*/FEdGraphTerminalType()), /*bIsPersistent =*/false));
+	ActionRegistrar.AddBlueprintAction(ActionKey, MakeTempVarNodeSpawner(FEdGraphPinType(UEdGraphSchema_K2::PC_Boolean, NAME_None, nullptr, EPinContainerType::None, /*bIsReference =*/ false, /*InValueTerminalType=*/FEdGraphTerminalType()), /*bIsPersistent =*/false));
+	ActionRegistrar.AddBlueprintAction(ActionKey, MakeTempVarNodeSpawner(FEdGraphPinType(UEdGraphSchema_K2::PC_Boolean, NAME_None, nullptr, EPinContainerType::Array, /*bIsReference =*/ false, /*InValueTerminalType=*/FEdGraphTerminalType()), /*bIsPersistent =*/false));
+	ActionRegistrar.AddBlueprintAction(ActionKey, MakeTempVarNodeSpawner(FEdGraphPinType(UEdGraphSchema_K2::PC_String, NAME_None, nullptr, EPinContainerType::None, /*bIsReference =*/ false, /*InValueTerminalType=*/FEdGraphTerminalType()), /*bIsPersistent =*/false));
+	ActionRegistrar.AddBlueprintAction(ActionKey, MakeTempVarNodeSpawner(FEdGraphPinType(UEdGraphSchema_K2::PC_String, NAME_None, nullptr, EPinContainerType::Array, /*bIsReference =*/ false, /*InValueTerminalType=*/FEdGraphTerminalType()), /*bIsPersistent =*/false));
+	ActionRegistrar.AddBlueprintAction(ActionKey, MakeTempVarNodeSpawner(FEdGraphPinType(UEdGraphSchema_K2::PC_Text, NAME_None, nullptr, EPinContainerType::None, /*bIsReference =*/ false, /*InValueTerminalType=*/FEdGraphTerminalType()), /*bIsPersistent =*/false));
+	ActionRegistrar.AddBlueprintAction(ActionKey, MakeTempVarNodeSpawner(FEdGraphPinType(UEdGraphSchema_K2::PC_Text, NAME_None, nullptr, EPinContainerType::Array, /*bIsReference =*/ false, /*InValueTerminalType=*/FEdGraphTerminalType()), /*bIsPersistent =*/false));
 	
-	ActionRegistrar.AddBlueprintAction(ActionKey, MakeTempVarNodeSpawner(FEdGraphPinType(K2Schema->PC_Wildcard, FString(), nullptr, EPinContainerType::None, /*bIsReference =*/ false, /*InValueTerminalType=*/FEdGraphTerminalType()), /*bIsPersistent =*/false));
-	ActionRegistrar.AddBlueprintAction(ActionKey, MakeTempVarNodeSpawner(FEdGraphPinType(K2Schema->PC_Wildcard, FString(), nullptr, EPinContainerType::Array, /*bIsReference =*/ false, /*InValueTerminalType=*/FEdGraphTerminalType()), /*bIsPersistent =*/false));
+	ActionRegistrar.AddBlueprintAction(ActionKey, MakeTempVarNodeSpawner(FEdGraphPinType(UEdGraphSchema_K2::PC_Wildcard, NAME_None, nullptr, EPinContainerType::None, /*bIsReference =*/ false, /*InValueTerminalType=*/FEdGraphTerminalType()), /*bIsPersistent =*/false));
+	ActionRegistrar.AddBlueprintAction(ActionKey, MakeTempVarNodeSpawner(FEdGraphPinType(UEdGraphSchema_K2::PC_Wildcard, NAME_None, nullptr, EPinContainerType::Array, /*bIsReference =*/ false, /*InValueTerminalType=*/FEdGraphTerminalType()), /*bIsPersistent =*/false));
 
 	UScriptStruct* VectorStruct  = TBaseStructure<FVector>::Get();
-	ActionRegistrar.AddBlueprintAction(ActionKey, MakeTempVarNodeSpawner(FEdGraphPinType(K2Schema->PC_Struct, TEXT("Vector"), VectorStruct, EPinContainerType::None, /*bIsReference =*/ false, /*InValueTerminalType=*/FEdGraphTerminalType()), /*bIsPersistent =*/false));
-	ActionRegistrar.AddBlueprintAction(ActionKey, MakeTempVarNodeSpawner(FEdGraphPinType(K2Schema->PC_Struct, TEXT("Vector"), VectorStruct, EPinContainerType::Array, /*bIsReference =*/ false, /*InValueTerminalType=*/FEdGraphTerminalType()), /*bIsPersistent =*/false));
+	ActionRegistrar.AddBlueprintAction(ActionKey, MakeTempVarNodeSpawner(FEdGraphPinType(UEdGraphSchema_K2::PC_Struct, TEXT("Vector"), VectorStruct, EPinContainerType::None, /*bIsReference =*/ false, /*InValueTerminalType=*/FEdGraphTerminalType()), /*bIsPersistent =*/false));
+	ActionRegistrar.AddBlueprintAction(ActionKey, MakeTempVarNodeSpawner(FEdGraphPinType(UEdGraphSchema_K2::PC_Struct, TEXT("Vector"), VectorStruct, EPinContainerType::Array, /*bIsReference =*/ false, /*InValueTerminalType=*/FEdGraphTerminalType()), /*bIsPersistent =*/false));
 	
 	UScriptStruct* RotatorStruct = TBaseStructure<FRotator>::Get();
-	ActionRegistrar.AddBlueprintAction(ActionKey, MakeTempVarNodeSpawner(FEdGraphPinType(K2Schema->PC_Struct, TEXT("Rotator"), RotatorStruct, EPinContainerType::None, /*bIsReference =*/ false, /*InValueTerminalType=*/FEdGraphTerminalType()), /*bIsPersistent =*/false));
-	ActionRegistrar.AddBlueprintAction(ActionKey, MakeTempVarNodeSpawner(FEdGraphPinType(K2Schema->PC_Struct, TEXT("Rotator"), RotatorStruct, EPinContainerType::Array, /*bIsReference =*/ false, /*InValueTerminalType=*/FEdGraphTerminalType()), /*bIsPersistent =*/false));
+	ActionRegistrar.AddBlueprintAction(ActionKey, MakeTempVarNodeSpawner(FEdGraphPinType(UEdGraphSchema_K2::PC_Struct, TEXT("Rotator"), RotatorStruct, EPinContainerType::None, /*bIsReference =*/ false, /*InValueTerminalType=*/FEdGraphTerminalType()), /*bIsPersistent =*/false));
+	ActionRegistrar.AddBlueprintAction(ActionKey, MakeTempVarNodeSpawner(FEdGraphPinType(UEdGraphSchema_K2::PC_Struct, TEXT("Rotator"), RotatorStruct, EPinContainerType::Array, /*bIsReference =*/ false, /*InValueTerminalType=*/FEdGraphTerminalType()), /*bIsPersistent =*/false));
 	
 	UScriptStruct* TransformStruct = TBaseStructure<FTransform>::Get();
-	ActionRegistrar.AddBlueprintAction(ActionKey, MakeTempVarNodeSpawner(FEdGraphPinType(K2Schema->PC_Struct, TEXT("Transform"), TransformStruct, EPinContainerType::None, /*bIsReference =*/ false, /*InValueTerminalType=*/FEdGraphTerminalType()), /*bIsPersistent =*/false));
-	ActionRegistrar.AddBlueprintAction(ActionKey, MakeTempVarNodeSpawner(FEdGraphPinType(K2Schema->PC_Struct, TEXT("Transform"), TransformStruct, EPinContainerType::Array, /*bIsReference =*/ false, /*InValueTerminalType=*/FEdGraphTerminalType()), /*bIsPersistent =*/false));
+	ActionRegistrar.AddBlueprintAction(ActionKey, MakeTempVarNodeSpawner(FEdGraphPinType(UEdGraphSchema_K2::PC_Struct, TEXT("Transform"), TransformStruct, EPinContainerType::None, /*bIsReference =*/ false, /*InValueTerminalType=*/FEdGraphTerminalType()), /*bIsPersistent =*/false));
+	ActionRegistrar.AddBlueprintAction(ActionKey, MakeTempVarNodeSpawner(FEdGraphPinType(UEdGraphSchema_K2::PC_Struct, TEXT("Transform"), TransformStruct, EPinContainerType::Array, /*bIsReference =*/ false, /*InValueTerminalType=*/FEdGraphTerminalType()), /*bIsPersistent =*/false));
 
 	UScriptStruct* BlendSampleStruct = FindObjectChecked<UScriptStruct>(ANY_PACKAGE, TEXT("BlendSampleData"));
-	ActionRegistrar.AddBlueprintAction(ActionKey, MakeTempVarNodeSpawner(FEdGraphPinType(K2Schema->PC_Struct, TEXT("BlendSampleData"), BlendSampleStruct, EPinContainerType::None, /*bIsReference =*/ false, /*InValueTerminalType=*/FEdGraphTerminalType()), /*bIsPersistent =*/false));
-	ActionRegistrar.AddBlueprintAction(ActionKey, MakeTempVarNodeSpawner(FEdGraphPinType(K2Schema->PC_Struct, TEXT("BlendSampleData"), BlendSampleStruct, EPinContainerType::Array, /*bIsReference =*/ false, /*InValueTerminalType=*/FEdGraphTerminalType()), /*bIsPersistent =*/false));
+	ActionRegistrar.AddBlueprintAction(ActionKey, MakeTempVarNodeSpawner(FEdGraphPinType(UEdGraphSchema_K2::PC_Struct, TEXT("BlendSampleData"), BlendSampleStruct, EPinContainerType::None, /*bIsReference =*/ false, /*InValueTerminalType=*/FEdGraphTerminalType()), /*bIsPersistent =*/false));
+	ActionRegistrar.AddBlueprintAction(ActionKey, MakeTempVarNodeSpawner(FEdGraphPinType(UEdGraphSchema_K2::PC_Struct, TEXT("BlendSampleData"), BlendSampleStruct, EPinContainerType::Array, /*bIsReference =*/ false, /*InValueTerminalType=*/FEdGraphTerminalType()), /*bIsPersistent =*/false));
 
 	// add persistent bool and int types (for macro graphs)
-	ActionRegistrar.AddBlueprintAction(ActionKey, MakeTempVarNodeSpawner(FEdGraphPinType(K2Schema->PC_Int, FString(), nullptr, EPinContainerType::None, /*bIsReference =*/ false, /*InValueTerminalType=*/FEdGraphTerminalType()), /*bIsPersistent =*/true));
-	ActionRegistrar.AddBlueprintAction(ActionKey, MakeTempVarNodeSpawner(FEdGraphPinType(K2Schema->PC_Boolean, FString(), nullptr, EPinContainerType::None, /*bIsReference =*/ false, /*InValueTerminalType=*/FEdGraphTerminalType()), /*bIsPersistent =*/true));
+	ActionRegistrar.AddBlueprintAction(ActionKey, MakeTempVarNodeSpawner(FEdGraphPinType(UEdGraphSchema_K2::PC_Int, NAME_None, nullptr, EPinContainerType::None, /*bIsReference =*/ false, /*InValueTerminalType=*/FEdGraphTerminalType()), /*bIsPersistent =*/true));
+	ActionRegistrar.AddBlueprintAction(ActionKey, MakeTempVarNodeSpawner(FEdGraphPinType(UEdGraphSchema_K2::PC_Boolean, NAME_None, nullptr, EPinContainerType::None, /*bIsReference =*/ false, /*InValueTerminalType=*/FEdGraphTerminalType()), /*bIsPersistent =*/true));
 }
 
 FText UK2Node_TemporaryVariable::GetMenuCategory() const
diff --git a/Engine/Source/Editor/BlueprintGraph/Private/K2Node_Timeline.cpp b/Engine/Source/Editor/BlueprintGraph/Private/K2Node_Timeline.cpp
--- a/Engine/Source/Editor/BlueprintGraph/Private/K2Node_Timeline.cpp
+++ b/Engine/Source/Editor/BlueprintGraph/Private/K2Node_Timeline.cpp
@@ -37,23 +37,23 @@ UK2Node_Timeline::UK2Node_Timeline(const FObjectInitializer& ObjectInitializer)
 	bIgnoreTimeDilation = false;
 }
 
-static FString PlayPinName(TEXT("Play"));
-static FString PlayFromStartPinName(TEXT("PlayFromStart"));
-static FString StopPinName(TEXT("Stop"));
-static FString UpdatePinName(TEXT("Update"));
-static FString ReversePinName(TEXT("Reverse"));
-static FString ReverseFromEndPinName(TEXT("ReverseFromEnd"));
-static FString FinishedPinName(TEXT("Finished"));
-static FString NewTimePinName(TEXT("NewTime"));
-static FString SetNewTimePinName(TEXT("SetNewTime"));
-static FString DirectionPinName(TEXT("Direction"));
+static FName PlayPinName(TEXT("Play"));
+static FName PlayFromStartPinName(TEXT("PlayFromStart"));
+static FName StopPinName(TEXT("Stop"));
+static FName UpdatePinName(TEXT("Update"));
+static FName ReversePinName(TEXT("Reverse"));
+static FName ReverseFromEndPinName(TEXT("ReverseFromEnd"));
+static FName FinishedPinName(TEXT("Finished"));
+static FName NewTimePinName(TEXT("NewTime"));
+static FName SetNewTimePinName(TEXT("SetNewTime"));
+static FName DirectionPinName(TEXT("Direction"));
 
 namespace 
 {
-	UEdGraphPin* GetPin (const UK2Node_Timeline* Timeline, const FString& PinName, EEdGraphPinDirection DesiredDirection) 
+	UEdGraphPin* GetPin (const UK2Node_Timeline* Timeline, const FName PinName, EEdGraphPinDirection DesiredDirection) 
 	{
 		UEdGraphPin* Pin = Timeline->FindPin(PinName);
-		check(Pin != NULL);
+		check(Pin);
 		check(Pin->Direction == DesiredDirection);
 		return Pin;
 	}
@@ -64,21 +64,21 @@ void UK2Node_Timeline::AllocateDefaultPins()
 	const UEdGraphSchema_K2* K2Schema = GetDefault<UEdGraphSchema_K2>();
 	bCanRenameNode = 1;
 
-	CreatePin(EGPD_Input, K2Schema->PC_Exec, FString(), nullptr, PlayPinName);
-	CreatePin(EGPD_Input, K2Schema->PC_Exec, FString(), nullptr, PlayFromStartPinName);
-	CreatePin(EGPD_Input, K2Schema->PC_Exec, FString(), nullptr, StopPinName);
-	CreatePin(EGPD_Input,  K2Schema->PC_Exec, FString(), nullptr, ReversePinName);
-	CreatePin(EGPD_Input,  K2Schema->PC_Exec, FString(), nullptr, ReverseFromEndPinName);
+	CreatePin(EGPD_Input, UEdGraphSchema_K2::PC_Exec, PlayPinName);
+	CreatePin(EGPD_Input, UEdGraphSchema_K2::PC_Exec, PlayFromStartPinName);
+	CreatePin(EGPD_Input, UEdGraphSchema_K2::PC_Exec, StopPinName);
+	CreatePin(EGPD_Input,  UEdGraphSchema_K2::PC_Exec, ReversePinName);
+	CreatePin(EGPD_Input,  UEdGraphSchema_K2::PC_Exec, ReverseFromEndPinName);
 
-	CreatePin(EGPD_Output, K2Schema->PC_Exec, FString(), nullptr, UpdatePinName);
-	CreatePin(EGPD_Output, K2Schema->PC_Exec, FString(), nullptr, FinishedPinName);
+	CreatePin(EGPD_Output, UEdGraphSchema_K2::PC_Exec, UpdatePinName);
+	CreatePin(EGPD_Output, UEdGraphSchema_K2::PC_Exec, FinishedPinName);
 
-	CreatePin(EGPD_Input, K2Schema->PC_Exec, FString(), nullptr, SetNewTimePinName);
+	CreatePin(EGPD_Input, UEdGraphSchema_K2::PC_Exec, SetNewTimePinName);
 
-	UEdGraphPin* NewPositionPin = CreatePin(EGPD_Input, K2Schema->PC_Float, FString(), nullptr, NewTimePinName);
+	UEdGraphPin* NewPositionPin = CreatePin(EGPD_Input, UEdGraphSchema_K2::PC_Float, NewTimePinName);
 	K2Schema->SetPinAutogeneratedDefaultValue(NewPositionPin, TEXT("0.0"));
 
-	CreatePin(EGPD_Output, K2Schema->PC_Byte, FString(), FTimeline::GetTimelineDirectionEnum(), DirectionPinName);
+	CreatePin(EGPD_Output, UEdGraphSchema_K2::PC_Byte, FTimeline::GetTimelineDirectionEnum(), DirectionPinName);
 
 	UBlueprint* Blueprint = GetBlueprint();
 	check(Blueprint);
@@ -89,30 +89,26 @@ void UK2Node_Timeline::AllocateDefaultPins()
 		// Ensure the timeline template is fully loaded or the node representation will be wrong.
 		PreloadObject(Timeline);
 
-		for(int32 i=0; i<Timeline->FloatTracks.Num(); i++)
+		for (const FTTFloatTrack& FloatTrack : Timeline->FloatTracks)
 		{
-			FTTFloatTrack& FloatTrack = Timeline->FloatTracks[i];
-			CreatePin(EGPD_Output, K2Schema->PC_Float, FString(), nullptr, FloatTrack.TrackName.ToString());			
+			CreatePin(EGPD_Output, UEdGraphSchema_K2::PC_Float, FloatTrack.TrackName);			
 		}
 
-		for(int32 i=0; i<Timeline->VectorTracks.Num(); i++)
+		UScriptStruct* VectorStruct = TBaseStructure<FVector>::Get();
+		for (const FTTVectorTrack& VectorTrack : Timeline->VectorTracks)
 		{
-			FTTVectorTrack& VectorTrack = Timeline->VectorTracks[i];
-			UScriptStruct* VectorStruct = TBaseStructure<FVector>::Get();
-			CreatePin(EGPD_Output, K2Schema->PC_Struct, FString(), VectorStruct, VectorTrack.TrackName.ToString());			
+			CreatePin(EGPD_Output, UEdGraphSchema_K2::PC_Struct, VectorStruct, VectorTrack.TrackName);			
 		}
 
-		for(int32 i=0; i<Timeline->LinearColorTracks.Num(); i++)
+		UScriptStruct* LinearColorStruct = TBaseStructure<FLinearColor>::Get();
+		for (const FTTLinearColorTrack& LinearColorTrack : Timeline->LinearColorTracks)
 		{
-			FTTLinearColorTrack& LinearColorTrack = Timeline->LinearColorTracks[i];
-			UScriptStruct* LinearColorStruct = TBaseStructure<FLinearColor>::Get();
-			CreatePin(EGPD_Output, K2Schema->PC_Struct, FString(), LinearColorStruct, LinearColorTrack.TrackName.ToString());			
+			CreatePin(EGPD_Output, UEdGraphSchema_K2::PC_Struct, LinearColorStruct, LinearColorTrack.TrackName);			
 		}
 
-		for(int32 i=0; i<Timeline->EventTracks.Num(); i++)
+		for (const FTTEventTrack& EventTrack : Timeline->EventTracks)
 		{
-			FTTEventTrack& EventTrack = Timeline->EventTracks[i];
-			CreatePin(EGPD_Output, K2Schema->PC_Exec, FString(), nullptr, EventTrack.TrackName.ToString());			
+			CreatePin(EGPD_Output, UEdGraphSchema_K2::PC_Exec, EventTrack.TrackName);
 		}
 
 		// cache play status
@@ -319,14 +315,11 @@ FText UK2Node_Timeline::GetNodeTitle(ENodeTitleType::Type TitleType) const
 
 UEdGraphPin* UK2Node_Timeline::GetDirectionPin() const
 {
-	const UEdGraphSchema_K2* K2Schema = GetDefault<UEdGraphSchema_K2>();
-	check(NULL != K2Schema);
-
 	UEdGraphPin* Pin = FindPin(DirectionPinName);
-	if(Pin && K2Schema)
+	if (Pin)
 	{
 		const bool bIsOutput = (EGPD_Output == Pin->Direction);
-		const bool bProperType = (K2Schema->PC_Byte == Pin->PinType.PinCategory);
+		const bool bProperType = (UEdGraphSchema_K2::PC_Byte == Pin->PinType.PinCategory);
 		const bool bProperSubCategoryObj = (Pin->PinType.PinSubCategoryObject == FTimeline::GetTimelineDirectionEnum());
 		if(bIsOutput && bProperType && bProperSubCategoryObj)
 		{
@@ -584,19 +577,19 @@ void UK2Node_Timeline::ExpandNode(FKismetCompilerContext& CompilerContext, UEdGr
 		for(int32 i=0; i<Timeline->FloatTracks.Num(); i++)
 		{
 			const FName TrackName = Timeline->FloatTracks[i].TrackName;
-			ExpandForPin(FindPin(TrackName.ToString()), Timeline->GetTrackPropertyName(TrackName), CompilerContext, SourceGraph);
+			ExpandForPin(FindPin(TrackName), Timeline->GetTrackPropertyName(TrackName), CompilerContext, SourceGraph);
 		}
 
 		for(int32 i=0; i<Timeline->VectorTracks.Num(); i++)
 		{
 			const FName TrackName = Timeline->VectorTracks[i].TrackName;
-			ExpandForPin(FindPin(TrackName.ToString()), Timeline->GetTrackPropertyName(TrackName), CompilerContext, SourceGraph);
+			ExpandForPin(FindPin(TrackName), Timeline->GetTrackPropertyName(TrackName), CompilerContext, SourceGraph);
 		}
 
 		for(int32 i=0; i<Timeline->LinearColorTracks.Num(); i++)
 		{
 			const FName TrackName = Timeline->LinearColorTracks[i].TrackName;
-			ExpandForPin(FindPin(TrackName.ToString()), Timeline->GetTrackPropertyName(TrackName), CompilerContext, SourceGraph);
+			ExpandForPin(FindPin(TrackName), Timeline->GetTrackPropertyName(TrackName), CompilerContext, SourceGraph);
 		}
 	}
 }
diff --git a/Engine/Source/Editor/BlueprintGraph/Private/K2Node_Tunnel.cpp b/Engine/Source/Editor/BlueprintGraph/Private/K2Node_Tunnel.cpp
--- a/Engine/Source/Editor/BlueprintGraph/Private/K2Node_Tunnel.cpp
+++ b/Engine/Source/Editor/BlueprintGraph/Private/K2Node_Tunnel.cpp
@@ -69,7 +69,7 @@ FText UK2Node_Tunnel::GetNodeTitle(ENodeTitleType::Type TitleType) const
 	}
 }
 
-FString UK2Node_Tunnel::CreateUniquePinName(FString InSourcePinName) const
+FName UK2Node_Tunnel::CreateUniquePinName(FName InSourcePinName) const
 {
 	if (GetClass() == UK2Node_Tunnel::StaticClass())
 	{
@@ -80,13 +80,13 @@ FString UK2Node_Tunnel::CreateUniquePinName(FString InSourcePinName) const
 
 		if (TunnelEntry.IsValid() && TunnelResult.IsValid())
 		{
-			FString PinName(InSourcePinName);
+			FName PinName(InSourcePinName);
 
 			int32 Index = 1;
 			while (TunnelEntry.Get()->FindPin(PinName) != nullptr || TunnelResult.Get()->FindPin(PinName) != nullptr)
 			{
 				++Index;
-				PinName = InSourcePinName + FString::FromInt(Index);
+				PinName = *FString::Printf(TEXT("%s%d"),*InSourcePinName.ToString(), Index);
 			}
 
 			return PinName;
@@ -159,8 +159,7 @@ bool UK2Node_Tunnel::CanCreateUserDefinedPin(const FEdGraphPinType& InPinType, E
 {
 	bool bResult = true;
 	// Make sure that if this is an exec node we are allowed one.
-	const UEdGraphSchema_K2* Schema = GetDefault<UEdGraphSchema_K2>();
-	if (InPinType.PinCategory == Schema->PC_Exec && !CanModifyExecutionWires())
+	if (InPinType.PinCategory == UEdGraphSchema_K2::PC_Exec && !CanModifyExecutionWires())
 	{
 		OutErrorMessage = LOCTEXT("MultipleExecPinError", "Cannot support more exec pins!");
 		bResult = false;
@@ -194,11 +193,11 @@ UEdGraphPin* UK2Node_Tunnel::CreatePinFromUserDefinition(const TSharedPtr<FUserP
 	Schema->SetPinAutogeneratedDefaultValue(Result, NewPinInfo->PinDefaultValue);
 
 	// Make sure it mirrors onto the associated node
-	UEdGraphNode* TargetNode = ((InputSinkNode != NULL) ? InputSinkNode : OutputSourceNode);
-	if (Cast<UK2Node_Composite>(TargetNode) != NULL || Cast<UK2Node_MacroInstance>(TargetNode) != NULL)
+	UEdGraphNode* TargetNode = (InputSinkNode ? InputSinkNode : OutputSourceNode);
+	if (Cast<UK2Node_Composite>(TargetNode) || Cast<UK2Node_MacroInstance>(TargetNode))
 	{
 		UEdGraphPin* HasPinAlready = TargetNode->FindPin(Result->PinName);
-		if (HasPinAlready == NULL)
+		if (HasPinAlready == nullptr)
 		{
 			TargetNode->CreatePin(UEdGraphPin::GetComplementaryDirection(Direction), NewPinInfo->PinType, NewPinInfo->PinName);
 		}
@@ -230,7 +229,7 @@ bool UK2Node_Tunnel::CanModifyExecutionWires()
 	return true;
 }
 
-ERenamePinResult UK2Node_Tunnel::RenameUserDefinedPin(const FString& OldName, const FString& NewName, bool bTest)
+ERenamePinResult UK2Node_Tunnel::RenameUserDefinedPin(const FName OldName, const FName NewName, bool bTest)
 {
 	const ERenamePinResult ThisNodeResult = Super::RenameUserDefinedPin(OldName, NewName, bTest);
 	if(ERenamePinResult::ERenamePinResult_NameCollision == ThisNodeResult)
@@ -269,8 +268,6 @@ void UK2Node_Tunnel::ReallocatePinsDuringReconstruction(TArray<UEdGraphPin*>& Ol
 {
 	Super::ReallocatePinsDuringReconstruction(OldPins);
 
-	const UEdGraphSchema_K2* const Schema = GetDefault<UEdGraphSchema_K2>();
-
 	// determine if all wildcard pins are unlinked.
 	// if they are, we should revert them all back to wildcard status
 	bool bAllWildcardsAreUnlinked = true;
@@ -278,7 +275,7 @@ void UK2Node_Tunnel::ReallocatePinsDuringReconstruction(TArray<UEdGraphPin*>& Ol
 	{
 		// for each of the wildcard pins...
 		UEdGraphPin* const Pin = *PinIt;
-		if (Pin->PinType.PinCategory == Schema->PC_Wildcard)
+		if (Pin->PinType.PinCategory == UEdGraphSchema_K2::PC_Wildcard)
 		{
 			// find it in the old pins array (where it might not be a wildcard)
 			// and see if it's unlinked
@@ -303,7 +300,7 @@ void UK2Node_Tunnel::ReallocatePinsDuringReconstruction(TArray<UEdGraphPin*>& Ol
 		for (auto PinIt = Pins.CreateConstIterator(); PinIt; PinIt++)
 		{
 			UEdGraphPin* const Pin = *PinIt;
-			if (Pin->PinType.PinCategory == Schema->PC_Wildcard)
+			if (Pin->PinType.PinCategory == UEdGraphSchema_K2::PC_Wildcard)
 			{
 				// find it in the old pins and copy the type
 				for (auto OldPinIt = OldPins.CreateConstIterator(); OldPinIt; OldPinIt++)
diff --git a/Engine/Source/Editor/BlueprintGraph/Private/K2Node_TunnelBoundary.cpp b/Engine/Source/Editor/BlueprintGraph/Private/K2Node_TunnelBoundary.cpp
--- a/Engine/Source/Editor/BlueprintGraph/Private/K2Node_TunnelBoundary.cpp
+++ b/Engine/Source/Editor/BlueprintGraph/Private/K2Node_TunnelBoundary.cpp
@@ -158,7 +158,7 @@ void UK2Node_TunnelBoundary::CreateBoundaryNodesForTunnelInstance(UK2Node_Tunnel
 		TunnelBoundaryNode->TunnelBoundaryType = TBT_EndOfThread;
 		MessageLog.RegisterIntermediateTunnelNode(TunnelBoundaryNode, RegistrantTunnelInstance);
 		// Create exit point pin
-		UEdGraphPin* BoundaryTerminalPin = TunnelBoundaryNode->CreatePin(EGPD_Input, UEdGraphSchema_K2::PC_Exec, FString(), nullptr, TEXT("TunnelEndOfThread"));
+		UEdGraphPin* BoundaryTerminalPin = TunnelBoundaryNode->CreatePin(EGPD_Input, UEdGraphSchema_K2::PC_Exec, TEXT("TunnelEndOfThread"));
 		// Wire up the endpoint unwired links to the boundary node.
 		for (UEdGraphPin* TerminationPin : ExecutionEndpointPins)
 		{
@@ -259,12 +259,12 @@ void UK2Node_TunnelBoundary::WireUpTunnelEntry(UK2Node_Tunnel* TunnelInstance, U
 		// Find the true source pin
 		UEdGraphPin* SourcePin = TunnelInstance->FindPin(TunnelPin->PinName);
 		// Wire in the node
-		UEdGraphPin* OutputPin = CreatePin(EGPD_Output, UEdGraphSchema_K2::PC_Exec, FString(), nullptr, TEXT("TunnelEntryExec"));
+		UEdGraphPin* OutputPin = CreatePin(EGPD_Output, UEdGraphSchema_K2::PC_Exec, TEXT("TunnelEntryExec"));
 		for (UEdGraphPin* LinkedPin : TunnelPin->LinkedTo)
 		{
 			LinkedPin->MakeLinkTo(OutputPin);
 		}
-		UEdGraphPin* InputPin = CreatePin(EGPD_Input, UEdGraphSchema_K2::PC_Exec, FString(), nullptr, SourcePin->PinName);
+		UEdGraphPin* InputPin = CreatePin(EGPD_Input, UEdGraphSchema_K2::PC_Exec, SourcePin->PinName);
 		MessageLog.NotifyIntermediatePinCreation(InputPin, SourcePin);
 		TunnelPin->BreakAllPinLinks();
 		TunnelPin->MakeLinkTo(InputPin);
@@ -282,12 +282,12 @@ void UK2Node_TunnelBoundary::WireUpTunnelExit(UK2Node_Tunnel* TunnelInstance, UE
 		// Find the true source pin
 		UEdGraphPin* SourcePin = TunnelInstance->FindPin(TunnelPin->PinName);
 		// Wire in the node
-		UEdGraphPin* InputPin = CreatePin(EGPD_Input, UEdGraphSchema_K2::PC_Exec, FString(), nullptr, SourcePin->PinName);
+		UEdGraphPin* InputPin = CreatePin(EGPD_Input, UEdGraphSchema_K2::PC_Exec, SourcePin->PinName);
 		for (UEdGraphPin* LinkedPin : TunnelPin->LinkedTo)
 		{
 			LinkedPin->MakeLinkTo(InputPin);
 		}
-		UEdGraphPin* OutputPin = CreatePin(EGPD_Output, UEdGraphSchema_K2::PC_Exec, FString(), nullptr, TEXT("TunnelExitExec"));
+		UEdGraphPin* OutputPin = CreatePin(EGPD_Output, UEdGraphSchema_K2::PC_Exec, TEXT("TunnelExitExec"));
 		MessageLog.NotifyIntermediatePinCreation(OutputPin, SourcePin);
 		TunnelPin->BreakAllPinLinks();
 		TunnelPin->MakeLinkTo(OutputPin);
@@ -303,9 +303,10 @@ void UK2Node_TunnelBoundary::WireUpEntry(UEdGraphNode* SourceNode, UEdGraphPin*
 		// Create the pin base name
 		CreateBaseNodeName(SourceNode);
 		// Wire in the node
-		UEdGraphPin* InputPin = CreatePin(EGPD_Input, SourcePin->PinType.PinCategory, SourcePin->PinType.PinSubCategory, SourcePin->PinType.PinSubCategoryObject.Get(), TEXT("EntryBoundary"), SourcePin->PinType.ContainerType, SourcePin->PinType.bIsReference, SourcePin->PinType.bIsConst, INDEX_NONE, SourcePin->PinType.PinValueType);
+		const UEdGraphNode::FCreatePinParams PinParams(SourcePin->PinType);
+		UEdGraphPin* InputPin = CreatePin(EGPD_Input, SourcePin->PinType.PinCategory, SourcePin->PinType.PinSubCategory, SourcePin->PinType.PinSubCategoryObject.Get(), TEXT("EntryBoundary"), PinParams);
 		MessageLog.NotifyIntermediatePinCreation(InputPin, SourcePin);
-		UEdGraphPin* OutputPin = CreatePin(EGPD_Output, SourcePin->PinType.PinCategory, SourcePin->PinType.PinSubCategory, SourcePin->PinType.PinSubCategoryObject.Get(), SourcePin->PinName, SourcePin->PinType.ContainerType, SourcePin->PinType.bIsReference, SourcePin->PinType.bIsConst, INDEX_NONE, SourcePin->PinType.PinValueType);
+		UEdGraphPin* OutputPin = CreatePin(EGPD_Output, SourcePin->PinType.PinCategory, SourcePin->PinType.PinSubCategory, SourcePin->PinType.PinSubCategoryObject.Get(), SourcePin->PinName, PinParams);
 		for (UEdGraphPin* EntryPin : EntryPins)
 		{
 			for (UEdGraphPin* LinkedPin : EntryPin->LinkedTo)
@@ -329,9 +330,10 @@ void UK2Node_TunnelBoundary::WireUpExit(UEdGraphNode* SourceNode, UEdGraphPin* S
 		// Create the pin base name
 		CreateBaseNodeName(SourceNode);
 		// Wire in the node
-		UEdGraphPin* OutputPin = CreatePin(EGPD_Output, SourcePin->PinType.PinCategory, SourcePin->PinType.PinSubCategory, SourcePin->PinType.PinSubCategoryObject.Get(), SourcePin->PinName, SourcePin->PinType.ContainerType, SourcePin->PinType.bIsReference, SourcePin->PinType.bIsConst, INDEX_NONE, SourcePin->PinType.PinValueType);
+		const UEdGraphNode::FCreatePinParams PinParams(SourcePin->PinType);
+		UEdGraphPin* OutputPin = CreatePin(EGPD_Output, SourcePin->PinType.PinCategory, SourcePin->PinType.PinSubCategory, SourcePin->PinType.PinSubCategoryObject.Get(), SourcePin->PinName, PinParams);
 		MessageLog.NotifyIntermediatePinCreation(OutputPin, SourcePin);
-		UEdGraphPin* InputPin = CreatePin(EGPD_Input, SourcePin->PinType.PinCategory, SourcePin->PinType.PinSubCategory, SourcePin->PinType.PinSubCategoryObject.Get(), TEXT("ExitBoundary"), SourcePin->PinType.ContainerType, SourcePin->PinType.bIsReference, SourcePin->PinType.bIsConst, INDEX_NONE, SourcePin->PinType.PinValueType);
+		UEdGraphPin* InputPin = CreatePin(EGPD_Input, SourcePin->PinType.PinCategory, SourcePin->PinType.PinSubCategory, SourcePin->PinType.PinSubCategoryObject.Get(), TEXT("ExitBoundary"), PinParams);
 		for (UEdGraphPin* ExitPin : ExitPins)
 		{
 			for (UEdGraphPin* LinkedPin : ExitPin->LinkedTo)
diff --git a/Engine/Source/Editor/BlueprintGraph/Private/K2Node_Variable.cpp b/Engine/Source/Editor/BlueprintGraph/Private/K2Node_Variable.cpp
--- a/Engine/Source/Editor/BlueprintGraph/Private/K2Node_Variable.cpp
+++ b/Engine/Source/Editor/BlueprintGraph/Private/K2Node_Variable.cpp
@@ -59,7 +59,7 @@ void UK2Node_Variable::SetFromProperty(const UProperty* Property, bool bSelfCont
 	VariableReference.SetFromField<UProperty>(Property, bSelfContext);
 }
 
-bool UK2Node_Variable::CreatePinForVariable(EEdGraphPinDirection Direction, FString InPinName/* = FString()*/)
+bool UK2Node_Variable::CreatePinForVariable(EEdGraphPinDirection Direction, FName InPinName/* = NAME_None */)
 {
 	const UEdGraphSchema_K2* K2Schema = GetDefault<UEdGraphSchema_K2>();
 
@@ -88,9 +88,9 @@ bool UK2Node_Variable::CreatePinForVariable(EEdGraphPinDirection Direction, FStr
 
 	if (VariableProperty != nullptr)
 	{
-		const FString PinName = InPinName.IsEmpty()? GetVarNameString() : InPinName;
+		const FName PinName = InPinName.IsNone()? GetVarName() : InPinName;
 		// Create the pin
-		UEdGraphPin* VariablePin = CreatePin(Direction, FString(), FString(), nullptr, PinName);
+		UEdGraphPin* VariablePin = CreatePin(Direction, NAME_None, PinName);
 		K2Schema->ConvertPropertyToPinType(VariableProperty, /*out*/ VariablePin->PinType);
 		K2Schema->SetPinAutogeneratedDefaultValueBasedOnType(VariablePin);
 	}
@@ -163,7 +163,7 @@ void UK2Node_Variable::CreatePinForSelf()
 				TargetClass = TargetClass->GetAuthoritativeClass();
 			}
 
-			UEdGraphPin* TargetPin = CreatePin(EGPD_Input, K2Schema->PC_Object, FString(), TargetClass, K2Schema->PN_Self);
+			UEdGraphPin* TargetPin = CreatePin(EGPD_Input, UEdGraphSchema_K2::PC_Object, TargetClass, UEdGraphSchema_K2::PN_Self);
 			TargetPin->PinFriendlyName =  LOCTEXT("Target", "Target");
 
 			if (bSelfTarget)
@@ -178,41 +178,38 @@ void UK2Node_Variable::CreatePinForSelf()
 	}
 }
 
-bool UK2Node_Variable::RecreatePinForVariable(EEdGraphPinDirection Direction, TArray<UEdGraphPin*>& OldPins, FString InPinName/* = FString()*/)
+bool UK2Node_Variable::RecreatePinForVariable(EEdGraphPinDirection Direction, TArray<UEdGraphPin*>& OldPins, FName InPinName/* = NAME_None*/)
 {
 	// probably the node was pasted to a blueprint without the variable
 	// we don't want to beak any connection, so the pin will be recreated from old one, but compiler will throw error
 
 	// find old variable pin
-	const UEdGraphPin* OldVariablePin = NULL;
-	const FString PinName = InPinName.IsEmpty()? GetVarNameString() : InPinName;
-	for(auto Iter = OldPins.CreateConstIterator(); Iter; ++Iter)
+	const UEdGraphPin* OldVariablePin = nullptr;
+	const FName PinName = InPinName.IsNone() ? GetVarName() : InPinName;
+	for (const UEdGraphPin* Pin : OldPins)
 	{
-		if(const UEdGraphPin* Pin = *Iter)
+		if (Pin && PinName == Pin->PinName)
 		{
-			if(PinName == Pin->PinName)
-			{
-				OldVariablePin = Pin;
-				break;
-			}
+			OldVariablePin = Pin;
+			break;
 		}
 	}
 
-	if(nullptr != OldVariablePin)
+	if (nullptr != OldVariablePin)
 	{
 		// create new pin from old one
-		UEdGraphPin* VariablePin = CreatePin(Direction, FString(), FString(), nullptr, PinName);
+		UEdGraphPin* VariablePin = CreatePin(Direction, NAME_None, PinName);
 		VariablePin->PinType = OldVariablePin->PinType;
 		
 		const UEdGraphSchema_K2* K2Schema = GetDefault<UEdGraphSchema_K2>();
 		K2Schema->SetPinAutogeneratedDefaultValueBasedOnType(VariablePin);
 
-		Message_Note(*FString::Printf(TEXT("Pin for variable '%s' recreated, but the variable is missing."), *PinName));
+		Message_Note(*FString::Printf(TEXT("Pin for variable '%s' recreated, but the variable is missing."), *PinName.ToString()));
 		return true;
 	}
 	else
 	{
-		Message_Warn(*FString::Printf(TEXT("RecreatePinForVariable: '%s' pin not found"), *PinName));
+		Message_Warn(*FString::Printf(TEXT("RecreatePinForVariable: '%s' pin not found"), *PinName.ToString()));
 		return false;
 	}
 }
@@ -251,13 +248,13 @@ FString UK2Node_Variable::GetFindReferenceSearchString() const
 UK2Node::ERedirectType UK2Node_Variable::DoPinsMatchForReconstruction( const UEdGraphPin* NewPin, int32 NewPinIndex, const UEdGraphPin* OldPin, int32 OldPinIndex ) const 
 {
 	const UEdGraphSchema_K2* K2Schema = GetDefault<UEdGraphSchema_K2>();
-	if( OldPin->PinType.PinCategory == K2Schema->PC_Exec )
+	if( OldPin->PinType.PinCategory == UEdGraphSchema_K2::PC_Exec )
 	{
 		return Super::DoPinsMatchForReconstruction(NewPin, NewPinIndex, OldPin, OldPinIndex);
 	}
 
 	const bool bPinNamesMatch = (OldPin->PinName == NewPin->PinName);
-	const bool bCanMatchSelfs = bPinNamesMatch || ((OldPin->PinName == K2Schema->PN_Self) == (NewPin->PinName == K2Schema->PN_Self));
+	const bool bCanMatchSelfs = bPinNamesMatch || ((OldPin->PinName == UEdGraphSchema_K2::PN_Self) == (NewPin->PinName == UEdGraphSchema_K2::PN_Self));
 	const bool bTheSameDirection = (NewPin->Direction == OldPin->Direction);
 
 	if (bCanMatchSelfs && bTheSameDirection)
@@ -308,8 +305,8 @@ UK2Node::ERedirectType UK2Node_Variable::DoPinsMatchForReconstruction( const UEd
 				}
 
 				// Compare whether the names, ignoring the original variable's name in the case of renames, match
-				FString NewPinPropertyName = NewPin->PinName.RightChop(ParentmostNewPin->PinName.Len() + 1);
-				FString OldPinPropertyName = OldPin->PinName.RightChop(ParentmostOldPin->PinName.Len() + 1);
+				FString NewPinPropertyName = NewPin->PinName.ToString().RightChop(ParentmostNewPin->PinName.ToString().Len() + 1);
+				FString OldPinPropertyName = OldPin->PinName.ToString().RightChop(ParentmostOldPin->PinName.ToString().Len() + 1);
 
 				if (NewPinPropertyName != OldPinPropertyName)
 				{
@@ -321,13 +318,13 @@ UK2Node::ERedirectType UK2Node_Variable::DoPinsMatchForReconstruction( const UEd
 		}
 		else
 		{
-			const bool bNewPinIsObject = (NewPin->PinType.PinCategory == K2Schema->PC_Object);
+			const bool bNewPinIsObject = (NewPin->PinType.PinCategory == UEdGraphSchema_K2::PC_Object);
 
 			// Special Case: If we had a pin match, and the class isn't loaded 
 			//               yet because of a cyclic dependency, temporarily 
 			//               cast away the const, and fix up.
 			if ( bPinNamesMatch &&
-				(bNewPinIsObject || (NewPin->PinType.PinCategory == K2Schema->PC_Interface)) &&
+				(bNewPinIsObject || (NewPin->PinType.PinCategory == UEdGraphSchema_K2::PC_Interface)) &&
 				(NewPin->PinType.PinSubCategoryObject == NULL) )
 			{
 				// @TODO:  Fix this up to be less hacky
@@ -369,7 +366,7 @@ UK2Node::ERedirectType UK2Node_Variable::DoPinsMatchForReconstruction( const UEd
 			{
 				const UClass* PSCOClass = Cast<UClass>(OldPin->PinType.PinSubCategoryObject.Get());
 				const bool bOldIsBlueprint = PSCOClass && PSCOClass->IsChildOf(UBlueprint::StaticClass());
-				const bool bNewIsClass     = (NewPin->PinType.PinCategory == K2Schema->PC_Class);
+				const bool bNewIsClass     = (NewPin->PinType.PinCategory == UEdGraphSchema_K2::PC_Class);
 				// Special Case: If we're migrating from old blueprint references 
 				//               to class references, allow pins to be reconnected if coerced
 				if (bNewIsClass && bOldIsBlueprint)
@@ -394,7 +391,7 @@ UClass* UK2Node_Variable::GetVariableSourceClass() const
 UProperty* UK2Node_Variable::GetPropertyForVariable_Internal(UClass* OwningClass) const
 {
 	const FName VarName = GetVarName();
-	UEdGraphPin* VariablePin = FindPin(GetVarNameString());
+	UEdGraphPin* VariablePin = FindPin(VarName);
 
 	UProperty* VariableProperty = VariableReference.ResolveMember<UProperty>(OwningClass);
 
@@ -408,7 +405,7 @@ UProperty* UK2Node_Variable::GetPropertyForVariable_Internal(UClass* OwningClass
 		// If the variable has been remapped update the pin
 		else if (VariablePin && VarName != GetVarName())
 		{
-			VariablePin->PinName = GetVarNameString();
+			VariablePin->PinName = GetVarName();
 		}
 	}
 
@@ -441,7 +438,7 @@ UProperty* UK2Node_Variable::GetPropertyForVariableFromSkeleton() const
 
 UEdGraphPin* UK2Node_Variable::GetValuePin() const
 {
-	UEdGraphPin* Pin = FindPin(GetVarNameString());
+	UEdGraphPin* Pin = FindPin(GetVarName());
 	check(Pin == nullptr || Pin->Direction == EGPD_Output);
 	return Pin;
 }
@@ -589,7 +586,7 @@ void UK2Node_Variable::HandleVariableRenamed(UBlueprint* InBlueprint, UClass* In
 			VariableReference.SetExternalMember(InNewVarName, NodeRefClass);
 		}
 
-		RenameUserDefinedPin(InOldVarName.ToString(), InNewVarName.ToString());
+		RenameUserDefinedPin(InOldVarName, InNewVarName);
 	}
 }
 
@@ -805,7 +802,7 @@ FString UK2Node_Variable::GetDocumentationLink() const
 
 FString UK2Node_Variable::GetDocumentationExcerptName() const
 {
-	return GetVarName().ToString();
+	return GetVarNameString();
 }
 
 void UK2Node_Variable::AutowireNewNode(UEdGraphPin* FromPin)
@@ -833,7 +830,7 @@ void UK2Node_Variable::AutowireNewNode(UEdGraphPin* FromPin)
 					// BP Comms is highly likely at this point, if the source pin's type is a child of the variable's owner class, let's conform the "Target" pin
 					if(FromPin->PinType.PinSubCategoryObject == PropertyOwner || dynamic_cast<UClass*>(FromPin->PinType.PinSubCategoryObject.Get())->IsChildOf(PropertyOwner))
 					{
-						UEdGraphPin* TargetPin = FindPin(K2Schema->PN_Self);
+						UEdGraphPin* TargetPin = FindPin(UEdGraphSchema_K2::PN_Self);
 						if (TargetPin)
 						{
 							TargetPin->PinType.PinSubCategoryObject = PropertyOwner;
diff --git a/Engine/Source/Editor/BlueprintGraph/Private/K2Node_VariableGet.cpp b/Engine/Source/Editor/BlueprintGraph/Private/K2Node_VariableGet.cpp
--- a/Engine/Source/Editor/BlueprintGraph/Private/K2Node_VariableGet.cpp
+++ b/Engine/Source/Editor/BlueprintGraph/Private/K2Node_VariableGet.cpp
@@ -111,19 +111,16 @@ void UK2Node_VariableGet::CreateNonPurePins(TArray<UEdGraphPin*>* InOldPinsPtr)
 		}
 		// If there is no property and we are given some old pins to look at, find the old value pin and use the type there
 		// This allows nodes to be pasted into other BPs without access to the property
-		else if(InOldPinsPtr)
+		else if (InOldPinsPtr)
 		{
 			// find old variable pin and use the type.
-			const FString PinName = GetVarNameString();
-			for(auto Iter = InOldPinsPtr->CreateConstIterator(); Iter; ++Iter)
+			const FName PinName = GetVarName();
+			for (const UEdGraphPin* Pin : *InOldPinsPtr)
 			{
-				if(const UEdGraphPin* Pin = *Iter)
+				if (Pin && PinName == Pin->PinName)
 				{
-					if(PinName == Pin->PinName)
-					{
-						PinType = Pin->PinType;
-						break;
-					}
+					PinType = Pin->PinType;
+					break;
 				}
 			}
 
@@ -132,13 +129,13 @@ void UK2Node_VariableGet::CreateNonPurePins(TArray<UEdGraphPin*>* InOldPinsPtr)
 		if (IsValidTypeForNonPure(PinType))
 		{
 			// Input - Execution Pin
-			CreatePin(EGPD_Input, K2Schema->PC_Exec, FString(), nullptr, K2Schema->PN_Execute);
+			CreatePin(EGPD_Input, UEdGraphSchema_K2::PC_Exec, UEdGraphSchema_K2::PN_Execute);
 
 			// Output - Execution Pins
-			UEdGraphPin* ValidPin = CreatePin(EGPD_Output, K2Schema->PC_Exec, FString(), nullptr, K2Schema->PN_Then);
+			UEdGraphPin* ValidPin = CreatePin(EGPD_Output, UEdGraphSchema_K2::PC_Exec, UEdGraphSchema_K2::PN_Then);
 			ValidPin->PinFriendlyName = LOCTEXT("Valid", "Is Valid");
 
-			UEdGraphPin* InvalidPin = CreatePin(EGPD_Output, K2Schema->PC_Exec, FString(), nullptr, K2Schema->PN_Else);
+			UEdGraphPin* InvalidPin = CreatePin(EGPD_Output, UEdGraphSchema_K2::PC_Exec, UEdGraphSchema_K2::PN_Else);
 			InvalidPin->PinFriendlyName = LOCTEXT("Invalid", "Is Not Valid");
 		}
 		else
@@ -272,7 +269,7 @@ FText UK2Node_VariableGet::GetTooltipText() const
 FText UK2Node_VariableGet::GetNodeTitle(ENodeTitleType::Type TitleType) const
 {
 	// If there is only one variable being read, the title can be made the variable name
-	FString OutputPinName;
+	FName OutputPinName;
 	int32 NumOutputsFound = 0;
 
 	for (int32 PinIndex = 0; PinIndex < Pins.Num(); ++PinIndex)
@@ -306,7 +303,7 @@ FText UK2Node_VariableGet::GetNodeTitle(ENodeTitleType::Type TitleType) const
 	else if (CachedNodeTitle.IsOutOfDate(this))
 	{
 		FFormatNamedArguments Args;
-		Args.Add(TEXT("PinName"), FText::FromString(OutputPinName));
+		Args.Add(TEXT("PinName"), FText::FromName(OutputPinName));
 		// FText::Format() is slow, so we cache this to save on performance
 		CachedNodeTitle.SetCachedText(FText::Format(LOCTEXT("GetPinName", "Get {PinName}"), Args), this);
 	}
@@ -320,7 +317,7 @@ FNodeHandlingFunctor* UK2Node_VariableGet::CreateNodeHandler(FKismetCompilerCont
 
 bool UK2Node_VariableGet::IsValidTypeForNonPure(const FEdGraphPinType& InPinType)
 {
-	return !InPinType.IsContainer() && (InPinType.PinCategory == UObject::StaticClass()->GetName() ||InPinType.PinCategory == UClass::StaticClass()->GetName());
+	return !InPinType.IsContainer() && (InPinType.PinCategory == UObject::StaticClass()->GetFName() ||InPinType.PinCategory == UClass::StaticClass()->GetFName());
 }
 
 void UK2Node_VariableGet::GetContextMenuActions(const FGraphNodeContextMenuBuilder& Context) const
@@ -484,11 +481,11 @@ void UK2Node_VariableGet::ExpandNode(class FKismetCompilerContext& CompilerConte
 		UK2Node_CallFunction* IsValidFunction = CompilerContext.SpawnIntermediateNode<UK2Node_CallFunction>(this, SourceGraph);
 
 		// Based on if the type is an "Object" or a "Class" changes which function to use
-		if (ValuePin->PinType.PinCategory == UObject::StaticClass()->GetName())
+		if (ValuePin->PinType.PinCategory == UObject::StaticClass()->GetFName())
 		{
 			IsValidFunction->SetFromFunction(UKismetSystemLibrary::StaticClass()->FindFunctionByName(GET_MEMBER_NAME_CHECKED(UKismetSystemLibrary, IsValid)));
 		}
-		else if (ValuePin->PinType.PinCategory == UClass::StaticClass()->GetName())
+		else if (ValuePin->PinType.PinCategory == UClass::StaticClass()->GetFName())
 		{
 			IsValidFunction->SetFromFunction(UKismetSystemLibrary::StaticClass()->FindFunctionByName(GET_MEMBER_NAME_CHECKED(UKismetSystemLibrary, IsValidClass)));
 		}
diff --git a/Engine/Source/Editor/BlueprintGraph/Private/K2Node_VariableSet.cpp b/Engine/Source/Editor/BlueprintGraph/Private/K2Node_VariableSet.cpp
--- a/Engine/Source/Editor/BlueprintGraph/Private/K2Node_VariableSet.cpp
+++ b/Engine/Source/Editor/BlueprintGraph/Private/K2Node_VariableSet.cpp
@@ -81,8 +81,8 @@ UK2Node_VariableSet::UK2Node_VariableSet(const FObjectInitializer& ObjectInitial
 
 void UK2Node_VariableSet::AllocateDefaultPins()
 {
-	CreatePin(EGPD_Input, UEdGraphSchema_K2::PC_Exec, FString(), nullptr, UEdGraphSchema_K2::PN_Execute);
-	CreatePin(EGPD_Output, UEdGraphSchema_K2::PC_Exec, FString(), nullptr, UEdGraphSchema_K2::PN_Then);
+	CreatePin(EGPD_Input, UEdGraphSchema_K2::PC_Exec, UEdGraphSchema_K2::PN_Execute);
+	CreatePin(EGPD_Output, UEdGraphSchema_K2::PC_Exec, UEdGraphSchema_K2::PN_Then);
 
 	if (GetVarName() != NAME_None)
 	{
@@ -102,8 +102,8 @@ void UK2Node_VariableSet::AllocateDefaultPins()
 
 void UK2Node_VariableSet::ReallocatePinsDuringReconstruction(TArray<UEdGraphPin*>& OldPins)
 {
-	CreatePin(EGPD_Input, UEdGraphSchema_K2::PC_Exec, FString(), nullptr, UEdGraphSchema_K2::PN_Execute);
-	CreatePin(EGPD_Output, UEdGraphSchema_K2::PC_Exec, FString(), nullptr, UEdGraphSchema_K2::PN_Then);
+	CreatePin(EGPD_Input, UEdGraphSchema_K2::PC_Exec, UEdGraphSchema_K2::PN_Execute);
+	CreatePin(EGPD_Output, UEdGraphSchema_K2::PC_Exec, UEdGraphSchema_K2::PN_Then);
 
 	if (GetVarName() != NAME_None)
 	{
@@ -247,7 +247,7 @@ FText UK2Node_VariableSet::GetNodeTitle(ENodeTitleType::Type TitleType) const
 	const UEdGraphSchema_K2* K2Schema = GetDefault<UEdGraphSchema_K2>();
 
 	// If there is only one variable being written (one non-meta input pin), the title can be made the variable name
-	FString InputPinName;
+	FName InputPinName;
 	int32 NumInputsFound = 0;
 
 	for (int32 PinIndex = 0; PinIndex < Pins.Num(); ++PinIndex)
@@ -270,7 +270,7 @@ FText UK2Node_VariableSet::GetNodeTitle(ENodeTitleType::Type TitleType) const
 	else if (CachedNodeTitle.IsOutOfDate(this))
 	{
 		FFormatNamedArguments Args;
-		Args.Add(TEXT("PinName"), FText::FromString(InputPinName));
+		Args.Add(TEXT("PinName"), FText::FromName(InputPinName));
 
 		// FText::Format() is slow, so we cache this to save on performance
 		if (HasLocalRepNotify())
@@ -324,7 +324,7 @@ FNodeHandlingFunctor* UK2Node_VariableSet::CreateNodeHandler(FKismetCompilerCont
 	return new FKCHandler_VariableSet(CompilerContext);
 }
 
-FString UK2Node_VariableSet::GetVariableOutputPinName() const
+FName UK2Node_VariableSet::GetVariableOutputPinName() const
 {
 	return TEXT("Output_Get");
 }
@@ -344,7 +344,7 @@ FText UK2Node_VariableSet::GetPinNameOverride(const UEdGraphPin& Pin) const
 		return FText::GetEmpty();
 	}
 
-	return !Pin.PinFriendlyName.IsEmpty() ? Pin.PinFriendlyName : FText::FromString(Pin.PinName);
+	return !Pin.PinFriendlyName.IsEmpty() ? Pin.PinFriendlyName : FText::FromName(Pin.PinName);
 }
 
 void UK2Node_VariableSet::ValidateNodeDuringCompilation(FCompilerResultsLog& MessageLog) const
@@ -408,7 +408,7 @@ void UK2Node_VariableSet::ExpandNode(class FKismetCompilerContext& CompilerConte
 					VariableGetNode->VariableReference = VariableReference;
 					VariableGetNode->AllocateDefaultPins();
 					CompilerContext.MessageLog.NotifyIntermediateObjectCreation(VariableGetNode, this);
-					CompilerContext.MovePinLinksToIntermediate(*VariableGetPin, *VariableGetNode->FindPin(GetVarNameString()));
+					CompilerContext.MovePinLinksToIntermediate(*VariableGetPin, *VariableGetNode->FindPin(GetVarName()));
 
 					// Duplicate the connection to the self pin.
 					UEdGraphPin* SetSelfPin = K2Schema->FindSelfPin(*this, EGPD_Input);
@@ -462,7 +462,7 @@ void UK2Node_VariableSet::ExpandNode(class FKismetCompilerContext& CompilerConte
 				}
 				check(SetFunctionValuePin);
 
-				CompilerContext.MovePinLinksToIntermediate(*FindPin(GetVarNameString(), EGPD_Input), *SetFunctionValuePin);
+				CompilerContext.MovePinLinksToIntermediate(*FindPin(GetVarName(), EGPD_Input), *SetFunctionValuePin);
 			}
 		}
 	}
diff --git a/Engine/Source/Editor/BlueprintGraph/Private/K2Node_VariableSetRef.cpp b/Engine/Source/Editor/BlueprintGraph/Private/K2Node_VariableSetRef.cpp
--- a/Engine/Source/Editor/BlueprintGraph/Private/K2Node_VariableSetRef.cpp
+++ b/Engine/Source/Editor/BlueprintGraph/Private/K2Node_VariableSetRef.cpp
@@ -11,8 +11,8 @@
 #include "EditorCategoryUtils.h"
 #include "BlueprintActionDatabaseRegistrar.h"
 
-static FString TargetVarPinName(TEXT("Target"));
-static FString VarValuePinName(TEXT("Value"));
+static FName TargetVarPinName(TEXT("Target"));
+static FName VarValuePinName(TEXT("Value"));
 
 #define LOCTEXT_NAMESPACE "K2Node_VariableSetRef"
 
@@ -91,13 +91,14 @@ UK2Node_VariableSetRef::UK2Node_VariableSetRef(const FObjectInitializer& ObjectI
 
 void UK2Node_VariableSetRef::AllocateDefaultPins()
 {
-	const UEdGraphSchema_K2* K2Schema = GetDefault<UEdGraphSchema_K2>();
+	CreatePin(EGPD_Input, UEdGraphSchema_K2::PC_Exec, UEdGraphSchema_K2::PN_Execute);
+	CreatePin(EGPD_Output, UEdGraphSchema_K2::PC_Exec, UEdGraphSchema_K2::PN_Then);
 
-	CreatePin(EGPD_Input, K2Schema->PC_Exec, FString(), nullptr, K2Schema->PN_Execute);
-	CreatePin(EGPD_Output, K2Schema->PC_Exec, FString(), nullptr, K2Schema->PN_Then);
+	UEdGraphNode::FCreatePinParams PinParams;
+	PinParams.bIsReference = true;
+	CreatePin(EGPD_Input, UEdGraphSchema_K2::PC_Wildcard, TargetVarPinName, PinParams);
 
-	CreatePin(EGPD_Input, K2Schema->PC_Wildcard, FString(), nullptr, TargetVarPinName, EPinContainerType::None, true);
-	UEdGraphPin* ValuePin = CreatePin(EGPD_Input, K2Schema->PC_Wildcard, FString(), nullptr, VarValuePinName);
+	CreatePin(EGPD_Input, UEdGraphSchema_K2::PC_Wildcard, VarValuePinName);
 }
 
 void UK2Node_VariableSetRef::ReallocatePinsDuringReconstruction(TArray<UEdGraphPin*>& OldPins)
@@ -105,11 +106,10 @@ void UK2Node_VariableSetRef::ReallocatePinsDuringReconstruction(TArray<UEdGraphP
 	AllocateDefaultPins();
 
  	// Coerce the type of the node from the old pin, if available
-  	UEdGraphPin* OldTargetPin = NULL;
-  	for( auto OldPinIt = OldPins.CreateIterator(); OldPinIt; ++OldPinIt )
+  	UEdGraphPin* OldTargetPin = nullptr;
+  	for (UEdGraphPin* CurrPin : OldPins)
   	{
-  		UEdGraphPin* CurrPin = *OldPinIt;
-  		if( CurrPin->PinName == TargetVarPinName )
+  		if (CurrPin->PinName == TargetVarPinName)
   		{
   			OldTargetPin = CurrPin;
   			break;
@@ -136,7 +136,7 @@ FText UK2Node_VariableSetRef::GetNodeTitle(ENodeTitleType::Type TitleType) const
 	const UEdGraphSchema_K2* Schema = GetDefault<UEdGraphSchema_K2>();
 
 	UEdGraphPin* TargetPin = GetTargetPin();
-	if ((TargetPin == nullptr) || (TargetPin->PinType.PinCategory == Schema->PC_Wildcard))
+	if ((TargetPin == nullptr) || (TargetPin->PinType.PinCategory == UEdGraphSchema_K2::PC_Wildcard))
 	{
 		return NSLOCTEXT("K2Node", "SetRefVarNodeTitle", "Set By-Ref Var");
 	}
@@ -183,8 +183,6 @@ void UK2Node_VariableSetRef::NotifyPinConnectionListChanged(UEdGraphPin* Pin)
 
 void UK2Node_VariableSetRef::CoerceTypeFromPin(const UEdGraphPin* Pin)
 {
-	const UEdGraphSchema_K2* K2Schema = GetDefault<UEdGraphSchema_K2>();
-
 	UEdGraphPin* TargetPin = GetTargetPin();
 	UEdGraphPin* ValuePin = GetValuePin();
 
@@ -203,13 +201,13 @@ void UK2Node_VariableSetRef::CoerceTypeFromPin(const UEdGraphPin* Pin)
 	else
 	{
 		// Pin disconnected...revert to wildcard
-		TargetPin->PinType.PinCategory = K2Schema->PC_Wildcard;
-		TargetPin->PinType.PinSubCategory.Reset();
+		TargetPin->PinType.PinCategory = UEdGraphSchema_K2::PC_Wildcard;
+		TargetPin->PinType.PinSubCategory = NAME_None;
 		TargetPin->PinType.PinSubCategoryObject = nullptr;
 		TargetPin->BreakAllPinLinks();
 
-		ValuePin->PinType.PinCategory = K2Schema->PC_Wildcard;
-		ValuePin->PinType.PinSubCategory.Reset();
+		ValuePin->PinType.PinCategory = UEdGraphSchema_K2::PC_Wildcard;
+		ValuePin->PinType.PinSubCategory = NAME_None;
 		ValuePin->PinType.PinSubCategoryObject = nullptr;
 		ValuePin->BreakAllPinLinks();
 
diff --git a/Engine/Source/Editor/BlueprintGraph/Private/MakeStructHandler.cpp b/Engine/Source/Editor/BlueprintGraph/Private/MakeStructHandler.cpp
--- a/Engine/Source/Editor/BlueprintGraph/Private/MakeStructHandler.cpp
+++ b/Engine/Source/Editor/BlueprintGraph/Private/MakeStructHandler.cpp
@@ -111,7 +111,7 @@ void FKCHandler_MakeStruct::Compile(FKismetFunctionContext& Context, UEdGraphNod
 	{
 		if (Pin && !Pin->bOrphanedPin && (Pin != StructPin) && (Pin->Direction == EGPD_Input) && !CompilerContext.GetSchema()->IsMetaPin(*Pin))
 		{
-			UProperty* BoundProperty = FindField<UProperty>(Node->StructType, *(Pin->PinName));
+			UProperty* BoundProperty = FindField<UProperty>(Node->StructType, Pin->PinName);
 			check(BoundProperty);
 
 			// If the pin is not connectible, do not forward the net
@@ -151,7 +151,7 @@ void FKCHandler_MakeStruct::Compile(FKismetFunctionContext& Context, UEdGraphNod
 
 				// Literal Bool Term to set the OverrideProperty to
 				FBPTerminal* BoolTerm = Context.CreateLocalTerminal(ETerminalSpecification::TS_Literal);
-				BoolTerm->Type.PinCategory = CompilerContext.GetSchema()->PC_Boolean;
+				BoolTerm->Type.PinCategory = UEdGraphSchema_K2::PC_Boolean;
 				BoolTerm->bIsLiteral = true;
 				// If we are showing the pin, then we are overriding the property
 
diff --git a/Engine/Source/Editor/BlueprintGraph/Private/MathExpressionHandler.cpp b/Engine/Source/Editor/BlueprintGraph/Private/MathExpressionHandler.cpp
--- a/Engine/Source/Editor/BlueprintGraph/Private/MathExpressionHandler.cpp
+++ b/Engine/Source/Editor/BlueprintGraph/Private/MathExpressionHandler.cpp
@@ -17,14 +17,14 @@ struct KCHandler_MathExpressionHelper
 	static bool CreateMap(UEdGraphNode& KeyNode, UEdGraphNode& ValueNode, EEdGraphPinDirection KeyPinDirection, const UEdGraphSchema_K2* Schema, TMap<UEdGraphPin*, UEdGraphPin*>& OutMap)
 	{
 		check(Schema);
-		for (auto InnerInputPin : KeyNode.Pins)
+		for (UEdGraphPin* InnerInputPin : KeyNode.Pins)
 		{
 			if (!InnerInputPin || InnerInputPin->Direction != KeyPinDirection)
 			{
 				return false;
 			}
 
-			auto OuterInputPin = ValueNode.FindPin(InnerInputPin->PinName);
+			UEdGraphPin* OuterInputPin = ValueNode.FindPin(InnerInputPin->PinName);
 			if (!OuterInputPin
 				|| !Schema->ArePinsCompatible(InnerInputPin, OuterInputPin, nullptr, false))
 			{
@@ -85,7 +85,7 @@ FBlueprintCompiledStatement* FKCHandler_MathExpression::GenerateFunctionRPN(UEdG
 		{
 			UEdGraphPin* PinToTry = nullptr;
 			{
-				UEdGraphPin* PinMatch = CallFunctionNode->FindPin(Property->GetName());
+				UEdGraphPin* PinMatch = CallFunctionNode->FindPin(Property->GetFName());
 				const bool bGoodPin = PinMatch && FKismetCompilerUtilities::IsTypeCompatibleWithProperty(PinMatch, Property, CompilerContext.MessageLog, CompilerContext.GetSchema(), Context.NewClass);
 				PinToTry = bGoodPin ? FEdGraphUtilities::GetNetFromPin(PinMatch) : nullptr;
 			}
@@ -145,21 +145,21 @@ void FKCHandler_MathExpression::RegisterNets(FKismetFunctionContext& Context, UE
 {
 	FNodeHandlingFunctor::RegisterNets(Context, InNode);
 
-	auto Node_MathExpression = CastChecked<UK2Node_MathExpression>(InNode);
-	auto InnerEntryNode = Node_MathExpression->GetEntryNode();
-	auto InnerExitNode = Node_MathExpression->GetExitNode();
+	UK2Node_MathExpression* Node_MathExpression = CastChecked<UK2Node_MathExpression>(InNode);
+	UK2Node_Tunnel* InnerEntryNode = Node_MathExpression->GetEntryNode();
+	UK2Node_Tunnel* InnerExitNode = Node_MathExpression->GetExitNode();
 	check(Node_MathExpression->BoundGraph && InnerEntryNode && InnerExitNode);
 
-	for (auto InnerGraphNode : Node_MathExpression->BoundGraph->Nodes)
+	for (UEdGraphNode* InnerGraphNode : Node_MathExpression->BoundGraph->Nodes)
 	{
 		if (!InnerGraphNode || (InnerGraphNode == InnerEntryNode) || (InnerGraphNode == InnerExitNode))
 		{
 			continue;
 		}
 
-		if (auto GetVarNode = Cast<UK2Node_VariableGet>(InnerGraphNode))
+		if (UK2Node_VariableGet* GetVarNode = Cast<UK2Node_VariableGet>(InnerGraphNode))
 		{
-			auto VarHandler = GetVarNode->CreateNodeHandler(CompilerContext);
+			FNodeHandlingFunctor* VarHandler = GetVarNode->CreateNodeHandler(CompilerContext);
 			if (ensure(VarHandler))
 			{
 				VarHandler->RegisterNets(Context, GetVarNode);
@@ -167,13 +167,13 @@ void FKCHandler_MathExpression::RegisterNets(FKismetFunctionContext& Context, UE
 			continue;
 		}
 
-		for (auto Pin : InnerGraphNode->Pins)
+		for (UEdGraphPin* Pin : InnerGraphNode->Pins)
 		{
 			// Register fake terms for InlineGeneratedValues
 			if (Pin && (Pin->Direction == EEdGraphPinDirection::EGPD_Output) && Pin->LinkedTo.Num())
 			{
-				auto Linked = Pin->LinkedTo[0];
-				auto LinkedOwnerNode = Linked ? Linked->GetOwningNodeUnchecked() : nullptr;
+				UEdGraphPin* Linked = Pin->LinkedTo[0];
+				UEdGraphNode* LinkedOwnerNode = Linked ? Linked->GetOwningNodeUnchecked() : nullptr;
 				if (LinkedOwnerNode && (InnerExitNode != LinkedOwnerNode))
 				{
 					FBPTerminal* Term = new (Context.InlineGeneratedValues) FBPTerminal();
@@ -193,7 +193,7 @@ void FKCHandler_MathExpression::RegisterNets(FKismetFunctionContext& Context, UE
 
 void FKCHandler_MathExpression::RegisterNet(FKismetFunctionContext& Context, UEdGraphPin* Net) 
 {
-	FBPTerminal* Term = Context.CreateLocalTerminalFromPinAutoChooseScope(Net, Context.NetNameMap->MakeValidName(Net));
+	FBPTerminal* Term = Context.CreateLocalTerminalFromPinAutoChooseScope(Net, *Context.NetNameMap->MakeValidName(Net));
 	Context.NetMap.Add(Net, Term);
 }
 
@@ -202,8 +202,8 @@ void FKCHandler_MathExpression::Compile(FKismetFunctionContext& Context, UEdGrap
 	UK2Node_MathExpression* Node_MathExpression = CastChecked<UK2Node_MathExpression>(Node);
 	check(Context.Schema);
 
-	auto InnerExitNode = Node_MathExpression->GetExitNode();
-	auto InnerEntryNode = Node_MathExpression->GetEntryNode();
+	UK2Node_Tunnel* InnerExitNode = Node_MathExpression->GetExitNode();
+	UK2Node_Tunnel* InnerEntryNode = Node_MathExpression->GetEntryNode();
 
 	if (!InnerExitNode || !InnerEntryNode || (InnerExitNode->Pins.Num() != 1) || ((InnerExitNode->Pins.Num() + InnerEntryNode->Pins.Num()) != Node->Pins.Num()))
 	{
@@ -223,8 +223,8 @@ void FKCHandler_MathExpression::Compile(FKismetFunctionContext& Context, UEdGrap
 		return;
 	}
 
-	auto InnerOutputPin = InnerExitNode->Pins[0];
-	auto OuterOutputPinPtr = InnerToOuterOutput.Find(InnerOutputPin);
+	UEdGraphPin* InnerOutputPin = InnerExitNode->Pins[0];
+	UEdGraphPin** OuterOutputPinPtr = InnerToOuterOutput.Find(InnerOutputPin);
 	if (!InnerOutputPin || (InnerOutputPin->LinkedTo.Num() != 1) || !InnerOutputPin->LinkedTo[0] || !OuterOutputPinPtr || !*OuterOutputPinPtr)
 	{
 		Context.MessageLog.Error(*LOCTEXT("Compile_WrongOutputLink", "ICE - wrong output link - @@").ToString(), Node);
@@ -233,7 +233,7 @@ void FKCHandler_MathExpression::Compile(FKismetFunctionContext& Context, UEdGrap
 
 	FBPTerminal** OutputTermPtr = Context.NetMap.Find(FEdGraphUtilities::GetNetFromPin(*OuterOutputPinPtr));
 	FBPTerminal* OutputTerm = OutputTermPtr ? *OutputTermPtr : nullptr;
-	auto LastInnerNode = InnerOutputPin->LinkedTo[0]->GetOwningNodeUnchecked();
+	UEdGraphNode* LastInnerNode = InnerOutputPin->LinkedTo[0]->GetOwningNodeUnchecked();
 
 	FBlueprintCompiledStatement* DetachedStatement = GenerateFunctionRPN(LastInnerNode, Context, *Node_MathExpression, OutputTerm, InnerToOuterInput);
 	if (DetachedStatement)
diff --git a/Engine/Source/Editor/BlueprintGraph/Private/StructMemberNodeHandlers.cpp b/Engine/Source/Editor/BlueprintGraph/Private/StructMemberNodeHandlers.cpp
--- a/Engine/Source/Editor/BlueprintGraph/Private/StructMemberNodeHandlers.cpp
+++ b/Engine/Source/Editor/BlueprintGraph/Private/StructMemberNodeHandlers.cpp
@@ -58,13 +58,11 @@ static FBPTerminal* RegisterStructVar(FCompilerResultsLog& MessageLog, FKismetFu
 static void ResolveAndRegisterScopedStructTerm(FCompilerResultsLog& MessageLog, FKismetFunctionContext& Context, UScriptStruct* StructType, UEdGraphPin* Net, FBPTerminal* ContextTerm)
 {
 	// Find the property for the struct
-	UProperty* BoundProperty = FindField<UProperty>(StructType, *(Net->PinName));
-
-	if (BoundProperty != NULL)
+	if (UProperty* BoundProperty = FindField<UProperty>(StructType, Net->PinName))
 	{
 		// Create the term in the list
 		FBPTerminal* Term = new (Context.VariableReferences) FBPTerminal();
-		Term->CopyFromPin(Net, Net->PinName);
+		Term->CopyFromPin(Net, Net->PinName.ToString());
 		Term->AssociatedVarProperty = BoundProperty;
 		Context.NetMap.Add(Net, Term);
 		Term->Context = ContextTerm;
diff --git a/Engine/Source/Editor/BlueprintGraph/Private/VariableSetHandler.cpp b/Engine/Source/Editor/BlueprintGraph/Private/VariableSetHandler.cpp
--- a/Engine/Source/Editor/BlueprintGraph/Private/VariableSetHandler.cpp
+++ b/Engine/Source/Editor/BlueprintGraph/Private/VariableSetHandler.cpp
@@ -155,12 +155,12 @@ void FKCHandler_VariableSet::Transform(FKismetFunctionContext& Context, UEdGraph
 			CallFuncNode->AllocateDefaultPins();
 
 			// Copy self pin
-			UEdGraphPin* NewSelfPin = CallFuncNode->FindPinChecked(CompilerContext.GetSchema()->PN_Self);
-			UEdGraphPin* OldSelfPin = Node->FindPinChecked(CompilerContext.GetSchema()->PN_Self);
+			UEdGraphPin* NewSelfPin = CallFuncNode->FindPinChecked(UEdGraphSchema_K2::PN_Self);
+			UEdGraphPin* OldSelfPin = Node->FindPinChecked(UEdGraphSchema_K2::PN_Self);
 			NewSelfPin->CopyPersistentDataFromOldPin(*OldSelfPin);
 
 			// link new CallFuncNode -> Set Node
-			UEdGraphPin* OldExecPin = Node->FindPin(CompilerContext.GetSchema()->PN_Execute);
+			UEdGraphPin* OldExecPin = Node->FindPin(UEdGraphSchema_K2::PN_Execute);
 			check(OldExecPin);
 
 			UEdGraphPin* NewExecPin = CallFuncNode->GetExecPin();
@@ -179,12 +179,12 @@ void FKCHandler_VariableSet::Transform(FKismetFunctionContext& Context, UEdGraph
 			CallFuncNode->AllocateDefaultPins();
 
 			// Copy self pin
-			UEdGraphPin* NewSelfPin = CallFuncNode->FindPinChecked(CompilerContext.GetSchema()->PN_Self);
-			UEdGraphPin* OldSelfPin = Node->FindPinChecked(CompilerContext.GetSchema()->PN_Self);
+			UEdGraphPin* NewSelfPin = CallFuncNode->FindPinChecked(UEdGraphSchema_K2::PN_Self);
+			UEdGraphPin* OldSelfPin = Node->FindPinChecked(UEdGraphSchema_K2::PN_Self);
 			NewSelfPin->CopyPersistentDataFromOldPin(*OldSelfPin);
 
 			// link Set Node -> new CallFuncNode
-			UEdGraphPin* OldThenPin = Node->FindPin(CompilerContext.GetSchema()->PN_Then);
+			UEdGraphPin* OldThenPin = Node->FindPin(UEdGraphSchema_K2::PN_Then);
 			check(OldThenPin);
 
 			UEdGraphPin* NewThenPin = CallFuncNode->GetThenPin();
diff --git a/Engine/Source/Editor/Cascade/Private/Cascade.cpp b/Engine/Source/Editor/Cascade/Private/Cascade.cpp
--- a/Engine/Source/Editor/Cascade/Private/Cascade.cpp
+++ b/Engine/Source/Editor/Cascade/Private/Cascade.cpp
@@ -1673,10 +1673,6 @@ void FCascade::Tick(float DeltaTime)
 	}
 }
 
-bool FCascade::IsTickable() const
-{
-	return true;
-}
 TStatId FCascade::GetStatId() const
 {
 	RETURN_QUICK_DECLARE_CYCLE_STAT(FCascade, STATGROUP_Tickables);
diff --git a/Engine/Source/Editor/EnvironmentQueryEditor/Private/EnvironmentQueryGraphNode_Option.cpp b/Engine/Source/Editor/EnvironmentQueryEditor/Private/EnvironmentQueryGraphNode_Option.cpp
--- a/Engine/Source/Editor/EnvironmentQueryEditor/Private/EnvironmentQueryGraphNode_Option.cpp
+++ b/Engine/Source/Editor/EnvironmentQueryEditor/Private/EnvironmentQueryGraphNode_Option.cpp
@@ -21,7 +21,7 @@ UEnvironmentQueryGraphNode_Option::UEnvironmentQueryGraphNode_Option(const FObje
 
 void UEnvironmentQueryGraphNode_Option::AllocateDefaultPins()
 {
-	UEdGraphPin* Inputs = CreatePin(EGPD_Input, TEXT("Transition"), FString(), nullptr, TEXT("Out"));
+	UEdGraphPin* Inputs = CreatePin(EGPD_Input, TEXT("Transition"), TEXT("Out"));
 }
 
 void UEnvironmentQueryGraphNode_Option::PostPlacedNewNode()
@@ -55,6 +55,7 @@ void UEnvironmentQueryGraphNode_Option::ResetNodeOwner()
 	{
 		UObject* GraphOwner = GetGraph() ? GetGraph()->GetOuter() : nullptr;
 		OptionInstance->Generator->Rename(NULL, GraphOwner, REN_DontCreateRedirectors | REN_DoNotDirty);
+		OptionInstance->Generator->ClearFlags(RF_Transient);
 	}
 }
 
diff --git a/Engine/Source/Editor/EnvironmentQueryEditor/Private/EnvironmentQueryGraphNode_Root.cpp b/Engine/Source/Editor/EnvironmentQueryEditor/Private/EnvironmentQueryGraphNode_Root.cpp
--- a/Engine/Source/Editor/EnvironmentQueryEditor/Private/EnvironmentQueryGraphNode_Root.cpp
+++ b/Engine/Source/Editor/EnvironmentQueryEditor/Private/EnvironmentQueryGraphNode_Root.cpp
@@ -9,7 +9,7 @@ UEnvironmentQueryGraphNode_Root::UEnvironmentQueryGraphNode_Root(const FObjectIn
 
 void UEnvironmentQueryGraphNode_Root::AllocateDefaultPins()
 {
-	UEdGraphPin* Outputs = CreatePin(EGPD_Output, TEXT("Transition"), FString(), nullptr, TEXT("In"));
+	UEdGraphPin* Outputs = CreatePin(EGPD_Output, TEXT("Transition"), TEXT("In"));
 }
 
 FText UEnvironmentQueryGraphNode_Root::GetNodeTitle(ENodeTitleType::Type TitleType) const
diff --git a/Engine/Source/Editor/GameProjectGeneration/Private/GameProjectUtils.cpp b/Engine/Source/Editor/GameProjectGeneration/Private/GameProjectUtils.cpp
--- a/Engine/Source/Editor/GameProjectGeneration/Private/GameProjectUtils.cpp
+++ b/Engine/Source/Editor/GameProjectGeneration/Private/GameProjectUtils.cpp
@@ -70,6 +70,7 @@
 
 #include "PlatformInfo.h"
 #include "Blueprint/BlueprintSupport.h"
+#include "Settings/ProjectPackagingSettings.h"
 
 #define LOCTEXT_NAMESPACE "GameProjectUtils"
 
@@ -3369,6 +3370,9 @@ bool GameProjectUtils::ProjectRequiresBuild(const FName InPlatformInfoName)
 	// check to see if any plugins beyond the defaults have been enabled
 	bRequiresBuild |= IProjectManager::Get().IsNonDefaultPluginEnabled();
 
+	// check to see if Blueprint nativization is enabled in the Project settings
+	bRequiresBuild |= GetDefault<UProjectPackagingSettings>()->BlueprintNativizationMethod != EProjectPackagingBlueprintNativizationMethod::Disabled;
+
 	return bRequiresBuild;
 }
 
diff --git a/Engine/Source/Editor/GameplayTasksEditor/Classes/K2Node_LatentGameplayTaskCall.h b/Engine/Source/Editor/GameplayTasksEditor/Classes/K2Node_LatentGameplayTaskCall.h
--- a/Engine/Source/Editor/GameplayTasksEditor/Classes/K2Node_LatentGameplayTaskCall.h
+++ b/Engine/Source/Editor/GameplayTasksEditor/Classes/K2Node_LatentGameplayTaskCall.h
@@ -42,7 +42,7 @@ public:
 	bool ValidateActorArraySpawning(class FKismetCompilerContext& CompilerContext, bool bGenerateErrors);
 
 	UPROPERTY()
-	TArray<FString> SpawnParamPins;
+	TArray<FName> SpawnParamPins;
 
 	static void RegisterSpecializedTaskNodeClass(TSubclassOf<UK2Node_LatentGameplayTaskCall> NodeClass);
 protected:
