Copying //UE4/Dev-Framework to //UE4/Dev-Main (Source: //UE4/Dev-Framework @ 3779049)

============================
  MAJOR FEATURES & CHANGES
============================

Change 3626305 by [REDACTED_PERSON]

Change 3627162 by [REDACTED_PERSON]

	- Regression introduced in CL# 3613358.

Change 3756887 by [REDACTED_PERSON]

	Copy of CL #3756788

Change 3756888 by [REDACTED_PERSON]

	Copy of CL #3753986

Change 3759254 by [REDACTED_PERSON]

	Disable the duplicate PrimaryAssetId for editor only types like Maps. This can happen if content folk copy maps using the content browser, and does not actually cause a runtime problem. It still ensures for cooked types

Change 3759278 by [REDACTED_PERSON]

	Add IsTempPackage to FPackageName
	Fix issue where temp/memory packages shown in a content browser/asset audit window would spam the log when it failed to find source control info for them

Change 3759613 by [REDACTED_PERSON]

	Add support for casting between mismatched soft pointer types in nativized Blueprint C++ assignment statements and function calls.

	Change summary:
	- Extended FEmitHelper::GenerateAutomaticCast() to consider soft pointer terms and inject C++ code to explicitly cast the RHS when needed.

Change 3760040 by [REDACTED_PERSON]

	Add Call Stack control for viewing Blueprint call stacks when paused at a breakpoint, available from the Developer Tools menu

Change 3760955 by [REDACTED_PERSON]

	Fix conditional (SA/CIS issue).

Change 3761356 by [REDACTED_PERSON]

	Fix DLC staging rules to handle metadata correctly and remove debug log I accidentally added. The DLC staging now iterates in a similar way to the normal staging, it just may also excluded Engine

Change 3761859 by [REDACTED_PERSON]

	#ue4 - Fix crash in UStaticMesh::GetAssetRegistryTags() when FindObject is used during saving. Added Lex::ToString for physics enums ECollisionTraceFlag, EPhysicsType, and EBodyCollisionResponse.

	#tests QA game, content browser

Change 3761860 by mason.seay

	Submitting test content for Async Load issue

Change 3762559 by [REDACTED_PERSON]

	Change GetDefaultsAsString on Pin to always return an internal string so it can correctly be import texted, add GetDefaultsAsText for display purposes

Change 3764459 by [REDACTED_PERSON]

	PR #4224: Fix LoadLevelInstanceBySoftObjectPtr (Contributed by [REDACTED_PERSON])

Change 3764580 by [REDACTED_PERSON]

	Clean up delegates in UObjectGlobals.h, fixing several incorrect comments and moving some editor delegates into WITH_EDITOR

Change 3764602 by [REDACTED_PERSON]

	This is necessary because they would not get registered with various editor systems for the rest of the editor session, even if opened manually

Change 3764603 by [REDACTED_PERSON]

	Related to UE-52487, now that async loading blueprints in the editor properly registers them with the blueprint actions, we need to unregister them when automated tests want them to unload. Add a ClearEditorReferences function to UBlueprint that calls the OnUnloaded editor delegate, so EngineTest doesn't need to include the editor module

Change 3764768 by [REDACTED_PERSON]

Change 3766415 by [REDACTED_PERSON]

	Removing invalid assets. Most of these are out  of date.

Change 3766417 by [REDACTED_PERSON]

	Add warnings when we try to export a package without a type.

Change 3766514 by [REDACTED_PERSON]

	Added a #include to fix the build.

Change 3766542 by [REDACTED_PERSON]

	Added a #include to fix the build.

Change 3766558 by [REDACTED_PERSON]

	Rename variables to avoid warnings about hiding previous variable declarations.

Change 3767619 by [REDACTED_PERSON]

	bActorIsBeingDestroyed must be part of transaction history

Change 3767993 by [REDACTED_PERSON]

	Preserve graph editor controls when clicking on a hyper link, this speeds up navigation via the debugger  'step' command and Find in Blueprints control

Change 3768146 by [REDACTED_PERSON]

	Fix material instance dynamic not correctly finding object in details panel customization as a result soft path changes

Change 3769586 by [REDACTED_PERSON]

	Fix expose on spawn related error messages

Change 3769863 by [REDACTED_PERSON]

	Blueprint call stack now has access to frame offsets and can highlight nodes that are active on previous stack frames

Change 3771200 by [REDACTED_PERSON]

	CIS fix - add missing DO_BLUEPRINT_GUARD

Change 3771555 by [REDACTED_PERSON]

	Add transactions for several pin class changing actions which were missing them

Change 3771589 by [REDACTED_PERSON]

Change 3771683 by [REDACTED_PERSON]

	Call Stack polish: background color no longer changes when undocked, prettify-ing "ExecuteUbergraph_blahblah" in to "Event Graph", resizing works correctly, added overlay message when no call stack is available

Change 3771734 by [REDACTED_PERSON]

	Add entries for native code in the blueprint call stack view, clarifying re-entrancy

Change 3774293 by [REDACTED_PERSON]

Change 3774328 by [REDACTED_PERSON]

Change 3774342 by [REDACTED_PERSON]

Change 3774476 by [REDACTED_PERSON]

	Add class and function info to pin names for async nodes. This fixes a problem where redirectors for async node pins did not work.

	https://udn.unrealengine.com/questions/402882/propertyredirect-fails-with-uk2node-latentgameplay.html?childToView=403808

Change 3774645 by [REDACTED_PERSON]

	Refactor the variable rename checking in make/break struct to use the generic one I just added

Change 3775412 by [REDACTED_PERSON]

	UX improvements for Blueprint single-step debugging and breakpoints. Also added Step Out and Step Over debugging commands.

	Change summary:
	- Remapped the existing Step In command from F10 to F11 hotkey.
	- Mapped existing Step Over command to F10 and existing Step Out command to ALT-SHIFT-F11 hotkeys.
	- Added new (repurposed) icon assets for single-step debugging toolbar commands.
	- Modified FPlayWorldCommands::BuildToolbar() to add new Step Over and Step Out commands to the toolbar.
	- Modified FCompilerResultsLog::CalculateStableIdentifierForLatentActionManager() to remove special-case code for intermediate Tunnel Instance nodes, as these are now reverse-mapped through FullSourceBacktrackMap.
	- Modified FKismetDebugUtilities::CheckBreakConditions() to more generally manage the current graph stack (i.e. not just for Blueprint Function graphs). Also fixed a bug where we had failed to reset the target graph stack depth after completing a Step Out/Over iteration.
	- Modified FBlueprintDebugData::FindAllCodeLocationsFromSourceNode() to remove the additional iteration for the special Macro source node table (no longer required).
	- Modified FBlueprintDebugData::RegisterNodeToCodeAssociation() to remove the Macro-specific parameters and the additional insertions into the special Macro tables (no longer required).
	- Modified UK2Node_MathExpression::ValidateNodeDuringCompilation() to remove the special-case for Macro Instance source nodes, as Macro source nodes are now being mapped through the same table.
	- Added FindMatchingTunnelInstanceNode() as a utility method for now in BlueprintConnectionDrawingPolicy.cpp in order to match up Macro/Composite graph source nodes with nested Tunnel Instance nodes at the current graph level. *** TODO: For 4.19 we probably should revert back to using a secondary table in the debug data to map Tunnel Instance node hierarchies to code offsets in order to result in a faster lookup time here. ***
	- Modified FKismetConnectionDrawingPolicy::BuilldExecutionRoadmap() to replace the special-case for Macro Instance source nodes with a more general check for Tunnel Instance nodes that also handles Composite source nodes.
	- Revised UK2Node_TunnelBoundary to strip out most of what was being used to support the profiler, while keeping its basic compiled goto behavior in order to still function as a NOP node.
	- Added FKismetCompilerContext::SpawnIntermediateTunnelBoundaryNodes().
	- Modified FKismetCompilerContext::ExpandTunnelsAndMacros() to no longer overwrite intermediate Macro source node mappings in the lookup table with the Macro Instance source node that triggered the Macro graph expansion. Also revised the TunnelNode case to spawn intermediate TunnelBoundary (NOP) nodes around Macro and Composite gateways; this allows breakpoints to hit on the Tunnel nodes around a source graph expansion.
	- Modified FScriptBuilderBase::EmitInstrumentation() to remove special-case handling for Macro and Tunnel source nodes. These are now being mapped directly through the SourceBacktrackMap instead.
	- Removed alternate breakpoint icon assets for Macro Instance and Composite nodes (no longer needed).
	- Removed UK2Node::GetActiveBreakpointToolTipText() and its UK2Node_MacroInstance override (no longer required).
	- Removed special case in SGraphNodeK2Base::GetOverlayBrushes() for Macro Instance and Composite nodes (no longer needed).
	- Removed special-case mappings and interface methods for Tunnel nodes in FCompilerResultsLog (no longer required).
	- Removed the LineNumberToMacroSourceNodeMap and LineNumberToMacroInstanceNodeMap members from the FDebuggingInfoForSingleFunction struct (no longer in use).
	- Removed FBlueprintDebugData::FindMacroSourceNodeFromCodeLocation() and FindMacroInstanceNodesFromCodeLocation().
	- Removed FKismetDebugUtilities::FindMacroSourceNodeForCodeLocation() (no longer in use).
	- Removed special-case handling for Macro Instance nodes in FKismetDebugUtilities::OnScriptException() (no longer required). Macro source nodes are no longer being mapped to code offsets through a separate table, and we don't need to worry about saving/restoring the Active Object when debugging with a Macro Graph in the active tab.

Change 3776606 by mason.seay

	Updated content to prevent warning from appearing

Change 3777051 by [REDACTED_PERSON]

	ComponentTemplate references in UBlueprint can no be cleared after compiling the (blueprint defined) component

Change 3777108 by [REDACTED_PERSON]

	Look up call stack frame source name when caching a script call stack for display. This relies on debug data being generated for event stubs

Change 3778277 by [REDACTED_PERSON]

	Fixed potential null material reference causing crash.

Change 3778288 by [REDACTED_PERSON]

	PR #3957: Making FAlphaBlend BlueprintType in order to fix a bunch of broken UPROPERTY's as of 4.17 (Contributed by [REDACTED_PERSON])

Change 3778321 by [REDACTED_PERSON]

	Fix for a regression in BP script execution behavior related to misidentified latent node expansions from a macro source graph.

	Change summary:
	- Removed FCompilerResultsLog::FullSourceBacktrackMap (no longer in use).
	- Restored FCompilerResultsLog::IntermediateTunnelNodeToTunnelInstanceMap (which was in place prior to CL# 37754112); this table was being used to map intermediate nodes resulting from a tunnel instance node expansion back to the outer tunnel instance node that triggered the expansion. Its once again being used for that reason, but I reduced the scope a bit to only include the execution path within the expansion, as that's the only mapping that we need.
	- Restored FCompilerResultsLog::RegisterIntermediateTunnelNode(), but renamed it to NotifyIntermediateTunnelNode() to be consistent with the other parts of the MessageLog interface, and also removed the part of the implementation that was adding to a secondary macro expansion-to-source backtrack map (since macro expansion node lookup is now done through the main source backtrack map).
	- Restored FCompilerResultsLog::GetIntermediateTunnelInstance().
	- Modified FCompilerResultsLog::NotifyIntermediateObjectCreation() to remove the part of the implementation that was adding to the secondary node-only-to-source backtrack map (it was previously just a redundant copy of the main one except in the case of macro expansions).
	- Modified FCompilerResultsLog::CalculateStableIdentifierForLatentActionManager() to restore the calculation of a stable UUID for nodes sourced from a macro expansion, where we had incorporated the outer intermediate tunnel instance node chain.

Change 3778329 by [REDACTED_PERSON]

	PR #4241: Enforce calling superclass on ActorComponent::BeginPlay (Contributed by [REDACTED_PERSON])

Change 3778349 by [REDACTED_PERSON]

	Minor cleanup

Change 3759702 by [REDACTED_PERSON]

	Right now it doesn't package them at all, we could change it to package them as Debug Non-UFS if desired
	Change it so the asset audit UI will only load DevelopmentAssetRegistry.bin files, the cooked registry files don't have enough information any more to be useful
	Remove ability for runtime game to load DevelopmentAssetRegistry.bin, this ended up not being useful

Change 3766414 by [REDACTED_PERSON]

	Data validation plugin

Change 3769923 by [REDACTED_PERSON]

	Fix Texture EstimatedTotal memory to handle LOD bias, it now reports the largest possible size in a cooked game of any platform
	Fix many GetResourceSizeEx calls to match the new definition and improve accuracy
	Switched several editor tools to use EstimatedTotal now that it is more useful, and removed some unused memory stats
	Remove ResourceSize from UObject asset registry tags as it was misleading and inaccurate, for now it is only possible to get this for loaded objects
	Remove MapFileSize from Worlds as it redundant with the generic file size. Fixed the generic file size to display using the Size format
	Several UI fixes for Asset Audit and Size Map to deal with this change. Asset Audit no longer has the memory size columns, and the memory size drop down in Size Map is disabled for cooked builds

Change 3771365 by [REDACTED_PERSON]

diff --git a/Engine/Plugins/Editor/AssetManagerEditor/Source/AssetManagerEditor/Private/AssetManagerEditorModule.cpp b/Engine/Plugins/Editor/AssetManagerEditor/Source/AssetManagerEditor/Private/AssetManagerEditorModule.cpp
--- a/Engine/Plugins/Editor/AssetManagerEditor/Source/AssetManagerEditor/Private/AssetManagerEditorModule.cpp
+++ b/Engine/Plugins/Editor/AssetManagerEditor/Source/AssetManagerEditor/Private/AssetManagerEditorModule.cpp
@@ -959,8 +959,8 @@ bool FAssetManagerEditorModule::GetStringValueForCustomColumn(const FAssetData&
 			}
 			else
 			{
-				// Sizes are converted to kb
-				OutValue = Lex::ToString((IntegerValue + 512) / 1024);
+				// Display size properly
+				OutValue = FText::AsMemory(IntegerValue).ToString();
 			}
 			return true;
 		}
@@ -1084,6 +1084,17 @@ bool FAssetManagerEditorModule::GetIntegerValueForCustomColumn(const FAssetData&
 			return true;
 		}
 	}
+	else if (ColumnName == ResourceSizeName)
+	{
+		// Resource size can currently only be calculated for loaded assets, so load and check
+		UObject* Asset = AssetData.GetAsset();
+
+		if (Asset)
+		{
+			OutValue = Asset->GetResourceSizeBytes(EResourceSizeMode::EstimatedTotal);
+			return true;
+		}
+	}
 	else if (ColumnName == TotalUsageName)
 	{
 		int64 TotalWeight = 0;
@@ -1110,11 +1121,6 @@ bool FAssetManagerEditorModule::GetIntegerValueForCustomColumn(const FAssetData&
 		// Get base value of asset tag
 		if (AssetData.GetTagValue(ColumnName, OutValue))
 		{
-			// ResourceSize is in kb, convert back to bytes
-			if (ColumnName == ResourceSizeName)
-			{
-				OutValue *= 1024;
-			}
 			return true;
 		}
 	}
@@ -1156,19 +1162,14 @@ FString FAssetManagerEditorModule::GetSavedAssetRegistryPath(ITargetPlatform* Ta
 	FParse::Value(FCommandLine::Get(), TEXT("AssetRegistryFile="), CommandLinePath);
 	CommandLinePath.ReplaceInline(TEXT("[Platform]"), *PlatformName);
 	
-	// First try DevelopmentAssetRegistry.bin, then try AssetRegistry.bin
-	FString CookedAssetRegistry = FPaths::ProjectDir() / TEXT("AssetRegistry.bin");
-
-	FString CookedPath = CookedSandbox->ConvertToAbsolutePathForExternalAppForWrite(*CookedAssetRegistry).Replace(TEXT("[Platform]"), *PlatformName);
-	FString DevCookedPath = CookedPath.Replace(TEXT("AssetRegistry.bin"), TEXT("DevelopmentAssetRegistry.bin"));
+	// We can only load DevelopmentAssetRegistry, the normal asset registry doesn't have enough data to be useful
+	FString CookedDevelopmentAssetRegistry = FPaths::ProjectDir() / TEXT("Metadata") / TEXT("DevelopmentAssetRegistry.bin");
 
-	FString EditorCookedPath = EditorCookedSandbox->ConvertToAbsolutePathForExternalAppForWrite(*CookedAssetRegistry).Replace(TEXT("[Platform]"), *PlatformName);
-	FString DevEditorCookedPath = EditorCookedPath.Replace(TEXT("AssetRegistry.bin"), TEXT("DevelopmentAssetRegistry.bin"));
+	FString DevCookedPath = CookedSandbox->ConvertToAbsolutePathForExternalAppForWrite(*CookedDevelopmentAssetRegistry).Replace(TEXT("[Platform]"), *PlatformName);
+	FString DevEditorCookedPath = EditorCookedSandbox->ConvertToAbsolutePathForExternalAppForWrite(*CookedDevelopmentAssetRegistry).Replace(TEXT("[Platform]"), *PlatformName);
+	FString DevSharedCookedPath = FPaths::Combine(*FPaths::ProjectSavedDir(), TEXT("SharedIterativeBuild"), PlatformName, TEXT("Metadata"), TEXT("DevelopmentAssetRegistry.bin"));
 
-	FString SharedCookedPath = FPaths::Combine(*FPaths::ProjectSavedDir(), TEXT("SharedIterativeBuild"), PlatformName, TEXT("Cooked"), TEXT("AssetRegistry.bin"));
-	FString DevSharedCookedPath = SharedCookedPath.Replace(TEXT("AssetRegistry.bin"), TEXT("DevelopmentAssetRegistry.bin"));
-
-	// Try command line, then cooked, then build
+	// Try command line, then cooked, then shared build
 	if (!CommandLinePath.IsEmpty() && IFileManager::Get().FileExists(*CommandLinePath))
 	{
 		return CommandLinePath;
@@ -1179,31 +1180,16 @@ FString FAssetManagerEditorModule::GetSavedAssetRegistryPath(ITargetPlatform* Ta
 		return DevCookedPath;
 	}
 
-	if (IFileManager::Get().FileExists(*CookedPath))
-	{
-		return CookedPath;
-	}
-
 	if (IFileManager::Get().FileExists(*DevEditorCookedPath))
 	{
 		return DevEditorCookedPath;
 	}
 
-	if (IFileManager::Get().FileExists(*EditorCookedPath))
-	{
-		return EditorCookedPath;
-	}
-
 	if (IFileManager::Get().FileExists(*DevSharedCookedPath))
 	{
 		return DevSharedCookedPath;
 	}
 
-	if (IFileManager::Get().FileExists(*SharedCookedPath))
-	{
-		return SharedCookedPath;
-	}
-
 	return FString();
 }
 
@@ -1230,7 +1216,7 @@ void FAssetManagerEditorModule::SetCurrentRegistrySource(const FString& SourceNa
 
 	FAssetManagerEditorRegistrySource* NewSource = RegistrySourceMap.Find(SourceName);
 
-	if (ensure(NewSource))
+	if (NewSource)
 	{
 		CurrentRegistrySource = NewSource;
 
@@ -1369,6 +1355,13 @@ void FAssetManagerEditorModule::SetCurrentRegistrySource(const FString& SourceNa
 			CurrentRegistrySource->bManagementDataInitialized = true;
 		}
 	}
+	else
+	{
+		FNotificationInfo Info(FText::Format(LOCTEXT("LoadRegistryFailed", "Can't find registry source {0}! Reverting to Editor."), FText::FromString(SourceName)));
+		Info.ExpireDuration = 10.0f;
+		FSlateNotificationManager::Get().AddNotification(Info);
+		CurrentRegistrySource = RegistrySourceMap.Find(FAssetManagerEditorRegistrySource::EditorSourceName);
+	}
 
 	check(CurrentRegistrySource);
 	
@@ -1391,15 +1384,13 @@ void FAssetManagerEditorModule::RefreshRegistryData()
 {
 	UAssetManager::Get().UpdateManagementDatabase();
 
-	// Reset registry source, then set it again
-	if (!CurrentRegistrySource->bIsEditor && CurrentRegistrySource->RegistryState)
-	{
-		delete CurrentRegistrySource->RegistryState;
-		CurrentRegistrySource->RegistryState = nullptr;
-	}
+	// Rescan registry sources, try to restore the current one
+	FString OldSourceName = CurrentRegistrySource->SourceName;
+
+	CurrentRegistrySource = nullptr;
+	InitializeRegistrySources(false);
 
-	CurrentRegistrySource->bManagementDataInitialized = false;
-	SetCurrentRegistrySource(CurrentRegistrySource->SourceName);
+	SetCurrentRegistrySource(OldSourceName);
 }
 
 bool FAssetManagerEditorModule::IsPackageInCurrentRegistrySource(FName PackageName)
diff --git a/Engine/Plugins/Editor/AssetManagerEditor/Source/AssetManagerEditor/Private/SAssetAuditBrowser.cpp b/Engine/Plugins/Editor/AssetManagerEditor/Source/AssetManagerEditor/Private/SAssetAuditBrowser.cpp
--- a/Engine/Plugins/Editor/AssetManagerEditor/Source/AssetManagerEditor/Private/SAssetAuditBrowser.cpp
+++ b/Engine/Plugins/Editor/AssetManagerEditor/Source/AssetManagerEditor/Private/SAssetAuditBrowser.cpp
@@ -290,10 +290,8 @@ void SAssetAuditBrowser::Construct(const FArguments& InArgs)
 	// Add custom columns
 	Config.CustomColumns.Emplace(FPrimaryAssetId::PrimaryAssetTypeTag, LOCTEXT("AssetType", "Primary Type"), LOCTEXT("AssetTypeTooltip", "Primary Asset Type of this asset, if set"), UObject::FAssetRegistryTag::TT_Alphabetical, FOnGetCustomAssetColumnData::CreateSP(this, &SAssetAuditBrowser::GetStringValueForCustomColumn));
 	Config.CustomColumns.Emplace(FPrimaryAssetId::PrimaryAssetNameTag, LOCTEXT("AssetName", "Primary Name"), LOCTEXT("AssetNameTooltip", "Primary Asset Name of this asset, if set"), UObject::FAssetRegistryTag::TT_Alphabetical, FOnGetCustomAssetColumnData::CreateSP(this, &SAssetAuditBrowser::GetStringValueForCustomColumn));
-	Config.CustomColumns.Emplace(IAssetManagerEditorModule::ManagedResourceSizeName, LOCTEXT("ManagedResourceSize", "Memory Kb"), LOCTEXT("ManagedResourceSizeTooltip", "Memory used by both this asset and any other assets it manages, in kilobytes"), UObject::FAssetRegistryTag::TT_Numerical, FOnGetCustomAssetColumnData::CreateSP(this, &SAssetAuditBrowser::GetStringValueForCustomColumn));
-	Config.CustomColumns.Emplace(IAssetManagerEditorModule::ResourceSizeName, LOCTEXT("ResourceSize", "Exclusive Memory Kb"), LOCTEXT("ResourceSizeTooltip", "Memory used exclusively by this asset, in kilobytes"), UObject::FAssetRegistryTag::TT_Numerical, FOnGetCustomAssetColumnData::CreateSP(this, &SAssetAuditBrowser::GetStringValueForCustomColumn));
-	Config.CustomColumns.Emplace(IAssetManagerEditorModule::ManagedDiskSizeName, LOCTEXT("ManagedDiskSize", "Disk Kb"), LOCTEXT("ManagedDiskSizeTooltip", "Total disk space used by both this and all managed assets, in kilobytes"), UObject::FAssetRegistryTag::TT_Numerical, FOnGetCustomAssetColumnData::CreateSP(this, &SAssetAuditBrowser::GetStringValueForCustomColumn));
-	Config.CustomColumns.Emplace(IAssetManagerEditorModule::DiskSizeName, LOCTEXT("DiskSize", "Exclusive Disk Kb"), LOCTEXT("DiskSizeTooltip", "Size of saved file on disk for only this asset, in kilobytes"), UObject::FAssetRegistryTag::TT_Numerical, FOnGetCustomAssetColumnData::CreateSP(this, &SAssetAuditBrowser::GetStringValueForCustomColumn));
+	Config.CustomColumns.Emplace(IAssetManagerEditorModule::ManagedDiskSizeName, LOCTEXT("ManagedDiskSize", "Disk Size"), LOCTEXT("ManagedDiskSizeTooltip", "Total disk space used by both this and all managed assets"), UObject::FAssetRegistryTag::TT_Numerical, FOnGetCustomAssetColumnData::CreateSP(this, &SAssetAuditBrowser::GetStringValueForCustomColumn));
+	Config.CustomColumns.Emplace(IAssetManagerEditorModule::DiskSizeName, LOCTEXT("DiskSize", "Exclusive Disk Size"), LOCTEXT("DiskSizeTooltip", "Size of saved file on disk for only this asset"), UObject::FAssetRegistryTag::TT_Numerical, FOnGetCustomAssetColumnData::CreateSP(this, &SAssetAuditBrowser::GetStringValueForCustomColumn));
 	Config.CustomColumns.Emplace(IAssetManagerEditorModule::TotalUsageName, LOCTEXT("TotalUsage", "Total Usage"), LOCTEXT("TotalUsageTooltip", "Weighted count of Primary Assets that use this, higher usage means it's more likely to be in memory at runtime"), UObject::FAssetRegistryTag::TT_Numerical, FOnGetCustomAssetColumnData::CreateSP(this, &SAssetAuditBrowser::GetStringValueForCustomColumn));
 	Config.CustomColumns.Emplace(IAssetManagerEditorModule::CookRuleName, LOCTEXT("CookRule", "Cook Rule"), LOCTEXT("CookRuleTooltip", "Rather this asset will be cooked or not"), UObject::FAssetRegistryTag::TT_Alphabetical, FOnGetCustomAssetColumnData::CreateSP(this, &SAssetAuditBrowser::GetStringValueForCustomColumn));
 	Config.CustomColumns.Emplace(IAssetManagerEditorModule::ChunksName, LOCTEXT("Chunks", "Chunks"), LOCTEXT("ChunksTooltip", "List of chunks this will be added to when cooked"), UObject::FAssetRegistryTag::TT_Alphabetical, FOnGetCustomAssetColumnData::CreateSP(this, &SAssetAuditBrowser::GetStringValueForCustomColumn));
@@ -846,6 +844,16 @@ void SAssetAuditBrowser::SetCurrentRegistrySource(const FAssetManagerEditorRegis
 {
 	CurrentRegistrySource = RegistrySource;
 
+	// Refresh dropdown
+	TArray<const FAssetManagerEditorRegistrySource*> AvailableSources;
+	IAssetManagerEditorModule::Get().GetAvailableRegistrySources(AvailableSources);
+	SourceComboList.Reset();
+
+	for (const FAssetManagerEditorRegistrySource* ValidSource : AvailableSources)
+	{
+		SourceComboList.Add(MakeShared<FString>(ValidSource->SourceName));
+	}
+
 	RefreshAssetView();
 }
 
diff --git a/Engine/Plugins/Editor/AssetManagerEditor/Source/AssetManagerEditor/Private/SSizeMap.cpp b/Engine/Plugins/Editor/AssetManagerEditor/Source/AssetManagerEditor/Private/SSizeMap.cpp
--- a/Engine/Plugins/Editor/AssetManagerEditor/Source/AssetManagerEditor/Private/SSizeMap.cpp
+++ b/Engine/Plugins/Editor/AssetManagerEditor/Source/AssetManagerEditor/Private/SSizeMap.cpp
@@ -2,7 +2,6 @@
 
 #include "SSizeMap.h"
 #include "Modules/ModuleManager.h"
-#include "Engine/Texture2D.h"
 #include "Engine/AssetManager.h"
 #include "Editor.h"
 #include "AssetRegistryModule.h"
@@ -10,7 +9,6 @@
 #include "ClassIconFinder.h"
 #include "Math/UnitConversion.h"
 #include "Widgets/Text/STextBlock.h"
-#include "Widgets/Input/SComboBox.h"
 #include "Widgets/Input/SButton.h"
 #include "Framework/Commands/UIAction.h"
 #include "Framework/Commands/UICommandList.h"
@@ -22,6 +20,8 @@
 #include "DragAndDrop/AssetDragDropOp.h"
 #include "CollectionManagerModule.h"
 #include "ICollectionManager.h"
+#include "SlateApplication.h"
+#include "Misc/ScopedSlowTask.h"
 
 #define LOCTEXT_NAMESPACE "SizeMap"
 
@@ -64,6 +64,8 @@ TSharedRef<SWidget> SSizeMap::GenerateSizeTypeComboItem(TSharedPtr<FName> InItem
 
 void SSizeMap::HandleSizeTypeComboChanged(TSharedPtr<FName> Item, ESelectInfo::Type SelectInfo)
 {
+	FSlateApplication::Get().DismissAllMenus();
+
 	CurrentSizeType = *Item.Get();
 
 	RefreshMap();
@@ -87,6 +89,11 @@ FText SSizeMap::GetSizeTypeText(FName SizeType) const
 	return FText::GetEmpty();
 }
 
+bool SSizeMap::IsSizeTypeEnabled() const
+{
+	return CurrentRegistrySource && CurrentRegistrySource->bIsEditor;
+}
+
 FText SSizeMap::GetOverviewText() const
 {
 	return OverviewText;
@@ -143,6 +150,7 @@ void SSizeMap::Construct(const FArguments& InArgs)
 
 	CurrentSizeType = IAssetManagerEditorModule::DiskSizeName;
 	CurrentRegistrySource = EditorModule->GetCurrentRegistrySource(false);
+	bMemorySizeCached = false;
 
 	ChildSlot
 	[
@@ -194,10 +202,12 @@ void SSizeMap::Construct(const FArguments& InArgs)
 			.Padding(2.f)
 			.AutoWidth()
 			[
-				SNew(SComboBox<TSharedPtr<FName>>)
+				SAssignNew(ComboBoxWidget, SComboBox<TSharedPtr<FName>>)
 				.OptionsSource(&SizeTypeComboList)
 				.OnGenerateWidget(this, &SSizeMap::GenerateSizeTypeComboItem)
 				.OnSelectionChanged(this, &SSizeMap::HandleSizeTypeComboChanged)
+				.IsEnabled(this, &SSizeMap::IsSizeTypeEnabled)
+				.ToolTipText(LOCTEXT("SizeType_Tooltip", "Change to display disk size or memory size, memory size is only available for editor data"))
 				[
 					SNew(STextBlock)
 					.Text(this, &SSizeMap::GetSizeTypeComboText)
@@ -258,15 +268,29 @@ FReply SSizeMap::OnDrop(const FGeometry& MyGeometry, const FDragDropEvent& DragD
 
 void SSizeMap::SetRootAssetIdentifiers(const TArray<FAssetIdentifier>& NewRootAssetIdentifiers)
 {
-	RootAssetIdentifiers = NewRootAssetIdentifiers;
-
+	if (RootAssetIdentifiers != NewRootAssetIdentifiers)
+	{
+		bMemorySizeCached = false;
+		RootAssetIdentifiers = NewRootAssetIdentifiers;
+	}
+	
 	RefreshMap();
 }
 
 void SSizeMap::SetCurrentRegistrySource(const FAssetManagerEditorRegistrySource* RegistrySource)
 {
-	CurrentRegistrySource = RegistrySource;
-
+	if (CurrentRegistrySource != RegistrySource)
+	{
+		CurrentRegistrySource = RegistrySource;
+		bMemorySizeCached = false;
+		if (!IsSizeTypeEnabled())
+		{
+			// If size type is disabled reset to disk size
+			CurrentSizeType = IAssetManagerEditorModule::DiskSizeName;
+			ComboBoxWidget->ClearSelection();
+		}
+	}
+	
 	RefreshMap();
 }
 
@@ -576,21 +600,12 @@ void SSizeMap::GatherDependenciesRecursively(TSharedPtr<FAssetThumbnailPool>& In
 
 				if (AssetPackageName != NAME_None)
 				{
-					// If we're reading cooked data, this will fail for dependencies that are editor only. This is fine, they will have 0 size
 					if (EditorModule->GetIntegerValueForCustomColumn(NodeSizeMapData.AssetData, CurrentSizeType, FoundSize))
 					{
+						// If we're reading cooked data, this will fail for dependencies that are editor only. This is fine, they will have 0 size
 						NodeSizeMapData.AssetSize = FoundSize;
 						NodeSizeMapData.bHasKnownSize = true;
 					}
-
-					if (NodeSizeMapData.AssetSize == 0 && CurrentRegistrySource->bIsEditor && CurrentSizeType == IAssetManagerEditorModule::ResourceSizeName)
-					{
-						// Load asset and check
-						UObject* Asset = NodeSizeMapData.AssetData.GetAsset();
-
-						NodeSizeMapData.AssetSize = Asset->GetResourceSizeBytes(EResourceSizeMode::Exclusive);
-						NodeSizeMapData.bHasKnownSize = true;
-					}
 				}
 				else
 				{
@@ -781,6 +796,20 @@ void SSizeMap::RefreshMap()
 		return;
 	}
 
+	bool bShowSlowTask = false;
+	if (CurrentSizeType == IAssetManagerEditorModule::ResourceSizeName)
+	{
+		if (!bMemorySizeCached)
+		{
+			bMemorySizeCached = true;
+			bShowSlowTask = true;
+		}
+	}
+
+	// If we're refresing memory start a slow task as we need to laod assets
+	FScopedSlowTask SlowTask(0, LOCTEXT("ComputeMemorySizeSlowTask", "Finding Memory Size..."), bShowSlowTask);
+	SlowTask.MakeDialog();
+
 	// Wipe the current tree out
 	RootTreeMapNode->Children.Empty();
 	NodeSizeMapDataMap.Empty();
@@ -873,6 +902,7 @@ void SSizeMap::RefreshMap()
 
 void SSizeMap::OnInitialAssetRegistrySearchComplete()
 {
+	bMemorySizeCached = false;
 	RefreshMap();
 }
 
diff --git a/Engine/Plugins/Editor/AssetManagerEditor/Source/AssetManagerEditor/Private/SSizeMap.h b/Engine/Plugins/Editor/AssetManagerEditor/Source/AssetManagerEditor/Private/SSizeMap.h
--- a/Engine/Plugins/Editor/AssetManagerEditor/Source/AssetManagerEditor/Private/SSizeMap.h
+++ b/Engine/Plugins/Editor/AssetManagerEditor/Source/AssetManagerEditor/Private/SSizeMap.h
@@ -5,6 +5,7 @@
 #include "CoreMinimal.h"
 #include "Widgets/DeclarativeSyntaxSupport.h"
 #include "Widgets/SCompoundWidget.h"
+#include "Widgets/Input/SComboBox.h"
 #include "AssetData.h"
 #include "AssetManagerEditorModule.h"
 #include "ITreeMap.h"
@@ -105,10 +106,13 @@ protected:
 	FText GetSizeTypeComboText() const;
 	FText GetSizeTypeText(FName SizeType) const;
 	FText GetOverviewText() const;
+	bool IsSizeTypeEnabled() const;
 
 	TArray<TSharedPtr<FName>> SizeTypeComboList;
 	FName CurrentSizeType;
 	FText OverviewText;
+	bool bMemorySizeCached;
+	TSharedPtr<SComboBox<TSharedPtr<FName>>> ComboBoxWidget;
 
 	/** Our tree map widget */
 	TSharedPtr<class STreeMap> TreeMapWidget;
diff --git a/Engine/Plugins/Editor/DataValidation/Source/DataValidation/Private/DataValidationCommandlet.cpp b/Engine/Plugins/Editor/DataValidation/Source/DataValidation/Private/DataValidationCommandlet.cpp
new file mode 100644
--- /dev/null
+++ b/Engine/Plugins/Editor/DataValidation/Source/DataValidation/Private/DataValidationCommandlet.cpp
@@ -0,0 +1,49 @@
+// Copyright 1998-2017 Epic Games, Inc. All Rights Reserved.
+
+#include "DataValidationCommandlet.h"
+#include "DataValidationManager.h"
+
+#include "IAssetRegistry.h"
+#include "AssetRegistryHelpers.h"
+#include "AssetRegistryModule.h"
+
+
+DEFINE_LOG_CATEGORY_STATIC(LogDataValidation, Warning, All);
+
+// Commandlet for validating data
+int32 UDataValidationCommandlet::Main(const FString& FullCommandLine)
+{
+	UE_LOG(LogDataValidation, Log, TEXT("--------------------------------------------------------------------------------------------"));
+	UE_LOG(LogDataValidation, Log, TEXT("Running DataValidation Commandlet"));
+	TArray<FString> Tokens;
+	TArray<FString> Switches;
+	TMap<FString, FString> Params;
+	ParseCommandLine(*FullCommandLine, Tokens, Switches, Params);
+
+	// validate data
+	if (!ValidateData())
+	{
+		UE_LOG(LogDataValidation, Warning, TEXT("Errors occurred while validating data"));
+		return 2; // return something other than 1 for error since the engine will return 1 if any other system (possibly unrelated) logged errors during execution.
+	}
+
+	UE_LOG(LogDataValidation, Log, TEXT("Successfully finished running DataValidation Commandlet"));
+	UE_LOG(LogDataValidation, Log, TEXT("--------------------------------------------------------------------------------------------"));
+	return 0;
+}
+
+//static
+bool UDataValidationCommandlet::ValidateData()
+{
+	FAssetRegistryModule& AssetRegistryModule = FModuleManager::LoadModuleChecked<FAssetRegistryModule>(AssetRegistryConstants::ModuleName);
+
+	TArray<FAssetData> AssetDataList;
+	AssetRegistryModule.Get().GetAllAssets(AssetDataList);
+
+	UDataValidationManager* DataValidationManager = UDataValidationManager::Get();
+	check(DataValidationManager);
+
+	DataValidationManager->ValidateAssets(AssetDataList);
+
+	return true;
+}
diff --git a/Engine/Plugins/Editor/DataValidation/Source/DataValidation/Private/DataValidationManager.cpp b/Engine/Plugins/Editor/DataValidation/Source/DataValidation/Private/DataValidationManager.cpp
new file mode 100644
--- /dev/null
+++ b/Engine/Plugins/Editor/DataValidation/Source/DataValidation/Private/DataValidationManager.cpp
@@ -0,0 +1,173 @@
+// Copyright 1998-2017 Epic Games, Inc. All Rights Reserved.
+
+#include "DataValidationManager.h"
+
+#include "ModuleManager.h"
+#include "Developer/MessageLog/Public/MessageLogModule.h"
+#include "Logging/MessageLog.h"
+#include "ScopedSlowTask.h"
+
+#include "CoreGlobals.h"
+
+#define LOCTEXT_NAMESPACE "DataValidationManager"
+
+UDataValidationManager* GDataValidationManager = nullptr;
+
+/**
+ * UDataValidationManager
+ */
+
+UDataValidationManager::UDataValidationManager(const FObjectInitializer& ObjectInitializer)
+	: Super(ObjectInitializer)
+{
+	DataValidationManagerClassName = FSoftClassPath(TEXT("/Script/DataValidation.DataValidationManager"));
+}
+
+UDataValidationManager* UDataValidationManager::Get()
+{
+	if (GDataValidationManager == nullptr)
+	{
+		FSoftClassPath DataValidationManagerClassName = (UDataValidationManager::StaticClass()->GetDefaultObject<UDataValidationManager>())->DataValidationManagerClassName;
+
+		UClass* SingletonClass = DataValidationManagerClassName.TryLoadClass<UObject>();
+		checkf(SingletonClass != nullptr, TEXT("Data validation config value DataValidationManagerClassName is not a valid class name."));
+
+		GDataValidationManager = NewObject<UDataValidationManager>(GetTransientPackage(), SingletonClass, NAME_None);
+		checkf(GDataValidationManager != nullptr, TEXT("Data validation config value DataValidationManagerClassName is not a subclass of UDataValidationManager."))
+
+		GDataValidationManager->AddToRoot();
+		GDataValidationManager->Initialize();
+	}
+
+	return GDataValidationManager;
+}
+
+void UDataValidationManager::Initialize()
+{
+	FMessageLogInitializationOptions InitOptions;
+	InitOptions.bShowFilters = true;
+
+	FMessageLogModule& MessageLogModule = FModuleManager::LoadModuleChecked<FMessageLogModule>("MessageLog");
+	MessageLogModule.RegisterLogListing("DataValidation", LOCTEXT("DataValidation", "Data Validation"), InitOptions);
+}
+
+UDataValidationManager::~UDataValidationManager()
+{
+}
+
+EDataValidationResult UDataValidationManager::IsObjectValid(UObject* InObject, TArray<FText>& ValidationErrors) const
+{
+	if (ensure(InObject))
+	{
+		return InObject->IsDataValid(ValidationErrors);
+	}
+
+	return EDataValidationResult::NotValidated;
+}
+
+EDataValidationResult UDataValidationManager::IsAssetValid(FAssetData& AssetData, TArray<FText>& ValidationErrors) const
+{
+	if (AssetData.IsValid())
+	{
+		UObject* Obj = AssetData.GetAsset();
+		if (Obj)
+		{
+			return IsObjectValid(Obj, ValidationErrors);
+		}
+	}
+
+	return EDataValidationResult::Invalid;
+}
+
+void UDataValidationManager::ValidateAssets(TArray<FAssetData> AssetDataList, bool bSkipExcludedDirectories) const
+{
+	FScopedSlowTask SlowTask(1.0f, LOCTEXT("ValidatingDataTask", "Validating Data..."));
+	SlowTask.Visibility = ESlowTaskVisibility::ForceVisible;
+	SlowTask.MakeDialog();
+
+	FMessageLog DataValidationLog("DataValidation");
+
+	int32 NumAdded = 0;
+
+	int32 NumFilesChecked = 0;
+	int32 NumValidFiles = 0;
+	int32 NumInvalidFiles = 0;
+	int32 NumFilesSkipped = 0;
+	int32 NumFilesUnableToValidate = 0;
+
+	int32 NumFilesToValidate = AssetDataList.Num();
+
+	// Now add to map or update as needed
+	for (FAssetData& Data : AssetDataList)
+	{
+		SlowTask.EnterProgressFrame(1.0f / NumFilesToValidate, FText::Format(LOCTEXT("ValidatingFilename", "Validating {0}"), FText::FromString(Data.GetFullName())));
+
+		// Check exclusion path
+		if (bSkipExcludedDirectories && IsPathExcludedFromValidation(Data.PackageName.ToString()))
+		{
+			++NumFilesSkipped;
+			continue;
+		}
+
+		TArray<FText> ValidationErrors;
+		EDataValidationResult Result = IsAssetValid(Data, ValidationErrors);
+		++NumFilesChecked;
+
+		for (const FText& ErrorMsg : ValidationErrors)
+		{
+			FMessageLog("DataValidation").Error()->AddToken(FTextToken::Create(ErrorMsg));
+		}
+
+		if (Result == EDataValidationResult::Valid)
+		{
+			++NumValidFiles;
+		}
+		else
+		{
+			FFormatNamedArguments Arguments;
+			Arguments.Add(TEXT("Filename"), FText::FromName(Data.PackageName));
+			if (Result == EDataValidationResult::Invalid)
+			{
+				FMessageLog("DataValidation").Error()->AddToken(FTextToken::Create(FText::Format(LOCTEXT("InvalidDataResult", "{Filename} contains invalid data."), Arguments)));
+				++NumInvalidFiles;
+			}
+			else if (Result == EDataValidationResult::NotValidated)
+			{
+				FMessageLog("DataValidation").Info()->AddToken(FTextToken::Create(FText::Format(LOCTEXT("NotValidatedDataResult", "There is no data validation for {Filename}."), Arguments)));
+				++NumFilesUnableToValidate;
+			}
+		}
+	}
+
+	{
+		FFormatNamedArguments Arguments;
+		Arguments.Add(TEXT("Result"), NumInvalidFiles > 0 ? LOCTEXT("Failed", "FAILED") : LOCTEXT("Succeeded", "SUCCEEDED"));
+		Arguments.Add(TEXT("NumChecked"), NumFilesChecked);
+		Arguments.Add(TEXT("NumValid"), NumValidFiles);
+		Arguments.Add(TEXT("NumInvalid"), NumInvalidFiles);
+		Arguments.Add(TEXT("NumSkipped"), NumFilesSkipped);
+		Arguments.Add(TEXT("NumUnableToValidate"), NumFilesUnableToValidate);
+
+		TSharedRef<FTokenizedMessage> ValidationLog = NumInvalidFiles > 0 ? FMessageLog("DataValidation").Error() : FMessageLog("DataValidation").Info();
+
+		ValidationLog->AddToken(FTextToken::Create(FText::Format(LOCTEXT("SuccessOrFailure", "Data validation {Result}."), Arguments)));
+		ValidationLog->AddToken(FTextToken::Create(FText::Format(LOCTEXT("ResultsSummary", "Files Checked: {NumChecked}, Passed: {NumValid}, Failed: {NumInvalid}, Skipped: {NumSkipped}, Unable to validate: {NumUnableToValidate}"), Arguments)));
+	}
+
+	DataValidationLog.Open(EMessageSeverity::Info, true);
+}
+
+bool UDataValidationManager::IsPathExcludedFromValidation(const FString& Path) const
+{
+	for (const FDirectoryPath& ExcludedPath : ExcludedDirectories)
+	{
+		if (Path.Contains(ExcludedPath.Path))
+		{
+			return true;
+		}
+	}
+
+	return false;
+}
+
+#undef LOCTEXT_NAMESPACE
diff --git a/Engine/Plugins/Editor/DataValidation/Source/DataValidation/Private/DataValidationModule.cpp b/Engine/Plugins/Editor/DataValidation/Source/DataValidation/Private/DataValidationModule.cpp
new file mode 100644
--- /dev/null
+++ b/Engine/Plugins/Editor/DataValidation/Source/DataValidation/Private/DataValidationModule.cpp
@@ -0,0 +1,281 @@
+// Copyright 1998-2017 Epic Games, Inc. All Rights Reserved.
+
+#include "DataValidationModule.h"
+#include "DataValidationManager.h"
+#include "UObject/Object.h"
+#include "UObject/SoftObjectPath.h"
+#include "GameFramework/HUD.h"
+
+#include "SlateApplication.h"
+#include "MultiBoxBuilder.h"
+#include "MultiBoxExtender.h"
+#include "EditorStyleSet.h"
+#include "AssetRegistryModule.h"
+#include "MessageDialog.h"
+#include "DataValidationCommandlet.h"
+#include "LevelEditor.h"
+#include "Misc/FeedbackContext.h"
+#include "Misc/OutputDeviceConsole.h"
+#include "Modules/ModuleManager.h"
+
+#include "ContentBrowserModule.h"
+#include "ContentBrowserDelegates.h"
+
+#include "WorkspaceMenuStructure.h"
+#include "WorkspaceMenuStructureModule.h"
+
+#define LOCTEXT_NAMESPACE "DataValidationModule"
+
+class FDataValidationModule : public IDataValidationModule
+{
+	// Begin IModuleInterface
+	virtual void StartupModule() override;
+	virtual void ShutdownModule() override;
+	// End IModuleInterface
+
+	/** Validates selected assets and opens a window to report the results. If bValidateDependencies is true it will also validate any assets that the selected assets depend on. */
+	virtual void ValidateAssets(TArray<FAssetData> SelectedAssets, bool bValidateDependencies) override;
+
+	void ValidateFolders(TArray<FString> SelectedFolders);
+
+private:
+	TSharedRef<FExtender> OnExtendContentBrowserAssetSelectionMenu(const TArray<FAssetData>& SelectedAssets);
+	TSharedRef<FExtender> OnExtendContentBrowserPathSelectionMenu(const TArray<FString>& SelectedPaths);
+	void CreateDataValidationContentBrowserAssetMenu(FMenuBuilder& MenuBuilder, TArray<FAssetData> SelectedAssets);
+	void CreateDataValidationContentBrowserPathMenu(FMenuBuilder& MenuBuilder, TArray<FString> SelectedPaths);
+
+	// Adds Asset and any assets it depends on to the set DependentAssets
+	void FindAssetDependencies(const FAssetRegistryModule& AssetRegistryModule, const FAssetData& Asset, TSet<FAssetData>& DependentAssets);
+
+	static void DataValidationMenuCreationDelegate(FMenuBuilder& MenuBuilder);
+	static FText Menu_ValidateDataGetTitle();
+	static void Menu_ValidateData();
+
+	TSharedPtr<FExtender> MenuExtender;
+
+	FDelegateHandle ContentBrowserAssetExtenderDelegateHandle;
+	FDelegateHandle ContentBrowserPathExtenderDelegateHandle;
+};
+
+IMPLEMENT_MODULE(FDataValidationModule, DataValidation)
+
+void FDataValidationModule::StartupModule()
+{	
+	if (!IsRunningCommandlet() && !IsRunningGame() && FSlateApplication::IsInitialized())
+	{
+		// Register content browser hook
+		FContentBrowserModule& ContentBrowserModule = FModuleManager::LoadModuleChecked<FContentBrowserModule>(TEXT("ContentBrowser"));
+		TArray<FContentBrowserMenuExtender_SelectedAssets>& CBAssetMenuExtenderDelegates = ContentBrowserModule.GetAllAssetViewContextMenuExtenders();
+
+		CBAssetMenuExtenderDelegates.Add(FContentBrowserMenuExtender_SelectedAssets::CreateRaw(this, &FDataValidationModule::OnExtendContentBrowserAssetSelectionMenu));
+		ContentBrowserAssetExtenderDelegateHandle = CBAssetMenuExtenderDelegates.Last().GetHandle();
+
+		TArray<FContentBrowserMenuExtender_SelectedPaths>& CBFolderMenuExtenderDelegates = ContentBrowserModule.GetAllPathViewContextMenuExtenders();
+
+		CBFolderMenuExtenderDelegates.Add(FContentBrowserMenuExtender_SelectedPaths::CreateRaw(this, &FDataValidationModule::OnExtendContentBrowserPathSelectionMenu));
+		ContentBrowserPathExtenderDelegateHandle = CBFolderMenuExtenderDelegates.Last().GetHandle();
+
+		// add the File->DataValidation menu subsection
+
+		// make an extension to add the DataValidation function menu
+		MenuExtender = MakeShareable(new FExtender);
+		MenuExtender->AddMenuExtension("FileLoadAndSave", EExtensionHook::After, nullptr, FMenuExtensionDelegate::CreateStatic(&DataValidationMenuCreationDelegate));
+
+		FLevelEditorModule& LevelEditorModule = FModuleManager::LoadModuleChecked<FLevelEditorModule>("LevelEditor");
+		LevelEditorModule.GetMenuExtensibilityManager()->AddExtender(MenuExtender);
+	}
+}
+
+void FDataValidationModule::ShutdownModule()
+{
+	if (!IsRunningCommandlet() && !IsRunningGame())
+	{
+		FContentBrowserModule& ContentBrowserModule = FModuleManager::LoadModuleChecked<FContentBrowserModule>(TEXT("ContentBrowser"));
+		TArray<FContentBrowserMenuExtender_SelectedAssets>& CBMenuExtenderDelegates = ContentBrowserModule.GetAllAssetViewContextMenuExtenders();
+		CBMenuExtenderDelegates.RemoveAll([this](const FContentBrowserMenuExtender_SelectedAssets& Delegate) { return Delegate.GetHandle() == ContentBrowserAssetExtenderDelegateHandle; });
+		CBMenuExtenderDelegates.RemoveAll([this](const FContentBrowserMenuExtender_SelectedAssets& Delegate) { return Delegate.GetHandle() == ContentBrowserPathExtenderDelegateHandle; });
+
+		// remove menu extension
+		FLevelEditorModule& LevelEditorModule = FModuleManager::LoadModuleChecked<FLevelEditorModule>("LevelEditor");
+		LevelEditorModule.GetMenuExtensibilityManager()->RemoveExtender(MenuExtender);
+		MenuExtender = nullptr;
+	}
+}
+
+void FDataValidationModule::DataValidationMenuCreationDelegate(FMenuBuilder& MenuBuilder)
+{
+	MenuBuilder.BeginSection("DataValidation", LOCTEXT("DataValidation", "DataValidation"));
+	MenuBuilder.AddMenuEntry(
+		TAttribute<FText>::Create(&Menu_ValidateDataGetTitle),
+		LOCTEXT("ValidateDataTooltip", "Validates all user data in content directory."),
+		FSlateIcon(FEditorStyle::GetStyleSetName(), "DeveloperTools.MenuIcon"),
+		FUIAction(FExecuteAction::CreateStatic(&FDataValidationModule::Menu_ValidateData)));
+	MenuBuilder.EndSection();
+}
+
+FText FDataValidationModule::Menu_ValidateDataGetTitle()
+{
+	FAssetRegistryModule& AssetRegistryModule = FModuleManager::LoadModuleChecked<FAssetRegistryModule>("AssetRegistry");
+	if (AssetRegistryModule.Get().IsLoadingAssets())
+	{
+		return LOCTEXT("ValidateDataTitleDA", "Validate Data [Discovering Assets]");
+	}
+	return LOCTEXT("ValidateDataTitle", "Validate Data...");
+}
+
+void FDataValidationModule::Menu_ValidateData()
+{
+	// make sure the asset registry is finished building
+	FAssetRegistryModule& AssetRegistryModule = FModuleManager::LoadModuleChecked<FAssetRegistryModule>("AssetRegistry");
+	if (AssetRegistryModule.Get().IsLoadingAssets())
+	{
+		FMessageDialog::Open(EAppMsgType::Ok, LOCTEXT("AssetsStillScanningError", "Cannot run data validation while still discovering assets."));
+		return;
+	}
+
+	// validate the data
+	bool bSuccess = UDataValidationCommandlet::ValidateData();
+
+	// display an error if the task failed
+	if (!bSuccess)
+	{
+		FMessageDialog::Open(EAppMsgType::Ok, LOCTEXT("DataValidationError", "An error was encountered during data validation. See the log for details."));
+		return;
+	}
+}
+
+void FDataValidationModule::FindAssetDependencies(const FAssetRegistryModule& AssetRegistryModule, const FAssetData& Asset, TSet<FAssetData>& DependentAssets)
+{
+	if (Asset.IsValid())
+	{
+		UObject* Obj = Asset.GetAsset();
+		if (Obj)
+		{
+			const FName SelectedPackageName = Obj->GetOutermost()->GetFName();
+			FString PackageString = SelectedPackageName.ToString();
+			FString ObjectString = FString::Printf(TEXT("%s.%s"), *PackageString, *FPackageName::GetLongPackageAssetName(PackageString));
+
+			if (!DependentAssets.Contains(Asset))
+			{
+				DependentAssets.Add(Asset);
+
+				TArray<FName> Dependencies;
+				AssetRegistryModule.Get().GetDependencies(SelectedPackageName, Dependencies, EAssetRegistryDependencyType::Packages);
+
+				for (const FName Dependency : Dependencies)
+				{
+					const FString DependencyPackageString = Dependency.ToString();
+					FString DependencyObjectString = FString::Printf(TEXT("%s.%s"), *DependencyPackageString, *FPackageName::GetLongPackageAssetName(DependencyPackageString));
+
+					// recurse on each dependency
+					FName ObjectPath(*DependencyObjectString);
+					FAssetData DependentAsset = AssetRegistryModule.Get().GetAssetByObjectPath(ObjectPath);
+
+					FindAssetDependencies(AssetRegistryModule, DependentAsset, DependentAssets);
+				}
+			}
+		}
+	}
+}
+
+void FDataValidationModule::ValidateAssets(TArray<FAssetData> SelectedAssets, bool bValidateDependencies)
+{
+	FAssetRegistryModule& AssetRegistryModule = FModuleManager::LoadModuleChecked<FAssetRegistryModule>(TEXT("AssetRegistry"));
+
+	TSet<FAssetData> DependentAssets;
+	if (bValidateDependencies)
+	{
+		for (const FAssetData& Asset : SelectedAssets)
+		{
+			FindAssetDependencies(AssetRegistryModule, Asset, DependentAssets);
+		}
+
+		SelectedAssets = DependentAssets.Array();
+	}
+
+	UDataValidationManager* DataValidationManager = UDataValidationManager::Get();
+	if (DataValidationManager)
+	{
+		DataValidationManager->ValidateAssets(SelectedAssets, false);
+	}
+}
+
+void FDataValidationModule::ValidateFolders(TArray<FString> SelectedFolders)
+{
+	FAssetRegistryModule& AssetRegistryModule = FModuleManager::Get().LoadModuleChecked<FAssetRegistryModule>(TEXT("AssetRegistry"));
+
+	// Form a filter from the paths
+	FARFilter Filter;
+	Filter.bRecursivePaths = true;
+	for (const FString& Folder : SelectedFolders)
+	{
+		Filter.PackagePaths.Emplace(*Folder);
+	}
+
+	// Query for a list of assets in the selected paths
+	TArray<FAssetData> AssetList;
+	AssetRegistryModule.Get().GetAssets(Filter, AssetList);
+
+	ValidateAssets(AssetList, false);
+}
+
+// Extend content browser menu for groups of selected assets
+TSharedRef<FExtender> FDataValidationModule::OnExtendContentBrowserAssetSelectionMenu(const TArray<FAssetData>& SelectedAssets)
+{
+	TSharedRef<FExtender> Extender(new FExtender());
+
+	Extender->AddMenuExtension(
+		"AssetContextAdvancedActions",
+		EExtensionHook::After,
+		nullptr,
+		FMenuExtensionDelegate::CreateRaw(this, &FDataValidationModule::CreateDataValidationContentBrowserAssetMenu, SelectedAssets));
+
+	return Extender;
+}
+
+// Extend content browser menu for groups of selected assets
+void FDataValidationModule::CreateDataValidationContentBrowserAssetMenu(FMenuBuilder& MenuBuilder, TArray<FAssetData> SelectedAssets)
+{
+	MenuBuilder.AddMenuEntry
+	(
+		LOCTEXT("ValidateAssetsTabTitle", "Validate Assets"),
+		LOCTEXT("ValidateAssetsTooltipText", "Runs data validation on these assets."),
+		FSlateIcon(),
+		FUIAction(FExecuteAction::CreateRaw(this, &FDataValidationModule::ValidateAssets, SelectedAssets, false))
+	);
+	MenuBuilder.AddMenuEntry
+	(
+		LOCTEXT("ValidateAssetsAndDependenciesTabTitle", "Validate Assets and Dependencies"),
+		LOCTEXT("ValidateAssetsAndDependenciesTooltipText", "Runs data validation on these assets and all assets they depend on."),
+		FSlateIcon(),
+		FUIAction(FExecuteAction::CreateRaw(this, &FDataValidationModule::ValidateAssets, SelectedAssets, true))
+	);
+}
+
+// Extend content browser menu for groups of asset paths (folders)
+TSharedRef<FExtender> FDataValidationModule::OnExtendContentBrowserPathSelectionMenu(const TArray<FString>& SelectedPaths)
+{
+	TSharedRef<FExtender> Extender(new FExtender());
+
+	Extender->AddMenuExtension(
+		"PathContextBulkOperations",
+		EExtensionHook::After,
+		nullptr,
+		FMenuExtensionDelegate::CreateRaw(this, &FDataValidationModule::CreateDataValidationContentBrowserPathMenu, SelectedPaths));
+
+	return Extender;
+}
+
+// Extend content browser menu for groups of asset paths (folders)
+void FDataValidationModule::CreateDataValidationContentBrowserPathMenu(FMenuBuilder& MenuBuilder, TArray<FString> SelectedPaths)
+{
+	MenuBuilder.AddMenuEntry
+	(
+		LOCTEXT("ValidateAssetsPathTabTitle", "Validate Assets in Folder"),
+		LOCTEXT("ValidateAssetsPathTooltipText", "Runs data validation on the assets in the selected folder."),
+		FSlateIcon(),
+		FUIAction(FExecuteAction::CreateRaw(this, &FDataValidationModule::ValidateFolders, SelectedPaths))
+	);
+}
+
+#undef LOCTEXT_NAMESPACE
\ No newline at end of file
diff --git a/Engine/Plugins/Editor/DataValidation/Source/DataValidation/Public/DataValidationCommandlet.h b/Engine/Plugins/Editor/DataValidation/Source/DataValidation/Public/DataValidationCommandlet.h
new file mode 100644
--- /dev/null
+++ b/Engine/Plugins/Editor/DataValidation/Source/DataValidation/Public/DataValidationCommandlet.h
@@ -0,0 +1,26 @@
+// Copyright 1998-2017 Epic Games, Inc. All Rights Reserved.
+#pragma once
+
+#include "Commandlets/Commandlet.h"
+#include "UObject/Interface.h"
+#include "DataValidationCommandlet.generated.h"
+
+UCLASS(CustomConstructor)
+class DATAVALIDATION_API UDataValidationCommandlet : public UCommandlet
+{
+	GENERATED_UCLASS_BODY()
+
+public:
+	UDataValidationCommandlet(const FObjectInitializer& ObjectInitializer)
+		: Super(ObjectInitializer)
+	{
+		LogToConsole = false;
+	}
+
+	// Begin UCommandlet Interface
+	virtual int32 Main(const FString& Params) override;
+	// End UCommandlet Interface
+
+	// do the validation without creating a commandlet
+	static bool ValidateData();
+};
diff --git a/Engine/Plugins/Editor/DataValidation/Source/DataValidation/Public/DataValidationManager.h b/Engine/Plugins/Editor/DataValidation/Source/DataValidation/Public/DataValidationManager.h
new file mode 100644
--- /dev/null
+++ b/Engine/Plugins/Editor/DataValidation/Source/DataValidation/Public/DataValidationManager.h
@@ -0,0 +1,73 @@
+// Copyright 1998-2017 Epic Games, Inc. All Rights Reserved.
+
+#pragma once
+
+#include "CoreMinimal.h"
+#include "UObject/Object.h"
+#include "AssetData.h"
+#include "Engine/EngineTypes.h"
+#include "DataValidationManager.generated.h"
+
+class FLogWindowManager;
+class SLogWindow;
+class SWindow;
+class UUnitTest;
+
+/**
+ * Manages centralized execution and tracking of data validation, as well as handling console commands,
+ * and some misc tasks like local log hooking
+ */
+
+UCLASS(config=Editor)
+class DATAVALIDATION_API UDataValidationManager : public UObject
+{
+	GENERATED_UCLASS_BODY()
+
+public:
+	/**
+	 * Static getter for the data validation manager
+	 *
+	 * @return	Returns the data validation manager
+	 */
+	static UDataValidationManager* Get();
+
+	/**
+	 * Initialize the data validation manager
+	 */
+	virtual void Initialize();
+
+	/**
+	 * Destructor for handling removal of log registration
+	 */
+	virtual ~UDataValidationManager() override;
+
+	/**
+	 * @return Returns Valid if the object contains valid data; returns Invalid if the object contains invalid data; returns NotValidated if no validations was performed on the object
+	 */
+	virtual EDataValidationResult IsObjectValid(UObject* InObject, TArray<FText>& ValidationErrors) const;
+
+	/**
+	 * @return Returns Valid if the object pointed to by AssetData contains valid data; returns Invalid if the object contains invalid data or does not exist; returns NotValidated if no validations was performed on the object
+	 */
+	virtual EDataValidationResult IsAssetValid(FAssetData& AssetData, TArray<FText>& ValidationErrors) const;
+
+	virtual void ValidateAssets(TArray<FAssetData> AssetDataList, bool bSkipExcludedDirectories = true) const;
+
+protected:
+	/** 
+	 * @return Returns true if the current Path should be skipped for validation. Returns false otherwise.
+	 */
+	virtual bool IsPathExcludedFromValidation(const FString& Path) const;
+
+	/**
+	 * Directories to ignore for data validation. Useful for test assets
+	 */
+	UPROPERTY(config)
+	TArray<FDirectoryPath> ExcludedDirectories;
+
+private:
+
+	/** The class to instantiate as the manager object. Defaults to this class but can be overridden */
+	UPROPERTY(config)
+	FSoftClassPath DataValidationManagerClassName;
+};
diff --git a/Engine/Source/Developer/BlueprintCompilerCppBackend/Private/BlueprintCompilerCppBackendUtils.cpp b/Engine/Source/Developer/BlueprintCompilerCppBackend/Private/BlueprintCompilerCppBackendUtils.cpp
--- a/Engine/Source/Developer/BlueprintCompilerCppBackend/Private/BlueprintCompilerCppBackendUtils.cpp
+++ b/Engine/Source/Developer/BlueprintCompilerCppBackend/Private/BlueprintCompilerCppBackendUtils.cpp
@@ -1410,19 +1410,25 @@ bool FEmitHelper::GenerateAutomaticCast(FEmitterLocalContext& EmitterContext, co
 			return RType.IsArray() && LClass && RClass && (LClass->IsChildOf(RClass) || RClass->IsChildOf(LClass)) && (LClass != RClass);
 		};
 
+		// No need to check both types here because we already checked for a match above.
 		const bool bIsClassTerm = LType.PinCategory == UEdGraphSchema_K2::PC_Class;
-		auto GetTypeString = [bIsClassTerm](UClass* TermType, const UObjectProperty* AssociatedProperty)->FString
+		const bool bIsObjectTerm = LType.PinCategory == UEdGraphSchema_K2::PC_Object;
+		const bool bIsSoftObjTerm = LType.PinCategory == UEdGraphSchema_K2::PC_SoftClass || LType.PinCategory == UEdGraphSchema_K2::PC_SoftObject;
+		auto GetTypeString = [bIsClassTerm, bIsSoftObjTerm](UClass* TermType, const UObjectPropertyBase* AssociatedProperty)->FString
 		{
 			// favor the property's CPPType since it makes choices based off of things like CPF_UObjectWrapper (which 
 			// adds things like TSubclassof<> to the decl)... however, if the property type doesn't match the term 
 			// type, then ignore the property (this can happen for things like our array library, which uses wildcards 
 			// and custom thunks to allow differing types)
-			const bool bPropertyMatch = AssociatedProperty && ((AssociatedProperty->PropertyClass == TermType) ||
+			const bool bPropertyMatch = AssociatedProperty &&
+				(bIsSoftObjTerm ||
+				(AssociatedProperty->PropertyClass == TermType) ||
 				(bIsClassTerm && CastChecked<UClassProperty>(AssociatedProperty)->MetaClass == TermType));
 
 			if (bPropertyMatch)
 			{
 				// use GetCPPTypeCustom() so that it properly fills out nativized class names
+				// note: soft properties will use the term type that we pass in here in place of the internal MetaClass/PropertyClass, so we just always match them (above)
 				return AssociatedProperty->GetCPPTypeCustom(/*ExtendedTypeText=*/nullptr, CPPF_None, FEmitHelper::GetCppName(TermType));
 			}
 			else if (bIsClassTerm)
@@ -1434,15 +1440,16 @@ bool FEmitHelper::GenerateAutomaticCast(FEmitterLocalContext& EmitterContext, co
 
 		auto GetInnerTypeString = [GetTypeString](UClass* TermType, const UArrayProperty* ArrayProp)
 		{
-			const UObjectProperty* InnerProp = ArrayProp ? Cast<UObjectProperty>(ArrayProp->Inner) : nullptr;
+			const UObjectPropertyBase* InnerProp = ArrayProp ? Cast<UObjectPropertyBase>(ArrayProp->Inner) : nullptr;
 			return GetTypeString(TermType, InnerProp);
 		};
 
 		auto GenerateArrayCast = [&OutCastBegin, &OutCastEnd](const FString& LTypeStr, const FString& RTypeStr)
 		{
-			OutCastBegin = FString::Printf(TEXT("TArrayCaster< %s >("), *RTypeStr);
-			OutCastEnd   = FString::Printf(TEXT(").Get< %s >()"), *LTypeStr);
+			OutCastBegin = FString::Printf(TEXT("TArrayCaster<%s>("), *RTypeStr);
+			OutCastEnd   = FString::Printf(TEXT(").Get<%s>()"), *LTypeStr);
 		};
+
 		// CLASS/TSubClassOf<> to CLASS/TSubClassOf<>
 		if (bIsClassTerm)
 		{
@@ -1468,7 +1475,7 @@ bool FEmitHelper::GenerateAutomaticCast(FEmitterLocalContext& EmitterContext, co
 			}
 		}
 		// OBJECT to OBJECT
-		else if (LType.PinCategory == UEdGraphSchema_K2::PC_Object)
+		else if (bIsObjectTerm || bIsSoftObjTerm)
 		{
 			UClass* LClass = GetClassType(LType);
 			UClass* RClass = GetClassType(RType);
@@ -1476,14 +1483,25 @@ bool FEmitHelper::GenerateAutomaticCast(FEmitterLocalContext& EmitterContext, co
 			if (!RType.IsContainer() && LClass && RClass && (LType.bIsReference || bForceReference) && (LClass != RClass) && RClass->IsChildOf(LClass))
 			{
 				// when pointer is passed as reference, the type must be exactly the same
-				OutCastBegin = FString::Printf(TEXT("*(%s*)(&("), *GetTypeString(LClass, Cast<UObjectProperty>(LProp)));
+				OutCastBegin = FString::Printf(TEXT("*(%s*)(&("), *GetTypeString(LClass, Cast<UObjectPropertyBase>(LProp)));
 				OutCastEnd = TEXT("))");
 				return true;
 			}
 			if (!RType.IsContainer() && LClass && RClass && LClass->IsChildOf(RClass) && !RClass->IsChildOf(LClass))
 			{
-				OutCastBegin = FString::Printf(TEXT("CastChecked<%s>("), *FEmitHelper::GetCppName(LClass));
-				OutCastEnd = TEXT(", ECastCheckedType::NullAllowed)");
+				if (bIsObjectTerm)
+				{
+					OutCastBegin = FString::Printf(TEXT("CastChecked<%s>("), *FEmitHelper::GetCppName(LClass));
+					OutCastEnd = TEXT(", ECastCheckedType::NullAllowed)");
+				}
+				else
+				{
+					// TSoftClassPtr/TSoftObjectPtr cannot be implicitly downcast via assignment operator, but
+					// rather than emit an explicit cast here, we can just assign it to the wrapped object path.
+					OutCastBegin = TEXT("");
+					OutCastEnd = TEXT(".ToSoftObjectPath()");
+				}
+				
 				return true;
 			}
 			else if (RequiresArrayCast(LClass, RClass))
diff --git a/Engine/Source/Developer/TreeMap/STreeMap.cpp b/Engine/Source/Developer/TreeMap/STreeMap.cpp
--- a/Engine/Source/Developer/TreeMap/STreeMap.cpp
+++ b/Engine/Source/Developer/TreeMap/STreeMap.cpp
@@ -891,7 +891,7 @@ void STreeMap::TakeUndoSnapshot()
 
 FTreeMapNodeDataPtr STreeMap::FindNodeInCopiedTree( const FTreeMapNodeDataRef& NodeToFind, const FTreeMapNodeDataRef& OriginalNode, const FTreeMapNodeDataRef& CopiedNode ) const
 {
-	if( NodeToFind == OriginalNode )
+	if( NodeToFind.Get() == OriginalNode.Get() )
 	{
 		return CopiedNode;
 	}
diff --git a/Engine/Source/Editor/BlueprintGraph/Classes/K2Node.h b/Engine/Source/Editor/BlueprintGraph/Classes/K2Node.h
--- a/Engine/Source/Editor/BlueprintGraph/Classes/K2Node.h
+++ b/Engine/Source/Editor/BlueprintGraph/Classes/K2Node.h
@@ -232,9 +232,6 @@ class UK2Node : public UEdGraphNode
 	/** */
 	BLUEPRINTGRAPH_API virtual FText GetToolTipHeading() const;
 
-	/** Return tooltip text that explains the result of an active breakpoint on this node */
-	BLUEPRINTGRAPH_API virtual FText GetActiveBreakpointToolTipText() const;
-
 	/**
 	 * Determine if the node of this type should be filtered in the actions menu
 	 */
diff --git a/Engine/Source/Editor/BlueprintGraph/Classes/K2Node_BaseAsyncTask.h b/Engine/Source/Editor/BlueprintGraph/Classes/K2Node_BaseAsyncTask.h
--- a/Engine/Source/Editor/BlueprintGraph/Classes/K2Node_BaseAsyncTask.h
+++ b/Engine/Source/Editor/BlueprintGraph/Classes/K2Node_BaseAsyncTask.h
@@ -57,6 +57,9 @@ protected:
 	// Returns the factory function (checked)
 	UFunction* GetFactoryFunction() const;
 
+	/** Determines what the possible redirect pin names are **/
+	virtual void GetRedirectPinNames(const UEdGraphPin& Pin, TArray<FString>& RedirectPinNames) const;
+
 protected:
 	// The name of the function to call to create a proxy object
 	UPROPERTY()
diff --git a/Engine/Source/Editor/BlueprintGraph/Classes/K2Node_Knot.h b/Engine/Source/Editor/BlueprintGraph/Classes/K2Node_Knot.h
--- a/Engine/Source/Editor/BlueprintGraph/Classes/K2Node_Knot.h
+++ b/Engine/Source/Editor/BlueprintGraph/Classes/K2Node_Knot.h
@@ -54,8 +54,7 @@ public:
 
 private:
 	void PropagatePinType();
-	void PropagatePinTypeFromInput();
-	void PropagatePinTypeFromOutput();
+	void PropagatePinTypeFromDirection(bool bFromInput);
 
 	/** Recursion guard boolean to prevent PropagatePinType from infinitely recursing if you manage to create a loop of knots */
 	bool bRecursionGuard;
diff --git a/Engine/Source/Editor/BlueprintGraph/Classes/K2Node_MacroInstance.h b/Engine/Source/Editor/BlueprintGraph/Classes/K2Node_MacroInstance.h
--- a/Engine/Source/Editor/BlueprintGraph/Classes/K2Node_MacroInstance.h
+++ b/Engine/Source/Editor/BlueprintGraph/Classes/K2Node_MacroInstance.h
@@ -63,7 +63,6 @@ public:
 	virtual bool DrawNodeAsEntry() const override { return false; }
 	virtual void NotifyPinConnectionListChanged(UEdGraphPin* Pin) override;
 	virtual void PostReconstructNode() override;
-	virtual FText GetActiveBreakpointToolTipText() const override;
 	virtual bool HasExternalDependencies(TArray<class UStruct*>* OptionalOutput) const override;
 	virtual void GetNodeAttributes( TArray<TKeyValuePair<FString, FString>>& OutNodeAttributes ) const override;
 	virtual FText GetMenuCategory() const override;
diff --git a/Engine/Source/Editor/BlueprintGraph/Classes/K2Node_StructOperation.h b/Engine/Source/Editor/BlueprintGraph/Classes/K2Node_StructOperation.h
--- a/Engine/Source/Editor/BlueprintGraph/Classes/K2Node_StructOperation.h
+++ b/Engine/Source/Editor/BlueprintGraph/Classes/K2Node_StructOperation.h
@@ -46,8 +46,7 @@ protected:
 		virtual void CustomizePinData(UEdGraphPin* Pin, FName SourcePropertyName, int32 ArrayIndex, UProperty* Property) const override;
 		// End of FOptionalPinsUpdater interfac
 	};
-#if WITH_EDITOR
-	static bool DoRenamedPinsMatch(const UEdGraphPin* NewPin, const UEdGraphPin* OldPin, bool bStructInVaraiablesOut);
-#endif
+
+	bool DoRenamedPinsMatch(const UEdGraphPin* NewPin, const UEdGraphPin* OldPin, bool bStructInVariablesOut) const;
 };
 
diff --git a/Engine/Source/Editor/BlueprintGraph/Classes/K2Node_TunnelBoundary.h b/Engine/Source/Editor/BlueprintGraph/Classes/K2Node_TunnelBoundary.h
--- a/Engine/Source/Editor/BlueprintGraph/Classes/K2Node_TunnelBoundary.h
+++ b/Engine/Source/Editor/BlueprintGraph/Classes/K2Node_TunnelBoundary.h
@@ -4,21 +4,19 @@
 
 #include "CoreMinimal.h"
 #include "UObject/ObjectMacros.h"
-#include "Misc/Guid.h"
-#include "EdGraph/EdGraphPin.h"
 #include "K2Node.h"
 #include "K2Node_TunnelBoundary.generated.h"
 
-class FCompilerResultsLog;
-class UEdGraph;
+class FKismetCompilerContext;
 class UK2Node_Tunnel;
 
-enum ETunnelBoundaryType
+UENUM()
+enum class ETunnelBoundaryType : uint8
 {
-	TBT_Unknown = 0,
-	TBT_EntrySite,
-	TBT_ExitSite,
-	TBT_EndOfThread
+	Unknown = 0,
+	EntrySite,
+	InputSite,
+	OutputSite
 };
 
 UCLASS()
@@ -27,59 +25,27 @@ class BLUEPRINTGRAPH_API UK2Node_TunnelBoundary : public UK2Node
 	GENERATED_UCLASS_BODY()
 
 	/** Base Name */
-	UPROPERTY()
+	UPROPERTY(Transient)
 	FName BaseName;
 
-	/** Final exit site */
-	UPROPERTY()
-	FEdGraphPinReference FinalExitSite;
-
-	/** Wether or not this is an exit point */
+	/** Node Type */
+	UPROPERTY(Transient)
 	ETunnelBoundaryType TunnelBoundaryType;
 
 public:
 
-	//~ Begin of UK2Node Interface
+	//~ Begin UEdGraphNode Interface
 	virtual FText GetNodeTitle(ENodeTitleType::Type TitleType) const override;
-	virtual class FNodeHandlingFunctor* CreateNodeHandler(class FKismetCompilerContext& CompilerContext) const override;
-	//~ End UK2Node Interface
-
-	/** Creates Boundary nodes for a tunnel graph */
-	static void CreateBoundaryNodesForGraph(UEdGraph* TunnelGraph, class FCompilerResultsLog& MessageLog);
-
-	/** Creates Boundary nodes for a tunnel instance */
-	static void CreateBoundaryNodesForTunnelInstance(UK2Node_Tunnel* TunnelInstance, UEdGraph* TunnelGraph, class FCompilerResultsLog& MessageLog);
-
-	/** Creates Boundaries for expansion nodes */
-	static void CreateBoundariesForExpansionNodes(UEdGraphNode* SourceNode, TArray<UEdGraphNode*>& ExpansionNodes, TMap<UEdGraphPin*, UEdGraphPin*>& LinkedPinMap, FCompilerResultsLog& MessageLog);
-
-	/** Checks if the tunnel node is a pure tunnel rather than a tunnel instance */
-	static bool IsPureTunnel(UK2Node_Tunnel* Tunnel);
-
-protected:
-
-	/** Wires up the tunnel entry boundries */
-	void WireUpTunnelEntry(UK2Node_Tunnel* TunnelInstance, UEdGraphPin* TunnelPin, FCompilerResultsLog& MessageLog);
+	//~ End UEdGraphNode Interface
 
-	/** Wires up the tunnel exit boundries */
-	void WireUpTunnelExit(UK2Node_Tunnel* TunnelInstance, UEdGraphPin* TunnelPin, FCompilerResultsLog& MessageLog);
-
-	/** Wires up the entry boundries */
-	void WireUpEntry(UEdGraphNode* SourceNode, UEdGraphPin* SourcePin, const TArray<UEdGraphPin*>& EntryPins, FCompilerResultsLog& MessageLog);
-
-	/** Wires up the exit boundries */
-	void WireUpExit(UEdGraphNode* SourceNode, UEdGraphPin* SourcePin, const TArray<UEdGraphPin*>& ExitPins, FCompilerResultsLog& MessageLog);
-
-	/** Create base node name */
-	void CreateBaseNodeName(UEdGraphNode* SourceNode);
-
-	/** Build a guid map for nodes from the original graph */
-	static void BuildSourceNodeMap(UEdGraphNode* Tunnel, TMap<FGuid, const UEdGraphNode*>& SourceNodeMap);
-
-	/** Determines the true source tunnel instance */
-	static UEdGraphNode* FindTrueSourceTunnelInstance(UEdGraphNode* Tunnel, UEdGraphNode* SourceTunnelInstance);
+	//~ Begin K2Node Interface
+	virtual class FNodeHandlingFunctor* CreateNodeHandler(FKismetCompilerContext& CompilerContext) const override;
+	virtual int32 GetNodeRefreshPriority() const override { return EBaseNodeRefreshPriority::Low_UsesDependentWildcard; }
+	//~ End UK2Node Interface
 
-	/** Tries to map and locate tunnel exit or termination sites */
-	static void FindTunnelExitSiteInstances(UEdGraphPin* NodeEntryPin, TArray<UEdGraphPin*>& ExitPins, TArray<UEdGraphPin*>& VisitedPins, const UEdGraphNode* TunnelExit = nullptr);
+	/** Returns the type of tunnel boundary this node is */
+	ETunnelBoundaryType GetTunnelBoundaryType() const { return TunnelBoundaryType; }
 
-};
+	/** Set node attributes based on the tunnel source node */
+	void SetNodeAttributes(const UK2Node_Tunnel* SourceNode);
+};
\ No newline at end of file
diff --git a/Engine/Source/Editor/BlueprintGraph/Classes/K2Node_Variable.h b/Engine/Source/Editor/BlueprintGraph/Classes/K2Node_Variable.h
--- a/Engine/Source/Editor/BlueprintGraph/Classes/K2Node_Variable.h
+++ b/Engine/Source/Editor/BlueprintGraph/Classes/K2Node_Variable.h
@@ -141,6 +141,9 @@ public:
 	UProperty* GetPropertyForVariable() const;
 	UProperty* GetPropertyForVariableFromSkeleton() const;
 
+	/** Returns true if the variable names match, this looks for redirectors */
+	static bool DoesRenamedVariableMatch(FName OldVariableName, FName NewVariableName, UStruct* StructType);
+
 private:
 	UProperty* GetPropertyForVariable_Internal(UClass* OwningClass) const;
 
diff --git a/Engine/Source/Editor/BlueprintGraph/Private/K2Node.cpp b/Engine/Source/Editor/BlueprintGraph/Private/K2Node.cpp
--- a/Engine/Source/Editor/BlueprintGraph/Private/K2Node.cpp
+++ b/Engine/Source/Editor/BlueprintGraph/Private/K2Node.cpp
@@ -186,12 +186,7 @@ FText UK2Node::GetToolTipHeading() const
 	{
 		if (ExistingBreakpoint->IsEnabled())
 		{
-			Heading = LOCTEXT("EnabledBreakpoint", "Active Breakpoint");
-			FText ActiveBreakpointToolTipText = GetActiveBreakpointToolTipText();
-			if (!ActiveBreakpointToolTipText.IsEmpty())
-			{
-				Heading = FText::Format(FText::FromString("{0} - {1}"), Heading, ActiveBreakpointToolTipText);
-			}
+			Heading = LOCTEXT("EnabledBreakpoint", "Active Breakpoint - Execution will break at this location.");
 		}
 		else 
 		{
@@ -202,11 +197,6 @@ FText UK2Node::GetToolTipHeading() const
 	return Heading;
 }
 
-FText UK2Node::GetActiveBreakpointToolTipText() const
-{
-	return LOCTEXT("ActiveBreakpointToolTip", "Execution will break at this location.");
-}
-
 bool UK2Node::CreatePinsForFunctionEntryExit(const UFunction* Function, bool bForFunctionEntry)
 {
 	const UEdGraphSchema_K2* K2Schema = GetDefault<UEdGraphSchema_K2>();
@@ -842,7 +832,7 @@ void UK2Node::ValidateNodeDuringCompilation(FCompilerResultsLog& MessageLog) con
 	Super::ValidateNodeDuringCompilation(MessageLog);
 
 	// Since this might be an expanded node, the validation will have been done on the source node (which will return this if not expanded)
-	if (const UK2Node* SourceNode = Cast<UK2Node>(MessageLog.GetSourceNode(this)))
+	if (const UK2Node* SourceNode = Cast<UK2Node>(MessageLog.FindSourceObject(this)))
 	{
 		// If we don't commit any messages for the source node, then see if this node has messages to commit.
 		// This is primarily needed for nodes generated from inside macros.
diff --git a/Engine/Source/Editor/BlueprintGraph/Private/K2Node_BaseAsyncTask.cpp b/Engine/Source/Editor/BlueprintGraph/Private/K2Node_BaseAsyncTask.cpp
--- a/Engine/Source/Editor/BlueprintGraph/Private/K2Node_BaseAsyncTask.cpp
+++ b/Engine/Source/Editor/BlueprintGraph/Private/K2Node_BaseAsyncTask.cpp
@@ -475,6 +475,39 @@ UFunction* UK2Node_BaseAsyncTask::GetFactoryFunction() const
 	return FactoryFunction;
 }
 
+void UK2Node_BaseAsyncTask::GetRedirectPinNames(const UEdGraphPin& Pin, TArray<FString>& RedirectPinNames) const
+{
+	TMap<FString, FStringFormatArg> Args;
+
+	if (ProxyClass)
+	{
+		Args.Add(TEXT("ProxyClass"), ProxyClass->GetName());
+		Args.Add(TEXT("ProxyClassSeparator"), TEXT("."));
+	}
+	else
+	{
+		Args.Add(TEXT("ProxyClass"), TEXT(""));
+		Args.Add(TEXT("ProxyClassSeparator"), TEXT(""));
+	}
+	if (ProxyFactoryFunctionName != NAME_None)
+	{
+		Args.Add(TEXT("ProxyFactoryFunction"), ProxyFactoryFunctionName.ToString());
+		Args.Add(TEXT("ProxyFactoryFunctionSeparator"), TEXT("."));
+	}
+	else
+	{
+		Args.Add(TEXT("ProxyFactoryFunction"), TEXT(""));
+		Args.Add(TEXT("ProxyFactoryFunctionSeparator"), TEXT(""));
+	}
+
+	Args.Add(TEXT("PinName"), Pin.PinName.ToString());
+
+	FString FullPinName;
+	FullPinName = FString::Format(TEXT("{ProxyClass}{ProxyClassSeparator}{ProxyFactoryFunction}{ProxyFactoryFunctionSeparator}{PinName}"), Args);
+
+	RedirectPinNames.Add(FullPinName);
+}
+
 void UK2Node_BaseAsyncTask::ValidateNodeDuringCompilation(class FCompilerResultsLog& MessageLog) const
 {
 	Super::ValidateNodeDuringCompilation(MessageLog);
diff --git a/Engine/Source/Editor/BlueprintGraph/Private/K2Node_BreakStruct.cpp b/Engine/Source/Editor/BlueprintGraph/Private/K2Node_BreakStruct.cpp
--- a/Engine/Source/Editor/BlueprintGraph/Private/K2Node_BreakStruct.cpp
+++ b/Engine/Source/Editor/BlueprintGraph/Private/K2Node_BreakStruct.cpp
@@ -315,33 +315,13 @@ FLinearColor UK2Node_BreakStruct::GetNodeTitleColor() const
 	return UK2Node::GetNodeTitleColor();
 }
 
-UK2Node::ERedirectType UK2Node_BreakStruct::DoPinsMatchForReconstruction(const UEdGraphPin* NewPin, int32 NewPinIndex, const UEdGraphPin* OldPin, int32 OldPinIndex)  const
+UK2Node::ERedirectType UK2Node_BreakStruct::DoPinsMatchForReconstruction(const UEdGraphPin* NewPin, int32 NewPinIndex, const UEdGraphPin* OldPin, int32 OldPinIndex) const
 {
 	ERedirectType Result = UK2Node::DoPinsMatchForReconstruction(NewPin, NewPinIndex, OldPin, OldPinIndex);
 	if ((ERedirectType_None == Result) && DoRenamedPinsMatch(NewPin, OldPin, true))
 	{
 		Result = ERedirectType_Name;
 	}
-	else if ((ERedirectType_None == Result) && NewPin && OldPin)
-	{
-		if ((EGPD_Input == NewPin->Direction) && (EGPD_Input == OldPin->Direction))
-		{
-			const UEdGraphSchema_K2* K2Schema = GetDefault<UEdGraphSchema_K2>();
-			if (K2Schema->ArePinTypesCompatible( NewPin->PinType, OldPin->PinType))
-			{
-				Result = ERedirectType_Name;
-			}
-		}
-		else if ((EGPD_Output == NewPin->Direction) && (EGPD_Output == OldPin->Direction))
-		{
-			FName RedirectedPinName = UProperty::FindRedirectedPropertyName(StructType, OldPin->PinName);
-
-			if (RedirectedPinName != NAME_Name)
-			{
-				Result = (RedirectedPinName != NewPin->PinName ? ERedirectType_None : ERedirectType_Name);
-			}
-		}
-	}
 	return Result;
 }
 
diff --git a/Engine/Source/Editor/BlueprintGraph/Private/K2Node_ConstructObjectFromClass.cpp b/Engine/Source/Editor/BlueprintGraph/Private/K2Node_ConstructObjectFromClass.cpp
--- a/Engine/Source/Editor/BlueprintGraph/Private/K2Node_ConstructObjectFromClass.cpp
+++ b/Engine/Source/Editor/BlueprintGraph/Private/K2Node_ConstructObjectFromClass.cpp
@@ -215,7 +215,7 @@ void UK2Node_ConstructObjectFromClass::OnClassPinChanged()
 	// Cache all the pin connections to the ResultPin, we will attempt to recreate them
 	TArray<UEdGraphPin*> ResultPinConnectionList = ResultPin->LinkedTo;
 	// Because the archetype has changed, we break the output link as the output pin type will change
-	ResultPin->BreakAllPinLinks();
+	ResultPin->BreakAllPinLinks(true);
 
 	const UEdGraphSchema_K2* K2Schema = GetDefault<UEdGraphSchema_K2>();
 
diff --git a/Engine/Source/Editor/BlueprintGraph/Private/K2Node_Knot.cpp b/Engine/Source/Editor/BlueprintGraph/Private/K2Node_Knot.cpp
--- a/Engine/Source/Editor/BlueprintGraph/Private/K2Node_Knot.cpp
+++ b/Engine/Source/Editor/BlueprintGraph/Private/K2Node_Knot.cpp
@@ -92,7 +92,7 @@ void UK2Node_Knot::PropagatePinType()
 	{
 		if (Inputs->PinType.PinCategory != UEdGraphSchema_K2::PC_Wildcard)
 		{
-			PropagatePinTypeFromInput();
+			PropagatePinTypeFromDirection(true);
 			return;
 		}
 	}
@@ -101,7 +101,7 @@ void UK2Node_Knot::PropagatePinType()
 	{
 		if (Outputs->PinType.PinCategory != UEdGraphSchema_K2::PC_Wildcard)
 		{
-			PropagatePinTypeFromOutput();
+			PropagatePinTypeFromDirection(false);
 			return;
 		}
 	}
@@ -110,12 +110,12 @@ void UK2Node_Knot::PropagatePinType()
 	if (MyInputPin->LinkedTo.Num() > 0)
 	{
 		// If we can't mirror from output type, we should at least get the type information from the input connection chain
-		PropagatePinTypeFromInput();
+		PropagatePinTypeFromDirection(true);
 	}
 	else if (MyOutputPin->LinkedTo.Num() > 0)
 	{
 		// Try to mirror from output first to make sure we get appropriate member references
-		PropagatePinTypeFromOutput();
+		PropagatePinTypeFromDirection(false);
 	}
 	else
 	{
@@ -130,83 +130,51 @@ void UK2Node_Knot::PropagatePinType()
 	}
 }
 
-void UK2Node_Knot::PropagatePinTypeFromInput()
+void UK2Node_Knot::PropagatePinTypeFromDirection(bool bFromInput)
 {
 	if (bRecursionGuard)
 	{
 		return;
 	}
-	// Set the type of the pin based on input connections.
-	// We have to move up the chain of linked reroute nodes until we reach a node
-	// with type information before percolating that information down.
-	UEdGraphPin* MyInputPin = GetInputPin();
-	UEdGraphPin* MyOutputPin = GetOutputPin();
+	// Set the type of the pin based on the source connection, and then percolate
+	// that type information up until we no longer reach another Reroute node
+	UEdGraphPin* MySourcePin = bFromInput ? GetInputPin() : GetOutputPin();
+	UEdGraphPin* MyDestinationPin = bFromInput ? GetOutputPin() : GetInputPin();
 
 	TGuardValue<bool> RecursionGuard(bRecursionGuard, true);
 
-	for (UEdGraphPin* InPin : MyInputPin->LinkedTo)
+	// Make sure any source knot pins compute their type, this will try to call back
+	// into this function but the recursion guard will stop it
+	for (UEdGraphPin* InPin : MySourcePin->LinkedTo)
 	{
 		if (UK2Node_Knot* KnotNode = Cast<UK2Node_Knot>(InPin->GetOwningNode()))
 		{
-			KnotNode->PropagatePinTypeFromInput();
+			KnotNode->PropagatePinTypeFromDirection(bFromInput);
 		}
 	}
 
-	UEdGraphPin* TypeSource = MyInputPin->LinkedTo.Num() ? MyInputPin->LinkedTo[0] : nullptr;
+	UEdGraphPin* TypeSource = MySourcePin->LinkedTo.Num() ? MySourcePin->LinkedTo[0] : nullptr;
 	if (TypeSource)
 	{
-		MyInputPin->PinType = TypeSource->PinType;
-		MyOutputPin->PinType = TypeSource->PinType;
+		MySourcePin->PinType = TypeSource->PinType;
+		MyDestinationPin->PinType = TypeSource->PinType;
 
-		for (UEdGraphPin* InPin : MyInputPin->LinkedTo)
+		for (UEdGraphPin* LinkPin : MyDestinationPin->LinkedTo)
 		{
-			if (UK2Node* OwningNode = Cast<UK2Node>(InPin->GetOwningNode()))
-			{
-				if (!OwningNode->IsA<UK2Node_Knot>())
-				{
-					OwningNode->PinConnectionListChanged(InPin);
-				}
-			}
-		}
-	}
-	else
-	{
-		// TODO?
-	}
-}
-
-void UK2Node_Knot::PropagatePinTypeFromOutput()
-{
-	// Set the type of the pin based on the output connection, and then percolate
-	// that type information up until we no longer reach another Reroute node
-	UEdGraphPin* MyInputPin = GetInputPin();
-	UEdGraphPin* MyOutputPin = GetOutputPin();
-
-	UEdGraphPin* TypeSource = MyOutputPin->LinkedTo.Num() ? MyOutputPin->LinkedTo[0] : nullptr;
-	if (TypeSource)
-	{
-		MyInputPin->PinType = TypeSource->PinType;
-		MyOutputPin->PinType = TypeSource->PinType;
-
-		for (UEdGraphPin* InPin : MyInputPin->LinkedTo)
-		{
-			if (UK2Node* OwningNode = Cast<UK2Node>(InPin->GetOwningNode()))
+			if (UK2Node* OwningNode = Cast<UK2Node>(LinkPin->GetOwningNode()))
 			{
+				// Notify any pins in the destination direction
 				if (UK2Node_Knot* KnotNode = Cast<UK2Node_Knot>(OwningNode))
 				{
-					KnotNode->PropagatePinTypeFromOutput();
+					KnotNode->PropagatePinTypeFromDirection(bFromInput);
 				}
 				else
 				{
-					OwningNode->PinConnectionListChanged(InPin);
+					OwningNode->PinConnectionListChanged(LinkPin);
 				}
 			}
 		}
 	}
-	else 
-	{
-		// TODO?
-	}
 }
 
 void UK2Node_Knot::GetMenuActions(FBlueprintActionDatabaseRegistrar& ActionRegistrar) const
diff --git a/Engine/Source/Editor/BlueprintGraph/Private/K2Node_MacroInstance.cpp b/Engine/Source/Editor/BlueprintGraph/Private/K2Node_MacroInstance.cpp
--- a/Engine/Source/Editor/BlueprintGraph/Private/K2Node_MacroInstance.cpp
+++ b/Engine/Source/Editor/BlueprintGraph/Private/K2Node_MacroInstance.cpp
@@ -50,26 +50,33 @@ void UK2Node_MacroInstance::PostPasteNode()
 	const UBlueprint* InstanceOwner = GetBlueprint();
 
 	// Find the owner of the macro graph
-	const UEdGraph* MacroGraph = MacroGraphReference.GetGraph();
-	UObject* MacroOwner = MacroGraph->GetOuter();
-	UBlueprint* MacroOwnerBP = NULL;
-	while(MacroOwner)
+	if (const UEdGraph* MacroGraph = MacroGraphReference.GetGraph())
 	{
-		MacroOwnerBP = Cast<UBlueprint>(MacroOwner);
-		if(MacroOwnerBP)
+		UObject* MacroOwner = MacroGraph->GetOuter();
+		UBlueprint* MacroOwnerBP = nullptr;
+		while (MacroOwner)
 		{
-			break;
+			MacroOwnerBP = Cast<UBlueprint>(MacroOwner);
+			if (MacroOwnerBP)
+			{
+				break;
+			}
+
+			MacroOwner = MacroOwner->GetOuter();
 		}
 
-		MacroOwner = MacroOwner->GetOuter();
+		if ((MacroOwnerBP != nullptr)
+			&& (MacroOwnerBP->BlueprintType != BPTYPE_MacroLibrary)
+			&& (MacroOwnerBP != InstanceOwner))
+		{
+			// If this is a graph from another blueprint that is NOT a library, disallow the connection!
+			MacroGraphReference.SetGraph(nullptr);
+		}
 	}
-	
-	if((MacroOwnerBP != NULL)
-		&& (MacroOwnerBP->BlueprintType != BPTYPE_MacroLibrary)
-		&& (MacroOwnerBP != InstanceOwner))
+	else
 	{
-		// If this is a graph from another blueprint that is NOT a library, disallow the connection!
-		MacroGraphReference.SetGraph(NULL);
+		// Can't find the referenced macro, fully clear this reference
+		MacroGraphReference.SetGraph(nullptr);
 	}
 
 	Super::PostPasteNode();
@@ -83,7 +90,7 @@ void UK2Node_MacroInstance::AllocateDefaultPins()
 	PreloadObject(MacroGraphReference.GetBlueprint());
 	UEdGraph* MacroGraph = MacroGraphReference.GetGraph();
 	
-	if (MacroGraph != NULL)
+	if (MacroGraph != nullptr)
 	{
 		PreloadObject(MacroGraph);
 
@@ -106,7 +113,7 @@ void UK2Node_MacroInstance::AllocateDefaultPins()
 						UEdGraphPin* PortPin = *PinIt;
 
 						// We're not interested in any pins that have been expanded internally on the macro
-						if (PortPin->ParentPin == NULL)
+						if (PortPin->ParentPin == nullptr)
 						{
 							UEdGraphPin* NewLocalPin = CreatePin(UEdGraphPin::GetComplementaryDirection(PortPin->Direction), PortPin->PinType, PortPin->PinName);
 							Schema->SetPinAutogeneratedDefaultValue(NewLocalPin, PortPin->GetDefaultAsString());
@@ -186,7 +193,7 @@ FLinearColor UK2Node_MacroInstance::GetNodeTitleColor() const
 
 void UK2Node_MacroInstance::GetContextMenuActions(const FGraphNodeContextMenuBuilder& Context) const
 {
-	if ( Context.Pin == NULL )
+	if ( Context.Pin == nullptr )
 	{
 		Context.MenuBuilder->BeginSection("K2NodeMacroInstance", NSLOCTEXT("K2Node", "MacroInstanceHeader", "Macro Instance"));
 		{
@@ -222,7 +229,7 @@ FKismetUserDeclaredFunctionMetadata* UK2Node_MacroInstance::GetAssociatedGraphMe
 		}
 	}
 
-	return NULL;
+	return nullptr;
 }
 
 void UK2Node_MacroInstance::FindInContentBrowser(TWeakObjectPtr<UK2Node_MacroInstance> MacroInstance)
@@ -316,7 +323,7 @@ FString UK2Node_MacroInstance::GetDocumentationLink() const
 FString UK2Node_MacroInstance::GetDocumentationExcerptName() const
 {
 	UEdGraph* MacroGraph = MacroGraphReference.GetGraph();
-	if (MacroGraph != NULL)
+	if (MacroGraph != nullptr)
 	{
 		return MacroGraph->GetName();
 	}
@@ -351,7 +358,7 @@ FSlateIcon UK2Node_MacroInstance::GetIconAndTint(FLinearColor& OutColor) const
 	// Special case handling for standard macros
 	// @TODO Change this to a SlateBurushAsset pointer on the graph or something similar, to allow any macro to have an icon
 	UEdGraph* MacroGraph = MacroGraphReference.GetGraph();
-	if(MacroGraph != NULL && MacroGraph->GetOuter()->GetName() == TEXT("StandardMacros"))
+	if(MacroGraph != nullptr && MacroGraph->GetOuter()->GetName() == TEXT("StandardMacros"))
 	{
 		FName MacroName = FName(*MacroGraph->GetName());
 		if(	MacroName == TEXT("ForLoop" ) ||
@@ -438,11 +445,6 @@ void UK2Node_MacroInstance::PostFixupAllWildcardPins(bool bInAllWildcardPinsUnli
 	}
 }
 
-FText UK2Node_MacroInstance::GetActiveBreakpointToolTipText() const
-{
-	return LOCTEXT("ActiveBreakpointToolTip", "Execution will break inside the macro.");
-}
-
 bool UK2Node_MacroInstance::HasExternalDependencies(TArray<class UStruct*>* OptionalOutput) const
 {
 	UBlueprint* OtherBlueprint = MacroGraphReference.GetBlueprint();
diff --git a/Engine/Source/Editor/BlueprintGraph/Private/K2Node_MakeStruct.cpp b/Engine/Source/Editor/BlueprintGraph/Private/K2Node_MakeStruct.cpp
--- a/Engine/Source/Editor/BlueprintGraph/Private/K2Node_MakeStruct.cpp
+++ b/Engine/Source/Editor/BlueprintGraph/Private/K2Node_MakeStruct.cpp
@@ -266,33 +266,13 @@ FNodeHandlingFunctor* UK2Node_MakeStruct::CreateNodeHandler(FKismetCompilerConte
 	return new FKCHandler_MakeStruct(CompilerContext);
 }
 
-UK2Node::ERedirectType UK2Node_MakeStruct::DoPinsMatchForReconstruction(const UEdGraphPin* NewPin, int32 NewPinIndex, const UEdGraphPin* OldPin, int32 OldPinIndex)  const
+UK2Node::ERedirectType UK2Node_MakeStruct::DoPinsMatchForReconstruction(const UEdGraphPin* NewPin, int32 NewPinIndex, const UEdGraphPin* OldPin, int32 OldPinIndex) const
 {
 	ERedirectType Result = UK2Node::DoPinsMatchForReconstruction(NewPin, NewPinIndex, OldPin, OldPinIndex);
 	if ((ERedirectType_None == Result) && DoRenamedPinsMatch(NewPin, OldPin, false))
 	{
 		Result = ERedirectType_Name;
 	}
-	else if ((ERedirectType_None == Result) && NewPin && OldPin)
-	{
-		if ((EGPD_Output == NewPin->Direction) && (EGPD_Output == OldPin->Direction))
-		{
-			const UEdGraphSchema_K2* K2Schema = GetDefault<UEdGraphSchema_K2>();
-			if (K2Schema->ArePinTypesCompatible(NewPin->PinType, OldPin->PinType))
-			{
-				Result = ERedirectType_Name;
-			}
-		}
-		else if ((EGPD_Input == NewPin->Direction) && (EGPD_Input == OldPin->Direction))
-		{
-			FName RedirectedPinName = UProperty::FindRedirectedPropertyName(StructType, OldPin->PinName);
-
-			if (RedirectedPinName != NAME_Name)
-			{
-				Result = ((RedirectedPinName != NewPin->PinName) ? ERedirectType_None : ERedirectType_Name);
-			}
-		}
-	}
 	return Result;
 }
 
diff --git a/Engine/Source/Editor/BlueprintGraph/Private/K2Node_MathExpression.cpp b/Engine/Source/Editor/BlueprintGraph/Private/K2Node_MathExpression.cpp
--- a/Engine/Source/Editor/BlueprintGraph/Private/K2Node_MathExpression.cpp
+++ b/Engine/Source/Editor/BlueprintGraph/Private/K2Node_MathExpression.cpp
@@ -2702,17 +2702,7 @@ void UK2Node_MathExpression::ValidateNodeDuringCompilation(FCompilerResultsLog&
 	{
 		if(UObject const* SourceObject = MessageLog.FindSourceObject(this))
 		{
-			UK2Node_MathExpression const* MathExpression = nullptr;
-
-			// If the source object is a MacroInstance, we need to look elsewhere for the original MathExpression
-			if(Cast<UK2Node_MacroInstance>(SourceObject))
-			{
-				MathExpression = CastChecked<UK2Node_MathExpression>(MessageLog.GetSourceTunnelNode(this));
-			}
-			else
-			{
-				MathExpression = MessageLog.FindSourceObjectTypeChecked<UK2Node_MathExpression>(this);
-			}
+			UK2Node_MathExpression const* MathExpression = MessageLog.FindSourceObjectTypeChecked<UK2Node_MathExpression>(this);
 
 			// Should always be able to find the source math expression
 			check(MathExpression);
diff --git a/Engine/Source/Editor/BlueprintGraph/Private/K2Node_StructOperation.cpp b/Engine/Source/Editor/BlueprintGraph/Private/K2Node_StructOperation.cpp
--- a/Engine/Source/Editor/BlueprintGraph/Private/K2Node_StructOperation.cpp
+++ b/Engine/Source/Editor/BlueprintGraph/Private/K2Node_StructOperation.cpp
@@ -52,31 +52,35 @@ void UK2Node_StructOperation::FStructOperationOptionalPinManager::CustomizePinDa
 	}
 }
 
-bool UK2Node_StructOperation::DoRenamedPinsMatch(const UEdGraphPin* NewPin, const UEdGraphPin* OldPin, bool bStructInVaraiablesOut)
+bool UK2Node_StructOperation::DoRenamedPinsMatch(const UEdGraphPin* NewPin, const UEdGraphPin* OldPin, bool bStructInVariablesOut) const
 {
-	bool bResult = false;
 	if (NewPin && OldPin && (OldPin->Direction == NewPin->Direction))
 	{
-		const EEdGraphPinDirection StructDirection = bStructInVaraiablesOut ? EGPD_Input : EGPD_Output;
-		const EEdGraphPinDirection VariablesDirection = bStructInVaraiablesOut ? EGPD_Output : EGPD_Input;
+		const EEdGraphPinDirection StructDirection = bStructInVariablesOut ? EGPD_Input : EGPD_Output;
+		const EEdGraphPinDirection VariablesDirection = bStructInVariablesOut ? EGPD_Output : EGPD_Input;
 		const UEdGraphSchema_K2* K2Schema = GetDefault<UEdGraphSchema_K2>(); 
 		const bool bCompatible = K2Schema && K2Schema->ArePinTypesCompatible(NewPin->PinType, OldPin->PinType);
 
 		if (bCompatible && (StructDirection == OldPin->Direction))
 		{
-			// Struct name was changed
-			bResult = true;
+			// Struct name was changed, which is fine
+			return true;
 		}
 		else if (bCompatible && (VariablesDirection == OldPin->Direction))
 		{
-			// name of a member variable was changed
+			// Name of a member variable was changed, check guids and redirects
 			if ((NewPin->PersistentGuid == OldPin->PersistentGuid) && OldPin->PersistentGuid.IsValid())
 			{
-				bResult = true;
+				return true;
+			}
+
+			if (DoesRenamedVariableMatch(OldPin->PinName, NewPin->PinName, StructType))
+			{
+				return true;
 			}
 		}
 	}
-	return bResult;
+	return false;
 }
 
 FString UK2Node_StructOperation::GetPinMetaData(FName InPinName, FName InKey)
diff --git a/Engine/Source/Editor/BlueprintGraph/Private/K2Node_TunnelBoundary.cpp b/Engine/Source/Editor/BlueprintGraph/Private/K2Node_TunnelBoundary.cpp
--- a/Engine/Source/Editor/BlueprintGraph/Private/K2Node_TunnelBoundary.cpp
+++ b/Engine/Source/Editor/BlueprintGraph/Private/K2Node_TunnelBoundary.cpp
@@ -2,17 +2,11 @@
 
 #include "K2Node_TunnelBoundary.h"
 #include "EdGraphSchema_K2.h"
-#include "K2Node_Tunnel.h"
+#include "KismetCompilerMisc.h"
 #include "K2Node_Composite.h"
-#include "K2Node_ExecutionSequence.h"
 #include "K2Node_MacroInstance.h"
-#include "KismetCompiledFunctionContext.h"
-#include "KismetCompilerMisc.h"
-
-#define LOCTEXT_NAMESPACE "FKCHandler_TunnelBoundary"
 
-//////////////////////////////////////////////////////////////////////////
-// FKCHandler_TunnelBoundary
+#define LOCTEXT_NAMESPACE "UK2Node_TunnelBoundary"
 
 class FKCHandler_TunnelBoundary : public FNodeHandlingFunctor
 {
@@ -24,473 +18,71 @@ public:
 
 	virtual void Compile(FKismetFunctionContext& Context, UEdGraphNode* Node) override
 	{
-		if (UK2Node_TunnelBoundary* BoundaryNode = Cast<UK2Node_TunnelBoundary>(Node))
-		{
-			switch(BoundaryNode->TunnelBoundaryType)
-			{
-				case TBT_EntrySite:
-				case TBT_ExitSite:
-				{
-					GenerateSimpleThenGoto(Context, *Node);
-					break;
-				}
-				case TBT_EndOfThread:
-				{
-					FBlueprintCompiledStatement& ExitTunnelStatement = Context.AppendStatementForNode(Node);
-					ExitTunnelStatement.Type = KCST_InstrumentedTunnelEndOfThread;
-					FBlueprintCompiledStatement& EOTStatement = Context.AppendStatementForNode(Node);
-					EOTStatement.Type = KCST_EndOfThread;
-					break;
-				}
-			}
-		}
+		// Tunnel boundaries do nothing except continue the execution sequence.
+		GenerateSimpleThenGoto(Context, *Node);
 	}
 
 };
 
-//////////////////////////////////////////////////////////////////////////
-// UK2Node_TunnelBoundary
-
 UK2Node_TunnelBoundary::UK2Node_TunnelBoundary(const FObjectInitializer& ObjectInitializer)
 	: Super(ObjectInitializer)
-	, TunnelBoundaryType(TBT_Unknown)
 {
-	return;
+
 }
 
 FText UK2Node_TunnelBoundary::GetNodeTitle(ENodeTitleType::Type TitleType) const
 {
 	FFormatNamedArguments Args;
 	Args.Add(TEXT("BaseName"), FText::FromName(BaseName));
-	switch(TunnelBoundaryType)
-	{
-		case TBT_EntrySite:
-		{
-			Args.Add(TEXT("EntryType"), FText::FromString(TEXT("Tunnel Entry")));
-			break;
-		}
-		case TBT_ExitSite:
-		{
-			Args.Add(TEXT("EntryType"), FText::FromString(TEXT("Tunnel Exit")));
-			break;
-		}
-		case TBT_EndOfThread:
-		{
-			Args.Add(TEXT("EntryType"), FText::FromString(TEXT("Tunnel End Of Thread")));
-			break;
-		}
-	}
-	return FText::Format(LOCTEXT("UK2Node_TunnelBoundary_FullTitle", "{BaseName} {EntryType}"), Args);
-}
-
-FNodeHandlingFunctor* UK2Node_TunnelBoundary::CreateNodeHandler(FKismetCompilerContext& CompilerContext) const
-{
-	return new FKCHandler_TunnelBoundary(CompilerContext);
-}
-
-void UK2Node_TunnelBoundary::CreateBoundaryNodesForGraph(UEdGraph* TunnelGraph, FCompilerResultsLog& MessageLog)
-{
-	if (UK2Node_Tunnel* OuterTunnel = TunnelGraph->GetTypedOuter<UK2Node_Tunnel>())
-	{
-		UK2Node_Tunnel* SourceTunnel = Cast<UK2Node_Tunnel>(MessageLog.FindSourceObject(OuterTunnel));
-		CreateBoundaryNodesForTunnelInstance(SourceTunnel, TunnelGraph, MessageLog);
-	}
-}
-
-void UK2Node_TunnelBoundary::CreateBoundaryNodesForTunnelInstance(UK2Node_Tunnel* TunnelInstance, UEdGraph* TunnelGraph, FCompilerResultsLog& MessageLog)
-{
-	if (!TunnelGraph)
+	switch (TunnelBoundaryType)
 	{
-		return;
-	}
+	case ETunnelBoundaryType::EntrySite:
+		Args.Add(TEXT("TunnelBoundaryType"), FText::FromString(TEXT("Tunnel Entry Site")));
+		break;
 
-	TArray<UEdGraphPin*> TunnelEntryPins;
-	TArray<UEdGraphPin*> TunnelExitPins;
-	UEdGraphNode* TunnelExitNode = nullptr;
-	TArray<UK2Node_Tunnel*> Tunnels;
-	TunnelGraph->GetNodesOfClass<UK2Node_Tunnel>(Tunnels);
-	for (UK2Node_Tunnel* TunnelNode : Tunnels)
-	{
-		if (IsPureTunnel(TunnelNode))
-		{
-			for (UEdGraphPin* Pin : TunnelNode->Pins)
-			{
-				if (Pin->LinkedTo.Num() && Pin->PinType.PinCategory == UEdGraphSchema_K2::PC_Exec)
-				{
-					if (Pin->Direction == EGPD_Output)
-					{
-						TunnelEntryPins.Add(Pin);
-					}
-					else
-					{
-						TunnelExitNode = TunnelNode;
-						TunnelExitPins.Add(Pin);
-					}
-				}
-			}
-		}
-	}
-	// Find the blueprint tunnel instance.
-	UEdGraphNode* SourceTunnelInstance = Cast<UEdGraphNode>(MessageLog.FindSourceObject(TunnelInstance));
-	SourceTunnelInstance = FindTrueSourceTunnelInstance(TunnelInstance, SourceTunnelInstance);
-	UEdGraphNode* RegistrantTunnelInstance = TunnelInstance->IsA<UK2Node_Composite>() ? SourceTunnelInstance : TunnelInstance;
-	// Create the Boundary nodes for each unique entry/exit site.
-	TArray<UEdGraphPin*> ExecutionEndpointPins;
-	TArray<UEdGraphPin*> VisitedPins;
-	for (UEdGraphPin* EntryPin : TunnelEntryPins)
-	{
-		UK2Node_TunnelBoundary* TunnelBoundaryNode = TunnelGraph->CreateIntermediateNode<UK2Node_TunnelBoundary>();
-		MessageLog.NotifyIntermediateObjectCreation(TunnelBoundaryNode, TunnelInstance);
-		TunnelBoundaryNode->CreateNewGuid();
-		MessageLog.RegisterIntermediateTunnelNode(TunnelBoundaryNode, RegistrantTunnelInstance);
-		TunnelBoundaryNode->WireUpTunnelEntry(TunnelInstance, EntryPin, MessageLog);
-		for (UEdGraphPin* LinkedPin : EntryPin->LinkedTo)
-		{
-			FindTunnelExitSiteInstances(LinkedPin, ExecutionEndpointPins, VisitedPins, TunnelExitNode);
-		}
-	}
-	if (ExecutionEndpointPins.Num())
-	{
-		// Create boundary node.
-		UK2Node_TunnelBoundary* TunnelBoundaryNode = TunnelGraph->CreateIntermediateNode<UK2Node_TunnelBoundary>();
-		MessageLog.NotifyIntermediateObjectCreation(TunnelBoundaryNode, TunnelInstance);
-		TunnelBoundaryNode->CreateNewGuid();
-		TunnelBoundaryNode->TunnelBoundaryType = TBT_EndOfThread;
-		MessageLog.RegisterIntermediateTunnelNode(TunnelBoundaryNode, RegistrantTunnelInstance);
-		// Create exit point pin
-		UEdGraphPin* BoundaryTerminalPin = TunnelBoundaryNode->CreatePin(EGPD_Input, UEdGraphSchema_K2::PC_Exec, TEXT("TunnelEndOfThread"));
-		// Wire up the endpoint unwired links to the boundary node.
-		for (UEdGraphPin* TerminationPin : ExecutionEndpointPins)
-		{
-			TerminationPin->MakeLinkTo(BoundaryTerminalPin);
-		}
-	}
-	for (UEdGraphPin* ExitPin : TunnelExitPins)
-	{
-		UK2Node_TunnelBoundary* TunnelBoundaryNode = TunnelGraph->CreateIntermediateNode<UK2Node_TunnelBoundary>();
-		MessageLog.NotifyIntermediateObjectCreation(TunnelBoundaryNode, TunnelInstance);
-		TunnelBoundaryNode->CreateNewGuid();
-		MessageLog.RegisterIntermediateTunnelNode(TunnelBoundaryNode, RegistrantTunnelInstance);
-		TunnelBoundaryNode->WireUpTunnelExit(TunnelInstance, ExitPin, MessageLog);
-	}
-	// Build node guid map to locate true source node.
-	TMap<FGuid, const UEdGraphNode*> TrueSourceNodeMap;
-	BuildSourceNodeMap(SourceTunnelInstance, TrueSourceNodeMap);
-}
+	case ETunnelBoundaryType::InputSite:
+		Args.Add(TEXT("TunnelBoundaryType"), FText::FromString(TEXT("Tunnel Input Site")));
+		break;
 
-void UK2Node_TunnelBoundary::CreateBoundariesForExpansionNodes(UEdGraphNode* SourceNode, TArray<UEdGraphNode*>& ExpansionNodes, TMap<UEdGraphPin*, UEdGraphPin*>& LinkedPinMap, FCompilerResultsLog& MessageLog)
-{
-	// Find boundary points and map them to the source pins
-	UEdGraph* TargetGraph = nullptr;
-	TMap<UEdGraphPin*, TArray<UEdGraphPin*>> EntryPins;
-	TMap<UEdGraphPin*, TArray<UEdGraphPin*>> ExitPins;
-	for (UEdGraphNode* ExpansionNode : ExpansionNodes)
-	{
-		if (!TargetGraph)
-		{
-			TargetGraph = ExpansionNode->GetGraph();
-		}
-		for (UEdGraphPin* Pin : ExpansionNode->Pins)
-		{
-			for (UEdGraphPin* LinkedPin : Pin->LinkedTo)
-			{
-				//if (ExpansionNode->IsA<UK2Node_Event>() && Pin->Direction == EGPD_Output && Pin->PinType.PinCategory == UEdGraphSchema_K2::PC_Exec)
-				//{
-				//	// Add boundary for expansion events.
-				//	LinkedPinMap.Add(Pin, LinkedPin);
-				//}
-				if (UEdGraphPin** SourcePin = LinkedPinMap.Find(LinkedPin))
-				{
-					if (Pin->Direction == EGPD_Input)
-					{
-						TArray<UEdGraphPin*>& PinLinks = EntryPins.FindOrAdd(*SourcePin);
-						PinLinks.Add(Pin);
-					}
-					else
-					{
-						TArray<UEdGraphPin*>& PinLinks = ExitPins.FindOrAdd(*SourcePin);
-						PinLinks.Add(Pin);
-					}
-				}
-			}
-		}
+	case ETunnelBoundaryType::OutputSite:
+		Args.Add(TEXT("TunnelBoundaryType"), FText::FromString(TEXT("Tunnel Output Site")));
+		break;
 	}
-	// Create the Boundary nodes for each unique entry/exit site.
-	if (TargetGraph)
-	{
-		if (EntryPins.Num() > 0)
-		{
-			for (const auto& PinEntry : EntryPins)
-			{
-				// Create entry node and pins
-				UK2Node_TunnelBoundary* EntryBoundaryNode = TargetGraph->CreateIntermediateNode<UK2Node_TunnelBoundary>();
-				MessageLog.NotifyIntermediateObjectCreation(EntryBoundaryNode, SourceNode);
-				EntryBoundaryNode->CreateNewGuid();
-				EntryBoundaryNode->WireUpEntry(SourceNode, PinEntry.Key, PinEntry.Value, MessageLog);
-			}
-		}
-		if (ExitPins.Num() > 0)
-		{
-			for (const auto& PinExit : ExitPins)
-			{
-				// Create exit node and pins
-				UK2Node_TunnelBoundary* ExitBoundaryNode = TargetGraph->CreateIntermediateNode<UK2Node_TunnelBoundary>();
-				MessageLog.NotifyIntermediateObjectCreation(ExitBoundaryNode, SourceNode);
-				ExitBoundaryNode->CreateNewGuid();
-				ExitBoundaryNode->WireUpExit(SourceNode, PinExit.Key, PinExit.Value, MessageLog);
-			}
-		}
-	}
-}
 
-bool UK2Node_TunnelBoundary::IsPureTunnel(UK2Node_Tunnel* Tunnel)
-{
-	return !Tunnel->IsA<UK2Node_MacroInstance>() && !Tunnel->IsA<UK2Node_Composite>();
+	return FText::Format(LOCTEXT("UK2Node_TunnelBoundary_FullTitle", "{BaseName} - {TunnelBoundaryType}"), Args);
 }
 
-void UK2Node_TunnelBoundary::WireUpTunnelEntry(UK2Node_Tunnel* TunnelInstance, UEdGraphPin* TunnelPin, FCompilerResultsLog& MessageLog)
-{
-	if (TunnelInstance && TunnelPin)
-	{
-		// Mark as entry node
-		TunnelBoundaryType = TBT_EntrySite;
-		// Grab the graph name
-		CreateBaseNodeName(TunnelInstance);
-		// Find the true source pin
-		UEdGraphPin* SourcePin = TunnelInstance->FindPin(TunnelPin->PinName);
-		// Wire in the node
-		UEdGraphPin* OutputPin = CreatePin(EGPD_Output, UEdGraphSchema_K2::PC_Exec, TEXT("TunnelEntryExec"));
-		for (UEdGraphPin* LinkedPin : TunnelPin->LinkedTo)
-		{
-			LinkedPin->MakeLinkTo(OutputPin);
-		}
-		UEdGraphPin* InputPin = CreatePin(EGPD_Input, UEdGraphSchema_K2::PC_Exec, SourcePin->PinName);
-		MessageLog.NotifyIntermediatePinCreation(InputPin, SourcePin);
-		TunnelPin->BreakAllPinLinks();
-		TunnelPin->MakeLinkTo(InputPin);
-	}
-}
-
-void UK2Node_TunnelBoundary::WireUpTunnelExit(UK2Node_Tunnel* TunnelInstance, UEdGraphPin* TunnelPin, FCompilerResultsLog& MessageLog)
-{
-	if (TunnelInstance && TunnelPin)
-	{
-		// Mark as exit node
-		TunnelBoundaryType = TBT_ExitSite;
-		// Grab the graph name
-		CreateBaseNodeName(TunnelInstance);
-		// Find the true source pin
-		UEdGraphPin* SourcePin = TunnelInstance->FindPin(TunnelPin->PinName);
-		// Wire in the node
-		UEdGraphPin* InputPin = CreatePin(EGPD_Input, UEdGraphSchema_K2::PC_Exec, SourcePin->PinName);
-		for (UEdGraphPin* LinkedPin : TunnelPin->LinkedTo)
-		{
-			LinkedPin->MakeLinkTo(InputPin);
-		}
-		UEdGraphPin* OutputPin = CreatePin(EGPD_Output, UEdGraphSchema_K2::PC_Exec, TEXT("TunnelExitExec"));
-		MessageLog.NotifyIntermediatePinCreation(OutputPin, SourcePin);
-		TunnelPin->BreakAllPinLinks();
-		TunnelPin->MakeLinkTo(OutputPin);
-	}
-}
-
-void UK2Node_TunnelBoundary::WireUpEntry(UEdGraphNode* SourceNode, UEdGraphPin* SourcePin, const TArray<UEdGraphPin*>& EntryPins, FCompilerResultsLog& MessageLog)
-{
-	if (SourceNode && SourcePin && EntryPins.Num())
-	{
-		// Mark as entry node
-		TunnelBoundaryType = TBT_EntrySite;
-		// Create the pin base name
-		CreateBaseNodeName(SourceNode);
-		// Wire in the node
-		const UEdGraphNode::FCreatePinParams PinParams(SourcePin->PinType);
-		UEdGraphPin* InputPin = CreatePin(EGPD_Input, SourcePin->PinType.PinCategory, SourcePin->PinType.PinSubCategory, SourcePin->PinType.PinSubCategoryObject.Get(), TEXT("EntryBoundary"), PinParams);
-		MessageLog.NotifyIntermediatePinCreation(InputPin, SourcePin);
-		UEdGraphPin* OutputPin = CreatePin(EGPD_Output, SourcePin->PinType.PinCategory, SourcePin->PinType.PinSubCategory, SourcePin->PinType.PinSubCategoryObject.Get(), SourcePin->PinName, PinParams);
-		for (UEdGraphPin* EntryPin : EntryPins)
-		{
-			for (UEdGraphPin* LinkedPin : EntryPin->LinkedTo)
-			{
-				check (LinkedPin->Direction != InputPin->Direction);
-				LinkedPin->MakeLinkTo(InputPin);
-			}
-			EntryPin->BreakAllPinLinks();
-			check (EntryPin->Direction != OutputPin->Direction);
-			EntryPin->MakeLinkTo(OutputPin);
-		}
-	}
-}
-
-void UK2Node_TunnelBoundary::WireUpExit(UEdGraphNode* SourceNode, UEdGraphPin* SourcePin, const TArray<UEdGraphPin*>& ExitPins, FCompilerResultsLog& MessageLog)
+FNodeHandlingFunctor* UK2Node_TunnelBoundary::CreateNodeHandler(FKismetCompilerContext& CompilerContext) const
 {
-	if (SourceNode && SourcePin && ExitPins.Num())
-	{
-		// Mark as exit node
-		TunnelBoundaryType = TBT_ExitSite;
-		// Create the pin base name
-		CreateBaseNodeName(SourceNode);
-		// Wire in the node
-		const UEdGraphNode::FCreatePinParams PinParams(SourcePin->PinType);
-		UEdGraphPin* OutputPin = CreatePin(EGPD_Output, SourcePin->PinType.PinCategory, SourcePin->PinType.PinSubCategory, SourcePin->PinType.PinSubCategoryObject.Get(), SourcePin->PinName, PinParams);
-		MessageLog.NotifyIntermediatePinCreation(OutputPin, SourcePin);
-		UEdGraphPin* InputPin = CreatePin(EGPD_Input, SourcePin->PinType.PinCategory, SourcePin->PinType.PinSubCategory, SourcePin->PinType.PinSubCategoryObject.Get(), TEXT("ExitBoundary"), PinParams);
-		for (UEdGraphPin* ExitPin : ExitPins)
-		{
-			for (UEdGraphPin* LinkedPin : ExitPin->LinkedTo)
-			{
-				check (LinkedPin->Direction != OutputPin->Direction);
-				LinkedPin->MakeLinkTo(OutputPin);
-			}
-			ExitPin->BreakAllPinLinks();
-			check (ExitPin->Direction != InputPin->Direction);
-			ExitPin->MakeLinkTo(InputPin);
-		}
-	}
+	return new FKCHandler_TunnelBoundary(CompilerContext);
 }
 
-void UK2Node_TunnelBoundary::CreateBaseNodeName(UEdGraphNode* SourceNode)
+void UK2Node_TunnelBoundary::SetNodeAttributes(const UK2Node_Tunnel* SourceNode)
 {
-	UEdGraph* TunnelGraph = nullptr;
 	if (const UK2Node_MacroInstance* MacroInstance = Cast<UK2Node_MacroInstance>(SourceNode))
 	{
-		TunnelGraph = MacroInstance->GetMacroGraph();
+		// Macro tunnel instance node
+		const UEdGraph* TunnelGraph = MacroInstance->GetMacroGraph();
 		BaseName = TunnelGraph ? TunnelGraph->GetFName() : NAME_None;
+		TunnelBoundaryType = ETunnelBoundaryType::EntrySite;
 	}
 	else if (const UK2Node_Composite* CompositeInstance = Cast<UK2Node_Composite>(SourceNode))
 	{
-		TunnelGraph = CompositeInstance->BoundGraph;
+		// Composite tunnel instance node
+		const UEdGraph* TunnelGraph = CompositeInstance->BoundGraph;
 		BaseName = TunnelGraph ? TunnelGraph->GetFName() : NAME_None;
+		TunnelBoundaryType = ETunnelBoundaryType::EntrySite;
 	}
-	else
+	else if (SourceNode)
 	{
-		BaseName = SourceNode ? SourceNode->GetFName() : NAME_None;
+		// Input or output tunnel node
+		SetNodeAttributes(SourceNode->bCanHaveOutputs ? SourceNode->GetOutputSource() : SourceNode->GetInputSink());
+		TunnelBoundaryType = SourceNode->bCanHaveInputs ? ETunnelBoundaryType::OutputSite : ETunnelBoundaryType::InputSite;
 	}
-}
-
-void UK2Node_TunnelBoundary::BuildSourceNodeMap(UEdGraphNode* Tunnel, TMap<FGuid, const UEdGraphNode*>& SourceNodeMap)
-{
-	// Crawl the graph and locate source node guids
-	UEdGraph* TunnelGraph = nullptr;
-	if (UK2Node_MacroInstance* SourceMacroInstance = Cast<UK2Node_MacroInstance>(Tunnel))
-	{
-		TunnelGraph = SourceMacroInstance->GetMacroGraph();
-	}
-	else if (UK2Node_Composite* SourceCompositeInstance = Cast<UK2Node_Composite>(Tunnel))
-	{
-		TunnelGraph = SourceCompositeInstance->BoundGraph;
-	}
-	if (TunnelGraph)
-	{
-		for (UEdGraphNode* GraphNode : TunnelGraph->Nodes)
-		{
-			SourceNodeMap.Add(GraphNode->NodeGuid) = GraphNode;
-		}
-	}
-}
-
-UEdGraphNode* UK2Node_TunnelBoundary::FindTrueSourceTunnelInstance(UEdGraphNode* Tunnel, UEdGraphNode* SourceTunnelInstance)
-{
-	UEdGraphNode* SourceNode = nullptr;
-	if (Tunnel && SourceTunnelInstance)
-	{
-		if (Tunnel->NodeGuid == SourceTunnelInstance->NodeGuid)
-		{
-			SourceNode = SourceTunnelInstance;
-		}
-		else
-		{
-			UEdGraph* TunnelGraph = nullptr;
-			if (UK2Node_MacroInstance* SourceMacroInstance = Cast<UK2Node_MacroInstance>(SourceTunnelInstance))
-			{
-				TunnelGraph = SourceMacroInstance->GetMacroGraph();
-			}
-			else if (UK2Node_Composite* SourceCompositeInstance = Cast<UK2Node_Composite>(SourceTunnelInstance))
-			{
-				TunnelGraph = SourceCompositeInstance->BoundGraph;
-			}
-			if (TunnelGraph)
-			{
-				for (UEdGraphNode* GraphNode : TunnelGraph->Nodes)
-				{
-					if (GraphNode->NodeGuid == Tunnel->NodeGuid)
-					{
-						SourceNode = GraphNode;
-						break;
-					}
-
-					if (GraphNode->IsA<UK2Node_Composite>() || GraphNode->IsA<UK2Node_MacroInstance>())
-					{
-						SourceNode = FindTrueSourceTunnelInstance(Tunnel, GraphNode);
-					if (SourceNode)
-					{
-						break;
-					}
-				}
-			}
-		}
-	}
-	}
-	return SourceNode;
-}
-
-void UK2Node_TunnelBoundary::FindTunnelExitSiteInstances(UEdGraphPin* NodeEntryPin, TArray<UEdGraphPin*>& ExitPins, TArray<UEdGraphPin*>& VisitedPins, const UEdGraphNode* TunnelExitNode)
-{
-	UEdGraphNode* PinNode = NodeEntryPin->GetOwningNode();
-
-	// Grab pin node
-	if (PinNode != nullptr && PinNode != TunnelExitNode && !VisitedPins.Contains(NodeEntryPin))
+	else
 	{
-		VisitedPins.Push(NodeEntryPin);
-
-		// Find all exec out pins
-		TArray<UEdGraphPin*> ExecPins;
-		TArray<UEdGraphPin*> ConnectedExecPins;
-		for (UEdGraphPin* Pin : PinNode->Pins)
-		{
-			if (Pin->Direction == EGPD_Output && Pin->PinType.PinCategory == UEdGraphSchema_K2::PC_Exec)
-			{
-				ExecPins.Add(Pin);
-				if (Pin->LinkedTo.Num() > 0)
-				{
-					ConnectedExecPins.Add(Pin);
-				}
-			}
-		}
-		if (ConnectedExecPins.Num() > 1)
-		{
-			if (PinNode->IsA<UK2Node_ExecutionSequence>())
-			{
-				// Execution sequence style nodes, follow the last connected pin
-				for (UEdGraphPin* LinkedPin : ConnectedExecPins.Last()->LinkedTo)
-				{
-					FindTunnelExitSiteInstances(LinkedPin, ExitPins, VisitedPins);
-				}
-			}
-			else
-			{
-				// Branch style nodes, more tricky
-				for (UEdGraphPin* ExecPin : ConnectedExecPins)
-				{
-					for (UEdGraphPin* LinkedPin : ExecPin->LinkedTo)
-					{
-						FindTunnelExitSiteInstances(LinkedPin, ExitPins, VisitedPins);
-					}
-				}
-			}
-		}
-		else if (ConnectedExecPins.Num() == 1)
-		{
-			for (UEdGraphPin* LinkedPin : ConnectedExecPins[0]->LinkedTo)
-			{
-				FindTunnelExitSiteInstances(LinkedPin, ExitPins, VisitedPins);
-			}
-		}
-		else if (ExecPins.Num() > 0)
-		{
-			ExitPins.Add(ExecPins.Last());
-		}
-
-		VisitedPins.Pop(false);
+		BaseName = NAME_None;
+		TunnelBoundaryType = ETunnelBoundaryType::Unknown;
 	}
 }
 
diff --git a/Engine/Source/Editor/BlueprintGraph/Private/K2Node_Variable.cpp b/Engine/Source/Editor/BlueprintGraph/Private/K2Node_Variable.cpp
--- a/Engine/Source/Editor/BlueprintGraph/Private/K2Node_Variable.cpp
+++ b/Engine/Source/Editor/BlueprintGraph/Private/K2Node_Variable.cpp
@@ -11,6 +11,7 @@
 #include "Kismet2/BlueprintEditorUtils.h"
 #include "KismetCompilerMisc.h"
 #include "Kismet2/CompilerResultsLog.h"
+#include "Kismet2/StructureEditorUtils.h"
 #include "Styling/SlateIconFinder.h"
 #include "Logging/MessageLog.h"
 #include "SourceCodeNavigation.h"
@@ -305,10 +306,10 @@ UK2Node::ERedirectType UK2Node_Variable::DoPinsMatchForReconstruction( const UEd
 				}
 
 				// Compare whether the names, ignoring the original variable's name in the case of renames, match
-				FString NewPinPropertyName = NewPin->PinName.ToString().RightChop(ParentmostNewPin->PinName.ToString().Len() + 1);
-				FString OldPinPropertyName = OldPin->PinName.ToString().RightChop(ParentmostOldPin->PinName.ToString().Len() + 1);
+				FName NewPinPropertyName = FName(*NewPin->PinName.ToString().RightChop(ParentmostNewPin->PinName.ToString().Len() + 1));
+				FName OldPinPropertyName = FName(*OldPin->PinName.ToString().RightChop(ParentmostOldPin->PinName.ToString().Len() + 1));
 
-				if (NewPinPropertyName != OldPinPropertyName)
+				if (!DoesRenamedVariableMatch(OldPinPropertyName, NewPinPropertyName, Cast<UStruct>(NewPin->ParentPin->PinType.PinSubCategoryObject.Get())))
 				{
 					return ERedirectType_None;
 				}
@@ -426,6 +427,36 @@ UProperty* UK2Node_Variable::GetPropertyForVariable() const
 	return GetPropertyForVariable_Internal(GetBlueprintClassFromNode());
 }
 
+bool UK2Node_Variable::DoesRenamedVariableMatch(FName OldVariableName, FName NewVariableName, UStruct* StructType)
+{
+	if (NewVariableName == OldVariableName)
+	{
+		return true;
+	}
+
+	FGuid OldPropertyGuid = FStructureEditorUtils::GetGuidFromPropertyName(OldVariableName);
+	FGuid NewPropertyGuid = FStructureEditorUtils::GetGuidFromPropertyName(NewVariableName);
+
+	// If guids match this was a user struct rename
+	if (OldPropertyGuid.IsValid() && (OldPropertyGuid == NewPropertyGuid))
+	{
+		return true;
+	}
+
+	// Also check native rename if we can find the struct
+	if (StructType)
+	{
+		FName RedirectedPinName = UProperty::FindRedirectedPropertyName(StructType, OldVariableName);
+
+		if (NewVariableName == RedirectedPinName)
+		{
+			return true;
+		}
+	}
+
+	return false;
+}
+
 UProperty* UK2Node_Variable::GetPropertyForVariableFromSkeleton() const
 {
 	if (UClass* SkeletonClass = FBlueprintEditorUtils::GetSkeletonClass(VariableReference.GetMemberParentClass(GetBlueprintClassFromNode())))
diff --git a/Engine/Source/Editor/ContentBrowser/Private/AssetViewSortManager.cpp b/Engine/Source/Editor/ContentBrowser/Private/AssetViewSortManager.cpp
--- a/Engine/Source/Editor/ContentBrowser/Private/AssetViewSortManager.cpp
+++ b/Engine/Source/Editor/ContentBrowser/Private/AssetViewSortManager.cpp
@@ -384,7 +384,7 @@ void FAssetViewSortManager::SortList(TArray<TSharedPtr<FAssetViewItem>>& AssetIt
 			UObject::FAssetRegistryTag::ETagType TagType = FindAndRefreshCustomColumn(AssetItems, Tag, CustomColumns);
 			
 			// Since this SortData.Tag is not one of preset columns, sort by asset registry tag	
-			if (TagType != UObject::FAssetRegistryTag::ETagType::TT_Hidden && MajorityAssetType != NAME_None)
+			if (TagType == UObject::FAssetRegistryTag::ETagType::TT_Hidden && MajorityAssetType != NAME_None)
 			{
 				UClass* Class = FindObject<UClass>(ANY_PACKAGE, *MajorityAssetType.ToString());
 				if (Class)
diff --git a/Engine/Source/Editor/ContentBrowser/Private/AssetViewWidgets.cpp b/Engine/Source/Editor/ContentBrowser/Private/AssetViewWidgets.cpp
--- a/Engine/Source/Editor/ContentBrowser/Private/AssetViewWidgets.cpp
+++ b/Engine/Source/Editor/ContentBrowser/Private/AssetViewWidgets.cpp
@@ -975,7 +975,7 @@ void SAssetViewItem::UpdateSourceControlState(float InDeltaTime)
 	if ( !bSourceControlStateRequested && SourceControlStateDelay > 1.0f && ISourceControlModule::Get().IsEnabled() && AssetItem.IsValid() )
 	{
 		// Only update the SCC state for non-temporary asset items that aren't a built in script
-		if ( !AssetItem->IsTemporaryItem() && AssetItem->GetType() != EAssetItemType::Folder && !FPackageName::IsScriptPackage(CachedPackageName) )
+		if ( !AssetItem->IsTemporaryItem() && AssetItem->GetType() != EAssetItemType::Folder && !FPackageName::IsScriptPackage(CachedPackageName) && !FPackageName::IsMemoryPackage(CachedPackageName) && !FPackageName::IsTempPackage(CachedPackageName))
 		{
 			// Request the most recent SCC state for this asset
 			ISourceControlModule::Get().QueueStatusUpdate(CachedPackageFileName);
