Copy of CL #3812127. Putting in Main as several branches were running into it and it's a bad crash

diff --git a/Engine/Source/Editor/BlueprintGraph/Classes/K2Node_AssignDelegate.h b/Engine/Source/Editor/BlueprintGraph/Classes/K2Node_AssignDelegate.h
--- a/Engine/Source/Editor/BlueprintGraph/Classes/K2Node_AssignDelegate.h
+++ b/Engine/Source/Editor/BlueprintGraph/Classes/K2Node_AssignDelegate.h
@@ -25,9 +25,13 @@ public:
 	virtual FText GetMenuCategory() const override;
 	virtual bool IsCompatibleWithGraph(const UEdGraph* TargetGraph) const override;
 	virtual void PostPlacedNewNode() override;
+	virtual void AllocateDefaultPins() override;
 	// End of UEdGraphNode interface
 
 private:
 	/** Constructing FText strings can be costly, so we cache the node's title */
 	FNodeTextCache CachedListTitle;
+
+	/** If true, should try to spawn a bound custom event node on next allocate */
+	bool bShouldSpawnCustomEvent;
 };
diff --git a/Engine/Source/Editor/BlueprintGraph/Private/K2Node_AssignDelegate.cpp b/Engine/Source/Editor/BlueprintGraph/Private/K2Node_AssignDelegate.cpp
--- a/Engine/Source/Editor/BlueprintGraph/Private/K2Node_AssignDelegate.cpp
+++ b/Engine/Source/Editor/BlueprintGraph/Private/K2Node_AssignDelegate.cpp
@@ -15,6 +15,7 @@
 UK2Node_AssignDelegate::UK2Node_AssignDelegate(const FObjectInitializer& ObjectInitializer)
 	: Super(ObjectInitializer)
 {
+	bShouldSpawnCustomEvent = false;
 }
 
 //------------------------------------------------------------------------------
@@ -73,9 +74,18 @@ bool UK2Node_AssignDelegate::IsCompatibleWithGraph(const UEdGraph* TargetGraph)
 //------------------------------------------------------------------------------
 void UK2Node_AssignDelegate::PostPlacedNewNode()
 {
+	Super::PostPlacedNewNode();
+	bShouldSpawnCustomEvent = true;
+}
+
+void UK2Node_AssignDelegate::AllocateDefaultPins()
+{
+	Super::AllocateDefaultPins();
+
 	UEdGraphPin* InDelegatePin = GetDelegatePin();
-	if (InDelegatePin->LinkedTo.Num() == 0)
+	if (InDelegatePin->LinkedTo.Num() == 0 && bShouldSpawnCustomEvent)
 	{
+		bShouldSpawnCustomEvent = false;
 		if (UMulticastDelegateProperty* DelegateProp = Cast<UMulticastDelegateProperty>(GetProperty()))
 		{
 			FVector2D const LocationOffset(-150.f, +150.f);
