Merging //UE4/Dev-Main to Dev-Geometry (//UE4/Dev-Geometry)

diff --git a/Engine/Config/BaseEditorPerProjectUserSettings.ini b/Engine/Config/BaseEditorPerProjectUserSettings.ini
--- a/Engine/Config/BaseEditorPerProjectUserSettings.ini
+++ b/Engine/Config/BaseEditorPerProjectUserSettings.ini
@@ -141,6 +141,7 @@ bToolbarCustomization=False
 bBehaviorTreeEditor=False
 bEnableFindAndReplaceReferences =False
 bExampleLayersAndBlends = True
+bEnableBlueprintGraphBookmarks=False
 
 [/Script/UnrealEd.EditorLoadingSavingSettings]
 ; True if we should automatically load a default level at start up
diff --git a/Engine/Config/BaseEngine.ini b/Engine/Config/BaseEngine.ini
--- a/Engine/Config/BaseEngine.ini
+++ b/Engine/Config/BaseEngine.ini
@@ -49,7 +49,7 @@ ConsoleClassName=/Script/Engine.Console
 GameViewportClientClassName=/Script/Engine.GameViewportClient
 LocalPlayerClassName=/Script/Engine.LocalPlayer
 WorldSettingsClassName=/Script/Engine.WorldSettings
-NavigationSystemClassName=/Script/Engine.NavigationSystem
+NavigationSystemClassName=/Script/NavigationSystem.NavigationSystemV1
 AvoidanceManagerClassName=/Script/Engine.AvoidanceManager
 PhysicsCollisionHandlerClassName=/Script/Engine.PhysicsCollisionHandler
 LevelScriptActorClassName=/Script/Engine.LevelScriptActor
@@ -509,6 +509,54 @@ bUseStreamingPause=false
 +FunctionRedirects=(OldName="/Script/HeadMountedDisplay.HeadMountedDisplayFunctionLibrary.AddDeviceVisualizationComponent",NewName="/Script/HeadMountedDisplay.XRAssetFunctionLibrary.AddDeviceVisualizationComponentBlocking")
 +FunctionRedirects=(OldName="/Script/HeadMountedDisplay.HeadMountedDisplayFunctionLibrary.AddNamedDeviceVisualizationComponent",NewName="/Script/HeadMountedDisplay.XRAssetFunctionLibrary.AddNamedDeviceVisualizationComponentBlocking")
 
++ClassRedirects=(OldName="NavigationSystem",NewName="/Script/NavigationSystem.NavigationSystemV1")
++ClassRedirects=(OldName="NavMeshBoundsVolume",NewName="/Script/NavigationSystem.NavMeshBoundsVolume")
++ClassRedirects=(OldName="NavArea",NewName="/Script/NavigationSystem.NavArea")
++ClassRedirects=(OldName="NavAreaMeta",NewName="/Script/NavigationSystem.NavAreaMeta")
++ClassRedirects=(OldName="NavArea_Default",NewName="/Script/NavigationSystem.NavArea_Default")
++ClassRedirects=(OldName="NavArea_LowHeight",NewName="/Script/NavigationSystem.NavArea_LowHeight")
++ClassRedirects=(OldName="NavArea_Null",NewName="/Script/NavigationSystem.NavArea_Null")
++ClassRedirects=(OldName="NavArea_Obstacle",NewName="/Script/NavigationSystem.NavArea_Obstacle")
++ClassRedirects=(OldName="NavAreaMeta_SwitchByAgent",NewName="/Script/NavigationSystem.NavAreaMeta_SwitchByAgent")
++ClassRedirects=(OldName="NavigationQueryFilter",NewName="/Script/NavigationSystem.NavigationQueryFilter")
++ClassRedirects=(OldName="NavMeshRenderingComponent",NewName="/Script/NavigationSystem.NavMeshRenderingComponent")
++ClassRedirects=(OldName="RecastNavMesh",NewName="/Script/NavigationSystem.RecastNavMesh")
++ClassRedirects=(OldName="RecastNavMeshDataChunk",NewName="/Script/NavigationSystem.RecastNavMeshDataChunk")
++ClassRedirects=(OldName="AbstractNavData",NewName="/Script/NavigationSystem.AbstractNavData")
++ClassRedirects=(OldName="CrowdManagerBase",NewName="/Script/NavigationSystem.CrowdManagerBase")
++ClassRedirects=(OldName="NavCollision",NewName="/Script/NavigationSystem.NavCollision")
++ClassRedirects=(OldName="NavigationData",NewName="/Script/NavigationSystem.NavigationData")
++ClassRedirects=(OldName="NavigationInvokerComponent",NewName="/Script/NavigationSystem.NavigationInvokerComponent")
++ClassRedirects=(OldName="NavigationPath",NewName="/Script/NavigationSystem.NavigationPath")
++ClassRedirects=(OldName="NavigationTestingActor",NewName="/Script/NavigationSystem.NavigationTestingActor")
++ClassRedirects=(OldName="NavLinkComponent",NewName="/Script/NavigationSystem.NavLinkComponent")
++ClassRedirects=(OldName="NavLinkCustomComponent",NewName="/Script/NavigationSystem.NavLinkCustomComponent")
++ClassRedirects=(OldName="NavLinkRenderingComponent",NewName="/Script/NavigationSystem.NavLinkRenderingComponent")
++ClassRedirects=(OldName="NavLinkTrivial",NewName="/Script/NavigationSystem.NavLinkTrivial")
++ClassRedirects=(OldName="NavModifierComponent",NewName="/Script/NavigationSystem.NavModifierComponent")
++ClassRedirects=(OldName="NavModifierVolume",NewName="/Script/NavigationSystem.NavModifierVolume")
++ClassRedirects=(OldName="NavRelevantComponent",NewName="/Script/NavigationSystem.NavRelevantComponent")
++ClassRedirects=(OldName="RecastFilter_UseDefaultArea",NewName="/Script/NavigationSystem.RecastFilter_UseDefaultArea")
++ClassRedirects=(OldName="NavigationGraph",NewName="/Script/NavigationSystem.NavigationGraph")
++ClassRedirects=(OldName="NavigationGraphNode",NewName="/Script/NavigationSystem.NavigationGraphNode")
++ClassRedirects=(OldName="NavigationGraphNodeComponent",NewName="/Script/NavigationSystem.NavigationGraphNodeComponent")
++ClassRedirects=(OldName="NavigationPathGenerator",NewName="/Script/NavigationSystem.NavigationPathGenerator")
++ClassRedirects=(OldName="NavLinkCustomInterface",NewName="/Script/NavigationSystem.NavLinkCustomInterface")
++ClassRedirects=(OldName="NavLinkHostInterface",NewName="/Script/NavigationSystem.NavLinkHostInterface")
++ClassRedirects=(OldName="NavNodeInterface",NewName="/Script/NavigationSystem.NavNodeInterface")
++ClassRedirects=(OldName="NavLinkProxy",NewName="/Script/AIModule.NavLinkProxy")
++StructRedirects=(OldName="NavigationFilterArea",NewName="/Script/NavigationSystem.NavigationFilterArea")
++StructRedirects=(OldName="NavigationFilterFlags",NewName="/Script/NavigationSystem.NavigationFilterFlags")
++StructRedirects=(OldName="NavGraphEdge",NewName="/Script/NavigationSystem.NavGraphEdge")
++StructRedirects=(OldName="NavGraphNode",NewName="/Script/NavigationSystem.NavGraphNode")
++StructRedirects=(OldName="NavCollisionCylinder",NewName="/Script/NavigationSystem.NavCollisionCylinder")
++StructRedirects=(OldName="NavCollisionBox",NewName="/Script/NavigationSystem.NavCollisionBox")
++StructRedirects=(OldName="SupportedAreaData",NewName="/Script/NavigationSystem.SupportedAreaData")
++FunctionRedirects=(OldName="NavigationSystemV1.SimpleMoveToActor",NewName="AIBlueprintHelperLibrary.SimpleMoveToActor")
++FunctionRedirects=(OldName="NavigationSystemV1.SimpleMoveToLocation",NewName="AIBlueprintHelperLibrary.SimpleMoveToLocation")
+
+
+
 [CoreUObject.Metadata]
 ; Note: UnrealHeaderTool should be rerun after making changes to MetadataRedirects to catch any keys specified in class headers
 +MetadataRedirects=(OldKey="K2Protected", NewKey="BlueprintProtected")
@@ -913,7 +961,6 @@ bDeferDependencyLoads=true
 bForceDisableCookedDependencyDeferring=false
 bExecutionAfterReturn=false
 bCanSuppressAccessViolation=false
-bUseSelectRef=true
 bLoadNativeConvertedBPClassInEditor=false
 bIgnoreCompileOnLoadErrorsOnBuildMachine=false
 
@@ -1434,14 +1481,14 @@ ParticleEditorPromotionTest=(DefaultParticleAsset=(FilePath="../../Content/Tutor
 +EditorViewButtons=UV
 EditorViewButtonsObject=EditorCylinder
 
-[/Script/Engine.NavigationSystem]
+[/Script/NavigationSystem.NavigationSystemV1]
 bAutoCreateNavigationData=true
 bAddPlayersToGenerationSeeds=true
 
-[/Script/Engine.NavigationData]
+[/Script/NavigationSystem.NavigationData]
 RuntimeGeneration=Static
 
-[/Script/Engine.RecastNavMesh]
+[/Script/NavigationSystem.RecastNavMesh]
 ; runtime params
 TileSetUpdateInterval=1.0
 MaxTileGridWidth=256
@@ -1462,10 +1509,10 @@ MergeRegionSize=400.f
 bUseBetterOffsetsFromCorners=true
 
 
-[/Script/Engine.NavArea_Null]
+[/Script/NavigationSystem.NavArea_Null]
 DrawColor=(R=38,G=38,B=38,A=64)
 
-[/Script/Engine.NavArea_Default]
+[/Script/NavigationSystem.NavArea_Default]
 DrawColor=(R=140,G=255,B=0,A=164)
 
 [RemoteConfiguration]
@@ -1913,7 +1960,6 @@ DefaultCompletionMode=RestoreState
 
 [PlatformCrypto]
 PlatformRequiresDataCrypto=False
-
 [/Script/AppleARKit.AppleARKitSettings]
 bEnableLiveLinkForFaceTracking=true
 LiveLinkPublishingPort=11111
diff --git a/Engine/Config/BaseGame.ini b/Engine/Config/BaseGame.ini
--- a/Engine/Config/BaseGame.ini
+++ b/Engine/Config/BaseGame.ini
@@ -16,6 +16,7 @@ MAXCLIENTUPDATEINTERVAL=0.25f
 MaxMoveDeltaTime=0.125f
 ClientNetSendMoveDeltaTime=0.0166
 ClientNetSendMoveDeltaTimeThrottled=0.0222
+ClientNetSendMoveDeltaTimeStationary=0.0833
 ClientNetSendMoveThrottleAtNetSpeed=10000
 ClientNetSendMoveThrottleOverPlayerCount=10
 ClientAuthorativePosition=false
diff --git a/Engine/Plugins/2D/Paper2D/Source/Paper2D/Paper2D.Build.cs b/Engine/Plugins/2D/Paper2D/Source/Paper2D/Paper2D.Build.cs
--- a/Engine/Plugins/2D/Paper2D/Source/Paper2D/Paper2D.Build.cs
+++ b/Engine/Plugins/2D/Paper2D/Source/Paper2D/Paper2D.Build.cs
@@ -16,8 +16,9 @@ public class Paper2D : ModuleRules
 				"ShaderCore",
 				"RenderCore",
 				"RHI",
-                "SlateCore",
-                "Slate"
+        		        "SlateCore",
+        		        "Slate",
+                		"NavigationSystem"
 			}
 		);
 
diff --git a/Engine/Plugins/2D/Paper2D/Source/Paper2D/Private/PaperGroupedSpriteComponent.cpp b/Engine/Plugins/2D/Paper2D/Source/Paper2D/Private/PaperGroupedSpriteComponent.cpp
--- a/Engine/Plugins/2D/Paper2D/Source/Paper2D/Private/PaperGroupedSpriteComponent.cpp
+++ b/Engine/Plugins/2D/Paper2D/Source/Paper2D/Private/PaperGroupedSpriteComponent.cpp
@@ -1,13 +1,11 @@
 // Copyright 1998-2018 Epic Games, Inc. All Rights Reserved.
 
 #include "PaperGroupedSpriteComponent.h"
-#include "AI/Navigation/NavigationSystem.h"
+#include "AI/NavigationSystemBase.h"
 #include "Engine/CollisionProfile.h"
 
 #include "GroupedSpriteSceneProxy.h"
 #include "AI/NavigationSystemHelpers.h"
-#include "AI/NavigationOctree.h"
-
 
 #include "Logging/TokenizedMessage.h"
 #include "Logging/MessageLog.h"
@@ -44,7 +42,7 @@ int32 UPaperGroupedSpriteComponent::AddInstanceWithMaterial(const FTransform& Tr
 
 	MarkRenderStateDirty();
 
-	UNavigationSystem::UpdateComponentInNavOctree(*this);
+	FNavigationSystem::UpdateComponentData(*this);
 
 	return NewInstanceIndex;
 }
@@ -94,7 +92,7 @@ bool UPaperGroupedSpriteComponent::UpdateInstanceTransform(int32 InstanceIndex,
 	}
 
 	// Request navigation update
-	UNavigationSystem::UpdateComponentInNavOctree(*this);
+	FNavigationSystem::UpdateComponentData(*this);
 
 	FSpriteInstanceData& InstanceData = PerInstanceSpriteData[InstanceIndex];
 
@@ -115,7 +113,7 @@ bool UPaperGroupedSpriteComponent::UpdateInstanceTransform(int32 InstanceIndex,
 	}
 
 	// Request navigation update
-	UNavigationSystem::UpdateComponentInNavOctree(*this);
+	FNavigationSystem::UpdateComponentData(*this);
 
 	if (bMarkRenderStateDirty)
 	{
@@ -152,7 +150,7 @@ bool UPaperGroupedSpriteComponent::RemoveInstance(int32 InstanceIndex)
 	}
 
 	// Request navigation update
-	UNavigationSystem::UpdateComponentInNavOctree(*this);
+	FNavigationSystem::UpdateComponentData(*this);
 
 	// remove instance
 	PerInstanceSpriteData.RemoveAt(InstanceIndex);
@@ -182,7 +180,7 @@ void UPaperGroupedSpriteComponent::ClearInstances()
 	// Indicate we need to update render state to reflect changes
 	MarkRenderStateDirty();
 
-	UNavigationSystem::UpdateComponentInNavOctree(*this);
+	FNavigationSystem::UpdateComponentData(*this);
 }
 
 int32 UPaperGroupedSpriteComponent::GetInstanceCount() const
@@ -336,7 +334,7 @@ void UPaperGroupedSpriteComponent::PostEditUndo()
 {
 	Super::PostEditUndo();
 
-	UNavigationSystem::UpdateComponentInNavOctree(*this);
+	FNavigationSystem::UpdateComponentData(*this);
 }
 #endif
 
diff --git a/Engine/Plugins/2D/Paper2D/Source/Paper2DEditor/Paper2DEditor.Build.cs b/Engine/Plugins/2D/Paper2D/Source/Paper2DEditor/Paper2DEditor.Build.cs
--- a/Engine/Plugins/2D/Paper2D/Source/Paper2DEditor/Paper2DEditor.Build.cs
+++ b/Engine/Plugins/2D/Paper2D/Source/Paper2DEditor/Paper2DEditor.Build.cs
@@ -29,8 +29,9 @@ public class Paper2DEditor : ModuleRules
 				"EditorStyle",
 				"MeshPaint",
 				"EditorWidgets",
-				"Projects"
-			});
+				"Projects",
+				"NavigationSystem",
+            });
 
 		PrivateIncludePathModuleNames.AddRange(
 			new string[] {
diff --git a/Engine/Plugins/2D/Paper2D/Source/Paper2DEditor/Private/TileMapEditing/EdModeTileMap.cpp b/Engine/Plugins/2D/Paper2D/Source/Paper2DEditor/Private/TileMapEditing/EdModeTileMap.cpp
--- a/Engine/Plugins/2D/Paper2D/Source/Paper2DEditor/Private/TileMapEditing/EdModeTileMap.cpp
+++ b/Engine/Plugins/2D/Paper2D/Source/Paper2DEditor/Private/TileMapEditing/EdModeTileMap.cpp
@@ -1,7 +1,7 @@
 // Copyright 1998-2018 Epic Games, Inc. All Rights Reserved.
 
 #include "TileMapEditing/EdModeTileMap.h"
-#include "AI/Navigation/NavigationSystem.h"
+#include "NavigationSystem.h"
 #include "Framework/Application/SlateApplication.h"
 #include "Engine/Selection.h"
 #include "EditorModeManager.h"
@@ -54,7 +54,7 @@ void FTileMapDirtyRegion::PushToNavSystem() const
 	{
 		if (Component->IsNavigationRelevant())
 		{
-			if (UNavigationSystem* NavSys = UNavigationSystem::GetCurrent(Component))
+			if (UNavigationSystemV1* NavSys = FNavigationSystem::GetCurrent<UNavigationSystemV1>(Component))
 			{
 				NavSys->AddDirtyArea(DirtyRegionInWorldSpace, ENavigationDirtyFlag::All);
 			}
@@ -230,7 +230,7 @@ void FEdModeTileMap::FlushPendingDirtyRegions()
 	{
 		if (Component)
 		{
-			UNavigationSystem::UpdateComponentInNavOctree(*Component);
+			FNavigationSystem::UpdateComponentData(*Component);
 		}		
 	}
 
diff --git a/Engine/Plugins/Compositing/Composure/Source/Composure/Private/ComposurePlayerCompositingCameraModifier.h b/Engine/Plugins/Compositing/Composure/Source/Composure/Private/ComposurePlayerCompositingCameraModifier.h
--- a/Engine/Plugins/Compositing/Composure/Source/Composure/Private/ComposurePlayerCompositingCameraModifier.h
+++ b/Engine/Plugins/Compositing/Composure/Source/Composure/Private/ComposurePlayerCompositingCameraModifier.h
@@ -5,7 +5,7 @@
 #include "CoreMinimal.h"
 #include "UObject/ObjectMacros.h"
 #include "UObject/Object.h"
-#include "Classes/Camera/CameraModifier.h"
+#include "Camera/CameraModifier.h"
 #include "ComposurePlayerCompositingCameraModifier.generated.h"
 
 
diff --git a/Engine/Plugins/Editor/AssetManagerEditor/Source/AssetManagerEditor/Private/AssetManagerEditorModule.cpp b/Engine/Plugins/Editor/AssetManagerEditor/Source/AssetManagerEditor/Private/AssetManagerEditorModule.cpp
--- a/Engine/Plugins/Editor/AssetManagerEditor/Source/AssetManagerEditor/Private/AssetManagerEditorModule.cpp
+++ b/Engine/Plugins/Editor/AssetManagerEditor/Source/AssetManagerEditor/Private/AssetManagerEditorModule.cpp
@@ -1440,7 +1440,7 @@ void FAssetManagerEditorModule::SetCurrentRegistrySource(const FString& SourceNa
 
 void FAssetManagerEditorModule::RefreshRegistryData()
 {
-	UAssetManager::Get().UpdateManagementDatabase();
+	UAssetManager::Get().UpdateManagementDatabase(true);
 
 	// Rescan registry sources, try to restore the current one
 	FString OldSourceName = CurrentRegistrySource->SourceName;
diff --git a/Engine/Plugins/Runtime/ApexDestruction/Source/ApexDestruction/ApexDestruction.Build.cs b/Engine/Plugins/Runtime/ApexDestruction/Source/ApexDestruction/ApexDestruction.Build.cs
--- a/Engine/Plugins/Runtime/ApexDestruction/Source/ApexDestruction/ApexDestruction.Build.cs
+++ b/Engine/Plugins/Runtime/ApexDestruction/Source/ApexDestruction/ApexDestruction.Build.cs
@@ -12,11 +12,12 @@ public class ApexDestruction : ModuleRules
 				"Core",
 				"CoreUObject", // @todo Mac: for some reason it's needed to link in debug on Mac
 				"Engine",
-                "PhysX",
-                "APEX",
-                "RHI",
-                "RenderCore",
-                "ApexDestructionLib"
+				"PhysX",
+				"APEX",
+				"RHI",
+				"RenderCore",
+				"ApexDestructionLib",
+				"NavigationSystem",
             }
 		);
 
diff --git a/Engine/Plugins/Runtime/ApexDestruction/Source/ApexDestruction/Private/ApexDestructionCustomPayload.cpp b/Engine/Plugins/Runtime/ApexDestruction/Source/ApexDestruction/Private/ApexDestructionCustomPayload.cpp
--- a/Engine/Plugins/Runtime/ApexDestruction/Source/ApexDestruction/Private/ApexDestructionCustomPayload.cpp
+++ b/Engine/Plugins/Runtime/ApexDestruction/Source/ApexDestruction/Private/ApexDestructionCustomPayload.cpp
@@ -4,7 +4,7 @@
 #include "DestructibleComponent.h"
 #include "PhysXPublic.h"
 #include "ApexDestructionModule.h"
-#include "AI/Navigation/NavigationSystem.h"
+#include "AI/NavigationSystemBase.h"
 
 #if WITH_APEX
 
@@ -67,7 +67,7 @@ void FApexDestructionSyncActors::FinalizeSync(int32 SceneType)
 				DestructibleComponent->SyncComponentToRBPhysics();
 			}
 
-			UNavigationSystem::UpdateComponentInNavOctree(*DestructibleComponent);
+			FNavigationSystem::UpdateComponentData(*DestructibleComponent);
 		}
 	}
 
diff --git a/Engine/Plugins/Runtime/GameplayAbilities/Source/GameplayAbilities/Private/AbilitySystemBlueprintLibrary.cpp b/Engine/Plugins/Runtime/GameplayAbilities/Source/GameplayAbilities/Private/AbilitySystemBlueprintLibrary.cpp
--- a/Engine/Plugins/Runtime/GameplayAbilities/Source/GameplayAbilities/Private/AbilitySystemBlueprintLibrary.cpp
+++ b/Engine/Plugins/Runtime/GameplayAbilities/Source/GameplayAbilities/Private/AbilitySystemBlueprintLibrary.cpp
@@ -240,13 +240,15 @@ FGameplayTargetDataFilterHandle UAbilitySystemBlueprintLibrary::MakeFilterHandle
 
 FGameplayEffectSpecHandle UAbilitySystemBlueprintLibrary::MakeSpecHandle(UGameplayEffect* InGameplayEffect, AActor* InInstigator, AActor* InEffectCauser, float InLevel)
 {
-	FGameplayEffectContext* EffectContext = new FGameplayEffectContext(InInstigator, InEffectCauser);
+	FGameplayEffectContext* EffectContext = UAbilitySystemGlobals::Get().AllocGameplayEffectContext();
+	EffectContext->AddInstigator(InInstigator, InEffectCauser);
 	return FGameplayEffectSpecHandle(new FGameplayEffectSpec(InGameplayEffect, FGameplayEffectContextHandle(EffectContext), InLevel));
 }
 
 FGameplayEffectSpecHandle UAbilitySystemBlueprintLibrary::CloneSpecHandle(AActor* InNewInstigator, AActor* InEffectCauser, FGameplayEffectSpecHandle GameplayEffectSpecHandle_Clone)
 {
-	FGameplayEffectContext* EffectContext = new FGameplayEffectContext(InNewInstigator, InEffectCauser);
+	FGameplayEffectContext* EffectContext = UAbilitySystemGlobals::Get().AllocGameplayEffectContext();
+	EffectContext->AddInstigator(InNewInstigator, InEffectCauser);
 
 	return FGameplayEffectSpecHandle(new FGameplayEffectSpec(*GameplayEffectSpecHandle_Clone.Data.Get(), FGameplayEffectContextHandle(EffectContext)));
 }
@@ -647,7 +649,7 @@ bool UAbilitySystemBlueprintLibrary::GetGameplayCueDirection(AActor* TargetActor
 	return false;
 }
 
-bool UAbilitySystemBlueprintLibrary::DoesGameplayCueMeetTagRequirements(FGameplayCueParameters Parameters, FGameplayTagRequirements& SourceTagReqs, FGameplayTagRequirements& TargetTagReqs)
+bool UAbilitySystemBlueprintLibrary::DoesGameplayCueMeetTagRequirements(FGameplayCueParameters Parameters, const FGameplayTagRequirements& SourceTagReqs, const FGameplayTagRequirements& TargetTagReqs)
 {
 	return SourceTagReqs.RequirementsMet(Parameters.AggregatedSourceTags)
 		&& TargetTagReqs.RequirementsMet(Parameters.AggregatedSourceTags);
diff --git a/Engine/Plugins/Runtime/GameplayAbilities/Source/GameplayAbilities/Private/AbilitySystemComponent.cpp b/Engine/Plugins/Runtime/GameplayAbilities/Source/GameplayAbilities/Private/AbilitySystemComponent.cpp
--- a/Engine/Plugins/Runtime/GameplayAbilities/Source/GameplayAbilities/Private/AbilitySystemComponent.cpp
+++ b/Engine/Plugins/Runtime/GameplayAbilities/Source/GameplayAbilities/Private/AbilitySystemComponent.cpp
@@ -525,7 +525,7 @@ void UAbilitySystemComponent::NotifyTagMap_StackCountChange(const FGameplayTagCo
 
 // ------------------------------------------------------------------------
 
-FActiveGameplayEffectHandle UAbilitySystemComponent::ApplyGameplayEffectSpecToTarget(OUT FGameplayEffectSpec &Spec, UAbilitySystemComponent *Target, FPredictionKey PredictionKey)
+FActiveGameplayEffectHandle UAbilitySystemComponent::ApplyGameplayEffectSpecToTarget(const FGameplayEffectSpec &Spec, UAbilitySystemComponent *Target, FPredictionKey PredictionKey)
 {
 	SCOPE_CYCLE_COUNTER(STAT_AbilitySystemComp_ApplyGameplayEffectSpecToTarget);
 
@@ -551,7 +551,7 @@ FActiveGameplayEffectHandle UAbilitySystemComponent::ApplyGameplayEffectSpecToTa
 	return ReturnHandle;
 }
 
-FActiveGameplayEffectHandle UAbilitySystemComponent::ApplyGameplayEffectSpecToSelf(OUT FGameplayEffectSpec &Spec, FPredictionKey PredictionKey)
+FActiveGameplayEffectHandle UAbilitySystemComponent::ApplyGameplayEffectSpecToSelf(const FGameplayEffectSpec &Spec, FPredictionKey PredictionKey)
 {
 	// Scope lock the container after the addition has taken place to prevent the new effect from potentially getting mangled during the remainder
 	// of the add operation
@@ -668,7 +668,7 @@ FActiveGameplayEffectHandle UAbilitySystemComponent::ApplyGameplayEffectSpecToSe
 			{
 				ABILITY_VLOG(OwnerActor, Log, TEXT("Applied %s"), *OurCopyOfSpec->Def->GetFName().ToString());
 
-				for (FGameplayModifierInfo Modifier : Spec.Def->Modifiers)
+				for (const FGameplayModifierInfo& Modifier : Spec.Def->Modifiers)
 				{
 					float Magnitude = 0.f;
 					Modifier.ModifierMagnitude.AttemptCalculateMagnitude(Spec, Magnitude);
@@ -776,7 +776,7 @@ FActiveGameplayEffectHandle UAbilitySystemComponent::ApplyGameplayEffectSpecToSe
 	// Apply Linked effects
 	// todo: this is ignoring the returned handles, should we put them into a TArray and return all of the handles?
 	// ------------------------------------------------------
-	for (const FGameplayEffectSpecHandle TargetSpec: Spec.TargetEffectSpecs)
+	for (const FGameplayEffectSpecHandle& TargetSpec: Spec.TargetEffectSpecs)
 	{
 		if (TargetSpec.IsValid())
 		{
@@ -798,7 +798,7 @@ FActiveGameplayEffectHandle UAbilitySystemComponent::ApplyGameplayEffectSpecToSe
 	return MyHandle;
 }
 
-FActiveGameplayEffectHandle UAbilitySystemComponent::BP_ApplyGameplayEffectSpecToTarget(FGameplayEffectSpecHandle& SpecHandle, UAbilitySystemComponent* Target)
+FActiveGameplayEffectHandle UAbilitySystemComponent::BP_ApplyGameplayEffectSpecToTarget(const FGameplayEffectSpecHandle& SpecHandle, UAbilitySystemComponent* Target)
 {
 	FActiveGameplayEffectHandle ReturnHandle;
 	if (SpecHandle.IsValid() && Target)
@@ -809,7 +809,7 @@ FActiveGameplayEffectHandle UAbilitySystemComponent::BP_ApplyGameplayEffectSpecT
 	return ReturnHandle;
 }
 
-FActiveGameplayEffectHandle UAbilitySystemComponent::BP_ApplyGameplayEffectSpecToSelf(FGameplayEffectSpecHandle& SpecHandle)
+FActiveGameplayEffectHandle UAbilitySystemComponent::BP_ApplyGameplayEffectSpecToSelf(const FGameplayEffectSpecHandle& SpecHandle)
 {
 	FActiveGameplayEffectHandle ReturnHandle;
 	if (SpecHandle.IsValid())
diff --git a/Engine/Plugins/Runtime/GameplayAbilities/Source/GameplayAbilities/Private/AbilitySystemComponent_Abilities.cpp b/Engine/Plugins/Runtime/GameplayAbilities/Source/GameplayAbilities/Private/AbilitySystemComponent_Abilities.cpp
--- a/Engine/Plugins/Runtime/GameplayAbilities/Source/GameplayAbilities/Private/AbilitySystemComponent_Abilities.cpp
+++ b/Engine/Plugins/Runtime/GameplayAbilities/Source/GameplayAbilities/Private/AbilitySystemComponent_Abilities.cpp
@@ -2002,7 +2002,7 @@ void UAbilitySystemComponent::BindToInputComponent(UInputComponent* InputCompone
 	}
 }
 
-void UAbilitySystemComponent::BindAbilityActivationToInputComponent(UInputComponent* InputComponent, FGameplayAbiliyInputBinds BindInfo)
+void UAbilitySystemComponent::BindAbilityActivationToInputComponent(UInputComponent* InputComponent, FGameplayAbilityInputBinds BindInfo)
 {
 	UEnum* EnumBinds = BindInfo.GetBindEnum();
 
@@ -2052,7 +2052,7 @@ void UAbilitySystemComponent::BindAbilityActivationToInputComponent(UInputCompon
 	}
 }
 
-void UAbilitySystemComponent::SetBlockAbilityBindingsArray(FGameplayAbiliyInputBinds BindInfo)
+void UAbilitySystemComponent::SetBlockAbilityBindingsArray(FGameplayAbilityInputBinds BindInfo)
 {
 	UEnum* EnumBinds = BindInfo.GetBindEnum();
 	bIsNetDirty = true;
diff --git a/Engine/Plugins/Runtime/GameplayAbilities/Source/GameplayAbilities/Public/Abilities/GameplayAbilityTypes.h b/Engine/Plugins/Runtime/GameplayAbilities/Source/GameplayAbilities/Public/Abilities/GameplayAbilityTypes.h
--- a/Engine/Plugins/Runtime/GameplayAbilities/Source/GameplayAbilities/Public/Abilities/GameplayAbilityTypes.h
+++ b/Engine/Plugins/Runtime/GameplayAbilities/Source/GameplayAbilities/Public/Abilities/GameplayAbilityTypes.h
@@ -397,9 +397,9 @@ DECLARE_MULTICAST_DELEGATE_OneParam(FGameplayAbilityEndedDelegate, const FAbilit
  *	Structure that tells AbilitySystemComponent what to bind to an InputComponent (see BindAbilityActivationToInputComponent) 
  *	
  */
-struct FGameplayAbiliyInputBinds
+struct FGameplayAbilityInputBinds
 {
-	FGameplayAbiliyInputBinds(FString InConfirmTargetCommand, FString InCancelTargetCommand, FString InEnumName, int32 InConfirmTargetInputID = INDEX_NONE, int32 InCancelTargetInputID = INDEX_NONE)
+	FGameplayAbilityInputBinds(FString InConfirmTargetCommand, FString InCancelTargetCommand, FString InEnumName, int32 InConfirmTargetInputID = INDEX_NONE, int32 InCancelTargetInputID = INDEX_NONE)
 	: ConfirmTargetCommand(InConfirmTargetCommand)
 	, CancelTargetCommand(InCancelTargetCommand)
 	, EnumName(InEnumName)
diff --git a/Engine/Plugins/Runtime/GameplayAbilities/Source/GameplayAbilities/Public/AbilitySystemBlueprintLibrary.h b/Engine/Plugins/Runtime/GameplayAbilities/Source/GameplayAbilities/Public/AbilitySystemBlueprintLibrary.h
--- a/Engine/Plugins/Runtime/GameplayAbilities/Source/GameplayAbilities/Public/AbilitySystemBlueprintLibrary.h
+++ b/Engine/Plugins/Runtime/GameplayAbilities/Source/GameplayAbilities/Public/AbilitySystemBlueprintLibrary.h
@@ -234,7 +234,7 @@ class GAMEPLAYABILITIES_API UAbilitySystemBlueprintLibrary : public UBlueprintFu
 
 	/** Returns true if the aggregated source and target tags from the effect spec meets the tag requirements */
 	UFUNCTION(BlueprintPure, Category = "Ability|GameplayCue")
-	static bool DoesGameplayCueMeetTagRequirements(FGameplayCueParameters Parameters, UPARAM(ref) FGameplayTagRequirements& SourceTagReqs, UPARAM(ref) FGameplayTagRequirements& TargetTagReqs);
+	static bool DoesGameplayCueMeetTagRequirements(FGameplayCueParameters Parameters, const FGameplayTagRequirements& SourceTagReqs, const FGameplayTagRequirements& TargetTagReqs);
 
 	// -------------------------------------------------------------------------------
 	//		GameplayEffectSpec
diff --git a/Engine/Plugins/Runtime/GameplayAbilities/Source/GameplayAbilities/Public/AbilitySystemComponent.h b/Engine/Plugins/Runtime/GameplayAbilities/Source/GameplayAbilities/Public/AbilitySystemComponent.h
--- a/Engine/Plugins/Runtime/GameplayAbilities/Source/GameplayAbilities/Public/AbilitySystemComponent.h
+++ b/Engine/Plugins/Runtime/GameplayAbilities/Source/GameplayAbilities/Public/AbilitySystemComponent.h
@@ -295,14 +295,14 @@ class GAMEPLAYABILITIES_API UAbilitySystemComponent : public UGameplayTasksCompo
 	// --------------------------------------------
 	// Primary outward facing API for other systems:
 	// --------------------------------------------
-	FActiveGameplayEffectHandle ApplyGameplayEffectSpecToTarget(OUT FGameplayEffectSpec& GameplayEffect, UAbilitySystemComponent *Target, FPredictionKey PredictionKey=FPredictionKey());
-	FActiveGameplayEffectHandle ApplyGameplayEffectSpecToSelf(OUT FGameplayEffectSpec& GameplayEffect, FPredictionKey PredictionKey = FPredictionKey());
+	FActiveGameplayEffectHandle ApplyGameplayEffectSpecToTarget(const FGameplayEffectSpec& GameplayEffect, UAbilitySystemComponent *Target, FPredictionKey PredictionKey=FPredictionKey());
+	FActiveGameplayEffectHandle ApplyGameplayEffectSpecToSelf(const FGameplayEffectSpec& GameplayEffect, FPredictionKey PredictionKey = FPredictionKey());
 
 	UFUNCTION(BlueprintCallable, Category = GameplayEffects, meta=(DisplayName = "ApplyGameplayEffectSpecToTarget", ScriptName = "ApplyGameplayEffectSpecToTarget"))
-	FActiveGameplayEffectHandle BP_ApplyGameplayEffectSpecToTarget(UPARAM(ref) FGameplayEffectSpecHandle& SpecHandle, UAbilitySystemComponent* Target);
+	FActiveGameplayEffectHandle BP_ApplyGameplayEffectSpecToTarget(const FGameplayEffectSpecHandle& SpecHandle, UAbilitySystemComponent* Target);
 
 	UFUNCTION(BlueprintCallable, Category = GameplayEffects, meta=(DisplayName = "ApplyGameplayEffectSpecToSelf", ScriptName = "ApplyGameplayEffectSpecToSelf"))
-	FActiveGameplayEffectHandle BP_ApplyGameplayEffectSpecToSelf(UPARAM(ref) FGameplayEffectSpecHandle& SpecHandle);
+	FActiveGameplayEffectHandle BP_ApplyGameplayEffectSpecToSelf(const FGameplayEffectSpecHandle& SpecHandle);
 	
 	/** Gets the FActiveGameplayEffect based on the passed in Handle */
 	const UGameplayEffect* GetGameplayEffectDefForHandle(FActiveGameplayEffectHandle Handle);
@@ -1184,8 +1184,8 @@ public:
 
 	virtual void BindToInputComponent(UInputComponent* InputComponent);
 
-	void SetBlockAbilityBindingsArray(FGameplayAbiliyInputBinds BindInfo);
-	virtual void BindAbilityActivationToInputComponent(UInputComponent* InputComponent, FGameplayAbiliyInputBinds BindInfo);
+	void SetBlockAbilityBindingsArray(FGameplayAbilityInputBinds BindInfo);
+	virtual void BindAbilityActivationToInputComponent(UInputComponent* InputComponent, FGameplayAbilityInputBinds BindInfo);
 
 	virtual void AbilityLocalInputPressed(int32 InputID);
 	virtual void AbilityLocalInputReleased(int32 InputID);
diff --git a/Engine/Plugins/Runtime/GameplayAbilities/Source/GameplayAbilities/Public/AttributeSet.h b/Engine/Plugins/Runtime/GameplayAbilities/Source/GameplayAbilities/Public/AttributeSet.h
--- a/Engine/Plugins/Runtime/GameplayAbilities/Source/GameplayAbilities/Public/AttributeSet.h
+++ b/Engine/Plugins/Runtime/GameplayAbilities/Source/GameplayAbilities/Public/AttributeSet.h
@@ -205,7 +205,7 @@ public:
 	/** This signifies the attribute set can be ID'd by name over the network. */
 	void SetNetAddressable();
 
-	void InitFromMetaDataTable(const UDataTable* DataTable);
+	virtual void InitFromMetaDataTable(const UDataTable* DataTable);
 
 	FORCEINLINE AActor* GetOwningActor() const { return CastChecked<AActor>(GetOuter()); }
 	UAbilitySystemComponent* GetOwningAbilitySystemComponent() const;
diff --git a/Engine/Plugins/Tests/EditorTests/Source/EditorTests/EditorTests.Build.cs b/Engine/Plugins/Tests/EditorTests/Source/EditorTests/EditorTests.Build.cs
--- a/Engine/Plugins/Tests/EditorTests/Source/EditorTests/EditorTests.Build.cs
+++ b/Engine/Plugins/Tests/EditorTests/Source/EditorTests/EditorTests.Build.cs
@@ -43,7 +43,8 @@ public class EditorTests : ModuleRules
 				"ShaderCore",
                 "RenderCore",
 				"Engine",
-				"Slate",
+                "NavigationSystem",
+                "Slate",
 				"SlateCore",
 				"AssetTools",
 				"MainFrame",
diff --git a/Engine/Plugins/Tests/EditorTests/Source/EditorTests/Private/UnrealEd/EditorBuildPromotionTests.cpp b/Engine/Plugins/Tests/EditorTests/Source/EditorTests/Private/UnrealEd/EditorBuildPromotionTests.cpp
--- a/Engine/Plugins/Tests/EditorTests/Source/EditorTests/Private/UnrealEd/EditorBuildPromotionTests.cpp
+++ b/Engine/Plugins/Tests/EditorTests/Source/EditorTests/Private/UnrealEd/EditorBuildPromotionTests.cpp
@@ -23,7 +23,7 @@
 #include "Materials/MaterialInterface.h"
 #include "Engine/SkeletalMesh.h"
 #include "Components/StaticMeshComponent.h"
-#include "AI/Navigation/NavigationSystem.h"
+#include "AI/NavigationSystemBase.h"
 #include "Engine/Texture.h"
 #include "Engine/StaticMesh.h"
 #include "Kismet/KismetSystemLibrary.h"
@@ -2243,11 +2243,10 @@ namespace BuildPromotionTestHelper
 			GUnrealEd->Exec(World, TEXT("MAP REBUILD ALLVISIBLE"));
 			UE_LOG(LogEditorBuildPromotionTests, Display, TEXT("Rebuilt the map"));
 
-			if (World->GetWorldSettings()->bEnableNavigationSystem &&
-				World->GetNavigationSystem())
+			if (World->GetNavigationSystem())
 			{
 				// Invoke navmesh generator
-				World->GetNavigationSystem()->Build();
+				FNavigationSystem::Build(*World);
 				UE_LOG(LogEditorBuildPromotionTests, Display, TEXT("Built navigation"));
 			}
 
diff --git a/Engine/Source/Developer/BlueprintCompilerCppBackend/Private/BlueprintCompilerCppBackendBase.cpp b/Engine/Source/Developer/BlueprintCompilerCppBackend/Private/BlueprintCompilerCppBackendBase.cpp
--- a/Engine/Source/Developer/BlueprintCompilerCppBackend/Private/BlueprintCompilerCppBackendBase.cpp
+++ b/Engine/Source/Developer/BlueprintCompilerCppBackend/Private/BlueprintCompilerCppBackendBase.cpp
@@ -1340,7 +1340,6 @@ void FBlueprintCompilerCppBackendBase::EmitFileBeginning(const FString& CleanNam
 	FIncludeHeaderHelper::EmitIncludeHeader(bIncludeCodeHelpersInHeader ? EmitterContext.Header : EmitterContext.Body, TEXT("GeneratedCodeHelpers"), true);
 	FIncludeHeaderHelper::EmitIncludeHeader(EmitterContext.Header, TEXT("Blueprint/BlueprintSupport"), true);
 
-	FBackendHelperUMG::AdditionalHeaderIncludeForWidget(EmitterContext);
 	FBackendHelperAnim::AddHeaders(EmitterContext);
 
 	TSet<FString> AlreadyIncluded;
diff --git a/Engine/Source/Developer/BlueprintCompilerCppBackend/Private/BlueprintCompilerCppBackendGatherDependencies.cpp b/Engine/Source/Developer/BlueprintCompilerCppBackend/Private/BlueprintCompilerCppBackendGatherDependencies.cpp
--- a/Engine/Source/Developer/BlueprintCompilerCppBackend/Private/BlueprintCompilerCppBackendGatherDependencies.cpp
+++ b/Engine/Source/Developer/BlueprintCompilerCppBackend/Private/BlueprintCompilerCppBackendGatherDependencies.cpp
@@ -163,6 +163,11 @@ struct FFindAssetsToInclude : public FGatherConvertedClassDependenciesHelperBase
 						// This is a field owned by a class, so attempt to add the class as a dependency.
 						MaybeIncludeObjectAsDependency(OwnerClass, CurrentlyConvertedStruct);
 					}
+					else if (ObjAsBPGC && (CurrentlyConvertedStruct->IsA<UUserDefinedStruct>() || ObjAsBPGC->IsChildOf(CurrentlyConvertedStruct)))
+					{
+						// If an unconverted BP (e.g. DOBP) has a potentially circular dependency on the converted asset, wrap it inside a stub class to avoid an EDL cycle on load.
+						Dependencies.MarkUnconvertedClassAsNecessary(ObjAsBPGC);
+					}
 					else
 					{
 						// Add the class itself as a dependency.
@@ -524,6 +529,23 @@ bool FGatherConvertedClassDependencies::WillClassBeConverted(const UBlueprintGen
 	return false;
 }
 
+void FGatherConvertedClassDependencies::MarkUnconvertedClassAsNecessary(const UBlueprintGeneratedClass* InClass) const
+{
+	if (InClass && !InClass->HasAnyFlags(RF_ClassDefaultObject))
+	{
+		const UClass* ClassToInclude = FindOriginalClass(InClass);
+		checkSlow(ClassToInclude);
+
+		IBlueprintCompilerCppBackendModule& BackEndModule = (IBlueprintCompilerCppBackendModule&)IBlueprintCompilerCppBackendModule::Get();
+		const auto& IncludeUnconvertedClassDelegate = BackEndModule.OnIncludingUnconvertedBP();
+
+		if (IncludeUnconvertedClassDelegate.IsBound())
+		{
+			IncludeUnconvertedClassDelegate.Execute(Cast<UBlueprint>(ClassToInclude->ClassGeneratedBy), NativizationOptions);
+		}
+	}
+}
+
 void FGatherConvertedClassDependencies::DependenciesForHeader()
 {
 	TArray<UObject*> ObjectsToCheck;
@@ -715,7 +737,7 @@ TSet<const UObject*> FGatherConvertedClassDependencies::AllDependencies() const
 	return All;
 }
 
-class FArchiveReferencesInStructIntance : public FArchive
+class FArchiveReferencesInStructInstance : public FArchive
 {
 public:
 
@@ -738,7 +760,7 @@ public:
 	}
 	//~ End FArchive Interface
 
-	FArchiveReferencesInStructIntance()
+	FArchiveReferencesInStructInstance()
 	{
 		ArIsObjectReferenceCollector = true;
 		ArIsFilterEditorOnly = true;
@@ -746,14 +768,14 @@ public:
 };
 
 
-void FGatherConvertedClassDependencies::GatherAssetReferencedByUDSDefaultValue(TSet<UObject*>& Dependencies, UUserDefinedStruct* Struct)
+void FGatherConvertedClassDependencies::GatherAssetsReferencedByUDSDefaultValue(TSet<UObject*>& Dependencies, UUserDefinedStruct* Struct)
 {
 	if (Struct)
 	{
 		FStructOnScope StructOnScope(Struct);
 		Struct->InitializeDefaultValue(StructOnScope.GetStructMemory());
-		FArchiveReferencesInStructIntance ArchiveReferencesInStructIntance;
-		Struct->SerializeItem(ArchiveReferencesInStructIntance, StructOnScope.GetStructMemory(), nullptr);
-		Dependencies.Append(ArchiveReferencesInStructIntance.References);
+		FArchiveReferencesInStructInstance ArchiveReferencesInStructInstance;
+		Struct->SerializeItem(ArchiveReferencesInStructInstance, StructOnScope.GetStructMemory(), nullptr);
+		Dependencies.Append(ArchiveReferencesInStructInstance.References);
 	}
-}
+}
\ No newline at end of file
diff --git a/Engine/Source/Developer/BlueprintCompilerCppBackend/Private/BlueprintCompilerCppBackendUMG.cpp b/Engine/Source/Developer/BlueprintCompilerCppBackend/Private/BlueprintCompilerCppBackendUMG.cpp
--- a/Engine/Source/Developer/BlueprintCompilerCppBackend/Private/BlueprintCompilerCppBackendUMG.cpp
+++ b/Engine/Source/Developer/BlueprintCompilerCppBackend/Private/BlueprintCompilerCppBackendUMG.cpp
@@ -23,14 +23,6 @@ void FBackendHelperUMG::WidgetFunctionsInHeader(FEmitterLocalContext& Context)
 	}
 }
 
-void FBackendHelperUMG::AdditionalHeaderIncludeForWidget(FEmitterLocalContext& Context)
-{
-	if (Cast<UWidgetBlueprintGeneratedClass>(Context.GetCurrentlyGeneratedClass()))
-	{
-		Context.Header.AddLine(TEXT("#include \"Runtime/UMG/Public/UMG.h\""));
-	}
-}
-
 void FBackendHelperUMG::CreateClassSubobjects(FEmitterLocalContext& Context, bool bCreate, bool bInitialize)
 {
 	if (auto WidgetClass = Cast<UWidgetBlueprintGeneratedClass>(Context.GetCurrentlyGeneratedClass()))
diff --git a/Engine/Source/Developer/BlueprintCompilerCppBackend/Private/BlueprintCompilerCppBackendUtils.h b/Engine/Source/Developer/BlueprintCompilerCppBackend/Private/BlueprintCompilerCppBackendUtils.h
--- a/Engine/Source/Developer/BlueprintCompilerCppBackend/Private/BlueprintCompilerCppBackendUtils.h
+++ b/Engine/Source/Developer/BlueprintCompilerCppBackend/Private/BlueprintCompilerCppBackendUtils.h
@@ -370,8 +370,6 @@ struct FBackendHelperUMG
 {
 	static void WidgetFunctionsInHeader(FEmitterLocalContext& Context);
 
-	static void AdditionalHeaderIncludeForWidget(FEmitterLocalContext& EmitterContext);
-
 	// these function should use the same context as Constructor
 	static void CreateClassSubobjects(FEmitterLocalContext& Context, bool bCreate, bool bInitialize);
 	static void EmitWidgetInitializationFunctions(FEmitterLocalContext& Context);
diff --git a/Engine/Source/Developer/BlueprintCompilerCppBackend/Private/BlueprintCompilerCppBackendValueHelper.cpp b/Engine/Source/Developer/BlueprintCompilerCppBackend/Private/BlueprintCompilerCppBackendValueHelper.cpp
--- a/Engine/Source/Developer/BlueprintCompilerCppBackend/Private/BlueprintCompilerCppBackendValueHelper.cpp
+++ b/Engine/Source/Developer/BlueprintCompilerCppBackend/Private/BlueprintCompilerCppBackendValueHelper.cpp
@@ -824,8 +824,12 @@ struct FDefaultSubobjectData
 		TArray<FDefaultSubobjectData> NestedSubobjectsToInit;
 		for (UObject* DSO : NestedDefaultSubobjects)
 		{
-			FDefaultSubobjectData* SubobjectData = new(NestedSubobjectsToInit) FDefaultSubobjectData();
-			FEmitDefaultValueHelper::HandleInstancedSubobject(Context, DSO, /* bCreateInstance = */ false, /* bSkipEditorOnlyCheck = */ false, SubobjectData);
+			// We don't need to emit code to initialize nested default subobjects that are also editor-only, since they won't be used in a cooked build.
+			if (!DSO->IsEditorOnly())
+			{
+				FDefaultSubobjectData* SubobjectData = new(NestedSubobjectsToInit) FDefaultSubobjectData();
+				FEmitDefaultValueHelper::HandleInstancedSubobject(Context, DSO, /* bCreateInstance = */ false, /* bSkipEditorOnlyCheck = */ true, SubobjectData);
+			}
 		}
 
 		// Recursively emit code to initialize any nested default subobjects found above that that are now locally referenced within this scope block.
@@ -1189,7 +1193,7 @@ struct FFakeImportTableHelper
 				{
 					if (Subobject)
 					{
-						SerializeBeforeCreateCDODependencies.Add(Subobject->GetClass());
+						SerializeBeforeSerializeStructDependencies.Add(Subobject->GetClass());
 						SerializeBeforeCreateCDODependencies.Add(Subobject->GetClass()->GetDefaultObject());
 					}
 				}
@@ -1241,7 +1245,7 @@ void FEmitDefaultValueHelper::AddStaticFunctionsForDependencies(FEmitterLocalCon
 		TSet<UObject*> References;
 		for (UUserDefinedStruct* UDS : Context.StructsWithDefaultValuesUsed)
 		{
-			FGatherConvertedClassDependencies::GatherAssetReferencedByUDSDefaultValue(References, UDS);
+			FGatherConvertedClassDependencies::GatherAssetsReferencedByUDSDefaultValue(References, UDS);
 		}
 		for (UObject* Obj : References)
 		{
@@ -1773,11 +1777,7 @@ void FEmitDefaultValueHelper::GenerateConstructor(FEmitterLocalContext& Context)
 				if (DSO && DSO->GetClass()->HasAnyClassFlags(CLASS_DefaultToInstanced))
 				{
 					// Determine if this is an editor-only subobject.
-					bool bIsEditorOnlySubobject = false;
-					if (const UActorComponent* ActorComponent = Cast<UActorComponent>(DSO))
-					{
-						bIsEditorOnlySubobject = ActorComponent->IsEditorOnly();
-					}
+					const bool bIsEditorOnlySubobject = DSO->IsEditorOnly();
 
 					// Skip ctor code gen for editor-only subobjects, since they won't be used by the runtime. Any dependencies on editor-only subobjects will be handled later (see HandleInstancedSubobject).
 					if (!bIsEditorOnlySubobject)
diff --git a/Engine/Source/Developer/BlueprintCompilerCppBackend/Public/BlueprintCompilerCppBackendGatherDependencies.h b/Engine/Source/Developer/BlueprintCompilerCppBackend/Public/BlueprintCompilerCppBackendGatherDependencies.h
--- a/Engine/Source/Developer/BlueprintCompilerCppBackend/Public/BlueprintCompilerCppBackendGatherDependencies.h
+++ b/Engine/Source/Developer/BlueprintCompilerCppBackend/Public/BlueprintCompilerCppBackendGatherDependencies.h
@@ -49,8 +49,9 @@ public:
 
 public:
 	bool WillClassBeConverted(const UBlueprintGeneratedClass* InClass) const;
+	void MarkUnconvertedClassAsNecessary(const UBlueprintGeneratedClass* InClass) const;
 
-	static void GatherAssetReferencedByUDSDefaultValue(TSet<UObject*>& Dependencies, UUserDefinedStruct* Struct);
+	static void GatherAssetsReferencedByUDSDefaultValue(TSet<UObject*>& Dependencies, UUserDefinedStruct* Struct);
 
 	static bool IsFieldFromExcludedPackage(const UField* Field, const TSet<FName>& InExcludedModules);
 
diff --git a/Engine/Source/Developer/BlueprintNativeCodeGen/Private/BlueprintNativeCodeGenModule.cpp b/Engine/Source/Developer/BlueprintNativeCodeGen/Private/BlueprintNativeCodeGenModule.cpp
--- a/Engine/Source/Developer/BlueprintNativeCodeGen/Private/BlueprintNativeCodeGenModule.cpp
+++ b/Engine/Source/Developer/BlueprintNativeCodeGen/Private/BlueprintNativeCodeGenModule.cpp
@@ -156,6 +156,9 @@ void FBlueprintNativeCodeGenModule::MarkUnconvertedBlueprintAsNecessary(TSoftObj
 	FStatePerPlatform* StateForCurrentPlatform = StatesPerPlatform.Find(NativizationOptions.PlatformName);
 	if (ensure(StateForCurrentPlatform))
 	{
+		const UBlueprint* BP = BPPtr.Get();
+		UE_LOG(LogBlueprintCodeGen, Log, TEXT("Requiring stub class for unconverted Blueprint asset: %s"), BP ? *BP->GetName() : TEXT(""));
+
 		StateForCurrentPlatform->StubsRequiredByGeneratedCode.Add(BPPtr);
 	}
 }
diff --git a/Engine/Source/Developer/BlueprintNativeCodeGen/Private/BlueprintNativeCodeGenUtils.cpp b/Engine/Source/Developer/BlueprintNativeCodeGen/Private/BlueprintNativeCodeGenUtils.cpp
--- a/Engine/Source/Developer/BlueprintNativeCodeGen/Private/BlueprintNativeCodeGenUtils.cpp
+++ b/Engine/Source/Developer/BlueprintNativeCodeGen/Private/BlueprintNativeCodeGenUtils.cpp
@@ -30,8 +30,6 @@ DEFINE_LOG_CATEGORY(LogBlueprintCodeGen)
 namespace BlueprintNativeCodeGenUtilsImpl
 {
 	static FString CoreModuleName   = TEXT("Core");
-	static FString EngineModuleName = TEXT("Engine");
-	static FString EngineHeaderFile = TEXT("Engine.h");
 
 	// Used to cache the set of plugin dependencies.
 	static TSet<FString> PluginDependencies;
@@ -189,7 +187,6 @@ static bool BlueprintNativeCodeGenUtilsImpl::GenerateModuleSourceFiles(const FBl
 	FText FailureReason;
 
 	TArray<FString> PchIncludes;
-	PchIncludes.Add(EngineHeaderFile);
 	PchIncludes.Add(TEXT("GeneratedCodeHelpers.h"));
 	PchIncludes.Add(NativizedDependenciesFileName() + TEXT(".h"));
 
@@ -259,8 +256,6 @@ static bool BlueprintNativeCodeGenUtilsImpl::GenerateModuleBuildFile(const FBlue
 	TArray<FString> PublicDependencies;
 	// for IModuleInterface
 	PublicDependencies.Add(CoreModuleName);
-	// for Engine.h
-	PublicDependencies.Add(EngineModuleName);
 
 	if (GameProjectUtils::ProjectHasCodeFiles()) 
 	{
diff --git a/Engine/Source/Developer/FunctionalTesting/Classes/FunctionalTest.h b/Engine/Source/Developer/FunctionalTesting/Classes/FunctionalTest.h
--- a/Engine/Source/Developer/FunctionalTesting/Classes/FunctionalTest.h
+++ b/Engine/Source/Developer/FunctionalTesting/Classes/FunctionalTest.h
@@ -341,14 +341,14 @@ public:
 	 * Assert on a relationship between two integers.
 	 * @param What	A name to use in the message if the assert fails (What: expected {Actual} to be <ShouldBe> {Expected} for context '')
 	 */
-	UFUNCTION(BlueprintCallable, Category = "Asserts", DisplayName = "Assert Value (int)", meta = ( HidePin = "ContextObject", DefaultToSelf = "ContextObject"))
+	UFUNCTION(BlueprintCallable, Category = "Asserts", DisplayName = "Assert Value (Integer)", meta = ( HidePin = "ContextObject", DefaultToSelf = "ContextObject"))
 	bool AssertValue_Int(int32 Actual, EComparisonMethod ShouldBe, int32 Expected, const FString& What, const UObject* ContextObject = nullptr);
 
 	/**
 	 * Assert on a relationship between two floats.
 	 * @param What	A name to use in the message if the assert fails (What: expected {Actual} to be <ShouldBe> {Expected} for context '')
 	 */
-	UFUNCTION(BlueprintCallable, Category = "Asserts", DisplayName = "Assert Value (float)", meta = ( HidePin = "ContextObject", DefaultToSelf = "ContextObject"))
+	UFUNCTION(BlueprintCallable, Category = "Asserts", DisplayName = "Assert Value (Float)", meta = ( HidePin = "ContextObject", DefaultToSelf = "ContextObject"))
 	bool AssertValue_Float(float Actual, EComparisonMethod ShouldBe, float Expected, const FString& What, const UObject* ContextObject = nullptr);
 
 	/**
@@ -383,7 +383,7 @@ public:
 	* Assert that two ints are equal
 	* @param What	A name to use in the message if the assert fails (What: expected {Actual} to be Equal To {Expected} for context '')
 	*/
-	UFUNCTION(BlueprintCallable, Category = "Asserts", DisplayName = "Assert Equal (Int)", meta = (HidePin = "ContextObject", DefaultToSelf = "ContextObject"))
+	UFUNCTION(BlueprintCallable, Category = "Asserts", DisplayName = "Assert Equal (Integer)", meta = (HidePin = "ContextObject", DefaultToSelf = "ContextObject"))
 	bool AssertEqual_Int(const int Actual, const int Expected, const FString& What, const UObject* ContextObject = nullptr);
 
 	/**
diff --git a/Engine/Source/Developer/FunctionalTesting/FunctionalTesting.Build.cs b/Engine/Source/Developer/FunctionalTesting/FunctionalTesting.Build.cs
--- a/Engine/Source/Developer/FunctionalTesting/FunctionalTesting.Build.cs
+++ b/Engine/Source/Developer/FunctionalTesting/FunctionalTesting.Build.cs
@@ -14,6 +14,7 @@ public class FunctionalTesting : ModuleRules
                 "ShaderCore",
                 "Slate",
                 "MessageLog",
+                "NavigationSystem",
                 "AIModule",
                 "RenderCore",
                 "AssetRegistry",
diff --git a/Engine/Source/Developer/FunctionalTesting/Private/FunctionalAITest.cpp b/Engine/Source/Developer/FunctionalTesting/Private/FunctionalAITest.cpp
--- a/Engine/Source/Developer/FunctionalTesting/Private/FunctionalAITest.cpp
+++ b/Engine/Source/Developer/FunctionalTesting/Private/FunctionalAITest.cpp
@@ -5,11 +5,11 @@
 #include "Engine/World.h"
 #include "FunctionalTestingModule.h"
 #include "FunctionalTestingManager.h"
-#include "AI/Navigation/NavigationSystem.h"
+#include "NavigationSystem.h"
 #include "AIController.h"
 #include "Blueprint/AIBlueprintHelperLibrary.h"
-#include "AI/Navigation/RecastNavMesh.h"
-#include "AI/NavigationOctree.h"
+#include "NavMesh/RecastNavMesh.h"
+#include "NavigationOctree.h"
 
 AFunctionalAITest::AFunctionalAITest( const FObjectInitializer& ObjectInitializer )
 	: Super(ObjectInitializer)
@@ -99,10 +99,10 @@ void AFunctionalAITest::OnTimeout()
 	// - log pending navmesh rebuilds / dirty areas
 	// - check if area modifiers from navoctree were applied
 
-	UNavigationSystem* NavSys = bDebugNavMeshOnTimeout ? UNavigationSystem::GetCurrent(GetWorld()) : nullptr;
+	UNavigationSystemV1* NavSys = bDebugNavMeshOnTimeout ? FNavigationSystem::GetCurrent<UNavigationSystemV1>(GetWorld()) : nullptr;
 	if (NavSys)
 	{
-		const ARecastNavMesh* Navmesh = NavSys ? Cast<ARecastNavMesh>(NavSys->GetMainNavData()) : nullptr;
+		const ARecastNavMesh* Navmesh = NavSys ? Cast<ARecastNavMesh>(NavSys->GetDefaultNavDataInstance()) : nullptr;
 
 		UE_LOG(LogFunctionalTest, Log, TEXT("Test timed out, log details for: %s"), *GetNameSafe(Navmesh));
 		UE_LOG(LogFunctionalTest, Log, TEXT("> dirty areas? %s"), NavSys->HasDirtyAreasQueued() ? TEXT("YES") : TEXT("no"));
@@ -345,7 +345,7 @@ FVector AFunctionalAITest::GetRandomizedLocation(const FVector& Location) const
 
 bool AFunctionalAITest::IsNavMeshReady() const
 {
-	UNavigationSystem* NavSys = UNavigationSystem::GetCurrent(GetWorld());
+	UNavigationSystemV1* NavSys = FNavigationSystem::GetCurrent<UNavigationSystemV1>(GetWorld());
 	if (NavSys && NavSys->NavDataSet.Num() > 0 && !NavSys->IsNavigationBuildInProgress())
 	{
 		return true;
diff --git a/Engine/Source/Developer/FunctionalTesting/Private/FunctionalTestingManager.cpp b/Engine/Source/Developer/FunctionalTesting/Private/FunctionalTestingManager.cpp
--- a/Engine/Source/Developer/FunctionalTesting/Private/FunctionalTestingManager.cpp
+++ b/Engine/Source/Developer/FunctionalTesting/Private/FunctionalTestingManager.cpp
@@ -5,7 +5,7 @@
 #include "Engine/World.h"
 #include "EngineUtils.h"
 #include "FunctionalTestingModule.h"
-#include "AI/Navigation/NavigationSystem.h"
+#include "NavigationSystem.h"
 #include "EngineGlobals.h"
 #include "Engine/Engine.h"
 #include "Misc/RuntimeErrors.h"
@@ -163,7 +163,7 @@ void UFunctionalTestingManager::TriggerFirstValidTest()
 
 	const bool bIsWorldInitialized =
 		World->AreActorsInitialized() &&
-		!UNavigationSystem::IsNavigationBeingBuilt(World);
+		!UNavigationSystemV1::IsNavigationBeingBuilt(World);
 
 	if (bInitialDelayApplied == true && bIsWorldInitialized)
 	{
diff --git a/Engine/Source/Developer/LogVisualizer/Private/VisualLoggerDatabase.cpp b/Engine/Source/Developer/LogVisualizer/Private/VisualLoggerDatabase.cpp
--- a/Engine/Source/Developer/LogVisualizer/Private/VisualLoggerDatabase.cpp
+++ b/Engine/Source/Developer/LogVisualizer/Private/VisualLoggerDatabase.cpp
@@ -14,6 +14,9 @@ TSharedPtr< struct FVisualLoggerGraphsDatabase > FVisualLoggerGraphsDatabase::St
 void FVisualLoggerDBRow::AddItem(const FVisualLogDevice::FVisualLogEntryItem& NewItem) 
 { 
 	const int32 ItemIndex = Items.Add(NewItem);
+	const int32 HiddenIndex = HiddenItems.Add(false);
+	checkf(ItemIndex == HiddenIndex, TEXT("Items and HiddenItems should have matching indices."));
+
 	DBEvents.OnNewItem.Broadcast(*this, ItemIndex);
 }
 
@@ -35,13 +38,12 @@ const FVisualLogDevice::FVisualLogEntryItem& FVisualLoggerDBRow::GetCurrentItem(
 
 void FVisualLoggerDBRow::SetItemVisibility(int32 ItemIndex, bool IsVisible)
 { 
-	if (IsVisible)
-	{
-		HiddenItems.RemoveSingleSwap(ItemIndex);
-	}
-	else
+	const bool IsHidden = !IsVisible;
+	if (HiddenItems[ItemIndex] != IsHidden)
 	{
-		HiddenItems.AddUnique(ItemIndex);
+		HiddenItems[ItemIndex] = IsHidden;
+		NumHiddenItems += IsHidden ? 1 : -1;
+		check(NumHiddenItems >= 0);
 	}
 }
 
diff --git a/Engine/Source/Developer/LogVisualizer/Private/VisualLoggerDatabase.h b/Engine/Source/Developer/LogVisualizer/Private/VisualLoggerDatabase.h
--- a/Engine/Source/Developer/LogVisualizer/Private/VisualLoggerDatabase.h
+++ b/Engine/Source/Developer/LogVisualizer/Private/VisualLoggerDatabase.h
@@ -48,8 +48,8 @@ public:
 	int32 GetCurrentItemIndex() const { return CurrentItemIndex; }
 
 	void SetItemVisibility(int32 ItemIndex, bool IsVisible);
-	bool IsItemVisible(int32 ItemIndex) const { return HiddenItems.Find(ItemIndex) == INDEX_NONE; }
-	int32 GetNumberOfHiddenItems() const { return HiddenItems.Num(); }
+	bool IsItemVisible(int32 ItemIndex) const { return HiddenItems[ItemIndex] == false; }
+	int32 GetNumberOfHiddenItems() const { return NumHiddenItems; }
 
 	int32 GetClosestItem(float Time) const;
 	int32 GetClosestItem(float Time, float ScrubTime) const;
@@ -61,7 +61,8 @@ protected:
 	int32 CurrentItemIndex;
 
 	TArray<FVisualLogDevice::FVisualLogEntryItem>	Items;
-	TArray<int32> HiddenItems;
+	TBitArray<> HiddenItems;
+	int32 NumHiddenItems = 0;
 };
 
 struct FVisualLoggerDatabase
diff --git a/Engine/Source/Developer/LogVisualizer/Private/VisualLoggerRenderingActor.cpp b/Engine/Source/Developer/LogVisualizer/Private/VisualLoggerRenderingActor.cpp
--- a/Engine/Source/Developer/LogVisualizer/Private/VisualLoggerRenderingActor.cpp
+++ b/Engine/Source/Developer/LogVisualizer/Private/VisualLoggerRenderingActor.cpp
@@ -1,7 +1,7 @@
 // Copyright 1998-2018 Epic Games, Inc. All Rights Reserved.
 
 #include "VisualLoggerRenderingActor.h"
-#include "AI/Navigation/NavigationSystem.h"
+#include "AI/NavigationSystemBase.h"
 #include "VisualLogger/VisualLogger.h"
 #include "LogVisualizerSettings.h"
 #include "VisualLoggerDatabase.h"
@@ -11,6 +11,11 @@
 #endif // WITH_EDITOR
 #include "VisualLoggerRenderingComponent.h"
 
+namespace FDebugDrawing
+{
+	const FVector NavOffset(0, 0, 15);
+}
+
 class UVisualLoggerRenderingComponent;
 class FVisualLoggerSceneProxy final : public FDebugRenderSceneProxy
 {
@@ -396,7 +401,7 @@ void AVisualLoggerRenderingActor::GetDebugShapes(const FVisualLogDevice::FVisual
 		}
 
 
-		const FVector CorridorOffset = NavigationDebugDrawing::PathOffset * 1.25f;
+		const FVector CorridorOffset = FDebugDrawing::NavOffset * 1.25f;
 		const FColor Color = ElementToDraw->GetFColor();
 
 		switch (ElementToDraw->GetType())
diff --git a/Engine/Source/Editor/AnimGraph/Private/AnimGraphNode_SubInstance.cpp b/Engine/Source/Editor/AnimGraph/Private/AnimGraphNode_SubInstance.cpp
--- a/Engine/Source/Editor/AnimGraph/Private/AnimGraphNode_SubInstance.cpp
+++ b/Engine/Source/Editor/AnimGraph/Private/AnimGraphNode_SubInstance.cpp
@@ -207,7 +207,7 @@ void UAnimGraphNode_SubInstance::ReallocatePinsDuringReconstruction(TArray<UEdGr
 	{
 		if(OldPin && !AnimGraphDefaultSchema->IsPosePin(OldPin->PinType))
 		{
-			OldPin->bSavePinIfOrphaned = false;
+			OldPin->SetSavePinIfOrphaned(false);
 		}
 	}
 
diff --git a/Engine/Source/Editor/BlueprintGraph/Classes/K2Node_BaseAsyncTask.h b/Engine/Source/Editor/BlueprintGraph/Classes/K2Node_BaseAsyncTask.h
--- a/Engine/Source/Editor/BlueprintGraph/Classes/K2Node_BaseAsyncTask.h
+++ b/Engine/Source/Editor/BlueprintGraph/Classes/K2Node_BaseAsyncTask.h
@@ -41,6 +41,7 @@ class BLUEPRINTGRAPH_API UK2Node_BaseAsyncTask : public UK2Node
 	virtual FText GetNodeTitle(ENodeTitleType::Type TitleType) const override;
 	virtual bool IsCompatibleWithGraph(const UEdGraph* TargetGraph) const override;
 	virtual void ValidateNodeDuringCompilation(class FCompilerResultsLog& MessageLog) const override;
+	virtual void GetPinHoverText(const UEdGraphPin& Pin, FString& HoverTextOut) const override;
 	// End of UEdGraphNode interface
 
 	// UK2Node interface
@@ -101,4 +102,18 @@ protected:
 	// Pin Redirector support
 	static TMap<FName, FAsyncTaskPinRedirectMapInfo> AsyncTaskPinRedirectMap;
 	static bool bAsyncTaskPinRedirectMapInitialized;
+
+private:
+	/** Invalidates current pin tool tips, so that they will be refreshed before being displayed: */
+	void InvalidatePinTooltips() { bPinTooltipsValid = false; }
+
+	/**
+	* Creates hover text for the specified pin.
+	*
+	* @param   Pin				The pin you want hover text for (should belong to this node)
+	*/
+	void GeneratePinTooltip(UEdGraphPin& Pin) const;
+
+	/** Flag used to track validity of pin tooltips, when tooltips are invalid they will be refreshed before being displayed */
+	mutable bool bPinTooltipsValid;
 };
diff --git a/Engine/Source/Editor/BlueprintGraph/Classes/K2Node_ForEachElementInEnum.h b/Engine/Source/Editor/BlueprintGraph/Classes/K2Node_ForEachElementInEnum.h
--- a/Engine/Source/Editor/BlueprintGraph/Classes/K2Node_ForEachElementInEnum.h
+++ b/Engine/Source/Editor/BlueprintGraph/Classes/K2Node_ForEachElementInEnum.h
@@ -22,6 +22,7 @@ class UK2Node_ForEachElementInEnum : public UK2Node
 
 	BLUEPRINTGRAPH_API static const FName InsideLoopPinName;
 	BLUEPRINTGRAPH_API static const FName EnumOuputPinName;
+	BLUEPRINTGRAPH_API static const FName SkipHiddenPinName;
 
 	//~ Begin UEdGraphNode Interface
 	virtual void AllocateDefaultPins() override;
@@ -36,6 +37,7 @@ class UK2Node_ForEachElementInEnum : public UK2Node
 	virtual void ValidateNodeDuringCompilation(class FCompilerResultsLog& MessageLog) const override;
 	virtual void GetMenuActions(FBlueprintActionDatabaseRegistrar& ActionRegistrar) const override;
 	virtual FText GetMenuCategory() const override;
+	virtual void PostPlacedNewNode() override;
 	//~ End UK2Node Interface
 
 private:
diff --git a/Engine/Source/Editor/BlueprintGraph/Classes/K2Node_Select.h b/Engine/Source/Editor/BlueprintGraph/Classes/K2Node_Select.h
--- a/Engine/Source/Editor/BlueprintGraph/Classes/K2Node_Select.h
+++ b/Engine/Source/Editor/BlueprintGraph/Classes/K2Node_Select.h
@@ -9,16 +9,18 @@
 #include "Textures/SlateIcon.h"
 #include "K2Node.h"
 #include "NodeDependingOnEnumInterface.h"
+#include "K2Node_AddPinInterface.h"
 #include "K2Node_Select.generated.h"
 
 class FBlueprintActionDatabaseRegistrar;
 class UEdGraph;
 
 UCLASS(MinimalAPI, meta=(Keywords = "Ternary Select"))
-class UK2Node_Select : public UK2Node, public INodeDependingOnEnumInterface
+class UK2Node_Select : public UK2Node, public INodeDependingOnEnumInterface, public IK2Node_AddPinInterface
 {
 	GENERATED_UCLASS_BODY()
 
+private:
 	/** The number of selectable options this node currently has */
 	UPROPERTY()
 	int32 NumOptionPins;
@@ -41,7 +43,17 @@ class UK2Node_Select : public UK2Node, public INodeDependingOnEnumInterface
 
 	/** Whether we need to reconstruct the node after the pins have changed */
 	UPROPERTY(Transient)
-	bool bReconstructNode;
+	uint8 bReconstructNode:1;
+
+	uint8 bReconstructForPinTypeChange:1;
+
+	/**
+	* Respond to a change to the PinType for a selected pin
+	* @param Pin	The pin that is changing type
+	*/
+	void OnPinTypeChanged(UEdGraphPin* Pin);
+
+public:
 
 	//~ Begin UEdGraphNode Interface
 	virtual void AllocateDefaultPins() override;
@@ -74,6 +86,11 @@ class UK2Node_Select : public UK2Node, public INodeDependingOnEnumInterface
 	virtual bool ShouldBeReconstructedAfterEnumChanged() const override { return true; }
 	// End of INodeDependingOnEnumInterface
 
+	//~ IK2Node_AddPinInterface
+	virtual bool CanAddPin() const override;
+	virtual void AddInputPin() override;
+	//~ End IK2Node_AddPinInterface
+
 	/** Get the return value pin */
 	BLUEPRINTGRAPH_API UEdGraphPin* GetReturnValuePin() const;
 	/** Get the condition pin */
@@ -86,12 +103,8 @@ class UK2Node_Select : public UK2Node, public INodeDependingOnEnumInterface
 	/** Gets the name and class of the PrintString function */
 	BLUEPRINTGRAPH_API static void GetPrintStringFunction(FName& FunctionName, UClass** FunctionClass);
 
-	/** Adds a new option pin to the node */
-	BLUEPRINTGRAPH_API void AddOptionPinToNode();
 	/** Removes the last option pin from the node */
 	BLUEPRINTGRAPH_API void RemoveOptionPinToNode();
-	/** Return whether an option pin can be added to the node */
-	BLUEPRINTGRAPH_API virtual bool CanAddOptionPinToNode() const;
 	/** Return whether an option pin can be removed to the node */
 	BLUEPRINTGRAPH_API virtual bool CanRemoveOptionPinToNode() const;
 
diff --git a/Engine/Source/Editor/BlueprintGraph/Classes/K2Node_SpawnActorFromClass.h b/Engine/Source/Editor/BlueprintGraph/Classes/K2Node_SpawnActorFromClass.h
--- a/Engine/Source/Editor/BlueprintGraph/Classes/K2Node_SpawnActorFromClass.h
+++ b/Engine/Source/Editor/BlueprintGraph/Classes/K2Node_SpawnActorFromClass.h
@@ -24,7 +24,6 @@ class BLUEPRINTGRAPH_API UK2Node_SpawnActorFromClass : public UK2Node_ConstructO
 	virtual FText GetNodeTitle(ENodeTitleType::Type TitleType) const override;
 	virtual void GetPinHoverText(const UEdGraphPin& Pin, FString& HoverTextOut) const override;
 	virtual void ExpandNode(class FKismetCompilerContext& CompilerContext, UEdGraph* SourceGraph) override;
-	virtual bool HasExternalDependencies(TArray<class UStruct*>* OptionalOutput) const override;
 	virtual FSlateIcon GetIconAndTint(FLinearColor& OutColor) const override;
 	virtual bool IsCompatibleWithGraph(const UEdGraph* TargetGraph) const override;
 	//~ End UEdGraphNode Interface.
diff --git a/Engine/Source/Editor/BlueprintGraph/Private/CallFunctionHandler.cpp b/Engine/Source/Editor/BlueprintGraph/Private/CallFunctionHandler.cpp
--- a/Engine/Source/Editor/BlueprintGraph/Private/CallFunctionHandler.cpp
+++ b/Engine/Source/Editor/BlueprintGraph/Private/CallFunctionHandler.cpp
@@ -514,7 +514,7 @@ void FKCHandler_CallFunction::RegisterNets(FKismetFunctionContext& Context, UEdG
 		if (Function->HasMetaData(FBlueprintMetadata::MD_WorldContext))
 		{
 			UEdGraphSchema_K2 const* K2Schema = CompilerContext.GetSchema();
-			const bool bHasIntrinsicWorldContext = !K2Schema->IsStaticFunctionGraph(Context.SourceGraph) && FBlueprintEditorUtils::ImplentsGetWorld(Context.Blueprint);
+			const bool bHasIntrinsicWorldContext = !K2Schema->IsStaticFunctionGraph(Context.SourceGraph) && FBlueprintEditorUtils::ImplementsGetWorld(Context.Blueprint);
 
 			const FName WorldContextPinName = *Function->GetMetaData(FBlueprintMetadata::MD_WorldContext);
 
diff --git a/Engine/Source/Editor/BlueprintGraph/Private/EdGraphSchema_K2.cpp b/Engine/Source/Editor/BlueprintGraph/Private/EdGraphSchema_K2.cpp
--- a/Engine/Source/Editor/BlueprintGraph/Private/EdGraphSchema_K2.cpp
+++ b/Engine/Source/Editor/BlueprintGraph/Private/EdGraphSchema_K2.cpp
@@ -816,7 +816,7 @@ bool UEdGraphSchema_K2::CanFunctionBeUsedInGraph(const UClass* InClass, const UF
 					UBlueprint* BP = FBlueprintEditorUtils::FindBlueprintForGraph(InDestGraph);
 					const bool bIsFunctLib = BP && (EBlueprintType::BPTYPE_FunctionLibrary == BP->BlueprintType);
 					UClass* ParentClass = BP ? BP->ParentClass : NULL;
-					const bool bIncompatibleParent = ParentClass && (!FBlueprintEditorUtils::ImplentsGetWorld(BP) && !ParentClass->HasMetaDataHierarchical(FBlueprintMetadata::MD_ShowWorldContextPin));
+					const bool bIncompatibleParent = ParentClass && (!FBlueprintEditorUtils::ImplementsGetWorld(BP) && !ParentClass->HasMetaDataHierarchical(FBlueprintMetadata::MD_ShowWorldContextPin));
 					if (!bIsFunctLib && bIncompatibleParent)
 					{
 						if (OutReason != nullptr)
@@ -3633,11 +3633,15 @@ FText UEdGraphSchema_K2::TerminalTypeToText(const FName Category, const FName Su
 		}
 		else
 		{
-			FString SubCategoryObjName = SubCategoryObject->GetName();
+			FString SubCategoryObjName;
 			if (UField* SubCategoryField = Cast<UField>(SubCategoryObject))
 			{
 				SubCategoryObjName = SubCategoryField->GetDisplayNameText().ToString();
 			}
+			else
+			{
+				SubCategoryObjName = SubCategoryObject->GetName();
+			}
 
 			if (!bIsWeakPtr)
 			{
@@ -3652,6 +3656,12 @@ FText UEdGraphSchema_K2::TerminalTypeToText(const FName Category, const FName Su
 				{
 					PropertyText = FText::Format(LOCTEXT("ObjectAsTextWithoutCategory", "{ObjectName}"), Args);
 				}
+				// If this is a raw UObject reference don't display Object twice
+				else if (((Category == UEdGraphSchema_K2::PC_Object) || (Category == UEdGraphSchema_K2::PC_SoftObject)) && (SubCategoryObject == UObject::StaticClass()))
+				{
+					Args.Add(TEXT("Category"), UEdGraphSchema_K2::GetCategoryText(Category));
+					PropertyText = FText::Format(LOCTEXT("ObjectAsJustCategory", "{Category}"), Args);
+				}
 				else
 				{
 					Args.Add(TEXT("Category"), (!bIsInterface ? UEdGraphSchema_K2::GetCategoryText(Category) : UEdGraphSchema_K2::GetCategoryText(PC_Interface)));
@@ -3667,7 +3677,7 @@ FText UEdGraphSchema_K2::TerminalTypeToText(const FName Category, const FName Su
 			}
 		}
 	}
-	else if (SubCategory != TEXT(""))
+	else if (!SubCategory.IsNone())
 	{
 		FFormatNamedArguments Args;
 		Args.Add(TEXT("Category"), UEdGraphSchema_K2::GetCategoryText(Category));
@@ -4452,7 +4462,7 @@ void UEdGraphSchema_K2::HandleGraphBeingDeleted(UEdGraph& GraphBeingRemoved) con
 		}
 
 		// Remove from the list of recently edited documents
-		Blueprint->LastEditedDocuments.RemoveAll([&GraphBeingRemoved](const FEditedDocumentInfo& TestDoc) { return TestDoc.EditedObject == &GraphBeingRemoved; });
+		Blueprint->LastEditedDocuments.RemoveAll([&GraphBeingRemoved](const FEditedDocumentInfo& TestDoc) { return TestDoc.EditedObjectPath.ResolveObject() == &GraphBeingRemoved; });
 	}
 }
 
diff --git a/Engine/Source/Editor/BlueprintGraph/Private/K2Node.cpp b/Engine/Source/Editor/BlueprintGraph/Private/K2Node.cpp
--- a/Engine/Source/Editor/BlueprintGraph/Private/K2Node.cpp
+++ b/Engine/Source/Editor/BlueprintGraph/Private/K2Node.cpp
@@ -915,7 +915,7 @@ void UK2Node::RewireOldPinsToNewPins(TArray<UEdGraphPin*>& InOldPins, TArray<UEd
 		// * The pin has been flagged not be saved if orphaned
 		// * The pin is hidden and not a split pin
 		const bool bVisibleOrSplitPin = (!OldPin->bHidden || (OldPin->SubPins.Num() > 0));
-		if (UEdGraphPin::AreOrphanPinsEnabled() && !bMatched && bVisibleOrSplitPin && OldPin->bSavePinIfOrphaned)
+		if (UEdGraphPin::AreOrphanPinsEnabled() && !bMatched && bVisibleOrSplitPin && OldPin->ShouldSavePinIfOrphaned())
 		{
 			// The node can specify to save no pins, all pins, or all but exec pins. However, even if all is specified Execute and Then are never saved
 			 const bool bSaveOrphanedPin = ((OrphanedPinSaveMode == ESaveOrphanPinMode::SaveAll) ||
@@ -1411,7 +1411,7 @@ void FOptionalPinManager::EvaluateOldShownPins(const TArray<FOptionalPinFromProp
 		{
 			if (UEdGraphPin* Pin = Node->FindPin(ShowPinForProperty.PropertyName))
 			{
-				Pin->bSavePinIfOrphaned = false;
+				Pin->SetSavePinIfOrphaned(false);
 			}
 		}
 	}
diff --git a/Engine/Source/Editor/BlueprintGraph/Private/K2Node_BaseAsyncTask.cpp b/Engine/Source/Editor/BlueprintGraph/Private/K2Node_BaseAsyncTask.cpp
--- a/Engine/Source/Editor/BlueprintGraph/Private/K2Node_BaseAsyncTask.cpp
+++ b/Engine/Source/Editor/BlueprintGraph/Private/K2Node_BaseAsyncTask.cpp
@@ -30,6 +30,7 @@ UK2Node_BaseAsyncTask::UK2Node_BaseAsyncTask(const FObjectInitializer& ObjectIni
 	, ProxyFactoryClass(nullptr)
 	, ProxyClass(nullptr)
 	, ProxyActivateFunctionName(NAME_None)
+	, bPinTooltipsValid(false)
 {
 }
 
@@ -63,6 +64,8 @@ bool UK2Node_BaseAsyncTask::IsCompatibleWithGraph(const UEdGraph* TargetGraph) c
 
 void UK2Node_BaseAsyncTask::AllocateDefaultPins()
 {
+	InvalidatePinTooltips();
+
 	const UEdGraphSchema_K2* K2Schema = GetDefault<UEdGraphSchema_K2>();
 
 	CreatePin(EGPD_Input, UEdGraphSchema_K2::PC_Exec, UEdGraphSchema_K2::PN_Execute);
@@ -616,4 +619,47 @@ UK2Node::ERedirectType UK2Node_BaseAsyncTask::DoPinsMatchForReconstruction(const
 	return Super::DoPinsMatchForReconstruction(NewPin, NewPinIndex, OldPin, OldPinIndex);
 }
 
+void UK2Node_BaseAsyncTask::GeneratePinTooltip(UEdGraphPin& Pin) const
+{
+	ensure(Pin.GetOwningNode() == this);
+
+	UEdGraphSchema const* Schema = GetSchema();
+	check(Schema);
+	UEdGraphSchema_K2 const* const K2Schema = Cast<const UEdGraphSchema_K2>(Schema);
+
+	if (K2Schema == nullptr)
+	{
+		Schema->ConstructBasicPinTooltip(Pin, FText::GetEmpty(), Pin.PinToolTip);
+		return;
+	}
+
+	// get the class function object associated with this node
+	// Slight change from UK2Node_CallFunction (where this code is copied from)
+	// We're getting the Factory function instead of GetTargetFunction
+	UFunction* Function = GetFactoryFunction(); 
+	if (Function == nullptr)
+	{
+		Schema->ConstructBasicPinTooltip(Pin, FText::GetEmpty(), Pin.PinToolTip);
+		return;
+	}
+
+	UK2Node_CallFunction::GeneratePinTooltipFromFunction(Pin, Function);
+}
+
+void UK2Node_BaseAsyncTask::GetPinHoverText(const UEdGraphPin& Pin, FString& HoverTextOut) const
+{
+	if (!bPinTooltipsValid)
+	{
+		for (UEdGraphPin* P : Pins)
+		{
+			P->PinToolTip.Reset();
+			GeneratePinTooltip(*P);
+		}
+
+		bPinTooltipsValid = true;
+	}
+
+	return UK2Node::GetPinHoverText(Pin, HoverTextOut);
+}
+
 #undef LOCTEXT_NAMESPACE
diff --git a/Engine/Source/Editor/BlueprintGraph/Private/K2Node_CallFunction.cpp b/Engine/Source/Editor/BlueprintGraph/Private/K2Node_CallFunction.cpp
--- a/Engine/Source/Editor/BlueprintGraph/Private/K2Node_CallFunction.cpp
+++ b/Engine/Source/Editor/BlueprintGraph/Private/K2Node_CallFunction.cpp
@@ -1861,7 +1861,7 @@ void UK2Node_CallFunction::ValidateNodeDuringCompilation(class FCompilerResultsL
 			check(Blueprint);
 			UClass* ParentClass = Blueprint->ParentClass;
 			check(ParentClass);
-			if (ParentClass && !FBlueprintEditorUtils::ImplentsGetWorld(Blueprint) && !ParentClass->HasMetaDataHierarchical(FBlueprintMetadata::MD_ShowWorldContextPin))
+			if (ParentClass && !FBlueprintEditorUtils::ImplementsGetWorld(Blueprint) && !ParentClass->HasMetaDataHierarchical(FBlueprintMetadata::MD_ShowWorldContextPin))
 			{
 				MessageLog.Warning(*LOCTEXT("FunctionUnsafeInContext", "Function '@@' is unsafe to call from blueprints of class '@@'.").ToString(), this, ParentClass);
 			}
@@ -2336,7 +2336,7 @@ bool UK2Node_CallFunction::ReconnectPureExecPins(TArray<UEdGraphPin*>& OldPins)
 		}
 		if (PinExec)
 		{
-			PinExec->bSavePinIfOrphaned = false; 
+			PinExec->SetSavePinIfOrphaned(false); 
 
 			// look for old then pin
 			UEdGraphPin* PinThen = nullptr;
@@ -2350,7 +2350,7 @@ bool UK2Node_CallFunction::ReconnectPureExecPins(TArray<UEdGraphPin*>& OldPins)
 			}
 			if (PinThen)
 			{
-				PinThen->bSavePinIfOrphaned = false;
+				PinThen->SetSavePinIfOrphaned(false);
 
 				// reconnect all incoming links to old exec pin to the far end of the old then pin.
 				if (PinThen->LinkedTo.Num() > 0)
diff --git a/Engine/Source/Editor/BlueprintGraph/Private/K2Node_ForEachElementInEnum.cpp b/Engine/Source/Editor/BlueprintGraph/Private/K2Node_ForEachElementInEnum.cpp
--- a/Engine/Source/Editor/BlueprintGraph/Private/K2Node_ForEachElementInEnum.cpp
+++ b/Engine/Source/Editor/BlueprintGraph/Private/K2Node_ForEachElementInEnum.cpp
@@ -7,6 +7,7 @@
 #include "K2Node_AssignmentStatement.h"
 #include "K2Node_ExecutionSequence.h"
 #include "K2Node_IfThenElse.h"
+#include "K2Node_SwitchEnum.h"
 #include "K2Node_TemporaryVariable.h"
 #include "KismetCompiler.h"
 #include "Kismet/KismetNodeHelperLibrary.h"
@@ -138,13 +139,25 @@ UK2Node_ForEachElementInEnum::UK2Node_ForEachElementInEnum(const FObjectInitiali
 
 const FName UK2Node_ForEachElementInEnum::InsideLoopPinName(TEXT("LoopBody"));
 const FName UK2Node_ForEachElementInEnum::EnumOuputPinName(TEXT("EnumValue"));
+const FName UK2Node_ForEachElementInEnum::SkipHiddenPinName(TEXT("SkipHidden"));
 
 void UK2Node_ForEachElementInEnum::AllocateDefaultPins()
 {
+	UEdGraphSchema_K2 const* K2Schema = GetDefault<UEdGraphSchema_K2>();
+
 	CreatePin(EGPD_Input, UEdGraphSchema_K2::PC_Exec, UEdGraphSchema_K2::PN_Execute);
 
 	if (Enum)
 	{
+		if (UEdGraphPin* SkipHiddenPin = CreatePin(EGPD_Input, UEdGraphSchema_K2::PC_Boolean, SkipHiddenPinName))
+		{
+			// This is a non-standard option that likely won't need to be utilized much, so we make it advanced.
+			SkipHiddenPin->bAdvancedView = true;
+			AdvancedPinDisplay = ENodeAdvancedPins::Hidden;
+
+			K2Schema->ConstructBasicPinTooltip(*SkipHiddenPin, LOCTEXT("SkipHiddenPinToolTip", "Controls whether or not the loop will skip over hidden enumeration values."), SkipHiddenPin->PinToolTip);
+		}
+
 		CreatePin(EGPD_Output, UEdGraphSchema_K2::PC_Exec, InsideLoopPinName);
 		CreatePin(EGPD_Output, UEdGraphSchema_K2::PC_Byte, Enum, EnumOuputPinName);
 	}
@@ -210,10 +223,6 @@ void UK2Node_ForEachElementInEnum::ExpandNode(class FKismetCompilerContext& Comp
 
 	const UEdGraphSchema_K2* Schema = CompilerContext.GetSchema();
 
-	CompilerContext.MovePinLinksToIntermediate(*GetExecPin(), *ForLoop.StartLoopExecInPin);
-	CompilerContext.MovePinLinksToIntermediate(*FindPinChecked(UEdGraphSchema_K2::PN_Then), *ForLoop.LoopCompleteOutExecPin);
-	CompilerContext.MovePinLinksToIntermediate(*FindPinChecked(InsideLoopPinName), *ForLoop.InsideLoopExecOutPin);
-
 	UK2Node_GetNumEnumEntries* GetNumEnumEntries = CompilerContext.SpawnIntermediateNode<UK2Node_GetNumEnumEntries>(this, SourceGraph);
 	GetNumEnumEntries->Enum = Enum;
 	GetNumEnumEntries->AllocateDefaultPins();
@@ -230,6 +239,64 @@ void UK2Node_ForEachElementInEnum::ExpandNode(class FKismetCompilerContext& Comp
 	CastByteToEnum->bSafe = true;
 	CastByteToEnum->AllocateDefaultPins();
 	bResult &= Schema->TryCreateConnection(Conv_Func->FindPinChecked(UEdGraphSchema_K2::PN_ReturnValue), CastByteToEnum->FindPinChecked(UK2Node_CastByteToEnum::ByteInputPinName));
+
+	// Additional expansion logic to optionally exclude hidden values during runtime loop iteration
+	UK2Node_ExecutionSequence* SwitchOutputSequence = nullptr;
+	if (const UEdGraphPin* SkipHiddenValuesPin = FindPin(SkipHiddenPinName))
+	{
+		// Process only if the enum type contains at least one hidden value
+		int32 EnumIndex = 0;
+		bool bHasHiddenValues = false;
+		while (!bHasHiddenValues && EnumIndex < Enum->NumEnums() - 1)
+		{
+			bHasHiddenValues = Enum->HasMetaData(TEXT("Hidden"), EnumIndex) || Enum->HasMetaData(TEXT("Spacer"), EnumIndex++);
+		}
+
+		if (bHasHiddenValues)
+		{
+			// Skip hidden values branch (only included if something is linked to the "skip hidden" input pin)
+			UK2Node_IfThenElse* ShouldSkipHiddenBranch = nullptr;
+			if (SkipHiddenValuesPin->LinkedTo.Num() > 0)
+			{
+				bResult &= ensure(SkipHiddenValuesPin->LinkedTo.Num() == 1);
+				ShouldSkipHiddenBranch = CompilerContext.SpawnIntermediateNode<UK2Node_IfThenElse>(this, SourceGraph);
+				ShouldSkipHiddenBranch->AllocateDefaultPins();
+				bResult &= Schema->TryCreateConnection(ForLoop.InsideLoopExecOutPin, ShouldSkipHiddenBranch->GetExecPin());
+				bResult &= Schema->TryCreateConnection(SkipHiddenValuesPin->LinkedTo[0], ShouldSkipHiddenBranch->GetConditionPin());
+			}
+
+			// Enum switch node (only if we included a "should skip" test or if the "skip hidden" input pin default value is 'true')
+			if (ShouldSkipHiddenBranch || SkipHiddenValuesPin->GetDefaultAsString().Equals(TEXT("true"), ESearchCase::IgnoreCase))
+			{
+				// The switch node will internally exclude any hidden enum values when constructed
+				UK2Node_SwitchEnum* SwitchEnum = CompilerContext.SpawnIntermediateNode<UK2Node_SwitchEnum>(this, SourceGraph);
+				SwitchEnum->SetEnum(Enum);
+				SwitchEnum->bHasDefaultPin = false;
+				SwitchEnum->AllocateDefaultPins();
+				bResult &= Schema->TryCreateConnection(SwitchEnum->GetSelectionPin(), CastByteToEnum->FindPinChecked(UEdGraphSchema_K2::PN_ReturnValue));
+				bResult &= Schema->TryCreateConnection(SwitchEnum->GetExecPin(), ShouldSkipHiddenBranch ? ShouldSkipHiddenBranch->GetThenPin() : ForLoop.InsideLoopExecOutPin);
+
+				// Switch output execution sequence (direct all relevant output pins back to a single execution path)
+				SwitchOutputSequence = CompilerContext.SpawnIntermediateNode<UK2Node_ExecutionSequence>(this, SourceGraph);
+				SwitchOutputSequence->AllocateDefaultPins();
+				if (ShouldSkipHiddenBranch)
+				{
+					bResult &= Schema->TryCreateConnection(ShouldSkipHiddenBranch->GetElsePin(), SwitchOutputSequence->GetExecPin());
+				}
+				for (int32 SwitchCasePinIndex = 0; SwitchCasePinIndex < SwitchEnum->EnumEntries.Num() && bResult; ++SwitchCasePinIndex)
+				{
+					if (UEdGraphPin* SwitchCasePin = SwitchEnum->FindPin(SwitchEnum->EnumEntries[SwitchCasePinIndex]))
+					{
+						bResult &= Schema->TryCreateConnection(SwitchCasePin, SwitchOutputSequence->GetExecPin());
+					}
+				}
+			}
+		}
+	}
+
+	CompilerContext.MovePinLinksToIntermediate(*GetExecPin(), *ForLoop.StartLoopExecInPin);
+	CompilerContext.MovePinLinksToIntermediate(*FindPinChecked(UEdGraphSchema_K2::PN_Then), *ForLoop.LoopCompleteOutExecPin);
+	CompilerContext.MovePinLinksToIntermediate(*FindPinChecked(InsideLoopPinName), SwitchOutputSequence ? *SwitchOutputSequence->GetThenPinGivenIndex(0) : *ForLoop.InsideLoopExecOutPin);
 	CompilerContext.MovePinLinksToIntermediate(*FindPinChecked(EnumOuputPinName), *CastByteToEnum->FindPinChecked(UEdGraphSchema_K2::PN_ReturnValue));
 
 	if (!bResult)
@@ -268,4 +335,16 @@ FText UK2Node_ForEachElementInEnum::GetMenuCategory() const
 	return FEditorCategoryUtils::GetCommonCategory(FCommonEditorCategory::Enum);
 }
 
+void UK2Node_ForEachElementInEnum::PostPlacedNewNode()
+{
+	Super::PostPlacedNewNode();
+
+	// Skip hidden enumeration values by default for new node placements.
+	if (UEdGraphPin* SkipHiddenPin = FindPin(SkipHiddenPinName))
+	{
+		const UEdGraphSchema_K2* K2Schema = GetDefault<UEdGraphSchema_K2>();
+		K2Schema->SetPinAutogeneratedDefaultValue(SkipHiddenPin, TEXT("true"));
+	}
+}
+
 #undef LOCTEXT_NAMESPACE
diff --git a/Engine/Source/Editor/BlueprintGraph/Private/K2Node_GetClassDefaults.cpp b/Engine/Source/Editor/BlueprintGraph/Private/K2Node_GetClassDefaults.cpp
--- a/Engine/Source/Editor/BlueprintGraph/Private/K2Node_GetClassDefaults.cpp
+++ b/Engine/Source/Editor/BlueprintGraph/Private/K2Node_GetClassDefaults.cpp
@@ -538,29 +538,31 @@ void UK2Node_GetClassDefaults::CreateOutputPins(UClass* InClass)
 
 void UK2Node_GetClassDefaults::OnClassPinChanged()
 {
-	TArray<UEdGraphPin*> OldPins = Pins;
-	TArray<UEdGraphPin*> OldOutputPins;
+	Modify();
+
+	TArray<UEdGraphPin*> OldOutputPins = MoveTemp(Pins);
 
 	// Gather all current output pins
-	for(int32 PinIndex = 0; PinIndex < OldPins.Num(); ++PinIndex)
+	for(int32 PinIndex = 0; PinIndex < OldOutputPins.Num(); ++PinIndex)
 	{
-		UEdGraphPin* OldPin = OldPins[PinIndex];
-		if(OldPin->Direction == EGPD_Output)
+		UEdGraphPin* OldPin = OldOutputPins[PinIndex];
+		if(OldPin->Direction == EGPD_Input)
 		{
-			Pins.Remove(OldPin);
-			OldOutputPins.Add(OldPin);
+			OldOutputPins.RemoveAt(PinIndex--, 1, false);
+			Pins.Add(OldPin);
 		}
 	}
 
 	// Clear the current output pin settings (so they don't carry over to the new set)
-	ShowPinForProperties.Empty();
+	ShowPinForProperties.Reset();
 
 	// Create output pins for the new class type
 	UClass* InputClass = GetInputClass();
 	CreateOutputPins(InputClass);
 
-	// Destroy the previous set of output pins
-	DestroyPinList(OldOutputPins);
+	// Restore the split pins and reconnect output pins to previously linked inputs
+	RestoreSplitPins(OldOutputPins);
+	RewireOldPinsToNewPins(OldOutputPins, Pins);
 
 	// Notify the graph that the node has been changed
 	if(UEdGraph* Graph = GetGraph())
diff --git a/Engine/Source/Editor/BlueprintGraph/Private/K2Node_Select.cpp b/Engine/Source/Editor/BlueprintGraph/Private/K2Node_Select.cpp
--- a/Engine/Source/Editor/BlueprintGraph/Private/K2Node_Select.cpp
+++ b/Engine/Source/Editor/BlueprintGraph/Private/K2Node_Select.cpp
@@ -19,16 +19,24 @@
 
 #define LOCTEXT_NAMESPACE "K2Node_Select"
 
+namespace
+{
+	FName NAME_bPickOption0(TEXT("bPickOption0"));
+	FName NAME_Index(TEXT("Index"));
+	FName NAME_Option_0(TEXT("Option 0"));
+	FName NAME_Option_1(TEXT("Option 1"));
+}
+
 //////////////////////////////////////////////////////////////////////////
 // FKCHandler_Select
 
-class FKCHandler_SelectRef : public FNodeHandlingFunctor
+class FKCHandler_Select : public FNodeHandlingFunctor
 {
 protected:
 	TMap<UEdGraphNode*, FBPTerminal*> DefaultTermMap;
 
 public:
-	FKCHandler_SelectRef(FKismetCompilerContext& InCompilerContext)
+	FKCHandler_Select(FKismetCompilerContext& InCompilerContext)
 		: FNodeHandlingFunctor(InCompilerContext)
 	{
 	}
@@ -143,190 +151,6 @@ public:
 	}
 };
 
-class FKCHandler_Select : public FNodeHandlingFunctor
-{
-protected:
-	TMap<UEdGraphNode*, FBPTerminal*> BoolTermMap;
-
-public:
-	FKCHandler_Select(FKismetCompilerContext& InCompilerContext)
-		: FNodeHandlingFunctor(InCompilerContext)
-	{
-	}
-
-	virtual void RegisterNets(FKismetFunctionContext& Context, UEdGraphNode* Node) override
-	{
-		FNodeHandlingFunctor::RegisterNets(Context, Node);
-
-		// Create the net for the return value manually as it's a special case Output Direction pin
-		UK2Node_Select* SelectNode = Cast<UK2Node_Select>(Node);
-		UEdGraphPin* ReturnPin = SelectNode->GetReturnValuePin();
-
-		FBPTerminal* Term = Context.CreateLocalTerminalFromPinAutoChooseScope(ReturnPin, Context.NetNameMap->MakeValidName(ReturnPin));
-		Context.NetMap.Add(SelectNode->GetReturnValuePin(), Term);
-
-		// Create a term to determine if the compare was successful or not
-		FBPTerminal* BoolTerm = Context.CreateLocalTerminal();
-		BoolTerm->Type.PinCategory = UEdGraphSchema_K2::PC_Boolean;
-		BoolTerm->Source = Node;
-		BoolTerm->Name = Context.NetNameMap->MakeValidName(Node) + TEXT("_CmpSuccess");
-		BoolTermMap.Add(Node, BoolTerm);
-	}
-
-	virtual void Compile(FKismetFunctionContext& Context, UEdGraphNode* Node) override
-	{
-		// Cast the node and get all the input pins
-		UK2Node_Select* SelectNode = Cast<UK2Node_Select>(Node);
-		TArray<UEdGraphPin*> OptionPins;
-		SelectNode->GetOptionPins(OptionPins);
-		UEdGraphPin* IndexPin = SelectNode->GetIndexPin();
-
-		// Get the kismet term for the (Condition or Index) that will determine which option to use
-		UEdGraphPin* PinToTry = FEdGraphUtilities::GetNetFromPin(IndexPin);
-		FBPTerminal** ConditionTerm = Context.NetMap.Find(PinToTry);
-
-		// Get the kismet term for the return value
-		UEdGraphPin* ReturnPin = SelectNode->GetReturnValuePin();
-		FBPTerminal** ReturnTerm = Context.NetMap.Find(ReturnPin);
-
-		// Don't proceed if there is no return value or there is no selection
-		if (ConditionTerm != NULL && ReturnTerm != NULL)
-		{
-			FName ConditionalFunctionName = "";
-			UClass* ConditionalFunctionClass = NULL;
-			SelectNode->GetConditionalFunction(ConditionalFunctionName, &ConditionalFunctionClass);
-			UFunction* ConditionFunction = FindField<UFunction>(ConditionalFunctionClass, ConditionalFunctionName);
-
-			// Find the local boolean for use in the equality call function below (BoolTerm = result of EqualEqual_IntInt or NotEqual_BoolBool)
-			FBPTerminal* BoolTerm = BoolTermMap.FindRef(SelectNode);
-
-			// We need to keep a pointer to the previous IfNot statement so it can be linked to the next conditional statement
-			FBlueprintCompiledStatement* PrevIfNotStatement = NULL;
-
-			// Keep an array of all the unconditional goto statements so we can clean up their jumps after the noop statement is created
-			TArray<FBlueprintCompiledStatement*> GotoStatementList;
-
-			// Loop through all the options
-			for (int32 OptionIdx = 0; OptionIdx < OptionPins.Num(); OptionIdx++)
-			{
-				// Create a CallFunction statement with the condition function from the Select class
-				FBlueprintCompiledStatement& Statement = Context.AppendStatementForNode(Node);
-				Statement.Type = KCST_CallFunction;
-				Statement.FunctionToCall = ConditionFunction;
-				Statement.FunctionContext = NULL;
-				Statement.bIsParentContext = false;
-				// BoolTerm will be the return value of the condition statement
-				Statement.LHS = BoolTerm;
-				// The condition passed into the Select node
-				Statement.RHS.Add(*ConditionTerm);
-				// Create a local int for use in the equality call function below (LiteralTerm = the right hand side of the EqualEqual_IntInt or NotEqual_BoolBool statement)
-				FBPTerminal* LiteralTerm = Context.CreateLocalTerminal(ETerminalSpecification::TS_Literal);
-				LiteralTerm->bIsLiteral = true;
-				LiteralTerm->Type.PinCategory = UEdGraphSchema_K2::PC_Int;
-
-				if (UEnum* NodeEnum = SelectNode->GetEnum())
-				{
-					int32 EnumValue = NodeEnum->GetValueByName(OptionPins[OptionIdx]->PinName);
-					LiteralTerm->Name = FString::Printf(TEXT("%d"), EnumValue);
-				}
-				else
-				{
-					LiteralTerm->Name = FString::Printf(TEXT("%d"), OptionIdx);
-				}
-				Statement.RHS.Add(LiteralTerm);
-				// If there is a previous IfNot statement, hook this one to that one for jumping
-				if (PrevIfNotStatement)
-				{
-					Statement.bIsJumpTarget = true;
-					PrevIfNotStatement->TargetLabel = &Statement;
-				}
-
-				// Create a GotoIfNot statement using the BoolTerm from above as the condition
-				FBlueprintCompiledStatement* IfNotStatement = &Context.AppendStatementForNode(Node);
-				IfNotStatement->Type = KCST_GotoIfNot;
-				IfNotStatement->LHS = BoolTerm;
-
-				// Create an assignment statement
-				FBlueprintCompiledStatement& AssignStatement = Context.AppendStatementForNode(Node);
-				AssignStatement.Type = KCST_Assignment;
-				AssignStatement.LHS = *ReturnTerm;
-				// Get the kismet term from the option pin
-				UEdGraphPin* OptionPinToTry = FEdGraphUtilities::GetNetFromPin(OptionPins[OptionIdx]);
-				FBPTerminal** OptionTerm = Context.NetMap.Find(OptionPinToTry);
-				if (!OptionTerm)
-				{
-					Context.MessageLog.Error(*LOCTEXT("Error_UnregisterOptionPin", "Unregister option pin @@").ToString(), OptionPins[OptionIdx]);
-					return;
-				}
-				AssignStatement.RHS.Add(*OptionTerm);
-
-				// Create an unconditional goto to exit the node
-				FBlueprintCompiledStatement& GotoStatement = Context.AppendStatementForNode(Node);
-				GotoStatement.Type = KCST_UnconditionalGoto;
-				GotoStatementList.Add(&GotoStatement);
-
-				// If this is the last IfNot statement, hook the jump to an error message
-				if (OptionIdx == OptionPins.Num() - 1)
-				{
-					// Create a CallFunction statement for doing a print string of our error message
-					FBlueprintCompiledStatement& PrintStatement = Context.AppendStatementForNode(Node);
-					PrintStatement.Type = KCST_CallFunction;
-					PrintStatement.bIsJumpTarget = true;
-					FName PrintStringFunctionName = "";
-					UClass* PrintStringFunctionClass = NULL;
-					SelectNode->GetPrintStringFunction(PrintStringFunctionName, &PrintStringFunctionClass);
-					UFunction* PrintFunction = FindField<UFunction>(PrintStringFunctionClass, PrintStringFunctionName);
-					PrintStatement.FunctionToCall = PrintFunction;
-					PrintStatement.FunctionContext = NULL;
-					PrintStatement.bIsParentContext = false;
-
-					// Create a local int for use in the equality call function below (LiteralTerm = the right hand side of the EqualEqual_IntInt or NotEqual_BoolBool statement)
-					FBPTerminal* LiteralStringTerm = Context.CreateLocalTerminal(ETerminalSpecification::TS_Literal);
-					LiteralStringTerm->bIsLiteral = true;
-					LiteralStringTerm->Type.PinCategory = UEdGraphSchema_K2::PC_String;
-
-					FString SelectionNodeType;
-					if (IndexPin)
-					{
-						if (UEnum* EnumObject = Cast<UEnum>(IndexPin->PinType.PinSubCategoryObject.Get()))
-						{
-							SelectionNodeType = EnumObject->GetName();
-						}
-						else
-						{
-							// Not an enum, so just use the basic type
-							SelectionNodeType = IndexPin->PinType.PinCategory.ToString();
-						}
-					}
-					const UEdGraph* OwningGraph = Context.MessageLog.FindSourceObjectTypeChecked<UEdGraph>( SelectNode->GetGraph() );
-					LiteralStringTerm->Name = FText::Format(
-						LOCTEXT("SelectNodeIndexWarningFmt", "Graph {0}: Selection Node of type {1} failed! Out of bounds indexing of the options. There are only {2} options available."),
-						(OwningGraph) ? FText::FromString(OwningGraph->GetFullName()) : LOCTEXT("SelectNodeIndexWarningNoGraph", "NONE"),
-						(IndexPin) ? FText::FromString(SelectionNodeType) : LOCTEXT("SelectNodeIndexWarningNoPin", "NONE"),
-						OptionPins.Num()
-					).ToString();
-					PrintStatement.RHS.Add(LiteralStringTerm);
-
-					// Hook the IfNot statement's jump target to this statement
-					IfNotStatement->TargetLabel = &PrintStatement;
-				}
-
-				PrevIfNotStatement = IfNotStatement;
-			}
-
-			// Create a noop to jump to so the unconditional goto statements can exit the node after successful assignment
-			FBlueprintCompiledStatement& NopStatement = Context.AppendStatementForNode(Node);
-			NopStatement.Type = KCST_Nop;
-			NopStatement.bIsJumpTarget = true;
-			// Loop through the unconditional goto statements and fix their jump targets
-			for (FBlueprintCompiledStatement* GotoStatement : GotoStatementList)
-			{
-				GotoStatement->TargetLabel = &NopStatement;
-			}
-		}
-	}
-};
-
 UK2Node_Select::UK2Node_Select(const FObjectInitializer& ObjectInitializer)
 	: Super(ObjectInitializer)
 {
@@ -353,8 +177,6 @@ void UK2Node_Select::AllocateDefaultPins()
 		NumOptionPins = EnumEntries.Num();
 	}
 
-	static const FBoolConfigValueHelper UseSelectRef(TEXT("Kismet"), TEXT("bUseSelectRef"), GEngineIni);
-
 	// Create the option pins
 	for (int32 Idx = 0; Idx < NumOptionPins; Idx++)
 	{
@@ -377,7 +199,7 @@ void UK2Node_Select::AllocateDefaultPins()
 
 		if (NewPin)
 		{
-			NewPin->bDisplayAsMutableRef = UseSelectRef;
+			NewPin->bDisplayAsMutableRef = true;
 			if (IndexPinType.PinCategory == UEdGraphSchema_K2::PC_Boolean)
 			{
 				NewPin->PinFriendlyName = (Idx == 0 ? GFalse : GTrue);
@@ -394,7 +216,7 @@ void UK2Node_Select::AllocateDefaultPins()
 
 	// Create the return value
 	UEdGraphPin* ReturnPin = CreatePin(EGPD_Output, UEdGraphSchema_K2::PC_Wildcard, UEdGraphSchema_K2::PN_ReturnValue);
-	ReturnPin->bDisplayAsMutableRef = UseSelectRef;
+	ReturnPin->bDisplayAsMutableRef = true;
 
 	Super::AllocateDefaultPins();
 }
@@ -435,6 +257,19 @@ FText UK2Node_Select::GetNodeTitle(ENodeTitleType::Type TitleType) const
 
 UK2Node::ERedirectType UK2Node_Select::DoPinsMatchForReconstruction(const UEdGraphPin* NewPin, int32 NewPinIndex, const UEdGraphPin* OldPin, int32 OldPinIndex) const
 {
+	if (bReconstructForPinTypeChange)
+	{
+		// If we're reconstructing for the purposes of changing the index pin type then we want to
+		// keep our connections based on the index of the option pin
+		if (NewPin != GetIndexPin() && NewPin != GetReturnValuePin())
+		{
+			if (NewPinIndex == OldPinIndex)
+			{
+				return UK2Node::ERedirectType_Name;
+			}
+		}
+	}
+
 	// Check to see if the new pin name matches the old pin name.
 	if (Enum && (NewPinIndex < NumOptionPins) && (NewPin->PinName != OldPin->PinName))
 	{
@@ -469,11 +304,11 @@ void UK2Node_Select::ReallocatePinsDuringReconstruction(TArray<UEdGraphPin*>& Ol
 	UEdGraphPin* OldReturnPin = nullptr;
 	for (UEdGraphPin* OldPin : OldPins)
 	{
-		if (OldPin->PinName == TEXT("bPickOption0"))
+		if (OldPin->PinName == NAME_bPickOption0)
 		{
 			OldConditionPin = OldPin;
 		}
-		else if (OldPin->PinName == TEXT("Index"))
+		else if (OldPin->PinName == NAME_Index)
 		{
 			OldIndexPin = OldPin;
 		}
@@ -512,8 +347,8 @@ void UK2Node_Select::ReallocatePinsDuringReconstruction(TArray<UEdGraphPin*>& Ol
 			PinConnectionListChanged(IndexPin);
 		}
 
-		UEdGraphPin* OptionPin0 = FindPin(TEXT("Option 0"));
-		UEdGraphPin* OptionPin1 = FindPin(TEXT("Option 1"));
+		UEdGraphPin* OptionPin0 = FindPin(NAME_Option_0);
+		UEdGraphPin* OptionPin1 = FindPin(NAME_Option_1);
 
 		for (UEdGraphPin* OldPin : OldPins)
 		{
@@ -560,6 +395,7 @@ void UK2Node_Select::PostReconstructNode()
 {
 	// After ReconstructNode we must be sure that no additional reconstruction is required
 	bReconstructNode = false;
+	bReconstructForPinTypeChange = false;
 
 	UEdGraphPin* ReturnPin = GetReturnValuePin();
 	const bool bFillTypeFromConnected = ReturnPin && (ReturnPin->PinType.PinCategory == UEdGraphSchema_K2::PC_Wildcard);
@@ -587,7 +423,7 @@ void UK2Node_Select::PostReconstructNode()
 		}
 
 		ReturnPin->PinType = PinType;
-		PinTypeChanged(ReturnPin);
+		OnPinTypeChanged(ReturnPin);
 	}
 
 	Super::PostReconstructNode();
@@ -610,7 +446,7 @@ void UK2Node_Select::NotifyPinConnectionListChanged(UEdGraphPin* Pin)
 			{
 				Pin->PinType = LinkPin->PinType;
 
-				PinTypeChanged(Pin);
+				OnPinTypeChanged(Pin);
 			}
 		}
 	}
@@ -622,24 +458,71 @@ void UK2Node_Select::NotifyPinConnectionListChanged(UEdGraphPin* Pin)
 		UEdGraphPin* ReturnPin = FindPin(UEdGraphSchema_K2::PN_ReturnValue);
 
 		// See if this pin is one of the wildcard pins
-		bool bIsWildcardPin = (Pin == ReturnPin || OptionPins.Find(Pin) != INDEX_NONE) && Pin->PinType.PinCategory == UEdGraphSchema_K2::PC_Wildcard;
+		const bool bIsWildcardPin = ((Pin->PinType.PinCategory == UEdGraphSchema_K2::PC_Wildcard) && ((Pin == ReturnPin) || (OptionPins.Find(Pin) != INDEX_NONE)));
 
-		// If the pin was one of the wildcards we have to handle it specially
-		if (bIsWildcardPin)
+		TFunction<bool(UEdGraphPin*)> PinInUse = [&PinInUse](UEdGraphPin* PinToConsider)
 		{
-			// If the pin is linked, make sure the other wildcard pins match
-			if (Pin->LinkedTo.Num() > 0)
+			bool bPinInUse = ((PinToConsider->LinkedTo.Num() > 0) || (PinToConsider->ParentPin != nullptr) || !PinToConsider->DoesDefaultValueMatchAutogenerated());
+			if (!bPinInUse)
 			{
-				UEdGraphPin* LinkPin = Pin->LinkedTo[0];
+				for (UEdGraphPin* SubPin : PinToConsider->SubPins)
+				{
+					bPinInUse = PinInUse(SubPin);
+					if (bPinInUse)
+					{
+						break;
+					}
+				}
+			}
+			return bPinInUse;
+		};
 
-				if (Pin->PinType != LinkPin->PinType)
+		bool bPinsInUse = PinInUse(ReturnPin);
+		
+		if (!bPinsInUse)
+		{
+			for (UEdGraphPin* OptionPin : OptionPins)
+			{
+				bPinsInUse = PinInUse(OptionPin);
+				if (bPinsInUse)
 				{
-					Pin->PinType = LinkPin->PinType;
+					break;
+				}
+			}
+		}
+
+		bool bPinTypeChanged = false;
 
-					PinTypeChanged(Pin);
+		if (bPinsInUse)
+		{
+			// If the pin was one of the wildcards we have to handle it specially
+			if (bIsWildcardPin)
+			{
+				// If the pin is linked, make sure the other wildcard pins match
+				if (Pin->LinkedTo.Num() > 0)
+				{
+					UEdGraphPin* LinkPin = Pin->LinkedTo[0];
+
+					if (Pin->PinType != LinkPin->PinType)
+					{
+						Pin->PinType = LinkPin->PinType;
+						bPinTypeChanged = true;
+					}
 				}
 			}
 		}
+		else
+		{
+			bPinTypeChanged = true;
+			Pin->PinType.PinCategory = UEdGraphSchema_K2::PC_Wildcard;
+			Pin->PinType.PinSubCategory = NAME_None;
+			Pin->PinType.PinSubCategoryObject = nullptr;
+		}
+
+		if (bPinTypeChanged)
+		{
+			OnPinTypeChanged(Pin);
+		}
 	}
 }
 
@@ -719,8 +602,10 @@ void UK2Node_Select::GetPrintStringFunction(FName& FunctionName, UClass** Functi
 	*FunctionClass = UKismetSystemLibrary::StaticClass();
 }
 
-void UK2Node_Select::AddOptionPinToNode()
+void UK2Node_Select::AddInputPin()
 {
+	Modify();
+
 	const UEdGraphSchema_K2* K2Schema = GetDefault<UEdGraphSchema_K2>();
 
 	// Increment the pin count
@@ -794,7 +679,7 @@ void UK2Node_Select::NodeConnectionListChanged()
 	}
 }
 
-bool UK2Node_Select::CanAddOptionPinToNode() const
+bool UK2Node_Select::CanAddPin() const
 {
 	if (IndexPinType.PinCategory == UEdGraphSchema_K2::PC_Byte &&
 		IndexPinType.PinSubCategoryObject.IsValid() &&
@@ -827,10 +712,11 @@ bool UK2Node_Select::CanRemoveOptionPinToNode() const
 
 void UK2Node_Select::ChangePinType(UEdGraphPin* Pin)
 {
-	PinTypeChanged(Pin);
+	OnPinTypeChanged(Pin);
 
 	if (bReconstructNode)
 	{
+		bReconstructForPinTypeChange = true;
 		ReconstructNode();
 	}
 
@@ -876,6 +762,12 @@ bool UK2Node_Select::CanChangePinType(UEdGraphPin* Pin) const
 }
 
 void UK2Node_Select::PinTypeChanged(UEdGraphPin* Pin)
+{
+	bReconstructForPinTypeChange = true;
+	OnPinTypeChanged(Pin);
+}
+
+void UK2Node_Select::OnPinTypeChanged(UEdGraphPin* Pin)
 {
 	const UEdGraphSchema_K2* Schema = GetDefault<UEdGraphSchema_K2>();
 
@@ -888,7 +780,7 @@ void UK2Node_Select::PinTypeChanged(UEdGraphPin* Pin)
 			// Since it is an interactive action we want the pins to go away regardless of the new type
 			for (UEdGraphPin* PinToDiscard : Pins)
 			{
-				PinToDiscard->bSavePinIfOrphaned = false;
+				PinToDiscard->SetSavePinIfOrphaned(false);
 			}
 
 			if (IndexPinType.PinSubCategoryObject.IsValid())
@@ -968,10 +860,10 @@ void UK2Node_Select::PostPasteNode()
 		FString OldDefaultValue = IndexPin->DefaultValue;
 
 		// Corrects data in the index pin that is not valid after pasting
-		PinTypeChanged(IndexPin);
+		OnPinTypeChanged(IndexPin);
 
 		// Restore the default value of the index pin
-		IndexPin->DefaultValue = OldDefaultValue;
+		IndexPin->DefaultValue = MoveTemp(OldDefaultValue);
 	}
 }
 
@@ -994,10 +886,7 @@ bool UK2Node_Select::IsConnectionDisallowed(const UEdGraphPin* MyPin, const UEdG
 
 FNodeHandlingFunctor* UK2Node_Select::CreateNodeHandler(FKismetCompilerContext& CompilerContext) const
 {
-	static const FBoolConfigValueHelper UseSelectRef(TEXT("Kismet"), TEXT("bUseSelectRef"), GEngineIni);
-	return UseSelectRef
-		? static_cast<FNodeHandlingFunctor*>(new FKCHandler_SelectRef(CompilerContext))
-		: static_cast<FNodeHandlingFunctor*>(new FKCHandler_Select(CompilerContext));
+	return static_cast<FNodeHandlingFunctor*>(new FKCHandler_Select(CompilerContext));
 }
 
 void UK2Node_Select::GetMenuActions(FBlueprintActionDatabaseRegistrar& ActionRegistrar) const
@@ -1029,13 +918,6 @@ void UK2Node_Select::ExpandNode(class FKismetCompilerContext& CompilerContext, U
 {
 	Super::ExpandNode(CompilerContext, SourceGraph);
 
-	static const FBoolConfigValueHelper UseSelectRef(TEXT("Kismet"), TEXT("bUseSelectRef"), GEngineIni);
-	if (!UseSelectRef)
-	{
-		return;
-	}
-
-	bool bSuccess = false;
 	const UEdGraphSchema_K2* Schema = CompilerContext.GetSchema();
 	for (UEdGraphPin* Pin : Pins)
 	{
diff --git a/Engine/Source/Editor/BlueprintGraph/Private/K2Node_SpawnActorFromClass.cpp b/Engine/Source/Editor/BlueprintGraph/Private/K2Node_SpawnActorFromClass.cpp
--- a/Engine/Source/Editor/BlueprintGraph/Private/K2Node_SpawnActorFromClass.cpp
+++ b/Engine/Source/Editor/BlueprintGraph/Private/K2Node_SpawnActorFromClass.cpp
@@ -447,17 +447,4 @@ void UK2Node_SpawnActorFromClass::ExpandNode(class FKismetCompilerContext& Compi
 	SpawnNode->BreakAllNodeLinks();
 }
 
-bool UK2Node_SpawnActorFromClass::HasExternalDependencies(TArray<class UStruct*>* OptionalOutput) const
-{
-	UClass* SourceClass = GetClassToSpawn();
-	const UBlueprint* SourceBlueprint = GetBlueprint();
-	const bool bResult = (SourceClass != NULL) && (SourceClass->ClassGeneratedBy != SourceBlueprint);
-	if (bResult && OptionalOutput)
-	{
-		OptionalOutput->AddUnique(SourceClass);
-	}
-	const bool bSuperResult = Super::HasExternalDependencies(OptionalOutput);
-	return bSuperResult || bResult;
-}
-
 #undef LOCTEXT_NAMESPACE
diff --git a/Engine/Source/Editor/BlueprintGraph/Public/BlueprintEditorSettings.h b/Engine/Source/Editor/BlueprintGraph/Public/BlueprintEditorSettings.h
--- a/Engine/Source/Editor/BlueprintGraph/Public/BlueprintEditorSettings.h
+++ b/Engine/Source/Editor/BlueprintGraph/Public/BlueprintEditorSettings.h
@@ -5,6 +5,7 @@
 #include "CoreMinimal.h"
 #include "UObject/ObjectMacros.h"
 #include "UObject/Object.h"
+#include "Engine/Blueprint.h"
 #include "BlueprintEditorSettings.generated.h"
 
 UENUM()
@@ -89,6 +90,22 @@ public:
 	UPROPERTY(EditAnywhere, config, Category=Workflow)
 	bool bNavigateToNativeFunctionsFromCallNodes;
 
+	/** Blueprint bookmark database */
+	UPROPERTY(config)
+	TMap<FGuid, FEditedDocumentInfo> Bookmarks;
+
+	/** Blueprint bookmark nodes (for display) */
+	UPROPERTY(config)
+	TArray<FBPEditorBookmarkNode> BookmarkNodes;
+
+	/** If enabled, comment nodes will be included in the tree view display in the Bookmarks tab. */
+	UPROPERTY(config)
+	bool bIncludeCommentNodesInBookmarksTab;
+
+	/** If enabled, only the bookmarks for the current document will be shown in the Bookmarks tab. */
+	UPROPERTY(config)
+	bool bShowBookmarksForCurrentDocumentOnlyInTab;
+
 // Compiler Settings
 public:
 	/** Determines when to save Blueprints post-compile */
