Mirror 4083455 from Dev-Framework
Reduce use of template nodes - a template node was causing an ensure when asset reload code attempted to fix it up

diff --git a/Engine/Source/Editor/BlueprintGraph/Classes/EdGraphSchema_K2_Actions.h b/Engine/Source/Editor/BlueprintGraph/Classes/EdGraphSchema_K2_Actions.h
--- a/Engine/Source/Editor/BlueprintGraph/Classes/EdGraphSchema_K2_Actions.h
+++ b/Engine/Source/Editor/BlueprintGraph/Classes/EdGraphSchema_K2_Actions.h
@@ -4,10 +4,13 @@
 
 #include "CoreMinimal.h"
 #include "UObject/ObjectMacros.h"
+#include "UObject/UObjectGlobals.h"
 #include "Templates/SubclassOf.h"
 #include "UObject/UnrealType.h"
 #include "Engine/Blueprint.h"
+#include "EdGraph/EdGraph.h"
 #include "EdGraph/EdGraphSchema.h"
+
 #include "EdGraphSchema_K2_Actions.generated.h"
 
 class UEdGraph;
@@ -17,6 +20,15 @@ class UK2Node;
 * FEdGraphSchemaAction_K2NewNode
 *******************************************************************************/
 
+enum class EK2NewNodeFlags
+{
+	None = 0x0,
+
+	SelectNewNode = 0x1,
+	GotoNewNode = 0x2,
+};
+ENUM_CLASS_FLAGS(EK2NewNodeFlags)
+
 /** Action to add a node to the graph */
 USTRUCT()
 struct BLUEPRINTGRAPH_API FEdGraphSchemaAction_K2NewNode : public FEdGraphSchemaAction
@@ -47,21 +59,49 @@ struct BLUEPRINTGRAPH_API FEdGraphSchemaAction_K2NewNode : public FEdGraphSchema
 	{}
 
 	// FEdGraphSchemaAction interface
-	virtual UEdGraphNode* PerformAction(class UEdGraph* ParentGraph, UEdGraphPin* FromPin, const FVector2D Location, bool bSelectNewNode = true) override;
-	virtual UEdGraphNode* PerformAction(class UEdGraph* ParentGraph, TArray<UEdGraphPin*>& FromPins, const FVector2D Location, bool bSelectNewNode = true) override;
+	virtual UEdGraphNode* PerformAction(UEdGraph* ParentGraph, UEdGraphPin* FromPin, const FVector2D Location, bool bSelectNewNode = true) override;
+	virtual UEdGraphNode* PerformAction(UEdGraph* ParentGraph, TArray<UEdGraphPin*>& FromPins, const FVector2D Location, bool bSelectNewNode = true) override;
 	virtual void AddReferencedObjects(FReferenceCollector& Collector) override;
 	// End of FEdGraphSchemaAction interface
 
 	template <typename NodeType>
 	static NodeType* SpawnNodeFromTemplate(class UEdGraph* ParentGraph, NodeType* InTemplateNode, const FVector2D Location, bool bSelectNewNode = true)
 	{
-		FEdGraphSchemaAction_K2NewNode Action;
-		Action.NodeTemplate = InTemplateNode;
+		EK2NewNodeFlags Options = EK2NewNodeFlags::None;
+		if (bSelectNewNode)
+		{
+			Options |= EK2NewNodeFlags::SelectNewNode;
+		}
+		return Cast<NodeType>(CreateNode(ParentGraph, TArrayView<UEdGraphPin*>(), Location, InTemplateNode, Options));
+	}
+
+	template <typename NodeType>
+	static NodeType* SpawnNode(class UEdGraph* ParentGraph, const FVector2D Location, EK2NewNodeFlags Options, TFunctionRef<void(NodeType*)> InitializerFn)
+	{
+		UClass* NodeClass = NodeType::StaticClass();
+		return Cast<NodeType>(CreateNode(
+			ParentGraph,
+			TArrayView<UEdGraphPin*>(),
+			Location,
+			[NodeClass](UEdGraph* ParentGraph)->UK2Node*
+			{
+				return NewObject<UK2Node>(ParentGraph, NodeClass);
+			},
+			[InitializerFn](UK2Node* NewNode)
+			{
+				InitializerFn(CastChecked<NodeType>(NewNode));
+			},
+			Options));
+	}
 
-		return Cast<NodeType>(Action.PerformAction(ParentGraph, NULL, Location, bSelectNewNode));
+	template <typename NodeType>
+	static NodeType* SpawnNode(class UEdGraph* ParentGraph, const FVector2D Location, EK2NewNodeFlags Options)
+	{
+		return SpawnNode<NodeType>(ParentGraph, Location, Options, [](NodeType*){});
 	}
 
-	static UEdGraphNode* CreateNode(class UEdGraph* ParentGraph, UEdGraphPin* FromPin, const FVector2D Location, class UK2Node* NodeTemplate, bool bSelectNewNode);
+	static UEdGraphNode* CreateNode(UEdGraph* ParentGraph, TArrayView<UEdGraphPin*> FromPins, const FVector2D Location, UK2Node* NodeTemplate, EK2NewNodeFlags Options);
+	static UEdGraphNode* CreateNode(UEdGraph* ParentGraph, TArrayView<UEdGraphPin*> FromPins, const FVector2D Location, TFunctionRef<UK2Node*(UEdGraph*)> ConstructionFn, TFunctionRef<void(UK2Node*)> InitializerFn, EK2NewNodeFlags Options);
 };
 
 /*******************************************************************************
diff --git a/Engine/Source/Editor/BlueprintGraph/Private/EdGraphSchema_K2.cpp b/Engine/Source/Editor/BlueprintGraph/Private/EdGraphSchema_K2.cpp
--- a/Engine/Source/Editor/BlueprintGraph/Private/EdGraphSchema_K2.cpp
+++ b/Engine/Source/Editor/BlueprintGraph/Private/EdGraphSchema_K2.cpp
@@ -6924,7 +6924,7 @@ void UEdGraphSchema_K2::OnPinConnectionDoubleCicked(UEdGraphPin* PinA, UEdGraphP
 	UEdGraph* ParentGraph = PinA->GetOwningNode()->GetGraph();
 	if (!FBlueprintEditorUtils::IsGraphReadOnly(ParentGraph))
 	{
-		UK2Node_Knot* NewKnot = FEdGraphSchemaAction_K2NewNode::SpawnNodeFromTemplate<UK2Node_Knot>(ParentGraph, NewObject<UK2Node_Knot>(), KnotTopLeft);
+		UK2Node_Knot* NewKnot = FEdGraphSchemaAction_K2NewNode::SpawnNode<UK2Node_Knot>(ParentGraph, KnotTopLeft, EK2NewNodeFlags::SelectNewNode);
 
 		// Move the connections across (only notifying the knot, as the other two didn't really change)
 		PinA->BreakLinkTo(PinB);
@@ -6968,22 +6968,32 @@ void UEdGraphSchema_K2::ConfigureVarNode(UK2Node_Variable* InVarNode, FName InVa
 
 UK2Node_VariableGet* UEdGraphSchema_K2::SpawnVariableGetNode(const FVector2D GraphPosition, class UEdGraph* ParentGraph, FName VariableName, UStruct* Source) const
 {
-	UK2Node_VariableGet* NodeTemplate = NewObject<UK2Node_VariableGet>();
 	UBlueprint* Blueprint = FBlueprintEditorUtils::FindBlueprintForGraph(ParentGraph);
 
-	UEdGraphSchema_K2::ConfigureVarNode(NodeTemplate, VariableName, Source, Blueprint);
-
-	return FEdGraphSchemaAction_K2NewNode::SpawnNodeFromTemplate<UK2Node_VariableGet>(ParentGraph, NodeTemplate, GraphPosition);
+	return FEdGraphSchemaAction_K2NewNode::SpawnNode<UK2Node_VariableGet>(
+		ParentGraph,
+		GraphPosition,
+		EK2NewNodeFlags::SelectNewNode,
+		[VariableName, Source, Blueprint](UK2Node_VariableGet* NewInstance)
+		{
+			UEdGraphSchema_K2::ConfigureVarNode(NewInstance, VariableName, Source, Blueprint);
+		}
+	);
 }
 
 UK2Node_VariableSet* UEdGraphSchema_K2::SpawnVariableSetNode(const FVector2D GraphPosition, class UEdGraph* ParentGraph, FName VariableName, UStruct* Source) const
 {
-	UK2Node_VariableSet* NodeTemplate = NewObject<UK2Node_VariableSet>();
 	UBlueprint* Blueprint = FBlueprintEditorUtils::FindBlueprintForGraph(ParentGraph);
 
-	UEdGraphSchema_K2::ConfigureVarNode(NodeTemplate, VariableName, Source, Blueprint);
-
-	return FEdGraphSchemaAction_K2NewNode::SpawnNodeFromTemplate<UK2Node_VariableSet>(ParentGraph, NodeTemplate, GraphPosition);
+	return FEdGraphSchemaAction_K2NewNode::SpawnNode<UK2Node_VariableSet>(
+		ParentGraph,
+		GraphPosition,
+		EK2NewNodeFlags::SelectNewNode,
+		[VariableName, Source, Blueprint](UK2Node_VariableSet* NewInstance)
+		{
+			UEdGraphSchema_K2::ConfigureVarNode(NewInstance, VariableName, Source, Blueprint);
+		}
+	);
 }
 
 UEdGraphPin* UEdGraphSchema_K2::DropPinOnNode(UEdGraphNode* InTargetNode, const FName& InSourcePinName, const FEdGraphPinType& InSourcePinType, EEdGraphPinDirection InSourcePinDirection) const
diff --git a/Engine/Source/Editor/BlueprintGraph/Private/EdGraphSchema_K2_Actions.cpp b/Engine/Source/Editor/BlueprintGraph/Private/EdGraphSchema_K2_Actions.cpp
--- a/Engine/Source/Editor/BlueprintGraph/Private/EdGraphSchema_K2_Actions.cpp
+++ b/Engine/Source/Editor/BlueprintGraph/Private/EdGraphSchema_K2_Actions.cpp
@@ -182,13 +182,54 @@ UFunction* FEdGraphSchemaAction_K2Graph::GetFunction() const
 /////////////////////////////////////////////////////
 // FEdGraphSchemaAction_K2ViewNode
 
-UEdGraphNode* FEdGraphSchemaAction_K2NewNode::CreateNode(class UEdGraph* ParentGraph, UEdGraphPin* FromPin, const FVector2D Location, class UK2Node* NodeTemplate, bool bSelectNewNode/* = true*/)
+UEdGraphNode* FEdGraphSchemaAction_K2NewNode::CreateNode(
+	UEdGraph* ParentGraph,
+	TArrayView<UEdGraphPin*> FromPins,
+	const FVector2D Location,
+	UK2Node* NodeTemplate,
+	EK2NewNodeFlags Options)
+{
+	if (NodeTemplate)
+	{
+		return CreateNode(
+			ParentGraph,
+			FromPins,
+			Location,
+			[NodeTemplate](UEdGraph* ParentGraph)->UK2Node*
+			{
+				return DuplicateObject<UK2Node>(NodeTemplate, ParentGraph);
+			},
+			[](UK2Node*) {},
+			Options
+			);
+	}
+	return nullptr;
+}
+
+UEdGraphNode* FEdGraphSchemaAction_K2NewNode::CreateNode(
+	UEdGraph* ParentGraph,
+	TArrayView<UEdGraphPin*> FromPins,
+	const FVector2D Location,
+	TFunctionRef<UK2Node*(UEdGraph*)> ConstructionFn,
+	TFunctionRef<void(UK2Node*)> InitializerFn,
+	EK2NewNodeFlags Options)
 {
+	const bool bSelectNewNode = (Options & EK2NewNodeFlags::SelectNewNode) != EK2NewNodeFlags::None;
+	const bool bGotoNode = (Options & EK2NewNodeFlags::GotoNewNode) != EK2NewNodeFlags::None;
+
+	const FScopedTransaction Transaction(NSLOCTEXT("UnrealEd", "K2_AddNode", "Add Node"));
+	ParentGraph->Modify();
+	for(UEdGraphPin* FromPin : FromPins)
+	{
+		FromPin->Modify();
+	}
+
 	// Smart pointer that handles fixup after potential node reconstruction
-	FWeakGraphPinPtr FromPinPtr = FromPin;
+	FWeakGraphPinPtr FromPinPtr = FromPins.Num() > 0 ? FromPins[0] : nullptr;
 
 	// Duplicate template node to create new node
-	UEdGraphNode* ResultNode = DuplicateObject<UK2Node>(NodeTemplate, ParentGraph);
+	UK2Node* ResultNode = ConstructionFn(ParentGraph);
+	InitializerFn(ResultNode);
 	ResultNode->SetFlags(RF_Transactional);
 
 	ParentGraph->AddNode(ResultNode, true, bSelectNewNode);
@@ -218,79 +259,66 @@ UEdGraphNode* FEdGraphSchemaAction_K2NewNode::CreateNode(class UEdGraph* ParentG
 
 	// make sure to auto-wire after we position the new node (in case the 
 	// auto-wire creates a conversion node to put between them)
-	ResultNode->AutowireNewNode(FromPinPtr);
+	for (UEdGraphPin* FromPin : FromPins)
+	{
+		ResultNode->AutowireNewNode(FromPin);
+	}
 
 	// Update Analytics for the new nodes
-	FBlueprintEditorUtils::AnalyticsTrackNewNode( ResultNode );
+	FBlueprintEditorUtils::AnalyticsTrackNewNode(ResultNode);
 	// NOTE: At this point the node may have been reconstructed, depending on node type!
 
-	return ResultNode;
-}
-
-UEdGraphNode* FEdGraphSchemaAction_K2NewNode::PerformAction(class UEdGraph* ParentGraph, UEdGraphPin* FromPin, const FVector2D Location, bool bSelectNewNode/* = true*/)
-{
-	UEdGraphNode* ResultNode = nullptr;
+	UBlueprint* Blueprint = FBlueprintEditorUtils::FindBlueprintForGraphChecked(ParentGraph);
 
-	// If there is a template, we actually use it
-	if (NodeTemplate != nullptr)
+	// See if we need to recompile skeleton after adding this node, or just mark dirty
+	if (ResultNode->NodeCausesStructuralBlueprintChange())
 	{
-		const FScopedTransaction Transaction( NSLOCTEXT("UnrealEd", "K2_AddNode", "Add Node") );
-		ParentGraph->Modify();
-		if (FromPin)
-		{
-			FromPin->Modify();
-		}
-
-		ResultNode = CreateNode(ParentGraph, FromPin, Location, NodeTemplate, bSelectNewNode);
-
-		UBlueprint* Blueprint = FBlueprintEditorUtils::FindBlueprintForGraphChecked(ParentGraph);
-
-		// See if we need to recompile skeleton after adding this node, or just mark dirty
-		UK2Node* K2Node = Cast<UK2Node>(ResultNode);
-		check(K2Node);
-		if(K2Node->NodeCausesStructuralBlueprintChange())
-		{
-			FBlueprintEditorUtils::MarkBlueprintAsStructurallyModified(Blueprint);
-		}
-		else
-		{
-			FBlueprintEditorUtils::MarkBlueprintAsModified(Blueprint);
-		}
+		FBlueprintEditorUtils::MarkBlueprintAsStructurallyModified(Blueprint);
+	}
+	else
+	{
+		FBlueprintEditorUtils::MarkBlueprintAsModified(Blueprint);
+	}
 
-		// Clear any error messages resulting from placing a node.  They'll be flagged on the next compile
-		K2Node->ErrorMsg.Empty();
-		K2Node->bHasCompilerMessage = false;
+	// Clear any error messages resulting from placing a node.  They'll be flagged on the next compile
+	ResultNode->ErrorMsg.Empty();
+	ResultNode->bHasCompilerMessage = false;
 
-		if ( bGotoNode && ResultNode )
-		{
-			// Select existing node
-			FKismetEditorUtilities::BringKismetToFocusAttentionOnObject(ResultNode);
-		}
+	if (bGotoNode && ResultNode)
+	{
+		// Select existing node
+		FKismetEditorUtilities::BringKismetToFocusAttentionOnObject(ResultNode);
 	}
 
 	return ResultNode;
 }
 
-UEdGraphNode* FEdGraphSchemaAction_K2NewNode::PerformAction(class UEdGraph* ParentGraph, TArray<UEdGraphPin*>& FromPins, const FVector2D Location, bool bSelectNewNode/* = true*/) 
+UEdGraphNode* FEdGraphSchemaAction_K2NewNode::PerformAction(class UEdGraph* ParentGraph, UEdGraphPin* FromPin, const FVector2D Location, bool bSelectNewNode/* = true*/)
 {
-	UEdGraphNode* ResultNode = nullptr;
-
-	if (FromPins.Num() > 0)
+	EK2NewNodeFlags Options = EK2NewNodeFlags::None;
+	if (bSelectNewNode)
 	{
-		ResultNode = PerformAction(ParentGraph, FromPins[0], Location, bSelectNewNode);
-
-		// Try autowiring the rest of the pins
-		for (int32 Index = 1; Index < FromPins.Num(); ++Index)
-		{
-			ResultNode->AutowireNewNode(FromPins[Index]);
-		}
+		Options |= EK2NewNodeFlags::SelectNewNode;
 	}
-	else
+	if (bGotoNode)
 	{
-		ResultNode = PerformAction(ParentGraph, nullptr, Location, bSelectNewNode);
+		Options |= EK2NewNodeFlags::GotoNewNode;
 	}
+	return CreateNode(ParentGraph, MakeArrayView(&FromPin, 1), Location, NodeTemplate, Options);
+}
 
-	return ResultNode;
+UEdGraphNode* FEdGraphSchemaAction_K2NewNode::PerformAction(class UEdGraph* ParentGraph, TArray<UEdGraphPin*>& FromPins, const FVector2D Location, bool bSelectNewNode/* = true*/) 
+{
+	EK2NewNodeFlags Options = EK2NewNodeFlags::None;
+	if (bSelectNewNode)
+	{
+		Options |= EK2NewNodeFlags::SelectNewNode;
+	}
+	if (bGotoNode)
+	{
+		Options |= EK2NewNodeFlags::GotoNewNode;
+	}
+	return CreateNode(ParentGraph, FromPins, Location, NodeTemplate, Options);
 }
 
 void FEdGraphSchemaAction_K2NewNode::AddReferencedObjects( FReferenceCollector& Collector )
@@ -337,7 +365,7 @@ UEdGraphNode* FEdGraphSchemaAction_K2AssignDelegate::AssignDelegate(class UK2Nod
 			FromPin->Modify();
 		}
 
-		BindNode = Cast<UK2Node_AddDelegate>(CreateNode(ParentGraph, FromPin, Location, NodeTemplate, bSelectNewNode));
+		BindNode = Cast<UK2Node_AddDelegate>(CreateNode(ParentGraph, MakeArrayView(&FromPin, 1), Location, NodeTemplate, bSelectNewNode ? EK2NewNodeFlags::SelectNewNode : EK2NewNodeFlags::None));
 		UMulticastDelegateProperty* DelegateProperty = BindNode ? Cast<UMulticastDelegateProperty>(BindNode->GetProperty()) : nullptr;
 		if(DelegateProperty)
 		{
diff --git a/Engine/Source/Editor/BlueprintGraph/Private/K2Node_FormatText.cpp b/Engine/Source/Editor/BlueprintGraph/Private/K2Node_FormatText.cpp
--- a/Engine/Source/Editor/BlueprintGraph/Private/K2Node_FormatText.cpp
+++ b/Engine/Source/Editor/BlueprintGraph/Private/K2Node_FormatText.cpp
@@ -259,14 +259,16 @@ void UK2Node_FormatText::PostReconstructNode()
 				if (CurrentPin != FormatPin && CurrentPin->Direction == EGPD_Input && CurrentPin->LinkedTo.Num() == 0 && !CurrentPin->DefaultTextValue.IsEmpty())
 				{
 					// Create a new "Make Literal Text" function and add it to the graph
-					UK2Node_CallFunction* MakeLiteralText = nullptr;
-					{
-						UK2Node_CallFunction* MakeLiteralTextTemplate = NewObject<UK2Node_CallFunction>(GetGraph());
-						MakeLiteralTextTemplate->SetFromFunction(UKismetSystemLibrary::StaticClass()->FindFunctionByName(GET_MEMBER_NAME_CHECKED(UKismetSystemLibrary, MakeLiteralText)));
-
-						const FVector2D SpawnLocation = FVector2D(NodePosX - 300, NodePosY + (60 * (NumPinsFixedUp + 1)));
-						MakeLiteralText = FEdGraphSchemaAction_K2NewNode::SpawnNodeFromTemplate<UK2Node_CallFunction>(GetGraph(), MakeLiteralTextTemplate, SpawnLocation, /*bSelectNewNode*/false);
-					}
+					const FVector2D SpawnLocation = FVector2D(NodePosX - 300, NodePosY + (60 * (NumPinsFixedUp + 1)));
+					UK2Node_CallFunction* MakeLiteralText = FEdGraphSchemaAction_K2NewNode::SpawnNode<UK2Node_CallFunction>(
+						GetGraph(),
+						SpawnLocation,
+						EK2NewNodeFlags::None,
+						[](UK2Node_CallFunction* NewInstance)
+						{
+							NewInstance->SetFromFunction(UKismetSystemLibrary::StaticClass()->FindFunctionByName(GET_MEMBER_NAME_CHECKED(UKismetSystemLibrary, MakeLiteralText)));
+						}
+					);
 
 					// Set the new value and clear it on this pin to avoid it ever attempting this upgrade again (eg, if the "Make Literal Text" node was disconnected)
 					UEdGraphPin* LiteralValuePin = MakeLiteralText->FindPinChecked(TEXT("Value"));
diff --git a/Engine/Source/Editor/Kismet/Private/BPDelegateDragDropAction.h b/Engine/Source/Editor/Kismet/Private/BPDelegateDragDropAction.h
--- a/Engine/Source/Editor/Kismet/Private/BPDelegateDragDropAction.h
+++ b/Engine/Source/Editor/Kismet/Private/BPDelegateDragDropAction.h
@@ -55,8 +55,15 @@ public:
 	{
 		check(Params.Graph && Params.Property);
 		TNode* Node = NewObject<TNode>();
-		Node->SetFromProperty(Params.Property, Params.bSelfContext);
-		FEdGraphSchemaAction_K2NewNode::SpawnNodeFromTemplate<TNode>(Params.Graph, Node, Params.GraphPosition);
+		FEdGraphSchemaAction_K2NewNode::SpawnNode<TNode>(
+			Params.Graph,
+			Params.GraphPosition,
+			EK2NewNodeFlags::SelectNewNode,
+			[&Params](TNode* NewInstance)
+			{
+				NewInstance->SetFromProperty(Params.Property, Params.bSelfContext);
+			}
+		);
 		Params.AnalyticCallback.ExecuteIfBound();
 	}
 
diff --git a/Engine/Source/Editor/Kismet/Private/BPFunctionDragDropAction.cpp b/Engine/Source/Editor/Kismet/Private/BPFunctionDragDropAction.cpp
--- a/Engine/Source/Editor/Kismet/Private/BPFunctionDragDropAction.cpp
+++ b/Engine/Source/Editor/Kismet/Private/BPFunctionDragDropAction.cpp
@@ -338,15 +338,22 @@ FReply FKismetMacroDragDropAction::DroppedOnPanel(TSharedRef<SWidget> const& Pan
 	check(CanBeDroppedDelegate.IsBound());
 
 	FReply Reply = FReply::Unhandled();
-
 	FText CannotDropReason = FText::GetEmpty();
+	UEdGraph* MacroForCapture = Macro;
+	FNodeCreationAnalytic& AnalyticCallbackForCapture = AnalyticCallback;
+
 	if (CanBeDroppedDelegate.Execute(TSharedPtr<FEdGraphSchemaAction>(), &Graph, CannotDropReason))
 	{
-		UK2Node_MacroInstance* MacroTemplate = NewObject<UK2Node_MacroInstance>();
-		MacroTemplate->SetMacroGraph(Macro);
-		AnalyticCallback.ExecuteIfBound();
-		
-		FEdGraphSchemaAction_K2NewNode::SpawnNodeFromTemplate<UK2Node_MacroInstance>(&Graph, MacroTemplate, GraphPosition);
+		FEdGraphSchemaAction_K2NewNode::SpawnNode<UK2Node_MacroInstance>(
+			&Graph,
+			GraphPosition,
+			EK2NewNodeFlags::SelectNewNode,
+			[MacroForCapture, &AnalyticCallbackForCapture](UK2Node_MacroInstance* NewInstance)
+			{
+				NewInstance->SetMacroGraph(MacroForCapture);
+				AnalyticCallbackForCapture.ExecuteIfBound();
+			}
+		);
 		Reply = FReply::Handled();
 	}
 	return Reply;
diff --git a/Engine/Source/Editor/Kismet/Private/BlueprintEditor.cpp b/Engine/Source/Editor/Kismet/Private/BlueprintEditor.cpp
--- a/Engine/Source/Editor/Kismet/Private/BlueprintEditor.cpp
+++ b/Engine/Source/Editor/Kismet/Private/BlueprintEditor.cpp
@@ -3015,10 +3015,15 @@ void FBlueprintEditor::OnGraphEditorDropActor(const TArray< TWeakObjectPtr<AActo
 			AActor* DroppedActor = Actors[i].Get();
 			if (DroppedActor&& (DroppedActor->GetLevel() == BlueprintLevel) && !DroppedActor->IsChildActor())
 			{
-				UK2Node_Literal* LiteralNodeTemplate = NewObject<UK2Node_Literal>();
-				LiteralNodeTemplate->SetObjectRef(DroppedActor);
-
-				UK2Node_Literal* ActorRefNode = FEdGraphSchemaAction_K2NewNode::SpawnNodeFromTemplate<UK2Node_Literal>(Graph, LiteralNodeTemplate, NodeLocation);
+				UK2Node_Literal* ActorRefNode = FEdGraphSchemaAction_K2NewNode::SpawnNode<UK2Node_Literal>(
+					Graph,
+					NodeLocation,
+					EK2NewNodeFlags::SelectNewNode,
+					[DroppedActor](UK2Node_Literal* NewInstance)
+					{
+						NewInstance->SetObjectRef(DroppedActor);
+					}
+				);
 				NodeLocation.Y += UEdGraphSchema_K2::EstimateNodeHeight(ActorRefNode);
 			}
 		}
@@ -3036,11 +3041,15 @@ void FBlueprintEditor::OnGraphEditorDropStreamingLevel(const TArray< TWeakObject
 		if ((DroppedLevel != NULL) && 
 			(DroppedLevel->IsA(ULevelStreamingKismet::StaticClass()))) 
 		{
-			const FVector2D NodeLocation = DropLocation + (i * FVector2D(0,80));
-				
-			UK2Node_CallFunction* NodeTemplate = NewObject<UK2Node_CallFunction>(Graph);
-			NodeTemplate->SetFromFunction(TargetFunc);
-			UK2Node_CallFunction* Node = FEdGraphSchemaAction_K2NewNode::SpawnNodeFromTemplate<UK2Node_CallFunction>(Graph, NodeTemplate, NodeLocation);
+			UK2Node_CallFunction* Node = FEdGraphSchemaAction_K2NewNode::SpawnNode<UK2Node_CallFunction>(
+				Graph,
+				DropLocation + (i * FVector2D(0, 80)),
+				EK2NewNodeFlags::SelectNewNode,
+				[TargetFunc](UK2Node_CallFunction* NewInstance)
+				{
+					NewInstance->SetFromFunction(TargetFunc);
+				}
+			);
 						
 			// Set dropped level package name
 			UEdGraphPin* PackageNameInputPin = Node->FindPinChecked(TEXT("PackageName"));
@@ -6699,8 +6708,7 @@ void FBlueprintEditor::CollapseNodes(TSet<UEdGraphNode*>& InCollapsableNodes)
 	// Create the composite node that will serve as the gateway into the subgraph
 	UK2Node_Composite* GatewayNode = NULL;
 	{
-		UK2Node_Composite* TemplateNode = NewObject<UK2Node_Composite>();
-		GatewayNode = FEdGraphSchemaAction_K2NewNode::SpawnNodeFromTemplate<UK2Node_Composite>(SourceGraph, TemplateNode, FVector2D(0,0));
+		GatewayNode = FEdGraphSchemaAction_K2NewNode::SpawnNode<UK2Node_Composite>(SourceGraph, FVector2D(0,0), EK2NewNodeFlags::SelectNewNode);
 		GatewayNode->bCanRenameNode = true;
 		check(GatewayNode);
 	}
@@ -6789,10 +6797,15 @@ UEdGraph* FBlueprintEditor::CollapseSelectionToMacro(TSharedPtr<SGraphEditor> In
 	DestinationGraph = FBlueprintEditorUtils::CreateNewGraph(GetBlueprintObj(), DocumentName, UEdGraph::StaticClass(), UEdGraphSchema_K2::StaticClass());
 	FBlueprintEditorUtils::AddMacroGraph(GetBlueprintObj(), DestinationGraph, /*bIsUserCreated=*/ true, NULL);
 
-	UK2Node_MacroInstance* MacroTemplate = NewObject<UK2Node_MacroInstance>();
-	MacroTemplate->SetMacroGraph(DestinationGraph);
-
-	UK2Node_MacroInstance* GatewayNode = FEdGraphSchemaAction_K2NewNode::SpawnNodeFromTemplate<UK2Node_MacroInstance>(SourceGraph, MacroTemplate, FVector2D(0.0f, 0.0f), false);
+	UK2Node_MacroInstance* GatewayNode = FEdGraphSchemaAction_K2NewNode::SpawnNode<UK2Node_MacroInstance>(
+		SourceGraph,
+		FVector2D(0.0f, 0.0f),
+		EK2NewNodeFlags::None,
+		[DestinationGraph](UK2Node_MacroInstance* NewInstance)
+		{
+			NewInstance->SetMacroGraph(DestinationGraph);
+		}
+	);
 
 	TArray<UK2Node_Tunnel*> TunnelNodes;
 	GatewayNode->GetMacroGraph()->GetNodesOfClass(TunnelNodes);
diff --git a/Engine/Source/Editor/Kismet/Private/SMyBlueprint.cpp b/Engine/Source/Editor/Kismet/Private/SMyBlueprint.cpp
--- a/Engine/Source/Editor/Kismet/Private/SMyBlueprint.cpp
+++ b/Engine/Source/Editor/Kismet/Private/SMyBlueprint.cpp
@@ -2082,12 +2082,16 @@ void SMyBlueprint::ImplementFunction(FEdGraphSchemaAction_K2Graph* GraphAction)
 		}
 		else
 		{
-			UK2Node_Event* NewEventNodeTemplate = NewObject<UK2Node_Event>();
-			NewEventNodeTemplate->EventReference.SetExternalMember(EventName, OverrideFuncClass);
-			NewEventNodeTemplate->bOverrideFunction = true;
-
-			const FVector2D NewNodePos = EventGraph->GetGoodPlaceForNewNode();
-			UK2Node_Event* NewEventNode = FEdGraphSchemaAction_K2NewNode::SpawnNodeFromTemplate<UK2Node_Event>(EventGraph, NewEventNodeTemplate, NewNodePos);
+			UK2Node_Event* NewEventNode = FEdGraphSchemaAction_K2NewNode::SpawnNode<UK2Node_Event>(
+				EventGraph,
+				EventGraph->GetGoodPlaceForNewNode(),
+				EK2NewNodeFlags::SelectNewNode,
+				[EventName, OverrideFuncClass](UK2Node_Event* NewInstance)
+				{
+					NewInstance->EventReference.SetExternalMember(EventName, OverrideFuncClass);
+					NewInstance->bOverrideFunction = true;
+				}
+			);
 			if (NewEventNode)
 			{
 				FKismetEditorUtilities::BringKismetToFocusAttentionOnObject(NewEventNode);
diff --git a/Engine/Source/Editor/UnrealEd/Private/Kismet2/Kismet2.cpp b/Engine/Source/Editor/UnrealEd/Private/Kismet2/Kismet2.cpp
--- a/Engine/Source/Editor/UnrealEd/Private/Kismet2/Kismet2.cpp
+++ b/Engine/Source/Editor/UnrealEd/Private/Kismet2/Kismet2.cpp
@@ -2012,11 +2012,17 @@ void FKismetEditorUtilities::CreateNewBoundEventForActor(AActor* Actor, FName Ev
 					const FVector2D NewNodePos = TargetGraph->GetGoodPlaceForNewNode();
 
 					// Create a new event node
-					UK2Node_ActorBoundEvent* EventNodeTemplate = NewObject<UK2Node_ActorBoundEvent>();
-					EventNodeTemplate->InitializeActorBoundEventParams(Actor, DelegateProperty);
-
+					UK2Node_ActorBoundEvent* EventNode = FEdGraphSchemaAction_K2NewNode::SpawnNode<UK2Node_ActorBoundEvent>(
+						TargetGraph,
+						NewNodePos,
+						EK2NewNodeFlags::SelectNewNode,
+						[Actor, DelegateProperty](UK2Node_ActorBoundEvent* NewInstance)
+						{
+							NewInstance->InitializeActorBoundEventParams(Actor, DelegateProperty);
+						}
+					);
 					// Finally, bring up kismet and jump to the new node
-					if (UK2Node_ActorBoundEvent* EventNode = FEdGraphSchemaAction_K2NewNode::SpawnNodeFromTemplate<UK2Node_ActorBoundEvent>(TargetGraph, EventNodeTemplate, NewNodePos))
+					if (EventNode)
 					{
 						BringKismetToFocusAttentionOnObject(EventNode);
 					}
@@ -2050,10 +2056,15 @@ void FKismetEditorUtilities::CreateNewBoundEventForClass(UClass* Class, FName Ev
 				const FVector2D NewNodePos = TargetGraph->GetGoodPlaceForNewNode();
 
 				// Create a new event node
-				UK2Node_ComponentBoundEvent* EventNodeTemplate = NewObject<UK2Node_ComponentBoundEvent>();
-				EventNodeTemplate->InitializeComponentBoundEventParams(ComponentProperty, DelegateProperty);
-
-				UK2Node_ComponentBoundEvent* EventNode = FEdGraphSchemaAction_K2NewNode::SpawnNodeFromTemplate<UK2Node_ComponentBoundEvent>(TargetGraph, EventNodeTemplate, NewNodePos);
+				UK2Node_ComponentBoundEvent* EventNode = FEdGraphSchemaAction_K2NewNode::SpawnNode<UK2Node_ComponentBoundEvent>(
+					TargetGraph,
+					NewNodePos,
+					EK2NewNodeFlags::SelectNewNode,
+					[ComponentProperty, DelegateProperty](UK2Node_ComponentBoundEvent* NewInstance)
+					{
+						NewInstance->InitializeComponentBoundEventParams(ComponentProperty, DelegateProperty);
+					}
+				);
 
 				// Finally, bring up kismet and jump to the new node
 				if ( EventNode != nullptr )
