Mirror 4081417 from Dev-Framework
Enum literals are now more similar to byte literals, which  makes them safer to pass by reference

diff --git a/Engine/Source/Editor/BlueprintGraph/Classes/K2Node_EnumLiteral.h b/Engine/Source/Editor/BlueprintGraph/Classes/K2Node_EnumLiteral.h
--- a/Engine/Source/Editor/BlueprintGraph/Classes/K2Node_EnumLiteral.h
+++ b/Engine/Source/Editor/BlueprintGraph/Classes/K2Node_EnumLiteral.h
@@ -34,7 +34,7 @@ class UK2Node_EnumLiteral : public UK2Node, public INodeDependingOnEnumInterface
 
 	//~ Begin UK2Node Interface
 	virtual bool IsNodePure() const override { return true; }
-	virtual FNodeHandlingFunctor* CreateNodeHandler(FKismetCompilerContext& CompilerContext) const override;
+	virtual void ExpandNode(class FKismetCompilerContext& CompilerContext, UEdGraph* SourceGraph) override;
 	virtual void GetMenuActions(FBlueprintActionDatabaseRegistrar& ActionRegistrar) const override;
 	virtual FText GetMenuCategory() const override;
 	//~ End UK2Node Interface
diff --git a/Engine/Source/Editor/BlueprintGraph/Private/K2Node_EnumLiteral.cpp b/Engine/Source/Editor/BlueprintGraph/Private/K2Node_EnumLiteral.cpp
--- a/Engine/Source/Editor/BlueprintGraph/Private/K2Node_EnumLiteral.cpp
+++ b/Engine/Source/Editor/BlueprintGraph/Private/K2Node_EnumLiteral.cpp
@@ -2,8 +2,12 @@
 
 
 #include "K2Node_EnumLiteral.h"
+#include "K2Node_CastByteToEnum.h"
 #include "EdGraphSchema_K2.h"
 #include "EdGraphUtilities.h"
+#include "K2Node_CallFunction.h"
+#include "Kismet/KismetSystemLibrary.h"
+#include "KismetCompiler.h"
 #include "KismetCompiledFunctionContext.h"
 #include "KismetCompilerMisc.h"
 #include "BlueprintFieldNodeSpawner.h"
@@ -72,41 +76,45 @@ FText UK2Node_EnumLiteral::GetNodeTitle(ENodeTitleType::Type TitleType) const
 	return CachedTooltip;
 }
 
-class FKCHandler_EnumLiteral : public FNodeHandlingFunctor
+void UK2Node_EnumLiteral::ExpandNode(class FKismetCompilerContext& CompilerContext, UEdGraph* SourceGraph)
 {
-public:
-	FKCHandler_EnumLiteral(FKismetCompilerContext& InCompilerContext)
-		: FNodeHandlingFunctor(InCompilerContext)
+	bool bSuccess = Enum != nullptr;
+	if (bSuccess)
 	{
-	}
-
-	virtual void RegisterNets(FKismetFunctionContext& Context, UEdGraphNode* Node) override
-	{
-		check(Context.Schema && Cast<UK2Node_EnumLiteral>(Node));
-		FNodeHandlingFunctor::RegisterNets(Context, Node);
-
-		UEdGraphPin* InPin = Node->FindPinChecked(UK2Node_EnumLiteral::GetEnumInputPinName());
-		UEdGraphPin* Net = FEdGraphUtilities::GetNetFromPin(InPin);
-		if (Context.NetMap.Find(Net) == NULL)
+		const FName FunctionName = GET_FUNCTION_NAME_CHECKED(UKismetSystemLibrary, MakeLiteralByte);
+		UK2Node_CallFunction* MakeLiteralByteNode = CompilerContext.SpawnIntermediateNode<UK2Node_CallFunction>(this, SourceGraph);
+		MakeLiteralByteNode->SetFromFunction(UKismetSystemLibrary::StaticClass()->FindFunctionByName(FunctionName));
+		MakeLiteralByteNode->AllocateDefaultPins();
+
+		UEdGraphPin* OrgInputPin = FindPinChecked(GetEnumInputPinName());
+		UEdGraphPin* NewInputPin = MakeLiteralByteNode->FindPinChecked(TEXT("Value"));
+
+		// read the enum value, convert to an integer:
+		int32 Index = Enum->GetIndexByName(*OrgInputPin->DefaultValue);
+		bSuccess = Index != INDEX_NONE;
+		if (bSuccess)
 		{
-			FBPTerminal* Term = Context.CreateLocalTerminalFromPinAutoChooseScope(Net, Context.NetNameMap->MakeValidName(Net));
-			Context.NetMap.Add(Net, Term);
+			NewInputPin->DefaultValue = FString::FromInt(Enum->GetValueByIndex(Index));
 		}
 
-		FBPTerminal** ValueSource = Context.NetMap.Find(Net);
-		check(ValueSource && *ValueSource);
-		UEdGraphPin* OutPin = Node->FindPinChecked(UEdGraphSchema_K2::PN_ReturnValue);
-		if (ensure(Context.NetMap.Find(OutPin) == NULL))
-		{
-			FBPTerminal* TerminalPtr = *ValueSource; //necessary because of CheckAddress in Map::Add
-			Context.NetMap.Add(OutPin, TerminalPtr);
-		}
+		UK2Node_CastByteToEnum* CastByte = CompilerContext.SpawnIntermediateNode<UK2Node_CastByteToEnum>(this, SourceGraph);
+		CastByte->Enum = Enum;
+		CastByte->AllocateDefaultPins();
+
+		UEdGraphPin* CastInputPin = CastByte->FindPinChecked(UK2Node_CastByteToEnum::ByteInputPinName);
+		UEdGraphPin* NewReturnPin = MakeLiteralByteNode->GetReturnValuePin();
+		check(nullptr != NewReturnPin);
+		bSuccess &= CompilerContext.GetSchema()->TryCreateConnection(NewReturnPin, CastInputPin);
+
+		UEdGraphPin* OrgReturnPin = FindPinChecked(UEdGraphSchema_K2::PN_ReturnValue);
+		UEdGraphPin* CastReturnPin = CastByte->FindPinChecked(UEdGraphSchema_K2::PN_ReturnValue);
+		CompilerContext.MovePinLinksToIntermediate(*OrgReturnPin, *CastReturnPin);
 	}
-};
 
-FNodeHandlingFunctor* UK2Node_EnumLiteral::CreateNodeHandler(FKismetCompilerContext& CompilerContext) const
-{
-	return new FKCHandler_EnumLiteral(CompilerContext);
+	if (!bSuccess)
+	{
+		CompilerContext.MessageLog.Error(*NSLOCTEXT("K2Node", "EnumLiteral_NullEnumError", "Undefined Enum in @@").ToString(), this);
+	}
 }
 
 void UK2Node_EnumLiteral::GetMenuActions(FBlueprintActionDatabaseRegistrar& ActionRegistrar) const
