Merge 4083850 from Dev-Framework
Prevent GameplayAbilitiesEditor from creating FBluerpintActionDatabase when not neccessary, improve performance of FBlueprintActionDatabase::RegisterAllNodeAction

diff --git a/Engine/Plugins/Runtime/GameplayAbilities/Source/GameplayAbilitiesEditor/Private/GameplayAbilitiesEditorModule.cpp b/Engine/Plugins/Runtime/GameplayAbilities/Source/GameplayAbilitiesEditor/Private/GameplayAbilitiesEditorModule.cpp
--- a/Engine/Plugins/Runtime/GameplayAbilities/Source/GameplayAbilitiesEditor/Private/GameplayAbilitiesEditorModule.cpp
+++ b/Engine/Plugins/Runtime/GameplayAbilities/Source/GameplayAbilitiesEditor/Private/GameplayAbilitiesEditorModule.cpp
@@ -248,7 +248,10 @@ void FGameplayAbilitiesEditorModule::GameplayTagTreeChanged()
 	FString PerfMessage = FString::Printf(TEXT("FGameplayAbilitiesEditorModule::GameplayTagTreeChanged"));
 	SCOPE_LOG_TIME_IN_SECONDS(*PerfMessage, nullptr)
 #endif
-	FBlueprintActionDatabase::Get().RefreshClassActions(UK2Node_GameplayCueEvent::StaticClass());
+	if (FBlueprintActionDatabase* BAD = FBlueprintActionDatabase::TryGet())
+	{
+		BAD->RefreshClassActions(UK2Node_GameplayCueEvent::StaticClass());
+	}
 }
 
 void FGameplayAbilitiesEditorModule::ShutdownModule()
diff --git a/Engine/Source/Editor/BlueprintGraph/Private/BlueprintActionDatabase.cpp b/Engine/Source/Editor/BlueprintGraph/Private/BlueprintActionDatabase.cpp
--- a/Engine/Source/Editor/BlueprintGraph/Private/BlueprintActionDatabase.cpp
+++ b/Engine/Source/Editor/BlueprintGraph/Private/BlueprintActionDatabase.cpp
@@ -1011,10 +1011,11 @@ static bool BlueprintActionDatabaseImpl::IsObjectValidForDatabase(UObject const*
  * FBlueprintActionDatabase
  ******************************************************************************/
 
+static FBlueprintActionDatabase* DatabaseInst = nullptr;
+
 //------------------------------------------------------------------------------
 FBlueprintActionDatabase& FBlueprintActionDatabase::Get()
 {
-	static FBlueprintActionDatabase* DatabaseInst = nullptr;
 	if (DatabaseInst == nullptr)
 	{
 		DatabaseInst = new FBlueprintActionDatabase();
@@ -1023,6 +1024,12 @@ FBlueprintActionDatabase& FBlueprintActionDatabase::Get()
 	return *DatabaseInst;
 }
 
+//------------------------------------------------------------------------------
+FBlueprintActionDatabase* FBlueprintActionDatabase::TryGet()
+{
+	return DatabaseInst;
+}
+
 //------------------------------------------------------------------------------
 FBlueprintActionDatabase::FBlueprintActionDatabase()
 {
@@ -1527,14 +1534,11 @@ FBlueprintActionDatabase::FActionRegistry const& FBlueprintActionDatabase::GetAl
 void FBlueprintActionDatabase::RegisterAllNodeActions(FBlueprintActionDatabaseRegistrar& Registrar)
 {
 	// nodes may have actions they wish to add for this asset
-	for (TObjectIterator<UClass> ClassIt; ClassIt; ++ClassIt)
-	{
-		UClass* NodeClass = *ClassIt;
-		if (!NodeClass->IsChildOf<UK2Node>())
-		{
-			continue;
-		}
+	TArray<UClass*> NodeClassList;
+	GetDerivedClasses(UK2Node::StaticClass(), NodeClassList);
 
+	for (UClass* NodeClass : NodeClassList)
+	{
 		TGuardValue< TSubclassOf<UEdGraphNode> > ScopedNodeClass(Registrar.GeneratingClass, NodeClass);
 		BlueprintActionDatabaseImpl::GetNodeSpecificActions(NodeClass, Registrar);
 	}
diff --git a/Engine/Source/Editor/BlueprintGraph/Public/BlueprintActionDatabase.h b/Engine/Source/Editor/BlueprintGraph/Public/BlueprintActionDatabase.h
--- a/Engine/Source/Editor/BlueprintGraph/Public/BlueprintActionDatabase.h
+++ b/Engine/Source/Editor/BlueprintGraph/Public/BlueprintActionDatabase.h
@@ -41,6 +41,9 @@ public:
 	 */
 	static FBlueprintActionDatabase& Get();
 
+	/** Getter to access the datbase singleton, will return null if the database has not been initialized */
+	static FBlueprintActionDatabase* TryGet();
+
 	// FTickableEditorObject interface
 	virtual void Tick(float DeltaTime) override;
 	virtual ETickableTickType GetTickableTickType() const override { return ETickableTickType::Always; }
