SA fix, remove unneeded null check
#rnx

diff --git a/Engine/Source/Editor/BlueprintGraph/Private/EdGraphSchema_K2_Actions.cpp b/Engine/Source/Editor/BlueprintGraph/Private/EdGraphSchema_K2_Actions.cpp
--- a/Engine/Source/Editor/BlueprintGraph/Private/EdGraphSchema_K2_Actions.cpp
+++ b/Engine/Source/Editor/BlueprintGraph/Private/EdGraphSchema_K2_Actions.cpp
@@ -284,7 +284,7 @@ UEdGraphNode* FEdGraphSchemaAction_K2NewNode::CreateNode(
 	ResultNode->ErrorMsg.Empty();
 	ResultNode->bHasCompilerMessage = false;
 
-	if (bGotoNode && ResultNode)
+	if (bGotoNode)
 	{
 		// Select existing node
 		FKismetEditorUtilities::BringKismetToFocusAttentionOnObject(ResultNode);
