Try and find the full path name for default value of class parameters

diff --git a/Engine/Source/Editor/BlueprintGraph/Private/EdGraphSchema_K2.cpp b/Engine/Source/Editor/BlueprintGraph/Private/EdGraphSchema_K2.cpp
--- a/Engine/Source/Editor/BlueprintGraph/Private/EdGraphSchema_K2.cpp
+++ b/Engine/Source/Editor/BlueprintGraph/Private/EdGraphSchema_K2.cpp
@@ -4798,6 +4798,15 @@ bool UEdGraphSchema_K2::FindFunctionParameterDefaultValue(const UFunction* Funct
 		// Specified default value in the metadata
 		OutString = MetadataDefaultValue;
 		bHasAutomaticValue = true;
+
+		// If the parameter is a class then try and get the full name as the metadata might just be the short name
+		if (Param->IsA<UClassProperty>() && !FPackageName::IsValidObjectPath(OutString))
+		{
+			if (UClass* DefaultClass = FindObject<UClass>(ANY_PACKAGE, *OutString, true))
+			{
+				OutString = DefaultClass->GetPathName();
+			}
+		}
 	}
 	else
 	{
