Fix shadow variable
#rnx

diff --git a/Engine/Source/Editor/BlueprintGraph/Classes/EdGraphSchema_K2_Actions.h b/Engine/Source/Editor/BlueprintGraph/Classes/EdGraphSchema_K2_Actions.h
--- a/Engine/Source/Editor/BlueprintGraph/Classes/EdGraphSchema_K2_Actions.h
+++ b/Engine/Source/Editor/BlueprintGraph/Classes/EdGraphSchema_K2_Actions.h
@@ -83,9 +83,9 @@ struct BLUEPRINTGRAPH_API FEdGraphSchemaAction_K2NewNode : public FEdGraphSchema
 			ParentGraph,
 			TArrayView<UEdGraphPin*>(),
 			Location,
-			[NodeClass](UEdGraph* ParentGraph)->UK2Node*
+			[NodeClass](UEdGraph* InParentGraph)->UK2Node*
 			{
-				return NewObject<UK2Node>(ParentGraph, NodeClass);
+				return NewObject<UK2Node>(InParentGraph, NodeClass);
 			},
 			[InitializerFn](UK2Node* NewNode)
 			{
