Fix for crash on adding OnClicked event to button in Widget BP

diff --git a/Engine/Source/Editor/BlueprintGraph/Private/EdGraphSchema_K2_Actions.cpp b/Engine/Source/Editor/BlueprintGraph/Private/EdGraphSchema_K2_Actions.cpp
--- a/Engine/Source/Editor/BlueprintGraph/Private/EdGraphSchema_K2_Actions.cpp
+++ b/Engine/Source/Editor/BlueprintGraph/Private/EdGraphSchema_K2_Actions.cpp
@@ -221,7 +221,10 @@ UEdGraphNode* FEdGraphSchemaAction_K2NewNode::CreateNode(
 	ParentGraph->Modify();
 	for(UEdGraphPin* FromPin : FromPins)
 	{
-		FromPin->Modify();
+		if (FromPin != nullptr)
+		{
+			FromPin->Modify();
+		}
 	}
 
 	// Smart pointer that handles fixup after potential node reconstruction
