Fix shadow variable warning
#rnx

diff --git a/Engine/Source/Editor/BlueprintGraph/Private/EdGraphSchema_K2_Actions.cpp b/Engine/Source/Editor/BlueprintGraph/Private/EdGraphSchema_K2_Actions.cpp
--- a/Engine/Source/Editor/BlueprintGraph/Private/EdGraphSchema_K2_Actions.cpp
+++ b/Engine/Source/Editor/BlueprintGraph/Private/EdGraphSchema_K2_Actions.cpp
@@ -195,9 +195,9 @@ UEdGraphNode* FEdGraphSchemaAction_K2NewNode::CreateNode(
 			ParentGraph,
 			FromPins,
 			Location,
-			[NodeTemplate](UEdGraph* ParentGraph)->UK2Node*
+			[NodeTemplate](UEdGraph* InParentGraph)->UK2Node*
 			{
-				return DuplicateObject<UK2Node>(NodeTemplate, ParentGraph);
+				return DuplicateObject<UK2Node>(NodeTemplate, InParentGraph);
 			},
 			[](UK2Node*) {},
 			Options
