[NO NON-PII COMMIT MESSAGE]

diff --git a/Engine/Source/Editor/BlueprintGraph/Private/K2Node.cpp b/Engine/Source/Editor/BlueprintGraph/Private/K2Node.cpp
--- a/Engine/Source/Editor/BlueprintGraph/Private/K2Node.cpp
+++ b/Engine/Source/Editor/BlueprintGraph/Private/K2Node.cpp
@@ -182,16 +182,21 @@ void UK2Node::FixupPinDefaultValues()
 
 FText UK2Node::GetToolTipHeading() const
 {
+	UBlueprint* Blueprint = FBlueprintEditorUtils::FindBlueprintForNode(this);
 	FText Heading = FText::GetEmpty();
-	if (UBreakpoint* ExistingBreakpoint = FKismetDebugUtilities::FindBreakpointForNode(GetBlueprint(), this))
+
+	if (Blueprint)
 	{
-		if (ExistingBreakpoint->IsEnabled())
-		{
-			Heading = LOCTEXT("EnabledBreakpoint", "Active Breakpoint - Execution will break at this location.");
-		}
-		else 
+		if (UBreakpoint* ExistingBreakpoint = FKismetDebugUtilities::FindBreakpointForNode(Blueprint, this))
 		{
-			Heading = LOCTEXT("DisabledBreakpoint", "Disabled Breakpoint");
+			if (ExistingBreakpoint->IsEnabled())
+			{
+				Heading = LOCTEXT("EnabledBreakpoint", "Active Breakpoint - Execution will break at this location.");
+			}
+			else
+			{
+				Heading = LOCTEXT("DisabledBreakpoint", "Disabled Breakpoint");
+			}
 		}
 	}
 
@@ -1544,12 +1549,17 @@ void UK2Node::GetPinHoverText(const UEdGraphPin& Pin, FString& HoverTextOut) con
 	if (HoverTextOut.IsEmpty())
 	{
 		UEdGraphSchema const* Schema = GetSchema();
-		check(Schema != NULL);
+		check(Schema != nullptr);
 		Schema->ConstructBasicPinTooltip(Pin, FText::GetEmpty(), HoverTextOut);
 	}	
 
-	UBlueprint* Blueprint = FBlueprintEditorUtils::FindBlueprintForNodeChecked(this);
-	check(Blueprint != NULL);
+	UBlueprint* Blueprint = FBlueprintEditorUtils::FindBlueprintForNode(this);
+
+	// If this is an orphaned node, just end now
+	if (Blueprint == nullptr)
+	{
+		return;
+	}
 
 	// If this resides in an intermediate graph, show the UObject name for debug purposes
 	if(Blueprint->IntermediateGeneratedGraphs.Contains(GetGraph()))
@@ -1559,14 +1569,14 @@ void UK2Node::GetPinHoverText(const UEdGraphPin& Pin, FString& HoverTextOut) con
 
 	UObject* ActiveObject = Blueprint->GetObjectBeingDebugged();
 	// if there is no object being debugged, then we don't need to tack on any of the following data
-	if (ActiveObject == NULL)
+	if (ActiveObject == nullptr)
 	{
 		return;
 	}
 
 	// Switch the blueprint to the one that generated the object being debugged (e.g. in case we're inside a Macro BP while debugging)
 	Blueprint = Cast<UBlueprint>(ActiveObject->GetClass()->ClassGeneratedBy);
-	if (Blueprint == NULL)
+	if (Blueprint == nullptr)
 	{
 		return;
 	}
