Add warning to detect default value changes when assets were saved in earlier versions of the editor, which may have had bad default value information on some call function nodes

diff --git a/Engine/Source/Editor/BlueprintGraph/Classes/K2Node.h b/Engine/Source/Editor/BlueprintGraph/Classes/K2Node.h
--- a/Engine/Source/Editor/BlueprintGraph/Classes/K2Node.h
+++ b/Engine/Source/Editor/BlueprintGraph/Classes/K2Node.h
@@ -409,7 +409,7 @@ protected:
 	BLUEPRINTGRAPH_API void ReconstructSinglePin(UEdGraphPin* NewPin, UEdGraphPin* OldPin, ERedirectType RedirectType);
 
 	// Helper function to rewire old pins to new pins during node reconstruction (or other regeneration of pins)
-	BLUEPRINTGRAPH_API void RewireOldPinsToNewPins(TArray<UEdGraphPin*>& InOldPins, TArray<UEdGraphPin*>& InNewPins);
+	BLUEPRINTGRAPH_API void RewireOldPinsToNewPins(TArray<UEdGraphPin*>& InOldPins, TArray<UEdGraphPin*>& InNewPins, TMap<UEdGraphPin*, UEdGraphPin*>* NewPinToOldPin);
 
 	// Helper function to properly destroy a set of pins
 	BLUEPRINTGRAPH_API void DestroyPinList(TArray<UEdGraphPin*>& InPins);
diff --git a/Engine/Source/Editor/BlueprintGraph/Private/K2Node.cpp b/Engine/Source/Editor/BlueprintGraph/Private/K2Node.cpp
--- a/Engine/Source/Editor/BlueprintGraph/Private/K2Node.cpp
+++ b/Engine/Source/Editor/BlueprintGraph/Private/K2Node.cpp
@@ -576,12 +576,35 @@ void UK2Node::ReconstructNode()
 	Pins.Reset();
 
 	// Recreate the new pins
+	TMap<UEdGraphPin*, UEdGraphPin*> NewPinsToOldPins;
 	ReallocatePinsDuringReconstruction(OldPins);
-	RewireOldPinsToNewPins(OldPins, Pins);
+	RewireOldPinsToNewPins(OldPins, Pins, &NewPinsToOldPins);
 
 	// Let subclasses do any additional work
 	PostReconstructNode();
 
+	if (Cast<UK2Node_CallFunction>(this) &&
+		Blueprint->CurrentMessageLog &&
+		GetLinkerCustomVersion(FReleaseObjectVersion::GUID) < FReleaseObjectVersion::PinDefaultValuesVerified)
+	{
+		for (UEdGraphPin* NewPin : Pins)
+		{
+			UEdGraphPin** OldPinEntry = NewPinsToOldPins.Find(NewPin);
+			if (OldPinEntry && *OldPinEntry)
+			{
+				UEdGraphPin* OldPin = *OldPinEntry;
+				if (NewPin->LinkedTo.Num() == 0 &&
+					!OldPin->AutogeneratedDefaultValue.IsEmpty() &&
+					NewPin->AutogeneratedDefaultValue != OldPin->AutogeneratedDefaultValue &&
+					OldPin->DefaultValue == OldPin->AutogeneratedDefaultValue &&
+					NewPin->DefaultValue != OldPin->DefaultValue)
+				{
+					Blueprint->CurrentMessageLog->Warning(*LOCTEXT("VerifyDefaultValues", "Default value for @@ on @@ has changed and this asset is from a version that may have had incorrect default value information - verify and resave").ToString(), NewPin, this);
+				}
+			}
+		}
+	}
+
 	GetGraph()->NotifyGraphChanged();
 }
 
@@ -652,20 +675,7 @@ UK2Node::ERedirectType UK2Node::DoPinsMatchForReconstruction(const UEdGraphPin*
 	// if the pin names do match
 	if (NewPin->PinName == OldPin->PinName)
 	{
-		// Make sure we're not dealing with a menu node
-		UEdGraph* OuterGraph = GetGraph();
-		if( OuterGraph && OuterGraph->Schema )
-		{
-			const UEdGraphSchema_K2* K2Schema = Cast<const UEdGraphSchema_K2>(GetSchema());
-			if (!K2Schema || GetBlueprint()->bIsRegeneratingOnLoad || K2Schema->IsSelfPin(*NewPin) || K2Schema->ArePinTypesCompatible(OldPin->PinType, NewPin->PinType) )
-			{
-				RedirectType = ERedirectType_Name;
-			}
-			else
-			{
-				RedirectType = ERedirectType_None;
-			}
-		}
+		RedirectType = ERedirectType_Name;
 	}
 	else
 	{
@@ -874,7 +884,7 @@ FString UK2Node::GetPinMetaData(FName InPinName, FName InKey)
 	return FString();
 }
 
-void UK2Node::RewireOldPinsToNewPins(TArray<UEdGraphPin*>& InOldPins, TArray<UEdGraphPin*>& InNewPins)
+void UK2Node::RewireOldPinsToNewPins(TArray<UEdGraphPin*>& InOldPins, TArray<UEdGraphPin*>& InNewPins, TMap<UEdGraphPin*, UEdGraphPin*>* NewPinToOldPin)
 {
 	TArray<UEdGraphPin*> OrphanedOldPins;
 	TMap<UEdGraphPin*, UEdGraphPin*> MatchedPins; // Old to New
@@ -907,6 +917,10 @@ void UK2Node::RewireOldPinsToNewPins(TArray<UEdGraphPin*>& InOldPins, TArray<UEd
 			{
 				ReconstructSinglePin(NewPin, OldPin, RedirectType);
 				MatchedPins.Add(OldPin, NewPin);
+				if (NewPinToOldPin)
+				{
+					NewPinToOldPin->Add(NewPin, OldPin);
+				}
 				bMatched = true;
 				break;
 			}
diff --git a/Engine/Source/Editor/BlueprintGraph/Private/K2Node_ConstructObjectFromClass.cpp b/Engine/Source/Editor/BlueprintGraph/Private/K2Node_ConstructObjectFromClass.cpp
--- a/Engine/Source/Editor/BlueprintGraph/Private/K2Node_ConstructObjectFromClass.cpp
+++ b/Engine/Source/Editor/BlueprintGraph/Private/K2Node_ConstructObjectFromClass.cpp
@@ -226,7 +226,7 @@ void UK2Node_ConstructObjectFromClass::OnClassPinChanged()
 	}
 
 	// Rewire the old pins to the new pins so connections are maintained if possible
-	RewireOldPinsToNewPins(OldClassPins, Pins);
+	RewireOldPinsToNewPins(OldClassPins, Pins, nullptr);
 
 	// Refresh the UI for the graph so the pin changes show up
 	GetGraph()->NotifyGraphChanged();
diff --git a/Engine/Source/Editor/BlueprintGraph/Private/K2Node_GetClassDefaults.cpp b/Engine/Source/Editor/BlueprintGraph/Private/K2Node_GetClassDefaults.cpp
--- a/Engine/Source/Editor/BlueprintGraph/Private/K2Node_GetClassDefaults.cpp
+++ b/Engine/Source/Editor/BlueprintGraph/Private/K2Node_GetClassDefaults.cpp
@@ -562,7 +562,7 @@ void UK2Node_GetClassDefaults::OnClassPinChanged()
 
 	// Restore the split pins and reconnect output pins to previously linked inputs
 	RestoreSplitPins(OldOutputPins);
-	RewireOldPinsToNewPins(OldOutputPins, Pins);
+	RewireOldPinsToNewPins(OldOutputPins, Pins, nullptr);
 
 	// Notify the graph that the node has been changed
 	if(UEdGraph* Graph = GetGraph())
diff --git a/Engine/Source/Editor/KismetCompiler/Private/KismetCompiler.cpp b/Engine/Source/Editor/KismetCompiler/Private/KismetCompiler.cpp
--- a/Engine/Source/Editor/KismetCompiler/Private/KismetCompiler.cpp
+++ b/Engine/Source/Editor/KismetCompiler/Private/KismetCompiler.cpp
@@ -2948,6 +2948,7 @@ void FKismetCompilerContext::ExpansionStep(UEdGraph* Graph, bool bAllowUbergraph
 	{
 		BP_SCOPED_COMPILER_EVENT_STAT(EKismetCompilerStats_Expansion);
 
+		// First we need to expand knot nodes, so it will remote disconnected knots
 		// Collapse any remaining tunnels or macros
 		ExpandTunnelsAndMacros(Graph);
 
@@ -3813,7 +3814,7 @@ void FKismetCompilerContext::CompileClassLayout(EInternalCompilerFlags InternalF
 		}
 	}
 
-	if (CompileOptions.DoesRequireBytecodeGeneration())
+	if (CompileOptions.DoesRequireBytecodeGeneration() && !Blueprint->bIsRegeneratingOnLoad)
 	{
 		TArray<UEdGraph*> AllGraphs;
 		Blueprint->GetAllGraphs(AllGraphs);
diff --git a/Engine/Source/Runtime/Core/Public/UObject/ReleaseObjectVersion.h b/Engine/Source/Runtime/Core/Public/UObject/ReleaseObjectVersion.h
--- a/Engine/Source/Runtime/Core/Public/UObject/ReleaseObjectVersion.h
+++ b/Engine/Source/Runtime/Core/Public/UObject/ReleaseObjectVersion.h
@@ -55,6 +55,9 @@ struct CORE_API FReleaseObjectVersion
 		// HISMC Cluster Tree migration to add new data
 		HISMCClusterTreeMigration,
 
+		// Default values on pins in blueprints could be saved incoherently
+		PinDefaultValuesVerified,
+
 		// -----<new versions can be added above this line>-------------------------------------------------
 		VersionPlusOne,
 		LatestVersion = VersionPlusOne - 1
