Pins with default values must be orphaned, so that users can move data to new type if possible

diff --git a/Engine/Source/Editor/BlueprintGraph/Private/K2Node.cpp b/Engine/Source/Editor/BlueprintGraph/Private/K2Node.cpp
--- a/Engine/Source/Editor/BlueprintGraph/Private/K2Node.cpp
+++ b/Engine/Source/Editor/BlueprintGraph/Private/K2Node.cpp
@@ -675,7 +675,20 @@ UK2Node::ERedirectType UK2Node::DoPinsMatchForReconstruction(const UEdGraphPin*
 	// if the pin names do match
 	if (NewPin->PinName == OldPin->PinName)
 	{
-		RedirectType = ERedirectType_Name;
+		// If the old pin had a default value, only match the new pin if the type is compatible:
+		const UEdGraphSchema_K2* K2Schema = Cast<const UEdGraphSchema_K2>(GetSchema());
+		if (K2Schema &&
+			OldPin->Direction == EGPD_Input &&
+			OldPin->LinkedTo.Num() == 0 &&
+			OldPin->DefaultValue != OldPin->AutogeneratedDefaultValue &&
+			!K2Schema->ArePinTypesCompatible(OldPin->PinType, NewPin->PinType))
+		{
+			RedirectType = ERedirectType_None;
+		}
+		else
+		{
+			RedirectType = ERedirectType_Name;
+		}
 	}
 	else
 	{
