Use DoesDefaultValueMatchAutogenerated so that pin text values are not considered as 'matching'
#rnx

diff --git a/Engine/Source/Editor/BlueprintGraph/Private/K2Node.cpp b/Engine/Source/Editor/BlueprintGraph/Private/K2Node.cpp
--- a/Engine/Source/Editor/BlueprintGraph/Private/K2Node.cpp
+++ b/Engine/Source/Editor/BlueprintGraph/Private/K2Node.cpp
@@ -680,7 +680,7 @@ UK2Node::ERedirectType UK2Node::DoPinsMatchForReconstruction(const UEdGraphPin*
 		if (K2Schema &&
 			OldPin->Direction == EGPD_Input &&
 			OldPin->LinkedTo.Num() == 0 &&
-			OldPin->DefaultValue != OldPin->AutogeneratedDefaultValue &&
+			!OldPin->DoesDefaultValueMatchAutogenerated() &&
 			!K2Schema->ArePinTypesCompatible(OldPin->PinType, NewPin->PinType))
 		{
 			RedirectType = ERedirectType_None;
