Include fix
#rnx

diff --git a/Engine/Source/Editor/BlueprintGraph/Classes/EdGraphSchema_K2_Actions.h b/Engine/Source/Editor/BlueprintGraph/Classes/EdGraphSchema_K2_Actions.h
--- a/Engine/Source/Editor/BlueprintGraph/Classes/EdGraphSchema_K2_Actions.h
+++ b/Engine/Source/Editor/BlueprintGraph/Classes/EdGraphSchema_K2_Actions.h
@@ -4,10 +4,13 @@
 
 #include "CoreMinimal.h"
 #include "UObject/ObjectMacros.h"
+#include "UObject/UObjectGlobals.h"
 #include "Templates/SubclassOf.h"
 #include "UObject/UnrealType.h"
 #include "Engine/Blueprint.h"
+#include "EdGraph/EdGraph.h"
 #include "EdGraph/EdGraphSchema.h"
+
 #include "EdGraphSchema_K2_Actions.generated.h"
 
 class UEdGraph;
