(4.20.2) Remove check as the changes to pin orphaning means the container type can be expected not to match.
#rnx

diff --git a/Engine/Source/Editor/BlueprintGraph/Private/K2Node_CallArrayFunction.cpp b/Engine/Source/Editor/BlueprintGraph/Private/K2Node_CallArrayFunction.cpp
--- a/Engine/Source/Editor/BlueprintGraph/Private/K2Node_CallArrayFunction.cpp
+++ b/Engine/Source/Editor/BlueprintGraph/Private/K2Node_CallArrayFunction.cpp
@@ -86,7 +86,6 @@ void UK2Node_CallArrayFunction::NotifyPinConnectionListChanged(UEdGraphPin* Pin)
 			{
 				UEdGraphPin* LinkedTo = Pin->LinkedTo[0];
 				check(LinkedTo);
-				check(Pin->PinType.ContainerType == LinkedTo->PinType.ContainerType);
 
 				Pin->PinType.PinCategory = LinkedTo->PinType.PinCategory;
 				Pin->PinType.PinSubCategory = LinkedTo->PinType.PinSubCategory;
