Fix linker issue
#rnx

diff --git a/Engine/Source/Editor/BlueprintGraph/Classes/K2Node_Event.h b/Engine/Source/Editor/BlueprintGraph/Classes/K2Node_Event.h
--- a/Engine/Source/Editor/BlueprintGraph/Classes/K2Node_Event.h
+++ b/Engine/Source/Editor/BlueprintGraph/Classes/K2Node_Event.h
@@ -91,7 +91,7 @@ class UK2Node_Event : public UK2Node_EditablePinBase, public IK2Node_EventNodeIn
 	//~ End UK2Node Interface
 
 	//~ Begin IK2Node_EventNodeInterface Interface.
-	virtual TSharedPtr<FEdGraphSchemaAction> GetEventNodeAction(const FText& ActionCategory) override;
+	BLUEPRINTGRAPH_API virtual TSharedPtr<FEdGraphSchemaAction> GetEventNodeAction(const FText& ActionCategory) override;
 	//~ End IK2Node_EventNodeInterface Interface.
 
 	/** Checks whether the parameters for this event node are compatible with the specified function entry node */
