Optimize switch code generation for unlinked cases when there is also no default case.
Per: https://udn.unrealengine.com/questions/451077/blueprint-switch-statement-performance.html
#rnx

diff --git a/Engine/Source/Editor/BlueprintGraph/Private/K2Node_Switch.cpp b/Engine/Source/Editor/BlueprintGraph/Private/K2Node_Switch.cpp
--- a/Engine/Source/Editor/BlueprintGraph/Private/K2Node_Switch.cpp
+++ b/Engine/Source/Editor/BlueprintGraph/Private/K2Node_Switch.cpp
@@ -84,6 +84,10 @@ public:
 			UEdGraphPin* FuncPin = SwitchNode->GetFunctionPin();
 			FBPTerminal* FuncContext = Context.NetMap.FindRef(FuncPin);
 			UEdGraphPin* DefaultPin = SwitchNode->GetDefaultPin();
+			
+			// We don't need to generate if checks if there are no connections to it if there is no default pin or if the default pin is not linked 
+			// If there is a default pin that is linked then it would fall through to that default if we do not generate the cases
+			const bool bCanSkipUnlinkedCase = (DefaultPin == nullptr || DefaultPin->LinkedTo.Num() == 0);
 
 			// Pull out function to use
 			UClass* FuncClass = Cast<UClass>(FuncPin->PinType.PinSubCategoryObject.Get());
@@ -95,7 +99,7 @@ public:
 			{
 				UEdGraphPin* Pin = *PinIt;
 
-				if ((Pin->Direction == EGPD_Output) && (Pin != DefaultPin))
+				if ((Pin->Direction == EGPD_Output) && (Pin != DefaultPin) && (!bCanSkipUnlinkedCase || Pin->LinkedTo.Num() > 0))
 				{
 					// Create a term for the switch case value
 					FBPTerminal* CaseValueTerm = new (Context.Literals) FBPTerminal();
