Fix get data table row result pin not properly updating if it was split.

diff --git a/Engine/Source/Editor/BlueprintGraph/Private/K2Node_GetDataTableRow.cpp b/Engine/Source/Editor/BlueprintGraph/Private/K2Node_GetDataTableRow.cpp
--- a/Engine/Source/Editor/BlueprintGraph/Private/K2Node_GetDataTableRow.cpp
+++ b/Engine/Source/Editor/BlueprintGraph/Private/K2Node_GetDataTableRow.cpp
@@ -92,6 +92,12 @@ void UK2Node_GetDataTableRow::SetReturnTypeForStruct(UScriptStruct* NewRowStruct
 	if (NewRowStruct != OldRowStruct)
 	{
 		UEdGraphPin* ResultPin = GetResultPin();
+
+		if (ResultPin->SubPins.Num() > 0)
+		{
+			GetSchema()->RecombinePin(ResultPin);
+		}
+
 		// NOTE: purposefully not disconnecting the ResultPin (even though it changed type)... we want the user to see the old
 		//       connections, and incompatible connections will produce an error (plus, some super-struct connections may still be valid)
 		ResultPin->PinType.PinSubCategoryObject = NewRowStruct;
