Don't allow set member input pin to be split at all

diff --git a/Engine/Source/Editor/BlueprintGraph/Private/K2Node_SetFieldsInStruct.cpp b/Engine/Source/Editor/BlueprintGraph/Private/K2Node_SetFieldsInStruct.cpp
--- a/Engine/Source/Editor/BlueprintGraph/Private/K2Node_SetFieldsInStruct.cpp
+++ b/Engine/Source/Editor/BlueprintGraph/Private/K2Node_SetFieldsInStruct.cpp
@@ -359,13 +359,7 @@ bool UK2Node_SetFieldsInStruct::CanSplitPin(const UEdGraphPin* Pin) const
 		UEdGraphPin* InputPin = FindPinChecked(SetFieldsInStructHelper::StructRefPinName(), EGPD_Input);
 		if (Pin == InputPin)
 		{
-			for (const FOptionalPinFromProperty& OptionalProperty : ShowPinForProperties)
-			{
-				if (OptionalProperty.bShowPin)
-				{
-					return false;
-				}
-			}
+			return false;
 		}
 
 		return true;
