Fix GetDataTableRow title not updating when changing type

diff --git a/Engine/Source/Editor/BlueprintGraph/Private/K2Node_GetDataTableRow.cpp b/Engine/Source/Editor/BlueprintGraph/Private/K2Node_GetDataTableRow.cpp
--- a/Engine/Source/Editor/BlueprintGraph/Private/K2Node_GetDataTableRow.cpp
+++ b/Engine/Source/Editor/BlueprintGraph/Private/K2Node_GetDataTableRow.cpp
@@ -102,6 +102,8 @@ void UK2Node_GetDataTableRow::SetReturnTypeForStruct(UScriptStruct* NewRowStruct
 		//       connections, and incompatible connections will produce an error (plus, some super-struct connections may still be valid)
 		ResultPin->PinType.PinSubCategoryObject = NewRowStruct;
 		ResultPin->PinType.PinCategory = (NewRowStruct == nullptr) ? UEdGraphSchema_K2::PC_Wildcard : UEdGraphSchema_K2::PC_Struct;
+
+		CachedNodeTitle.Clear();
 	}
 }
 
