Fix for issues compiling blueprints after changing a UserDefinedStructure

diff --git a/Engine/Source/Editor/BlueprintGraph/Private/EdGraphSchema_K2.cpp b/Engine/Source/Editor/BlueprintGraph/Private/EdGraphSchema_K2.cpp
--- a/Engine/Source/Editor/BlueprintGraph/Private/EdGraphSchema_K2.cpp
+++ b/Engine/Source/Editor/BlueprintGraph/Private/EdGraphSchema_K2.cpp
@@ -3345,6 +3345,15 @@ bool UEdGraphSchema_K2::GetPropertyCategoryInfo(const UProperty* TestProperty, F
 	{
 		OutCategory = PC_Struct;
 		OutSubCategoryObject = StructProperty->Struct;
+		// Match IsTypeCompatibleWithProperty and erase REINST_ structs here:
+		if(UUserDefinedStruct* UDS = Cast<UUserDefinedStruct>(StructProperty->Struct))
+		{
+			UUserDefinedStruct* RealStruct = UDS->PrimaryStruct.Get();
+			if(RealStruct)
+			{
+				OutSubCategoryObject = RealStruct;
+			}
+		}
 	}
 	else if (TestProperty->IsA<UFloatProperty>())
 	{
