Validate CommutativeAssociativeBinaryOperator metadata in UHT and also at the K2Node level in case the function changes to an invalid construct later

diff --git a/Engine/Source/Editor/BlueprintGraph/Classes/K2Node_CommutativeAssociativeBinaryOperator.h b/Engine/Source/Editor/BlueprintGraph/Classes/K2Node_CommutativeAssociativeBinaryOperator.h
--- a/Engine/Source/Editor/BlueprintGraph/Classes/K2Node_CommutativeAssociativeBinaryOperator.h
+++ b/Engine/Source/Editor/BlueprintGraph/Classes/K2Node_CommutativeAssociativeBinaryOperator.h
@@ -41,6 +41,7 @@ public:
 
 	// UEdGraphNode interface
 	virtual void AllocateDefaultPins() override;
+	virtual void ValidateNodeDuringCompilation(class FCompilerResultsLog& MessageLog) const override;
 	// End of UEdGraphNode interface
 
 	// UK2Node interface
diff --git a/Engine/Source/Editor/BlueprintGraph/Private/K2Node_CommutativeAssociativeBinaryOperator.cpp b/Engine/Source/Editor/BlueprintGraph/Private/K2Node_CommutativeAssociativeBinaryOperator.cpp
--- a/Engine/Source/Editor/BlueprintGraph/Private/K2Node_CommutativeAssociativeBinaryOperator.cpp
+++ b/Engine/Source/Editor/BlueprintGraph/Private/K2Node_CommutativeAssociativeBinaryOperator.cpp
@@ -103,50 +103,57 @@ FEdGraphPinType UK2Node_CommutativeAssociativeBinaryOperator::GetType() const
 	return FEdGraphPinType();
 }
 
-void UK2Node_CommutativeAssociativeBinaryOperator::AllocateDefaultPins()
+void UK2Node_CommutativeAssociativeBinaryOperator::ValidateNodeDuringCompilation(class FCompilerResultsLog& MessageLog) const
 {
-	Super::AllocateDefaultPins();
+	Super::ValidateNodeDuringCompilation(MessageLog);
 
-	const UFunction* Function = GetTargetFunction();
-	if(NULL != Function)
+	if (const UFunction* Function = GetTargetFunction())
 	{
-		check(Function->HasAnyFunctionFlags(FUNC_BlueprintPure));
-	
-#if DO_CHECK
-		ensure(FindOutPin());
-		ensure((FindSelfPin() ? 4 : 3) ==  Pins.Num());
+		if (Function->HasAnyFunctionFlags(FUNC_BlueprintPure))
 		{
-			//VALIDATION
+			UEdGraphPin* SelfPin = FindSelfPin();
 			const FEdGraphPinType InputType = GetType();
-			int32 NativeInputPinsNum = 0;
-			for (int32 PinIt = 0; PinIt < Pins.Num(); PinIt++)
+			int32 OutPinCount = 0;
+			bool bConsistentPinType = true;
+
+			for (UEdGraphPin* Pin : Pins)
 			{
-				const UEdGraphPin* Pin = Pins[PinIt];
-				if (Pin != FindSelfPin())
+				if (Pin->Direction == EEdGraphPinDirection::EGPD_Output)
+				{
+					++OutPinCount;
+				}
+				else if (Pin != SelfPin)
 				{
-					ensure(InputType == Pin->PinType);
-					NativeInputPinsNum += (EEdGraphPinDirection::EGPD_Input ==  Pin->Direction) ? 1 : 0;
+					if (Pin->PinType != InputType)
+					{
+						bConsistentPinType = false;
+					}
 				}
 			}
-			ensure(BinaryOperatorInputsNum == NativeInputPinsNum);
-		}
-#endif // DO_CHECK
 
-		for (int32 i = 0; i < NumAdditionalInputs; ++i)
+			if (OutPinCount != 1)
+			{
+				MessageLog.Error(*LOCTEXT("InvalidOutPin", "@@ - Commutative associative binary operators must have a single output pin.").ToString(), this);
+			}
+			if (!bConsistentPinType)
+			{
+				MessageLog.Error(*LOCTEXT("InconsistentType", "@@ - Commutative associative binary operator inputs must all be of the same type.").ToString(), this);
+			}
+		}
+		else
 		{
-			AddInputPinInner(i);
+			MessageLog.Error(*LOCTEXT("MustBePure", "@@ - Commutative associative binary operators must be pure functions.").ToString(), this);
 		}
 	}
-	else
+}
+
+void UK2Node_CommutativeAssociativeBinaryOperator::AllocateDefaultPins()
+{
+	Super::AllocateDefaultPins();
+
+	for (int32 i = 0; i < NumAdditionalInputs; ++i)
 	{
-		const UClass* FunctionParentClass = FunctionReference.GetMemberParentClass(GetBlueprintClassFromNode());
-		Message_Error(
-			FText::Format(
-				LOCTEXT("NoFunction_ErrorFmt", "CommutativeAssociativeBinaryOperator has no function: '{0}' class: '{1}'"),
-				FText::FromString(FunctionReference.GetMemberName().ToString()),
-				FunctionParentClass ? FText::FromString(FunctionParentClass->GetName()) : LOCTEXT("NoFunction_ErrorNone", "None")
-			).ToString()
-		);
+		AddInputPinInner(i);
 	}
 }
 
diff --git a/Engine/Source/Programs/UnrealHeaderTool/Private/HeaderParser.cpp b/Engine/Source/Programs/UnrealHeaderTool/Private/HeaderParser.cpp
--- a/Engine/Source/Programs/UnrealHeaderTool/Private/HeaderParser.cpp
+++ b/Engine/Source/Programs/UnrealHeaderTool/Private/HeaderParser.cpp
@@ -951,6 +951,61 @@ namespace
 			}
 			break;
 
+			case ECheckedMetadataSpecifier::CommutativeAssociativeBinaryOperator:
+			{
+				if (UFunction* Function = Cast<UFunction>(Field))
+				{
+					bool bGoodParams = (Function->NumParms == 3);
+					if (bGoodParams)
+					{
+						UProperty* FirstParam = nullptr;
+						UProperty* SecondParam = nullptr;
+						UProperty* ReturnValue = nullptr;
+
+						TFieldIterator<UProperty> It(Function);
+
+						auto GetNextParam = [&]()
+						{
+							if (It)
+							{
+								if (It->HasAnyPropertyFlags(CPF_ReturnParm))
+								{
+									ReturnValue = *It;
+								}
+								else
+								{
+									if (FirstParam == nullptr)
+									{
+										FirstParam = *It;
+									}
+									else if (SecondParam == nullptr)
+									{
+										SecondParam = *It;
+									}
+								}
+								++It;
+							}
+						};
+
+						GetNextParam();
+						GetNextParam();
+						GetNextParam();
+						ensure(!It);
+
+						if (ReturnValue == nullptr || SecondParam == nullptr || !SecondParam->SameType(FirstParam))
+						{
+							bGoodParams = false;
+						}
+					}
+
+					if (!bGoodParams)
+					{
+						UE_LOG_ERROR_UHT(TEXT("Commutative asssociative binary operators must have exactly 2 parameters of the same type and a return value."));
+					}
+				}
+			}
+			break;
+
 			case ECheckedMetadataSpecifier::ExpandEnumAsExecs:
 			{
 				if (UFunction* Function = Cast<UFunction>(Field))
diff --git a/Engine/Source/Programs/UnrealHeaderTool/Private/Specifiers/CheckedMetadataSpecifiers.def b/Engine/Source/Programs/UnrealHeaderTool/Private/Specifiers/CheckedMetadataSpecifiers.def
--- a/Engine/Source/Programs/UnrealHeaderTool/Private/Specifiers/CheckedMetadataSpecifiers.def
+++ b/Engine/Source/Programs/UnrealHeaderTool/Private/Specifiers/CheckedMetadataSpecifiers.def
@@ -5,6 +5,7 @@
 CHECKED_METADATA_SPECIFIER(BlueprintProtected)
 CHECKED_METADATA_SPECIFIER(ClampMax)
 CHECKED_METADATA_SPECIFIER(ClampMin)
+CHECKED_METADATA_SPECIFIER(CommutativeAssociativeBinaryOperator)
 CHECKED_METADATA_SPECIFIER(DevelopmentStatus)
 CHECKED_METADATA_SPECIFIER(ExpandEnumAsExecs)
 CHECKED_METADATA_SPECIFIER(UIMax)
