Add comments to ERedirectType
#rnx

diff --git a/Engine/Source/Editor/BlueprintGraph/Classes/K2Node.h b/Engine/Source/Editor/BlueprintGraph/Classes/K2Node.h
--- a/Engine/Source/Editor/BlueprintGraph/Classes/K2Node.h
+++ b/Engine/Source/Editor/BlueprintGraph/Classes/K2Node.h
@@ -400,8 +400,11 @@ protected:
 
 	enum ERedirectType
 	{
+		/* The pins do not match */
 		ERedirectType_None,
+		/* The pins match by name or redirect and have the same type (or we're ok with the mismatched type) */
 		ERedirectType_Name,
+		/* The pins match via a redirect and the value needs to also be redirected */
 		ERedirectType_Value
 	};
 
