Don't allow a new pin to be matched twice

diff --git a/Engine/Source/Editor/BlueprintGraph/Private/K2Node.cpp b/Engine/Source/Editor/BlueprintGraph/Private/K2Node.cpp
--- a/Engine/Source/Editor/BlueprintGraph/Private/K2Node.cpp
+++ b/Engine/Source/Editor/BlueprintGraph/Private/K2Node.cpp
@@ -900,9 +900,12 @@ FString UK2Node::GetPinMetaData(FName InPinName, FName InKey)
 void UK2Node::RewireOldPinsToNewPins(TArray<UEdGraphPin*>& InOldPins, TArray<UEdGraphPin*>& InNewPins, TMap<UEdGraphPin*, UEdGraphPin*>* NewPinToOldPin)
 {
 	TArray<UEdGraphPin*> OrphanedOldPins;
+	TArray<bool> NewPinMatched; // Tracks whether a NewPin has already been matched to an OldPin
 	TMap<UEdGraphPin*, UEdGraphPin*> MatchedPins; // Old to New
 
-	const bool bSaveUnconnectedDefaultPins = (InNewPins.Num() == 0 && (IsA<UK2Node_CallFunction>() || IsA<UK2Node_MacroInstance>()));
+	const int32 NumNewPins = InNewPins.Num();
+	NewPinMatched.AddDefaulted(NumNewPins);
+	const bool bSaveUnconnectedDefaultPins = (NumNewPins == 0 && (IsA<UK2Node_CallFunction>() || IsA<UK2Node_MacroInstance>()));
 
 	// Rewire any connection to pins that are matched by name (O(N^2) right now)
 	// NOTE: we iterate backwards through the list because ReconstructSinglePin()
@@ -917,27 +920,29 @@ void UK2Node::RewireOldPinsToNewPins(TArray<UEdGraphPin*>& InOldPins, TArray<UEd
 
 		// common case is for InOldPins and InNewPins to match, so we start searching from the current index:
 		bool bMatched = false;
-		const int32 NumNewPins = InNewPins.Num();
 		int32 NewPinIndex = (NumNewPins ? OldPinIndex % NumNewPins : 0);
 		for (int32 NewPinCount = NumNewPins - 1; NewPinCount >= 0; --NewPinCount)
 		{
-			// if InNewPins grows in this loop then we may skip entries and fail to find a match:
+			// if InNewPins grows then we may skip entries and fail to find a match or NewPinMatched will not be accurate
 			check(NumNewPins == InNewPins.Num());
-			UEdGraphPin* NewPin = InNewPins[NewPinIndex];
-
-			const ERedirectType RedirectType = DoPinsMatchForReconstruction(NewPin, NewPinIndex, OldPin, OldPinIndex);
-			if (RedirectType != ERedirectType_None)
+			if (!NewPinMatched[NewPinIndex])
 			{
-				ReconstructSinglePin(NewPin, OldPin, RedirectType);
-				MatchedPins.Add(OldPin, NewPin);
-				if (NewPinToOldPin)
+				UEdGraphPin* NewPin = InNewPins[NewPinIndex];
+
+				const ERedirectType RedirectType = DoPinsMatchForReconstruction(NewPin, NewPinIndex, OldPin, OldPinIndex);
+				if (RedirectType != ERedirectType_None)
 				{
-					NewPinToOldPin->Add(NewPin, OldPin);
+					ReconstructSinglePin(NewPin, OldPin, RedirectType);
+					MatchedPins.Add(OldPin, NewPin);
+					if (NewPinToOldPin)
+					{
+						NewPinToOldPin->Add(NewPin, OldPin);
+					}
+					bMatched = true;
+					NewPinMatched[NewPinIndex] = true;
+					break;
 				}
-				bMatched = true;
-				break;
 			}
-
 			NewPinIndex = (NewPinIndex + 1) % InNewPins.Num();
 		}
 
