Fixed compilation error describing undetermined type when calling an interface function with Set or Map properties

diff --git a/Engine/Source/Editor/BlueprintGraph/Private/K2Node_Message.cpp b/Engine/Source/Editor/BlueprintGraph/Private/K2Node_Message.cpp
--- a/Engine/Source/Editor/BlueprintGraph/Private/K2Node_Message.cpp
+++ b/Engine/Source/Editor/BlueprintGraph/Private/K2Node_Message.cpp
@@ -1,6 +1,7 @@
 // Copyright 1998-2019 Epic Games, Inc. All Rights Reserved.
 
 #include "K2Node_Message.h"
+#include "EdGraphUtilities.h"
 #include "Engine/LevelScriptActor.h"
 #include "Engine/LevelStreaming.h"
 #include "EdGraphSchema_K2.h"
@@ -11,6 +12,8 @@
 #include "KismetCompilerMisc.h"
 #include "KismetCompiler.h"
 #include "Kismet/KismetArrayLibrary.h"
+#include "Kismet/BlueprintSetLibrary.h"
+#include "Kismet/BlueprintMapLibrary.h"
 
 #define LOCTEXT_NAMESPACE "K2Node_Message"
 
@@ -247,6 +250,10 @@ void UK2Node_Message::ExpandNode(class FKismetCompilerContext& CompilerContext,
 
 		UFunction* ArrayClearFunction = UKismetArrayLibrary::StaticClass()->FindFunctionByName(FName(TEXT("Array_Clear")));
 		check(ArrayClearFunction);
+		UFunction* SetClearFunction = UBlueprintSetLibrary::StaticClass()->FindFunctionByName(FName(TEXT("Set_Clear")));
+		check(SetClearFunction);
+		UFunction* MapClearFunction = UBlueprintMapLibrary::StaticClass()->FindFunctionByName(FName(TEXT("Map_Clear")));
+		check(MapClearFunction);
 
 		bool const bIsPureMessageFunc = Super::IsNodePure();
 		// Variable pins - Try to associate variable inputs to the message node with the variable inputs and outputs to the call function node
@@ -321,6 +328,32 @@ void UK2Node_Message::ExpandNode(class FKismetCompilerContext& CompilerContext,
 
 							DefaultValueThenPin = ClearArray->GetThenPin();
 						} 
+						else if(CurrentPin->PinType.IsSet())
+						{
+							UK2Node_CallFunction* ClearSet = CompilerContext.SpawnIntermediateNode<UK2Node_CallFunction>(this, SourceGraph);
+							DefaultValueNode = ClearSet;
+							ClearSet->SetFromFunction(SetClearFunction);
+							ClearSet->AllocateDefaultPins();
+
+							UEdGraphPin* SetPin = FEdGraphUtilities::FindSetParamPin(SetClearFunction, ClearSet);
+							Schema->TryCreateConnection(SetPin, VarOutPin);
+							ClearSet->PinConnectionListChanged(SetPin);
+
+							DefaultValueThenPin = ClearSet->GetThenPin();
+						}
+						else if(CurrentPin->PinType.IsMap())
+						{
+							UK2Node_CallFunction* ClearMap = CompilerContext.SpawnIntermediateNode<UK2Node_CallFunction>(this, SourceGraph);
+							DefaultValueNode = ClearMap;
+							ClearMap->SetFromFunction(MapClearFunction);
+							ClearMap->AllocateDefaultPins();
+
+							UEdGraphPin* MapPin = FEdGraphUtilities::FindMapParamPin(MapClearFunction, ClearMap);
+							Schema->TryCreateConnection(MapPin, VarOutPin);
+							ClearMap->PinConnectionListChanged(MapPin);
+
+							DefaultValueThenPin = ClearMap->GetThenPin();
+						}
 						else
 						{
 							UK2Node_AssignmentStatement* AssignDefaultValue = CompilerContext.SpawnIntermediateNode<UK2Node_AssignmentStatement>(this, SourceGraph);
