Fixed literal text values set on function return nodes never being gathered for localization

This is due to them being user-defined pins, and user-defined pins always set their value as the autogenerated default, which results in the normal pin gather in GatherGraphNodeForLocalization considering them default as they pass the DoesDefaultValueMatchAutogenerated test.

In this case the return node is the only place the values are defined, so we must gather them.

diff --git a/Engine/Source/Editor/BlueprintGraph/Private/K2Node_FunctionResult.cpp b/Engine/Source/Editor/BlueprintGraph/Private/K2Node_FunctionResult.cpp
--- a/Engine/Source/Editor/BlueprintGraph/Private/K2Node_FunctionResult.cpp
+++ b/Engine/Source/Editor/BlueprintGraph/Private/K2Node_FunctionResult.cpp
@@ -2,6 +2,7 @@
 
 #include "K2Node_FunctionResult.h"
 #include "Misc/CoreMisc.h"
+#include "Serialization/PropertyLocalizationDataGathering.h"
 #include "EdGraphSchema_K2.h"
 #include "K2Node_FunctionEntry.h"
 
@@ -118,9 +119,35 @@ public:
 	}
 };
 
+#if WITH_EDITORONLY_DATA
+namespace
+{
+	void GatherFunctionResultNodeForLocalization(const UObject* const Object, FPropertyLocalizationDataGatherer& PropertyLocalizationDataGatherer, const EPropertyLocalizationGathererTextFlags GatherTextFlags)
+	{
+		const UK2Node_FunctionResult* const FunctionResultNode = CastChecked<UK2Node_FunctionResult>(Object);
+
+		// Function Result (aka, Return) nodes always report their values as being the default
+		// but we still need to gather them as they are the only place the values are defined
+		const FString PathToObject = FunctionResultNode->GetPathName();
+		for (const UEdGraphPin* Pin : FunctionResultNode->Pins)
+		{
+			if (!Pin->DefaultTextValue.IsEmpty())
+			{
+				PropertyLocalizationDataGatherer.GatherTextInstance(Pin->DefaultTextValue, FString::Printf(TEXT("%s.%s"), *PathToObject, *Pin->GetName()), /*bIsEditorOnly*/true);
+			}
+		}
+
+		PropertyLocalizationDataGatherer.GatherLocalizationDataFromObject(FunctionResultNode, GatherTextFlags);
+	}
+}
+#endif
+
 UK2Node_FunctionResult::UK2Node_FunctionResult(const FObjectInitializer& ObjectInitializer)
 	: Super(ObjectInitializer)
 {
+#if WITH_EDITORONLY_DATA
+	{ static const FAutoRegisterLocalizationDataGatheringCallback AutomaticRegistrationOfLocalizationGatherer(UK2Node_FunctionResult::StaticClass(), &GatherFunctionResultNodeForLocalization); }
+#endif
 }
 
 FText UK2Node_FunctionResult::GetNodeTitle(ENodeTitleType::Type TitleType) const
