#UE4 Fix a bug where if you delete a custom event in a parent class that is overridden in a child it creates a substitute without a NodeGuid.

diff --git a/Engine/Source/Editor/BlueprintGraph/Private/EdGraphSchema_K2.cpp b/Engine/Source/Editor/BlueprintGraph/Private/EdGraphSchema_K2.cpp
--- a/Engine/Source/Editor/BlueprintGraph/Private/EdGraphSchema_K2.cpp
+++ b/Engine/Source/Editor/BlueprintGraph/Private/EdGraphSchema_K2.cpp
@@ -6234,6 +6234,9 @@ UEdGraphNode* UEdGraphSchema_K2::CreateSubstituteNode(UEdGraphNode* Node, const
 			CustomEventNode->NodePosX = EventNode->NodePosX;
 			CustomEventNode->NodePosY = EventNode->NodePosY;
 
+			// Reuse the same GUID as the replaced node
+			CustomEventNode->NodeGuid = EventNode->NodeGuid;
+
 			// Build a function name that is appropriate for the event we're replacing
 			FString FunctionName;
 			const UK2Node_ActorBoundEvent* ActorBoundEventNode = Cast<const UK2Node_ActorBoundEvent>(EventNode);
diff --git a/Engine/Source/Runtime/Engine/Classes/EdGraph/EdGraphSchema.h b/Engine/Source/Runtime/Engine/Classes/EdGraph/EdGraphSchema.h
--- a/Engine/Source/Runtime/Engine/Classes/EdGraph/EdGraphSchema.h
+++ b/Engine/Source/Runtime/Engine/Classes/EdGraph/EdGraphSchema.h
@@ -881,7 +881,7 @@ class ENGINE_API UEdGraphSchema : public UObject
 	virtual void ReconstructNode(UEdGraphNode& TargetNode, bool bIsBatchRequest=false) const;
 
 	/**
-	 * Attempts to construct a substitute node that is unique within its graph
+	 * Attempts to construct a substitute node that is unique within its graph. If this call returns non-null node, it is expected for the caller to destroy the node that was passed in.
 	 *
 	 * @param	Node			The node to replace
 	 * @param	Graph			The destination graph
