[NO NON-PII COMMIT MESSAGE]

diff --git a/Engine/Source/Editor/BlueprintGraph/Private/K2Node_MakeStruct.cpp b/Engine/Source/Editor/BlueprintGraph/Private/K2Node_MakeStruct.cpp
--- a/Engine/Source/Editor/BlueprintGraph/Private/K2Node_MakeStruct.cpp
+++ b/Engine/Source/Editor/BlueprintGraph/Private/K2Node_MakeStruct.cpp
@@ -69,7 +69,7 @@ void UK2Node_MakeStruct::FMakeStructPinManager::CustomizePinData(UEdGraphPin* Pi
 		if (nullptr != SampleStructMemory)
 		{
 			FString NewDefaultValue;
-			if (Property->ExportText_InContainer(0, NewDefaultValue, SampleStructMemory, SampleStructMemory, nullptr, PPF_None))
+			if (FBlueprintEditorUtils::PropertyValueToString(Property, SampleStructMemory, NewDefaultValue))
 			{
 				if (Schema->IsPinDefaultValid(Pin, NewDefaultValue, nullptr, FText::GetEmpty()).IsEmpty())
 				{
