Allow devs to opt-in to exposing deprecated functions for continued use in the Blueprint editor.

diff --git a/Engine/Source/Editor/BlueprintGraph/Private/BlueprintEditorSettings.cpp b/Engine/Source/Editor/BlueprintGraph/Private/BlueprintEditorSettings.cpp
--- a/Engine/Source/Editor/BlueprintGraph/Private/BlueprintEditorSettings.cpp
+++ b/Engine/Source/Editor/BlueprintGraph/Private/BlueprintEditorSettings.cpp
@@ -16,6 +16,7 @@ UBlueprintEditorSettings::UBlueprintEditorSettings(const FObjectInitializer& Obj
 	, bSplitContextTargetSettings(true)
 	, bExposeAllMemberComponentFunctions(true)
 	, bShowContextualFavorites(false)
+	, bExposeDeprecatedFunctions(false)
 	, bCompactCallOnMemberNodes(false)
 	, bFlattenFavoritesMenus(true)
 	, bFavorPureCastNodes(false)
diff --git a/Engine/Source/Editor/BlueprintGraph/Private/BlueprintFunctionNodeSpawner.cpp b/Engine/Source/Editor/BlueprintGraph/Private/BlueprintFunctionNodeSpawner.cpp
--- a/Engine/Source/Editor/BlueprintGraph/Private/BlueprintFunctionNodeSpawner.cpp
+++ b/Engine/Source/Editor/BlueprintGraph/Private/BlueprintFunctionNodeSpawner.cpp
@@ -17,6 +17,7 @@
 #include "ObjectEditorUtils.h"
 #include "BlueprintNodeSpawnerUtils.h"
 #include "BlueprintEditorSettings.h"
+#include "Kismet2/BlueprintEditorUtils.h"
 
 #define LOCTEXT_NAMESPACE "BlueprintFunctionNodeSpawner"
 
@@ -263,6 +264,11 @@ UBlueprintFunctionNodeSpawner* UBlueprintFunctionNodeSpawner::Create(TSubclassOf
 		MenuSignature.Tooltip = MenuSignature.MenuName;
 	}
 
+	if (Function->HasMetaData(FBlueprintMetadata::MD_DeprecatedFunction))
+	{
+		MenuSignature.MenuName = FBlueprintEditorUtils::GetDeprecatedMemberMenuItemName(MenuSignature.MenuName);
+	}
+
 	//--------------------------------------
 	// Post-Spawn Setup
 	//--------------------------------------
diff --git a/Engine/Source/Editor/BlueprintGraph/Private/EdGraphSchema_K2.cpp b/Engine/Source/Editor/BlueprintGraph/Private/EdGraphSchema_K2.cpp
--- a/Engine/Source/Editor/BlueprintGraph/Private/EdGraphSchema_K2.cpp
+++ b/Engine/Source/Editor/BlueprintGraph/Private/EdGraphSchema_K2.cpp
@@ -907,13 +907,13 @@ UFunction* UEdGraphSchema_K2::GetCallableParentFunction(UFunction* Function)
 bool UEdGraphSchema_K2::CanUserKismetCallFunction(const UFunction* Function)
 {
 	return Function && 
-		(Function->HasAllFunctionFlags(FUNC_BlueprintCallable) && !Function->HasAllFunctionFlags(FUNC_Delegate) && !Function->GetBoolMetaData(FBlueprintMetadata::MD_BlueprintInternalUseOnly) && !Function->HasMetaData(FBlueprintMetadata::MD_DeprecatedFunction));
+		(Function->HasAllFunctionFlags(FUNC_BlueprintCallable) && !Function->HasAllFunctionFlags(FUNC_Delegate) && !Function->GetBoolMetaData(FBlueprintMetadata::MD_BlueprintInternalUseOnly) && (!Function->HasMetaData(FBlueprintMetadata::MD_DeprecatedFunction) || GetDefault<UBlueprintEditorSettings>()->bExposeDeprecatedFunctions));
 }
 
 bool UEdGraphSchema_K2::CanKismetOverrideFunction(const UFunction* Function)
 {
 	return  Function && 
-		(Function->HasAllFunctionFlags(FUNC_BlueprintEvent) && !Function->HasAllFunctionFlags(FUNC_Delegate) && !Function->GetBoolMetaData(FBlueprintMetadata::MD_BlueprintInternalUseOnly) && !Function->HasMetaData(FBlueprintMetadata::MD_DeprecatedFunction));
+		(Function->HasAllFunctionFlags(FUNC_BlueprintEvent) && !Function->HasAllFunctionFlags(FUNC_Delegate) && !Function->GetBoolMetaData(FBlueprintMetadata::MD_BlueprintInternalUseOnly) && (!Function->HasMetaData(FBlueprintMetadata::MD_DeprecatedFunction) || GetDefault<UBlueprintEditorSettings>()->bExposeDeprecatedFunctions));
 }
 
 bool UEdGraphSchema_K2::HasFunctionAnyOutputParameter(const UFunction* InFunction)
diff --git a/Engine/Source/Editor/BlueprintGraph/Public/BlueprintEditorSettings.h b/Engine/Source/Editor/BlueprintGraph/Public/BlueprintEditorSettings.h
--- a/Engine/Source/Editor/BlueprintGraph/Public/BlueprintEditorSettings.h
+++ b/Engine/Source/Editor/BlueprintGraph/Public/BlueprintEditorSettings.h
@@ -46,6 +46,10 @@ public:
 	UPROPERTY(EditAnywhere, config, Category=Workflow, meta=(DisplayName="Context Menu: Show Favorites Section"))
 	bool bShowContextualFavorites;
 
+	/** If enabled, deprecated functions will be visible in the context menu and will be available for override implementation. By default, functions marked as deprecated are not exposed in either case. */
+	UPROPERTY(EditAnywhere, config, Category = Workflow)
+	bool bExposeDeprecatedFunctions;
+
 	/** If enabled, then call-on-member actions will be spawned as a single node (instead of a GetMember + FunctionCall node). */
 	UPROPERTY(EditAnywhere, config, Category=Workflow)
 	bool bCompactCallOnMemberNodes;
diff --git a/Engine/Source/Editor/Kismet/Private/SMyBlueprint.cpp b/Engine/Source/Editor/Kismet/Private/SMyBlueprint.cpp
--- a/Engine/Source/Editor/Kismet/Private/SMyBlueprint.cpp
+++ b/Engine/Source/Editor/Kismet/Private/SMyBlueprint.cpp
@@ -1248,6 +1248,11 @@ void SMyBlueprint::CollectAllActions(FGraphActionListBuilderBase& OutAllActions)
 				FunctionDesc = FText::FromString(Function->GetName());
 			}
 
+			if (Function->HasMetaData(FBlueprintMetadata::MD_DeprecatedFunction))
+			{
+				FunctionDesc = FBlueprintEditorUtils::GetDeprecatedMemberMenuItemName(FunctionDesc);
+			}
+
 			FText FunctionCategory = Function->GetMetaDataText(FBlueprintMetadata::MD_FunctionCategory, TEXT("UObjectCategory"), Function->GetFullGroupName(false));
 
 			TSharedPtr<FEdGraphSchemaAction_K2Graph> NewFuncAction = MakeShareable(new FEdGraphSchemaAction_K2Graph(EEdGraphSchemaAction_K2Graph::Function, FunctionCategory, FunctionDesc, FunctionTooltip, 1, NodeSectionID::FUNCTION_OVERRIDABLE));
diff --git a/Engine/Source/Editor/UnrealEd/Private/Kismet2/BlueprintEditorUtils.cpp b/Engine/Source/Editor/UnrealEd/Private/Kismet2/BlueprintEditorUtils.cpp
--- a/Engine/Source/Editor/UnrealEd/Private/Kismet2/BlueprintEditorUtils.cpp
+++ b/Engine/Source/Editor/UnrealEd/Private/Kismet2/BlueprintEditorUtils.cpp
@@ -9158,6 +9158,13 @@ FString FBlueprintEditorUtils::GetClassNameWithoutSuffix(const UClass* Class)
 	}
 }
 
+FText FBlueprintEditorUtils::GetDeprecatedMemberMenuItemName(const FText& MemberName)
+{
+	FFormatNamedArguments Args;
+	Args.Add(TEXT("MemberName"), MemberName);
+	return FText::Format(LOCTEXT("DeprecatedMemberMenuItemName", "{MemberName} (Deprecated)"), Args);
+}
+
 FText FBlueprintEditorUtils::GetDeprecatedMemberUsageNodeWarning(const FText& MemberName, const FText& DetailedMessage)
 {
 	static FText UnknownName = LOCTEXT("UnknownDeprecatedMemberName", "[unknown]");
diff --git a/Engine/Source/Editor/UnrealEd/Public/Kismet2/BlueprintEditorUtils.h b/Engine/Source/Editor/UnrealEd/Public/Kismet2/BlueprintEditorUtils.h
--- a/Engine/Source/Editor/UnrealEd/Public/Kismet2/BlueprintEditorUtils.h
+++ b/Engine/Source/Editor/UnrealEd/Public/Kismet2/BlueprintEditorUtils.h
@@ -1626,6 +1626,13 @@ public:
 	 */
 	static FString GetClassNameWithoutSuffix(const UClass* Class);
 
+	/**
+	 * Returns a formatted menu item label for a deprecated variable or function member with the given name.
+	 *
+	 * @param MemberName		(Required) User-facing name of the deprecated variable or function.
+	 */
+	static FText GetDeprecatedMemberMenuItemName(const FText& MemberName);
+
 	/**
 	 * Returns a formatted warning message regarding usage of a deprecated variable or function member with the given name.
 	 *
