Fix crash pasting timeline with external curve that has been force deleted

diff --git a/Engine/Source/Editor/BlueprintGraph/Private/K2Node_Timeline.cpp b/Engine/Source/Editor/BlueprintGraph/Private/K2Node_Timeline.cpp
--- a/Engine/Source/Editor/BlueprintGraph/Private/K2Node_Timeline.cpp
+++ b/Engine/Source/Editor/BlueprintGraph/Private/K2Node_Timeline.cpp
@@ -201,7 +201,7 @@ void UK2Node_Timeline::PostPasteNode()
 			for( auto TrackIt = Template->FloatTracks.CreateIterator(); TrackIt; ++TrackIt )
 			{
 				FTTFloatTrack& Track = *TrackIt;
-				if(Track.CurveFloat->GetOuter()->IsA(UBlueprint::StaticClass()))
+				if (!Track.bIsExternalCurve && Track.CurveFloat->GetOuter()->IsA(UBlueprint::StaticClass()))
 				{
 					Track.CurveFloat->Rename(*Template->MakeUniqueCurveName(Track.CurveFloat, Track.CurveFloat->GetOuter()), Blueprint, REN_DontCreateRedirectors);
 				}
@@ -210,7 +210,7 @@ void UK2Node_Timeline::PostPasteNode()
 			for( auto TrackIt = Template->EventTracks.CreateIterator(); TrackIt; ++TrackIt )
 			{
 				FTTEventTrack& Track = *TrackIt;
-				if(Track.CurveKeys->GetOuter()->IsA(UBlueprint::StaticClass()))
+				if (!Track.bIsExternalCurve && Track.CurveKeys->GetOuter()->IsA(UBlueprint::StaticClass()))
 				{
 					Track.CurveKeys->Rename(*Template->MakeUniqueCurveName(Track.CurveKeys, Track.CurveKeys->GetOuter()), Blueprint, REN_DontCreateRedirectors);
 				}
@@ -219,7 +219,7 @@ void UK2Node_Timeline::PostPasteNode()
 			for( auto TrackIt = Template->VectorTracks.CreateIterator(); TrackIt; ++TrackIt )
 			{
 				FTTVectorTrack& Track = *TrackIt;
-				if(Track.CurveVector->GetOuter()->IsA(UBlueprint::StaticClass()))
+				if (!Track.bIsExternalCurve && Track.CurveVector->GetOuter()->IsA(UBlueprint::StaticClass()))
 				{
 					Track.CurveVector->Rename(*Template->MakeUniqueCurveName(Track.CurveVector, Track.CurveVector->GetOuter()), Blueprint, REN_DontCreateRedirectors);
 				}
@@ -228,7 +228,7 @@ void UK2Node_Timeline::PostPasteNode()
 			for( auto TrackIt = Template->LinearColorTracks.CreateIterator(); TrackIt; ++TrackIt )
 			{
 				FTTLinearColorTrack& Track = *TrackIt;
-				if(Track.CurveLinearColor->GetOuter()->IsA(UBlueprint::StaticClass()))
+				if (!Track.bIsExternalCurve && Track.CurveLinearColor->GetOuter()->IsA(UBlueprint::StaticClass()))
 				{
 					Track.CurveLinearColor->Rename(*Template->MakeUniqueCurveName(Track.CurveLinearColor, Track.CurveLinearColor->GetOuter()), Blueprint, REN_DontCreateRedirectors);
 				}
