PR #5133: UE-64479 Fix the icon display for the Timeline node replication status (Contributed by [REDACTED_PERSON])
PR #5132: UE-64479: Replicated icon shows up for timelines in BP (Contributed by [REDACTED_PERSON])

diff --git a/Engine/Source/Editor/BlueprintGraph/Classes/K2Node_Timeline.h b/Engine/Source/Editor/BlueprintGraph/Classes/K2Node_Timeline.h
--- a/Engine/Source/Editor/BlueprintGraph/Classes/K2Node_Timeline.h
+++ b/Engine/Source/Editor/BlueprintGraph/Classes/K2Node_Timeline.h
@@ -37,7 +37,7 @@ class UK2Node_Timeline : public UK2Node
 	UPROPERTY(Transient)
 	uint32 bLoop:1;
 
-	/** If the timeline is set to loop */
+	/** If the timeline is set to replicate */
 	UPROPERTY(Transient)
 	uint32 bReplicated:1;
 
@@ -59,6 +59,7 @@ class UK2Node_Timeline : public UK2Node
 	virtual TSharedPtr<class INameValidatorInterface> MakeNameValidator() const override;
 	virtual FText GetTooltipText() const override;
 	virtual FString GetDocumentationExcerptName() const override;
+	virtual FName GetCornerIcon() const override;
 	virtual FSlateIcon GetIconAndTint(FLinearColor& OutColor) const override;
 	virtual bool ShouldShowNodeProperties() const override { return true; }
 	virtual UObject* GetJumpTargetForDoubleClick() const override;
diff --git a/Engine/Source/Editor/BlueprintGraph/Private/K2Node_Timeline.cpp b/Engine/Source/Editor/BlueprintGraph/Private/K2Node_Timeline.cpp
--- a/Engine/Source/Editor/BlueprintGraph/Private/K2Node_Timeline.cpp
+++ b/Engine/Source/Editor/BlueprintGraph/Private/K2Node_Timeline.cpp
@@ -114,6 +114,7 @@ void UK2Node_Timeline::AllocateDefaultPins()
 		// cache play status
 		bAutoPlay = Timeline->bAutoPlay;
 		bLoop = Timeline->bLoop;
+		bReplicated = Timeline->bReplicated;
 		bIgnoreTimeDilation = Timeline->bIgnoreTimeDilation;
 	}
 
@@ -181,6 +182,7 @@ void UK2Node_Timeline::PostPasteNode()
 		{
 			bAutoPlay = Template->bAutoPlay;
 			bLoop = Template->bLoop;
+			bReplicated = Template->bReplicated;
 			bIgnoreTimeDilation = Template->bIgnoreTimeDilation;
 		}
 	}
@@ -192,6 +194,7 @@ void UK2Node_Timeline::PostPasteNode()
 		UTimelineTemplate* Template = DuplicateObject<UTimelineTemplate>(OldTimeline, Blueprint->GeneratedClass, TimelineTemplateName);
 		bAutoPlay = Template->bAutoPlay;
 		bLoop = Template->bLoop;
+		bReplicated = Template->bReplicated;
 		bIgnoreTimeDilation = Template->bIgnoreTimeDilation;
 		Template->SetFlags(RF_Transactional);
 		Blueprint->Timelines.Add(Template);
@@ -599,6 +602,15 @@ FText UK2Node_Timeline::GetTooltipText() const
 	return LOCTEXT("TimelineTooltip", "Timeline node allows values to be keyframed over time.\nDouble click to open timeline editor.");
 }
 
+FName UK2Node_Timeline::GetCornerIcon() const
+{
+	if (bReplicated)
+	{
+		return TEXT("Graph.Replication.Replicated");
+	}
+	return Super::GetCornerIcon();
+}
+
 FSlateIcon UK2Node_Timeline::GetIconAndTint(FLinearColor& OutColor) const
 {
 	static FSlateIcon Icon("EditorStyle", "GraphEditor.Timeline_16x");
