"Assign [Delegate Name]" blueprint context menu entries now properly create a custom event node again

diff --git a/Engine/Source/Editor/BlueprintGraph/Classes/K2Node_AssignDelegate.h b/Engine/Source/Editor/BlueprintGraph/Classes/K2Node_AssignDelegate.h
--- a/Engine/Source/Editor/BlueprintGraph/Classes/K2Node_AssignDelegate.h
+++ b/Engine/Source/Editor/BlueprintGraph/Classes/K2Node_AssignDelegate.h
@@ -25,13 +25,9 @@ public:
 	virtual FText GetMenuCategory() const override;
 	virtual bool IsCompatibleWithGraph(const UEdGraph* TargetGraph) const override;
 	virtual void PostPlacedNewNode() override;
-	virtual void AllocateDefaultPins() override;
 	// End of UEdGraphNode interface
 
 private:
 	/** Constructing FText strings can be costly, so we cache the node's title */
 	FNodeTextCache CachedListTitle;
-
-	/** If true, should try to spawn a bound custom event node on next allocate */
-	bool bShouldSpawnCustomEvent;
 };
diff --git a/Engine/Source/Editor/BlueprintGraph/Private/K2Node_AssignDelegate.cpp b/Engine/Source/Editor/BlueprintGraph/Private/K2Node_AssignDelegate.cpp
--- a/Engine/Source/Editor/BlueprintGraph/Private/K2Node_AssignDelegate.cpp
+++ b/Engine/Source/Editor/BlueprintGraph/Private/K2Node_AssignDelegate.cpp
@@ -15,7 +15,6 @@
 UK2Node_AssignDelegate::UK2Node_AssignDelegate(const FObjectInitializer& ObjectInitializer)
 	: Super(ObjectInitializer)
 {
-	bShouldSpawnCustomEvent = false;
 }
 
 //------------------------------------------------------------------------------
@@ -75,17 +74,10 @@ bool UK2Node_AssignDelegate::IsCompatibleWithGraph(const UEdGraph* TargetGraph)
 void UK2Node_AssignDelegate::PostPlacedNewNode()
 {
 	Super::PostPlacedNewNode();
-	bShouldSpawnCustomEvent = true;
-}
-
-void UK2Node_AssignDelegate::AllocateDefaultPins()
-{
-	Super::AllocateDefaultPins();
-
+	
 	UEdGraphPin* InDelegatePin = GetDelegatePin();
-	if (InDelegatePin->LinkedTo.Num() == 0 && bShouldSpawnCustomEvent)
+	if (InDelegatePin->LinkedTo.Num() == 0)
 	{
-		bShouldSpawnCustomEvent = false;
 		if (UMulticastDelegateProperty* DelegateProp = Cast<UMulticastDelegateProperty>(GetProperty()))
 		{
 			FVector2D const LocationOffset(-150.f, +150.f);
