Make FKCHandler_CallFunction::GetFunctionNameFromNode return FName instead of FString
#rnx

diff --git a/Engine/Source/Editor/BlueprintGraph/Private/CallFunctionHandler.cpp b/Engine/Source/Editor/BlueprintGraph/Private/CallFunctionHandler.cpp
--- a/Engine/Source/Editor/BlueprintGraph/Private/CallFunctionHandler.cpp
+++ b/Engine/Source/Editor/BlueprintGraph/Private/CallFunctionHandler.cpp
@@ -437,7 +437,7 @@ void FKCHandler_CallFunction::CreateFunctionCallStatement(FKismetFunctionContext
 	}
 	else
 	{
-		FString WarningMessage = FText::Format(LOCTEXT("FindFunction_ErrorFmt", "Could not find the function '{0}' called from @@"), FText::FromString(GetFunctionNameFromNode(Node))).ToString();
+		FString WarningMessage = FText::Format(LOCTEXT("FindFunction_ErrorFmt", "Could not find the function '{0}' called from @@"), FText::FromString(GetFunctionNameFromNode(Node).ToString())).ToString();
 		CompilerContext.MessageLog.Warning(*WarningMessage, Node);
 	}
 }
@@ -589,9 +589,8 @@ UFunction* FKCHandler_CallFunction::FindFunction(FKismetFunctionContext& Context
 
 	if (CallingContext)
 	{
-		FString FunctionName = GetFunctionNameFromNode(Node);
-
-		return CallingContext->FindFunctionByName(*FunctionName);
+		const FName FunctionName = GetFunctionNameFromNode(Node);
+		return CallingContext->FindFunctionByName(FunctionName);
 	}
 
 	return nullptr;
@@ -712,17 +711,17 @@ void FKCHandler_CallFunction::CheckIfFunctionIsCallable(UFunction* Function, FKi
 }
 
 // Get the name of the function to call from the node
-FString FKCHandler_CallFunction::GetFunctionNameFromNode(UEdGraphNode* Node) const
+FName FKCHandler_CallFunction::GetFunctionNameFromNode(UEdGraphNode* Node) const
 {
 	UK2Node_CallFunction* CallFuncNode = Cast<UK2Node_CallFunction>(Node);
 	if (CallFuncNode)
 	{
-		return CallFuncNode->FunctionReference.GetMemberName().ToString();
+		return CallFuncNode->FunctionReference.GetMemberName();
 	}
 	else
 	{
 		CompilerContext.MessageLog.Error(*NSLOCTEXT("KismetCompiler", "UnableResolveFunctionName_Error", "Unable to resolve function name for @@").ToString(), Node);
-		return TEXT("");
+		return NAME_None;
 	}
 }
 
diff --git a/Engine/Source/Editor/BlueprintGraph/Public/CallFunctionHandler.h b/Engine/Source/Editor/BlueprintGraph/Public/CallFunctionHandler.h
--- a/Engine/Source/Editor/BlueprintGraph/Public/CallFunctionHandler.h
+++ b/Engine/Source/Editor/BlueprintGraph/Public/CallFunctionHandler.h
@@ -55,7 +55,7 @@ public:
 
 private:
 	// Get the name of the function to call from the node
-	virtual FString GetFunctionNameFromNode(UEdGraphNode* Node) const;
+	virtual FName GetFunctionNameFromNode(UEdGraphNode* Node) const;
 	UClass* GetCallingContext(FKismetFunctionContext& Context, UEdGraphNode* Node);
 	UClass* GetTrueCallingClass(FKismetFunctionContext& Context, UEdGraphPin* SelfPin);
 
