(4.21.1) Fix potential ensure when iterating LinkedTo as BreakLinkTo can mutate it.
#rnx

diff --git a/Engine/Source/Editor/BlueprintGraph/Private/K2Node_MakeContainer.cpp b/Engine/Source/Editor/BlueprintGraph/Private/K2Node_MakeContainer.cpp
--- a/Engine/Source/Editor/BlueprintGraph/Private/K2Node_MakeContainer.cpp
+++ b/Engine/Source/Editor/BlueprintGraph/Private/K2Node_MakeContainer.cpp
@@ -344,7 +344,8 @@ void UK2Node_MakeContainer::PropagatePinType()
 				}
 
 				// Verify that all previous connections to this pin are still valid with the new type
-				for (UEdGraphPin* ConnectedPin : CurrentPin->LinkedTo)
+				TArray<UEdGraphPin*> LinkedToCopy = CurrentPin->LinkedTo;
+				for (UEdGraphPin* ConnectedPin : LinkedToCopy)
 				{
 					if (!Schema->ArePinsCompatible(CurrentPin, ConnectedPin, CallingContext))
 					{
