Remove breakpoints reference a graph when the graph is deleted

diff --git a/Engine/Source/Editor/BlueprintGraph/Private/EdGraphSchema_K2.cpp b/Engine/Source/Editor/BlueprintGraph/Private/EdGraphSchema_K2.cpp
--- a/Engine/Source/Editor/BlueprintGraph/Private/EdGraphSchema_K2.cpp
+++ b/Engine/Source/Editor/BlueprintGraph/Private/EdGraphSchema_K2.cpp
@@ -2,6 +2,7 @@
 
 #include "EdGraphSchema_K2.h"
 #include "BlueprintCompilationManager.h"
+#include "Engine/Breakpoint.h"
 #include "Modules/ModuleManager.h"
 #include "UObject/Interface.h"
 #include "UObject/UnrealType.h"
@@ -4487,9 +4488,20 @@ void UEdGraphSchema_K2::HandleGraphBeingDeleted(UEdGraph& GraphBeingRemoved) con
 			NodeToDelete->Modify();
 			NodeToDelete->DestroyNode();
 		}
+		
+		// likely tagged as modified by caller, but make sure:
+		Blueprint->Modify();
 
 		// Remove from the list of recently edited documents
 		Blueprint->LastEditedDocuments.RemoveAll([&GraphBeingRemoved](const FEditedDocumentInfo& TestDoc) { return TestDoc.EditedObjectPath.ResolveObject() == &GraphBeingRemoved; });
+
+		// Remove any BPs that reference a node in this graph:
+		Blueprint->Breakpoints.RemoveAll(
+			[&GraphBeingRemoved](UBreakpoint* Breakpoint)
+			{
+				return !Breakpoint || (Breakpoint->GetLocation() && Breakpoint->GetLocation()->IsIn(&GraphBeingRemoved));
+			}
+		);
 	}
 }
 
