This allows re-evaluating knot nodes to fix type problems

diff --git a/Engine/Plugins/Editor/AssetManagerEditor/Source/AssetManagerEditor/Private/ReferenceViewer/ReferenceViewerSchema.cpp b/Engine/Plugins/Editor/AssetManagerEditor/Source/AssetManagerEditor/Private/ReferenceViewer/ReferenceViewerSchema.cpp
--- a/Engine/Plugins/Editor/AssetManagerEditor/Source/AssetManagerEditor/Private/ReferenceViewer/ReferenceViewerSchema.cpp
+++ b/Engine/Plugins/Editor/AssetManagerEditor/Source/AssetManagerEditor/Private/ReferenceViewer/ReferenceViewerSchema.cpp
@@ -106,13 +106,13 @@ void UReferenceViewerSchema::BreakSinglePinLink(UEdGraphPin* SourcePin, UEdGraph
 	// Don't allow breaking any links
 }
 
-FPinConnectionResponse UReferenceViewerSchema::MovePinLinks(UEdGraphPin& MoveFromPin, UEdGraphPin& MoveToPin, bool bIsItermediateMove) const
+FPinConnectionResponse UReferenceViewerSchema::MovePinLinks(UEdGraphPin& MoveFromPin, UEdGraphPin& MoveToPin, bool bIsIntermediateMove, bool bNotifyLinkedNodes) const
 {
 	// Don't allow moving any links
 	return FPinConnectionResponse(CONNECT_RESPONSE_DISALLOW, FString());
 }
 
-FPinConnectionResponse UReferenceViewerSchema::CopyPinLinks(UEdGraphPin& CopyFromPin, UEdGraphPin& CopyToPin, bool bIsItermediateCopy) const
+FPinConnectionResponse UReferenceViewerSchema::CopyPinLinks(UEdGraphPin& CopyFromPin, UEdGraphPin& CopyToPin, bool bIsIntermediateCopy) const
 {
 	// Don't allow copying any links
 	return FPinConnectionResponse(CONNECT_RESPONSE_DISALLOW, FString());
diff --git a/Engine/Plugins/Editor/AssetManagerEditor/Source/AssetManagerEditor/Private/ReferenceViewer/ReferenceViewerSchema.h b/Engine/Plugins/Editor/AssetManagerEditor/Source/AssetManagerEditor/Private/ReferenceViewer/ReferenceViewerSchema.h
--- a/Engine/Plugins/Editor/AssetManagerEditor/Source/AssetManagerEditor/Private/ReferenceViewer/ReferenceViewerSchema.h
+++ b/Engine/Plugins/Editor/AssetManagerEditor/Source/AssetManagerEditor/Private/ReferenceViewer/ReferenceViewerSchema.h
@@ -22,8 +22,8 @@ public:
 	virtual FLinearColor GetPinTypeColor(const FEdGraphPinType& PinType) const override;
 	virtual void BreakPinLinks(UEdGraphPin& TargetPin, bool bSendsNodeNotifcation) const override;
 	virtual void BreakSinglePinLink(UEdGraphPin* SourcePin, UEdGraphPin* TargetPin) const override;
-	virtual FPinConnectionResponse MovePinLinks(UEdGraphPin& MoveFromPin, UEdGraphPin& MoveToPin, bool bIsItermeadiateMove = false) const override;
-	virtual FPinConnectionResponse CopyPinLinks(UEdGraphPin& CopyFromPin, UEdGraphPin& CopyToPin, bool bIsItermeadiateCopy = false) const override;
+	virtual FPinConnectionResponse MovePinLinks(UEdGraphPin& MoveFromPin, UEdGraphPin& MoveToPin, bool bIsIntermediateMove = false, bool bNotifyLinkedNodes = false) const override;
+	virtual FPinConnectionResponse CopyPinLinks(UEdGraphPin& CopyFromPin, UEdGraphPin& CopyToPin, bool bIsIntermediateCopy = false) const override;
 	virtual class FConnectionDrawingPolicy* CreateConnectionDrawingPolicy(int32 InBackLayerID, int32 InFrontLayerID, float InZoomFactor, const FSlateRect& InClippingRect, class FSlateWindowElementList& InDrawElements, class UEdGraph* InGraphObj) const override;
 	virtual void DroppedAssetsOnGraph(const TArray<FAssetData>& Assets, const FVector2D& GraphPosition, UEdGraph* Graph) const override;
 	virtual void GetAssetsGraphHoverMessage(const TArray<FAssetData>& Assets, const UEdGraph* HoverGraph, FString& OutTooltipText, bool& OutOkIcon) const override;
diff --git a/Engine/Source/Editor/BlueprintGraph/Private/EdGraphSchema_K2.cpp b/Engine/Source/Editor/BlueprintGraph/Private/EdGraphSchema_K2.cpp
--- a/Engine/Source/Editor/BlueprintGraph/Private/EdGraphSchema_K2.cpp
+++ b/Engine/Source/Editor/BlueprintGraph/Private/EdGraphSchema_K2.cpp
@@ -5719,7 +5719,7 @@ struct FBackwardCompatibilityConversionHelper
 				UEdGraphPin* ExecPin = OldNode->GetExecPin();
 				UEdGraphPin* ExecCastPin = CastNode->GetExecPin();
 				check(ExecCastPin);
-				if (!ExecPin || !Schema.MovePinLinks(*ExecPin, *ExecCastPin).CanSafeConnect())
+				if (!ExecPin || !Schema.MovePinLinks(*ExecPin, *ExecCastPin, false, true).CanSafeConnect())
 				{
 					UE_LOG(LogBlueprint, Warning, TEXT("BackwardCompatibilityNodeConversion Error 'cannot connect' in blueprint: %s, pin: %s"),
 						Blueprint ? *Blueprint->GetName() : TEXT("Unknown"),
@@ -5749,7 +5749,7 @@ struct FBackwardCompatibilityConversionHelper
 
 				UEdGraphPin* CastSourcePin = CastNode->GetCastSourcePin();
 				check(CastSourcePin);
-				if (!Schema.MovePinLinks(*OldBlueprintPin, *CastSourcePin).CanSafeConnect())
+				if (!Schema.MovePinLinks(*OldBlueprintPin, *CastSourcePin, false, true).CanSafeConnect())
 				{
 					UE_LOG(LogBlueprint, Warning, TEXT("BackwardCompatibilityNodeConversion Error 'cannot connect' in blueprint: %s, pin: %s"),
 						Blueprint ? *Blueprint->GetName() : TEXT("Unknown"),
@@ -5780,7 +5780,7 @@ struct FBackwardCompatibilityConversionHelper
 					if (OldPin)
 					{
 						OldPins.Add(OldPin);
-						if (!Schema.MovePinLinks(*OldPin, *Pin).CanSafeConnect())
+						if (!Schema.MovePinLinks(*OldPin, *Pin, false, true).CanSafeConnect())
 						{
 							UE_LOG(LogBlueprint, Warning, TEXT("BackwardCompatibilityNodeConversion Error 'cannot connect' in blueprint: %s, pin: %s"),
 								Blueprint ? *Blueprint->GetName() : TEXT("Unknown"),
@@ -5964,7 +5964,7 @@ bool UEdGraphSchema_K2::ReplaceOldNodeWithNew(UK2Node* OldNode, UK2Node* NewNode
 			{
 				continue;
 			}
-			else if (!Schema->MovePinLinks(*OldPin, *NewPin).CanSafeConnect())
+			else if (!Schema->MovePinLinks(*OldPin, *NewPin, false, true).CanSafeConnect())
 			{
 				UE_LOG(LogBlueprint, Warning, TEXT("BackwardCompatibilityNodeConversion Error 'cannot safely move pin %s to %s' in blueprint: %s"),
 					*OldPin->PinName.ToString(),
diff --git a/Engine/Source/Editor/PhysicsAssetEditor/Private/PhysicsAssetGraph/PhysicsAssetGraphSchema.cpp b/Engine/Source/Editor/PhysicsAssetEditor/Private/PhysicsAssetGraph/PhysicsAssetGraphSchema.cpp
--- a/Engine/Source/Editor/PhysicsAssetEditor/Private/PhysicsAssetGraph/PhysicsAssetGraphSchema.cpp
+++ b/Engine/Source/Editor/PhysicsAssetEditor/Private/PhysicsAssetGraph/PhysicsAssetGraphSchema.cpp
@@ -52,13 +52,13 @@ void UPhysicsAssetGraphSchema::BreakSinglePinLink(UEdGraphPin* SourcePin, UEdGra
 	// Don't allow breaking any links
 }
 
-FPinConnectionResponse UPhysicsAssetGraphSchema::MovePinLinks(UEdGraphPin& MoveFromPin, UEdGraphPin& MoveToPin, bool bIsItermeadiateMove) const
+FPinConnectionResponse UPhysicsAssetGraphSchema::MovePinLinks(UEdGraphPin& MoveFromPin, UEdGraphPin& MoveToPin, bool bIsIntermediateMove, bool bNotifyLinkedNodes) const
 {
 	// Don't allow moving any links
 	return FPinConnectionResponse(CONNECT_RESPONSE_DISALLOW, TEXT(""));
 }
 
-FPinConnectionResponse UPhysicsAssetGraphSchema::CopyPinLinks(UEdGraphPin& CopyFromPin, UEdGraphPin& CopyToPin, bool bIsItermeadiateCopy) const
+FPinConnectionResponse UPhysicsAssetGraphSchema::CopyPinLinks(UEdGraphPin& CopyFromPin, UEdGraphPin& CopyToPin, bool bIsIntermediateCopy) const
 {
 	// Don't allow copying any links
 	return FPinConnectionResponse(CONNECT_RESPONSE_DISALLOW, TEXT(""));
diff --git a/Engine/Source/Editor/PhysicsAssetEditor/Private/PhysicsAssetGraph/PhysicsAssetGraphSchema.h b/Engine/Source/Editor/PhysicsAssetEditor/Private/PhysicsAssetGraph/PhysicsAssetGraphSchema.h
--- a/Engine/Source/Editor/PhysicsAssetEditor/Private/PhysicsAssetGraph/PhysicsAssetGraphSchema.h
+++ b/Engine/Source/Editor/PhysicsAssetEditor/Private/PhysicsAssetGraph/PhysicsAssetGraphSchema.h
@@ -25,8 +25,8 @@ public:
 	virtual FLinearColor GetPinTypeColor(const FEdGraphPinType& PinType) const override;
 	virtual void BreakPinLinks(UEdGraphPin& TargetPin, bool bSendsNodeNotifcation) const override;
 	virtual void BreakSinglePinLink(UEdGraphPin* SourcePin, UEdGraphPin* TargetPin) const override;
-	virtual FPinConnectionResponse MovePinLinks(UEdGraphPin& MoveFromPin, UEdGraphPin& MoveToPin, bool bIsItermeadiateMove = false) const override;
-	virtual FPinConnectionResponse CopyPinLinks(UEdGraphPin& CopyFromPin, UEdGraphPin& CopyToPin, bool bIsItermeadiateCopy = false) const override;
+	virtual FPinConnectionResponse MovePinLinks(UEdGraphPin& MoveFromPin, UEdGraphPin& MoveToPin, bool bIsIntermediateMove = false, bool bNotifyLinkedNodes = false) const override;
+	virtual FPinConnectionResponse CopyPinLinks(UEdGraphPin& CopyFromPin, UEdGraphPin& CopyToPin, bool bIsIntermediateCopy = false) const override;
 	virtual class FConnectionDrawingPolicy* CreateConnectionDrawingPolicy(int32 InBackLayerID, int32 InFrontLayerID, float InZoomFactor, const FSlateRect& InClippingRect, class FSlateWindowElementList& InDrawElements, class UEdGraph* InGraphObj) const override;
 	virtual bool ShouldAlwaysPurgeOnModification() const override { return false; }
 	virtual FPinConnectionResponse CanCreateNewNodes(UEdGraphPin* InSourcePin) const override;
diff --git a/Engine/Source/Runtime/Engine/Classes/EdGraph/EdGraphSchema.h b/Engine/Source/Runtime/Engine/Classes/EdGraph/EdGraphSchema.h
--- a/Engine/Source/Runtime/Engine/Classes/EdGraph/EdGraphSchema.h
+++ b/Engine/Source/Runtime/Engine/Classes/EdGraph/EdGraphSchema.h
@@ -842,8 +842,15 @@ class ENGINE_API UEdGraphSchema : public UObject
 	/** Handles double-clicking on a pin<->pin connection */
 	virtual void OnPinConnectionDoubleCicked(UEdGraphPin* PinA, UEdGraphPin* PinB, const FVector2D& GraphPosition) const { }
 
-	/** Break links on this pin and create links instead on MoveToPin */
-	virtual FPinConnectionResponse MovePinLinks(UEdGraphPin& MoveFromPin, UEdGraphPin& MoveToPin, bool bIsIntermediateMove = false) const;
+	/**
+	 * Break links on this pin and create links instead on MoveToPin
+	 * 
+	 * @param	MoveFromPin			Pin we are breaking links from
+	 * @param	MoveToPin			Pin we are copying links to
+	 * @param	bIsIntermediateMove	Allows linking to transient pins, should only be true when called from utility functions
+	 * @param	bNotifyLinkedNodes	If true, it will notify linked nodes if it fails to move connection, this allows type fixup
+	 */
+	virtual FPinConnectionResponse MovePinLinks(UEdGraphPin& MoveFromPin, UEdGraphPin& MoveToPin, bool bIsIntermediateMove = false, bool bNotifyLinkedNodes = false) const;
 	 
 	/** Copies pin links from one pin to another without breaking the original links */
 	virtual FPinConnectionResponse CopyPinLinks(UEdGraphPin& CopyFromPin, UEdGraphPin& CopyToPin, bool bIsIntermediateCopy = false) const;
diff --git a/Engine/Source/Runtime/Engine/Private/EdGraph/EdGraphSchema.cpp b/Engine/Source/Runtime/Engine/Private/EdGraph/EdGraphSchema.cpp
--- a/Engine/Source/Runtime/Engine/Private/EdGraph/EdGraphSchema.cpp
+++ b/Engine/Source/Runtime/Engine/Private/EdGraph/EdGraphSchema.cpp
@@ -622,7 +622,7 @@ void UEdGraphSchema::BreakSinglePinLink(UEdGraphPin* SourcePin, UEdGraphPin* Tar
 #endif	//#if WITH_EDITOR
 }
 
-FPinConnectionResponse UEdGraphSchema::MovePinLinks(UEdGraphPin& MoveFromPin, UEdGraphPin& MoveToPin, bool bIsIntermediateMove) const
+FPinConnectionResponse UEdGraphSchema::MovePinLinks(UEdGraphPin& MoveFromPin, UEdGraphPin& MoveToPin, bool bIsIntermediateMove, bool bNotifyLinkedNodes) const
 {
 #if WITH_EDITOR
 	ensureMsgf(bIsIntermediateMove || !MoveToPin.GetOwningNode()->GetGraph()->HasAnyFlags(RF_Transient),
@@ -633,7 +633,7 @@ FPinConnectionResponse UEdGraphSchema::MovePinLinks(UEdGraphPin& MoveFromPin, UE
 	// First copy the current set of links
 	TArray<UEdGraphPin*> CurrentLinks = MoveFromPin.LinkedTo;
 	// Then break all links at pin we are moving from
-	MoveFromPin.BreakAllPinLinks();
+	MoveFromPin.BreakAllPinLinks(false);
 	// Try and make each new connection
 	for (int32 i=0; i<CurrentLinks.Num(); i++)
 	{
@@ -644,12 +644,24 @@ FPinConnectionResponse UEdGraphSchema::MovePinLinks(UEdGraphPin& MoveFromPin, UE
 #endif
 		{
 			FPinConnectionResponse Response = CanCreateConnection(&MoveToPin, NewLink);
+
+#if WITH_EDITOR
+			if (!Response.CanSafeConnect() && bNotifyLinkedNodes)
+			{
+				// Connection failed, so notify and try again
+				if (UEdGraphNode* LinkedToNode = NewLink->GetOwningNodeUnchecked())
+				{
+					LinkedToNode->PinConnectionListChanged(NewLink);
+					Response = CanCreateConnection(&MoveToPin, NewLink);
+				}
+			}
+#endif
 			if (Response.CanSafeConnect())
 			{
 				MoveToPin.MakeLinkTo(NewLink);
-			}
+			}	
 			else
-			{
+			{ 
 				FinalResponse = Response;
 			}
 		}
