Add code to keep orphan pins around for any links lost during backward compatible node conversion, this avoids silent data loss

diff --git a/Engine/Source/Editor/BlueprintGraph/Private/EdGraphSchema_K2.cpp b/Engine/Source/Editor/BlueprintGraph/Private/EdGraphSchema_K2.cpp
--- a/Engine/Source/Editor/BlueprintGraph/Private/EdGraphSchema_K2.cpp
+++ b/Engine/Source/Editor/BlueprintGraph/Private/EdGraphSchema_K2.cpp
@@ -5619,6 +5619,39 @@ bool UEdGraphSchema_K2::FadeNodeWhenDraggingOffPin(const UEdGraphNode* Node, con
 
 struct FBackwardCompatibilityConversionHelper
 {
+	// Re-add orphaned pins to deal with any links that were lost during converstion
+	static bool RestoreOrphanLinks(UEdGraphPin* OldPin, UEdGraphPin* NewPin, UK2Node* NewNode, const TArray<UEdGraphPin*>& OldLinks)
+	{
+		// See if there are any links that didn't get copied, including to orphan pins or if the newpin is null
+		TArray<UEdGraphPin*> OrphanedLinks;
+
+		for (UEdGraphPin* OldLink : OldLinks)
+		{
+			if (!NewPin || !NewPin->LinkedTo.Contains(OldLink))
+			{
+				OrphanedLinks.Add(OldLink);
+			}
+		}
+
+		if (OrphanedLinks.Num() > 0)
+		{
+			// Add an orphan pin so warning/connections are not silently lost
+			UEdGraphPin* OrphanPin = NewNode->CreatePin(OldPin->Direction, OldPin->PinType, OldPin->PinName);
+
+			OrphanPin->bOrphanedPin = true;
+			OrphanPin->bNotConnectable = true;
+
+			for (UEdGraphPin* OldLink : OrphanedLinks)
+			{
+				OrphanPin->MakeLinkTo(OldLink);
+			}
+
+			return true;
+		}
+
+		return false;
+	}
+
 	static bool ConvertNode(
 		UK2Node* OldNode,
 		const FString& BlueprintPinName,
@@ -5779,13 +5812,17 @@ struct FBackwardCompatibilityConversionHelper
 					UEdGraphPin* OldPin = OldNode->FindPin(Pin->PinName);
 					if (OldPin)
 					{
+						TArray<UEdGraphPin*> OldLinks = OldPin->LinkedTo;
 						OldPins.Add(OldPin);
+
 						if (!Schema.MovePinLinks(*OldPin, *Pin, false, true).CanSafeConnect())
 						{
 							UE_LOG(LogBlueprint, Warning, TEXT("BackwardCompatibilityNodeConversion Error 'cannot connect' in blueprint: %s, pin: %s"),
 								Blueprint ? *Blueprint->GetName() : TEXT("Unknown"),
 								*Pin->PinName.ToString());
 						}
+
+						FBackwardCompatibilityConversionHelper::RestoreOrphanLinks(OldPin, Pin, NewNode, OldLinks);
 					}
 					else
 					{
@@ -5958,6 +5995,7 @@ bool UEdGraphSchema_K2::ReplaceOldNodeWithNew(UK2Node* OldNode, UK2Node* NewNode
 		{
 			UEdGraphPin* OldPin = OldNode->Pins[PinIdx];
 			UEdGraphPin* NewPin = NewPinArray[PinIdx];
+			TArray<UEdGraphPin*> OldLinks = OldPin->LinkedTo;
 
 			// could be null, meaning they didn't want to map this OldPin to anything
 			if (NewPin == nullptr)
@@ -5976,6 +6014,8 @@ bool UEdGraphSchema_K2::ReplaceOldNodeWithNew(UK2Node* OldNode, UK2Node* NewNode
 				// for wildcard pins, which may have to react to being connected with
 				NewNode->NotifyPinConnectionListChanged(NewPin);
 			}
+
+			FBackwardCompatibilityConversionHelper::RestoreOrphanLinks(OldPin, NewPin, NewNode, OldLinks);
 		}
 
 		NewNode->NodeComment = OldNode->NodeComment;
