Fix crash retreiving a graph for a graph that is (probably) in the process of being closed/deleted

diff --git a/Engine/Source/Editor/BlueprintGraph/Private/K2Node_CallFunction.cpp b/Engine/Source/Editor/BlueprintGraph/Private/K2Node_CallFunction.cpp
--- a/Engine/Source/Editor/BlueprintGraph/Private/K2Node_CallFunction.cpp
+++ b/Engine/Source/Editor/BlueprintGraph/Private/K2Node_CallFunction.cpp
@@ -1618,7 +1618,7 @@ bool UK2Node_CallFunction::ShouldDrawAsBead() const
 bool UK2Node_CallFunction::ShouldShowNodeProperties() const
 {
 	// Show node properties if this corresponds to a function graph
-	if (FunctionReference.GetMemberName() != NAME_None)
+	if (FunctionReference.GetMemberName() != NAME_None && HasValidBlueprint())
 	{
 		return FindObject<UEdGraph>(GetBlueprint(), *(FunctionReference.GetMemberName().ToString())) != NULL;
 	}
