PR #5336: Add auto conversion from class to string (Contributed by [REDACTED_PERSON])

diff --git a/Engine/Source/Editor/BlueprintGraph/Private/EdGraphSchema_K2.cpp b/Engine/Source/Editor/BlueprintGraph/Private/EdGraphSchema_K2.cpp
--- a/Engine/Source/Editor/BlueprintGraph/Private/EdGraphSchema_K2.cpp
+++ b/Engine/Source/Editor/BlueprintGraph/Private/EdGraphSchema_K2.cpp
@@ -2564,6 +2564,15 @@ bool UEdGraphSchema_K2::SearchForAutocastFunction(const UEdGraphPin* OutputPin,
 			FunctionOwner = Function->GetOwnerClass();
 		}
 	}
+	else if (OutputPin->PinType.PinCategory == PC_Class)
+	{
+		if (InputPin->PinType.PinCategory == PC_String)
+		{
+			UFunction* Function = UKismetSystemLibrary::StaticClass()->FindFunctionByName(GET_MEMBER_NAME_CHECKED(UKismetSystemLibrary, GetClassDisplayName));
+			TargetFunction = Function->GetFName();
+			FunctionOwner = Function->GetOwnerClass();
+		}
+	}
 	else if (OutputPin->PinType.PinCategory == PC_Struct)
 	{
 		const UScriptStruct* OutputStructType = Cast<const UScriptStruct>(OutputPin->PinType.PinSubCategoryObject.Get());
