[NO NON-PII COMMIT MESSAGE]

diff --git a/Engine/Source/Editor/BlueprintGraph/Private/K2Node.cpp b/Engine/Source/Editor/BlueprintGraph/Private/K2Node.cpp
--- a/Engine/Source/Editor/BlueprintGraph/Private/K2Node.cpp
+++ b/Engine/Source/Editor/BlueprintGraph/Private/K2Node.cpp
@@ -82,6 +82,7 @@ void UK2Node::Serialize(FArchive& Ar)
 			if (!Pin->bDefaultValueIsIgnored && !Pin->DefaultValue.IsEmpty() )
 			{
 				// If looking for references during save, expand any default values on the pins
+				// This is only reliable when saving in the editor, the cook case is handled below
 				if (Ar.IsObjectReferenceCollector() && Pin->PinType.PinCategory == UEdGraphSchema_K2::PC_Struct && Pin->PinType.PinSubCategoryObject.IsValid())
 				{
 					UScriptStruct* Struct = Cast<UScriptStruct>(Pin->PinType.PinSubCategoryObject.Get());
@@ -159,25 +160,37 @@ void UK2Node::FixupPinDefaultValues()
 		}
 	}
 
-	// Fix asset ptr pins
-	if (LinkerFrameworkVersion < FFrameworkObjectVersion::ChangeAssetPinsToString)
+	// Fix soft object ptr pins
+	if (GIsEditor || LinkerFrameworkVersion < FFrameworkObjectVersion::ChangeAssetPinsToString)
 	{
-		bool bFoundPin = false;
-		for (int32 i = 0; i < Pins.Num() && !bFoundPin; ++i)
+		FSoftObjectPathSerializationScope SetPackage(GetOutermost()->GetFName(), NAME_None, ESoftObjectPathCollectType::AlwaysCollect, ESoftObjectPathSerializeType::SkipSerializeIfArchiveHasSize);
+		for (int32 i = 0; i < Pins.Num(); ++i)
 		{
 			UEdGraphPin* Pin = Pins[i];
-
 			if (Pin->PinType.PinCategory == UEdGraphSchema_K2::PC_SoftObject || Pin->PinType.PinCategory == UEdGraphSchema_K2::PC_SoftClass)
 			{
-				if (Pin->DefaultObject && Pin->DefaultValue.IsEmpty())
+				// Fix old assetptr pins
+				if (LinkerFrameworkVersion < FFrameworkObjectVersion::ChangeAssetPinsToString)
+				{
+					if (Pin->DefaultObject && Pin->DefaultValue.IsEmpty())
+					{
+						Pin->DefaultValue = Pin->DefaultObject->GetPathName();
+						Pin->DefaultObject = nullptr;
+					}
+				}
+
+				// In editor, fixup soft object ptrs on load on to handle redirects and finding refs for cooking
+				// We're not handling soft object ptrs inside FStructs because it's a rare edge case and would be a performance hit on load
+				if (GIsEditor && !Pin->DefaultValue.IsEmpty())
 				{
-					Pin->DefaultValue = Pin->DefaultObject->GetPathName();
-					Pin->DefaultObject = nullptr;
+					FSoftObjectPath TempRef(Pin->DefaultValue);
+					TempRef.PostLoadPath();
+					TempRef.PreSavePath();
+					Pin->DefaultValue = TempRef.ToString();
 				}
 			}
 		}
 	}
-
 }
 
 FText UK2Node::GetToolTipHeading() const
diff --git a/Engine/Source/Editor/BlueprintGraph/Private/K2Node_FunctionEntry.cpp b/Engine/Source/Editor/BlueprintGraph/Private/K2Node_FunctionEntry.cpp
--- a/Engine/Source/Editor/BlueprintGraph/Private/K2Node_FunctionEntry.cpp
+++ b/Engine/Source/Editor/BlueprintGraph/Private/K2Node_FunctionEntry.cpp
@@ -248,6 +248,7 @@ void UK2Node_FunctionEntry::Serialize(FArchive& Ar)
 			if (!LocalVariable.DefaultValue.IsEmpty())
 			{
 				// If looking for references during save, expand any default values on the local variables
+				// This is only reliable when saving in the editor, the cook case is handled below
 				if (Ar.IsObjectReferenceCollector() && LocalVariable.VarType.PinCategory == UEdGraphSchema_K2::PC_Struct && LocalVariable.VarType.PinSubCategoryObject.IsValid())
 				{
 					UScriptStruct* Struct = Cast<UScriptStruct>(LocalVariable.VarType.PinSubCategoryObject.Get());
@@ -335,6 +336,24 @@ void UK2Node_FunctionEntry::Serialize(FArchive& Ar)
 				}
 			}
 		}
+
+		// In editor, fixup soft object ptrs on load on to handle redirects and finding refs for cooking
+		// We're not handling soft object ptrs inside FStructs because it's a rare edge case and would be a performance hit on load
+		if (GIsEditor)
+		{
+			FSoftObjectPathSerializationScope SetPackage(GetOutermost()->GetFName(), NAME_None, ESoftObjectPathCollectType::AlwaysCollect, ESoftObjectPathSerializeType::SkipSerializeIfArchiveHasSize);
+
+			for (FBPVariableDescription& LocalVariable : LocalVariables)
+			{
+				if (!LocalVariable.DefaultValue.IsEmpty() && (LocalVariable.VarType.PinCategory == UEdGraphSchema_K2::PC_SoftObject || LocalVariable.VarType.PinCategory == UEdGraphSchema_K2::PC_SoftClass))
+				{
+					FSoftObjectPath TempRef(LocalVariable.DefaultValue);
+					TempRef.PostLoadPath();
+					TempRef.PreSavePath();
+					LocalVariable.DefaultValue = TempRef.ToString();
+				}
+			}
+		}
 	}
 }
 
