Use metadata static instead of redefining it
#rnx

diff --git a/Engine/Source/Editor/BlueprintGraph/Private/K2Node_VariableGet.cpp b/Engine/Source/Editor/BlueprintGraph/Private/K2Node_VariableGet.cpp
--- a/Engine/Source/Editor/BlueprintGraph/Private/K2Node_VariableGet.cpp
+++ b/Engine/Source/Editor/BlueprintGraph/Private/K2Node_VariableGet.cpp
@@ -189,7 +189,6 @@ FText UK2Node_VariableGet::GetPropertyTooltip(UProperty const* VariableProperty)
 		UClass* SourceClass = VariableProperty->GetOwnerClass();
 		// discover if the variable property is a non blueprint user variable
 		bool const bIsNativeVariable = (SourceClass != nullptr) && (SourceClass->ClassGeneratedBy == nullptr);
-		FName const TooltipMetaKey(TEXT("tooltip"));
 
 		FText SubTooltip;
 		if (bIsNativeVariable)
@@ -199,7 +198,7 @@ FText UK2Node_VariableGet::GetPropertyTooltip(UProperty const* VariableProperty)
 			{
 				// See if the native property has a tooltip
 				SubTooltip = PropertyTooltip;
-				FString TooltipName = FString::Printf(TEXT("%s.%s"), *VarName.ToString(), *TooltipMetaKey.ToString());
+				FString TooltipName = FString::Printf(TEXT("%s.%s"), *VarName.ToString(), *FBlueprintMetadata::MD_Tooltip.ToString());
 				FText::FindText(*VariableProperty->GetFullGroupName(true), *TooltipName, SubTooltip);
 			}
 		}
@@ -208,7 +207,7 @@ FText UK2Node_VariableGet::GetPropertyTooltip(UProperty const* VariableProperty)
 			if (UBlueprint* VarBlueprint = Cast<UBlueprint>(SourceClass->ClassGeneratedBy))
 			{
 				FString UserTooltipData;
-				if (FBlueprintEditorUtils::GetBlueprintVariableMetaData(VarBlueprint, VarName, VariableProperty->GetOwnerStruct(), TooltipMetaKey, UserTooltipData))
+				if (FBlueprintEditorUtils::GetBlueprintVariableMetaData(VarBlueprint, VarName, VariableProperty->GetOwnerStruct(), FBlueprintMetadata::MD_Tooltip, UserTooltipData))
 				{
 					SubTooltip = FText::FromString(UserTooltipData);
 				}
@@ -229,13 +228,12 @@ FText UK2Node_VariableGet::GetPropertyTooltip(UProperty const* VariableProperty)
 
 FText UK2Node_VariableGet::GetBlueprintVarTooltip(FBPVariableDescription const& VarDesc)
 {
-	FName const TooltipMetaKey(TEXT("tooltip"));
-	int32 const MetaIndex = VarDesc.FindMetaDataEntryIndexForKey(TooltipMetaKey);
+	int32 const MetaIndex = VarDesc.FindMetaDataEntryIndexForKey(FBlueprintMetadata::MD_Tooltip);
 	bool const bHasTooltipData = (MetaIndex != INDEX_NONE);
 
 	if (bHasTooltipData)
 	{
-		FString UserTooltipData = VarDesc.GetMetaData(TooltipMetaKey);
+		FString UserTooltipData = VarDesc.GetMetaData(FBlueprintMetadata::MD_Tooltip);
 
 		FFormatNamedArguments Args;
 		Args.Add(TEXT("VarName"), FText::FromName(VarDesc.VarName));
diff --git a/Engine/Source/Editor/BlueprintGraph/Private/K2Node_VariableSet.cpp b/Engine/Source/Editor/BlueprintGraph/Private/K2Node_VariableSet.cpp
--- a/Engine/Source/Editor/BlueprintGraph/Private/K2Node_VariableSet.cpp
+++ b/Engine/Source/Editor/BlueprintGraph/Private/K2Node_VariableSet.cpp
@@ -152,7 +152,6 @@ FText UK2Node_VariableSet::GetPropertyTooltip(UProperty const* VariableProperty)
 		UClass* SourceClass = VariableProperty->GetOwnerClass();
 		// discover if the variable property is a non blueprint user variable
 		bool const bIsNativeVariable = (SourceClass != nullptr) && (SourceClass->ClassGeneratedBy == nullptr);
-		FName const TooltipMetaKey(TEXT("tooltip"));
 
 		FText SubTooltip;
 		if (bIsNativeVariable)
@@ -162,7 +161,7 @@ FText UK2Node_VariableSet::GetPropertyTooltip(UProperty const* VariableProperty)
 			{
 				// See if the native property has a tooltip
 				SubTooltip = PropertyTooltip;
-				FString TooltipName = FString::Printf(TEXT("%s.%s"), *VarName.ToString(), *TooltipMetaKey.ToString());
+				FString TooltipName = FString::Printf(TEXT("%s.%s"), *VarName.ToString(), *FBlueprintMetadata::MD_Tooltip.ToString());
 				FText::FindText(*VariableProperty->GetFullGroupName(true), *TooltipName, SubTooltip);
 			}
 		}
@@ -171,7 +170,7 @@ FText UK2Node_VariableSet::GetPropertyTooltip(UProperty const* VariableProperty)
 			if (UBlueprint* VarBlueprint = Cast<UBlueprint>(SourceClass->ClassGeneratedBy))
 			{
 				FString UserTooltipData;
-				if (FBlueprintEditorUtils::GetBlueprintVariableMetaData(VarBlueprint, VarName, VariableProperty->GetOwnerStruct(), TooltipMetaKey, UserTooltipData))
+				if (FBlueprintEditorUtils::GetBlueprintVariableMetaData(VarBlueprint, VarName, VariableProperty->GetOwnerStruct(), FBlueprintMetadata::MD_Tooltip, UserTooltipData))
 				{
 					SubTooltip = FText::FromString(UserTooltipData);
 				}
@@ -205,13 +204,12 @@ FText UK2Node_VariableSet::GetPropertyTooltip(UProperty const* VariableProperty)
 
 FText UK2Node_VariableSet::GetBlueprintVarTooltip(FBPVariableDescription const& VarDesc)
 {
-	FName const TooltipMetaKey(TEXT("tooltip"));
-	int32 const MetaIndex = VarDesc.FindMetaDataEntryIndexForKey(TooltipMetaKey);
+	int32 const MetaIndex = VarDesc.FindMetaDataEntryIndexForKey(FBlueprintMetadata::MD_Tooltip);
 	bool const bHasTooltipData = (MetaIndex != INDEX_NONE);
 
 	if (bHasTooltipData)
 	{
-		FString UserTooltipData = VarDesc.GetMetaData(TooltipMetaKey);
+		FString UserTooltipData = VarDesc.GetMetaData(FBlueprintMetadata::MD_Tooltip);
 
 		FFormatNamedArguments Args;
 		Args.Add(TEXT("VarName"), FText::FromName(VarDesc.VarName));
