PR #5284: Allow multiple ExpandEnumsAsExecs parameters (Contributed by [REDACTED_PERSON])

diff --git a/Engine/Source/Editor/BlueprintGraph/Classes/K2Node_CallFunction.h b/Engine/Source/Editor/BlueprintGraph/Classes/K2Node_CallFunction.h
--- a/Engine/Source/Editor/BlueprintGraph/Classes/K2Node_CallFunction.h
+++ b/Engine/Source/Editor/BlueprintGraph/Classes/K2Node_CallFunction.h
@@ -64,6 +64,8 @@ private:
 	/** Flag used to track validity of pin tooltips, when tooltips are invalid they will be refreshed before being displayed */
 	mutable bool bPinTooltipsValid;
 
+	TArray<UEdGraphPin*> ExpandAsEnumPins;
+
 public:
 
 	// UObject interface
@@ -199,6 +201,8 @@ public:
 	/** */
 	static FSlateIcon GetPaletteIconForFunction(UFunction const* Function, FLinearColor& OutColor);
 
+	static void GetExpandEnumPinNames(const UFunction* Function, TArray<FName>& EnumNamesToCheck);
+
 private: 
 	/* Looks at function metadata and properties to determine if this node should be using enum to exec expansion */
 	void DetermineWantsEnumToExecExpansion(const UFunction* Function);
diff --git a/Engine/Source/Editor/BlueprintGraph/Private/EdGraphSchema_K2.cpp b/Engine/Source/Editor/BlueprintGraph/Private/EdGraphSchema_K2.cpp
--- a/Engine/Source/Editor/BlueprintGraph/Private/EdGraphSchema_K2.cpp
+++ b/Engine/Source/Editor/BlueprintGraph/Private/EdGraphSchema_K2.cpp
@@ -204,7 +204,7 @@ struct FUnloadedAssetData
 		, AssetFriendlyName(FText::FromString(FName::NameToDisplayString(InAsset.AssetName.ToString(), false)))
 		, PossibleObjectReferenceTypes(InPossibleObjectReferenceTypes)
 	{
-		InAsset.GetTagValue("Tooltip", Tooltip);
+		InAsset.GetTagValue(FBlueprintMetadata::MD_Tooltip, Tooltip);
 		if (Tooltip.IsEmpty())
 		{
 			Tooltip = FText::FromString(InAsset.ObjectPath.ToString());
diff --git a/Engine/Source/Editor/BlueprintGraph/Private/K2Node_CallFunction.cpp b/Engine/Source/Editor/BlueprintGraph/Private/K2Node_CallFunction.cpp
--- a/Engine/Source/Editor/BlueprintGraph/Private/K2Node_CallFunction.cpp
+++ b/Engine/Source/Editor/BlueprintGraph/Private/K2Node_CallFunction.cpp
@@ -19,6 +19,7 @@
 #include "K2Node_FunctionEntry.h"
 #include "K2Node_IfThenElse.h"
 #include "K2Node_TemporaryVariable.h"
+#include "K2Node_ExecutionSequence.h"
 #include "Kismet2/BlueprintEditorUtils.h"
 #include "Classes/EditorStyleSettings.h"
 #include "Editor.h"
@@ -557,7 +558,12 @@ void UK2Node_CallFunction::GetPinHoverText(const UEdGraphPin& Pin, FString& Hove
 	{
 		for (UEdGraphPin* P : Pins)
 		{
-			P->PinToolTip.Empty();
+			if (!P->PinToolTip.IsEmpty() && ExpandAsEnumPins.Contains(P))
+			{
+				continue;
+			}
+
+			P->PinToolTip.Reset();
 			GeneratePinTooltip(*P);
 		}
 
@@ -735,57 +741,105 @@ void UK2Node_CallFunction::CreateExecPinsForFunctionCall(const UFunction* Functi
 	bool bCreateSingleExecInputPin = true;
 	bool bCreateThenPin = true;
 	
+	ExpandAsEnumPins.Reset();
+
 	// If not pure, create exec pins
 	if (!bIsPureFunc)
 	{
 		// If we want enum->exec expansion, and it is not disabled, do it now
 		if(bWantsEnumToExecExpansion)
 		{
-			const FString& EnumParamName = Function->GetMetaData(FBlueprintMetadata::MD_ExpandEnumAsExecs);
+			TArray<FName> EnumNames;
+			GetExpandEnumPinNames(Function, EnumNames);
 
-			UProperty* Prop = nullptr;
-			UEnum* Enum = nullptr;
+			UProperty* PreviousInput = nullptr;
 
-			if(UByteProperty* ByteProp = FindField<UByteProperty>(Function, FName(*EnumParamName)))
-			{
-				Prop = ByteProp;
-				Enum = ByteProp->Enum;
-			}
-			else if(UEnumProperty* EnumProp = FindField<UEnumProperty>(Function, FName(*EnumParamName)))
+			for (const FName& EnumParamName : EnumNames)
 			{
-				Prop = EnumProp;
-				Enum = EnumProp->GetEnum();
-			}
+				UProperty* Prop = nullptr;
+				UEnum* Enum = nullptr;
 
-			if(Prop != nullptr && Enum != nullptr)
-			{
-				const bool bIsFunctionInput = !Prop->HasAnyPropertyFlags(CPF_ReturnParm) &&
-					(!Prop->HasAnyPropertyFlags(CPF_OutParm) ||
-					 Prop->HasAnyPropertyFlags(CPF_ReferenceParm));
-				const EEdGraphPinDirection Direction = bIsFunctionInput ? EGPD_Input : EGPD_Output;
-				
-				// yay, found it! Now create exec pin for each
-				int32 NumExecs = (Enum->NumEnums() - 1);
-				for(int32 ExecIdx=0; ExecIdx<NumExecs; ExecIdx++)
+				if (UByteProperty* ByteProp = FindField<UByteProperty>(Function, EnumParamName))
 				{
-					bool const bShouldBeHidden = Enum->HasMetaData(TEXT("Hidden"), ExecIdx) || Enum->HasMetaData(TEXT("Spacer"), ExecIdx);
-					if (!bShouldBeHidden)
-					{
-						// Can't use Enum->GetNameByIndex here because it doesn't do namespace mangling
-						const FName ExecName = *Enum->GetNameStringByIndex(ExecIdx);
-						CreatePin(Direction, UEdGraphSchema_K2::PC_Exec, ExecName);
-					}
+					Prop = ByteProp;
+					Enum = ByteProp->Enum;
 				}
-				
-				if (bIsFunctionInput)
+				else if (UEnumProperty* EnumProp = FindField<UEnumProperty>(Function, EnumParamName))
 				{
-					// If using ExpandEnumAsExec for input, don't want to add a input exec pin
-					bCreateSingleExecInputPin = false;
+					Prop = EnumProp;
+					Enum = EnumProp->GetEnum();
 				}
-				else
+
+				if (Prop != nullptr && Enum != nullptr)
 				{
-					// If using ExpandEnumAsExec for output, don't want to add a "then" pin
-					bCreateThenPin = false;
+					const bool bIsFunctionInput = !Prop->HasAnyPropertyFlags(CPF_ReturnParm) &&
+						(!Prop->HasAnyPropertyFlags(CPF_OutParm) ||
+						Prop->HasAnyPropertyFlags(CPF_ReferenceParm));
+					const EEdGraphPinDirection Direction = bIsFunctionInput ? EGPD_Input : EGPD_Output;
+
+					if (bIsFunctionInput)
+					{
+						if (PreviousInput)
+						{
+							bHasCompilerMessage = true;
+							ErrorType = EMessageSeverity::Error;
+							ErrorMsg = FString::Printf(TEXT("Parameter '%s' is listed as an ExpandEnumAsExecs input, but %s already was one. Only one is permitted."), *EnumParamName.ToString(), *PreviousInput->GetName());
+							break;
+						}
+						PreviousInput = Prop;
+					}
+
+					// yay, found it! Now create exec pin for each
+					int32 NumExecs = (Enum->NumEnums() - 1);
+					for (int32 ExecIdx = 0; ExecIdx < NumExecs; ExecIdx++)
+					{
+						bool const bShouldBeHidden = Enum->HasMetaData(TEXT("Hidden"), ExecIdx) || Enum->HasMetaData(TEXT("Spacer"), ExecIdx);
+						if (!bShouldBeHidden)
+						{
+							// Can't use Enum->GetNameByIndex here because it doesn't do namespace mangling
+							const FString NameStr = Enum->GetNameStringByIndex(ExecIdx);
+							
+							UEdGraphPin* CreatedPin = nullptr;
+
+							// todo: really only makes sense if there are multiple outputs
+							if (bIsFunctionInput || EnumNames.Num() == 1)
+							{
+								CreatedPin = CreatePin(Direction, UEdGraphSchema_K2::PC_Exec, *NameStr);
+							}
+							else
+							{
+								CreatedPin = CreatePin(Direction, UEdGraphSchema_K2::PC_Exec, *NameStr);
+								CreatedPin->PinFriendlyName = FText::FromString(FString::Printf(TEXT("(%s) %s"), *Prop->GetDisplayNameText().ToString(), *NameStr));
+							}
+
+							ExpandAsEnumPins.Add(CreatedPin);
+
+							if (Enum->HasMetaData(TEXT("Tooltip"), ExecIdx))
+							{
+								FString EnumTooltip = Enum->GetMetaData(TEXT("Tooltip"), ExecIdx);
+
+								if (const UEdGraphSchema_K2* const K2Schema = Cast<const UEdGraphSchema_K2>(GetSchema()))
+								{
+									K2Schema->ConstructBasicPinTooltip(*CreatedPin, FText::FromString(EnumTooltip), CreatedPin->PinToolTip);
+								}
+								else
+								{
+									CreatedPin->PinToolTip = EnumTooltip;
+								}
+							}
+						}
+					}
+
+					if (bIsFunctionInput)
+					{
+						// If using ExpandEnumAsExec for input, don't want to add a input exec pin
+						bCreateSingleExecInputPin = false;
+					}
+					else
+					{
+						// If using ExpandEnumAsExec for output, don't want to add a "then" pin
+						bCreateThenPin = false;
+					}
 				}
 			}
 		}
@@ -814,21 +868,72 @@ void UK2Node_CallFunction::DetermineWantsEnumToExecExpansion(const UFunction* Fu
 
 	if (Function->HasMetaData(FBlueprintMetadata::MD_ExpandEnumAsExecs))
 	{
-		const FString& EnumParamName = Function->GetMetaData(FBlueprintMetadata::MD_ExpandEnumAsExecs);
-		UByteProperty* EnumProp = FindField<UByteProperty>(Function, FName(*EnumParamName));
-		if((EnumProp != NULL && EnumProp->Enum != NULL) || FindField<UEnumProperty>(Function, FName(*EnumParamName)))
+		TArray<FName> EnumNamesToCheck;
+		GetExpandEnumPinNames(Function, EnumNamesToCheck);
+
+		for (int32 i = EnumNamesToCheck.Num() - 1; i >= 0; --i)
 		{
-			bWantsEnumToExecExpansion = true;
+			const FName& EnumParamName = EnumNamesToCheck[i];
+
+			UByteProperty* EnumProp = FindField<UByteProperty>(Function, EnumParamName);
+			if ((EnumProp != NULL && EnumProp->Enum != NULL) || FindField<UEnumProperty>(Function, EnumParamName))
+			{
+				bWantsEnumToExecExpansion = true;
+				EnumNamesToCheck.RemoveAt(i);
+			}
 		}
-		else
+		
+		if (bWantsEnumToExecExpansion && EnumNamesToCheck.Num() > 0 && !bHasCompilerMessage)
+		{
+			bHasCompilerMessage = true;
+			ErrorType = EMessageSeverity::Warning;
+
+			if (EnumNamesToCheck.Num() == 1)
+			{
+				ErrorMsg = FText::Format(LOCTEXT("EnumToExecExpansionFailedFmt", "Unable to find enum parameter with name '{0}' to expand for @@"), FText::FromName(EnumNamesToCheck[0])).ToString();
+			}
+			else
+			{
+				FString ParamNames;
+
+				for (const FName& Name : EnumNamesToCheck)
+				{
+					if (!ParamNames.IsEmpty())
+					{
+						ParamNames += TEXT(", ");
+					}
+
+					ParamNames += Name.ToString();
+				}
+
+				ErrorMsg = FText::Format(LOCTEXT("EnumToExecExpansionFailedMultipleFmt", "Unable to find enum parameters for names:\n '{{0}}' \nto expand for @@"), FText::FromString(ParamNames)).ToString();
+			}
+		}
+	}
+}
+
+void UK2Node_CallFunction::GetExpandEnumPinNames(const UFunction* Function, TArray<FName>& EnumNamesToCheck)
+{ 
+	EnumNamesToCheck.Reset();
+
+	// todo: use metadatacache if/when that's accepted.
+	const FString& EnumParamString = Function->GetMetaData(FBlueprintMetadata::MD_ExpandEnumAsExecs);
+
+	TArray<FString> RawGroupings;
+	EnumParamString.ParseIntoArray(RawGroupings, TEXT(","), false);
+	for (const FString& RawGroup : RawGroupings)
+	{
+		TArray<FString> IndividualEntries;
+		RawGroup.ParseIntoArray(IndividualEntries, TEXT("|"));
+
+		for (const FString& Entry : IndividualEntries)
 		{
-			if (!bHasCompilerMessage)
+			if (Entry.IsEmpty())
 			{
-				//put in warning state
-				bHasCompilerMessage = true;
-				ErrorType = EMessageSeverity::Warning;
-				ErrorMsg = FText::Format(LOCTEXT("EnumToExecExpansionFailedFmt", "Unable to find enum parameter with name '{0}' to expand for @@"), FText::FromString(EnumParamName)).ToString();
+				continue;
 			}
+
+			EnumNamesToCheck.Add(*Entry);
 		}
 	}
 }
@@ -986,14 +1091,19 @@ bool UK2Node_CallFunction::CreatePinsForFunctionCall(const UFunction* Function)
 		bAllPinsGood = bAllPinsGood && bPinGood;
 	}
 
-	// If we have an 'enum to exec' parameter, set its default value to something valid so we don't get warnings
+	// If we have 'enum to exec' parameters, set their default value to something valid so we don't get warnings
 	if(bWantsEnumToExecExpansion)
 	{
-		const FString& EnumParamName = Function->GetMetaData(FBlueprintMetadata::MD_ExpandEnumAsExecs);
-		UEdGraphPin* EnumParamPin = FindPin(EnumParamName);
-		if (UEnum* PinEnum = (EnumParamPin ? Cast<UEnum>(EnumParamPin->PinType.PinSubCategoryObject.Get()) : NULL))
+		TArray<FName> EnumNamesToCheck;
+		GetExpandEnumPinNames(Function, EnumNamesToCheck);
+
+		for (const FName& Name : EnumNamesToCheck)
 		{
-			EnumParamPin->DefaultValue = PinEnum->GetNameStringByIndex(0);
+			UEdGraphPin* EnumParamPin = FindPin(Name);
+			if (UEnum* PinEnum = (EnumParamPin ? Cast<UEnum>(EnumParamPin->PinType.PinSubCategoryObject.Get()) : NULL))
+			{
+				EnumParamPin->DefaultValue = PinEnum->GetNameStringByIndex(0);
+			}
 		}
 	}
 
@@ -1815,6 +1925,7 @@ void UK2Node_CallFunction::ValidateNodeDuringCompilation(class FCompilerResultsL
 	}
 	else if (Function->HasMetaData(FBlueprintMetadata::MD_ExpandEnumAsExecs) && bWantsEnumToExecExpansion == false)
 	{
+		// will technically not have a properly formatted output for multiple params... but /shrug. 
 		const FString& EnumParamName = Function->GetMetaData(FBlueprintMetadata::MD_ExpandEnumAsExecs);
 		MessageLog.Warning(*FText::Format(LOCTEXT("EnumToExecExpansionFailedFmt", "Unable to find enum parameter with name '{0}' to expand for @@"), FText::FromString(EnumParamName)).ToString(), this);
 	}
@@ -2004,103 +2115,163 @@ void UK2Node_CallFunction::ExpandNode(class FKismetCompilerContext& CompilerCont
 	{
 		if(Function)
 		{
-			// Get the metadata that identifies which param is the enum, and try and find it
-			const FString& EnumParamName = Function->GetMetaData(FBlueprintMetadata::MD_ExpandEnumAsExecs);
+			TArray<FName> EnumNamesToCheck;
+			GetExpandEnumPinNames(Function, EnumNamesToCheck);
 
-			UEnum* Enum = nullptr;
+			bool bAlreadyHandleInput = false;
 
-			if (UByteProperty* ByteProp = FindField<UByteProperty>(Function, FName(*EnumParamName)))
-			{
-				Enum = ByteProp->Enum;
-			}
-			else if (UEnumProperty* EnumProp = FindField<UEnumProperty>(Function, FName(*EnumParamName)))
-			{
-				Enum = EnumProp->GetEnum();
-			}
+			UEdGraphPin* OutMainExecutePin = nullptr;
+			UK2Node_ExecutionSequence* SpawnedSequenceNode = nullptr;
+			int32 OutSequenceIndex = 0;
 
-			UEdGraphPin* EnumParamPin = FindPinChecked(EnumParamName);
-			if(Enum != nullptr)
+			for (const FName& EnumParamName : EnumNamesToCheck)
 			{
-				// Expanded as input execs pins
-				if (EnumParamPin->Direction == EGPD_Input)
+				UEnum* Enum = nullptr;
+
+				if (UByteProperty* ByteProp = FindField<UByteProperty>(Function, EnumParamName))
 				{
-					// Create normal exec input
-					UEdGraphPin* ExecutePin = CreatePin(EGPD_Input, UEdGraphSchema_K2::PC_Exec, UEdGraphSchema_K2::PN_Execute);
-
-					// Create temp enum variable
-					UK2Node_TemporaryVariable* TempEnumVarNode = CompilerContext.SpawnIntermediateNode<UK2Node_TemporaryVariable>(this, SourceGraph);
-					TempEnumVarNode->VariableType.PinCategory = UEdGraphSchema_K2::PC_Byte;
-					TempEnumVarNode->VariableType.PinSubCategoryObject = Enum;
-					TempEnumVarNode->AllocateDefaultPins();
-					// Get the output pin
-					UEdGraphPin* TempEnumVarOutput = TempEnumVarNode->GetVariablePin();
-
-					// Connect temp enum variable to (hidden) enum pin
-					Schema->TryCreateConnection(TempEnumVarOutput, EnumParamPin);
-
-					// Now we want to iterate over other exec inputs...
-					for(int32 PinIdx=Pins.Num()-1; PinIdx>=0; PinIdx--)
+					Enum = ByteProp->Enum;
+				}
+				else if (UEnumProperty* EnumProp = FindField<UEnumProperty>(Function, EnumParamName))
+				{
+					Enum = EnumProp->GetEnum();
+				}
+
+				UEdGraphPin* EnumParamPin = FindPin(EnumParamName);
+				if (Enum && EnumParamPin)
+				{
+					// Expanded as input execs pins
+					if (EnumParamPin->Direction == EGPD_Input)
 					{
-						UEdGraphPin* Pin = Pins[PinIdx];
-						if( Pin != NULL && 
-							Pin != ExecutePin &&
-							Pin->Direction == EGPD_Input && 
-							Pin->PinType.PinCategory == UEdGraphSchema_K2::PC_Exec )
+						if (bAlreadyHandleInput)
 						{
-							// Create node to set the temp enum var
-							UK2Node_AssignmentStatement* AssignNode = CompilerContext.SpawnIntermediateNode<UK2Node_AssignmentStatement>(this, SourceGraph);
-							AssignNode->AllocateDefaultPins();
+							CompilerContext.MessageLog.Error(TEXT("@@ Already provided an input enum parameter for ExpandEnumAsExecs. Only one is permitted."), this);
+							return;
+						}
 
-							// Move connections from fake 'enum exec' pint to this assignment node
-							CompilerContext.MovePinLinksToIntermediate(*Pin, *AssignNode->GetExecPin());
+						bAlreadyHandleInput = true;
 
-							// Connect this to out temp enum var
-							Schema->TryCreateConnection(AssignNode->GetVariablePin(), TempEnumVarOutput);
+						// Create normal exec input
+						UEdGraphPin* ExecutePin = CreatePin(EGPD_Input, UEdGraphSchema_K2::PC_Exec, UEdGraphSchema_K2::PN_Execute);
 
-							// Connect exec output to 'real' exec pin
-							Schema->TryCreateConnection(AssignNode->GetThenPin(), ExecutePin);
+						// Create temp enum variable
+						UK2Node_TemporaryVariable* TempEnumVarNode = CompilerContext.SpawnIntermediateNode<UK2Node_TemporaryVariable>(this, SourceGraph);
+						TempEnumVarNode->VariableType.PinCategory = UEdGraphSchema_K2::PC_Byte;
+						TempEnumVarNode->VariableType.PinSubCategoryObject = Enum;
+						TempEnumVarNode->AllocateDefaultPins();
+						// Get the output pin
+						UEdGraphPin* TempEnumVarOutput = TempEnumVarNode->GetVariablePin();
 
-							// set the literal enum value to set to
-							AssignNode->GetValuePin()->DefaultValue = Pin->PinName.ToString();
+						// Connect temp enum variable to (hidden) enum pin
+						Schema->TryCreateConnection(TempEnumVarOutput, EnumParamPin);
 
-							// Finally remove this 'cosmetic' exec pin
-							Pins[PinIdx]->MarkPendingKill();
-							Pins.RemoveAt(PinIdx);
+						// Now we want to iterate over other exec inputs...
+						for (int32 PinIdx = Pins.Num() - 1; PinIdx >= 0; PinIdx--)
+						{
+							UEdGraphPin* Pin = Pins[PinIdx];
+							if (Pin != NULL &&
+								Pin != ExecutePin &&
+								Pin->Direction == EGPD_Input &&
+								Pin->PinType.PinCategory == UEdGraphSchema_K2::PC_Exec)
+							{
+								// Create node to set the temp enum var
+								UK2Node_AssignmentStatement* AssignNode = CompilerContext.SpawnIntermediateNode<UK2Node_AssignmentStatement>(this, SourceGraph);
+								AssignNode->AllocateDefaultPins();
+
+								// Move connections from fake 'enum exec' pint to this assignment node
+								CompilerContext.MovePinLinksToIntermediate(*Pin, *AssignNode->GetExecPin());
+
+								// Connect this to out temp enum var
+								Schema->TryCreateConnection(AssignNode->GetVariablePin(), TempEnumVarOutput);
+
+								// Connect exec output to 'real' exec pin
+								Schema->TryCreateConnection(AssignNode->GetThenPin(), ExecutePin);
+
+								// set the literal enum value to set to
+								AssignNode->GetValuePin()->DefaultValue = Pin->PinName.ToString();
+
+								// Finally remove this 'cosmetic' exec pin
+								Pins[PinIdx]->MarkPendingKill();
+								Pins.RemoveAt(PinIdx);
+							}
 						}
 					}
-				}
-				// Expanded as output execs pins
-				else if (EnumParamPin->Direction == EGPD_Output)
-				{
-					// Create normal exec output
-					UEdGraphPin* ExecutePin = CreatePin(EGPD_Output, UEdGraphSchema_K2::PC_Exec, UEdGraphSchema_K2::PN_Execute);
-						
-					// Create a SwitchEnum node to switch on the output enum
-					UK2Node_SwitchEnum* SwitchEnumNode = CompilerContext.SpawnIntermediateNode<UK2Node_SwitchEnum>(this, SourceGraph);
-					UEnum* EnumObject = Cast<UEnum>(EnumParamPin->PinType.PinSubCategoryObject.Get());
-					SwitchEnumNode->SetEnum(EnumObject);
-					SwitchEnumNode->AllocateDefaultPins();
-						
-					// Hook up execution to the switch node
-					Schema->TryCreateConnection(ExecutePin, SwitchEnumNode->GetExecPin());
-					// Connect (hidden) enum pin to switch node's selection pin
-					Schema->TryCreateConnection(EnumParamPin, SwitchEnumNode->GetSelectionPin());
-						
-					// Now we want to iterate over other exec outputs
-					for(int32 PinIdx=Pins.Num()-1; PinIdx>=0; PinIdx--)
+					// Expanded as output execs pins
+					else if (EnumParamPin->Direction == EGPD_Output)
 					{
-						UEdGraphPin* Pin = Pins[PinIdx];
-						if( Pin != NULL &&
-							Pin != ExecutePin &&
-							Pin->Direction == EGPD_Output &&
-							Pin->PinType.PinCategory == UEdGraphSchema_K2::PC_Exec )
+						if (!OutMainExecutePin)
+						{
+							// Create normal exec output -- only once though.
+							OutMainExecutePin = CreatePin(EGPD_Output, UEdGraphSchema_K2::PC_Exec, UEdGraphSchema_K2::PN_Then);
+						}
+						else
+						{	
+							// set up a sequence so we can call one after another.
+							if (!SpawnedSequenceNode)
+							{
+								SpawnedSequenceNode = CompilerContext.SpawnIntermediateNode<UK2Node_ExecutionSequence>(this, SourceGraph);
+								SpawnedSequenceNode->AllocateDefaultPins();
+								CompilerContext.MovePinLinksToIntermediate(*OutMainExecutePin, *SpawnedSequenceNode->GetThenPinGivenIndex(OutSequenceIndex++));
+								Schema->TryCreateConnection(OutMainExecutePin, SpawnedSequenceNode->Pins[0]);
+							}
+						}
+
+						// Create a SwitchEnum node to switch on the output enum
+						UK2Node_SwitchEnum* SwitchEnumNode = CompilerContext.SpawnIntermediateNode<UK2Node_SwitchEnum>(this, SourceGraph);
+						UEnum* EnumObject = Cast<UEnum>(EnumParamPin->PinType.PinSubCategoryObject.Get());
+						SwitchEnumNode->SetEnum(EnumObject);
+						SwitchEnumNode->AllocateDefaultPins();
+
+						// Hook up execution to the switch node
+						if (!SpawnedSequenceNode)
+						{
+							Schema->TryCreateConnection(OutMainExecutePin, SwitchEnumNode->GetExecPin());
+						}
+						else
+						{
+							UEdGraphPin* SequenceOutput = SpawnedSequenceNode->GetThenPinGivenIndex(OutSequenceIndex);
+
+							if (!SequenceOutput)
+							{
+								SpawnedSequenceNode->AddInputPin();
+								SequenceOutput = SpawnedSequenceNode->GetThenPinGivenIndex(OutSequenceIndex);
+							}
+
+							Schema->TryCreateConnection(SequenceOutput, SwitchEnumNode->GetExecPin());
+							OutSequenceIndex++;
+						}
+						// Connect (hidden) enum pin to switch node's selection pin
+						Schema->TryCreateConnection(EnumParamPin, SwitchEnumNode->GetSelectionPin());
+
+						// Now we want to iterate over other exec outputs corresponding to the enum.
+						// the first pins created are the ExpandEnumAsExecs pins, and they're all made at the same time.
+						for (int32 PinIdx = Enum->NumEnums() - 2; PinIdx >= 0; PinIdx--)
 						{
-							// Move connections from fake 'enum exec' pin to this switch node
-							CompilerContext.MovePinLinksToIntermediate(*Pin, *SwitchEnumNode->FindPinChecked(Pin->PinName));
-								
-							// Finally remove this 'cosmetic' exec pin
-							Pins[PinIdx]->MarkPendingKill();
-							Pins.RemoveAt(PinIdx);
+							UEdGraphPin* Pin = Pins[PinIdx];
+
+							if (Pin &&
+								Pin != OutMainExecutePin &&
+								Pin->Direction == EGPD_Output &&
+								Pin->PinType.PinCategory == UEdGraphSchema_K2::PC_Exec)
+							{
+								if (UEdGraphPin* FoundPin = SwitchEnumNode->FindPin(Pin->PinName))
+								{
+									if (!FoundPin->LinkedTo.Contains(Pin))
+									{
+										// Move connections from fake 'enum exec' pin to this switch node
+										CompilerContext.MovePinLinksToIntermediate(*Pin, *FoundPin);
+
+										// Finally remove this 'cosmetic' exec pin
+										Pins[PinIdx]->MarkPendingKill();
+										Pins.RemoveAt(PinIdx);
+									}
+								}
+								// Have passed the relevant entries... no more work to do here.
+								else
+								{
+									break;
+								}
+							}
 						}
 					}
 				}
diff --git a/Engine/Source/Editor/UnrealEd/Private/Kismet2/BlueprintEditorUtils.cpp b/Engine/Source/Editor/UnrealEd/Private/Kismet2/BlueprintEditorUtils.cpp
--- a/Engine/Source/Editor/UnrealEd/Private/Kismet2/BlueprintEditorUtils.cpp
+++ b/Engine/Source/Editor/UnrealEd/Private/Kismet2/BlueprintEditorUtils.cpp
@@ -4137,10 +4137,20 @@ void FBlueprintEditorUtils::GetHiddenPinsForFunction(UEdGraph const* Graph, UFun
 
 			const FName& Key = It.Key();
 
-			if (Key == NAME_LatentInfo || Key == NAME_HidePin || Key == FBlueprintMetadata::MD_ExpandEnumAsExecs)
+			if (Key == NAME_LatentInfo || Key == NAME_HidePin)
 			{
 				HiddenPins.Add(*It.Value());
 			}
+			else if (Key == FBlueprintMetadata::MD_ExpandEnumAsExecs)
+			{
+				TArray<FName> EnumPinNames;
+				UK2Node_CallFunction::GetExpandEnumPinNames(Function, EnumPinNames);
+				
+				for (const FName& EnumName : EnumPinNames)
+				{
+					HiddenPins.Add(EnumName);
+				}
+			}
 			else if (Key == FBlueprintMetadata::MD_InternalUseParam)
 			{
 				const FName HiddenPinName = *It.Value();
diff --git a/Engine/Source/Programs/UnrealHeaderTool/Private/HeaderParser.cpp b/Engine/Source/Programs/UnrealHeaderTool/Private/HeaderParser.cpp
--- a/Engine/Source/Programs/UnrealHeaderTool/Private/HeaderParser.cpp
+++ b/Engine/Source/Programs/UnrealHeaderTool/Private/HeaderParser.cpp
@@ -1010,9 +1010,46 @@ namespace
 			{
 				if (UFunction* Function = Cast<UFunction>(Field))
 				{
-					if (FHeaderParser::FindField(Function, *InValue, false) == nullptr)
+					// multiple entry parsing in the same format as eg SetParam.
+					TArray<FString> RawGroupings;
+					InValue.ParseIntoArray(RawGroupings, TEXT(","), false);
+
+					UProperty* FirstInput = nullptr;
+					for (const FString& RawGroup : RawGroupings)
 					{
-						UE_LOG_ERROR_UHT(TEXT("Function does not have a parameter named '%s'"), *InValue);
+						TArray<FString> IndividualEntries;
+						RawGroup.ParseIntoArray(IndividualEntries, TEXT("|"));
+
+						for (const FString& Entry : IndividualEntries)
+						{
+							if (Entry.IsEmpty())
+							{
+								continue;
+							}
+							
+							UField* FoundField = FHeaderParser::FindField(Function, *Entry, false);
+							if (!FoundField)
+							{
+								UE_LOG_ERROR_UHT(TEXT("Function does not have a parameter named '%s'"), *Entry);
+							}
+							else if (UProperty* Prop = Cast<UProperty>(FoundField))
+							{
+								if (!Prop->HasAnyPropertyFlags(CPF_ReturnParm) &&
+
+								    (!Prop->HasAnyPropertyFlags(CPF_OutParm) ||
+									Prop->HasAnyPropertyFlags(CPF_ReferenceParm)))
+								{
+									if (!FirstInput)
+									{
+										FirstInput = Prop;
+									}
+									else
+									{
+										UE_LOG_ERROR_UHT(TEXT("Function already specified an ExpandEnumAsExec input (%s), but '%s' is also an input parameter. Only one is permitted."), *FirstInput->GetName(), *Entry);
+									}
+								}
+							}
+						}
 					}
 				}
 			}
