Editor only functions can't be added to runtime blueprints, only blutilities and editor widgets

diff --git a/Engine/Source/Editor/BlueprintGraph/Private/BlueprintActionFilter.cpp b/Engine/Source/Editor/BlueprintGraph/Private/BlueprintActionFilter.cpp
--- a/Engine/Source/Editor/BlueprintGraph/Private/BlueprintActionFilter.cpp
+++ b/Engine/Source/Editor/BlueprintGraph/Private/BlueprintActionFilter.cpp
@@ -1718,7 +1718,7 @@ static bool BlueprintActionFilterImpl::IsHiddenInNonEditorBlueprint(FBlueprintAc
 
 	if (Function)
 	{
-		const bool bIsEditorOnlyFunction = IsEditorOnlyObject(Function);
+		const bool bIsEditorOnlyFunction = IsEditorOnlyObject(Function) || Function->HasAnyFunctionFlags(FUNC_EditorOnly);
 		
 		if (bIsEditorOnlyFunction)
 		{
