Moving editor function in runtime blueprints compile error to ValidateFunctionOnCompile, adding additional logic to handle WITH_EDITOR wrapped functions

diff --git a/Engine/Source/Editor/BlueprintGraph/Private/K2Node_CallFunction.cpp b/Engine/Source/Editor/BlueprintGraph/Private/K2Node_CallFunction.cpp
--- a/Engine/Source/Editor/BlueprintGraph/Private/K2Node_CallFunction.cpp
+++ b/Engine/Source/Editor/BlueprintGraph/Private/K2Node_CallFunction.cpp
@@ -327,19 +327,6 @@ void FDynamicOutputHelper::VerifyNode(const UK2Node_CallFunction* FuncNode, FCom
 			}
 		}
 	}
-	
-	// Ensure that editor module BP exposed UFunctions can only be called in blueprints for which the baseclass is also part of an editor module
-	const UClass* FunctionClass = FuncNode->FunctionReference.GetMemberParentClass();
-	const bool bIsEditorOnlyFunction = FunctionClass && IsEditorOnlyObject(FunctionClass);
-
-	const UBlueprint* Blueprint = FuncNode->GetBlueprint();
-	const UClass* BlueprintClass = Blueprint->ParentClass;
-	bool bIsEditorOnlyBlueprintBaseClass = !BlueprintClass || IsEditorOnlyObject(BlueprintClass);
-	if (bIsEditorOnlyFunction && !bIsEditorOnlyBlueprintBaseClass)
-	{
-		FText const ErrorFormat = LOCTEXT("BlueprintEditorOnly", "Function in Editor Only Module '@@' cannot be called within the Non-Editor module blueprint base class '@@'.");
-		MessageLog.Error(*ErrorFormat.ToString(), FuncNode, BlueprintClass);
-	}
 }
 
 UK2Node_CallFunction* FDynamicOutputHelper::GetFunctionNode() const
@@ -1818,6 +1805,20 @@ void UK2Node_CallFunction::ValidateNodeDuringCompilation(class FCompilerResultsL
 		const FString& EnumParamName = Function->GetMetaData(FBlueprintMetadata::MD_ExpandEnumAsExecs);
 		MessageLog.Warning(*FText::Format(LOCTEXT("EnumToExecExpansionFailedFmt", "Unable to find enum parameter with name '{0}' to expand for @@"), FText::FromString(EnumParamName)).ToString(), this);
 	}
+	
+	// Ensure that editor module BP exposed UFunctions can only be called in blueprints for which the base class is also part of an editor module
+	// Also check for functions wrapped in WITH_EDITOR 
+	if ((IsEditorOnlyObject(Function) || Function->HasAnyFunctionFlags(FUNC_EditorOnly)))
+	{	
+		const UClass* BlueprintClass = Blueprint->ParentClass;
+		bool bIsEditorOnlyBlueprintBaseClass = !BlueprintClass || IsEditorOnlyObject(BlueprintClass);
+		if (!bIsEditorOnlyBlueprintBaseClass)
+		{
+			FString const FunctName = Function->GetName();
+			FText const WarningFormat = LOCTEXT("EditorFunctionFmt", "Cannot use the editor function \"{0}\" in this runtime Blueprint. Only for use in Editor Utility Blueprints and Blutilities.");
+			MessageLog.Error(*FText::Format(WarningFormat, FText::FromString(FunctName)).ToString(), this);
+		}
+	}
 
 	if (Function)
 	{
