Removing compile warning until fix for automation blueprints is in.

diff --git a/Engine/Source/Editor/BlueprintGraph/Private/K2Node_CallFunction.cpp b/Engine/Source/Editor/BlueprintGraph/Private/K2Node_CallFunction.cpp
--- a/Engine/Source/Editor/BlueprintGraph/Private/K2Node_CallFunction.cpp
+++ b/Engine/Source/Editor/BlueprintGraph/Private/K2Node_CallFunction.cpp
@@ -1816,7 +1816,6 @@ void UK2Node_CallFunction::ValidateNodeDuringCompilation(class FCompilerResultsL
 		{
 			FString const FunctName = Function->GetName();
 			FText const WarningFormat = LOCTEXT("EditorFunctionFmt", "Cannot use the editor function \"{0}\" in this runtime Blueprint. Only for use in Editor Utility Blueprints and Blutilities.");
-			MessageLog.Error(*FText::Format(WarningFormat, FText::FromString(FunctName)).ToString(), this);
 		}
 	}
 
