Reparenting editor-only automation tests to editor automation classes

diff --git a/Engine/Source/Editor/BlueprintGraph/Private/K2Node_CallFunction.cpp b/Engine/Source/Editor/BlueprintGraph/Private/K2Node_CallFunction.cpp
--- a/Engine/Source/Editor/BlueprintGraph/Private/K2Node_CallFunction.cpp
+++ b/Engine/Source/Editor/BlueprintGraph/Private/K2Node_CallFunction.cpp
@@ -1808,7 +1808,8 @@ void UK2Node_CallFunction::ValidateNodeDuringCompilation(class FCompilerResultsL
 	
 	// Ensure that editor module BP exposed UFunctions can only be called in blueprints for which the base class is also part of an editor module
 	// Also check for functions wrapped in WITH_EDITOR 
-	if ((IsEditorOnlyObject(Function) || Function->HasAnyFunctionFlags(FUNC_EditorOnly)))
+	if (Function && Blueprint &&
+		(IsEditorOnlyObject(Function) || Function->HasAnyFunctionFlags(FUNC_EditorOnly)))
 	{	
 		const UClass* BlueprintClass = Blueprint->ParentClass;
 		bool bIsEditorOnlyBlueprintBaseClass = !BlueprintClass || IsEditorOnlyObject(BlueprintClass);
@@ -1816,6 +1817,7 @@ void UK2Node_CallFunction::ValidateNodeDuringCompilation(class FCompilerResultsL
 		{
 			FString const FunctName = Function->GetName();
 			FText const WarningFormat = LOCTEXT("EditorFunctionFmt", "Cannot use the editor function \"{0}\" in this runtime Blueprint. Only for use in Editor Utility Blueprints and Blutilities.");
+			MessageLog.Error(*FText::Format(WarningFormat, FText::FromString(FunctName)).ToString(), this);
 		}
 	}
 
diff --git a/Engine/Source/Editor/Blutility/Classes/EditorAutomationActor.h b/Engine/Source/Editor/Blutility/Classes/EditorAutomationActor.h
new file mode 100644
--- /dev/null
+++ b/Engine/Source/Editor/Blutility/Classes/EditorAutomationActor.h
@@ -0,0 +1,21 @@
+// Copyright 1998-2018 Epic Games, Inc. All Rights Reserved.
+
+/**
+ * Base class of any editor-only automation test actors
+ */
+
+#pragma once
+
+#include "CoreMinimal.h"
+#include "GameFramework/Actor.h"
+#include "UObject/ScriptMacros.h"
+#include "EditorAutomationActor.generated.h"
+
+
+UCLASS(Abstract, Blueprintable)
+class BLUTILITY_API AEditorAutomationActor : public AActor
+{
+	GENERATED_UCLASS_BODY()
+
+	
+};
diff --git a/Engine/Source/Editor/Blutility/Classes/EditorAutomationObject.h b/Engine/Source/Editor/Blutility/Classes/EditorAutomationObject.h
new file mode 100644
--- /dev/null
+++ b/Engine/Source/Editor/Blutility/Classes/EditorAutomationObject.h
@@ -0,0 +1,22 @@
+// Copyright 1998-2018 Epic Games, Inc. All Rights Reserved.
+
+/**
+ * Base class of any editor-only automation test objects
+ */
+
+#pragma once
+
+#include "CoreMinimal.h"
+#include "UObject/ObjectMacros.h"
+#include "UObject/Object.h"
+#include "UObject/ScriptMacros.h"
+#include "EditorAutomationObject.generated.h"
+
+
+UCLASS(Abstract, Blueprintable)
+class BLUTILITY_API UEditorAutomationObject : public UObject
+{
+	GENERATED_UCLASS_BODY()
+
+	
+};
diff --git a/Engine/Source/Editor/Blutility/Private/EditorAutomationActor.cpp b/Engine/Source/Editor/Blutility/Private/EditorAutomationActor.cpp
new file mode 100644
--- /dev/null
+++ b/Engine/Source/Editor/Blutility/Private/EditorAutomationActor.cpp
@@ -0,0 +1,11 @@
+// Copyright 1998-2018 Epic Games, Inc. All Rights Reserved.
+
+#include "EditorAutomationActor.h"
+
+
+/////////////////////////////////////////////////////
+
+AEditorAutomationActor::AEditorAutomationActor(const FObjectInitializer& ObjectInitializer)
+	: Super(ObjectInitializer)
+{
+}
diff --git a/Engine/Source/Editor/Blutility/Private/EditorAutomationObject.cpp b/Engine/Source/Editor/Blutility/Private/EditorAutomationObject.cpp
new file mode 100644
--- /dev/null
+++ b/Engine/Source/Editor/Blutility/Private/EditorAutomationObject.cpp
@@ -0,0 +1,10 @@
+// Copyright 1998-2018 Epic Games, Inc. All Rights Reserved.
+
+#include "EditorAutomationObject.h"
+
+/////////////////////////////////////////////////////
+
+UEditorAutomationObject::UEditorAutomationObject(const FObjectInitializer& ObjectInitializer)
+	: Super(ObjectInitializer)
+{
+}
