PR #5386: UE-67872: Use separate ThenIndex for pin naming (Contributed by [REDACTED_PERSON])

diff --git a/Engine/Source/Editor/BlueprintGraph/Private/K2Node_ExecutionSequence.cpp b/Engine/Source/Editor/BlueprintGraph/Private/K2Node_ExecutionSequence.cpp
--- a/Engine/Source/Editor/BlueprintGraph/Private/K2Node_ExecutionSequence.cpp
+++ b/Engine/Source/Editor/BlueprintGraph/Private/K2Node_ExecutionSequence.cpp
@@ -207,12 +207,14 @@ void UK2Node_ExecutionSequence::InsertPinIntoExecutionNode(UEdGraphPin* PinToIns
 		CreatePin(EGPD_Output, UEdGraphSchema_K2::PC_Exec, GetUniquePinName(), Params);
 
 		// refresh names on the pin list:
+		int32 ThenIndex = 0;
 		for (int32 Idx = 0; Idx < Pins.Num(); ++Idx)
 		{
 			UEdGraphPin* PotentialPin = Pins[Idx];
 			if (UEdGraphSchema_K2::IsExecPin(*PotentialPin) && (PotentialPin->Direction == EGPD_Output))
 			{
-				PotentialPin->PinName = GetPinNameGivenIndex(Idx);
+				PotentialPin->PinName = GetPinNameGivenIndex(ThenIndex);
+				++ThenIndex;
 			}
 		}
 	}
