PR #5418: Make macro call titles friendly (Contributed by [REDACTED_PERSON])

diff --git a/Engine/Source/Editor/BlueprintGraph/Private/K2Node_MacroInstance.cpp b/Engine/Source/Editor/BlueprintGraph/Private/K2Node_MacroInstance.cpp
--- a/Engine/Source/Editor/BlueprintGraph/Private/K2Node_MacroInstance.cpp
+++ b/Engine/Source/Editor/BlueprintGraph/Private/K2Node_MacroInstance.cpp
@@ -10,6 +10,7 @@
 #include "Editor.h"
 #include "EditorCategoryUtils.h"
 #include "BlueprintActionFilter.h"
+#include "Classes/EditorStyleSettings.h"
 
 #define LOCTEXT_NAMESPACE "K2Node_MacroInstance"
 
@@ -177,6 +178,10 @@ FText UK2Node_MacroInstance::GetNodeTitle(ENodeTitleType::Type TitleType) const
 	if (MacroGraph)
 	{
 		Result = FText::FromString(MacroGraph->GetName());
+		if ((GEditor != NULL) && (GetDefault<UEditorStyleSettings>()->bShowFriendlyNames))
+		{
+			Result = FText::FromString(FName::NameToDisplayString(Result.ToString(), false));
+		}
 	}
 
 	return Result;
