PR #5435: Add `Create` and `New` keywords to `ConstructObject` (Contributed by [REDACTED_PERSON])

diff --git a/Engine/Source/Editor/BlueprintGraph/Classes/K2Node_ConstructObjectFromClass.h b/Engine/Source/Editor/BlueprintGraph/Classes/K2Node_ConstructObjectFromClass.h
--- a/Engine/Source/Editor/BlueprintGraph/Classes/K2Node_ConstructObjectFromClass.h
+++ b/Engine/Source/Editor/BlueprintGraph/Classes/K2Node_ConstructObjectFromClass.h
@@ -21,6 +21,7 @@ class BLUEPRINTGRAPH_API UK2Node_ConstructObjectFromClass : public UK2Node
 	virtual FText GetNodeTitle(ENodeTitleType::Type TitleType) const override;
 	virtual void PinDefaultValueChanged(UEdGraphPin* Pin) override;
 	virtual FText GetTooltipText() const override;
+	virtual FText GetKeywords() const override;
 	virtual bool HasExternalDependencies(TArray<class UStruct*>* OptionalOutput) const override;
 	virtual bool IsCompatibleWithGraph(const UEdGraph* TargetGraph) const override;
 	virtual void PinConnectionListChanged(UEdGraphPin* Pin);
diff --git a/Engine/Source/Editor/BlueprintGraph/Private/K2Node_ConstructObjectFromClass.cpp b/Engine/Source/Editor/BlueprintGraph/Private/K2Node_ConstructObjectFromClass.cpp
--- a/Engine/Source/Editor/BlueprintGraph/Private/K2Node_ConstructObjectFromClass.cpp
+++ b/Engine/Source/Editor/BlueprintGraph/Private/K2Node_ConstructObjectFromClass.cpp
@@ -398,4 +398,9 @@ bool UK2Node_ConstructObjectFromClass::HasExternalDependencies(TArray<class UStr
 	return bSuperResult || bResult;
 }
 
+FText UK2Node_ConstructObjectFromClass::GetKeywords() const
+{
+	return LOCTEXT("ConstructObjectKeywords", "Create New");
+}
+
 #undef LOCTEXT_NAMESPACE
