PR #5388: Remove empty () from pin context menus (Contributed by [REDACTED_PERSON])
Clean up Straighten connection messages to be consistent with Jump and Break

diff --git a/Engine/Source/Editor/BlueprintGraph/Private/EdGraphSchema_K2.cpp b/Engine/Source/Editor/BlueprintGraph/Private/EdGraphSchema_K2.cpp
--- a/Engine/Source/Editor/BlueprintGraph/Private/EdGraphSchema_K2.cpp
+++ b/Engine/Source/Editor/BlueprintGraph/Private/EdGraphSchema_K2.cpp
@@ -1450,11 +1450,24 @@ void UEdGraphSchema_K2::GetContextMenuActions(const UEdGraph* CurrentGraph, cons
 						FText PinName = Pin->GetDisplayName();
 						FText NodeName = Pin->GetOwningNode()->GetNodeTitle(ENodeTitleType::ListView);
 
+						FText EntryTitle;
+						FText EntryTooltip;
+						if (PinName.IsEmpty())
+						{
+							EntryTitle = FText::Format(LOCTEXT("StraightenDescription_SinglePin_NoDisplayName", "Straighten Connection to {0}"), NodeName);
+							EntryTooltip = FText::Format(LOCTEXT("StraightenDescription_SinglePin_NoDisplayName_Tip", "Straighten the connection between this pin, and {0}"), NodeName);
+						}
+						else
+						{
+							EntryTitle = FText::Format(LOCTEXT("StraightenDescription_SinglePin", "Straighten Connection to {0} ({1})"), NodeName, PinName);
+							EntryTooltip = FText::Format(LOCTEXT("StraightenDescription_SinglePin_Node_Tip", "Straighten the connection between this pin, and {0} ({1})"), NodeName, PinName);
+						}
+
 						MenuBuilder->AddMenuEntry(
 							FGraphEditorCommands::Get().StraightenConnections,
 							NAME_None,
-							FText::Format(LOCTEXT("StraightenDescription_SinglePin", "Straighten Connection to {0} ({1})"), NodeName, PinName),
-							FText::Format(LOCTEXT("StraightenDescription_SinglePin_Node_Tip", "Straighten the connection between this pin, and {0} ({1})"), NodeName, PinName),
+							EntryTitle,
+							EntryTooltip,
 							FSlateIcon(NAME_None, NAME_None, NAME_None)
 						);
 					}
@@ -1963,14 +1976,15 @@ void UEdGraphSchema_K2::GetBreakLinkToSubMenuActions( class FMenuBuilder& MenuBu
 		UEdGraphPin* Pin = *Links;
 		FText Title = Pin->GetOwningNode()->GetNodeTitle(ENodeTitleType::ListView);
 		FString TitleString = Title.ToString();
-		if ( Pin->PinName != TEXT("") )
+		const FText PinDisplayName = Pin->GetDisplayName();
+		if (!PinDisplayName.IsEmpty())
 		{
-			TitleString = FString::Printf(TEXT("%s (%s)"), *TitleString, *Pin->GetDisplayName().ToString());
+			TitleString = FString::Printf(TEXT("%s (%s)"), *TitleString, *PinDisplayName.ToString());
 
 			// Add name of connection if possible
 			FFormatNamedArguments Args;
 			Args.Add( TEXT("NodeTitle"), Title );
-			Args.Add( TEXT("PinName"), Pin->GetDisplayName() );
+			Args.Add( TEXT("PinName"), PinDisplayName );
 			Title = FText::Format( LOCTEXT("BreakDescPin", "{NodeTitle} ({PinName})"), Args );
 		}
 
@@ -1983,11 +1997,11 @@ void UEdGraphSchema_K2::GetBreakLinkToSubMenuActions( class FMenuBuilder& MenuBu
 
 		if ( Count == 0 )
 		{
-			Description = FText::Format( LOCTEXT("BreakDesc", "Break link to {NodeTitle}"), Args );
+			Description = FText::Format( LOCTEXT("BreakDesc", "Break Link to {NodeTitle}"), Args );
 		}
 		else
 		{
-			Description = FText::Format( LOCTEXT("BreakDescMulti", "Break link to {NodeTitle} ({NumberOfNodes})"), Args );
+			Description = FText::Format( LOCTEXT("BreakDescMulti", "Break Link to {NodeTitle} ({NumberOfNodes})"), Args );
 		}
 		++Count;
 
@@ -2006,14 +2020,15 @@ void UEdGraphSchema_K2::GetJumpToConnectionSubMenuActions( class FMenuBuilder& M
 	{
 		FText Title = PinLink->GetOwningNode()->GetNodeTitle(ENodeTitleType::ListView);
 		FString TitleString = Title.ToString();
-		if ( PinLink->PinName != TEXT("") )
+		const FText PinDisplayName = PinLink->GetDisplayName();
+		if (!PinDisplayName.IsEmpty())
 		{
-			TitleString = FString::Printf(TEXT("%s (%s)"), *TitleString, *PinLink->GetDisplayName().ToString());
+			TitleString = FString::Printf(TEXT("%s (%s)"), *TitleString, *PinDisplayName.ToString());
 
 			// Add name of connection if possible
 			FFormatNamedArguments Args;
 			Args.Add( TEXT("NodeTitle"), Title );
-			Args.Add( TEXT("PinName"), PinLink->GetDisplayName() );
+			Args.Add( TEXT("PinName"), PinDisplayName );
 			Title = FText::Format( LOCTEXT("JumpToDescPin", "{NodeTitle} ({PinName})"), Args );
 		}
 
@@ -2075,15 +2090,43 @@ void UEdGraphSchema_K2::GetStraightenConnectionToSubMenuActions( class FMenuBuil
 		NAME_None, LOCTEXT("StraightenAllConnections", "All Connected Pins"),
 		TAttribute<FText>(), FSlateIcon(NAME_None, NAME_None, NAME_None) );
 
-	for (auto& Pair : NodeToPins)
+	for (const TPair<UEdGraphNode*, TArray<UEdGraphPin*>>& Pair : NodeToPins)
 	{
-		FText NodeName = Pair.Key->GetNodeTitle(ENodeTitleType::ListView);
 		for (UEdGraphPin* Pin : Pair.Value)
 		{
-			FText PinName = Pin->GetDisplayName();
+			FText Title = Pair.Key->GetNodeTitle(ENodeTitleType::ListView);
+			FString TitleString = Title.ToString();
+			const FText PinDisplayName = Pin->GetDisplayName();
+			if (!PinDisplayName.IsEmpty())
+			{
+				TitleString = FString::Printf(TEXT("%s (%s)"), *TitleString, *PinDisplayName.ToString());
+
+				// Add name of connection if possible
+				FFormatNamedArguments Args;
+				Args.Add(TEXT("NodeTitle"), Title);
+				Args.Add(TEXT("PinName"), PinDisplayName);
+				Title = FText::Format(LOCTEXT("StraightenToDescPin", "{NodeTitle} ({PinName})"), Args);
+			}
+			uint32 &Count = LinkTitleCount.FindOrAdd(TitleString);
+
+			FText Description;
+			FFormatNamedArguments Args;
+			Args.Add(TEXT("NodeTitle"), Title);
+			Args.Add(TEXT("NumberOfNodes"), Count);
+
+			if (Count == 0)
+			{
+				Description = FText::Format(LOCTEXT("StraightenDesc", "Straighten connection to {NodeTitle}"), Args);
+			}
+			else
+			{
+				Description = FText::Format(LOCTEXT("StraightendDescMulti", "Straighten connection to {NodeTitle} ({NumberOfNodes})"), Args);
+			}
+			++Count;
+
 			MenuBuilder.AddMenuEntry(
-				FText::Format(LOCTEXT("StraightenDescription_Node", "{0} ({1})"), NodeName, Pin->GetDisplayName()),
-				FText::Format(LOCTEXT("StraightenDescription_Node_Tip", "Straighten the connection between this pin, and {0} ({1})"), NodeName, Pin->GetDisplayName()),
+				Description,
+				Description,
 				FSlateIcon(),
 				FExecuteAction::CreateLambda([=]{
 				if (const FUIAction* UIAction = MenuCommandList->GetActionForCommand(FGraphEditorCommands::Get().StraightenConnections))
