Added Blueprint Nodes for getting Engine and Editor Subsystems.

diff --git a/Engine/Source/Editor/BlueprintGraph/Classes/K2Node_GetSubsystem.h b/Engine/Source/Editor/BlueprintGraph/Classes/K2Node_GetSubsystem.h
--- a/Engine/Source/Editor/BlueprintGraph/Classes/K2Node_GetSubsystem.h
+++ b/Engine/Source/Editor/BlueprintGraph/Classes/K2Node_GetSubsystem.h
@@ -57,15 +57,52 @@ class UK2Node_GetSubsystemFromPC : public UK2Node_GetSubsystem
 public:
 
 	//~Begin UEdGraphNode interface.
-	void AllocateDefaultPins();
-	void ExpandNode(class FKismetCompilerContext& CompilerContext, UEdGraph* SourceGraph);
-	FText GetTooltipText() const;
+	virtual void AllocateDefaultPins() override;
+	virtual void ExpandNode(class FKismetCompilerContext& CompilerContext, UEdGraph* SourceGraph) override;
+	virtual FText GetTooltipText() const override;
 	//~End UEdGraphNode interface.
 
-	void GetMenuActions(FBlueprintActionDatabaseRegistrar& ActionRegistrar) const;
-	FText GetMenuCategory() const;
+	virtual void GetMenuActions(FBlueprintActionDatabaseRegistrar& ActionRegistrar) const override;
+	virtual FText GetMenuCategory() const override;
 
 	/** Get the world context input pin, can return NULL */
 	UEdGraphPin* GetPlayerControllerPin() const;
+
+};
+
+UCLASS()
+class UK2Node_GetEngineSubsystem : public UK2Node_GetSubsystem
+{
+	GENERATED_BODY()
+public:
+
+	//~Begin UEdGraphNode interface.
+	virtual void AllocateDefaultPins() override;
+	virtual void ExpandNode(class FKismetCompilerContext& CompilerContext, UEdGraph* SourceGraph) override;
+	virtual FText GetTooltipText() const override;
+	//~End UEdGraphNode interface.
+
+	virtual void GetMenuActions(FBlueprintActionDatabaseRegistrar& ActionRegistrar) const override;
+	virtual FText GetMenuCategory() const override;
 	
+};
+
+UCLASS()
+class UK2Node_GetEditorSubsystem : public UK2Node_GetSubsystem
+{
+	GENERATED_BODY()
+public:
+
+	//~Begin UEdGraphNode interface.
+	virtual void AllocateDefaultPins() override;
+	virtual void ExpandNode(class FKismetCompilerContext& CompilerContext, UEdGraph* SourceGraph) override;
+	virtual FText GetTooltipText() const override;
+	virtual void ValidateNodeDuringCompilation(class FCompilerResultsLog& MessageLog) const override;
+	//~End UEdGraphNode interface.
+
+	virtual bool IsActionFilteredOut(class FBlueprintActionFilter const& Filter) override;
+
+	virtual void GetMenuActions(FBlueprintActionDatabaseRegistrar& ActionRegistrar) const override;
+	virtual FText GetMenuCategory() const override;
+
 };
\ No newline at end of file
diff --git a/Engine/Source/Editor/BlueprintGraph/Private/K2Node_GetSubsystem.cpp b/Engine/Source/Editor/BlueprintGraph/Private/K2Node_GetSubsystem.cpp
--- a/Engine/Source/Editor/BlueprintGraph/Private/K2Node_GetSubsystem.cpp
+++ b/Engine/Source/Editor/BlueprintGraph/Private/K2Node_GetSubsystem.cpp
@@ -9,10 +9,15 @@
 #include "Kismet2/BlueprintEditorUtils.h"
 #include "BlueprintActionDatabaseRegistrar.h"
 
+#include "Subsystems/EngineSubsystem.h"
 #include "Subsystems/GameInstanceSubsystem.h"
 #include "Subsystems/LocalPlayerSubsystem.h"
+#include "EditorSubsystem.h"
 #include "Subsystems/SubsystemBlueprintLibrary.h"
+#include "Subsystems/EditorSubsystemBlueprintLibrary.h"
 #include "GameFramework/PlayerController.h"
+#include "Kismet2/BlueprintEditorUtils.h"
+
 
 // ************************************************************************************
 //    UK2Node_GetSubsystem
@@ -430,4 +435,274 @@ UEdGraphPin* UK2Node_GetSubsystemFromPC::GetPlayerControllerPin() const
 	UEdGraphPin* Pin = FindPin(TEXT("PlayerController"));
 	check(Pin == NULL || Pin->Direction == EGPD_Input);
 	return Pin;
+}
+
+
+// ************************************************************************************
+//    UK2Node_GetEngineSubsystem
+// ************************************************************************************
+
+void UK2Node_GetEngineSubsystem::AllocateDefaultPins()
+{
+	// Add blueprint pin
+	if (!CustomClass)
+	{
+		CreatePin(EGPD_Input, UEdGraphSchema_K2::PC_Class, USubsystem::StaticClass(), TEXT("Class"));
+	}
+
+	// Result pin
+	CreatePin(EGPD_Output, UEdGraphSchema_K2::PC_Object, (CustomClass ? (UClass*)CustomClass : UEngineSubsystem::StaticClass()), UEdGraphSchema_K2::PN_ReturnValue);
+
+	// Skip the UK2Node_GetSubsystem implementation
+	UK2Node::AllocateDefaultPins();
+}
+
+void UK2Node_GetEngineSubsystem::ExpandNode(class FKismetCompilerContext& CompilerContext, UEdGraph* SourceGraph)
+{
+	// Skip the UK2Node_GetSubsystem implementation
+	UK2Node::ExpandNode(CompilerContext, SourceGraph);
+
+	static const FName Class_ParamName(TEXT("Class"));
+
+	UK2Node_GetEngineSubsystem* GetEngineSubsystemNode = this;
+	UEdGraphPin* SpawnClassPin = GetEngineSubsystemNode->GetClassPin();
+	UEdGraphPin* SpawnNodeResult = GetEngineSubsystemNode->GetResultPin();
+
+	UClass* SpawnClass = (SpawnClassPin != nullptr) ? Cast<UClass>(SpawnClassPin->DefaultObject) : nullptr;
+	if (SpawnClassPin && (SpawnClassPin->LinkedTo.Num() == 0) && !SpawnClass)
+	{
+		CompilerContext.MessageLog.Error(*NSLOCTEXT("K2Node", "GetSubsystem_Error", "Node @@ must have a class specified.").ToString(), GetEngineSubsystemNode);
+		GetEngineSubsystemNode->BreakAllNodeLinks();
+		return;
+	}
+
+	// Choose appropriate underlying Getter
+	FName Get_FunctionName;
+	if (CustomClass->IsChildOf<ULocalPlayerSubsystem>())
+	{
+		Get_FunctionName = GET_FUNCTION_NAME_CHECKED(USubsystemBlueprintLibrary, GetEngineSubsystem);
+	}
+	else
+	{
+		CompilerContext.MessageLog.Error(*NSLOCTEXT("K2Node", "GetSubsystem_Error", "Node @@ must have a class specified.").ToString(), GetEngineSubsystemNode);
+		GetEngineSubsystemNode->BreakAllNodeLinks();
+		return;
+	}
+
+	//////////////////////////////////////////////////////////////////////////
+	// create 'USubsystemBlueprintLibrary::GetEngineSubsystem' call node
+	UK2Node_CallFunction* CallGetNode = CompilerContext.SpawnIntermediateNode<UK2Node_CallFunction>(GetEngineSubsystemNode, SourceGraph);
+	CallGetNode->FunctionReference.SetExternalMember(Get_FunctionName, USubsystemBlueprintLibrary::StaticClass());
+	CallGetNode->AllocateDefaultPins();
+
+	UEdGraphPin* CallCreateClassTypePin = CallGetNode->FindPinChecked(Class_ParamName);
+	UEdGraphPin* CallCreateResult = CallGetNode->GetReturnValuePin();
+
+	if (SpawnClassPin && SpawnClassPin->LinkedTo.Num() > 0)
+	{
+		// Copy the 'class' connection from the spawn node to 'USubsystemBlueprintLibrary::GetLocalPlayerSubSystemFromPlayerController'
+		CompilerContext.MovePinLinksToIntermediate(*SpawnClassPin, *CallCreateClassTypePin);
+	}
+	else
+	{
+		// Copy class literal onto 'USubsystemBlueprintLibrary::GetLocalPlayerSubSystemFromPlayerController' call
+		CallCreateClassTypePin->DefaultObject = *CustomClass;
+	}
+
+	// Move result connection from spawn node to 'USubsystemBlueprintLibrary::Get[something]Subsystem'
+	CallCreateResult->PinType = SpawnNodeResult->PinType;
+	CompilerContext.MovePinLinksToIntermediate(*SpawnNodeResult, *CallCreateResult);
+
+	//////////////////////////////////////////////////////////////////////////
+
+	// Break any links to the expanded node
+	GetEngineSubsystemNode->BreakAllNodeLinks();
+}
+
+void UK2Node_GetEngineSubsystem::GetMenuActions(FBlueprintActionDatabaseRegistrar& ActionRegistrar) const
+{
+	static TArray<UClass*> Subclasses;
+	Subclasses.Reset();
+	GetDerivedClasses(UEngineSubsystem::StaticClass(), Subclasses);
+
+	auto CustomizeCallback = [](UEdGraphNode* Node, bool bIsTemplateNode, UClass* Subclass)
+	{
+		auto TypedNode = CastChecked<UK2Node_GetEngineSubsystem>(Node);
+		TypedNode->Initialize(Subclass);
+	};
+
+	UClass* ActionKey = GetClass();
+	if (ActionRegistrar.IsOpenForRegistration(ActionKey))
+	{
+		for (auto& Iter : Subclasses)
+		{
+			UBlueprintNodeSpawner* Spawner = UBlueprintNodeSpawner::Create(ActionKey);
+			check(Spawner);
+
+			Spawner->CustomizeNodeDelegate = UBlueprintNodeSpawner::FCustomizeNodeDelegate::CreateStatic(CustomizeCallback, Iter);
+			ActionRegistrar.AddBlueprintAction(ActionKey, Spawner);
+		}
+	}
+}
+
+FText UK2Node_GetEngineSubsystem::GetMenuCategory() const
+{
+	return NSLOCTEXT("K2Node", "GetEngineSubsystem_MenuCategory", "Engine Subsystems");
+}
+
+FText UK2Node_GetEngineSubsystem::GetTooltipText() const
+{
+	if (CustomClass)
+	{
+		return FText::FormatNamed(NSLOCTEXT("K2Node", "GetEngineSubsystem_TooltipFormat", "Get {ClassName} an Engine Subsystem"), TEXT("ClassName"), CustomClass->GetDisplayNameText());
+	}
+
+	return NSLOCTEXT("K2Node", "GetEngineSubsystem_InvalidSubsystemTypeTooltip", "Invalid Subsystem Type");
+}
+
+// ************************************************************************************
+//    UK2Node_GetEditorSubsystem
+// ************************************************************************************
+
+void UK2Node_GetEditorSubsystem::AllocateDefaultPins()
+{
+	// Add blueprint pin
+	if (!CustomClass)
+	{
+		CreatePin(EGPD_Input, UEdGraphSchema_K2::PC_Class, USubsystem::StaticClass(), TEXT("Class"));
+	}
+
+	// Result pin
+	CreatePin(EGPD_Output, UEdGraphSchema_K2::PC_Object, (CustomClass ? (UClass*)CustomClass : UEditorSubsystem::StaticClass()), UEdGraphSchema_K2::PN_ReturnValue);
+
+	// Skip the UK2Node_GetSubsystem implementation
+	UK2Node::AllocateDefaultPins();
+}
+
+void UK2Node_GetEditorSubsystem::ExpandNode(class FKismetCompilerContext& CompilerContext, UEdGraph* SourceGraph)
+{
+	// Skip the UK2Node_GetSubsystem implementation
+	UK2Node::ExpandNode(CompilerContext, SourceGraph);
+
+	static const FName Class_ParamName(TEXT("Class"));
+
+	UK2Node_GetEditorSubsystem* GetEditorSubsystemNode = this;
+	UEdGraphPin* SpawnClassPin = GetEditorSubsystemNode->GetClassPin();
+	UEdGraphPin* SpawnNodeResult = GetEditorSubsystemNode->GetResultPin();
+
+	UClass* SpawnClass = (SpawnClassPin != nullptr) ? Cast<UClass>(SpawnClassPin->DefaultObject) : nullptr;
+	if (SpawnClassPin && (SpawnClassPin->LinkedTo.Num() == 0) && !SpawnClass)
+	{
+		CompilerContext.MessageLog.Error(*NSLOCTEXT("K2Node", "GetSubsystem_Error", "Node @@ must have a class specified.").ToString(), GetEditorSubsystemNode);
+		GetEditorSubsystemNode->BreakAllNodeLinks();
+		return;
+	}
+
+	// Choose appropriate underlying Getter
+	FName Get_FunctionName;
+	if (CustomClass->IsChildOf<ULocalPlayerSubsystem>())
+	{
+		Get_FunctionName = GET_FUNCTION_NAME_CHECKED(UEditorSubsystemBlueprintLibrary, GetEditorSubsystem);
+	}
+	else
+	{
+		CompilerContext.MessageLog.Error(*NSLOCTEXT("K2Node", "GetSubsystem_Error", "Node @@ must have a class specified.").ToString(), GetEditorSubsystemNode);
+		GetEditorSubsystemNode->BreakAllNodeLinks();
+		return;
+	}
+
+	//////////////////////////////////////////////////////////////////////////
+	// create 'USubsystemBlueprintLibrary::GetEditorSubsystem' call node
+	UK2Node_CallFunction* CallGetNode = CompilerContext.SpawnIntermediateNode<UK2Node_CallFunction>(GetEditorSubsystemNode, SourceGraph);
+	CallGetNode->FunctionReference.SetExternalMember(Get_FunctionName, USubsystemBlueprintLibrary::StaticClass());
+	CallGetNode->AllocateDefaultPins();
+
+	UEdGraphPin* CallCreateClassTypePin = CallGetNode->FindPinChecked(Class_ParamName);
+	UEdGraphPin* CallCreateResult = CallGetNode->GetReturnValuePin();
+
+	if (SpawnClassPin && SpawnClassPin->LinkedTo.Num() > 0)
+	{
+		// Copy the 'class' connection from the spawn node to 'USubsystemBlueprintLibrary::GetLocalPlayerSubSystemFromPlayerController'
+		CompilerContext.MovePinLinksToIntermediate(*SpawnClassPin, *CallCreateClassTypePin);
+	}
+	else
+	{
+		// Copy class literal onto 'USubsystemBlueprintLibrary::GetLocalPlayerSubSystemFromPlayerController' call
+		CallCreateClassTypePin->DefaultObject = *CustomClass;
+	}
+
+	// Move result connection from spawn node to 'USubsystemBlueprintLibrary::Get[something]Subsystem'
+	CallCreateResult->PinType = SpawnNodeResult->PinType;
+	CompilerContext.MovePinLinksToIntermediate(*SpawnNodeResult, *CallCreateResult);
+
+	//////////////////////////////////////////////////////////////////////////
+
+	// Break any links to the expanded node
+	GetEditorSubsystemNode->BreakAllNodeLinks();
+}
+
+void UK2Node_GetEditorSubsystem::GetMenuActions(FBlueprintActionDatabaseRegistrar& ActionRegistrar) const
+{
+	static TArray<UClass*> Subclasses;
+	Subclasses.Reset();
+	GetDerivedClasses(UEditorSubsystem::StaticClass(), Subclasses);
+
+	auto CustomizeCallback = [](UEdGraphNode* Node, bool bIsTemplateNode, UClass* Subclass)
+	{
+		auto TypedNode = CastChecked<UK2Node_GetEditorSubsystem>(Node);
+		TypedNode->Initialize(Subclass);
+	};
+
+	UClass* ActionKey = GetClass();
+	if (ActionRegistrar.IsOpenForRegistration(ActionKey))
+	{
+		for (auto& Iter : Subclasses)
+		{
+			UBlueprintNodeSpawner* Spawner = UBlueprintNodeSpawner::Create(ActionKey);
+			check(Spawner);
+
+			Spawner->CustomizeNodeDelegate = UBlueprintNodeSpawner::FCustomizeNodeDelegate::CreateStatic(CustomizeCallback, Iter);
+			ActionRegistrar.AddBlueprintAction(ActionKey, Spawner);
+		}
+	}
+}
+
+FText UK2Node_GetEditorSubsystem::GetMenuCategory() const
+{
+	return NSLOCTEXT("K2Node", "GetEditorSubsystem_MenuCategory", "Editor Subsystems");
+}
+
+FText UK2Node_GetEditorSubsystem::GetTooltipText() const
+{
+	if (CustomClass)
+	{
+		return FText::FormatNamed(NSLOCTEXT("K2Node", "GetEditorSubsystem_TooltipFormat", "Get {ClassName} an Editor Subsystem"), TEXT("ClassName"), CustomClass->GetDisplayNameText());
+	}
+
+	return NSLOCTEXT("K2Node", "GetEditorSubsystem_InvalidSubsystemTypeTooltip", "Invalid Subsystem Type");
+}
+
+bool UK2Node_GetEditorSubsystem::IsActionFilteredOut(class FBlueprintActionFilter const& Filter)
+{
+	for(UBlueprint* BP : Filter.Context.Blueprints)
+	{
+		if (!FBlueprintEditorUtils::IsBlutility(BP))
+		{
+			return true;
+		}
+	}
+	return false;
+}
+
+void UK2Node_GetEditorSubsystem::ValidateNodeDuringCompilation(class FCompilerResultsLog& MessageLog) const
+{
+	Super::ValidateNodeDuringCompilation(MessageLog);
+
+	const UBlueprint* BP = FBlueprintEditorUtils::FindBlueprintForNodeChecked(this);
+
+	if (!FBlueprintEditorUtils::IsBlutility(BP))
+	{
+		const FText ErrorMsg = NSLOCTEXT("K2Node", "GetSubsystem_Error", "Editor Subsystems can only be used in Editor Utilities / Blutilities");
+		MessageLog.Error(*ErrorMsg.ToString(), this);
+	}
 }
\ No newline at end of file
diff --git a/Engine/Source/Editor/UnrealEd/Private/Subsystems/EditorSubsystemBlueprintLibrary.cpp b/Engine/Source/Editor/UnrealEd/Private/Subsystems/EditorSubsystemBlueprintLibrary.cpp
new file mode 100644
--- /dev/null
+++ b/Engine/Source/Editor/UnrealEd/Private/Subsystems/EditorSubsystemBlueprintLibrary.cpp
@@ -0,0 +1,10 @@
+// Copyright 1998-2019 Epic Games, Inc. All Rights Reserved.
+
+#include "Subsystems/EditorSubsystemBlueprintLibrary.h"
+
+#include "Editor.h"
+
+/*static*/ UEditorSubsystem* UEditorSubsystemBlueprintLibrary::GetEditorSubsystem(TSubclassOf<UEditorSubsystem> Class)
+{
+	return GEditor->GetEditorSubsystemBase(Class);
+}
diff --git a/Engine/Source/Editor/UnrealEd/Public/Subsystems/EditorSubsystemBlueprintLibrary.h b/Engine/Source/Editor/UnrealEd/Public/Subsystems/EditorSubsystemBlueprintLibrary.h
new file mode 100644
--- /dev/null
+++ b/Engine/Source/Editor/UnrealEd/Public/Subsystems/EditorSubsystemBlueprintLibrary.h
@@ -0,0 +1,22 @@
+// Copyright 1998-2019 Epic Games, Inc. All Rights Reserved.
+
+#pragma once
+
+#include "Templates/SubclassOf.h"
+#include "Kismet/BlueprintFunctionLibrary.h"
+#include "EditorSubsystem.h"
+
+#include "EditorSubsystemBlueprintLibrary.generated.h"
+
+UCLASS()
+class UNREALED_API UEditorSubsystemBlueprintLibrary : public UBlueprintFunctionLibrary
+{
+	GENERATED_BODY()
+
+public:
+
+	/** Get a Local Player Subsystem from the Local Player associated with the provided context */
+	UFUNCTION(BlueprintPure, Category = "Editor Subsystems", meta = (BlueprintInternalUseOnly = "true"))
+	static UEditorSubsystem* GetEditorSubsystem(TSubclassOf<UEditorSubsystem> Class);
+
+};
\ No newline at end of file
diff --git a/Engine/Source/Runtime/Engine/Private/Subsystems/SubsystemBlueprintLibrary.cpp b/Engine/Source/Runtime/Engine/Private/Subsystems/SubsystemBlueprintLibrary.cpp
--- a/Engine/Source/Runtime/Engine/Private/Subsystems/SubsystemBlueprintLibrary.cpp
+++ b/Engine/Source/Runtime/Engine/Private/Subsystems/SubsystemBlueprintLibrary.cpp
@@ -7,6 +7,11 @@
 #include "GameFramework/PlayerController.h"
 #include "Blueprint/UserWidget.h"
 
+/*static*/  UEngineSubsystem* USubsystemBlueprintLibrary::GetEngineSubsystem(TSubclassOf<UEngineSubsystem> Class)
+{
+	return GEngine->GetEngineSubsystemBase(Class);
+}
+
 /*static*/ UGameInstanceSubsystem* USubsystemBlueprintLibrary::GetGameInstanceSubsystem(UObject* ContextObject, TSubclassOf<UGameInstanceSubsystem> Class)
 {
 	if (const UWorld* World = ThisClass::GetWorldFrom(ContextObject))
diff --git a/Engine/Source/Runtime/Engine/Public/Subsystems/SubsystemBlueprintLibrary.h b/Engine/Source/Runtime/Engine/Public/Subsystems/SubsystemBlueprintLibrary.h
--- a/Engine/Source/Runtime/Engine/Public/Subsystems/SubsystemBlueprintLibrary.h
+++ b/Engine/Source/Runtime/Engine/Public/Subsystems/SubsystemBlueprintLibrary.h
@@ -4,6 +4,7 @@
 
 #include "Templates/SubclassOf.h"
 #include "Kismet/BlueprintFunctionLibrary.h"
+#include "Subsystems/EngineSubsystem.h"
 #include "Subsystems/GameInstanceSubsystem.h"
 #include "Subsystems/LocalPlayerSubsystem.h"
 
@@ -15,6 +16,11 @@ class ENGINE_API USubsystemBlueprintLibrary : public UBlueprintFunctionLibrary
 	GENERATED_BODY()
 
 public:
+
+	/** Get a Game Instance Subsystem from the Game Instance associated with the provided context */
+	UFUNCTION(BlueprintPure, Category = "Engine Subsystems", meta = (BlueprintInternalUseOnly = "true"))
+	static UEngineSubsystem* GetEngineSubsystem(TSubclassOf<UEngineSubsystem> Class);
+
 	/** Get a Game Instance Subsystem from the Game Instance associated with the provided context */
 	UFUNCTION(BlueprintPure, Category = "GameInstance Subsystems", meta = (WorldContext = "ContextObject", BlueprintInternalUseOnly = "true"))
 	static UGameInstanceSubsystem* GetGameInstanceSubsystem(UObject* ContextObject, TSubclassOf<UGameInstanceSubsystem> Class);
