Missed a build file, in another CL.

diff --git a/Engine/Source/Editor/BlueprintGraph/BlueprintGraph.Build.cs b/Engine/Source/Editor/BlueprintGraph/BlueprintGraph.Build.cs
--- a/Engine/Source/Editor/BlueprintGraph/BlueprintGraph.Build.cs
+++ b/Engine/Source/Editor/BlueprintGraph/BlueprintGraph.Build.cs
@@ -22,6 +22,7 @@ public class BlueprintGraph : ModuleRules
                 "InputCore",
 				"Slate",
                 "EditorStyle",
+				"EditorSubsystem",
 			}
 		);
 
