Warning fix.

diff --git a/Engine/Source/Editor/BlueprintGraph/Private/K2Node_GetSubsystem.cpp b/Engine/Source/Editor/BlueprintGraph/Private/K2Node_GetSubsystem.cpp
--- a/Engine/Source/Editor/BlueprintGraph/Private/K2Node_GetSubsystem.cpp
+++ b/Engine/Source/Editor/BlueprintGraph/Private/K2Node_GetSubsystem.cpp
@@ -702,7 +702,7 @@ void UK2Node_GetEditorSubsystem::ValidateNodeDuringCompilation(class FCompilerRe
 
 	if (!FBlueprintEditorUtils::IsBlutility(BP))
 	{
-		const FText ErrorMsg = NSLOCTEXT("K2Node", "GetSubsystem_Error", "Editor Subsystems can only be used in Editor Utilities / Blutilities");
-		MessageLog.Error(*ErrorMsg.ToString(), this);
+		const FText ErrorText = NSLOCTEXT("K2Node", "GetSubsystem_Error", "Editor Subsystems can only be used in Editor Utilities / Blutilities");
+		MessageLog.Error(*ErrorText.ToString(), this);
 	}
 }
\ No newline at end of file
