Make priming max per frame a cvar

diff --git a/Engine/Source/Editor/BlueprintGraph/Private/BlueprintActionDatabase.cpp b/Engine/Source/Editor/BlueprintGraph/Private/BlueprintActionDatabase.cpp
--- a/Engine/Source/Editor/BlueprintGraph/Private/BlueprintActionDatabase.cpp
+++ b/Engine/Source/Editor/BlueprintGraph/Private/BlueprintActionDatabase.cpp
@@ -1098,6 +1098,14 @@ void FBlueprintActionDatabase::AddReferencedObjects(FReferenceCollector& Collect
 	}
 }
 
+int32 GBlueprintDatabasePrimingMaxPerFrame = 16;
+static FAutoConsoleVariableRef CVarBlueprintDatabasePrimingMaxPerFrame(
+	TEXT("bp.DatabasePrimingMaxPerFrame"),
+	GBlueprintDatabasePrimingMaxPerFrame,
+	TEXT("How many entries should be primed in to the database per frame."),
+	ECVF_Default
+);
+
 //------------------------------------------------------------------------------
 void FBlueprintActionDatabase::Tick(float DeltaTime)
 {
@@ -1122,10 +1130,9 @@ void FBlueprintActionDatabase::Tick(float DeltaTime)
 	
 	// priming every database entry at once would cause a hitch, so we spread it 
 	// out over several frames
-	static int32 const PrimingMaxPerFrame = 16;
 	int32 PrimedCount = 0;
 
-	while ((ActionPrimingQueue.Num() > 0) && (PrimedCount < PrimingMaxPerFrame))
+	while ((ActionPrimingQueue.Num() > 0) && (PrimedCount < GBlueprintDatabasePrimingMaxPerFrame))
 	{
 		auto ActionIndex = ActionPrimingQueue.CreateIterator();	
 			
@@ -1136,7 +1143,7 @@ void FBlueprintActionDatabase::Tick(float DeltaTime)
 			if (FActionList* ClassActionList = ActionRegistry.Find(ActionsKey))
 			{
 				int32& ActionListIndex = ActionIndex.Value();
-				for (; (ActionListIndex < ClassActionList->Num()) && (PrimedCount < PrimingMaxPerFrame); ++ActionListIndex)
+				for (; (ActionListIndex < ClassActionList->Num()) && (PrimedCount < GBlueprintDatabasePrimingMaxPerFrame); ++ActionListIndex)
 				{
 					UBlueprintNodeSpawner* Action = (*ClassActionList)[ActionListIndex];
 					Action->Prime();
