Delegate functions target pins should never be self, but rather that of the class that contains the delegate.
- Fixes redirected delegate nodes not ending up with a target pin type of the containing blueprint
Delegate nodes now correctly use DisplayName if specified on the property

diff --git a/Engine/Source/Editor/BlueprintGraph/Classes/K2Node_BaseMCDelegate.h b/Engine/Source/Editor/BlueprintGraph/Classes/K2Node_BaseMCDelegate.h
--- a/Engine/Source/Editor/BlueprintGraph/Classes/K2Node_BaseMCDelegate.h
+++ b/Engine/Source/Editor/BlueprintGraph/Classes/K2Node_BaseMCDelegate.h
@@ -48,7 +48,7 @@ public:
 
 	BLUEPRINTGRAPH_API UProperty* GetProperty() const
 	{
-		return DelegateReference.ResolveMember<UProperty>(GetBlueprintClassFromNode());
+		return DelegateReference.ResolveMember<UMulticastDelegateProperty>(GetBlueprintClassFromNode());
 	}
 
 	BLUEPRINTGRAPH_API FName GetPropertyName() const
@@ -56,6 +56,12 @@ public:
 		return DelegateReference.GetMemberName();
 	}
 
+	BLUEPRINTGRAPH_API FText GetPropertyDisplayName() const
+	{
+		UProperty* Prop = GetProperty();
+		return (Prop ? Prop->GetDisplayNameText() : FText::FromName(GetPropertyName()));
+	}
+
 	BLUEPRINTGRAPH_API UFunction* GetDelegateSignature(bool bForceNotFromSkelClass = false) const;
 
 	BLUEPRINTGRAPH_API UEdGraphPin* GetDelegatePin() const;
diff --git a/Engine/Source/Editor/BlueprintGraph/Private/BlueprintDelegateNodeSpawner.cpp b/Engine/Source/Editor/BlueprintGraph/Private/BlueprintDelegateNodeSpawner.cpp
--- a/Engine/Source/Editor/BlueprintGraph/Private/BlueprintDelegateNodeSpawner.cpp
+++ b/Engine/Source/Editor/BlueprintGraph/Private/BlueprintDelegateNodeSpawner.cpp
@@ -93,9 +93,8 @@ UBlueprintDelegateNodeSpawner* UBlueprintDelegateNodeSpawner::Create(TSubclassOf
 		{
 			UBlueprint* Blueprint = FBlueprintEditorUtils::FindBlueprintForNodeChecked(NewNode);
 			UClass* OwnerClass = MCDProperty->GetOwnerClass();
-			bool const bIsSelfContext = Blueprint->SkeletonGeneratedClass->IsChildOf(OwnerClass);
 
-			DelegateNode->SetFromProperty(MCDProperty, bIsSelfContext);
+			DelegateNode->SetFromProperty(MCDProperty, false);
 		}
 	};
 	NodeSpawner->SetNodeFieldDelegate = FSetNodeFieldDelegate::CreateStatic(SetDelegateLambda);
diff --git a/Engine/Source/Editor/BlueprintGraph/Private/K2Node_MCDelegate.cpp b/Engine/Source/Editor/BlueprintGraph/Private/K2Node_MCDelegate.cpp
--- a/Engine/Source/Editor/BlueprintGraph/Private/K2Node_MCDelegate.cpp
+++ b/Engine/Source/Editor/BlueprintGraph/Private/K2Node_MCDelegate.cpp
@@ -82,12 +82,9 @@ void UK2Node_BaseMCDelegate::AllocateDefaultPins()
 	{
 		PropertyOwnerClass = PropertyOwnerClass->GetAuthoritativeClass();
 	}
-	const UBlueprint* Blueprint = GetBlueprint();
 	
-	const bool bUseSelf = Blueprint && (PropertyOwnerClass == Blueprint->GeneratedClass);
+	UEdGraphPin* SelfPin = CreatePin(EGPD_Input, UEdGraphSchema_K2::PC_Object, PropertyOwnerClass, UEdGraphSchema_K2::PN_Self);
 
-	UEdGraphPin* SelfPin = NULL;
-	if (bUseSelf)
 	{
 		SelfPin = CreatePin(EGPD_Input, UEdGraphSchema_K2::PC_Object, UEdGraphSchema_K2::PSC_Self, UEdGraphSchema_K2::PN_Self);
 	}
@@ -96,7 +93,7 @@ void UK2Node_BaseMCDelegate::AllocateDefaultPins()
 		SelfPin = CreatePin(EGPD_Input, UEdGraphSchema_K2::PC_Object, PropertyOwnerClass, UEdGraphSchema_K2::PN_Self);
 	}
 
-	if(SelfPin)
+	if (SelfPin)
 	{
 		SelfPin->PinFriendlyName = NSLOCTEXT("K2Node", "BaseMCDelegateSelfPinName", "Target");
 	}
@@ -204,7 +201,7 @@ bool UK2Node_BaseMCDelegate::HasExternalDependencies(TArray<class UStruct*>* Opt
 {
 	const UBlueprint* SourceBlueprint = GetBlueprint();
 
-	auto MCProperty = GetProperty();
+	UProperty* MCProperty = GetProperty();
 	UClass* PropertySourceClass = MCProperty ? MCProperty->GetOwnerClass() : nullptr;
 	const bool bPropertyResult = (PropertySourceClass != NULL) && (PropertySourceClass->ClassGeneratedBy != SourceBlueprint);
 	if (bPropertyResult && OptionalOutput)
@@ -212,7 +209,7 @@ bool UK2Node_BaseMCDelegate::HasExternalDependencies(TArray<class UStruct*>* Opt
 		OptionalOutput->AddUnique(PropertySourceClass);
 	}
 
-	auto Signature = GetDelegateSignature(true);
+	UFunction* Signature = GetDelegateSignature(true);
 	UClass* SignatureSourceClass = Signature ? Signature->GetOwnerClass() : nullptr;
 	const bool bSignatureResult = (SignatureSourceClass != NULL) && (SignatureSourceClass->ClassGeneratedBy != SourceBlueprint);
 	if (bSignatureResult && OptionalOutput)
@@ -245,7 +242,7 @@ void UK2Node_BaseMCDelegate::AutowireNewNode(UEdGraphPin* FromPin)
 		{
 			if (FromPin->PinType.PinSubCategoryObject.IsValid() && FromPin->PinType.PinSubCategoryObject->IsA(UClass::StaticClass()))
 			{
-				UProperty* DelegateProperty = DelegateReference.ResolveMember<UProperty>(GetBlueprintClassFromNode());
+				UProperty* DelegateProperty = DelegateReference.ResolveMember<UMulticastDelegateProperty>(GetBlueprintClassFromNode());
 				if (DelegateProperty)
 				{
 					UClass* DelegateOwner = DelegateProperty->GetOwnerClass();
@@ -306,7 +303,7 @@ FText UK2Node_AddDelegate::GetNodeTitle(ENodeTitleType::Type TitleType) const
 	if (CachedNodeTitle.IsOutOfDate(this))
 	{
 		FFormatNamedArguments Args;
-		Args.Add(TEXT("PropertyName"), FText::FromName(GetPropertyName()));
+		Args.Add(TEXT("PropertyName"), GetPropertyDisplayName());
 		// FText::Format() is slow, so we cache this to save on performance
 		CachedNodeTitle.SetCachedText(FText::Format(NSLOCTEXT("K2Node", "AddDelegate", "Bind Event to {PropertyName}"), Args), this);
 	}
@@ -337,7 +334,7 @@ FText UK2Node_ClearDelegate::GetNodeTitle(ENodeTitleType::Type TitleType) const
 	if (CachedNodeTitle.IsOutOfDate(this))
 	{
 		FFormatNamedArguments Args;
-		Args.Add(TEXT("PropertyName"), FText::FromName(GetPropertyName()));
+		Args.Add(TEXT("PropertyName"), GetPropertyDisplayName());
 		// FText::Format() is slow, so we cache this to save on performance
 		CachedNodeTitle.SetCachedText(FText::Format(NSLOCTEXT("K2Node", "ClearDelegate", "Unbind all Events from {PropertyName}"), Args), this);
 	}
@@ -375,7 +372,7 @@ FText UK2Node_RemoveDelegate::GetNodeTitle(ENodeTitleType::Type TitleType) const
 	if (CachedNodeTitle.IsOutOfDate(this))
 	{
 		FFormatNamedArguments Args;
-		Args.Add(TEXT("PropertyName"), FText::FromName(GetPropertyName()));
+		Args.Add(TEXT("PropertyName"), GetPropertyDisplayName());
 		// FText::Format() is slow, so we cache this to save on performance
 		CachedNodeTitle.SetCachedText(FText::Format(NSLOCTEXT("K2Node", "RemoveDelegate", "Unbind Event from {PropertyName}"), Args), this);
 	}
@@ -428,7 +425,7 @@ FText UK2Node_CallDelegate::GetNodeTitle(ENodeTitleType::Type TitleType) const
 	if (CachedNodeTitle.IsOutOfDate(this))
 	{
 		FFormatNamedArguments Args;
-		Args.Add(TEXT("PropertyName"), FText::FromName(GetPropertyName()));
+		Args.Add(TEXT("PropertyName"), GetPropertyDisplayName());
 		// FText::Format() is slow, so we cache this to save on performance
 		CachedNodeTitle.SetCachedText(FText::Format(NSLOCTEXT("K2Node", "CallDelegate", "Call {PropertyName}"), Args), this);
 	}
diff --git a/Engine/Source/Runtime/Engine/Classes/Engine/MemberReference.h b/Engine/Source/Runtime/Engine/Classes/Engine/MemberReference.h
--- a/Engine/Source/Runtime/Engine/Classes/Engine/MemberReference.h
+++ b/Engine/Source/Runtime/Engine/Classes/Engine/MemberReference.h
@@ -369,8 +369,8 @@ public:
 					MemberParent  = ParentAsClass;
 
 					// Re-evaluate self-ness against the redirect if we were given a valid SelfScope
-					// For functions we don't want to go from not-self to self as we lose the knowledge of overriding
-					if (SelfScope != nullptr && (bSelfContext || TFieldType::StaticClass() != UFunction::StaticClass()))
+					// For functions and multicast delegates we don't want to go from not-self to self as the target pin type should remain consistent
+					if (SelfScope != nullptr && (bSelfContext || (TFieldType::StaticClass() != UFunction::StaticClass() && TFieldType::StaticClass() != UMulticastDelegateProperty::StaticClass())))
 					{
 						SetGivenSelfScope(MemberName, MemberGuid, ParentAsClass, SelfScope);
 					}
