Remove code that was intended to have been deleted that snuck back in
#rnx

diff --git a/Engine/Source/Editor/BlueprintGraph/Private/K2Node_MCDelegate.cpp b/Engine/Source/Editor/BlueprintGraph/Private/K2Node_MCDelegate.cpp
--- a/Engine/Source/Editor/BlueprintGraph/Private/K2Node_MCDelegate.cpp
+++ b/Engine/Source/Editor/BlueprintGraph/Private/K2Node_MCDelegate.cpp
@@ -85,14 +85,6 @@ void UK2Node_BaseMCDelegate::AllocateDefaultPins()
 	
 	UEdGraphPin* SelfPin = CreatePin(EGPD_Input, UEdGraphSchema_K2::PC_Object, PropertyOwnerClass, UEdGraphSchema_K2::PN_Self);
 
-	{
-		SelfPin = CreatePin(EGPD_Input, UEdGraphSchema_K2::PC_Object, UEdGraphSchema_K2::PSC_Self, UEdGraphSchema_K2::PN_Self);
-	}
-	else
-	{
-		SelfPin = CreatePin(EGPD_Input, UEdGraphSchema_K2::PC_Object, PropertyOwnerClass, UEdGraphSchema_K2::PN_Self);
-	}
-
 	if (SelfPin)
 	{
 		SelfPin->PinFriendlyName = NSLOCTEXT("K2Node", "BaseMCDelegateSelfPinName", "Target");
