Use full function description when no [REDACTED_HANDLE] or [REDACTED_HANDLE] info is provided

diff --git a/Engine/Source/Editor/BlueprintGraph/Private/K2Node_CallFunction.cpp b/Engine/Source/Editor/BlueprintGraph/Private/K2Node_CallFunction.cpp
--- a/Engine/Source/Editor/BlueprintGraph/Private/K2Node_CallFunction.cpp
+++ b/Engine/Source/Editor/BlueprintGraph/Private/K2Node_CallFunction.cpp
@@ -1506,6 +1506,16 @@ void UK2Node_CallFunction::GeneratePinTooltipFromFunction(UEdGraphPin& Pin, cons
 
 	} while (CurStrPos < FullToolTipLen);
 
+	// If we have no parameter or return value descriptions the full description will be relevant in describing the return value:
+	if( bReturnPin && 
+		Pin.PinToolTip.IsEmpty() && 
+		FunctionToolTipText.Find(TEXT("@param")) == INDEX_NONE && 
+		FunctionToolTipText.Find(TEXT("@return")) == INDEX_NONE)
+	{
+		// for the return pin, default to using the function description if no @return tag was provided:
+		Pin.PinToolTip = Function->GetToolTipText().ToString();
+	}
+
 	GetDefault<UEdGraphSchema_K2>()->ConstructBasicPinTooltip(Pin, FText::FromString(Pin.PinToolTip), Pin.PinToolTip);
 }
 
