Fixing parent checks in GetSubsystem nodes

diff --git a/Engine/Source/Editor/BlueprintGraph/Private/K2Node_GetSubsystem.cpp b/Engine/Source/Editor/BlueprintGraph/Private/K2Node_GetSubsystem.cpp
--- a/Engine/Source/Editor/BlueprintGraph/Private/K2Node_GetSubsystem.cpp
+++ b/Engine/Source/Editor/BlueprintGraph/Private/K2Node_GetSubsystem.cpp
@@ -478,7 +478,7 @@ void UK2Node_GetEngineSubsystem::ExpandNode(class FKismetCompilerContext& Compil
 
 	// Choose appropriate underlying Getter
 	FName Get_FunctionName;
-	if (CustomClass->IsChildOf<ULocalPlayerSubsystem>())
+	if (CustomClass->IsChildOf<UEngineSubsystem>())
 	{
 		Get_FunctionName = GET_FUNCTION_NAME_CHECKED(USubsystemBlueprintLibrary, GetEngineSubsystem);
 	}
@@ -600,7 +600,7 @@ void UK2Node_GetEditorSubsystem::ExpandNode(class FKismetCompilerContext& Compil
 
 	// Choose appropriate underlying Getter
 	FName Get_FunctionName;
-	if (CustomClass->IsChildOf<ULocalPlayerSubsystem>())
+	if (CustomClass->IsChildOf<UEditorSubsystem>())
 	{
 		Get_FunctionName = GET_FUNCTION_NAME_CHECKED(UEditorSubsystemBlueprintLibrary, GetEditorSubsystem);
 	}
