Don't do redundant work
#rnx

diff --git a/Engine/Source/Editor/BlueprintGraph/Private/K2Node_ActorBoundEvent.cpp b/Engine/Source/Editor/BlueprintGraph/Private/K2Node_ActorBoundEvent.cpp
--- a/Engine/Source/Editor/BlueprintGraph/Private/K2Node_ActorBoundEvent.cpp
+++ b/Engine/Source/Editor/BlueprintGraph/Private/K2Node_ActorBoundEvent.cpp
@@ -62,7 +62,6 @@ void UK2Node_ActorBoundEvent::ReconstructNode()
 			// Found a remapped property, update the node
 			TargetDelegateProp = NewProperty;
 			DelegatePropertyName = NewProperty->GetFName();
-			CachedNodeTitle.MarkDirty();
 		}
 	}
 
diff --git a/Engine/Source/Editor/BlueprintGraph/Private/K2Node_ComponentBoundEvent.cpp b/Engine/Source/Editor/BlueprintGraph/Private/K2Node_ComponentBoundEvent.cpp
--- a/Engine/Source/Editor/BlueprintGraph/Private/K2Node_ComponentBoundEvent.cpp
+++ b/Engine/Source/Editor/BlueprintGraph/Private/K2Node_ComponentBoundEvent.cpp
@@ -147,7 +147,6 @@ void UK2Node_ComponentBoundEvent::ReconstructNode()
 			// Found a remapped property, update the node
 			TargetDelegateProp = NewProperty;
 			DelegatePropertyName = NewProperty->GetFName();
-			CachedNodeTitle.MarkDirty();
 		}
 	}
 
