PR #5496: Blueprint search improvements (Contributed by [REDACTED_PERSON])

Change summary:
- Include enum literal values in search indexing, as well as the name of functions referenced by [REDACTED_PERSON] nodes.

diff --git a/Engine/Source/Editor/BlueprintGraph/Classes/K2Node_CreateDelegate.h b/Engine/Source/Editor/BlueprintGraph/Classes/K2Node_CreateDelegate.h
--- a/Engine/Source/Editor/BlueprintGraph/Classes/K2Node_CreateDelegate.h
+++ b/Engine/Source/Editor/BlueprintGraph/Classes/K2Node_CreateDelegate.h
@@ -17,7 +17,7 @@ class UK2Node_CreateDelegate : public UK2Node
 {
 	GENERATED_UCLASS_BODY()
 
-	UPROPERTY()
+	UPROPERTY(meta = (BlueprintSearchable = "true"))
 	FName SelectedFunctionName;
 
 	UPROPERTY()
diff --git a/Engine/Source/Editor/BlueprintGraph/Classes/K2Node_EnumLiteral.h b/Engine/Source/Editor/BlueprintGraph/Classes/K2Node_EnumLiteral.h
--- a/Engine/Source/Editor/BlueprintGraph/Classes/K2Node_EnumLiteral.h
+++ b/Engine/Source/Editor/BlueprintGraph/Classes/K2Node_EnumLiteral.h
@@ -30,6 +30,7 @@ class UK2Node_EnumLiteral : public UK2Node, public INodeDependingOnEnumInterface
 	virtual FText GetNodeTitle(ENodeTitleType::Type TitleType) const override;
 	virtual FSlateIcon GetIconAndTint(FLinearColor& OutColor) const override;
 	virtual void ValidateNodeDuringCompilation(class FCompilerResultsLog& MessageLog) const override;
+	virtual void AddSearchMetaDataInfo(TArray<struct FSearchTagDataPair>& OutTaggedMetaData) const override;
 	//~ End UEdGraphNode Interface
 
 	//~ Begin UK2Node Interface
diff --git a/Engine/Source/Editor/BlueprintGraph/Private/K2Node_EnumLiteral.cpp b/Engine/Source/Editor/BlueprintGraph/Private/K2Node_EnumLiteral.cpp
--- a/Engine/Source/Editor/BlueprintGraph/Private/K2Node_EnumLiteral.cpp
+++ b/Engine/Source/Editor/BlueprintGraph/Private/K2Node_EnumLiteral.cpp
@@ -12,6 +12,7 @@
 #include "KismetCompilerMisc.h"
 #include "BlueprintFieldNodeSpawner.h"
 #include "EditorCategoryUtils.h"
+#include "FindInBlueprintManager.h"
 #include "BlueprintActionDatabaseRegistrar.h"
 
 const FName UK2Node_EnumLiteral::GetEnumInputPinName()
@@ -34,6 +35,34 @@ void UK2Node_EnumLiteral::ValidateNodeDuringCompilation(class FCompilerResultsLo
 	}
 }
 
+void UK2Node_EnumLiteral::AddSearchMetaDataInfo(TArray<struct FSearchTagDataPair>& OutTaggedMetaData) const
+{
+	Super::AddSearchMetaDataInfo(OutTaggedMetaData);
+
+	const UEdGraphPin* Pin = FindPinChecked(GetEnumInputPinName());
+	if (!Pin->DefaultValue.IsEmpty())
+	{
+		OutTaggedMetaData.Add(FSearchTagDataPair(FFindInBlueprintSearchTags::FiB_NativeName, FText::FromString(Pin->DefaultValue)));
+
+		const int32 ValueIndex = Enum ? Enum->GetIndexByName(*Enum->GenerateFullEnumName(*Pin->DefaultValue)) : INDEX_NONE;
+		if (ValueIndex != INDEX_NONE)
+		{
+			FText SearchName = FText::FormatOrdered(NSLOCTEXT("K2Node", "EnumLiteral_SearchName", "{0} - {1}"), GetTooltipText(), Enum->GetDisplayNameText(ValueIndex));
+
+			// Find the Name, populated by Super::AddSearchMetaDataInfo
+			for (FSearchTagDataPair& SearchData : OutTaggedMetaData)
+			{
+				// Should always be the first item, but there is no guarantee
+				if (SearchData.Key.CompareTo(FFindInBlueprintSearchTags::FiB_Name) == 0)
+				{
+					SearchData.Value = SearchName;
+					break;
+				}
+			}
+		}
+	}
+}
+
 void UK2Node_EnumLiteral::AllocateDefaultPins()
 {
 	const UEdGraphSchema_K2* Schema = GetDefault<UEdGraphSchema_K2>();
