Fix deprecated API usage (CIS).

diff --git a/Engine/Source/Editor/BlueprintGraph/Private/K2Node_EnumLiteral.cpp b/Engine/Source/Editor/BlueprintGraph/Private/K2Node_EnumLiteral.cpp
--- a/Engine/Source/Editor/BlueprintGraph/Private/K2Node_EnumLiteral.cpp
+++ b/Engine/Source/Editor/BlueprintGraph/Private/K2Node_EnumLiteral.cpp
@@ -47,7 +47,7 @@ void UK2Node_EnumLiteral::AddSearchMetaDataInfo(TArray<struct FSearchTagDataPair
 		const int32 ValueIndex = Enum ? Enum->GetIndexByName(*Enum->GenerateFullEnumName(*Pin->DefaultValue)) : INDEX_NONE;
 		if (ValueIndex != INDEX_NONE)
 		{
-			FText SearchName = FText::FormatOrdered(NSLOCTEXT("K2Node", "EnumLiteral_SearchName", "{0} - {1}"), GetTooltipText(), Enum->GetDisplayNameText(ValueIndex));
+			FText SearchName = FText::FormatOrdered(NSLOCTEXT("K2Node", "EnumLiteral_SearchName", "{0} - {1}"), GetTooltipText(), Enum->GetDisplayNameTextByIndex(ValueIndex));
 
 			// Find the Name, populated by Super::AddSearchMetaDataInfo
 			for (FSearchTagDataPair& SearchData : OutTaggedMetaData)
