Use metadata statics instead of recreating FNames each time
#rnx

diff --git a/Engine/Source/Editor/BlueprintGraph/Private/EdGraphSchema_K2.cpp b/Engine/Source/Editor/BlueprintGraph/Private/EdGraphSchema_K2.cpp
--- a/Engine/Source/Editor/BlueprintGraph/Private/EdGraphSchema_K2.cpp
+++ b/Engine/Source/Editor/BlueprintGraph/Private/EdGraphSchema_K2.cpp
@@ -1293,7 +1293,7 @@ bool UEdGraphSchema_K2::PinHasSplittableStructType(const UEdGraphPin* InGraphPin
 				bCanSplit = UK2Node_MakeStruct::CanBeSplit(StructType);
 				if (!bCanSplit)
 				{
-					const FString& MetaData = StructType->GetMetaData(TEXT("HasNativeMake"));
+					const FString& MetaData = StructType->GetMetaData(FBlueprintMetadata::MD_NativeMakeFunction);
 					UFunction* Function = FindObject<UFunction>(NULL, *MetaData, true);
 					bCanSplit = (Function != NULL);
 				}
@@ -1303,7 +1303,7 @@ bool UEdGraphSchema_K2::PinHasSplittableStructType(const UEdGraphPin* InGraphPin
 				bCanSplit = UK2Node_BreakStruct::CanBeSplit(StructType);
 				if (!bCanSplit)
 				{
-					const FString& MetaData = StructType->GetMetaData(TEXT("HasNativeBreak"));
+					const FString& MetaData = StructType->GetMetaData(FBlueprintMetadata::MD_NativeBreakFunction);
 					UFunction* Function = FindObject<UFunction>(NULL, *MetaData, true);
 					bCanSplit = (Function != NULL);
 				}
@@ -6800,7 +6800,7 @@ UK2Node* UEdGraphSchema_K2::CreateSplitPinNode(UEdGraphPin* Pin, const FCreateSp
 		}
 		else
 		{
-			const FString& MetaData = StructType->GetMetaData(TEXT("HasNativeMake"));
+			const FString& MetaData = StructType->GetMetaData(FBlueprintMetadata::MD_NativeMakeFunction);
 			const UFunction* Function = FindObject<UFunction>(nullptr, *MetaData, true);
 
 			UK2Node_CallFunction* CallFunctionNode;
@@ -6848,7 +6848,7 @@ UK2Node* UEdGraphSchema_K2::CreateSplitPinNode(UEdGraphPin* Pin, const FCreateSp
 		}
 		else
 		{
-			const FString& MetaData = StructType->GetMetaData(TEXT("HasNativeBreak"));
+			const FString& MetaData = StructType->GetMetaData(FBlueprintMetadata::MD_NativeBreakFunction);
 			const UFunction* Function = FindObject<UFunction>(nullptr, *MetaData, true);
 
 			UK2Node_CallFunction* CallFunctionNode;
diff --git a/Engine/Source/Editor/BlueprintGraph/Private/K2Node_BreakStruct.cpp b/Engine/Source/Editor/BlueprintGraph/Private/K2Node_BreakStruct.cpp
--- a/Engine/Source/Editor/BlueprintGraph/Private/K2Node_BreakStruct.cpp
+++ b/Engine/Source/Editor/BlueprintGraph/Private/K2Node_BreakStruct.cpp
@@ -154,7 +154,7 @@ static bool CanCreatePinForProperty(const UProperty* Property)
 
 bool UK2Node_BreakStruct::CanBeBroken(const UScriptStruct* Struct, const bool bForInternalUse)
 {
-	if (Struct && !Struct->HasMetaData(TEXT("HasNativeBreak")) && UEdGraphSchema_K2::IsAllowableBlueprintVariableType(Struct, bForInternalUse))
+	if (Struct && !Struct->HasMetaData(FBlueprintMetadata::MD_NativeBreakFunction) && UEdGraphSchema_K2::IsAllowableBlueprintVariableType(Struct, bForInternalUse))
 	{
 		for (TFieldIterator<UProperty> It(Struct); It; ++It)
 		{
diff --git a/Engine/Source/Editor/BlueprintGraph/Private/K2Node_MakeStruct.cpp b/Engine/Source/Editor/BlueprintGraph/Private/K2Node_MakeStruct.cpp
--- a/Engine/Source/Editor/BlueprintGraph/Private/K2Node_MakeStruct.cpp
+++ b/Engine/Source/Editor/BlueprintGraph/Private/K2Node_MakeStruct.cpp
@@ -243,7 +243,7 @@ FLinearColor UK2Node_MakeStruct::GetNodeTitleColor() const
 
 bool UK2Node_MakeStruct::CanBeMade(const UScriptStruct* Struct, const bool bForInternalUse)
 {
-	return (Struct && !Struct->HasMetaData(TEXT("HasNativeMake")) && UEdGraphSchema_K2::IsAllowableBlueprintVariableType(Struct, bForInternalUse));
+	return (Struct && !Struct->HasMetaData(FBlueprintMetadata::MD_NativeMakeFunction) && UEdGraphSchema_K2::IsAllowableBlueprintVariableType(Struct, bForInternalUse));
 }
 
 bool UK2Node_MakeStruct::CanBeSplit(const UScriptStruct* Struct)
