PR #5551: Don't create WorldContext pin when creating node with wire (Contributed by [REDACTED_PERSON])

diff --git a/Engine/Source/Editor/BlueprintGraph/Private/K2Node.cpp b/Engine/Source/Editor/BlueprintGraph/Private/K2Node.cpp
--- a/Engine/Source/Editor/BlueprintGraph/Private/K2Node.cpp
+++ b/Engine/Source/Editor/BlueprintGraph/Private/K2Node.cpp
@@ -266,10 +266,26 @@ void UK2Node::AutowireNewNode(UEdGraphPin* FromPin)
 			UEdGraphPin* Pin = Pins[i];
 			check(Pin);
 
-			// Never consider for auto-wiring a hidden pin being connected to a Wildcard. It is never what the user expects
-			if (Pin->bHidden && FromPin->PinType.PinCategory == UEdGraphSchema_K2::PC_Wildcard)
+			if (Pin->bHidden)
 			{
-				continue;
+				// Never consider for auto-wiring a hidden pin being connected to a Wildcard. It is never what the user expects
+				if (FromPin->PinType.PinCategory == UEdGraphSchema_K2::PC_Wildcard)
+				{
+					continue;
+				}
+
+				// Never connect wires to hidden WorldContextObject pins
+				if (UK2Node_CallFunction* CallFunctionNode = Cast<UK2Node_CallFunction>(this))
+				{				
+					if (UFunction* NodeTargetFunction = CallFunctionNode->GetTargetFunction())
+					{
+						const FString& WorldContextPinName = NodeTargetFunction->GetMetaData(FBlueprintMetadata::MD_WorldContext);
+						if (!WorldContextPinName.IsEmpty() && WorldContextPinName == Pin->PinName.ToString())
+						{
+							continue;
+						}
+					}
+				}
 			}
 
 			ECanCreateConnectionResponse ConnectResponse = K2Schema->CanCreateConnection(FromPin, Pin).Response;
