Fix EditoronlyBPFunctionRedirects to work correctly if the old function no longer exists, and not crash if the new one doesn't

diff --git a/Engine/Source/Editor/BlueprintGraph/Private/EdGraphSchema_K2.cpp b/Engine/Source/Editor/BlueprintGraph/Private/EdGraphSchema_K2.cpp
--- a/Engine/Source/Editor/BlueprintGraph/Private/EdGraphSchema_K2.cpp
+++ b/Engine/Source/Editor/BlueprintGraph/Private/EdGraphSchema_K2.cpp
@@ -5969,19 +5969,21 @@ struct FBackwardCompatibilityConversionHelper
 	{
 		if (ConversionParams.FuncScope)
 		{
-			const UFunction* OldFunc = ConversionParams.FuncScope->FindFunctionByName(ConversionParams.OldFuncName);
-			check(OldFunc);
 			const UFunction* NewFunc = ConversionParams.FuncScope->FindFunctionByName(ConversionParams.NewFuncName);
-			check(NewFunc);
-
-			for (UK2Node_CallFunction* Node : Nodes)
+			if (ensureMsgf(NewFunc, TEXT("Can't find conversion function %s on %s!"), *ConversionParams.NewFuncName.ToString(), *ConversionParams.FuncScope->GetName()))
 			{
-				if (OldFunc == Node->GetTargetFunction())
+				for (UK2Node_CallFunction* Node : Nodes)
 				{
-					UK2Node_CallFunction* NewNode = NewObject<UK2Node_CallFunction>(Graph);
-					NewNode->SetFromFunction(NewFunc);
-					ConvertNode(Node, ConversionParams.BlueprintPinName, NewNode,
-						ConversionParams.ClassPinName, Schema, bOnlyWithDefaultBlueprint);
+					// Check to see if the class scope and name are the same, we can't depend on the UFunction still existing
+					UClass* MemberParent = Node->FunctionReference.GetMemberParentClass(Node->GetBlueprintClassFromNode());
+
+					if (MemberParent == ConversionParams.FuncScope && Node->FunctionReference.GetMemberName() == ConversionParams.OldFuncName)
+					{
+						UK2Node_CallFunction* NewNode = NewObject<UK2Node_CallFunction>(Graph);
+						NewNode->SetFromFunction(NewFunc);
+						ConvertNode(Node, ConversionParams.BlueprintPinName, NewNode,
+							ConversionParams.ClassPinName, Schema, bOnlyWithDefaultBlueprint);
+					}
 				}
 			}
 		}
