Based on PR #5467 but changed the metadata name

diff --git a/Engine/Source/Editor/BlueprintGraph/Classes/EdGraphSchema_K2.h b/Engine/Source/Editor/BlueprintGraph/Classes/EdGraphSchema_K2.h
--- a/Engine/Source/Editor/BlueprintGraph/Classes/EdGraphSchema_K2.h
+++ b/Engine/Source/Editor/BlueprintGraph/Classes/EdGraphSchema_K2.h
@@ -114,6 +114,9 @@ public:
 	// [FunctionMetadata] Indicates that the function should be drawn with this title over the function name
 	static const FName MD_DisplayName;
 
+	// [FunctionMetadata] Indicates the display name of the return value pin
+	static const FName MD_ReturnDisplayName;
+
 	// [FunctionMetadata] Indicates that a particular function parameter is for internal use only, which means it will be both hidden and not connectible.
 	static const FName MD_InternalUseParam;
 
diff --git a/Engine/Source/Editor/BlueprintGraph/Private/EdGraphSchema_K2.cpp b/Engine/Source/Editor/BlueprintGraph/Private/EdGraphSchema_K2.cpp
--- a/Engine/Source/Editor/BlueprintGraph/Private/EdGraphSchema_K2.cpp
+++ b/Engine/Source/Editor/BlueprintGraph/Private/EdGraphSchema_K2.cpp
@@ -117,6 +117,7 @@ const FName FBlueprintMetadata::MD_DeprecatedFunction(TEXT("DeprecatedFunction")
 const FName FBlueprintMetadata::MD_DeprecationMessage(TEXT("DeprecationMessage"));
 const FName FBlueprintMetadata::MD_CompactNodeTitle(TEXT("CompactNodeTitle"));
 const FName FBlueprintMetadata::MD_DisplayName(TEXT("DisplayName"));
+const FName FBlueprintMetadata::MD_ReturnDisplayName(TEXT("ReturnDisplayName"));
 const FName FBlueprintMetadata::MD_InternalUseParam(TEXT("InternalUseParam"));
 
 const FName FBlueprintMetadata::MD_PropertyGetFunction(TEXT("BlueprintGetter"));
diff --git a/Engine/Source/Editor/BlueprintGraph/Private/K2Node_CallFunction.cpp b/Engine/Source/Editor/BlueprintGraph/Private/K2Node_CallFunction.cpp
--- a/Engine/Source/Editor/BlueprintGraph/Private/K2Node_CallFunction.cpp
+++ b/Engine/Source/Editor/BlueprintGraph/Private/K2Node_CallFunction.cpp
@@ -1037,6 +1037,10 @@ bool UK2Node_CallFunction::CreatePinsForFunctionCall(const UFunction* Function)
 			{
 				Pin->PinFriendlyName = FText::FromString(PinDisplayName);
 			}
+			else if (Function->GetReturnProperty() == Param && Function->HasMetaData(FBlueprintMetadata::MD_ReturnDisplayName))
+			{
+				Pin->PinFriendlyName = Function->GetMetaDataText(FBlueprintMetadata::MD_ReturnDisplayName);
+			}
 
 			//Flag pin as read only for const reference property
 			Pin->bDefaultValueIsIgnored = Param->HasAllPropertyFlags(CPF_ConstParm | CPF_ReferenceParm) && (!Function->HasMetaData(FBlueprintMetadata::MD_AutoCreateRefTerm) || Pin->PinType.IsContainer());
