Based on PR #5362, but this only sets the display name and doesn't change the internal name, to avoid backward compatibility issues

diff --git a/Engine/Source/Editor/BlueprintGraph/Private/K2Node_BaseAsyncTask.cpp b/Engine/Source/Editor/BlueprintGraph/Private/K2Node_BaseAsyncTask.cpp
--- a/Engine/Source/Editor/BlueprintGraph/Private/K2Node_BaseAsyncTask.cpp
+++ b/Engine/Source/Editor/BlueprintGraph/Private/K2Node_BaseAsyncTask.cpp
@@ -73,10 +73,15 @@ void UK2Node_BaseAsyncTask::AllocateDefaultPins()
 
 	bool bExposeProxy = false;
 	bool bHideThen = false;
+	FText ExposeProxyDisplayName;
 	for (const UStruct* TestStruct = ProxyClass; TestStruct; TestStruct = TestStruct->GetSuperStruct())
 	{
 		bExposeProxy |= TestStruct->HasMetaData(TEXT("ExposedAsyncProxy"));
 		bHideThen |= TestStruct->HasMetaData(TEXT("HideThen"));
+		if (ExposeProxyDisplayName.IsEmpty())
+		{
+			ExposeProxyDisplayName = TestStruct->GetMetaDataText(TEXT("ExposedAsyncProxy"));
+		}
 	}
 
 	if (!bHideThen)
@@ -86,7 +91,11 @@ void UK2Node_BaseAsyncTask::AllocateDefaultPins()
 
 	if (bExposeProxy)
 	{
-		CreatePin(EGPD_Output, UEdGraphSchema_K2::PC_Object, ProxyClass, FBaseAsyncTaskHelper::GetAsyncTaskProxyName());
+		UEdGraphPin* ProxyPin = CreatePin(EGPD_Output, UEdGraphSchema_K2::PC_Object, ProxyClass, FBaseAsyncTaskHelper::GetAsyncTaskProxyName());
+		if (!ExposeProxyDisplayName.IsEmpty())
+		{
+			ProxyPin->PinFriendlyName = ExposeProxyDisplayName;
+		}
 	}
 
 	UFunction* Function = GetFactoryFunction();
diff --git a/Engine/Source/Runtime/GameplayTasks/Classes/GameplayTask.h b/Engine/Source/Runtime/GameplayTasks/Classes/GameplayTask.h
--- a/Engine/Source/Runtime/GameplayTasks/Classes/GameplayTask.h
+++ b/Engine/Source/Runtime/GameplayTasks/Classes/GameplayTask.h
@@ -136,7 +136,7 @@ public:
 	FString GetDebugDescription() const;
 };
 
-UCLASS(Abstract, BlueprintType, meta = (ExposedAsyncProxy), config = Game)
+UCLASS(Abstract, BlueprintType, meta = (ExposedAsyncProxy=AsyncTask), config = Game)
 class GAMEPLAYTASKS_API UGameplayTask : public UObject, public IGameplayTaskOwnerInterface
 {
 	GENERATED_BODY()
