[NO NON-PII COMMIT MESSAGE]

diff --git a/Engine/Source/Editor/BlueprintGraph/Private/EdGraphSchema_K2.cpp b/Engine/Source/Editor/BlueprintGraph/Private/EdGraphSchema_K2.cpp
--- a/Engine/Source/Editor/BlueprintGraph/Private/EdGraphSchema_K2.cpp
+++ b/Engine/Source/Editor/BlueprintGraph/Private/EdGraphSchema_K2.cpp
@@ -2401,7 +2401,22 @@ private:
 			// Exclude the bitmask subcategory string from integral types so that autocast will work.
 			PinSubCategory.Reset();
 		}
-		return FString::Printf(TEXT("%s;%s;%s"), *PinType.PinCategory.ToString(), *PinSubCategory, Obj ? *Obj->GetPathName() : TEXT(""));
+		
+		FString TypeString = FString::Printf(TEXT("%s;%s;%s;%d"), *PinType.PinCategory.ToString(), *PinSubCategory, Obj ? *Obj->GetPathName() : TEXT(""), (int32)PinType.ContainerType);
+
+		if (PinType.ContainerType == EPinContainerType::Map)
+		{
+			// Add value type to string
+			Obj = PinType.PinValueType.TerminalSubCategoryObject.Get();
+			PinSubCategory = PinType.PinValueType.TerminalSubCategory.ToString();
+			if (PinSubCategory.StartsWith(UEdGraphSchema_K2::PSC_Bitmask.ToString()))
+			{
+				PinSubCategory.Reset();
+			}
+			return FString::Printf(TEXT("%s;%s;%s;%s"), *TypeString, *PinType.PinValueType.TerminalCategory.ToString(), *PinSubCategory, Obj ? *Obj->GetPathName() : TEXT(""));
+		}
+
+		return TypeString;
 	}
 
 	static FString GenerateCastData(const FEdGraphPinType& InputPinType, const FEdGraphPinType& OutputPinType)
@@ -2571,68 +2586,72 @@ bool UEdGraphSchema_K2::SearchForAutocastFunction(const UEdGraphPin* OutputPin,
 			FunctionOwner = Function->GetOwnerClass();
 			return true;
 		}
-		return false;
-	}
-
-	// SPECIAL CASES, not supported by FAutocastFunctionMap
-	if ((OutputPin->PinType.PinCategory == PC_Interface) && (InputPin->PinType.PinCategory == PC_Object))
-	{
-		UClass const* InputClass = Cast<UClass const>(InputPin->PinType.PinSubCategoryObject.Get());
-
-		bool const bInputIsUObject = (InputClass && (InputClass == UObject::StaticClass()));
-		if (bInputIsUObject)
-		{
-			UFunction* Function = UKismetSystemLibrary::StaticClass()->FindFunctionByName(GET_MEMBER_NAME_CHECKED(UKismetSystemLibrary, Conv_InterfaceToObject));
-			TargetFunction = Function->GetFName();
-			FunctionOwner = Function->GetOwnerClass();
-		}
+		
+		// Skip the other special cases if container check fails, but allow checking the autocast map
 	}
-	else if (OutputPin->PinType.PinCategory == PC_Object)
+	else
 	{
-		UClass const* OutputClass = Cast<UClass const>(OutputPin->PinType.PinSubCategoryObject.Get());
-		if (InputPin->PinType.PinCategory == PC_Class)
+		// SPECIAL CASES, not supported by FAutocastFunctionMap.
+		if ((OutputPin->PinType.PinCategory == PC_Interface) && (InputPin->PinType.PinCategory == PC_Object))
 		{
 			UClass const* InputClass = Cast<UClass const>(InputPin->PinType.PinSubCategoryObject.Get());
-			if ((OutputClass != nullptr) &&
-				(InputClass != nullptr) &&
-				OutputClass->IsChildOf(InputClass))
+
+			bool const bInputIsUObject = (InputClass && (InputClass == UObject::StaticClass()));
+			if (bInputIsUObject)
 			{
-				UFunction* Function = UGameplayStatics::StaticClass()->FindFunctionByName(GET_MEMBER_NAME_CHECKED(UGameplayStatics, GetObjectClass));
+				UFunction* Function = UKismetSystemLibrary::StaticClass()->FindFunctionByName(GET_MEMBER_NAME_CHECKED(UKismetSystemLibrary, Conv_InterfaceToObject));
 				TargetFunction = Function->GetFName();
 				FunctionOwner = Function->GetOwnerClass();
 			}
 		}
-		else if (InputPin->PinType.PinCategory == PC_String)
+		else if (OutputPin->PinType.PinCategory == PC_Object)
 		{
-			UFunction* Function = UKismetSystemLibrary::StaticClass()->FindFunctionByName(GET_MEMBER_NAME_CHECKED(UKismetSystemLibrary, GetDisplayName));
-			TargetFunction = Function->GetFName();
-			FunctionOwner = Function->GetOwnerClass();
-		}
-	}
-	else if (OutputPin->PinType.PinCategory == PC_Class)
-	{
-		if (InputPin->PinType.PinCategory == PC_String)
-		{
-			UFunction* Function = UKismetSystemLibrary::StaticClass()->FindFunctionByName(GET_MEMBER_NAME_CHECKED(UKismetSystemLibrary, GetClassDisplayName));
-			TargetFunction = Function->GetFName();
-			FunctionOwner = Function->GetOwnerClass();
+			UClass const* OutputClass = Cast<UClass const>(OutputPin->PinType.PinSubCategoryObject.Get());
+			if (InputPin->PinType.PinCategory == PC_Class)
+			{
+				UClass const* InputClass = Cast<UClass const>(InputPin->PinType.PinSubCategoryObject.Get());
+				if ((OutputClass != nullptr) &&
+					(InputClass != nullptr) &&
+					OutputClass->IsChildOf(InputClass))
+				{
+					UFunction* Function = UGameplayStatics::StaticClass()->FindFunctionByName(GET_MEMBER_NAME_CHECKED(UGameplayStatics, GetObjectClass));
+					TargetFunction = Function->GetFName();
+					FunctionOwner = Function->GetOwnerClass();
+				}
+			}
+			else if (InputPin->PinType.PinCategory == PC_String)
+			{
+				UFunction* Function = UKismetSystemLibrary::StaticClass()->FindFunctionByName(GET_MEMBER_NAME_CHECKED(UKismetSystemLibrary, GetDisplayName));
+				TargetFunction = Function->GetFName();
+				FunctionOwner = Function->GetOwnerClass();
+			}
 		}
-	}
-	else if (OutputPin->PinType.PinCategory == PC_Struct)
-	{
-		const UScriptStruct* OutputStructType = Cast<const UScriptStruct>(OutputPin->PinType.PinSubCategoryObject.Get());
-		if (OutputStructType == TBaseStructure<FRotator>::Get())
+		else if (OutputPin->PinType.PinCategory == PC_Class)
 		{
-			const UScriptStruct* InputStructType = Cast<const UScriptStruct>(InputPin->PinType.PinSubCategoryObject.Get());
-			if ((InputPin->PinType.PinCategory == PC_Struct) && (InputStructType == TBaseStructure<FTransform>::Get()))
+			if (InputPin->PinType.PinCategory == PC_String)
 			{
-				UFunction* Function = UKismetMathLibrary::StaticClass()->FindFunctionByName(GET_MEMBER_NAME_CHECKED(UKismetMathLibrary, MakeTransform));
+				UFunction* Function = UKismetSystemLibrary::StaticClass()->FindFunctionByName(GET_MEMBER_NAME_CHECKED(UKismetSystemLibrary, GetClassDisplayName));
 				TargetFunction = Function->GetFName();
 				FunctionOwner = Function->GetOwnerClass();
 			}
 		}
+		else if (OutputPin->PinType.PinCategory == PC_Struct)
+		{
+			const UScriptStruct* OutputStructType = Cast<const UScriptStruct>(OutputPin->PinType.PinSubCategoryObject.Get());
+			if (OutputStructType == TBaseStructure<FRotator>::Get())
+			{
+				const UScriptStruct* InputStructType = Cast<const UScriptStruct>(InputPin->PinType.PinSubCategoryObject.Get());
+				if ((InputPin->PinType.PinCategory == PC_Struct) && (InputStructType == TBaseStructure<FTransform>::Get()))
+				{
+					UFunction* Function = UKismetMathLibrary::StaticClass()->FindFunctionByName(GET_MEMBER_NAME_CHECKED(UKismetMathLibrary, MakeTransform));
+					TargetFunction = Function->GetFName();
+					FunctionOwner = Function->GetOwnerClass();
+				}
+			}
+		}
 	}
 
+	// Try looking for a marked up autocast if we've not found a built-in one that works
 	if (TargetFunction == NAME_None)
 	{
 		const FAutocastFunctionMap& AutocastFunctionMap = FAutocastFunctionMap::Get();
