ConstructObject node should accept class variables of type Object connected to the Class pin.

diff --git a/Engine/Source/Editor/BlueprintGraph/Private/K2Node_GenericCreateObject.cpp b/Engine/Source/Editor/BlueprintGraph/Private/K2Node_GenericCreateObject.cpp
--- a/Engine/Source/Editor/BlueprintGraph/Private/K2Node_GenericCreateObject.cpp
+++ b/Engine/Source/Editor/BlueprintGraph/Private/K2Node_GenericCreateObject.cpp
@@ -18,7 +18,8 @@ struct FK2Node_GenericCreateObject_Utils
 			&& (bAllowAbstract || !ObjectClass->HasAnyClassFlags(CLASS_Abstract))
 			&& !ObjectClass->HasAnyClassFlags(CLASS_Deprecated | CLASS_NewerVersionExists);
 
-		if (bCanSpawnObject)
+		// UObject is a special case where if we are allowing abstract we are going to allow it through even though it doesn't have BlueprintType on it
+		if (bCanSpawnObject && (!bAllowAbstract || (*ObjectClass != UObject::StaticClass())))
 		{
 			static const FName BlueprintTypeName(TEXT("BlueprintType"));
 			static const FName NotBlueprintTypeName(TEXT("NotBlueprintType"));
