PR #5649: Make ConstructObject searchable (Contributed by [REDACTED_PERSON])

diff --git a/Engine/Source/Editor/BlueprintGraph/Classes/K2Node_ConstructObjectFromClass.h b/Engine/Source/Editor/BlueprintGraph/Classes/K2Node_ConstructObjectFromClass.h
--- a/Engine/Source/Editor/BlueprintGraph/Classes/K2Node_ConstructObjectFromClass.h
+++ b/Engine/Source/Editor/BlueprintGraph/Classes/K2Node_ConstructObjectFromClass.h
@@ -27,6 +27,7 @@ class BLUEPRINTGRAPH_API UK2Node_ConstructObjectFromClass : public UK2Node
 	virtual void PinConnectionListChanged(UEdGraphPin* Pin);
 	virtual void GetPinHoverText(const UEdGraphPin& Pin, FString& HoverTextOut) const override;
 	virtual void PostPlacedNewNode() override;
+	virtual void AddSearchMetaDataInfo(TArray<struct FSearchTagDataPair>& OutTaggedMetaData) const override;
 	//~ End UEdGraphNode Interface.
 
 	//~ Begin UK2Node Interface
diff --git a/Engine/Source/Editor/BlueprintGraph/Private/K2Node_ConstructObjectFromClass.cpp b/Engine/Source/Editor/BlueprintGraph/Private/K2Node_ConstructObjectFromClass.cpp
--- a/Engine/Source/Editor/BlueprintGraph/Private/K2Node_ConstructObjectFromClass.cpp
+++ b/Engine/Source/Editor/BlueprintGraph/Private/K2Node_ConstructObjectFromClass.cpp
@@ -7,6 +7,7 @@
 #include "BlueprintNodeSpawner.h"
 #include "EditorCategoryUtils.h"
 #include "BlueprintActionDatabaseRegistrar.h"
+#include "FindInBlueprintManager.h"
 
 struct FK2Node_ConstructObjectFromClassHelper
 {
@@ -176,6 +177,12 @@ void UK2Node_ConstructObjectFromClass::PostPlacedNewNode()
 	}
 }
 
+void UK2Node_ConstructObjectFromClass::AddSearchMetaDataInfo(TArray<struct FSearchTagDataPair>& OutTaggedMetaData) const
+{
+	Super::AddSearchMetaDataInfo(OutTaggedMetaData);
+	OutTaggedMetaData.Add(FSearchTagDataPair(FFindInBlueprintSearchTags::FiB_NativeName, CachedNodeTitle.GetCachedText()));
+}
+
 bool UK2Node_ConstructObjectFromClass::IsSpawnVarPin(UEdGraphPin* Pin) const
 {
 	return(	Pin->PinName != UEdGraphSchema_K2::PN_Execute &&
