[NO NON-PII COMMIT MESSAGE]

diff --git a/Engine/Source/Editor/BlueprintGraph/Classes/K2Node_AddComponent.h b/Engine/Source/Editor/BlueprintGraph/Classes/K2Node_AddComponent.h
--- a/Engine/Source/Editor/BlueprintGraph/Classes/K2Node_AddComponent.h
+++ b/Engine/Source/Editor/BlueprintGraph/Classes/K2Node_AddComponent.h
@@ -43,6 +43,7 @@ class UK2Node_AddComponent : public UK2Node_CallFunction
 	virtual FString GetDocumentationExcerptName() const override;
 	virtual bool IsCompatibleWithGraph(UEdGraph const* Graph) const override;
 	virtual void ReconstructNode() override;
+	virtual void FindDiffs(class UEdGraphNode* OtherNode, struct FDiffResults& Results) override;
 	//~ End UEdGraphNode Interface
 
 	//~ Begin K2Node Interface
diff --git a/Engine/Source/Editor/BlueprintGraph/Private/K2Node_AddComponent.cpp b/Engine/Source/Editor/BlueprintGraph/Private/K2Node_AddComponent.cpp
--- a/Engine/Source/Editor/BlueprintGraph/Private/K2Node_AddComponent.cpp
+++ b/Engine/Source/Editor/BlueprintGraph/Private/K2Node_AddComponent.cpp
@@ -14,6 +14,7 @@
 #include "UObject/ReleaseObjectVersion.h"
 #include "KismetCompilerMisc.h"
 #include "KismetCompiler.h"
+#include "DiffResults.h"
 #include "UObject/BlueprintsObjectVersion.h" // for ComponentTemplateClassSupport
 
 #define LOCTEXT_NAMESPACE "K2Node_AddComponent"
@@ -407,6 +408,28 @@ void UK2Node_AddComponent::ReconstructNode()
 	}
 }
 
+void UK2Node_AddComponent::FindDiffs(class UEdGraphNode* OtherNode, struct FDiffResults& Results)
+{
+	Super::FindDiffs(OtherNode, Results);
+
+	if (UK2Node_AddComponent* OtherGraphNode = Cast<UK2Node_AddComponent>(OtherNode))
+	{
+		UActorComponent* MyComponent = GetTemplateFromNode();
+		UActorComponent* OtherComponent = OtherGraphNode->GetTemplateFromNode();
+		if (MyComponent && OtherComponent)
+		{
+			FDiffSingleResult Diff;
+			Diff.Diff = EDiffType::NODE_PROPERTY;
+			Diff.Node1 = this;
+			Diff.Node2 = OtherNode;
+			Diff.ToolTip = LOCTEXT("DIF_ComponentTemplatePropertyToolTip", "A property of the component template has changed");
+			Diff.DisplayColor = FLinearColor(0.25f, 0.71f, 0.85f);
+
+			DiffProperties(MyComponent->GetClass(), OtherComponent->GetClass(), MyComponent, OtherComponent, Results, Diff);
+		}
+	}
+}
+
 void UK2Node_AddComponent::ExpandNode(class FKismetCompilerContext& CompilerContext, UEdGraph* SourceGraph)
 {
 	Super::ExpandNode(CompilerContext, SourceGraph);
