Fully clear entry node's member reference on duplicate to prevent it from colliding with the function it was duplicated from

diff --git a/Engine/Source/Editor/BlueprintGraph/Private/EdGraphSchema_K2.cpp b/Engine/Source/Editor/BlueprintGraph/Private/EdGraphSchema_K2.cpp
--- a/Engine/Source/Editor/BlueprintGraph/Private/EdGraphSchema_K2.cpp
+++ b/Engine/Source/Editor/BlueprintGraph/Private/EdGraphSchema_K2.cpp
@@ -6508,7 +6508,12 @@ UEdGraph* UEdGraphSchema_K2::DuplicateGraph(UEdGraph* GraphToDuplicate) const
 					if (EntryNode->FunctionReference.GetMemberName() == GraphToDuplicate->GetFName())
 					{
 						EntryNode->Modify();
-						EntryNode->FunctionReference.SetMemberName(NewGraph->GetFName());
+						
+						// We're duplicating the graph, so fully reset the member reference (including the GUID!)
+						FMemberReference NewRef;
+						NewRef.SetMemberName(NewGraph->GetFName());
+						EntryNode->FunctionReference = NewRef;
+
 						break;
 					}
 				}
