Add a compiler warning for nodes that reference a deprecated event dispatcher (MC delegate) variable in a Blueprint class.

diff --git a/Engine/Source/Editor/BlueprintGraph/Classes/K2Node_BaseMCDelegate.h b/Engine/Source/Editor/BlueprintGraph/Classes/K2Node_BaseMCDelegate.h
--- a/Engine/Source/Editor/BlueprintGraph/Classes/K2Node_BaseMCDelegate.h
+++ b/Engine/Source/Editor/BlueprintGraph/Classes/K2Node_BaseMCDelegate.h
@@ -33,6 +33,8 @@ public:
 	virtual bool HasExternalDependencies(TArray<class UStruct*>* OptionalOutput) const override;
 	virtual void GetNodeAttributes( TArray<TKeyValuePair<FString, FString>>& OutNodeAttributes ) const override;
 	virtual void AutowireNewNode(UEdGraphPin* FromPin) override;
+	virtual bool IsDeprecated() const override;
+	virtual FString GetDeprecationMessage() const override;
 	// End of UK2Node interface
 
 	// UEdGraphNode interface
diff --git a/Engine/Source/Editor/BlueprintGraph/Classes/K2Node_Variable.h b/Engine/Source/Editor/BlueprintGraph/Classes/K2Node_Variable.h
--- a/Engine/Source/Editor/BlueprintGraph/Classes/K2Node_Variable.h
+++ b/Engine/Source/Editor/BlueprintGraph/Classes/K2Node_Variable.h
@@ -71,8 +71,8 @@ public:
 	virtual void AutowireNewNode(UEdGraphPin* FromPin) override;
 	virtual bool CanPasteHere(const UEdGraph* TargetGraph) const override;
 	virtual void PostPasteNode() override;
-	virtual bool IsDeprecated() const;
-	virtual FString GetDeprecationMessage() const;
+	virtual bool IsDeprecated() const override;
+	virtual FString GetDeprecationMessage() const override;
 	virtual UObject* GetJumpTargetForDoubleClick() const override;
 	virtual bool CanJumpToDefinition() const override;
 	virtual void JumpToDefinition() const override;
diff --git a/Engine/Source/Editor/BlueprintGraph/Private/K2Node_MCDelegate.cpp b/Engine/Source/Editor/BlueprintGraph/Private/K2Node_MCDelegate.cpp
--- a/Engine/Source/Editor/BlueprintGraph/Private/K2Node_MCDelegate.cpp
+++ b/Engine/Source/Editor/BlueprintGraph/Private/K2Node_MCDelegate.cpp
@@ -269,6 +269,33 @@ void UK2Node_BaseMCDelegate::AutowireNewNode(UEdGraphPin* FromPin)
 	}
 }
 
+bool UK2Node_BaseMCDelegate::IsDeprecated() const
+{
+	if (Super::IsDeprecated() || DelegateReference.IsDeprecated())
+	{
+		return true;
+	}
+
+	UProperty* DelegateProperty = DelegateReference.ResolveMember<UProperty>(GetBlueprintClassFromNode());
+	if (DelegateProperty
+		&& (DelegateProperty->HasAllPropertyFlags(CPF_Deprecated) || DelegateProperty->HasMetaData(FBlueprintMetadata::MD_DeprecationMessage)))
+	{
+		return true;
+	}
+	return false;
+}
+
+FString UK2Node_BaseMCDelegate::GetDeprecationMessage() const
+{
+	UProperty* DelegateProperty = DelegateReference.ResolveMember<UProperty>(GetBlueprintClassFromNode());
+	if (DelegateProperty && DelegateProperty->HasMetaData(FBlueprintMetadata::MD_DeprecationMessage))
+	{
+		return FString::Printf(TEXT("%s %s"), *NSLOCTEXT("K2Node", "DelegateDeprecated_Warning", "@@ is deprecated;").ToString(), *DelegateProperty->GetMetaData(FBlueprintMetadata::MD_DeprecationMessage));
+	}
+
+	return Super::GetDeprecationMessage();
+}
+
 /////// UK2Node_AddDelegate ///////////
 
 UK2Node_AddDelegate::UK2Node_AddDelegate(const FObjectInitializer& ObjectInitializer)
