Fix various consistency issues with the editor UI response to deprecated member usage in a Blueprint class.

Change summary:
- Added FBlueprintEditorUtils::GetDeprecatedMemberUsageNodeWarning() to format a consistent compiler warning message
- Modified Variable, CallFunction and Delegate binding nodes to use the shared compiler warning message for deprecation.
- Fixed an issue that caused the deprecation compiler warning message to appear twice on overridden event nodes.
- Retained backwards-compatibility on Variable nodes for native references that only include 'DeprecationMessage' metadata (these were previously considered by the editor to be deprecated).
- Modified FunctionEntry and Event nodes to display a deprecation compiler warning only in the override case. The warning will not appear on the original definition. This includes custom event nodes as well.
- Updated the warning message in UK2Node_CreateDelegate::IsValid() for the !UEdGraphSchema_K2::FunctionCanBeUsedInDelegate() case (which also fails if the function has been marked as deprecated).

diff --git a/Engine/Source/Editor/BlueprintGraph/Classes/K2Node_CustomEvent.h b/Engine/Source/Editor/BlueprintGraph/Classes/K2Node_CustomEvent.h
--- a/Engine/Source/Editor/BlueprintGraph/Classes/K2Node_CustomEvent.h
+++ b/Engine/Source/Editor/BlueprintGraph/Classes/K2Node_CustomEvent.h
@@ -49,6 +49,9 @@ class UK2Node_CustomEvent : public UK2Node_Event
 	virtual void AutowireNewNode(UEdGraphPin* FromPin) override;
 	virtual void AddSearchMetaDataInfo(TArray<struct FSearchTagDataPair>& OutTaggedMetaData) const override;
 	virtual FText GetKeywords() const override;
+	virtual bool IsDeprecated() const override { return bIsDeprecated; }
+	virtual bool ShouldWarnOnDeprecation() const override;
+	virtual FString GetDeprecationMessage() const override;
 	//~ End UEdGraphNode Interface
 
 	//~ Begin UK2Node Interface
diff --git a/Engine/Source/Editor/BlueprintGraph/Classes/K2Node_Event.h b/Engine/Source/Editor/BlueprintGraph/Classes/K2Node_Event.h
--- a/Engine/Source/Editor/BlueprintGraph/Classes/K2Node_Event.h
+++ b/Engine/Source/Editor/BlueprintGraph/Classes/K2Node_Event.h
@@ -67,6 +67,7 @@ class UK2Node_Event : public UK2Node_EditablePinBase, public IK2Node_EventNodeIn
 	BLUEPRINTGRAPH_API virtual bool IsCompatibleWithGraph(const UEdGraph* TargetGraph) const override;
 	BLUEPRINTGRAPH_API virtual FName GetCornerIcon() const override;
 	BLUEPRINTGRAPH_API virtual bool IsDeprecated() const override;
+	BLUEPRINTGRAPH_API virtual bool ShouldWarnOnDeprecation() const override;
 	BLUEPRINTGRAPH_API virtual FString GetDeprecationMessage() const override;
 	BLUEPRINTGRAPH_API virtual UObject* GetJumpTargetForDoubleClick() const override;
 	BLUEPRINTGRAPH_API virtual FSlateIcon GetIconAndTint(FLinearColor& OutColor) const override;
diff --git a/Engine/Source/Editor/BlueprintGraph/Classes/K2Node_FunctionEntry.h b/Engine/Source/Editor/BlueprintGraph/Classes/K2Node_FunctionEntry.h
--- a/Engine/Source/Editor/BlueprintGraph/Classes/K2Node_FunctionEntry.h
+++ b/Engine/Source/Editor/BlueprintGraph/Classes/K2Node_FunctionEntry.h
@@ -41,6 +41,7 @@ class UK2Node_FunctionEntry : public UK2Node_FunctionTerminator
 	virtual FText GetNodeTitle(ENodeTitleType::Type TitleType) const override;
 	virtual bool CanUserDeleteNode() const override { return false; }
 	virtual bool IsDeprecated() const override;
+	virtual bool ShouldWarnOnDeprecation() const override;
 	virtual FString GetDeprecationMessage() const override;
 	virtual FText GetTooltipText() const override;
 	virtual void FindDiffs(UEdGraphNode* OtherNode, struct FDiffResults& Results) override;
diff --git a/Engine/Source/Editor/BlueprintGraph/Private/K2Node_CallFunction.cpp b/Engine/Source/Editor/BlueprintGraph/Private/K2Node_CallFunction.cpp
--- a/Engine/Source/Editor/BlueprintGraph/Private/K2Node_CallFunction.cpp
+++ b/Engine/Source/Editor/BlueprintGraph/Private/K2Node_CallFunction.cpp
@@ -451,13 +451,14 @@ bool UK2Node_CallFunction::ShouldWarnOnDeprecation() const
 
 FString UK2Node_CallFunction::GetDeprecationMessage() const
 {
-	UFunction* Function = GetTargetFunction();
-	if (Function && Function->HasMetaData(FBlueprintMetadata::MD_DeprecationMessage))
+	FText Result;
+	if (UFunction* Function = GetTargetFunction())
 	{
-		return FString::Printf(TEXT("%s %s"), *LOCTEXT("CallFunctionDeprecated_Warning", "@@ is deprecated;").ToString(), *Function->GetMetaData(FBlueprintMetadata::MD_DeprecationMessage));
+		FString DetailedMessage = Function->GetMetaData(FBlueprintMetadata::MD_DeprecationMessage);
+		Result = FBlueprintEditorUtils::GetDeprecatedMemberUsageNodeWarning(GetUserFacingFunctionName(Function), FText::FromString(DetailedMessage));
 	}
-
-	return Super::GetDeprecationMessage();
+	
+	return Result.ToString();
 }
 
 
diff --git a/Engine/Source/Editor/BlueprintGraph/Private/K2Node_CreateDelegate.cpp b/Engine/Source/Editor/BlueprintGraph/Private/K2Node_CreateDelegate.cpp
--- a/Engine/Source/Editor/BlueprintGraph/Private/K2Node_CreateDelegate.cpp
+++ b/Engine/Source/Editor/BlueprintGraph/Private/K2Node_CreateDelegate.cpp
@@ -158,7 +158,7 @@ bool UK2Node_CreateDelegate::IsValid(FString* OutMsg, bool bDontUseSkeletalClass
 	{
 		if (OutMsg)
 		{
-			*OutMsg = NSLOCTEXT("K2Node", "Function_cannot_be_used_in_delegate", "The selected function/event is not bindable - is the function/event pure or latent?").ToString();
+			*OutMsg = NSLOCTEXT("K2Node", "Function_cannot_be_used_in_delegate", "The selected function/event is not bindable - is the function/event deprecated, pure or latent?").ToString();
 		}
 		return false;
 	}
diff --git a/Engine/Source/Editor/BlueprintGraph/Private/K2Node_CustomEvent.cpp b/Engine/Source/Editor/BlueprintGraph/Private/K2Node_CustomEvent.cpp
--- a/Engine/Source/Editor/BlueprintGraph/Private/K2Node_CustomEvent.cpp
+++ b/Engine/Source/Editor/BlueprintGraph/Private/K2Node_CustomEvent.cpp
@@ -519,4 +519,17 @@ FText UK2Node_CustomEvent::GetKeywords() const
 	return FText::Format(LOCTEXT("CustomEventKeywords", "{ParentKeywords} Custom"), Args);
 }
 
+bool UK2Node_CustomEvent::ShouldWarnOnDeprecation() const
+{
+	// Only warn on override usage. This allows the source event to be marked as deprecated in the class that defines it without warning.
+	return IsOverride();
+}
+
+FString UK2Node_CustomEvent::GetDeprecationMessage() const
+{
+	FText EventName = FText::FromName(GetFunctionName());
+	FText DetailedMessage = FText::FromString(DeprecationMessage);
+	return FBlueprintEditorUtils::GetDeprecatedMemberUsageNodeWarning(EventName, DetailedMessage).ToString();
+}
+
 #undef LOCTEXT_NAMESPACE
diff --git a/Engine/Source/Editor/BlueprintGraph/Private/K2Node_Event.cpp b/Engine/Source/Editor/BlueprintGraph/Private/K2Node_Event.cpp
--- a/Engine/Source/Editor/BlueprintGraph/Private/K2Node_Event.cpp
+++ b/Engine/Source/Editor/BlueprintGraph/Private/K2Node_Event.cpp
@@ -789,17 +789,23 @@ bool UK2Node_Event::IsDeprecated() const
 	return false;
 }
 
+bool UK2Node_Event::ShouldWarnOnDeprecation() const
+{
+	// Only warn on override usage. This allows the source event to be marked as deprecated in the class that defines it without warning.
+	return bOverrideFunction;
+}
+
 FString UK2Node_Event::GetDeprecationMessage() const
 {
+	FText Result;
 	if (UFunction* Function = EventReference.ResolveMember<UFunction>(GetBlueprintClassFromNode()))
 	{
-		if (Function->HasMetaData(FBlueprintMetadata::MD_DeprecationMessage))
-		{
-			return FString::Printf(TEXT("%s %s"), *LOCTEXT("EventDeprecated_Warning", "@@ is deprecated;").ToString(), *Function->GetMetaData(FBlueprintMetadata::MD_DeprecationMessage));
-		}
+		FName EventName = GetFunctionName();
+		FString DetailedMessage = Function->GetMetaData(FBlueprintMetadata::MD_DeprecationMessage);
+		Result = FBlueprintEditorUtils::GetDeprecatedMemberUsageNodeWarning(FText::FromName(EventName), FText::FromString(DetailedMessage));
 	}
 
-	return Super::GetDeprecationMessage();
+	return Result.ToString();
 }
 
 UObject* UK2Node_Event::GetJumpTargetForDoubleClick() const
diff --git a/Engine/Source/Editor/BlueprintGraph/Private/K2Node_FunctionEntry.cpp b/Engine/Source/Editor/BlueprintGraph/Private/K2Node_FunctionEntry.cpp
--- a/Engine/Source/Editor/BlueprintGraph/Private/K2Node_FunctionEntry.cpp
+++ b/Engine/Source/Editor/BlueprintGraph/Private/K2Node_FunctionEntry.cpp
@@ -470,17 +470,23 @@ bool UK2Node_FunctionEntry::IsDeprecated() const
 	return false;
 }
 
+bool UK2Node_FunctionEntry::ShouldWarnOnDeprecation() const
+{
+	// Only warn on non-editable (i.e. override) usage. This allows the source graph to be marked as deprecated in the class that defines it without warning.
+	return !IsEditable();
+}
+
 FString UK2Node_FunctionEntry::GetDeprecationMessage() const
 {
+	FText Result;
 	if (UFunction* const Function = FunctionReference.ResolveMember<UFunction>(GetBlueprintClassFromNode()))
 	{
-		if (Function->HasMetaData(FBlueprintMetadata::MD_DeprecationMessage))
-		{
-			return FString::Printf(TEXT("%s %s"), *LOCTEXT("FunctionDeprecated_Warning", "@@ is deprecated;").ToString(), *Function->GetMetaData(FBlueprintMetadata::MD_DeprecationMessage));
-		}
+		FName FunctionName = FunctionReference.GetMemberName();
+		FString DetailedMessage = Function->GetMetaData(FBlueprintMetadata::MD_DeprecationMessage);
+		Result = FBlueprintEditorUtils::GetDeprecatedMemberUsageNodeWarning(FText::FromName(FunctionName), FText::FromString(DetailedMessage));
 	}
 
-	return Super::GetDeprecationMessage();
+	return Result.ToString();
 }
 
 FText UK2Node_FunctionEntry::GetTooltipText() const
diff --git a/Engine/Source/Editor/BlueprintGraph/Private/K2Node_MCDelegate.cpp b/Engine/Source/Editor/BlueprintGraph/Private/K2Node_MCDelegate.cpp
--- a/Engine/Source/Editor/BlueprintGraph/Private/K2Node_MCDelegate.cpp
+++ b/Engine/Source/Editor/BlueprintGraph/Private/K2Node_MCDelegate.cpp
@@ -17,6 +17,7 @@
 #include "K2Node_CallDelegate.h"
 #include "K2Node_ClearDelegate.h"
 #include "K2Node_RemoveDelegate.h"
+#include "Kismet2/BlueprintEditorUtils.h"
 
 #include "KismetCompiler.h"
 #include "DelegateNodeHandlers.h"
@@ -271,29 +272,26 @@ void UK2Node_BaseMCDelegate::AutowireNewNode(UEdGraphPin* FromPin)
 
 bool UK2Node_BaseMCDelegate::IsDeprecated() const
 {
-	if (Super::IsDeprecated() || DelegateReference.IsDeprecated())
-	{
-		return true;
-	}
-
-	UProperty* DelegateProperty = DelegateReference.ResolveMember<UProperty>(GetBlueprintClassFromNode());
-	if (DelegateProperty
-		&& (DelegateProperty->HasAllPropertyFlags(CPF_Deprecated) || DelegateProperty->HasMetaData(FBlueprintMetadata::MD_DeprecationMessage)))
-	{
-		return true;
-	}
-	return false;
+	// Check if either the node itself or the referenced delegate is deprecated.
+	return Super::IsDeprecated() || DelegateReference.IsDeprecated();
 }
 
 FString UK2Node_BaseMCDelegate::GetDeprecationMessage() const
 {
-	UProperty* DelegateProperty = DelegateReference.ResolveMember<UProperty>(GetBlueprintClassFromNode());
-	if (DelegateProperty && DelegateProperty->HasMetaData(FBlueprintMetadata::MD_DeprecationMessage))
+	// Handle the default case where the node itself has been deprecated (return the default message).
+	if (Super::IsDeprecated())
+	{
+		return Super::GetDeprecationMessage();
+	}
+
+	FText Result;
+	if (UProperty* DelegateProperty = DelegateReference.ResolveMember<UProperty>(GetBlueprintClassFromNode()))
 	{
-		return FString::Printf(TEXT("%s %s"), *NSLOCTEXT("K2Node", "DelegateDeprecated_Warning", "@@ is deprecated;").ToString(), *DelegateProperty->GetMetaData(FBlueprintMetadata::MD_DeprecationMessage));
+		FString DetailedMessage = DelegateProperty->GetMetaData(FBlueprintMetadata::MD_DeprecationMessage);
+		Result = FBlueprintEditorUtils::GetDeprecatedMemberUsageNodeWarning(GetPropertyDisplayName(), FText::FromString(DetailedMessage));
 	}
 
-	return Super::GetDeprecationMessage();
+	return Result.ToString();
 }
 
 /////// UK2Node_AddDelegate ///////////
diff --git a/Engine/Source/Editor/BlueprintGraph/Private/K2Node_Variable.cpp b/Engine/Source/Editor/BlueprintGraph/Private/K2Node_Variable.cpp
--- a/Engine/Source/Editor/BlueprintGraph/Private/K2Node_Variable.cpp
+++ b/Engine/Source/Editor/BlueprintGraph/Private/K2Node_Variable.cpp
@@ -954,29 +954,41 @@ void UK2Node_Variable::PostPasteNode()
 
 bool UK2Node_Variable::IsDeprecated() const
 {
+	// Check if either the node itself or the referenced variable is deprecated.
 	if (Super::IsDeprecated() || VariableReference.IsDeprecated())
 	{
 		return true;
 	}
-
-	UProperty* VariableProperty = VariableReference.ResolveMember<UProperty>(GetBlueprintClassFromNode());
-	if (VariableProperty 
-		&& (VariableProperty->HasAllPropertyFlags(CPF_Deprecated) || VariableProperty->HasMetaData(FBlueprintMetadata::MD_DeprecationMessage)))
+	else if (UProperty* VariableProperty = VariableReference.ResolveMember<UProperty>(GetBlueprintClassFromNode()))
 	{
-		return true;
+		// Backcompat: Allow variables tagged only with 'DeprecationMessage' meta to be seen as deprecated if inherited from a native parent class.
+		const bool bHasDeprecationMessage = VariableProperty->HasMetaData(FBlueprintMetadata::MD_DeprecationMessage);
+		if (bHasDeprecationMessage && VariableProperty->GetOuter()->IsNative())
+		{
+			return true;
+		}
 	}
+
 	return false;
 }
 
 FString UK2Node_Variable::GetDeprecationMessage() const
 {
-	UProperty* VariableProperty = VariableReference.ResolveMember<UProperty>(GetBlueprintClassFromNode());
-	if (VariableProperty && VariableProperty->HasMetaData(FBlueprintMetadata::MD_DeprecationMessage))
+	// Handle the default case where the node itself has been deprecated (return the default message).
+	if (Super::IsDeprecated())
 	{
-		return FString::Printf(TEXT("%s %s"), *LOCTEXT("PropertyDeprecated_Warning", "@@ is deprecated;").ToString(), *VariableProperty->GetMetaData(FBlueprintMetadata::MD_DeprecationMessage));
+		return Super::GetDeprecationMessage();
 	}
 
-	return Super::GetDeprecationMessage();
+	FText Result;
+	if (UProperty* VariableProperty = VariableReference.ResolveMember<UProperty>(GetBlueprintClassFromNode()))
+	{
+		FName MemberName = VariableReference.GetMemberName();
+		FString DetailedMessage = VariableProperty->GetMetaData(FBlueprintMetadata::MD_DeprecationMessage);
+		Result = FBlueprintEditorUtils::GetDeprecatedMemberUsageNodeWarning(FText::FromName(MemberName), FText::FromString(DetailedMessage));
+	}
+	
+	return Result.ToString();
 }
 
 UObject* UK2Node_Variable::GetJumpTargetForDoubleClick() const
diff --git a/Engine/Source/Editor/KismetCompiler/Private/KismetCompiler.cpp b/Engine/Source/Editor/KismetCompiler/Private/KismetCompiler.cpp
--- a/Engine/Source/Editor/KismetCompiler/Private/KismetCompiler.cpp
+++ b/Engine/Source/Editor/KismetCompiler/Private/KismetCompiler.cpp
@@ -3150,7 +3150,9 @@ void FKismetCompilerContext::VerifyValidOverrideEvent(const UEdGraph* Graph)
 				{
 					if (EventNode->IsDeprecated())
 					{
-						MessageLog.Warning(*EventNode->GetDeprecationMessage(), EventNode);
+						// The event cannot be placed because it has been deprecated. However, we already emit a
+						// warning in FGraphCompilerContext::ValidateNode(), so there's no need to repeat it here.
+						continue;
 					}
 					else if(!Function->HasAllFunctionFlags(FUNC_Const))	// ...allow legacy event nodes that override methods declared as 'const' to pass.
 					{
diff --git a/Engine/Source/Editor/UnrealEd/Private/Kismet2/BlueprintEditorUtils.cpp b/Engine/Source/Editor/UnrealEd/Private/Kismet2/BlueprintEditorUtils.cpp
--- a/Engine/Source/Editor/UnrealEd/Private/Kismet2/BlueprintEditorUtils.cpp
+++ b/Engine/Source/Editor/UnrealEd/Private/Kismet2/BlueprintEditorUtils.cpp
@@ -9020,6 +9020,17 @@ FString FBlueprintEditorUtils::GetClassNameWithoutSuffix(const UClass* Class)
 	}
 }
 
+FText FBlueprintEditorUtils::GetDeprecatedMemberUsageNodeWarning(const FText& MemberName, const FText& DetailedMessage)
+{
+	static FText UnknownName = LOCTEXT("UnknownDeprecatedMemberName", "[unknown]");
+	static FText DefaultMessage = LOCTEXT("DefaultDeprecatedMemberUsageDetails", "Please replace or remove it.");
+
+	FFormatNamedArguments Args;
+	Args.Add("MemberName", ensure(!MemberName.IsEmpty()) ? MemberName : UnknownName);
+	Args.Add("DetailedMessage", DetailedMessage.IsEmpty() ? DefaultMessage : DetailedMessage);
+	return FText::Format(LOCTEXT("DeprecatedMemberUsageNodeWarning", "@@: Usage of '{MemberName}' has been deprecated. {DetailedMessage}"), Args);
+}
+
 UK2Node_FunctionResult* FBlueprintEditorUtils::FindOrCreateFunctionResultNode(UK2Node_EditablePinBase* InFunctionEntryNode)
 {
 	UK2Node_FunctionResult* FunctionResult = nullptr;
diff --git a/Engine/Source/Editor/UnrealEd/Public/Kismet2/BlueprintEditorUtils.h b/Engine/Source/Editor/UnrealEd/Public/Kismet2/BlueprintEditorUtils.h
--- a/Engine/Source/Editor/UnrealEd/Public/Kismet2/BlueprintEditorUtils.h
+++ b/Engine/Source/Editor/UnrealEd/Public/Kismet2/BlueprintEditorUtils.h
@@ -1593,6 +1593,14 @@ public:
 	 */
 	static FString GetClassNameWithoutSuffix(const UClass* Class);
 
+	/**
+	 * Returns a formatted warning message regarding usage of a deprecated variable or function member with the given name.
+	 *
+	 * @param MemberName		(Required) User-facing name of the deprecated variable or function.
+	 * @param DetailedMessage	(Optional) Instructional text or other details from the owner. If empty, a default message will be used.
+	 */
+	static FText GetDeprecatedMemberUsageNodeWarning(const FText& MemberName, const FText& DetailedMessage);
+
 	/**
 	 * Remove overridden component templates from instance component handlers when a parent class disables editable when inherited boolean.
 	 */
