Add FindSignatureFunction and fix function entry tooltip to work for newly created functions
Fix the soft reference fixup code in K2Node to handle references inside a struct, and update CustomEvent/FunctionEntry to update their user pins defaults after a fixup to editor pin data
Right now only local variables are using the value cache as we have too many blueprint pin defaults that cannot be safely copied into properties, will re-evaluate later and fixup comments

diff --git a/Engine/Source/Editor/BlueprintGraph/Classes/EdGraphSchema_K2.h b/Engine/Source/Editor/BlueprintGraph/Classes/EdGraphSchema_K2.h
--- a/Engine/Source/Editor/BlueprintGraph/Classes/EdGraphSchema_K2.h
+++ b/Engine/Source/Editor/BlueprintGraph/Classes/EdGraphSchema_K2.h
@@ -548,6 +548,9 @@ public:
 	/** Returns true if the pin has a value field that can be edited inline */
 	bool PinDefaultValueIsEditable(const UEdGraphPin& InGraphPin) const;
 
+	/** Returns true if the pin has a custom default string format and it is not safe to use ExportText */
+	bool PinHasCustomDefaultFormat(const UEdGraphPin& InGraphPin) const;
+
 	struct FCreateSplitPinNodeParams
 	{
 		FCreateSplitPinNodeParams(const bool bInTransient)
diff --git a/Engine/Source/Editor/BlueprintGraph/Classes/K2Node.h b/Engine/Source/Editor/BlueprintGraph/Classes/K2Node.h
--- a/Engine/Source/Editor/BlueprintGraph/Classes/K2Node.h
+++ b/Engine/Source/Editor/BlueprintGraph/Classes/K2Node.h
@@ -270,6 +270,9 @@ class UK2Node : public UEdGraphNode
 	/** Expands a node while compiling, which may add additional nodes or delete this node */
 	BLUEPRINTGRAPH_API virtual void ExpandNode(class FKismetCompilerContext& CompilerContext, UEdGraph* SourceGraph);
 
+	/** Clears out any cached data that needs to be regenerated after a structural blueprint change */
+	BLUEPRINTGRAPH_API virtual void ClearCachedBlueprintData(UBlueprint* Blueprint) {}
+
 	/** Performs a node-specific deprecation fixup, which may delete this node and replace it with another one */
 	BLUEPRINTGRAPH_API virtual void ConvertDeprecatedNode(UEdGraph* Graph, bool bOnlySafeChanges) {}
 
@@ -455,7 +458,11 @@ protected:
 		}
 	}
 
-	void FixupPinDefaultValues();
+	/** Handle backwards compatible fixes on load */
+	BLUEPRINTGRAPH_API virtual void FixupPinDefaultValues();
+
+	/** Fixes up structure/soft object ref pins, on both save and load */
+	BLUEPRINTGRAPH_API virtual void FixupPinStringDataReferences(FArchive* SavingArchive);
 
 private:
 
diff --git a/Engine/Source/Editor/BlueprintGraph/Classes/K2Node_CustomEvent.h b/Engine/Source/Editor/BlueprintGraph/Classes/K2Node_CustomEvent.h
--- a/Engine/Source/Editor/BlueprintGraph/Classes/K2Node_CustomEvent.h
+++ b/Engine/Source/Editor/BlueprintGraph/Classes/K2Node_CustomEvent.h
@@ -55,9 +55,10 @@ class UK2Node_CustomEvent : public UK2Node_Event
 	//~ End UEdGraphNode Interface
 
 	//~ Begin UK2Node Interface
-	BLUEPRINTGRAPH_API virtual void ValidateNodeDuringCompilation(class FCompilerResultsLog& MessageLog) const override;
+	virtual void ValidateNodeDuringCompilation(class FCompilerResultsLog& MessageLog) const override;
 	virtual void GetMenuActions(FBlueprintActionDatabaseRegistrar& ActionRegistrar) const override;
-	BLUEPRINTGRAPH_API virtual bool NodeCausesStructuralBlueprintChange() const override { return true; }
+	virtual void FixupPinStringDataReferences(FArchive* SavingArchive) override;
+	virtual bool NodeCausesStructuralBlueprintChange() const override { return true; }
 	//~ End UK2Node Interface
 
 	//~ Begin UK2Node_EditablePinBase Interface
diff --git a/Engine/Source/Editor/BlueprintGraph/Classes/K2Node_EditablePinBase.h b/Engine/Source/Editor/BlueprintGraph/Classes/K2Node_EditablePinBase.h
--- a/Engine/Source/Editor/BlueprintGraph/Classes/K2Node_EditablePinBase.h
+++ b/Engine/Source/Editor/BlueprintGraph/Classes/K2Node_EditablePinBase.h
@@ -152,7 +152,7 @@ class UK2Node_EditablePinBase : public UK2Node
 	 */
 	virtual UEdGraphPin* CreatePinFromUserDefinition(const TSharedPtr<FUserPinInfo> NewPinInfo) { return nullptr; }
 
-	// Modifies the default value of an existing pin on the node.
+	/** Modifies the default value of an existing pin on the node, this will update both the UserPinInfo and the linked editor pin */
 	BLUEPRINTGRAPH_API virtual bool ModifyUserDefinedPinDefaultValue(TSharedPtr<FUserPinInfo> PinInfo, const FString& NewDefaultValue);
 
 	/**
@@ -160,6 +160,24 @@ class UK2Node_EditablePinBase : public UK2Node
 	 */
 	BLUEPRINTGRAPH_API virtual bool CreateUserDefinedPinsForFunctionEntryExit(const UFunction* Function, bool bForFunctionEntry);
 
+	/**
+	 * Copies data from any user pins matching properties in VariableStruct into the VariableStructData. Will also copy LocalVariables if passed in
+	 *
+	 * @param	VariableStruct		UStruct defining the variables to copy
+	 * @param	VariableStructData	Raw data for a variable cache
+	 * @param	LocalVariables		Optional list of local variables to copy
+	 */
+	BLUEPRINTGRAPH_API bool UpdateVariableStructFromDefaults(const UStruct* VariableStruct, uint8* VariableStructData, TArray<FBPVariableDescription>* LocalVariables = nullptr);
+
+	/**
+	 * Copies data from VariableStruct into the user defined pins as well as an optional list of local variables
+	 *
+	 * @param	VariableStruct		UStruct defining the variables to copy
+	 * @param	VariableStructData	Raw data for a variable cache
+	 * @param	LocalVariables		Optional list of local variables to copy
+	 */
+	BLUEPRINTGRAPH_API bool UpdateDefaultsFromVariableStruct(const UStruct* VariableStruct, uint8* VariableStructData, TArray<FBPVariableDescription>* LocalVariables = nullptr);
+
 	/**
 	 * Can this node have execution wires added or removed?
 	 */
@@ -174,5 +192,9 @@ class UK2Node_EditablePinBase : public UK2Node
 	 * Should this node require 'const' for pass-by-reference parameters?
 	 */
 	virtual bool ShouldUseConstRefParams() const { return false;  }
+
+private:
+	/** Internal function that just updates the UEdGraphPin, separate to avoid recursive update calls */
+	bool UpdateEdGraphPinDefaultValue(TSharedPtr<FUserPinInfo> PinInfo, FString& NewDefaultValue);
 };
 
diff --git a/Engine/Source/Editor/BlueprintGraph/Classes/K2Node_FunctionEntry.h b/Engine/Source/Editor/BlueprintGraph/Classes/K2Node_FunctionEntry.h
--- a/Engine/Source/Editor/BlueprintGraph/Classes/K2Node_FunctionEntry.h
+++ b/Engine/Source/Editor/BlueprintGraph/Classes/K2Node_FunctionEntry.h
@@ -6,6 +6,7 @@
 #include "UObject/ObjectMacros.h"
 #include "K2Node_FunctionTerminator.h"
 #include "Engine/Blueprint.h"
+#include "UObject/StructOnScope.h"
 #include "K2Node_FunctionEntry.generated.h"
 
 class UEdGraph;
@@ -34,6 +35,7 @@ class UK2Node_FunctionEntry : public UK2Node_FunctionTerminator
 	//~ Begin UObject Interface
 	virtual void Serialize(FArchive& Ar) override;
 	virtual void PreSave(const class ITargetPlatform* TargetPlatform) override;
+	virtual void PostLoad() override;
 	//~ End UObject Interface
 
 	//~ Begin UEdGraphNode Interface
@@ -53,6 +55,8 @@ class UK2Node_FunctionEntry : public UK2Node_FunctionTerminator
 	virtual void GetRedirectPinNames(const UEdGraphPin& Pin, TArray<FString>& RedirectPinNames) const override;
 	virtual void ExpandNode(class FKismetCompilerContext& CompilerContext, UEdGraph* SourceGraph) override;
 	virtual void PostReconstructNode() override;
+	virtual void ClearCachedBlueprintData(UBlueprint* Blueprint) override;
+	virtual void FixupPinStringDataReferences(FArchive* SavingArchive) override;
 	//~ End UK2Node Interface
 
 	//~ Begin UK2Node_EditablePinBase Interface
@@ -66,6 +70,15 @@ class UK2Node_FunctionEntry : public UK2Node_FunctionTerminator
 	virtual UEdGraphPin* CreatePinFromUserDefinition(const TSharedPtr<FUserPinInfo> NewPinInfo) override;
 	//~ End K2Node_FunctionTerminator Interface
 
+	/** Gets the UFunction and function variable cache structure that should be used for serialization fixups for function variables. If bForceRefresh is true it will always recreate the cache */
+	BLUEPRINTGRAPH_API TSharedPtr<FStructOnScope> GetFunctionVariableCache(bool bForceRefresh = false);
+
+	/** Copies data from the local/input variable defaults into the variable cache */
+	BLUEPRINTGRAPH_API bool RefreshFunctionVariableCache();
+
+	/** Handles updating loaded default values, by going default string into variable cache and back, if bForceRefresh it will happen even if the cache is already setup */
+	BLUEPRINTGRAPH_API bool UpdateLoadedDefaultValues(bool bForceRefresh = false);
+
 	// Removes an output pin from the node
 	BLUEPRINTGRAPH_API void RemoveOutputPin(UEdGraphPin* PinToRemove);
 
@@ -101,5 +114,12 @@ protected:
 	/** Any extra flags that the function may need */
 	UPROPERTY()
 	int32 ExtraFlags;
+
+	/** Holds a an in-memory representation of the UFunction struct, used to fixup local and user variables */
+	TSharedPtr<FStructOnScope> FunctionVariableCache;
+
+	/** True if we've updated the default values on this node at least once */
+	UPROPERTY(Transient)
+	bool bUpdatedDefaultValuesOnLoad;
 };
 
diff --git a/Engine/Source/Editor/BlueprintGraph/Classes/K2Node_FunctionTerminator.h b/Engine/Source/Editor/BlueprintGraph/Classes/K2Node_FunctionTerminator.h
--- a/Engine/Source/Editor/BlueprintGraph/Classes/K2Node_FunctionTerminator.h
+++ b/Engine/Source/Editor/BlueprintGraph/Classes/K2Node_FunctionTerminator.h
@@ -15,9 +15,7 @@ class UK2Node_FunctionTerminator : public UK2Node_EditablePinBase
 {
 	GENERATED_UCLASS_BODY()
 
-	/**
-	 * Reference to the function signature.
-	 */
+	/** Reference to the function signature. This is only resolvable by default if this is an inherited function */
 	UPROPERTY()
 	FMemberReference FunctionReference;
 
@@ -42,7 +40,10 @@ class UK2Node_FunctionTerminator : public UK2Node_EditablePinBase
 	//~ End UK2Node_EditablePinBase Interface
 
 	/** Promotes the node from being a part of an interface override to a full function that allows for parameter and result pin additions */
-	virtual void PromoteFromInterfaceOverride(bool bIsPrimaryTerminator = true);
+	BLUEPRINTGRAPH_API virtual void PromoteFromInterfaceOverride(bool bIsPrimaryTerminator = true);
+
+	/** Returns the UFunction that this node actually represents, this will work for both inherited and newly created functions */
+	BLUEPRINTGRAPH_API UFunction* FindSignatureFunction() const;
 
 private:
 	/** (DEPRECATED) Function signature class. Replaced by the 'FunctionReference' property. */
diff --git a/Engine/Source/Editor/BlueprintGraph/Classes/K2Node_Tunnel.h b/Engine/Source/Editor/BlueprintGraph/Classes/K2Node_Tunnel.h
--- a/Engine/Source/Editor/BlueprintGraph/Classes/K2Node_Tunnel.h
+++ b/Engine/Source/Editor/BlueprintGraph/Classes/K2Node_Tunnel.h
@@ -54,6 +54,7 @@ class BLUEPRINTGRAPH_API UK2Node_Tunnel : public UK2Node_EditablePinBase
 	virtual bool DrawNodeAsEntry() const override;
 	virtual bool DrawNodeAsExit() const override;
 	virtual bool NodeCausesStructuralBlueprintChange() const override { return true; }
+	virtual void ClearCachedBlueprintData(UBlueprint* Blueprint) override;
 	//~ End UK2Node Interface
 
 	//~ Begin UK2Node_EditablePinBase Interface.
diff --git a/Engine/Source/Editor/BlueprintGraph/Private/EdGraphSchema_K2.cpp b/Engine/Source/Editor/BlueprintGraph/Private/EdGraphSchema_K2.cpp
--- a/Engine/Source/Editor/BlueprintGraph/Private/EdGraphSchema_K2.cpp
+++ b/Engine/Source/Editor/BlueprintGraph/Private/EdGraphSchema_K2.cpp
@@ -1345,6 +1345,20 @@ bool UEdGraphSchema_K2::PinDefaultValueIsEditable(const UEdGraphPin& InGraphPin)
 	return true;
 }
 
+bool UEdGraphSchema_K2::PinHasCustomDefaultFormat(const UEdGraphPin& InGraphPin) const
+{
+	if (InGraphPin.PinType.PinCategory == PC_Struct)
+	{
+		// Some struct types have custom formats for default value for historical reasons
+		UObject const& SubCategoryObject = *InGraphPin.PinType.PinSubCategoryObject;
+		return &SubCategoryObject == VectorStruct
+			|| &SubCategoryObject == RotatorStruct
+			|| &SubCategoryObject == TransformStruct
+			|| &SubCategoryObject == LinearColorStruct;
+	}
+	return false;
+}
+
 void UEdGraphSchema_K2::SelectAllNodesInDirection(TEnumAsByte<enum EEdGraphPinDirection> InDirection, UEdGraph* Graph, UEdGraphPin* InGraphPin)
 {
 	/** Traverses the node graph out from the specified pin, logging each node that it visits along the way. */
diff --git a/Engine/Source/Editor/BlueprintGraph/Private/K2Node.cpp b/Engine/Source/Editor/BlueprintGraph/Private/K2Node.cpp
--- a/Engine/Source/Editor/BlueprintGraph/Private/K2Node.cpp
+++ b/Engine/Source/Editor/BlueprintGraph/Private/K2Node.cpp
@@ -76,48 +76,26 @@ void UK2Node::Serialize(FArchive& Ar)
 	Ar.UsingCustomVersion(FFrameworkObjectVersion::GUID);
 
 	if (Ar.IsSaving())
-	{	
-		for (UEdGraphPin* Pin : Pins)
+	{
+		if (Ar.IsObjectReferenceCollector() || Ar.Tell() < 0)
 		{
-			if (!Pin->bDefaultValueIsIgnored && !Pin->DefaultValue.IsEmpty() )
-			{
-				// If looking for references during save, expand any default values on the pins
-				// This is only reliable when saving in the editor, the cook case is handled below
-				if (Ar.IsObjectReferenceCollector() && Pin->PinType.PinCategory == UEdGraphSchema_K2::PC_Struct && Pin->PinType.PinSubCategoryObject.IsValid())
-				{
-					UScriptStruct* Struct = Cast<UScriptStruct>(Pin->PinType.PinSubCategoryObject.Get());
-
-					if (Struct)
-					{
-						TSharedPtr<FStructOnScope> StructData = MakeShareable(new FStructOnScope(Struct));
-
-						// Import the literal text to a dummy struct and then serialize that. Hard object references will not properly import, this is only useful for soft references!
-						FOutputDeviceNull NullOutput;
-						Struct->ImportText(*Pin->DefaultValue, StructData->GetStructMemory(), nullptr, PPF_SerializedAsImportText, &NullOutput, Pin->PinName.ToString());						
-						Struct->SerializeItem(Ar, StructData->GetStructMemory(), nullptr);
-					}
-				}
-
-				if (Pin->PinType.PinCategory == UEdGraphSchema_K2::PC_SoftObject || Pin->PinType.PinCategory == UEdGraphSchema_K2::PC_SoftClass)
-				{
-					FSoftObjectPath TempRef(Pin->DefaultValue);
-
-					// Serialize the asset reference, this will do the save fixup. It won't actually serialize the string if this is a real archive like linkersave
-					FSoftObjectPathSerializationScope DisableSerialize(NAME_None, NAME_None, ESoftObjectPathCollectType::AlwaysCollect, ESoftObjectPathSerializeType::SkipSerializeIfArchiveHasSize);
-					Ar << TempRef;
-
-					Pin->DefaultValue = TempRef.ToString();
-				}
-			}
+			// When this is a reference collector/modifier, serialize some pins as structs
+			FixupPinStringDataReferences(&Ar);
 		}
 	}
 
 	Super::Serialize(Ar);
 
-	if (Ar.IsLoading())
+	if (Ar.IsLoading() && (Ar.GetPortFlags() | PPF_Duplicate) == 0)
 	{
 		// Fix up pin default values, must be done before post load
 		FixupPinDefaultValues();
+
+		if (GIsEditor)
+		{
+			// We need to serialize string data references on load in editor builds so the cooker knows about them
+			FixupPinStringDataReferences(nullptr);
+		}
 	}
 }
 
@@ -161,33 +139,82 @@ void UK2Node::FixupPinDefaultValues()
 	}
 
 	// Fix soft object ptr pins
-	if (GIsEditor || LinkerFrameworkVersion < FFrameworkObjectVersion::ChangeAssetPinsToString)
+	if (LinkerFrameworkVersion < FFrameworkObjectVersion::ChangeAssetPinsToString)
 	{
-		FSoftObjectPathSerializationScope SetPackage(GetOutermost()->GetFName(), NAME_None, ESoftObjectPathCollectType::AlwaysCollect, ESoftObjectPathSerializeType::SkipSerializeIfArchiveHasSize);
 		for (int32 i = 0; i < Pins.Num(); ++i)
 		{
 			UEdGraphPin* Pin = Pins[i];
 			if (Pin->PinType.PinCategory == UEdGraphSchema_K2::PC_SoftObject || Pin->PinType.PinCategory == UEdGraphSchema_K2::PC_SoftClass)
 			{
 				// Fix old assetptr pins
-				if (LinkerFrameworkVersion < FFrameworkObjectVersion::ChangeAssetPinsToString)
+				if (Pin->DefaultObject && Pin->DefaultValue.IsEmpty())
 				{
-					if (Pin->DefaultObject && Pin->DefaultValue.IsEmpty())
+					Pin->DefaultValue = Pin->DefaultObject->GetPathName();
+					Pin->DefaultObject = nullptr;
+				}
+			}
+		}
+	}
+}
+
+void UK2Node::FixupPinStringDataReferences(FArchive* SavingArchive)
+{
+	FSoftObjectPathSerializationScope SetPackage(GetOutermost()->GetFName(), NAME_None, ESoftObjectPathCollectType::AlwaysCollect, ESoftObjectPathSerializeType::SkipSerializeIfArchiveHasSize);
+
+	// This code expans some pin types into their real representation, optionally serializes them, and then writes them out as strings again
+	const UEdGraphSchema_K2* K2Schema = GetDefault<UEdGraphSchema_K2>();
+
+	for (UEdGraphPin* Pin : Pins)
+	{
+		if (!Pin->bDefaultValueIsIgnored && !Pin->DefaultValue.IsEmpty())
+		{
+			// We skip this for structs like FVector that have a custom format for default values
+			if (Pin->PinType.PinCategory == UEdGraphSchema_K2::PC_Struct && Pin->PinType.PinSubCategoryObject.IsValid() && !K2Schema->PinHasCustomDefaultFormat(*Pin))
+			{
+				UScriptStruct* Struct = Cast<UScriptStruct>(Pin->PinType.PinSubCategoryObject.Get());
+				if (Struct)
+				{
+					TSharedPtr<FStructOnScope> StructData = MakeShareable(new FStructOnScope(Struct));
+					StructData->SetPackage(GetOutermost());
+
+					// Import the literal text to a dummy struct and then serialize that. Hard object references will not properly import, this is only useful for soft references!
+					FOutputDeviceNull NullOutput;
+					Struct->ImportText(*Pin->DefaultValue, StructData->GetStructMemory(), this, PPF_SerializedAsImportText, &NullOutput, Pin->PinName.ToString());
+
+					if (SavingArchive)
+					{
+						// If we're saving, use the archive to do any replacements
+						Struct->SerializeItem(*SavingArchive, StructData->GetStructMemory(), nullptr);
+					}
+					
+					// Convert back to the default value string as we might have changed
+					FString NewValue;
+					Struct->ExportText(NewValue, StructData->GetStructMemory(), nullptr, this, PPF_SerializedAsImportText, nullptr);
+
+					if (Pin->DefaultValue != NewValue)
 					{
-						Pin->DefaultValue = Pin->DefaultObject->GetPathName();
-						Pin->DefaultObject = nullptr;
+						Pin->DefaultValue = NewValue;
 					}
 				}
+			}
 
-				// In editor, fixup soft object ptrs on load on to handle redirects and finding refs for cooking
-				// We're not handling soft object ptrs inside FStructs because it's a rare edge case and would be a performance hit on load
-				if (GIsEditor && !Pin->DefaultValue.IsEmpty())
+			if (Pin->PinType.PinCategory == UEdGraphSchema_K2::PC_SoftObject || Pin->PinType.PinCategory == UEdGraphSchema_K2::PC_SoftClass)
+			{
+				FSoftObjectPath TempRef(Pin->DefaultValue);
+
+				if (SavingArchive)
+				{
+					// Serialize it directly, this won't do anything if it's a real archive like LinkerSave
+					*SavingArchive << TempRef;
+				}
+				else
 				{
-					FSoftObjectPath TempRef(Pin->DefaultValue);
+					// Transform it as strings
 					TempRef.PostLoadPath(GetLinker());
 					TempRef.PreSavePath();
-					Pin->DefaultValue = TempRef.ToString();
 				}
+				
+				Pin->DefaultValue = TempRef.ToString();
 			}
 		}
 	}
diff --git a/Engine/Source/Editor/BlueprintGraph/Private/K2Node_CustomEvent.cpp b/Engine/Source/Editor/BlueprintGraph/Private/K2Node_CustomEvent.cpp
--- a/Engine/Source/Editor/BlueprintGraph/Private/K2Node_CustomEvent.cpp
+++ b/Engine/Source/Editor/BlueprintGraph/Private/K2Node_CustomEvent.cpp
@@ -343,6 +343,31 @@ void UK2Node_CustomEvent::GetMenuActions(FBlueprintActionDatabaseRegistrar& Acti
 	}
 }
 
+void UK2Node_CustomEvent::FixupPinStringDataReferences(FArchive* SavingArchive)
+{
+	Super::FixupPinStringDataReferences(SavingArchive);
+	if (SavingArchive)
+	{ 
+		// If any of our pins got fixed up, we need to refresh our user pin default values
+		// For custom events, the Pin default values are authoritative
+		for (TSharedPtr<FUserPinInfo> PinInfo : UserDefinedPins)
+		{
+			if (UEdGraphPin* Pin = FindPin(PinInfo->PinName))
+			{
+				if (Pin->Direction == PinInfo->DesiredPinDirection)
+				{
+					FString DefaultsString = Pin->GetDefaultAsString();
+
+					if (DefaultsString != PinInfo->PinDefaultValue)
+					{
+						ModifyUserDefinedPinDefaultValue(PinInfo, Pin->GetDefaultAsString());
+					}
+				}
+			}
+		}
+	}
+}
+
 void UK2Node_CustomEvent::ReconstructNode()
 {
 	CachedNodeTitle.MarkDirty();
diff --git a/Engine/Source/Editor/BlueprintGraph/Private/K2Node_EditablePinBase.cpp b/Engine/Source/Editor/BlueprintGraph/Private/K2Node_EditablePinBase.cpp
--- a/Engine/Source/Editor/BlueprintGraph/Private/K2Node_EditablePinBase.cpp
+++ b/Engine/Source/Editor/BlueprintGraph/Private/K2Node_EditablePinBase.cpp
@@ -5,6 +5,7 @@
 #include "UObject/FrameworkObjectVersion.h"
 #include "Misc/FeedbackContext.h"
 #include "EdGraphSchema_K2.h"
+#include "Kismet2/BlueprintEditorUtils.h"
 #include "Kismet2/KismetDebugUtilities.h"
 
 // Ensure that the UserDefinedPin's "desired direction" matches the direction of 
@@ -256,6 +257,23 @@ void UK2Node_EditablePinBase::Serialize(FArchive& Ar)
 				}
 			}
 
+			// We need to fixup text default values because many of them got saved with the wrong package namespace
+			if (!PinInfo->PinDefaultValue.IsEmpty() && PinInfo->PinType.PinCategory == UEdGraphSchema_K2::PC_Text)
+			{
+				FString UseDefaultValue;
+				UObject* UseDefaultObject = nullptr;
+				FText UseDefaultText;
+
+				K2Schema->GetPinDefaultValuesFromString(PinInfo->PinType, this, PinInfo->PinDefaultValue, UseDefaultValue, UseDefaultObject, UseDefaultText, /*bPreserveTextIdentity*/false);
+
+				if (!UseDefaultText.IsEmpty())
+				{
+					// This will have the namespace changed if needed
+					PinInfo->PinDefaultValue.Reset();
+					FTextStringHelper::WriteToBuffer(PinInfo->PinDefaultValue, UseDefaultText);
+				}
+			}
+
 			UserDefinedPins.Add(PinInfo);
 		}
 	}
@@ -323,14 +341,25 @@ void UK2Node_EditablePinBase::PinDefaultValueChanged(UEdGraphPin* Pin)
 
 bool UK2Node_EditablePinBase::ModifyUserDefinedPinDefaultValue(TSharedPtr<FUserPinInfo> PinInfo, const FString& InDefaultValue)
 {
-	const UEdGraphSchema_K2* K2Schema = GetDefault<UEdGraphSchema_K2>();
 	FString NewDefaultValue = InDefaultValue;
 
+	if (!UpdateEdGraphPinDefaultValue(PinInfo, NewDefaultValue))
+	{
+		return false;
+	}
+
+	PinInfo->PinDefaultValue = NewDefaultValue;
+	return true;
+}
+
+bool UK2Node_EditablePinBase::UpdateEdGraphPinDefaultValue(TSharedPtr<FUserPinInfo> PinInfo, FString& NewDefaultValue)
+{
+	const UEdGraphSchema_K2* K2Schema = GetDefault<UEdGraphSchema_K2>();
 	// Find and modify the current pin
 	if (UEdGraphPin* OldPin = FindPin(PinInfo->PinName))
 	{
 		FString SavedDefaultValue = OldPin->DefaultValue;
-		
+
 		K2Schema->SetPinAutogeneratedDefaultValue(OldPin, NewDefaultValue);
 
 		// Validate the new default value
@@ -345,7 +374,6 @@ bool UK2Node_EditablePinBase::ModifyUserDefinedPinDefaultValue(TSharedPtr<FUserP
 		}
 	}
 
-	PinInfo->PinDefaultValue = NewDefaultValue;
 	return true;
 }
 
@@ -376,3 +404,152 @@ bool UK2Node_EditablePinBase::CreateUserDefinedPinsForFunctionEntryExit(const UF
 
 	return bAllPinsGood;
 }
+
+bool UK2Node_EditablePinBase::UpdateVariableStructFromDefaults(const UStruct* VariableStruct, uint8* VariableStructData, TArray<FBPVariableDescription>* LocalVariables /*= nullptr*/)
+{
+	const UEdGraphSchema_K2* K2Schema = GetDefault<UEdGraphSchema_K2>();
+
+	if (!VariableStruct || !VariableStructData)
+	{
+		return false;
+	}
+
+	/* This isn't safe because we have too many invalid pin default values
+	for (TSharedPtr<FUserPinInfo> PinInfo : UserDefinedPins)
+	{
+		if (PinInfo->DesiredPinDirection == EGPD_Output && !PinInfo->PinDefaultValue.IsEmpty())
+		{
+			UProperty* PinProperty = VariableStruct->FindPropertyByName(PinInfo->PinName);
+			if (PinProperty && (!PinProperty->HasAnyPropertyFlags(CPF_OutParm) || PinProperty->HasAnyPropertyFlags(CPF_ReferenceParm)))
+			{
+				FEdGraphPinType PinType;
+				K2Schema->ConvertPropertyToPinType(PinProperty, PinType);
+
+				if (PinType != PinInfo->PinType)
+				{
+					UE_LOG(LogBlueprint, Log, TEXT("Pin type for variable %s does not match type on struct %s during UpdateVariableStructFromDefaults, ignoring old default"), *PinInfo->PinName.ToString(), *VariableStruct->GetName());
+				}
+				else
+				{
+					FBlueprintEditorUtils::PropertyValueFromString(PinProperty, PinInfo->PinDefaultValue, VariableStructData);
+				}
+			}
+			else
+			{
+				UE_LOG(LogBlueprint, Log, TEXT("Could not find property %s on struct %s during UpdateVariableStructFromDefaults"), *PinInfo->PinName.ToString(), *VariableStruct->GetName());
+			}
+		}
+	} */
+
+	if (LocalVariables)
+	{
+		for (FBPVariableDescription& LocalVariable : *LocalVariables)
+		{
+			if (!LocalVariable.DefaultValue.IsEmpty())
+			{
+				UProperty* PinProperty = VariableStruct->FindPropertyByName(LocalVariable.VarName);
+
+				if (PinProperty && (!PinProperty->HasAnyPropertyFlags(CPF_OutParm) || PinProperty->HasAnyPropertyFlags(CPF_ReferenceParm)))
+				{
+					FEdGraphPinType PinType;
+					K2Schema->ConvertPropertyToPinType(PinProperty, /*out*/ PinType);
+
+					if (PinType != LocalVariable.VarType)
+					{
+						UE_LOG(LogBlueprint, Log, TEXT("Pin type for local variable %s does not match type on struct %s during UpdateVariableStructFromDefaults, ignoring old default"), *LocalVariable.VarName.ToString(), *VariableStruct->GetName());
+					}
+					else
+					{
+						FBlueprintEditorUtils::PropertyValueFromString(PinProperty, LocalVariable.DefaultValue, VariableStructData);
+					}
+				}
+				else
+				{
+					UE_LOG(LogBlueprint, Log, TEXT("Could not find local variable property %s on struct %s during UpdateVariableStructFromDefaults"), *LocalVariable.VarName.ToString(), *VariableStruct->GetName());
+				}
+			}
+		}
+	}
+
+	return true;
+}
+
+bool UK2Node_EditablePinBase::UpdateDefaultsFromVariableStruct(const UStruct* VariableStruct, uint8* VariableStructData, TArray<FBPVariableDescription>* LocalVariables /*= nullptr*/)
+{
+	const UEdGraphSchema_K2* K2Schema = GetDefault<UEdGraphSchema_K2>();
+
+	if (!VariableStruct || !VariableStructData)
+	{
+		return false;
+	}
+
+	/* This isn't safe because we have too many invalid pin default values
+	for (TSharedPtr<FUserPinInfo> PinInfo : UserDefinedPins)
+	{
+		if (PinInfo->DesiredPinDirection == EGPD_Output)
+		{
+			UProperty* PinProperty = VariableStruct->FindPropertyByName(PinInfo->PinName);
+			if (PinProperty && (!PinProperty->HasAnyPropertyFlags(CPF_OutParm) || PinProperty->HasAnyPropertyFlags(CPF_ReferenceParm)))
+			{
+				FEdGraphPinType PinType;
+				K2Schema->ConvertPropertyToPinType(PinProperty, PinType);
+
+				if (PinType != PinInfo->PinType)
+				{
+					UE_LOG(LogBlueprint, Log, TEXT("Pin type for variable %s does not match type on struct %s during UpdateDefaultsFromVariableStruct, ignoring old default"), *PinInfo->PinName.ToString(), *VariableStruct->GetName());
+				}
+				else
+				{
+					FString NewDefaultValue;
+					FBlueprintEditorUtils::PropertyValueToString(PinProperty, VariableStructData, NewDefaultValue);
+
+					if (NewDefaultValue != PinInfo->PinDefaultValue)
+					{
+						// Call the internal function to avoid recursion back into this update
+						if (UpdateEdGraphPinDefaultValue(PinInfo, NewDefaultValue))
+						{
+							PinInfo->PinDefaultValue = NewDefaultValue;
+						}
+						else
+						{
+							UE_LOG(LogBlueprint, Warning, TEXT("Failed to update variable %s on struct %s during UpdateDefaultsFromVariableStruct with new default %s"), *PinInfo->PinName.ToString(), *VariableStruct->GetName(), *NewDefaultValue);
+						}
+					}
+				}
+			}
+			else
+			{
+				UE_LOG(LogBlueprint, Log, TEXT("Could not find property %s on struct %s during UpdateDefaultsFromVariableStruct"), *PinInfo->PinName.ToString(), *VariableStruct->GetName());
+			}
+		}
+	}*/
+
+	if (LocalVariables)
+	{
+		for (FBPVariableDescription& LocalVariable : *LocalVariables)
+		{
+			UProperty* PinProperty = VariableStruct->FindPropertyByName(LocalVariable.VarName);
+
+			if (PinProperty && (!PinProperty->HasAnyPropertyFlags(CPF_OutParm) || PinProperty->HasAnyPropertyFlags(CPF_ReferenceParm)))
+			{
+				FEdGraphPinType PinType;
+				K2Schema->ConvertPropertyToPinType(PinProperty, /*out*/ PinType);
+
+				if (PinType != LocalVariable.VarType)
+				{
+					UE_LOG(LogBlueprint, Log, TEXT("Pin type for local variable %s does not match type on struct %s during UpdateDefaultsFromVariableStruct, ignoring old default"), *LocalVariable.VarName.ToString(), *VariableStruct->GetName());
+				}
+				else
+				{
+					FBlueprintEditorUtils::PropertyValueToString(PinProperty, VariableStructData, LocalVariable.DefaultValue);
+				}
+			}
+			else
+			{
+				UE_LOG(LogBlueprint, Log, TEXT("Could not find local variable property %s on struct %s during UpdateDefaultsFromVariableStruct"), *LocalVariable.VarName.ToString(), *VariableStruct->GetName());
+			}
+		}
+	}
+
+	return true;
+}
diff --git a/Engine/Source/Editor/BlueprintGraph/Private/K2Node_FunctionEntry.cpp b/Engine/Source/Editor/BlueprintGraph/Private/K2Node_FunctionEntry.cpp
--- a/Engine/Source/Editor/BlueprintGraph/Private/K2Node_FunctionEntry.cpp
+++ b/Engine/Source/Editor/BlueprintGraph/Private/K2Node_FunctionEntry.cpp
@@ -196,6 +196,7 @@ UK2Node_FunctionEntry::UK2Node_FunctionEntry(const FObjectInitializer& ObjectIni
 {
 	// Enforce const-correctness by default
 	bEnforceConstCorrectness = true;
+	bUpdatedDefaultValuesOnLoad = false;
 }
 
 void UK2Node_FunctionEntry::PreSave(const class ITargetPlatform* TargetPlatform)
@@ -203,38 +204,23 @@ void UK2Node_FunctionEntry::PreSave(const class ITargetPlatform* TargetPlatform)
 	Super::PreSave(TargetPlatform);
 	
 	const UBlueprint* Blueprint = HasValidBlueprint() ? GetBlueprint() : nullptr;
-	if (LocalVariables.Num() > 0 && Blueprint)
+	if (Blueprint && (LocalVariables.Num() > 0 || UserDefinedPins.Num() > 0))
 	{
-		// This code is here as it's unsafe to call when GIsSavingPackage is true
-		UFunction* Function = FindField<UFunction>(Blueprint->SkeletonGeneratedClass, *GetOuter()->GetName());
+		// Forcibly fixup defaults before we save
+		UpdateLoadedDefaultValues(true);
+	}
+}
 
-		if (Function)
-		{
-			if (Function->GetStructureSize() > 0 || !ensure(Function->PropertyLink == nullptr))
-			{
-				TSharedPtr<FStructOnScope> LocalVarData = MakeShareable(new FStructOnScope(Function));
+void UK2Node_FunctionEntry::PostLoad()
+{
+	Super::PostLoad();
 
-				for (TFieldIterator<UProperty> It(Function); It; ++It)
-				{
-					if (const UProperty* Property = *It)
-					{
-						const UStructProperty* PotentialUDSProperty = Cast<const UStructProperty>(Property);
-						// UDS requires default data even when the LocalVariable value is empty
+	if (GIsEditor)
+	{
+		// In the editor, we need to handle processing function default values at load time so they get picked up properly by the cooker
+		// This normally won't do anything because it gets called during the duplicate save during BP compilation, but if compilation gets skipped we need to make sure they get updated
 
-						for (FBPVariableDescription& LocalVar : LocalVariables)
-						{
-							if (LocalVar.VarName == Property->GetFName() && !LocalVar.DefaultValue.IsEmpty())
-							{
-								// Go to property and back, this handles redirector fixup and will sanitize the output
-								// The asset registry only knows about these references because when the node is expanded it turns into a hard reference
-								FBlueprintEditorUtils::PropertyValueFromString(Property, LocalVar.DefaultValue, LocalVarData->GetStructMemory());
-								FBlueprintEditorUtils::PropertyValueToString(Property, LocalVarData->GetStructMemory(), LocalVar.DefaultValue);
-							}
-						}
-					}
-				}
-			}
-		}
+		UpdateLoadedDefaultValues();
 	}
 }
 
@@ -245,38 +231,18 @@ void UK2Node_FunctionEntry::Serialize(FArchive& Ar)
 	Ar.UsingCustomVersion(FBlueprintsObjectVersion::GUID);
 
 	if (Ar.IsSaving())
-	{
-		for (FBPVariableDescription& LocalVariable : LocalVariables)
+	{		
+		if (Ar.IsObjectReferenceCollector() || Ar.Tell() < 0)
 		{
-			if (!LocalVariable.DefaultValue.IsEmpty())
+			// If this is explicitly a reference collector, or it's a save with no backing archive, then we want to use the function variable cache if it exists
+			// It's not safe to regenerate the cache at this point as we could be in GIsSaving
+			if (FunctionVariableCache.IsValid() && FunctionVariableCache->IsValid())
 			{
-				// If looking for references during save, expand any default values on the local variables
-				// This is only reliable when saving in the editor, the cook case is handled below
-				if (Ar.IsObjectReferenceCollector() && LocalVariable.VarType.PinCategory == UEdGraphSchema_K2::PC_Struct && LocalVariable.VarType.PinSubCategoryObject.IsValid())
-				{
-					UScriptStruct* Struct = Cast<UScriptStruct>(LocalVariable.VarType.PinSubCategoryObject.Get());
-
-					if (Struct)
-					{
-						TSharedPtr<FStructOnScope> StructData = MakeShareable(new FStructOnScope(Struct));
-
-						// Import the literal text to a dummy struct and then serialize that. Hard object references will not properly import, this is only useful for soft references!
-						FOutputDeviceNull NullOutput;
-						Struct->ImportText(*LocalVariable.DefaultValue, StructData->GetStructMemory(), nullptr, PPF_SerializedAsImportText, &NullOutput, LocalVariable.VarName.ToString());
-						Struct->SerializeItem(Ar, StructData->GetStructMemory(), nullptr);
-					}
-				}
+				UStruct* Struct = const_cast<UStruct*>(FunctionVariableCache->GetStruct());
+				Struct->SerializeBin(Ar, FunctionVariableCache->GetStructMemory());
 
-				if (LocalVariable.VarType.PinCategory == UEdGraphSchema_K2::PC_SoftObject || LocalVariable.VarType.PinCategory == UEdGraphSchema_K2::PC_SoftClass)
-				{
-					FSoftObjectPath TempRef(LocalVariable.DefaultValue);
-
-					// Serialize the asset reference, this will do the save fixup. It won't actually serialize the string if this is a real archive like linkersave
-					FSoftObjectPathSerializationScope DisableSerialize(NAME_None, NAME_None, ESoftObjectPathCollectType::AlwaysCollect, ESoftObjectPathSerializeType::SkipSerializeIfArchiveHasSize);
-					Ar << TempRef;
-
-					LocalVariable.DefaultValue = TempRef.ToString();
-				}
+				// Copy back into defaults as they may have changed
+				UpdateDefaultsFromVariableStruct(FunctionVariableCache->GetStruct(), FunctionVariableCache->GetStructMemory(), &LocalVariables);
 			}
 		}
 	}
@@ -339,24 +305,6 @@ void UK2Node_FunctionEntry::Serialize(FArchive& Ar)
 				}
 			}
 		}
-
-		// In editor, fixup soft object ptrs on load on to handle redirects and finding refs for cooking
-		// We're not handling soft object ptrs inside FStructs because it's a rare edge case and would be a performance hit on load
-		if (GIsEditor)
-		{
-			FSoftObjectPathSerializationScope SetPackage(GetOutermost()->GetFName(), NAME_None, ESoftObjectPathCollectType::AlwaysCollect, ESoftObjectPathSerializeType::SkipSerializeIfArchiveHasSize);
-
-			for (FBPVariableDescription& LocalVariable : LocalVariables)
-			{
-				if (!LocalVariable.DefaultValue.IsEmpty() && (LocalVariable.VarType.PinCategory == UEdGraphSchema_K2::PC_SoftObject || LocalVariable.VarType.PinCategory == UEdGraphSchema_K2::PC_SoftClass))
-				{
-					FSoftObjectPath TempRef(LocalVariable.DefaultValue);
-					TempRef.PostLoadPath(&Ar);
-					TempRef.PreSavePath();
-					LocalVariable.DefaultValue = TempRef.ToString();
-				}
-			}
-		}
 	}
 }
 
@@ -371,8 +319,12 @@ FText UK2Node_FunctionEntry::GetNodeTitle(ENodeTitleType::Type TitleType) const
 
 void UK2Node_FunctionEntry::AllocateDefaultPins()
 {
+	// Update our default values before copying them into pins
+	UpdateLoadedDefaultValues();
+
 	CreatePin(EGPD_Output, UEdGraphSchema_K2::PC_Exec, UEdGraphSchema_K2::PN_Then);
 
+	// Find any pins inherited from parent
 	if (UFunction* Function = FunctionReference.ResolveMember<UFunction>(GetBlueprintClassFromNode()))
 	{
 		CreatePinsForFunctionEntryExit(Function, /*bIsFunctionEntry=*/ true);
@@ -435,6 +387,67 @@ UEdGraphPin* UK2Node_FunctionEntry::CreatePinFromUserDefinition(const TSharedPtr
 	return NewPin;
 }
 
+TSharedPtr<FStructOnScope> UK2Node_FunctionEntry::GetFunctionVariableCache(bool bForceRefresh)
+{
+	if (bForceRefresh && FunctionVariableCache.IsValid())
+	{
+		// On force refresh, delete old one if it exists
+		FunctionVariableCache.Reset();
+	}
+
+	if (!FunctionVariableCache.IsValid())
+	{
+		if (UFunction* const Function = FindSignatureFunction())
+		{
+			if (LocalVariables.Num() > 0 || UserDefinedPins.Num() > 0)
+			{
+				FunctionVariableCache = MakeShared<FStructOnScope>(Function);
+				FunctionVariableCache->SetPackage(GetOutermost());
+
+				RefreshFunctionVariableCache();
+			}
+		}
+	}
+	
+	return FunctionVariableCache;
+}
+
+bool UK2Node_FunctionEntry::RefreshFunctionVariableCache()
+{
+	GetFunctionVariableCache(false);
+
+	if (FunctionVariableCache.IsValid())
+	{
+		// Update the cache if it was created
+		return UpdateVariableStructFromDefaults(FunctionVariableCache->GetStruct(), FunctionVariableCache->GetStructMemory(), &LocalVariables);
+	}
+	return false;
+}
+
+bool UK2Node_FunctionEntry::UpdateLoadedDefaultValues(bool bForceRefresh)
+{
+	// If we don't have a cache or it's force refresh, create one
+	if (!bUpdatedDefaultValuesOnLoad || bForceRefresh)
+	{
+		GetFunctionVariableCache(bForceRefresh);
+
+		bUpdatedDefaultValuesOnLoad = true;
+
+		if (FunctionVariableCache.IsValid())
+		{
+			// Now copy back into the default value strings
+			return UpdateDefaultsFromVariableStruct(FunctionVariableCache->GetStruct(), FunctionVariableCache->GetStructMemory(), &LocalVariables);
+		}
+		else
+		{
+			// No variable cache created
+			return true;
+		}
+	}
+
+	return false;
+}
+
 FNodeHandlingFunctor* UK2Node_FunctionEntry::CreateNodeHandler(FKismetCompilerContext& CompilerContext) const
 {
 	return new FKCHandler_FunctionEntry(CompilerContext);
@@ -462,6 +475,7 @@ void UK2Node_FunctionEntry::GetRedirectPinNames(const UEdGraphPin& Pin, TArray<F
 
 bool UK2Node_FunctionEntry::IsDeprecated() const
 {
+	// We only show deprecated for inherited functions
 	if (UFunction* const Function = FunctionReference.ResolveMember<UFunction>(GetBlueprintClassFromNode()))
 	{
 		return Function->HasMetaData(FBlueprintMetadata::MD_DeprecatedFunction);
@@ -491,7 +505,7 @@ FString UK2Node_FunctionEntry::GetDeprecationMessage() const
 
 FText UK2Node_FunctionEntry::GetTooltipText() const
 {
-	if (UFunction* const Function = FunctionReference.ResolveMember<UFunction>(GetBlueprintClassFromNode()))
+	if (UFunction* const Function = FindSignatureFunction())
 	{
 		return FText::FromString(UK2Node_CallFunction::GetDefaultTooltipForFunction(Function));
 	}
@@ -693,6 +707,36 @@ void UK2Node_FunctionEntry::PostReconstructNode()
 	Super::PostReconstructNode();
 }
 
+void UK2Node_FunctionEntry::ClearCachedBlueprintData(UBlueprint* Blueprint)
+{
+	FunctionVariableCache.Reset();
+}
+
+void UK2Node_FunctionEntry::FixupPinStringDataReferences(FArchive* SavingArchive)
+{
+	Super::FixupPinStringDataReferences(SavingArchive);
+	if (SavingArchive)
+	{
+		// If any of our pins got fixed up, we need to refresh our user pin default values
+		// For custom events, the Pin default values are authoritative
+		for (TSharedPtr<FUserPinInfo> PinInfo : UserDefinedPins)
+		{
+			if (UEdGraphPin* Pin = FindPin(PinInfo->PinName))
+			{
+				if (Pin->Direction == PinInfo->DesiredPinDirection)
+				{
+					FString DefaultsString = Pin->GetDefaultAsString();
+
+					if (DefaultsString != PinInfo->PinDefaultValue)
+					{
+						ModifyUserDefinedPinDefaultValue(PinInfo, Pin->GetDefaultAsString());
+					}
+				}
+			}
+		}
+	}
+}
+
 bool UK2Node_FunctionEntry::ModifyUserDefinedPinDefaultValue(TSharedPtr<FUserPinInfo> PinInfo, const FString& NewDefaultValue)
 {
 	if (Super::ModifyUserDefinedPinDefaultValue(PinInfo, NewDefaultValue))
@@ -700,6 +744,8 @@ bool UK2Node_FunctionEntry::ModifyUserDefinedPinDefaultValue(TSharedPtr<FUserPin
 		const UEdGraphSchema_K2* K2Schema = GetDefault<UEdGraphSchema_K2>();
 		K2Schema->HandleParameterDefaultValueChanged(this);
 
+		RefreshFunctionVariableCache();
+
 		return true;
 	}
 	return false;
diff --git a/Engine/Source/Editor/BlueprintGraph/Private/K2Node_FunctionTerminator.cpp b/Engine/Source/Editor/BlueprintGraph/Private/K2Node_FunctionTerminator.cpp
--- a/Engine/Source/Editor/BlueprintGraph/Private/K2Node_FunctionTerminator.cpp
+++ b/Engine/Source/Editor/BlueprintGraph/Private/K2Node_FunctionTerminator.cpp
@@ -121,6 +121,20 @@ void UK2Node_FunctionTerminator::PromoteFromInterfaceOverride(bool bIsPrimaryTer
 	Schema->ReconstructNode(*this, true);
 }
 
+UFunction* UK2Node_FunctionTerminator::FindSignatureFunction() const
+{
+	UClass* FoundClass = GetBlueprintClassFromNode();
+	UFunction* FoundFunction = FunctionReference.ResolveMember<UFunction>(FoundClass);
+
+	if (!FoundFunction && FoundClass && GetOuter())
+	{
+		// The resolve will fail if this is a locally-created function, so search using the event graph name
+		FoundFunction = FindField<UFunction>(FoundClass, *GetOuter()->GetName());
+	}
+
+	return FoundFunction;
+}
+
 void UK2Node_FunctionTerminator::ValidateNodeDuringCompilation(FCompilerResultsLog& MessageLog) const
 {
 	Super::ValidateNodeDuringCompilation(MessageLog);
diff --git a/Engine/Source/Editor/BlueprintGraph/Private/K2Node_Tunnel.cpp b/Engine/Source/Editor/BlueprintGraph/Private/K2Node_Tunnel.cpp
--- a/Engine/Source/Editor/BlueprintGraph/Private/K2Node_Tunnel.cpp
+++ b/Engine/Source/Editor/BlueprintGraph/Private/K2Node_Tunnel.cpp
@@ -177,6 +177,12 @@ bool UK2Node_Tunnel::CanCreateUserDefinedPin(const FEdGraphPinType& InPinType, E
 	return bResult;
 }
 
+void UK2Node_Tunnel::ClearCachedBlueprintData(UBlueprint* Blueprint)
+{
+	// Remove data marking graphs as latent, this will be re-cache'd as needed
+	MetaData.HasLatentFunctions = INDEX_NONE;
+}
+
 UEdGraphPin* UK2Node_Tunnel::CreatePinFromUserDefinition(const TSharedPtr<FUserPinInfo> NewPinInfo)
 {
 	// Create the new pin
diff --git a/Engine/Source/Editor/Kismet/Private/BlueprintDetailsCustomization.cpp b/Engine/Source/Editor/Kismet/Private/BlueprintDetailsCustomization.cpp
--- a/Engine/Source/Editor/Kismet/Private/BlueprintDetailsCustomization.cpp
+++ b/Engine/Source/Editor/Kismet/Private/BlueprintDetailsCustomization.cpp
@@ -830,7 +830,6 @@ void FBlueprintVarActionDetails::CustomizeDetails( IDetailLayoutBuilder& DetailL
 				UFunction* StructScope = Cast<UFunction>(VariableProperty->GetOuter());
 				check(StructScope);
 
-				TSharedPtr<FStructOnScope> StructData = MakeShareable(new FStructOnScope((UFunction*)StructScope));
 				UEdGraph* Graph = FBlueprintEditorUtils::FindScopeGraph(GetBlueprintObj(), (UFunction*)StructScope);
 
 				// Find the function entry nodes in the current graph
@@ -843,18 +842,9 @@ void FBlueprintVarActionDetails::CustomizeDetails( IDetailLayoutBuilder& DetailL
 				const UStructProperty* PotentialUDSProperty = Cast<const UStructProperty>(VariableProperty);
 				
 				UK2Node_FunctionEntry* FuncEntry = EntryNodes[0];
-				for (const FBPVariableDescription& LocalVar : FuncEntry->LocalVariables)
-				{
-					if(LocalVar.VarName == VariableProperty->GetFName()) //Property->GetFName())
-					{
-						// Only set the default value if there is one
-						if(!LocalVar.DefaultValue.IsEmpty())
-						{
-							FBlueprintEditorUtils::PropertyValueFromString(VariableProperty, LocalVar.DefaultValue, StructData->GetStructMemory());
-						}
-						break;
-					}
-				}
+
+				// This uses the cached struct data inside the node, which is updated on change
+				TSharedPtr<FStructOnScope> StructData = FuncEntry->GetFunctionVariableCache();
 
 				if(BlueprintEditor.IsValid())
 				{
@@ -2651,6 +2641,7 @@ void FBlueprintVarActionDetails::OnFinishedChangingProperties(const FPropertyCha
 					FuncEntry->Modify();
 					GetBlueprintObj()->Modify();
 					LocalVar.DefaultValue = DefaultValueString;
+					FuncEntry->RefreshFunctionVariableCache();
 					FBlueprintEditorUtils::MarkBlueprintAsModified(GetBlueprintObj());
 					break;
 				}
diff --git a/Engine/Source/Editor/UnrealEd/Private/Kismet2/BlueprintEditorUtils.cpp b/Engine/Source/Editor/UnrealEd/Private/Kismet2/BlueprintEditorUtils.cpp
--- a/Engine/Source/Editor/UnrealEd/Private/Kismet2/BlueprintEditorUtils.cpp
+++ b/Engine/Source/Editor/UnrealEd/Private/Kismet2/BlueprintEditorUtils.cpp
@@ -2123,15 +2123,17 @@ void FBlueprintEditorUtils::MarkBlueprintAsModified(UBlueprint* Blueprint, FProp
 	Blueprint->bCachedDependenciesUpToDate = false;
 	if (Blueprint->Status != BS_BeingCreated)
 	{
+		// This clears any cached data, which includes the macro tunnel node data
 		TArray<UEdGraph*> AllGraphs;
 		Blueprint->GetAllGraphs(AllGraphs);
-		for (int32 i = 0; i < AllGraphs.Num(); i++)
+		for (UEdGraph* GraphToClear : AllGraphs)
 		{
-			UK2Node_EditablePinBase* EntryNode = GetEntryNode(AllGraphs[i]);
-			if(UK2Node_Tunnel* TunnelNode = ExactCast<UK2Node_Tunnel>(EntryNode))
+			for (UEdGraphNode* Node : GraphToClear->Nodes)
 			{
-				// Remove data marking graphs as latent, this will be re-cache'd as needed
-				TunnelNode->MetaData.HasLatentFunctions = INDEX_NONE;
+				if (UK2Node* BPNode = Cast<UK2Node>(Node))
+				{
+					BPNode->ClearCachedBlueprintData(Blueprint);
+				}
 			}
 		}
 
