CIS fix

diff --git a/Engine/Source/Editor/BlueprintGraph/Classes/K2Node_EditablePinBase.h b/Engine/Source/Editor/BlueprintGraph/Classes/K2Node_EditablePinBase.h
--- a/Engine/Source/Editor/BlueprintGraph/Classes/K2Node_EditablePinBase.h
+++ b/Engine/Source/Editor/BlueprintGraph/Classes/K2Node_EditablePinBase.h
@@ -167,7 +167,7 @@ class UK2Node_EditablePinBase : public UK2Node
 	 * @param	VariableStructData	Raw data for a variable cache
 	 * @param	LocalVariables		Optional list of local variables to copy
 	 */
-	BLUEPRINTGRAPH_API bool UpdateVariableStructFromDefaults(const UStruct* VariableStruct, uint8* VariableStructData, TArray<FBPVariableDescription>* LocalVariables = nullptr);
+	BLUEPRINTGRAPH_API bool UpdateVariableStructFromDefaults(const UStruct* VariableStruct, uint8* VariableStructData, TArray<struct FBPVariableDescription>* LocalVariables = nullptr);
 
 	/**
 	 * Copies data from VariableStruct into the user defined pins as well as an optional list of local variables
@@ -176,7 +176,7 @@ class UK2Node_EditablePinBase : public UK2Node
 	 * @param	VariableStructData	Raw data for a variable cache
 	 * @param	LocalVariables		Optional list of local variables to copy
 	 */
-	BLUEPRINTGRAPH_API bool UpdateDefaultsFromVariableStruct(const UStruct* VariableStruct, uint8* VariableStructData, TArray<FBPVariableDescription>* LocalVariables = nullptr);
+	BLUEPRINTGRAPH_API bool UpdateDefaultsFromVariableStruct(const UStruct* VariableStruct, uint8* VariableStructData, TArray<struct FBPVariableDescription>* LocalVariables = nullptr);
 
 	/**
 	 * Can this node have execution wires added or removed?
