Fix order of operation warning
#rnx

diff --git a/Engine/Source/Editor/BlueprintGraph/Private/K2Node.cpp b/Engine/Source/Editor/BlueprintGraph/Private/K2Node.cpp
--- a/Engine/Source/Editor/BlueprintGraph/Private/K2Node.cpp
+++ b/Engine/Source/Editor/BlueprintGraph/Private/K2Node.cpp
@@ -86,7 +86,7 @@ void UK2Node::Serialize(FArchive& Ar)
 
 	Super::Serialize(Ar);
 
-	if (Ar.IsLoading() && (Ar.GetPortFlags() | PPF_Duplicate) == 0)
+	if (Ar.IsLoading() && ((Ar.GetPortFlags() | PPF_Duplicate) == 0))
 	{
 		// Fix up pin default values, must be done before post load
 		FixupPinDefaultValues();
