Pin fixup was being skipped, and was using the wrong parameters for export text
Added object owner parameter to BlueprintEditorUtils functions so they transfer text properly, deleted manual fixup
Change local variable diffing to use new cache, so text data is diffed properly

diff --git a/Engine/Source/Editor/AnimGraph/Private/AnimBlueprintNodeOptionalPinManager.cpp b/Engine/Source/Editor/AnimGraph/Private/AnimBlueprintNodeOptionalPinManager.cpp
--- a/Engine/Source/Editor/AnimGraph/Private/AnimBlueprintNodeOptionalPinManager.cpp
+++ b/Engine/Source/Editor/AnimGraph/Private/AnimBlueprintNodeOptionalPinManager.cpp
@@ -115,7 +115,7 @@ void FAnimBlueprintNodeOptionalPinManager::PostRemovedOldPin(FOptionalPinFromPro
 		if (UEdGraphPin* OldPin = OldPinMap.FindRef(OldPinName))
 		{
 			// Convert DefaultValue/DefaultValueObject and push back into the struct
-			FBlueprintEditorUtils::PropertyValueFromString_Direct(Property, OldPin->GetDefaultAsString(), PropertyAddress);
+			FBlueprintEditorUtils::PropertyValueFromString_Direct(Property, OldPin->GetDefaultAsString(), PropertyAddress, OldPin->GetOwningNodeUnchecked());
 		}
 	}
 }
diff --git a/Engine/Source/Editor/AnimGraph/Private/AnimGraphNode_CustomProperty.cpp b/Engine/Source/Editor/AnimGraph/Private/AnimGraphNode_CustomProperty.cpp
--- a/Engine/Source/Editor/AnimGraph/Private/AnimGraphNode_CustomProperty.cpp
+++ b/Engine/Source/Editor/AnimGraph/Private/AnimGraphNode_CustomProperty.cpp
@@ -75,7 +75,7 @@ void UAnimGraphNode_CustomProperty::ReallocatePinsDuringReconstruction(TArray<UE
 			FString CDODefaultValueString;
 			uint8* ContainerPtr = reinterpret_cast<uint8*>(TargetClass->GetDefaultObject());
 
-			if(FBlueprintEditorUtils::PropertyValueToString(Property, ContainerPtr, CDODefaultValueString))
+			if(FBlueprintEditorUtils::PropertyValueToString(Property, ContainerPtr, CDODefaultValueString, this))
 			{
 				// If we successfully pulled a value, set it to the pin
 				Schema->TrySetDefaultValue(*NewPin, CDODefaultValueString);
diff --git a/Engine/Source/Editor/BlueprintGraph/Classes/K2Node_EditablePinBase.h b/Engine/Source/Editor/BlueprintGraph/Classes/K2Node_EditablePinBase.h
--- a/Engine/Source/Editor/BlueprintGraph/Classes/K2Node_EditablePinBase.h
+++ b/Engine/Source/Editor/BlueprintGraph/Classes/K2Node_EditablePinBase.h
@@ -160,24 +160,6 @@ class UK2Node_EditablePinBase : public UK2Node
 	 */
 	BLUEPRINTGRAPH_API virtual bool CreateUserDefinedPinsForFunctionEntryExit(const UFunction* Function, bool bForFunctionEntry);
 
-	/**
-	 * Copies data from any user pins matching properties in VariableStruct into the VariableStructData. Will also copy LocalVariables if passed in
-	 *
-	 * @param	VariableStruct		UStruct defining the variables to copy
-	 * @param	VariableStructData	Raw data for a variable cache
-	 * @param	LocalVariables		Optional list of local variables to copy
-	 */
-	BLUEPRINTGRAPH_API bool UpdateVariableStructFromDefaults(const UStruct* VariableStruct, uint8* VariableStructData, TArray<struct FBPVariableDescription>* LocalVariables = nullptr);
-
-	/**
-	 * Copies data from VariableStruct into the user defined pins as well as an optional list of local variables
-	 *
-	 * @param	VariableStruct		UStruct defining the variables to copy
-	 * @param	VariableStructData	Raw data for a variable cache
-	 * @param	LocalVariables		Optional list of local variables to copy
-	 */
-	BLUEPRINTGRAPH_API bool UpdateDefaultsFromVariableStruct(const UStruct* VariableStruct, uint8* VariableStructData, TArray<struct FBPVariableDescription>* LocalVariables = nullptr);
-
 	/**
 	 * Can this node have execution wires added or removed?
 	 */
diff --git a/Engine/Source/Editor/BlueprintGraph/Classes/K2Node_FunctionEntry.h b/Engine/Source/Editor/BlueprintGraph/Classes/K2Node_FunctionEntry.h
--- a/Engine/Source/Editor/BlueprintGraph/Classes/K2Node_FunctionEntry.h
+++ b/Engine/Source/Editor/BlueprintGraph/Classes/K2Node_FunctionEntry.h
@@ -70,10 +70,10 @@ class UK2Node_FunctionEntry : public UK2Node_FunctionTerminator
 	virtual UEdGraphPin* CreatePinFromUserDefinition(const TSharedPtr<FUserPinInfo> NewPinInfo) override;
 	//~ End K2Node_FunctionTerminator Interface
 
-	/** Gets the UFunction and function variable cache structure that should be used for serialization fixups for function variables. If bForceRefresh is true it will always recreate the cache */
+	/** Gets the UFunction and function variable cache structure that should be used for serialization fixups for local variables. If bForceRefresh is true it will always recreate the cache */
 	BLUEPRINTGRAPH_API TSharedPtr<FStructOnScope> GetFunctionVariableCache(bool bForceRefresh = false);
 
-	/** Copies data from the local/input variable defaults into the variable cache */
+	/** Copies data from the local variable defaults into the variable cache */
 	BLUEPRINTGRAPH_API bool RefreshFunctionVariableCache();
 
 	/** Handles updating loaded default values, by going default string into variable cache and back, if bForceRefresh it will happen even if the cache is already setup */
@@ -111,6 +111,12 @@ class UK2Node_FunctionEntry : public UK2Node_FunctionTerminator
 	}
 
 protected:
+	/** Copies data from any local variables matching properties in VariableStruct into the VariableStructData */
+	BLUEPRINTGRAPH_API bool UpdateVariableStructFromDefaults(const UStruct* VariableStruct, uint8* VariableStructData);
+
+	/** Copies data from VariableStruct into the local variables */
+	BLUEPRINTGRAPH_API bool UpdateDefaultsFromVariableStruct(const UStruct* VariableStruct, uint8* VariableStructData);
+
 	/** Any extra flags that the function may need */
 	UPROPERTY()
 	int32 ExtraFlags;
@@ -119,7 +125,6 @@ protected:
 	TSharedPtr<FStructOnScope> FunctionVariableCache;
 
 	/** True if we've updated the default values on this node at least once */
-	UPROPERTY(Transient)
 	bool bUpdatedDefaultValuesOnLoad;
 };
 
diff --git a/Engine/Source/Editor/BlueprintGraph/Private/K2Node.cpp b/Engine/Source/Editor/BlueprintGraph/Private/K2Node.cpp
--- a/Engine/Source/Editor/BlueprintGraph/Private/K2Node.cpp
+++ b/Engine/Source/Editor/BlueprintGraph/Private/K2Node.cpp
@@ -86,7 +86,7 @@ void UK2Node::Serialize(FArchive& Ar)
 
 	Super::Serialize(Ar);
 
-	if (Ar.IsLoading() && ((Ar.GetPortFlags() | PPF_Duplicate) == 0))
+	if (Ar.IsLoading() && ((Ar.GetPortFlags() & PPF_Duplicate) == 0))
 	{
 		// Fix up pin default values, must be done before post load
 		FixupPinDefaultValues();
@@ -189,7 +189,7 @@ void UK2Node::FixupPinStringDataReferences(FArchive* SavingArchive)
 					
 					// Convert back to the default value string as we might have changed
 					FString NewValue;
-					Struct->ExportText(NewValue, StructData->GetStructMemory(), nullptr, this, PPF_SerializedAsImportText, nullptr);
+					Struct->ExportText(NewValue, StructData->GetStructMemory(), StructData->GetStructMemory(), this, PPF_SerializedAsImportText, nullptr);
 
 					if (Pin->DefaultValue != NewValue)
 					{
diff --git a/Engine/Source/Editor/BlueprintGraph/Private/K2Node_AddComponent.cpp b/Engine/Source/Editor/BlueprintGraph/Private/K2Node_AddComponent.cpp
--- a/Engine/Source/Editor/BlueprintGraph/Private/K2Node_AddComponent.cpp
+++ b/Engine/Source/Editor/BlueprintGraph/Private/K2Node_AddComponent.cpp
@@ -116,7 +116,7 @@ void UK2Node_AddComponent::AllocatePinsForExposedVariables()
 					if ((ClassDefaultObject != nullptr) && K2Schema->PinDefaultValueIsEditable(*Pin))
 					{
 						FString DefaultValueAsString;
-						const bool bDefaultValueSet = FBlueprintEditorUtils::PropertyValueToString(Property, reinterpret_cast<const uint8*>(ClassDefaultObject), DefaultValueAsString);
+						const bool bDefaultValueSet = FBlueprintEditorUtils::PropertyValueToString(Property, reinterpret_cast<const uint8*>(ClassDefaultObject), DefaultValueAsString, this);
 						check(bDefaultValueSet);
 						K2Schema->SetPinAutogeneratedDefaultValue(Pin, DefaultValueAsString);
 					}
diff --git a/Engine/Source/Editor/BlueprintGraph/Private/K2Node_ConstructObjectFromClass.cpp b/Engine/Source/Editor/BlueprintGraph/Private/K2Node_ConstructObjectFromClass.cpp
--- a/Engine/Source/Editor/BlueprintGraph/Private/K2Node_ConstructObjectFromClass.cpp
+++ b/Engine/Source/Editor/BlueprintGraph/Private/K2Node_ConstructObjectFromClass.cpp
@@ -122,7 +122,7 @@ void UK2Node_ConstructObjectFromClass::CreatePinsForClass(UClass* InClass, TArra
 				if (ClassDefaultObject && K2Schema->PinDefaultValueIsEditable(*Pin))
 				{
 					FString DefaultValueAsString;
-					const bool bDefaultValueSet = FBlueprintEditorUtils::PropertyValueToString(Property, reinterpret_cast<const uint8*>(ClassDefaultObject), DefaultValueAsString);
+					const bool bDefaultValueSet = FBlueprintEditorUtils::PropertyValueToString(Property, reinterpret_cast<const uint8*>(ClassDefaultObject), DefaultValueAsString, this);
 					check(bDefaultValueSet);
 					K2Schema->SetPinAutogeneratedDefaultValue(Pin, DefaultValueAsString);
 				}
diff --git a/Engine/Source/Editor/BlueprintGraph/Private/K2Node_CustomEvent.cpp b/Engine/Source/Editor/BlueprintGraph/Private/K2Node_CustomEvent.cpp
--- a/Engine/Source/Editor/BlueprintGraph/Private/K2Node_CustomEvent.cpp
+++ b/Engine/Source/Editor/BlueprintGraph/Private/K2Node_CustomEvent.cpp
@@ -360,7 +360,7 @@ void UK2Node_CustomEvent::FixupPinStringDataReferences(FArchive* SavingArchive)
 
 					if (DefaultsString != PinInfo->PinDefaultValue)
 					{
-						ModifyUserDefinedPinDefaultValue(PinInfo, Pin->GetDefaultAsString());
+						ModifyUserDefinedPinDefaultValue(PinInfo, DefaultsString);
 					}
 				}
 			}
diff --git a/Engine/Source/Editor/BlueprintGraph/Private/K2Node_EditablePinBase.cpp b/Engine/Source/Editor/BlueprintGraph/Private/K2Node_EditablePinBase.cpp
--- a/Engine/Source/Editor/BlueprintGraph/Private/K2Node_EditablePinBase.cpp
+++ b/Engine/Source/Editor/BlueprintGraph/Private/K2Node_EditablePinBase.cpp
@@ -258,13 +258,13 @@ void UK2Node_EditablePinBase::Serialize(FArchive& Ar)
 			}
 
 			// We need to fixup text default values because many of them got saved with the wrong package namespace
-			if (!PinInfo->PinDefaultValue.IsEmpty() && PinInfo->PinType.PinCategory == UEdGraphSchema_K2::PC_Text)
+			/*if (!PinInfo->PinDefaultValue.IsEmpty() && PinInfo->PinType.PinCategory == UEdGraphSchema_K2::PC_Text)
 			{
 				FString UseDefaultValue;
 				UObject* UseDefaultObject = nullptr;
 				FText UseDefaultText;
 
-				K2Schema->GetPinDefaultValuesFromString(PinInfo->PinType, this, PinInfo->PinDefaultValue, UseDefaultValue, UseDefaultObject, UseDefaultText, /*bPreserveTextIdentity*/false);
+				K2Schema->GetPinDefaultValuesFromString(PinInfo->PinType, this, PinInfo->PinDefaultValue, UseDefaultValue, UseDefaultObject, UseDefaultText,false);
 
 				if (!UseDefaultText.IsEmpty())
 				{
@@ -272,7 +272,7 @@ void UK2Node_EditablePinBase::Serialize(FArchive& Ar)
 					PinInfo->PinDefaultValue.Reset();
 					FTextStringHelper::WriteToBuffer(PinInfo->PinDefaultValue, UseDefaultText);
 				}
-			}
+			}*/
 
 			UserDefinedPins.Add(PinInfo);
 		}
@@ -404,152 +404,3 @@ bool UK2Node_EditablePinBase::CreateUserDefinedPinsForFunctionEntryExit(const UF
 
 	return bAllPinsGood;
 }
-
-bool UK2Node_EditablePinBase::UpdateVariableStructFromDefaults(const UStruct* VariableStruct, uint8* VariableStructData, TArray<FBPVariableDescription>* LocalVariables /*= nullptr*/)
-{
-	const UEdGraphSchema_K2* K2Schema = GetDefault<UEdGraphSchema_K2>();
-
-	if (!VariableStruct || !VariableStructData)
-	{
-		return false;
-	}
-
-	/* This isn't safe because we have too many invalid pin default values
-	for (TSharedPtr<FUserPinInfo> PinInfo : UserDefinedPins)
-	{
-		if (PinInfo->DesiredPinDirection == EGPD_Output && !PinInfo->PinDefaultValue.IsEmpty())
-		{
-			UProperty* PinProperty = VariableStruct->FindPropertyByName(PinInfo->PinName);
-			if (PinProperty && (!PinProperty->HasAnyPropertyFlags(CPF_OutParm) || PinProperty->HasAnyPropertyFlags(CPF_ReferenceParm)))
-			{
-				FEdGraphPinType PinType;
-				K2Schema->ConvertPropertyToPinType(PinProperty, PinType);
-
-				if (PinType != PinInfo->PinType)
-				{
-					UE_LOG(LogBlueprint, Log, TEXT("Pin type for variable %s does not match type on struct %s during UpdateVariableStructFromDefaults, ignoring old default"), *PinInfo->PinName.ToString(), *VariableStruct->GetName());
-				}
-				else
-				{
-					FBlueprintEditorUtils::PropertyValueFromString(PinProperty, PinInfo->PinDefaultValue, VariableStructData);
-				}
-			}
-			else
-			{
-				UE_LOG(LogBlueprint, Log, TEXT("Could not find property %s on struct %s during UpdateVariableStructFromDefaults"), *PinInfo->PinName.ToString(), *VariableStruct->GetName());
-			}
-		}
-	} */
-
-	if (LocalVariables)
-	{
-		for (FBPVariableDescription& LocalVariable : *LocalVariables)
-		{
-			if (!LocalVariable.DefaultValue.IsEmpty())
-			{
-				UProperty* PinProperty = VariableStruct->FindPropertyByName(LocalVariable.VarName);
-
-				if (PinProperty && (!PinProperty->HasAnyPropertyFlags(CPF_OutParm) || PinProperty->HasAnyPropertyFlags(CPF_ReferenceParm)))
-				{
-					FEdGraphPinType PinType;
-					K2Schema->ConvertPropertyToPinType(PinProperty, /*out*/ PinType);
-
-					if (PinType != LocalVariable.VarType)
-					{
-						UE_LOG(LogBlueprint, Log, TEXT("Pin type for local variable %s does not match type on struct %s during UpdateVariableStructFromDefaults, ignoring old default"), *LocalVariable.VarName.ToString(), *VariableStruct->GetName());
-					}
-					else
-					{
-						FBlueprintEditorUtils::PropertyValueFromString(PinProperty, LocalVariable.DefaultValue, VariableStructData);
-					}
-				}
-				else
-				{
-					UE_LOG(LogBlueprint, Log, TEXT("Could not find local variable property %s on struct %s during UpdateVariableStructFromDefaults"), *LocalVariable.VarName.ToString(), *VariableStruct->GetName());
-				}
-			}
-		}
-	}
-
-	return true;
-}
-
-bool UK2Node_EditablePinBase::UpdateDefaultsFromVariableStruct(const UStruct* VariableStruct, uint8* VariableStructData, TArray<FBPVariableDescription>* LocalVariables /*= nullptr*/)
-{
-	const UEdGraphSchema_K2* K2Schema = GetDefault<UEdGraphSchema_K2>();
-
-	if (!VariableStruct || !VariableStructData)
-	{
-		return false;
-	}
-
-	/* This isn't safe because we have too many invalid pin default values
-	for (TSharedPtr<FUserPinInfo> PinInfo : UserDefinedPins)
-	{
-		if (PinInfo->DesiredPinDirection == EGPD_Output)
-		{
-			UProperty* PinProperty = VariableStruct->FindPropertyByName(PinInfo->PinName);
-			if (PinProperty && (!PinProperty->HasAnyPropertyFlags(CPF_OutParm) || PinProperty->HasAnyPropertyFlags(CPF_ReferenceParm)))
-			{
-				FEdGraphPinType PinType;
-				K2Schema->ConvertPropertyToPinType(PinProperty, PinType);
-
-				if (PinType != PinInfo->PinType)
-				{
-					UE_LOG(LogBlueprint, Log, TEXT("Pin type for variable %s does not match type on struct %s during UpdateDefaultsFromVariableStruct, ignoring old default"), *PinInfo->PinName.ToString(), *VariableStruct->GetName());
-				}
-				else
-				{
-					FString NewDefaultValue;
-					FBlueprintEditorUtils::PropertyValueToString(PinProperty, VariableStructData, NewDefaultValue);
-
-					if (NewDefaultValue != PinInfo->PinDefaultValue)
-					{
-						// Call the internal function to avoid recursion back into this update
-						if (UpdateEdGraphPinDefaultValue(PinInfo, NewDefaultValue))
-						{
-							PinInfo->PinDefaultValue = NewDefaultValue;
-						}
-						else
-						{
-							UE_LOG(LogBlueprint, Warning, TEXT("Failed to update variable %s on struct %s during UpdateDefaultsFromVariableStruct with new default %s"), *PinInfo->PinName.ToString(), *VariableStruct->GetName(), *NewDefaultValue);
-						}
-					}
-				}
-			}
-			else
-			{
-				UE_LOG(LogBlueprint, Log, TEXT("Could not find property %s on struct %s during UpdateDefaultsFromVariableStruct"), *PinInfo->PinName.ToString(), *VariableStruct->GetName());
-			}
-		}
-	}*/
-
-	if (LocalVariables)
-	{
-		for (FBPVariableDescription& LocalVariable : *LocalVariables)
-		{
-			UProperty* PinProperty = VariableStruct->FindPropertyByName(LocalVariable.VarName);
-
-			if (PinProperty && (!PinProperty->HasAnyPropertyFlags(CPF_OutParm) || PinProperty->HasAnyPropertyFlags(CPF_ReferenceParm)))
-			{
-				FEdGraphPinType PinType;
-				K2Schema->ConvertPropertyToPinType(PinProperty, /*out*/ PinType);
-
-				if (PinType != LocalVariable.VarType)
-				{
-					UE_LOG(LogBlueprint, Log, TEXT("Pin type for local variable %s does not match type on struct %s during UpdateDefaultsFromVariableStruct, ignoring old default"), *LocalVariable.VarName.ToString(), *VariableStruct->GetName());
-				}
-				else
-				{
-					FBlueprintEditorUtils::PropertyValueToString(PinProperty, VariableStructData, LocalVariable.DefaultValue);
-				}
-			}
-			else
-			{
-				UE_LOG(LogBlueprint, Log, TEXT("Could not find local variable property %s on struct %s during UpdateDefaultsFromVariableStruct"), *LocalVariable.VarName.ToString(), *VariableStruct->GetName());
-			}
-		}
-	}
-
-	return true;
-}
diff --git a/Engine/Source/Editor/BlueprintGraph/Private/K2Node_FunctionEntry.cpp b/Engine/Source/Editor/BlueprintGraph/Private/K2Node_FunctionEntry.cpp
--- a/Engine/Source/Editor/BlueprintGraph/Private/K2Node_FunctionEntry.cpp
+++ b/Engine/Source/Editor/BlueprintGraph/Private/K2Node_FunctionEntry.cpp
@@ -204,7 +204,7 @@ void UK2Node_FunctionEntry::PreSave(const class ITargetPlatform* TargetPlatform)
 	Super::PreSave(TargetPlatform);
 	
 	const UBlueprint* Blueprint = HasValidBlueprint() ? GetBlueprint() : nullptr;
-	if (Blueprint && (LocalVariables.Num() > 0 || UserDefinedPins.Num() > 0))
+	if (Blueprint && LocalVariables.Num() > 0)
 	{
 		// Forcibly fixup defaults before we save
 		UpdateLoadedDefaultValues(true);
@@ -242,7 +242,7 @@ void UK2Node_FunctionEntry::Serialize(FArchive& Ar)
 				Struct->SerializeBin(Ar, FunctionVariableCache->GetStructMemory());
 
 				// Copy back into defaults as they may have changed
-				UpdateDefaultsFromVariableStruct(FunctionVariableCache->GetStruct(), FunctionVariableCache->GetStructMemory(), &LocalVariables);
+				UpdateDefaultsFromVariableStruct(FunctionVariableCache->GetStruct(), FunctionVariableCache->GetStructMemory());
 			}
 		}
 	}
@@ -395,11 +395,11 @@ TSharedPtr<FStructOnScope> UK2Node_FunctionEntry::GetFunctionVariableCache(bool
 		FunctionVariableCache.Reset();
 	}
 
-	if (!FunctionVariableCache.IsValid())
+	if (!FunctionVariableCache.IsValid() || !FunctionVariableCache->IsValid())
 	{
 		if (UFunction* const Function = FindSignatureFunction())
 		{
-			if (LocalVariables.Num() > 0 || UserDefinedPins.Num() > 0)
+			if (LocalVariables.Num() > 0)
 			{
 				FunctionVariableCache = MakeShared<FStructOnScope>(Function);
 				FunctionVariableCache->SetPackage(GetOutermost());
@@ -419,7 +419,7 @@ bool UK2Node_FunctionEntry::RefreshFunctionVariableCache()
 	if (FunctionVariableCache.IsValid())
 	{
 		// Update the cache if it was created
-		return UpdateVariableStructFromDefaults(FunctionVariableCache->GetStruct(), FunctionVariableCache->GetStructMemory(), &LocalVariables);
+		return UpdateVariableStructFromDefaults(FunctionVariableCache->GetStruct(), FunctionVariableCache->GetStructMemory());
 	}
 	return false;
 }
@@ -436,7 +436,7 @@ bool UK2Node_FunctionEntry::UpdateLoadedDefaultValues(bool bForceRefresh)
 		if (FunctionVariableCache.IsValid())
 		{
 			// Now copy back into the default value strings
-			return UpdateDefaultsFromVariableStruct(FunctionVariableCache->GetStruct(), FunctionVariableCache->GetStructMemory(), &LocalVariables);
+			return UpdateDefaultsFromVariableStruct(FunctionVariableCache->GetStruct(), FunctionVariableCache->GetStructMemory());
 		}
 		else
 		{
@@ -448,6 +448,96 @@ bool UK2Node_FunctionEntry::UpdateLoadedDefaultValues(bool bForceRefresh)
 	return false;
 }
 
+void UK2Node_FunctionEntry::ClearCachedBlueprintData(UBlueprint* Blueprint)
+{
+	FunctionVariableCache.Reset();
+}
+
+bool UK2Node_FunctionEntry::UpdateVariableStructFromDefaults(const UStruct* VariableStruct, uint8* VariableStructData)
+{
+	const UEdGraphSchema_K2* K2Schema = GetDefault<UEdGraphSchema_K2>();
+
+	if (!VariableStruct || !VariableStructData)
+	{
+		return false;
+	}
+
+	for (FBPVariableDescription& LocalVariable : LocalVariables)
+	{
+		if (!LocalVariable.DefaultValue.IsEmpty())
+		{
+			UProperty* PinProperty = VariableStruct->FindPropertyByName(LocalVariable.VarName);
+
+			if (PinProperty && (!PinProperty->HasAnyPropertyFlags(CPF_OutParm) || PinProperty->HasAnyPropertyFlags(CPF_ReferenceParm)))
+			{
+				FEdGraphPinType PinType;
+				K2Schema->ConvertPropertyToPinType(PinProperty, /*out*/ PinType);
+
+				if (PinType != LocalVariable.VarType)
+				{
+					//UE_LOG(LogBlueprint, Log, TEXT("Pin type for local variable %s does not match type on struct %s during UpdateVariableStructFromDefaults, ignoring old default"), *LocalVariable.VarName.ToString(), *VariableStruct->GetName());
+				}
+				else
+				{
+					FBlueprintEditorUtils::PropertyValueFromString(PinProperty, LocalVariable.DefaultValue, VariableStructData, this);
+				}
+			}
+			else
+			{
+				//UE_LOG(LogBlueprint, Log, TEXT("Could not find local variable property %s on struct %s during UpdateVariableStructFromDefaults"), *LocalVariable.VarName.ToString(), *VariableStruct->GetName());
+			}
+		}
+	}
+
+	return true;
+}
+
+bool UK2Node_FunctionEntry::UpdateDefaultsFromVariableStruct(const UStruct* VariableStruct, uint8* VariableStructData)
+{
+	const UEdGraphSchema_K2* K2Schema = GetDefault<UEdGraphSchema_K2>();
+
+	if (!VariableStruct || !VariableStructData)
+	{
+		return false;
+	}
+
+	for (FBPVariableDescription& LocalVariable : LocalVariables)
+	{
+		if (!LocalVariable.DefaultValue.IsEmpty())
+		{
+			// We don't want to write out fields that were empty before, as they were guaranteed to not have actual real data
+			UProperty* PinProperty = VariableStruct->FindPropertyByName(LocalVariable.VarName);
+
+			if (PinProperty && (!PinProperty->HasAnyPropertyFlags(CPF_OutParm) || PinProperty->HasAnyPropertyFlags(CPF_ReferenceParm)))
+			{
+				FEdGraphPinType PinType;
+				K2Schema->ConvertPropertyToPinType(PinProperty, /*out*/ PinType);
+
+				if (PinType != LocalVariable.VarType)
+				{
+					//UE_LOG(LogBlueprint, Log, TEXT("Pin type for local variable %s does not match type on struct %s during UpdateDefaultsFromVariableStruct, ignoring old default"), *LocalVariable.VarName.ToString(), *VariableStruct->GetName());
+				}
+				else
+				{
+					FString NewValue;
+					FBlueprintEditorUtils::PropertyValueToString(PinProperty, VariableStructData, NewValue, this);
+					if (NewValue != LocalVariable.DefaultValue)
+					{
+						LocalVariable.DefaultValue = NewValue;
+					}
+				}
+			}
+			else
+			{
+				//UE_LOG(LogBlueprint, Log, TEXT("Could not find local variable property %s on struct %s during UpdateDefaultsFromVariableStruct"), *LocalVariable.VarName.ToString(), *VariableStruct->GetName());
+			}
+		}
+	}
+
+	return true;
+}
+
+
 FNodeHandlingFunctor* UK2Node_FunctionEntry::CreateNodeHandler(FKismetCompilerContext& CompilerContext) const
 {
 	return new FKCHandler_FunctionEntry(CompilerContext);
@@ -547,7 +637,17 @@ void UK2Node_FunctionEntry::FindDiffs(UEdGraphNode* OtherNode, struct FDiffResul
 
 		for (int32 i = 0; i < LocalVariables.Num() && !bLocalVarsDiffer; i++)
 		{
-			if (!FBPVariableDescription::StaticStruct()->CompareScriptStruct(&LocalVariables[i], &OtherFunction->LocalVariables[i], 0))
+			const FBPVariableDescription& ThisVar = LocalVariables[i];
+			const FBPVariableDescription& OtherVar = OtherFunction->LocalVariables[i];
+			
+			// Can't do a raw compare, for local variable defaults we need to compare the struct
+			if (ThisVar.VarName != OtherVar.VarName
+				|| ThisVar.VarType != OtherVar.VarType
+				|| ThisVar.FriendlyName != OtherVar.FriendlyName
+				|| !ThisVar.Category.EqualTo(OtherVar.Category)
+				|| ThisVar.PropertyFlags != OtherVar.PropertyFlags
+				|| ThisVar.RepNotifyFunc != OtherVar.RepNotifyFunc
+				|| ThisVar.ReplicationCondition != OtherVar.ReplicationCondition)
 			{
 				bLocalVarsDiffer = true;
 			}
@@ -559,11 +659,29 @@ void UK2Node_FunctionEntry::FindDiffs(UEdGraphNode* OtherNode, struct FDiffResul
 			Diff.Diff = EDiffType::NODE_PROPERTY;
 			Diff.Node1 = this;
 			Diff.Node2 = OtherNode;
-			Diff.DisplayString = LOCTEXT("DIF_FunctionLocalVariables", "Function local variables have changed");
+			Diff.DisplayString = LOCTEXT("DIF_FunctionLocalVariables", "Function local variables have changed in structure");
 			Diff.DisplayColor = FLinearColor(0.25f, 0.71f, 0.85f);
 
 			Results.Add(Diff);
 		}
+		else
+		{
+			TSharedPtr<FStructOnScope> MyLocals = GetFunctionVariableCache();
+			TSharedPtr<FStructOnScope> OtherLocals = OtherFunction->GetFunctionVariableCache();
+
+			if (MyLocals.IsValid() && MyLocals->IsValid() && OtherLocals.IsValid() && OtherLocals->IsValid())
+			{
+				// Check for local var diffs
+				FDiffSingleResult Diff;
+				Diff.Diff = EDiffType::NODE_PROPERTY;
+				Diff.Node1 = this;
+				Diff.Node2 = OtherNode;
+				Diff.ToolTip = LOCTEXT("DIF_FunctionLocalVariableDefaults", "Function local variable default values have changed");
+				Diff.DisplayColor = FLinearColor(0.25f, 0.71f, 0.85f);
+
+				DiffProperties(const_cast<UStruct*>(MyLocals->GetStruct()), const_cast<UStruct*>(OtherLocals->GetStruct()), MyLocals->GetStructMemory(), OtherLocals->GetStructMemory(), Results, Diff);
+			}
+		}
 	}
 }
 
@@ -707,11 +825,6 @@ void UK2Node_FunctionEntry::PostReconstructNode()
 	Super::PostReconstructNode();
 }
 
-void UK2Node_FunctionEntry::ClearCachedBlueprintData(UBlueprint* Blueprint)
-{
-	FunctionVariableCache.Reset();
-}
-
 void UK2Node_FunctionEntry::FixupPinStringDataReferences(FArchive* SavingArchive)
 {
 	Super::FixupPinStringDataReferences(SavingArchive);
@@ -729,7 +842,7 @@ void UK2Node_FunctionEntry::FixupPinStringDataReferences(FArchive* SavingArchive
 
 					if (DefaultsString != PinInfo->PinDefaultValue)
 					{
-						ModifyUserDefinedPinDefaultValue(PinInfo, Pin->GetDefaultAsString());
+						ModifyUserDefinedPinDefaultValue(PinInfo, DefaultsString);
 					}
 				}
 			}
diff --git a/Engine/Source/Editor/BlueprintGraph/Private/K2Node_MakeVariable.cpp b/Engine/Source/Editor/BlueprintGraph/Private/K2Node_MakeVariable.cpp
--- a/Engine/Source/Editor/BlueprintGraph/Private/K2Node_MakeVariable.cpp
+++ b/Engine/Source/Editor/BlueprintGraph/Private/K2Node_MakeVariable.cpp
@@ -100,7 +100,7 @@ void UK2Node_MakeVariable::SetupVariable(const FBPVariableDescription& InVariabl
 
 	// make input pins for every value in the default value, for map pins we'll make a pair of inputs:
 	TSharedPtr<FStructOnScope> StructData = MakeShareable(new FStructOnScope(Scope));
-	FBlueprintEditorUtils::PropertyValueFromString(Property, VariableType.DefaultValue, StructData->GetStructMemory());
+	FBlueprintEditorUtils::PropertyValueFromString(Property, VariableType.DefaultValue, StructData->GetStructMemory(), this);
 
 	if(VariableType.VarType.IsArray())
 	{
diff --git a/Engine/Source/Editor/GameplayTasksEditor/Private/K2Node_LatentGameplayTaskCall.cpp b/Engine/Source/Editor/GameplayTasksEditor/Private/K2Node_LatentGameplayTaskCall.cpp
--- a/Engine/Source/Editor/GameplayTasksEditor/Private/K2Node_LatentGameplayTaskCall.cpp
+++ b/Engine/Source/Editor/GameplayTasksEditor/Private/K2Node_LatentGameplayTaskCall.cpp
@@ -218,7 +218,7 @@ void UK2Node_LatentGameplayTaskCall::CreatePinsForClass(UClass* InClass)
 			if (ClassDefaultObject && K2Schema->PinDefaultValueIsEditable(*Pin))
 			{
 				FString DefaultValueAsString;
-				const bool bDefaultValueSet = FBlueprintEditorUtils::PropertyValueToString(Property, reinterpret_cast<const uint8*>(ClassDefaultObject), DefaultValueAsString);
+				const bool bDefaultValueSet = FBlueprintEditorUtils::PropertyValueToString(Property, reinterpret_cast<const uint8*>(ClassDefaultObject), DefaultValueAsString, this);
 				check(bDefaultValueSet);
 				K2Schema->SetPinAutogeneratedDefaultValue(Pin, DefaultValueAsString);
 			}
@@ -434,7 +434,7 @@ bool UK2Node_LatentGameplayTaskCall::ConnectSpawnProperties(UClass* ClassToSpawn
 					// We don't want to generate an assignment node unless the default value 
 					// differs from the value in the CDO:
 					FString DefaultValueAsString;
-					FBlueprintEditorUtils::PropertyValueToString(Property, (uint8*)ClassToSpawn->ClassDefaultObject, DefaultValueAsString);
+					FBlueprintEditorUtils::PropertyValueToString(Property, (uint8*)ClassToSpawn->ClassDefaultObject, DefaultValueAsString, this);
 					if (DefaultValueAsString == SpawnVarPin->DefaultValue)
 					{
 						continue;
diff --git a/Engine/Source/Editor/Kismet/Private/BlueprintDetailsCustomization.cpp b/Engine/Source/Editor/Kismet/Private/BlueprintDetailsCustomization.cpp
--- a/Engine/Source/Editor/Kismet/Private/BlueprintDetailsCustomization.cpp
+++ b/Engine/Source/Editor/Kismet/Private/BlueprintDetailsCustomization.cpp
@@ -2625,12 +2625,12 @@ void FBlueprintVarActionDetails::OnFinishedChangingProperties(const FPropertyCha
 
 	if (InStructData.IsValid())
 	{
-		bDefaultValueSet = FBlueprintEditorUtils::PropertyValueToString(DirectProperty, InStructData->GetStructMemory(), DefaultValueString);
+		UK2Node_FunctionEntry* FuncEntry = Cast<UK2Node_FunctionEntry>(InEntryNode.Get());
+
+		bDefaultValueSet = FBlueprintEditorUtils::PropertyValueToString(DirectProperty, InStructData->GetStructMemory(), DefaultValueString, FuncEntry);
 
 		if(bDefaultValueSet)
 		{
-			UK2Node_FunctionEntry* FuncEntry = Cast<UK2Node_FunctionEntry>(InEntryNode.Get());
-
 			// Search out the correct local variable in the Function Entry Node and set the default value
 			for (FBPVariableDescription& LocalVar : FuncEntry->LocalVariables)
 			{
diff --git a/Engine/Source/Editor/Kismet/Private/UserDefinedStructureEditor.cpp b/Engine/Source/Editor/Kismet/Private/UserDefinedStructureEditor.cpp
--- a/Engine/Source/Editor/Kismet/Private/UserDefinedStructureEditor.cpp
+++ b/Engine/Source/Editor/Kismet/Private/UserDefinedStructureEditor.cpp
@@ -216,7 +216,7 @@ void FDefaultValueDetails::OnFinishedChangingProperties(const FPropertyChangedEv
 				{
 					if (StructData.IsValid() && StructData->IsValid())
 					{
-						bDefaultValueSet = FBlueprintEditorUtils::PropertyValueToString(DirectProperty, StructData->GetStructMemory(), DefaultValueString);
+						bDefaultValueSet = FBlueprintEditorUtils::PropertyValueToString(DirectProperty, StructData->GetStructMemory(), DefaultValueString, OwnerStruct);
 					}
 				}
 
diff --git a/Engine/Source/Editor/UnrealEd/Private/Kismet2/BlueprintEditorUtils.cpp b/Engine/Source/Editor/UnrealEd/Private/Kismet2/BlueprintEditorUtils.cpp
--- a/Engine/Source/Editor/UnrealEd/Private/Kismet2/BlueprintEditorUtils.cpp
+++ b/Engine/Source/Editor/UnrealEd/Private/Kismet2/BlueprintEditorUtils.cpp
@@ -8150,12 +8150,12 @@ bool FBlueprintEditorUtils::IsObjectADebugCandidate( AActor* InActorObject, UBlu
 	return bPassesFlags && bCanDebugThisObject;
 }
 
-bool FBlueprintEditorUtils::PropertyValueFromString(const UProperty* Property, const FString& StrValue, uint8* Container)
+bool FBlueprintEditorUtils::PropertyValueFromString(const UProperty* Property, const FString& StrValue, uint8* Container, UObject* OwningObject)
 {
-	return PropertyValueFromString_Direct(Property, StrValue, Property->ContainerPtrToValuePtr<uint8>(Container));
+	return PropertyValueFromString_Direct(Property, StrValue, Property->ContainerPtrToValuePtr<uint8>(Container), OwningObject);
 }
 
-bool FBlueprintEditorUtils::PropertyValueFromString_Direct(const UProperty* Property, const FString& StrValue, uint8* DirectValue)
+bool FBlueprintEditorUtils::PropertyValueFromString_Direct(const UProperty* Property, const FString& StrValue, uint8* DirectValue, UObject* OwningObject)
 {
 	bool bParseSucceeded = true;
 	if (!Property->IsA(UStructProperty::StaticClass()))
@@ -8221,7 +8221,7 @@ bool FBlueprintEditorUtils::PropertyValueFromString_Direct(const UProperty* Prop
 		else if (Property->IsA(UTextProperty::StaticClass()))
 		{
 			FStringOutputDevice ImportError;
-			const TCHAR* EndOfParsedBuff = Property->ImportText(*StrValue, DirectValue, PPF_SerializedAsImportText, nullptr, &ImportError);
+			const TCHAR* EndOfParsedBuff = Property->ImportText(*StrValue, DirectValue, PPF_SerializedAsImportText, OwningObject, &ImportError);
 			bParseSucceeded = EndOfParsedBuff && ImportError.IsEmpty();
 		}
 		else
@@ -8232,7 +8232,7 @@ bool FBlueprintEditorUtils::PropertyValueFromString_Direct(const UProperty* Prop
 				: *StrValue;
 
 			FStringOutputDevice ImportError;
-			const TCHAR* EndOfParsedBuff = Property->ImportText(*StrValue, DirectValue, PPF_SerializedAsImportText, nullptr, &ImportError);
+			const TCHAR* EndOfParsedBuff = Property->ImportText(*StrValue, DirectValue, PPF_SerializedAsImportText, OwningObject, &ImportError);
 			bParseSucceeded = EndOfParsedBuff && ImportError.IsEmpty();
 		}
 	}
@@ -8279,7 +8279,7 @@ bool FBlueprintEditorUtils::PropertyValueFromString_Direct(const UProperty* Prop
 			ensure(1 == StructProperty->ArrayDim);
 
 			FStringOutputDevice ImportError;
-			const TCHAR* EndOfParsedBuff = StructProperty->ImportText(StrValue.IsEmpty() ? TEXT("()") : *StrValue, DirectValue, PPF_SerializedAsImportText, nullptr, &ImportError);
+			const TCHAR* EndOfParsedBuff = StructProperty->ImportText(StrValue.IsEmpty() ? TEXT("()") : *StrValue, DirectValue, PPF_SerializedAsImportText, OwningObject, &ImportError);
 			bParseSucceeded &= EndOfParsedBuff && ImportError.IsEmpty();
 		}
 	}
@@ -8287,12 +8287,12 @@ bool FBlueprintEditorUtils::PropertyValueFromString_Direct(const UProperty* Prop
 	return bParseSucceeded;
 }
 
-bool FBlueprintEditorUtils::PropertyValueToString(const UProperty* Property, const uint8* Container, FString& OutForm)
+bool FBlueprintEditorUtils::PropertyValueToString(const UProperty* Property, const uint8* Container, FString& OutForm, UObject* OwningObject)
 {
-	return PropertyValueToString_Direct(Property, Property->ContainerPtrToValuePtr<const uint8>(Container), OutForm);
+	return PropertyValueToString_Direct(Property, Property->ContainerPtrToValuePtr<const uint8>(Container), OutForm, OwningObject);
 }
 
-bool FBlueprintEditorUtils::PropertyValueToString_Direct(const UProperty* Property, const uint8* DirectValue, FString& OutForm)
+bool FBlueprintEditorUtils::PropertyValueToString_Direct(const UProperty* Property, const uint8* DirectValue, FString& OutForm, UObject* OwningObject)
 {
 	check(Property && DirectValue);
 	OutForm.Reset();
@@ -8336,7 +8336,7 @@ bool FBlueprintEditorUtils::PropertyValueToString_Direct(const UProperty* Proper
 	if (OutForm.IsEmpty())
 	{
 		const uint8* DefaultValue = DirectValue;	
-		bSucceeded = Property->ExportText_Direct(OutForm, DirectValue, DefaultValue, nullptr, PPF_SerializedAsImportText);
+		bSucceeded = Property->ExportText_Direct(OutForm, DirectValue, DefaultValue, OwningObject, PPF_SerializedAsImportText);
 	}
 	return bSucceeded;
 }
diff --git a/Engine/Source/Editor/UnrealEd/Private/Kismet2/StructureEditorUtils.cpp b/Engine/Source/Editor/UnrealEd/Private/Kismet2/StructureEditorUtils.cpp
--- a/Engine/Source/Editor/UnrealEd/Private/Kismet2/StructureEditorUtils.cpp
+++ b/Engine/Source/Editor/UnrealEd/Private/Kismet2/StructureEditorUtils.cpp
@@ -447,7 +447,7 @@ bool FStructureEditorUtils::ChangeVariableDefaultValue(UUserDefinedStruct* Struc
 			const UProperty* Property = FindField<UProperty>(Struct, VarDesc->VarName);
 			FStructOnScope StructDefaultMem(Struct);
 			bAdvancedValidation = StructDefaultMem.IsValid() && Property &&
-				FBlueprintEditorUtils::PropertyValueFromString(Property, NewDefaultValue, StructDefaultMem.GetStructMemory());
+				FBlueprintEditorUtils::PropertyValueFromString(Property, NewDefaultValue, StructDefaultMem.GetStructMemory(), Struct);
 		}
 
 		if (bAdvancedValidation)
diff --git a/Engine/Source/Editor/UnrealEd/Private/UserDefinedStructEditorData.cpp b/Engine/Source/Editor/UnrealEd/Private/UserDefinedStructEditorData.cpp
--- a/Engine/Source/Editor/UnrealEd/Private/UserDefinedStructEditorData.cpp
+++ b/Engine/Source/Editor/UnrealEd/Private/UserDefinedStructEditorData.cpp
@@ -206,7 +206,7 @@ void UUserDefinedStructEditorData::ReinitializeDefaultInstance(FString* OutLog)
 				FStructVariableDescription* VarDesc = VariablesDescriptions.FindByPredicate(FStructureEditorUtils::FFindByGuidHelper<FStructVariableDescription>(VarGuid));
 				if (VarDesc && !VarDesc->CurrentDefaultValue.IsEmpty())
 				{
-					if (!FBlueprintEditorUtils::PropertyValueFromString(Property, VarDesc->CurrentDefaultValue, StructData))
+					if (!FBlueprintEditorUtils::PropertyValueFromString(Property, VarDesc->CurrentDefaultValue, StructData, ScriptStruct))
 					{
 						const FString Message = FString::Printf(TEXT("Cannot parse value. Property: %s String: \"%s\" ")
 							, *Property->GetDisplayNameText().ToString()
diff --git a/Engine/Source/Editor/UnrealEd/Public/Kismet2/BlueprintEditorUtils.h b/Engine/Source/Editor/UnrealEd/Public/Kismet2/BlueprintEditorUtils.h
--- a/Engine/Source/Editor/UnrealEd/Public/Kismet2/BlueprintEditorUtils.h
+++ b/Engine/Source/Editor/UnrealEd/Public/Kismet2/BlueprintEditorUtils.h
@@ -1156,16 +1156,16 @@ public:
 	static bool IsVariableUsed(const UBlueprint* Blueprint, const FName& Name, UEdGraph* LocalGraphScope = nullptr);
 
 	/** Copies the value from the passed in string into a property. ContainerMem points to the Struct or Class containing Property */
-	static bool PropertyValueFromString(const UProperty* Property, const FString& StrValue, uint8* Container);
+	static bool PropertyValueFromString(const UProperty* Property, const FString& StrValue, uint8* Container, UObject* OwningObject = nullptr);
 
 	/** Copies the value from the passed in string into a property. DirectValue is the raw memory address of the property value */
-	static bool PropertyValueFromString_Direct(const UProperty* Property, const FString& StrValue, uint8* DirectValue);
+	static bool PropertyValueFromString_Direct(const UProperty* Property, const FString& StrValue, uint8* DirectValue, UObject* OwningObject = nullptr);
 
 	/** Copies the value from a property into the string OutForm. ContainerMem points to the Struct or Class containing Property */
-	static bool PropertyValueToString(const UProperty* Property, const uint8* Container, FString& OutForm);
+	static bool PropertyValueToString(const UProperty* Property, const uint8* Container, FString& OutForm, UObject* OwningObject = nullptr);
 
 	/** Copies the value from a property into the string OutForm. DirectValue is the raw memory address of the property value */
-	static bool PropertyValueToString_Direct(const UProperty* Property, const uint8* DirectValue, FString& OutForm);
+	static bool PropertyValueToString_Direct(const UProperty* Property, const uint8* DirectValue, FString& OutForm, UObject* OwningObject = nullptr);
 
 	/** Call PostEditChange() on all Actors based on the given Blueprint */
 	static void PostEditChangeBlueprintActors(UBlueprint* Blueprint, bool bComponentEditChange = false);
diff --git a/Engine/Source/Runtime/Engine/Classes/EdGraph/EdGraphNode.h b/Engine/Source/Runtime/Engine/Classes/EdGraph/EdGraphNode.h
--- a/Engine/Source/Runtime/Engine/Classes/EdGraph/EdGraphNode.h
+++ b/Engine/Source/Runtime/Engine/Classes/EdGraph/EdGraphNode.h
@@ -835,7 +835,7 @@ public:
 	
 protected:
 	/**
-	 * Finds the difference in properties of node instance
+	 * Finds the difference in properties of node instance, for subobjects
 	 *
 	 * @param StructA The struct of the class we are looking at LHS
 	 * @param StructB The struct of the class we are looking at RHS
@@ -846,6 +846,18 @@ protected:
 	 */
 	virtual void DiffProperties(UClass* StructA, UClass* StructB, UObject* DataA, UObject* DataB, FDiffResults& Results, FDiffSingleResult& Diff) const;
 
+	/**
+	 * Finds the difference in properties of node instance, for arbitrary UStructs
+	 *
+	 * @param StructA The struct we are looking at LHS
+	 * @param StructB The struct we are looking at RHS
+	 * @param DataA The raw data we are comparing LHS
+	 * @param DataB The raw data we are comparing RHS
+	 * @param Results The Results where differences are stored
+	 * @param Diff The single result with default parameters setup
+	 */
+	virtual void DiffProperties(UStruct* StructA, UStruct* StructB, uint8* DataA, uint8* DataB, FDiffResults& Results, FDiffSingleResult& Diff) const;
+
 	// Returns a human-friendly description of the property in the form "PropertyName: Value"
 	virtual FString GetPropertyNameAndValueForDiff(const UProperty* Prop, const uint8* PropertyAddr) const;
 
diff --git a/Engine/Source/Runtime/Engine/Private/EdGraph/EdGraphNode.cpp b/Engine/Source/Runtime/Engine/Private/EdGraph/EdGraphNode.cpp
--- a/Engine/Source/Runtime/Engine/Private/EdGraph/EdGraphNode.cpp
+++ b/Engine/Source/Runtime/Engine/Private/EdGraph/EdGraphNode.cpp
@@ -229,6 +229,7 @@ FString UEdGraphNode::GetPropertyNameAndValueForDiff(const UProperty* Prop, cons
 	return FString::Printf(TEXT("%s: %s"), *FName::NameToDisplayString(Prop->GetName(), bIsBool), *ExportedStringValue);
 }
 
+
 void UEdGraphNode::DiffProperties(UClass* StructA, UClass* StructB, UObject* DataA, UObject* DataB, FDiffResults& Results, FDiffSingleResult& Diff) const
 {
 	// Find the common parent class in case the other node isn't of the same type
@@ -238,10 +239,23 @@ void UEdGraphNode::DiffProperties(UClass* StructA, UClass* StructB, UObject* Dat
 		ClassToViewAs = ClassToViewAs->GetSuperClass();
 	}
 
-	// Run through all the properties
-	for (TFieldIterator<UProperty> PropertyIt(ClassToViewAs, EFieldIteratorFlags::IncludeSuper); PropertyIt; ++PropertyIt)
+	DiffProperties(ClassToViewAs, ClassToViewAs, (uint8*)DataA, (uint8*)DataB, Results, Diff);
+}
+
+void UEdGraphNode::DiffProperties(UStruct* StructA, UStruct* StructB, uint8* DataA, uint8* DataB, FDiffResults& Results, FDiffSingleResult& Diff) const
+{
+	// Run through all the properties in the first struct
+	for (TFieldIterator<UProperty> PropertyIt(StructA, EFieldIteratorFlags::IncludeSuper); PropertyIt; ++PropertyIt)
 	{
 		UProperty* Prop = *PropertyIt;
+		UProperty* PropB = StructB->FindPropertyByName(Prop->GetFName());
+
+		if (!PropB || Prop->GetClass() != PropB->GetClass())
+		{
+			// Skip if properties don't match
+			continue;
+		}
+
 		// skip properties we cant see
 		if (!Prop->HasAnyPropertyFlags(CPF_Edit | CPF_BlueprintVisible) ||
 			Prop->HasAnyPropertyFlags(CPF_Transient) ||
@@ -253,7 +267,7 @@ void UEdGraphNode::DiffProperties(UClass* StructA, UClass* StructB, UObject* Dat
 		}
 
 		const FString ValueStringA = GetPropertyNameAndValueForDiff(Prop, Prop->ContainerPtrToValuePtr<uint8>(DataA));
-		const FString ValueStringB = GetPropertyNameAndValueForDiff(Prop, Prop->ContainerPtrToValuePtr<uint8>(DataB));
+		const FString ValueStringB = GetPropertyNameAndValueForDiff(PropB, PropB->ContainerPtrToValuePtr<uint8>(DataB));
 
 		if (ValueStringA != ValueStringB)
 		{
