More fixes/improvements to Blueprint graph node deprecation UI.

Change summary:
- Added "(deprecated)" to the display bar UI at the top of user-defined Blueprint event and function graphs.
- Added FEdGraphNodeDeprecationResponse along with EEdGraphNodeDeprecationType and EEdGraphNodeDeprecationMessageType enums.
- Added UEdGraphNode::HasDeprecatedReference() to differentiate between IsDeprecated() (which now means the node itself is deprecated) vs. a reference to a deprecated class member (e.g. variable, function).
- Deprecated UEdGraphNode::GetDeprecationMessage() and UEdGraphNode::ShouldWarnOnDeprecation() APIs in favor of UEdGraphNode::GetDeprecationResponse() (new). Node subclasses should override this when they need to override the message type (EEdGraphNodeDeprecationMessageType) and/or the message text based on the type of deprecation that's being reported by the node (EEdGraphNodeDeprecationType).
- A default deprecation response is now constructed for both cases by UEdGraphNode.
- Added a case for deprecation "notes" to FGraphCompilerContext::ValidateNode(). These are currently used to visually identify deprecated function entry nodes and event nodes in the Blueprint class in which they are deprecated, as these represent the definitions of those functions and not actual usage or call site (which we continue to warn about).
- Modified various node subclasses in which it was necessary to convert over to using a GetDeprecationResponse() method override.

diff --git a/Engine/Plugins/Editor/GameplayTagsEditor/Source/GameplayTagsEditor/Classes/GameplayTagsK2Node_LiteralGameplayTag.h b/Engine/Plugins/Editor/GameplayTagsEditor/Source/GameplayTagsEditor/Classes/GameplayTagsK2Node_LiteralGameplayTag.h
--- a/Engine/Plugins/Editor/GameplayTagsEditor/Source/GameplayTagsEditor/Classes/GameplayTagsK2Node_LiteralGameplayTag.h
+++ b/Engine/Plugins/Editor/GameplayTagsEditor/Source/GameplayTagsEditor/Classes/GameplayTagsK2Node_LiteralGameplayTag.h
@@ -35,7 +35,7 @@ class UGameplayTagsK2Node_LiteralGameplayTag : public UK2Node
 	virtual bool IsNodePure() const override { return true; }
 	virtual bool IsDeprecated() const override { return true; } 
 	virtual void ConvertDeprecatedNode(UEdGraph* Graph, bool bOnlySafeChanges) override;
-	virtual FString GetDeprecationMessage() const override;
+	virtual FEdGraphNodeDeprecationResponse GetDeprecationResponse(EEdGraphNodeDeprecationType DeprecationType) const override;
 	//~ End UK2Node Interface
 #endif
 
diff --git a/Engine/Plugins/Editor/GameplayTagsEditor/Source/GameplayTagsEditor/Private/GameplayTagsK2Node_LiteralGameplayTag.cpp b/Engine/Plugins/Editor/GameplayTagsEditor/Source/GameplayTagsEditor/Private/GameplayTagsK2Node_LiteralGameplayTag.cpp
--- a/Engine/Plugins/Editor/GameplayTagsEditor/Source/GameplayTagsEditor/Private/GameplayTagsK2Node_LiteralGameplayTag.cpp
+++ b/Engine/Plugins/Editor/GameplayTagsEditor/Source/GameplayTagsEditor/Private/GameplayTagsK2Node_LiteralGameplayTag.cpp
@@ -57,9 +57,12 @@ FText UGameplayTagsK2Node_LiteralGameplayTag::GetMenuCategory() const
 	return LOCTEXT("ActionMenuCategory", "Gameplay Tags");
 }
 
-FString UGameplayTagsK2Node_LiteralGameplayTag::GetDeprecationMessage() const
+FEdGraphNodeDeprecationResponse UGameplayTagsK2Node_LiteralGameplayTag::GetDeprecationResponse(EEdGraphNodeDeprecationType DeprecationType) const
 {
-	return LOCTEXT("NodeDeprecated_Warning", "@@ is deprecated, replace with Make Literal GameplayTagContainer function call").ToString();
+	FEdGraphNodeDeprecationResponse Response = Super::GetDeprecationResponse(DeprecationType);
+	Response.MessageText = LOCTEXT("NodeDeprecated_Warning", "@@ is deprecated, replace with Make Literal GameplayTagContainer function call");
+
+	return Response;
 }
 
 void UGameplayTagsK2Node_LiteralGameplayTag::ConvertDeprecatedNode(UEdGraph* Graph, bool bOnlySafeChanges)
diff --git a/Engine/Plugins/Experimental/ControlRig/Source/ControlRigDeveloper/Private/Graph/ControlRigGraphNode.cpp b/Engine/Plugins/Experimental/ControlRig/Source/ControlRigDeveloper/Private/Graph/ControlRigGraphNode.cpp
--- a/Engine/Plugins/Experimental/ControlRig/Source/ControlRigDeveloper/Private/Graph/ControlRigGraphNode.cpp
+++ b/Engine/Plugins/Experimental/ControlRig/Source/ControlRigDeveloper/Private/Graph/ControlRigGraphNode.cpp
@@ -184,13 +184,10 @@ bool UControlRigGraphNode::IsDeprecated() const
 	return Super::IsDeprecated();
 }
 
-bool UControlRigGraphNode::ShouldWarnOnDeprecation() const
+FEdGraphNodeDeprecationResponse UControlRigGraphNode::GetDeprecationResponse(EEdGraphNodeDeprecationType DeprecationType) const
 {
-	return true;
-}
+	FEdGraphNodeDeprecationResponse Response = Super::GetDeprecationResponse(DeprecationType);
 
-FString UControlRigGraphNode::GetDeprecationMessage() const
-{
 	UScriptStruct* ScriptStruct = GetUnitScriptStruct();
 	if (ScriptStruct)
 	{
@@ -198,10 +195,13 @@ FString UControlRigGraphNode::GetDeprecationMessage() const
 		ScriptStruct->GetStringMetaDataHierarchical(UControlRig::DeprecatedMetaName, &DeprecatedMetadata);
 		if (!DeprecatedMetadata.IsEmpty())
 		{
-			return FString::Printf(TEXT("Warning: This node is deprecated from: %s"), *DeprecatedMetadata);
+			FFormatNamedArguments Args;
+			Args.Add(TEXT("DeprecatedMetadata"), FText::FromString(DeprecatedMetadata));
+			Response.MessageText = FText::Format(LOCTEXT("ControlRigGraphNodeDeprecationMessage", "Warning: This node is deprecated from: {DeprecatedMetadata}"), Args);
 		}
 	}
-	return Super::GetDeprecationMessage();
+
+	return Response;
 }
 
 void UControlRigGraphNode::ReallocatePinsDuringReconstruction(const TArray<UEdGraphPin*>& OldPins)
diff --git a/Engine/Plugins/Experimental/ControlRig/Source/ControlRigDeveloper/Public/Graph/ControlRigGraphNode.h b/Engine/Plugins/Experimental/ControlRig/Source/ControlRigDeveloper/Public/Graph/ControlRigGraphNode.h
--- a/Engine/Plugins/Experimental/ControlRig/Source/ControlRigDeveloper/Public/Graph/ControlRigGraphNode.h
+++ b/Engine/Plugins/Experimental/ControlRig/Source/ControlRigDeveloper/Public/Graph/ControlRigGraphNode.h
@@ -177,9 +177,8 @@ public:
 	virtual void AutowireNewNode(UEdGraphPin* FromPin) override;	
 	virtual void PrepareForCopying() override;
 
-	virtual bool IsDeprecated() const;
-	virtual bool ShouldWarnOnDeprecation() const;
-	virtual FString GetDeprecationMessage() const;
+	virtual bool IsDeprecated() const override;
+	virtual FEdGraphNodeDeprecationResponse GetDeprecationResponse(EEdGraphNodeDeprecationType DeprecationType) const override;
 
 	/** Set the cached dimensions of this node */
 	void SetDimensions(const FVector2D& InDimensions) { Dimensions = InDimensions; }
diff --git a/Engine/Source/Editor/BlueprintGraph/Classes/K2Node_BaseMCDelegate.h b/Engine/Source/Editor/BlueprintGraph/Classes/K2Node_BaseMCDelegate.h
--- a/Engine/Source/Editor/BlueprintGraph/Classes/K2Node_BaseMCDelegate.h
+++ b/Engine/Source/Editor/BlueprintGraph/Classes/K2Node_BaseMCDelegate.h
@@ -33,8 +33,8 @@ public:
 	virtual bool HasExternalDependencies(TArray<class UStruct*>* OptionalOutput) const override;
 	virtual void GetNodeAttributes( TArray<TKeyValuePair<FString, FString>>& OutNodeAttributes ) const override;
 	virtual void AutowireNewNode(UEdGraphPin* FromPin) override;
-	virtual bool IsDeprecated() const override;
-	virtual FString GetDeprecationMessage() const override;
+	virtual bool HasDeprecatedReference() const override;
+	virtual FEdGraphNodeDeprecationResponse GetDeprecationResponse(EEdGraphNodeDeprecationType DeprecationType) const override;
 	// End of UK2Node interface
 
 	// UEdGraphNode interface
diff --git a/Engine/Source/Editor/BlueprintGraph/Classes/K2Node_CallFunction.h b/Engine/Source/Editor/BlueprintGraph/Classes/K2Node_CallFunction.h
--- a/Engine/Source/Editor/BlueprintGraph/Classes/K2Node_CallFunction.h
+++ b/Engine/Source/Editor/BlueprintGraph/Classes/K2Node_CallFunction.h
@@ -80,9 +80,8 @@ public:
 	virtual FText GetTooltipText() const override;
 	virtual FText GetNodeTitle(ENodeTitleType::Type TitleType) const override;
 	virtual FString GetDescriptiveCompiledName() const override;
-	virtual bool IsDeprecated() const override;
-	virtual bool ShouldWarnOnDeprecation() const override;
-	virtual FString GetDeprecationMessage() const override;
+	virtual bool HasDeprecatedReference() const override;
+	virtual FEdGraphNodeDeprecationResponse GetDeprecationResponse(EEdGraphNodeDeprecationType DeprecationType) const override;
 	virtual void PostPlacedNewNode() override;
 	virtual FString GetDocumentationLink() const override;
 	virtual FString GetDocumentationExcerptName() const override;
diff --git a/Engine/Source/Editor/BlueprintGraph/Classes/K2Node_CustomEvent.h b/Engine/Source/Editor/BlueprintGraph/Classes/K2Node_CustomEvent.h
--- a/Engine/Source/Editor/BlueprintGraph/Classes/K2Node_CustomEvent.h
+++ b/Engine/Source/Editor/BlueprintGraph/Classes/K2Node_CustomEvent.h
@@ -49,9 +49,8 @@ class UK2Node_CustomEvent : public UK2Node_Event
 	virtual void AutowireNewNode(UEdGraphPin* FromPin) override;
 	virtual void AddSearchMetaDataInfo(TArray<struct FSearchTagDataPair>& OutTaggedMetaData) const override;
 	virtual FText GetKeywords() const override;
-	virtual bool IsDeprecated() const override { return bIsDeprecated; }
-	virtual bool ShouldWarnOnDeprecation() const override;
-	virtual FString GetDeprecationMessage() const override;
+	virtual bool HasDeprecatedReference() const override { return bIsDeprecated; }
+	virtual FEdGraphNodeDeprecationResponse GetDeprecationResponse(EEdGraphNodeDeprecationType DeprecationType) const override;
 	//~ End UEdGraphNode Interface
 
 	//~ Begin UK2Node Interface
diff --git a/Engine/Source/Editor/BlueprintGraph/Classes/K2Node_Event.h b/Engine/Source/Editor/BlueprintGraph/Classes/K2Node_Event.h
--- a/Engine/Source/Editor/BlueprintGraph/Classes/K2Node_Event.h
+++ b/Engine/Source/Editor/BlueprintGraph/Classes/K2Node_Event.h
@@ -66,9 +66,8 @@ class UK2Node_Event : public UK2Node_EditablePinBase, public IK2Node_EventNodeIn
 	BLUEPRINTGRAPH_API virtual bool CanPasteHere(const UEdGraph* TargetGraph) const override;
 	BLUEPRINTGRAPH_API virtual bool IsCompatibleWithGraph(const UEdGraph* TargetGraph) const override;
 	BLUEPRINTGRAPH_API virtual FName GetCornerIcon() const override;
-	BLUEPRINTGRAPH_API virtual bool IsDeprecated() const override;
-	BLUEPRINTGRAPH_API virtual bool ShouldWarnOnDeprecation() const override;
-	BLUEPRINTGRAPH_API virtual FString GetDeprecationMessage() const override;
+	BLUEPRINTGRAPH_API virtual bool HasDeprecatedReference() const override;
+	BLUEPRINTGRAPH_API virtual FEdGraphNodeDeprecationResponse GetDeprecationResponse(EEdGraphNodeDeprecationType DeprecationType) const override;
 	BLUEPRINTGRAPH_API virtual UObject* GetJumpTargetForDoubleClick() const override;
 	BLUEPRINTGRAPH_API virtual FSlateIcon GetIconAndTint(FLinearColor& OutColor) const override;
 	BLUEPRINTGRAPH_API virtual FString GetFindReferenceSearchString() const override;
diff --git a/Engine/Source/Editor/BlueprintGraph/Classes/K2Node_FunctionEntry.h b/Engine/Source/Editor/BlueprintGraph/Classes/K2Node_FunctionEntry.h
--- a/Engine/Source/Editor/BlueprintGraph/Classes/K2Node_FunctionEntry.h
+++ b/Engine/Source/Editor/BlueprintGraph/Classes/K2Node_FunctionEntry.h
@@ -42,9 +42,8 @@ class UK2Node_FunctionEntry : public UK2Node_FunctionTerminator
 	virtual void AllocateDefaultPins() override;
 	virtual FText GetNodeTitle(ENodeTitleType::Type TitleType) const override;
 	virtual bool CanUserDeleteNode() const override { return false; }
-	virtual bool IsDeprecated() const override;
-	virtual bool ShouldWarnOnDeprecation() const override;
-	virtual FString GetDeprecationMessage() const override;
+	virtual bool HasDeprecatedReference() const override;
+	virtual FEdGraphNodeDeprecationResponse GetDeprecationResponse(EEdGraphNodeDeprecationType DeprecationType) const override;
 	virtual FText GetTooltipText() const override;
 	virtual void FindDiffs(UEdGraphNode* OtherNode, struct FDiffResults& Results) override;
 	//~ End UEdGraphNode Interface
diff --git a/Engine/Source/Editor/BlueprintGraph/Classes/K2Node_SpawnActor.h b/Engine/Source/Editor/BlueprintGraph/Classes/K2Node_SpawnActor.h
--- a/Engine/Source/Editor/BlueprintGraph/Classes/K2Node_SpawnActor.h
+++ b/Engine/Source/Editor/BlueprintGraph/Classes/K2Node_SpawnActor.h
@@ -24,8 +24,7 @@ class UK2Node_SpawnActor : public UK2Node
 	virtual void PinDefaultValueChanged(UEdGraphPin* Pin) override;
 	virtual FText GetTooltipText() const override;
 	virtual bool IsDeprecated() const override;
-	virtual bool ShouldWarnOnDeprecation() const override;
-	virtual FString GetDeprecationMessage() const override;
+	virtual FEdGraphNodeDeprecationResponse GetDeprecationResponse(EEdGraphNodeDeprecationType DeprecationType) const override;
 	virtual FSlateIcon GetIconAndTint(FLinearColor& OutColor) const override;
 	virtual bool IsCompatibleWithGraph(const UEdGraph* TargetGraph) const override;
 	//~ End UEdGraphNode Interface.
diff --git a/Engine/Source/Editor/BlueprintGraph/Classes/K2Node_Variable.h b/Engine/Source/Editor/BlueprintGraph/Classes/K2Node_Variable.h
--- a/Engine/Source/Editor/BlueprintGraph/Classes/K2Node_Variable.h
+++ b/Engine/Source/Editor/BlueprintGraph/Classes/K2Node_Variable.h
@@ -71,8 +71,8 @@ public:
 	virtual void AutowireNewNode(UEdGraphPin* FromPin) override;
 	virtual bool CanPasteHere(const UEdGraph* TargetGraph) const override;
 	virtual void PostPasteNode() override;
-	virtual bool IsDeprecated() const override;
-	virtual FString GetDeprecationMessage() const override;
+	virtual bool HasDeprecatedReference() const override;
+	virtual FEdGraphNodeDeprecationResponse GetDeprecationResponse(EEdGraphNodeDeprecationType DeprecationType) const override;
 	virtual UObject* GetJumpTargetForDoubleClick() const override;
 	virtual bool CanJumpToDefinition() const override;
 	virtual void JumpToDefinition() const override;
diff --git a/Engine/Source/Editor/BlueprintGraph/Private/EdGraphSchema_K2.cpp b/Engine/Source/Editor/BlueprintGraph/Private/EdGraphSchema_K2.cpp
--- a/Engine/Source/Editor/BlueprintGraph/Private/EdGraphSchema_K2.cpp
+++ b/Engine/Source/Editor/BlueprintGraph/Private/EdGraphSchema_K2.cpp
@@ -5428,6 +5428,11 @@ void UEdGraphSchema_K2::GetGraphDisplayInformation(const UEdGraph& Graph, /*out*
 		{
 			DisplayInfo.Notes.Add(TEXT("const"));
 		}
+
+		if (Function->HasMetaData(FBlueprintMetadata::MD_DeprecatedFunction))
+		{
+			DisplayInfo.Notes.Add(LOCTEXT("FunctionGraphDisplayInfo_Deprecated", "deprecated").ToString());
+		}
 	}
 
 	// Mark transient graphs as obviously so
diff --git a/Engine/Source/Editor/BlueprintGraph/Private/K2Node_CallFunction.cpp b/Engine/Source/Editor/BlueprintGraph/Private/K2Node_CallFunction.cpp
--- a/Engine/Source/Editor/BlueprintGraph/Private/K2Node_CallFunction.cpp
+++ b/Engine/Source/Editor/BlueprintGraph/Private/K2Node_CallFunction.cpp
@@ -437,31 +437,36 @@ UK2Node_CallFunction::UK2Node_CallFunction(const FObjectInitializer& ObjectIniti
 }
 
 
-bool UK2Node_CallFunction::IsDeprecated() const
+bool UK2Node_CallFunction::HasDeprecatedReference() const
 {
 	UFunction* Function = GetTargetFunction();
 	return (Function && Function->HasMetaData(FBlueprintMetadata::MD_DeprecatedFunction));
 }
 
-bool UK2Node_CallFunction::ShouldWarnOnDeprecation() const
+FEdGraphNodeDeprecationResponse UK2Node_CallFunction::GetDeprecationResponse(EEdGraphNodeDeprecationType DeprecationType) const
 {
-	// TEMP:  Do not warn in the case of SpawnActor, as we have a special upgrade path for those nodes
-	return (FunctionReference.GetMemberName() != FName(TEXT("BeginSpawningActorFromBlueprint")));
-}
-
-FString UK2Node_CallFunction::GetDeprecationMessage() const
-{
-	FText Result;
-	if (UFunction* Function = GetTargetFunction())
+	FEdGraphNodeDeprecationResponse Response = Super::GetDeprecationResponse(DeprecationType);
+	if (DeprecationType == EEdGraphNodeDeprecationType::NodeHasDeprecatedReference)
 	{
-		FString DetailedMessage = Function->GetMetaData(FBlueprintMetadata::MD_DeprecationMessage);
-		Result = FBlueprintEditorUtils::GetDeprecatedMemberUsageNodeWarning(GetUserFacingFunctionName(Function), FText::FromString(DetailedMessage));
+		// TEMP: Do not warn in the case of SpawnActor, as we have a special upgrade path for those nodes
+		if (FunctionReference.GetMemberName() == FName(TEXT("BeginSpawningActorFromBlueprint")))
+		{
+			Response.MessageType = EEdGraphNodeDeprecationMessageType::None;
+		}
+		else
+		{
+			UFunction* Function = GetTargetFunction();
+			if (ensureMsgf(Function != nullptr, TEXT("This node should not be able to report having a deprecated reference if the target function cannot be resolved.")))
+			{
+				FString DetailedMessage = Function->GetMetaData(FBlueprintMetadata::MD_DeprecationMessage);
+				Response.MessageText = FBlueprintEditorUtils::GetDeprecatedMemberUsageNodeWarning(GetUserFacingFunctionName(Function), FText::FromString(DetailedMessage));
+			}
+		}
 	}
 	
-	return Result.ToString();
+	return Response;
 }
 
-
 FText UK2Node_CallFunction::GetFunctionContextString() const
 {
 	FText ContextString;
diff --git a/Engine/Source/Editor/BlueprintGraph/Private/K2Node_CustomEvent.cpp b/Engine/Source/Editor/BlueprintGraph/Private/K2Node_CustomEvent.cpp
--- a/Engine/Source/Editor/BlueprintGraph/Private/K2Node_CustomEvent.cpp
+++ b/Engine/Source/Editor/BlueprintGraph/Private/K2Node_CustomEvent.cpp
@@ -544,17 +544,27 @@ FText UK2Node_CustomEvent::GetKeywords() const
 	return FText::Format(LOCTEXT("CustomEventKeywords", "{ParentKeywords} Custom"), Args);
 }
 
-bool UK2Node_CustomEvent::ShouldWarnOnDeprecation() const
+FEdGraphNodeDeprecationResponse UK2Node_CustomEvent::GetDeprecationResponse(EEdGraphNodeDeprecationType DeprecationType) const
 {
-	// Only warn on override usage. This allows the source event to be marked as deprecated in the class that defines it without warning.
-	return IsOverride();
-}
+	FEdGraphNodeDeprecationResponse Response = Super::GetDeprecationResponse(DeprecationType);
+	if (DeprecationType == EEdGraphNodeDeprecationType::NodeHasDeprecatedReference)
+	{
+		// Only warn on override usage.
+		if (IsOverride())
+		{
+			FText EventName = FText::FromName(GetFunctionName());
+			FText DetailedMessage = FText::FromString(DeprecationMessage);
+			Response.MessageText = FBlueprintEditorUtils::GetDeprecatedMemberUsageNodeWarning(EventName, DetailedMessage);
+		}
+		else
+		{
+			// Allow the source event to be marked as deprecated in the class that defines it without warning, but use a note to visually indicate that the definition itself has been deprecated.
+			Response.MessageType = EEdGraphNodeDeprecationMessageType::Note;
+			Response.MessageText = LOCTEXT("DeprecatedCustomEventMessage", "@@: This custom event has been marked as deprecated. It can be safely deleted if all references have been replaced or removed.");
+		}
+	}
 
-FString UK2Node_CustomEvent::GetDeprecationMessage() const
-{
-	FText EventName = FText::FromName(GetFunctionName());
-	FText DetailedMessage = FText::FromString(DeprecationMessage);
-	return FBlueprintEditorUtils::GetDeprecatedMemberUsageNodeWarning(EventName, DetailedMessage).ToString();
+	return Response;
 }
 
 #undef LOCTEXT_NAMESPACE
diff --git a/Engine/Source/Editor/BlueprintGraph/Private/K2Node_Event.cpp b/Engine/Source/Editor/BlueprintGraph/Private/K2Node_Event.cpp
--- a/Engine/Source/Editor/BlueprintGraph/Private/K2Node_Event.cpp
+++ b/Engine/Source/Editor/BlueprintGraph/Private/K2Node_Event.cpp
@@ -779,7 +779,7 @@ FText UK2Node_Event::GetMenuCategory() const
 	return FunctionCategory;
 }
 
-bool UK2Node_Event::IsDeprecated() const
+bool UK2Node_Event::HasDeprecatedReference() const
 {
 	if (UFunction* Function = EventReference.ResolveMember<UFunction>(GetBlueprintClassFromNode()))
 	{
@@ -789,23 +789,31 @@ bool UK2Node_Event::IsDeprecated() const
 	return false;
 }
 
-bool UK2Node_Event::ShouldWarnOnDeprecation() const
+FEdGraphNodeDeprecationResponse UK2Node_Event::GetDeprecationResponse(EEdGraphNodeDeprecationType DeprecationType) const
 {
-	// Only warn on override usage. This allows the source event to be marked as deprecated in the class that defines it without warning.
-	return bOverrideFunction;
-}
-
-FString UK2Node_Event::GetDeprecationMessage() const
-{
-	FText Result;
-	if (UFunction* Function = EventReference.ResolveMember<UFunction>(GetBlueprintClassFromNode()))
+	FEdGraphNodeDeprecationResponse Response = Super::GetDeprecationResponse(DeprecationType);
+	if (DeprecationType == EEdGraphNodeDeprecationType::NodeHasDeprecatedReference)
 	{
-		FName EventName = GetFunctionName();
-		FString DetailedMessage = Function->GetMetaData(FBlueprintMetadata::MD_DeprecationMessage);
-		Result = FBlueprintEditorUtils::GetDeprecatedMemberUsageNodeWarning(FText::FromName(EventName), FText::FromString(DetailedMessage));
+		// Only warn on override usage.
+		if (bOverrideFunction)
+		{
+			UFunction* Function = EventReference.ResolveMember<UFunction>(GetBlueprintClassFromNode());
+			if (ensureMsgf(Function != nullptr, TEXT("This node should not be able to report having a deprecated reference if the event override cannot be resolved.")))
+			{
+				FText EventName = FText::FromName(GetFunctionName());
+				FText DetailedMessage = FText::FromString(Function->GetMetaData(FBlueprintMetadata::MD_DeprecationMessage));
+				Response.MessageText = FBlueprintEditorUtils::GetDeprecatedMemberUsageNodeWarning(EventName, DetailedMessage);
+			}
+		}
+		else
+		{
+			// Allow the source event to be marked as deprecated in the class that defines it without warning, but use a note to visually indicate that the definition itself has been deprecated.
+			Response.MessageType = EEdGraphNodeDeprecationMessageType::Note;
+			Response.MessageText = LOCTEXT("DeprecatedEventMessage", "@@: This event has been marked as deprecated. It can be safely deleted if all references have been replaced or removed.");
+		}
 	}
 
-	return Result.ToString();
+	return Response;
 }
 
 UObject* UK2Node_Event::GetJumpTargetForDoubleClick() const
diff --git a/Engine/Source/Editor/BlueprintGraph/Private/K2Node_FunctionEntry.cpp b/Engine/Source/Editor/BlueprintGraph/Private/K2Node_FunctionEntry.cpp
--- a/Engine/Source/Editor/BlueprintGraph/Private/K2Node_FunctionEntry.cpp
+++ b/Engine/Source/Editor/BlueprintGraph/Private/K2Node_FunctionEntry.cpp
@@ -563,34 +563,44 @@ void UK2Node_FunctionEntry::GetRedirectPinNames(const UEdGraphPin& Pin, TArray<F
 	}
 }
 
-bool UK2Node_FunctionEntry::IsDeprecated() const
+bool UK2Node_FunctionEntry::HasDeprecatedReference() const
 {
 	// We only show deprecated for inherited functions
 	if (UFunction* const Function = FunctionReference.ResolveMember<UFunction>(GetBlueprintClassFromNode()))
 	{
 		return Function->HasMetaData(FBlueprintMetadata::MD_DeprecatedFunction);
 	}
-
-	return false;
-}
-
-bool UK2Node_FunctionEntry::ShouldWarnOnDeprecation() const
-{
-	// Only warn on non-editable (i.e. override) usage. This allows the source graph to be marked as deprecated in the class that defines it without warning.
-	return !IsEditable();
+	else
+	{
+		return MetaData.bIsDeprecated;
+	}
 }
 
-FString UK2Node_FunctionEntry::GetDeprecationMessage() const
+FEdGraphNodeDeprecationResponse UK2Node_FunctionEntry::GetDeprecationResponse(EEdGraphNodeDeprecationType DeprecationType) const
 {
-	FText Result;
-	if (UFunction* const Function = FunctionReference.ResolveMember<UFunction>(GetBlueprintClassFromNode()))
+	FEdGraphNodeDeprecationResponse Response = Super::GetDeprecationResponse(DeprecationType);
+	if (DeprecationType == EEdGraphNodeDeprecationType::NodeHasDeprecatedReference)
 	{
-		FName FunctionName = FunctionReference.GetMemberName();
-		FString DetailedMessage = Function->GetMetaData(FBlueprintMetadata::MD_DeprecationMessage);
-		Result = FBlueprintEditorUtils::GetDeprecatedMemberUsageNodeWarning(FText::FromName(FunctionName), FText::FromString(DetailedMessage));
+		// Only warn on non-editable (i.e. override) usage.
+		if (!IsEditable())
+		{
+			UFunction* const Function = FunctionReference.ResolveMember<UFunction>(GetBlueprintClassFromNode());
+			if (ensureMsgf(Function != nullptr, TEXT("This node should not be able to report having a deprecated reference if the override function cannot be resolved.")))
+			{
+				FText FunctionName = FText::FromName(FunctionReference.GetMemberName());
+				FText DetailedMessage = FText::FromString(Function->GetMetaData(FBlueprintMetadata::MD_DeprecationMessage));
+				Response.MessageText = FBlueprintEditorUtils::GetDeprecatedMemberUsageNodeWarning(FunctionName, DetailedMessage);
+			}
+		}
+		else
+		{
+			// Allow the function to be marked as deprecated in the class that defines it without warning, but use a note to visually indicate that the definition itself has been deprecated.
+			Response.MessageType = EEdGraphNodeDeprecationMessageType::Note;
+			Response.MessageText = LOCTEXT("DeprecatedFunctionMessage", "@@: This function has been marked as deprecated. It can be safely deleted if all references have been replaced or removed.");
+		}
 	}
 
-	return Result.ToString();
+	return Response;
 }
 
 FText UK2Node_FunctionEntry::GetTooltipText() const
diff --git a/Engine/Source/Editor/BlueprintGraph/Private/K2Node_MCDelegate.cpp b/Engine/Source/Editor/BlueprintGraph/Private/K2Node_MCDelegate.cpp
--- a/Engine/Source/Editor/BlueprintGraph/Private/K2Node_MCDelegate.cpp
+++ b/Engine/Source/Editor/BlueprintGraph/Private/K2Node_MCDelegate.cpp
@@ -270,28 +270,25 @@ void UK2Node_BaseMCDelegate::AutowireNewNode(UEdGraphPin* FromPin)
 	}
 }
 
-bool UK2Node_BaseMCDelegate::IsDeprecated() const
+bool UK2Node_BaseMCDelegate::HasDeprecatedReference() const
 {
-	// Check if either the node itself or the referenced delegate is deprecated.
-	return Super::IsDeprecated() || DelegateReference.IsDeprecated();
+	// Check if the referenced delegate is deprecated.
+	return DelegateReference.IsDeprecated();
 }
 
-FString UK2Node_BaseMCDelegate::GetDeprecationMessage() const
+FEdGraphNodeDeprecationResponse UK2Node_BaseMCDelegate::GetDeprecationResponse(EEdGraphNodeDeprecationType DeprecationType) const
 {
-	// Handle the default case where the node itself has been deprecated (return the default message).
-	if (Super::IsDeprecated())
+	FEdGraphNodeDeprecationResponse Response = Super::GetDeprecationResponse(DeprecationType);
+	if (DeprecationType == EEdGraphNodeDeprecationType::NodeHasDeprecatedReference)
 	{
-		return Super::GetDeprecationMessage();
-	}
-
-	FText Result;
-	if (UProperty* DelegateProperty = DelegateReference.ResolveMember<UProperty>(GetBlueprintClassFromNode()))
-	{
-		FString DetailedMessage = DelegateProperty->GetMetaData(FBlueprintMetadata::MD_DeprecationMessage);
-		Result = FBlueprintEditorUtils::GetDeprecatedMemberUsageNodeWarning(GetPropertyDisplayName(), FText::FromString(DetailedMessage));
+		if (UProperty* DelegateProperty = DelegateReference.ResolveMember<UProperty>(GetBlueprintClassFromNode()))
+		{
+			FString DetailedMessage = DelegateProperty->GetMetaData(FBlueprintMetadata::MD_DeprecationMessage);
+			Response.MessageText = FBlueprintEditorUtils::GetDeprecatedMemberUsageNodeWarning(GetPropertyDisplayName(), FText::FromString(DetailedMessage));
+		}
 	}
-
-	return Result.ToString();
+	
+	return Response;
 }
 
 /////// UK2Node_AddDelegate ///////////
diff --git a/Engine/Source/Editor/BlueprintGraph/Private/K2Node_SpawnActor.cpp b/Engine/Source/Editor/BlueprintGraph/Private/K2Node_SpawnActor.cpp
--- a/Engine/Source/Editor/BlueprintGraph/Private/K2Node_SpawnActor.cpp
+++ b/Engine/Source/Editor/BlueprintGraph/Private/K2Node_SpawnActor.cpp
@@ -494,14 +494,16 @@ bool UK2Node_SpawnActor::IsDeprecated() const
 	return false;
 }
 
-bool UK2Node_SpawnActor::ShouldWarnOnDeprecation() const
+FEdGraphNodeDeprecationResponse UK2Node_SpawnActor::GetDeprecationResponse(EEdGraphNodeDeprecationType DeprecationType) const
 {
-	return false;
-}
+	FEdGraphNodeDeprecationResponse Response = Super::GetDeprecationResponse(DeprecationType);
+	if (DeprecationType == EEdGraphNodeDeprecationType::NodeTypeIsDeprecated)
+	{
+		Response.MessageType = EEdGraphNodeDeprecationMessageType::None;
+		Response.MessageText = LOCTEXT("SpawnActorNodeOnlyDefaultBlueprint_Deprecatio", "Spawn Actor @@ is DEPRECATED and should be replaced by SpawnActorFromClass");
+	}
 
-FString UK2Node_SpawnActor::GetDeprecationMessage() const
-{
-	return LOCTEXT("SpawnActorNodeOnlyDefaultBlueprint_Deprecatio", "Spawn Actor @@ is DEPRECATED and should be replaced by SpawnActorFromClass").ToString();
+	return Response;
 }
 
 FSlateIcon UK2Node_SpawnActor::GetIconAndTint(FLinearColor& OutColor) const
diff --git a/Engine/Source/Editor/BlueprintGraph/Private/K2Node_Variable.cpp b/Engine/Source/Editor/BlueprintGraph/Private/K2Node_Variable.cpp
--- a/Engine/Source/Editor/BlueprintGraph/Private/K2Node_Variable.cpp
+++ b/Engine/Source/Editor/BlueprintGraph/Private/K2Node_Variable.cpp
@@ -952,10 +952,10 @@ void UK2Node_Variable::PostPasteNode()
 	}
 }
 
-bool UK2Node_Variable::IsDeprecated() const
+bool UK2Node_Variable::HasDeprecatedReference() const
 {
-	// Check if either the node itself or the referenced variable is deprecated.
-	if (Super::IsDeprecated() || VariableReference.IsDeprecated())
+	// Check if the referenced variable is deprecated.
+	if (VariableReference.IsDeprecated())
 	{
 		return true;
 	}
@@ -972,23 +972,20 @@ bool UK2Node_Variable::IsDeprecated() const
 	return false;
 }
 
-FString UK2Node_Variable::GetDeprecationMessage() const
+FEdGraphNodeDeprecationResponse UK2Node_Variable::GetDeprecationResponse(EEdGraphNodeDeprecationType DeprecationType) const
 {
-	// Handle the default case where the node itself has been deprecated (return the default message).
-	if (Super::IsDeprecated())
+	FEdGraphNodeDeprecationResponse Response = Super::GetDeprecationResponse(DeprecationType);
+	if (DeprecationType == EEdGraphNodeDeprecationType::NodeHasDeprecatedReference)
 	{
-		return Super::GetDeprecationMessage();
+		if (UProperty* VariableProperty = VariableReference.ResolveMember<UProperty>(GetBlueprintClassFromNode()))
+		{
+			FText MemberName = FText::FromName(VariableReference.GetMemberName());
+			FText DetailedMessage = FText::FromString(VariableProperty->GetMetaData(FBlueprintMetadata::MD_DeprecationMessage));
+			Response.MessageText = FBlueprintEditorUtils::GetDeprecatedMemberUsageNodeWarning(MemberName, DetailedMessage);
+		}
 	}
 
-	FText Result;
-	if (UProperty* VariableProperty = VariableReference.ResolveMember<UProperty>(GetBlueprintClassFromNode()))
-	{
-		FName MemberName = VariableReference.GetMemberName();
-		FString DetailedMessage = VariableProperty->GetMetaData(FBlueprintMetadata::MD_DeprecationMessage);
-		Result = FBlueprintEditorUtils::GetDeprecatedMemberUsageNodeWarning(FText::FromName(MemberName), FText::FromString(DetailedMessage));
-	}
-	
-	return Result.ToString();
+	return Response;
 }
 
 UObject* UK2Node_Variable::GetJumpTargetForDoubleClick() const
diff --git a/Engine/Source/Editor/KismetCompiler/Private/EdGraphCompiler.cpp b/Engine/Source/Editor/KismetCompiler/Private/EdGraphCompiler.cpp
--- a/Engine/Source/Editor/KismetCompiler/Private/EdGraphCompiler.cpp
+++ b/Engine/Source/Editor/KismetCompiler/Private/EdGraphCompiler.cpp
@@ -65,9 +65,24 @@ void FGraphCompilerContext::ValidatePin(const UEdGraphPin* Pin) const
 /** Validates that the node is schema compatible */
 void FGraphCompilerContext::ValidateNode(const UEdGraphNode* Node) const
 {
-	if (Node->IsDeprecated() && Node->ShouldWarnOnDeprecation())
+	const bool bIsNodeDeprecated = Node->IsDeprecated();
+	if (bIsNodeDeprecated || Node->HasDeprecatedReference())
 	{
-		MessageLog.Warning(*Node->GetDeprecationMessage(), Node);
+		EEdGraphNodeDeprecationType ResponseType = bIsNodeDeprecated ? EEdGraphNodeDeprecationType::NodeTypeIsDeprecated : EEdGraphNodeDeprecationType::NodeHasDeprecatedReference;
+		FEdGraphNodeDeprecationResponse Response = Node->GetDeprecationResponse(ResponseType);
+		if (Response.MessageType != EEdGraphNodeDeprecationMessageType::None && !Response.MessageText.IsEmpty())
+		{
+			switch (Response.MessageType)
+			{
+			case EEdGraphNodeDeprecationMessageType::Note:
+				MessageLog.Note(*Response.MessageText.ToString(), Node);
+				break;
+
+			case EEdGraphNodeDeprecationMessageType::Warning:
+				MessageLog.Warning(*Response.MessageText.ToString(), Node);
+				break;
+			}
+		}
 	}
 
 	for (int32 PinIndex = 0; PinIndex < Node->Pins.Num(); ++PinIndex)
diff --git a/Engine/Source/Editor/KismetCompiler/Private/KismetCompiler.cpp b/Engine/Source/Editor/KismetCompiler/Private/KismetCompiler.cpp
--- a/Engine/Source/Editor/KismetCompiler/Private/KismetCompiler.cpp
+++ b/Engine/Source/Editor/KismetCompiler/Private/KismetCompiler.cpp
@@ -3148,7 +3148,7 @@ void FKismetCompilerContext::VerifyValidOverrideEvent(const UEdGraph* Graph)
 					(EventNode->EventReference.GetMemberParentClass(EventNode->GetBlueprintClassFromNode()) == FuncClass) &&
 					(EventNode->EventReference.GetMemberName() == FuncName))
 				{
-					if (EventNode->IsDeprecated())
+					if (EventNode->HasDeprecatedReference())
 					{
 						// The event cannot be placed because it has been deprecated. However, we already emit a
 						// warning in FGraphCompilerContext::ValidateNode(), so there's no need to repeat it here.
diff --git a/Engine/Source/Runtime/Engine/Classes/EdGraph/EdGraphNode.h b/Engine/Source/Runtime/Engine/Classes/EdGraph/EdGraphNode.h
--- a/Engine/Source/Runtime/Engine/Classes/EdGraph/EdGraphNode.h
+++ b/Engine/Source/Runtime/Engine/Classes/EdGraph/EdGraphNode.h
@@ -170,6 +170,36 @@ struct FGraphNodeContextMenuBuilder
 	FGraphNodeContextMenuBuilder(const UEdGraph* InGraph, const UEdGraphNode* InNode, const UEdGraphPin* InPin, class FMenuBuilder* InMenuBuilder, bool bInDebuggingMode);
 };
 
+/** Deprecation types for node response. */
+enum class EEdGraphNodeDeprecationType
+{
+	/** The node type is deprecated. */
+	NodeTypeIsDeprecated,
+	/** The node references a deprecated member or type (e.g. variable or function). */
+	NodeHasDeprecatedReference
+};
+
+/** Deprecation response message types. */
+enum class EEdGraphNodeDeprecationMessageType
+{
+	/** No message. The Blueprint will compile successfully. */
+	None,
+	/** Emit the message as a note at compile time. This will appear as a note on the node and in the compiler log. */
+	Note,
+	/** Emit the message as a Blueprint compiler warning. This will appear as a warning on the node and in the compiler log. */
+	Warning
+};
+
+/** Deprecation response data. */
+struct FEdGraphNodeDeprecationResponse
+{
+	/** Message type. */
+	EEdGraphNodeDeprecationMessageType MessageType = EEdGraphNodeDeprecationMessageType::None;
+
+	/** Message text to display on the node and/or emit to the compile log. */
+	FText MessageText;
+};
+
 UCLASS()
 class ENGINE_API UEdGraphNode : public UObject
 {
@@ -723,10 +753,18 @@ public:
 	// Returns true if this node is deprecated
 	virtual bool IsDeprecated() const;
 
+	// Returns true if this node references a deprecated type or member
+	virtual bool HasDeprecatedReference() const { return false; }
+
+	// Returns the response to use when reporting a deprecation
+	virtual FEdGraphNodeDeprecationResponse GetDeprecationResponse(EEdGraphNodeDeprecationType DeprecationType) const;
+
 	// Returns true if this node should produce a compiler warning on deprecation
-	virtual bool ShouldWarnOnDeprecation() const { return true; }
+	UE_DEPRECATED(4.23, "Use GetDeprecationResponse instead.")
+	virtual bool ShouldWarnOnDeprecation() const;
 
 	// Returns the string to use when reporting the deprecation
+	UE_DEPRECATED(4.23, "Use GetDeprecationResponse instead.")
 	virtual FString GetDeprecationMessage() const;
 
 	// Returns the object that should be focused when double-clicking on this node
diff --git a/Engine/Source/Runtime/Engine/Private/EdGraph/EdGraphNode.cpp b/Engine/Source/Runtime/Engine/Private/EdGraph/EdGraphNode.cpp
--- a/Engine/Source/Runtime/Engine/Private/EdGraph/EdGraphNode.cpp
+++ b/Engine/Source/Runtime/Engine/Private/EdGraph/EdGraphNode.cpp
@@ -535,9 +535,54 @@ bool UEdGraphNode::IsDeprecated() const
 	return GetClass()->HasAnyClassFlags(CLASS_Deprecated);
 }
 
+FEdGraphNodeDeprecationResponse UEdGraphNode::GetDeprecationResponse(EEdGraphNodeDeprecationType DeprecationType) const
+{
+	FEdGraphNodeDeprecationResponse Response;
+
+	if (DeprecationType == EEdGraphNodeDeprecationType::NodeTypeIsDeprecated)
+	{
+		Response.MessageType = EEdGraphNodeDeprecationMessageType::Warning;
+		Response.MessageText = NSLOCTEXT("EdGraphCompiler", "NodeDeprecated_Warning", "@@ is deprecated; please replace or remove it.");
+	}
+	else if (DeprecationType == EEdGraphNodeDeprecationType::NodeHasDeprecatedReference)
+	{
+		Response.MessageType = EEdGraphNodeDeprecationMessageType::Warning;
+		Response.MessageText = NSLOCTEXT("EdGraphCompiler", "NodeDeprecatedReference_Note", "@@ has a deprecated reference; please replace or remove it.");
+	}
+
+	return Response;
+}
+
+// Deprecated; implemented for backwards-compatibility.
+bool UEdGraphNode::ShouldWarnOnDeprecation() const
+{
+	if (IsDeprecated())
+	{
+		return GetDeprecationResponse(EEdGraphNodeDeprecationType::NodeTypeIsDeprecated).MessageType == EEdGraphNodeDeprecationMessageType::Warning;
+	}
+	else if (HasDeprecatedReference())
+	{
+		return GetDeprecationResponse(EEdGraphNodeDeprecationType::NodeHasDeprecatedReference).MessageType == EEdGraphNodeDeprecationMessageType::Warning;
+	}
+
+	return false;
+}
+
+// Deprecated; implemented for backwards-compatibility.
 FString UEdGraphNode::GetDeprecationMessage() const
 {
-	return NSLOCTEXT("EdGraphCompiler", "NodeDeprecated_Warning", "@@ is deprecated; please replace or remove it.").ToString();
+	FText MessageText;
+
+	if (IsDeprecated())
+	{
+		MessageText = GetDeprecationResponse(EEdGraphNodeDeprecationType::NodeTypeIsDeprecated).MessageText;
+	}
+	else if (HasDeprecatedReference())
+	{
+		MessageText = GetDeprecationResponse(EEdGraphNodeDeprecationType::NodeHasDeprecatedReference).MessageText;
+	}
+
+	return MessageText.ToString();
 }
 
 void UEdGraphNode::AddReferencedObjects(UObject* InThis, FReferenceCollector& Collector)
