[NO NON-PII COMMIT MESSAGE]

diff --git a/Engine/Source/Editor/BlueprintGraph/Private/K2Node.cpp b/Engine/Source/Editor/BlueprintGraph/Private/K2Node.cpp
--- a/Engine/Source/Editor/BlueprintGraph/Private/K2Node.cpp
+++ b/Engine/Source/Editor/BlueprintGraph/Private/K2Node.cpp
@@ -44,13 +44,6 @@ UK2Node::UK2Node(const FObjectInitializer& ObjectInitializer)
 void UK2Node::PostLoad()
 {
 #if WITH_EDITORONLY_DATA
-
-	if (GIsEditor)
-	{
-		// We need to serialize string data references on load in editor builds so the cooker knows about them
-		FixupPinStringDataReferences(nullptr);
-	}
-
 	// Clean up win watches for any deprecated pins we are about to remove in Super::PostLoad
 	if (DeprecatedPins.Num() && HasValidBlueprint())
 	{
@@ -97,6 +90,12 @@ void UK2Node::Serialize(FArchive& Ar)
 	{
 		// Fix up pin default values, must be done before post load
 		FixupPinDefaultValues();
+
+		if (GIsEditor)
+		{
+			// We need to serialize string data references on load in editor builds so the cooker knows about them
+			FixupPinStringDataReferences(nullptr);
+		}
 	}
 }
 
@@ -162,8 +161,15 @@ void UK2Node::FixupPinStringDataReferences(FArchive* SavingArchive)
 {
 	FSoftObjectPathSerializationScope SetPackage(GetOutermost()->GetFName(), NAME_None, ESoftObjectPathCollectType::AlwaysCollect, ESoftObjectPathSerializeType::SkipSerializeIfArchiveHasSize);
 
-	// This code expans some pin types into their real representation, optionally serializes them, and then writes them out as strings again
+	// This code expands some pin types into their real representation, optionally serializes them, and then writes them out as strings again
 	const UEdGraphSchema_K2* K2Schema = GetDefault<UEdGraphSchema_K2>();
+	FLinkerLoad* LinkerLoad = GetLinker();
+
+	// Can't do any fixups without an archive of some sort
+	if (!SavingArchive && !LinkerLoad)
+	{
+		return;
+	}
 
 	for (UEdGraphPin* Pin : Pins)
 	{
@@ -175,6 +181,12 @@ void UK2Node::FixupPinStringDataReferences(FArchive* SavingArchive)
 				UScriptStruct* Struct = Cast<UScriptStruct>(Pin->PinType.PinSubCategoryObject.Get());
 				if (Struct)
 				{
+					// While loading, our user struct may not have been linked yet
+					if (Struct->HasAnyFlags(RF_NeedLoad) && LinkerLoad)
+					{
+						LinkerLoad->Preload(Struct);
+					}
+
 					TSharedPtr<FStructOnScope> StructData = MakeShareable(new FStructOnScope(Struct));
 					StructData->SetPackage(GetOutermost());
 
@@ -211,7 +223,7 @@ void UK2Node::FixupPinStringDataReferences(FArchive* SavingArchive)
 				else
 				{
 					// Transform it as strings
-					TempRef.PostLoadPath(GetLinker());
+					TempRef.PostLoadPath(LinkerLoad);
 					TempRef.PreSavePath();
 				}
 				
