Fixed anim notifies not showing up if you open a normal blueprint first

diff --git a/Engine/Source/Editor/BlueprintGraph/Private/BlueprintActionFilter.cpp b/Engine/Source/Editor/BlueprintGraph/Private/BlueprintActionFilter.cpp
--- a/Engine/Source/Editor/BlueprintGraph/Private/BlueprintActionFilter.cpp
+++ b/Engine/Source/Editor/BlueprintGraph/Private/BlueprintActionFilter.cpp
@@ -383,6 +383,15 @@ namespace BlueprintActionFilterImpl
 	 */
 	static bool IsIncompatibleAnimNotification( FBlueprintActionFilter const& Filter, FBlueprintActionInfo& BlueprintAction); 
 
+	/**
+	 * Rejection test that checks if an anim notify event is in an anim blueprint
+	 * 
+	 * @param  Filter			Holds the graph context for this test.
+	 * @param  BlueprintAction	The action you wish to query.
+	 * @return True if the action is an anim notify and the current blueprint is not an anim blueprint
+	 */
+	static bool IsAnimNotificationInAnimBlueprint( FBlueprintActionFilter const& Filter, FBlueprintActionInfo& BlueprintAction); 
+
 	/**
 	 * 
 	 * 
@@ -1563,6 +1572,22 @@ static bool BlueprintActionFilterImpl::IsNodeTemplateSelfFiltered(FBlueprintActi
 	return BlueprintAction.NodeSpawner->IsTemplateNodeFilteredOut(Filter);
 }
 
+//------------------------------------------------------------------------------
+static bool BlueprintActionFilterImpl::IsAnimNotificationInAnimBlueprint(FBlueprintActionFilter const& Filter, FBlueprintActionInfo& BlueprintAction)
+{
+	bool bIsFilteredOut = false;
+
+	if(!BlueprintAction.GetNodeClass()->IsChildOf<UK2Node_Event>())
+	{
+		if (const UAnimBlueprint* AnimBlueprintOwner = Cast<UAnimBlueprint>(BlueprintAction.GetActionOwner()))
+		{
+			bIsFilteredOut = !(Filter.Context.Blueprints.Num() == 1 && Filter.Context.Blueprints[0] == AnimBlueprintOwner);
+		}
+	}
+
+	return bIsFilteredOut;
+}
+
 //------------------------------------------------------------------------------
 static bool BlueprintActionFilterImpl::IsIncompatibleAnimNotification(FBlueprintActionFilter const& Filter, FBlueprintActionInfo& BlueprintAction)
 {
@@ -1904,6 +1929,7 @@ FBlueprintActionFilter::FBlueprintActionFilter(uint32 Flags/*= 0x00*/)
 	//
 	// this test in-particular spawns a template-node and then calls 
 	// AllocateDefaultPins() which is costly, so it should be very last!
+	AddRejectionTest(FRejectionTestDelegate::CreateStatic(IsAnimNotificationInAnimBlueprint));
 	AddRejectionTest(FRejectionTestDelegate::CreateStatic(IsIncompatibleAnimNotification));
 	AddRejectionTest(FRejectionTestDelegate::CreateStatic(IsNodeTemplateSelfFiltered));
 	AddRejectionTest(FRejectionTestDelegate::CreateStatic(IsMissingMatchingPinParam));
