Update the RestrictedToClasses metadata specifier to use a "," delimeter instead of spaces and ahve it support multiple class listings.

diff --git a/Engine/Source/Editor/BlueprintGraph/Private/BlueprintActionFilter.cpp b/Engine/Source/Editor/BlueprintGraph/Private/BlueprintActionFilter.cpp
--- a/Engine/Source/Editor/BlueprintGraph/Private/BlueprintActionFilter.cpp
+++ b/Engine/Source/Editor/BlueprintGraph/Private/BlueprintActionFilter.cpp
@@ -771,6 +771,15 @@ static bool BlueprintActionFilterImpl::IsRestrictedClassMember(FBlueprintActionF
 		if (ActionClass->HasMetaData(FBlueprintMetadata::MD_RestrictedToClasses))
 		{
 			FString const& ClassRestrictions = ActionClass->GetMetaData(FBlueprintMetadata::MD_RestrictedToClasses);
+			
+			// Parse the the metadata into an array that is delimited by ',' and trim whitespace
+			TArray<FString> ParsedClassRestrictions;
+			ClassRestrictions.ParseIntoArray(ParsedClassRestrictions, TEXT(","));
+			for (FString& ValidClassName : ParsedClassRestrictions)
+			{
+				ValidClassName = ValidClassName.TrimStartAndEnd();
+			}
+
 			for (UBlueprint const* TargetContext : FilterContext.Blueprints)
 			{
 				UClass* TargetClass = TargetContext->GeneratedClass;
@@ -778,7 +787,7 @@ static bool BlueprintActionFilterImpl::IsRestrictedClassMember(FBlueprintActionF
 				{
 					// Skip possible null classes (e.g. macros, etc)
 					continue;
-				};
+				}
 
 				bool bIsClassListed = false;
 				
@@ -788,7 +797,15 @@ static bool BlueprintActionFilterImpl::IsRestrictedClassMember(FBlueprintActionF
 				while (!bIsClassListed && (QueryClass != nullptr))
 				{
 					FString const ClassName = QueryClass->GetName();
-					bIsClassListed = (ClassName == ClassRestrictions) || !!FCString::StrfindDelim(*ClassRestrictions, *ClassName, TEXT(" "));
+					// If this class is on the list of valid classes
+					for (const FString& ValidClassName : ParsedClassRestrictions)
+					{
+						bIsClassListed = (ClassName == ValidClassName);
+						if (bIsClassListed)
+						{
+							break;
+						}
+					}
 					
 					QueryClass = QueryClass->GetSuperClass();
 				}
