Prevent multi-gate from having more than 32 entries.

diff --git a/Engine/Source/Editor/BlueprintGraph/Classes/K2Node_ExecutionSequence.h b/Engine/Source/Editor/BlueprintGraph/Classes/K2Node_ExecutionSequence.h
--- a/Engine/Source/Editor/BlueprintGraph/Classes/K2Node_ExecutionSequence.h
+++ b/Engine/Source/Editor/BlueprintGraph/Classes/K2Node_ExecutionSequence.h
@@ -37,7 +37,7 @@ class UK2Node_ExecutionSequence : public UK2Node, public IK2Node_AddPinInterface
 	//~ End UK2Node Interface
 
 	// IK2Node_AddPinInterface interface
-	BLUEPRINTGRAPH_API virtual void AddInputPin() override;
+	virtual void AddInputPin() override;
 	// End of IK2Node_AddPinInterface interface
 
 	//~ Begin K2Node_ExecutionSequence Interface
@@ -45,12 +45,6 @@ class UK2Node_ExecutionSequence : public UK2Node, public IK2Node_AddPinInterface
 	/** Gets a unique pin name, the next in the sequence */
 	FName GetUniquePinName();
 
-	/**
-	 * Adds a new execution pin to an execution node
-	 */
-	UE_DEPRECATED(4.17, "Use AddInputPin instead.")
-	BLUEPRINTGRAPH_API void AddPinToExecutionNode() { AddInputPin(); }
-
 	/**
 	 * Inserts a new execution pin, before the specified execution pin, into an execution node
 	 *
diff --git a/Engine/Source/Editor/BlueprintGraph/Classes/K2Node_MultiGate.h b/Engine/Source/Editor/BlueprintGraph/Classes/K2Node_MultiGate.h
--- a/Engine/Source/Editor/BlueprintGraph/Classes/K2Node_MultiGate.h
+++ b/Engine/Source/Editor/BlueprintGraph/Classes/K2Node_MultiGate.h
@@ -34,29 +34,34 @@ class UK2Node_MultiGate : public UK2Node_ExecutionSequence
 	virtual void GetMenuActions(FBlueprintActionDatabaseRegistrar& ActionRegistrar) const override;
 	//~ End UK2Node Interface
 
+	//~ IK2Node_AddPinInterface interface
+	virtual bool CanAddPin() const override;
+	//~ End IK2Node_AddPinInterface interface
+
 	/** Getting pin access */
-	BLUEPRINTGRAPH_API UEdGraphPin* GetResetPin() const;
-	BLUEPRINTGRAPH_API UEdGraphPin* GetIsRandomPin() const;
-	BLUEPRINTGRAPH_API UEdGraphPin* GetLoopPin() const;
-	BLUEPRINTGRAPH_API UEdGraphPin* GetStartIndexPin() const;
-	BLUEPRINTGRAPH_API void GetOutPins(TArray<UEdGraphPin*>& OutPins) const;
+	UEdGraphPin* GetResetPin() const;
+	UEdGraphPin* GetIsRandomPin() const;
+	UEdGraphPin* GetLoopPin() const;
+	UEdGraphPin* GetStartIndexPin() const;
+ 	void GetOutPins(TArray<UEdGraphPin*>& OutPins) const;
+	int32 GetNumOutPins() const;
 
 	/** Gets the name and class of the MarkBit function from the KismetNodeHelperLibrary */
-	BLUEPRINTGRAPH_API void GetMarkBitFunction(FName& FunctionName, UClass** FunctionClass);
+	void GetMarkBitFunction(FName& FunctionName, UClass** FunctionClass);
 	/** Gets the name and class of the HasUnmarkedBit function from the KismetNodeHelperLibrary */
-	BLUEPRINTGRAPH_API void GetHasUnmarkedBitFunction(FName& FunctionName, UClass** FunctionClass);
+	void GetHasUnmarkedBitFunction(FName& FunctionName, UClass** FunctionClass);
 	/** Gets the name and class of the MarkFirstUnmarkedBit function from the KismetNodeHelperLibrary */
-	BLUEPRINTGRAPH_API void GetUnmarkedBitFunction(FName& FunctionName, UClass** FunctionClass);
+	void GetUnmarkedBitFunction(FName& FunctionName, UClass** FunctionClass);
 	/** Gets the name and class of the Greater_IntInt function from the KismetMathLibrary */
-	BLUEPRINTGRAPH_API void GetConditionalFunction(FName& FunctionName, UClass** FunctionClass);
+	void GetConditionalFunction(FName& FunctionName, UClass** FunctionClass);
 	/** Gets the name and class of the EqualEqual_IntInt function from the KismetMathLibrary */
-	BLUEPRINTGRAPH_API void GetEqualityFunction(FName& FunctionName, UClass** FunctionClass);
+	void GetEqualityFunction(FName& FunctionName, UClass** FunctionClass);
 	/** Gets the name and class of the NotEqual_BoolBool function from the KismetMathLibrary */
-	BLUEPRINTGRAPH_API void GetBoolNotEqualFunction(FName& FunctionName, UClass** FunctionClass);
+	void GetBoolNotEqualFunction(FName& FunctionName, UClass** FunctionClass);
 	/** Gets the name and class of the PrintString function */
-	BLUEPRINTGRAPH_API void GetPrintStringFunction(FName& FunctionName, UClass** FunctionClass);
+	void GetPrintStringFunction(FName& FunctionName, UClass** FunctionClass);
 	/** Gets the name and class of the ClearAllBits function from the KismetNodeHelperLibrary */
-	BLUEPRINTGRAPH_API void GetClearAllBitsFunction(FName& FunctionName, UClass** FunctionClass);
+	void GetClearAllBitsFunction(FName& FunctionName, UClass** FunctionClass);
 
 private:
 	// Returns the exec output pin name for a given 0-based index
diff --git a/Engine/Source/Editor/BlueprintGraph/Private/K2Node_MultiGate.cpp b/Engine/Source/Editor/BlueprintGraph/Private/K2Node_MultiGate.cpp
--- a/Engine/Source/Editor/BlueprintGraph/Private/K2Node_MultiGate.cpp
+++ b/Engine/Source/Editor/BlueprintGraph/Private/K2Node_MultiGate.cpp
@@ -493,6 +493,12 @@ FText UK2Node_MultiGate::GetNodeTitle(ENodeTitleType::Type TitleType) const
 	return NSLOCTEXT("K2Node", "MultiGate", "MultiGate");
 }
 
+bool UK2Node_MultiGate::CanAddPin() const
+{
+	// Since an int32 is used to track which gates have been used, we can only have 32 outputs
+	return GetNumOutPins() < 32;
+}
+
 void UK2Node_MultiGate::AllocateDefaultPins()
 {
 	Super::AllocateDefaultPins();
@@ -547,18 +553,32 @@ UEdGraphPin* UK2Node_MultiGate::GetStartIndexPin() const
 
 void UK2Node_MultiGate::GetOutPins(TArray<UEdGraphPin*>& OutPins) const
 {
-	OutPins.Reset();
+	OutPins.Reset(Pins.Num());
 
-	const TCHAR* OutStr(TEXT("Out"));
 	for (UEdGraphPin* Pin : Pins)
 	{
-		if (Pin->PinName.ToString().StartsWith(OutStr))
+		if (Pin->Direction == EGPD_Output)
 		{
 			OutPins.Add(Pin);
 		}
 	}
 }
 
+int32 UK2Node_MultiGate::GetNumOutPins() const
+{
+	int32 NumOutPins = 0;
+
+	for (UEdGraphPin* Pin : Pins)
+	{
+		if (Pin->Direction == EGPD_Output)
+		{
+			++NumOutPins;
+		}
+	}
+
+	return NumOutPins;
+}
+
 void UK2Node_MultiGate::GetMarkBitFunction(FName& FunctionName, UClass** FunctionClass)
 {
 	FunctionName = GET_FUNCTION_NAME_CHECKED(UKismetNodeHelperLibrary, MarkBit);
diff --git a/Engine/Source/Editor/Kismet/Private/BlueprintEditor.cpp b/Engine/Source/Editor/Kismet/Private/BlueprintEditor.cpp
--- a/Engine/Source/Editor/Kismet/Private/BlueprintEditor.cpp
+++ b/Engine/Source/Editor/Kismet/Private/BlueprintEditor.cpp
@@ -4056,6 +4056,20 @@ void FBlueprintEditor::OnAddExecutionPin()
 
 bool FBlueprintEditor::CanAddExecutionPin() const
 {
+	const FGraphPanelSelectionSet& SelectedNodes = GetSelectedNodes();
+
+	// Iterate over all nodes, and see if all can have a pin added
+	for (FGraphPanelSelectionSet::TConstIterator It(SelectedNodes); It; ++It)
+	{
+		if (UK2Node_ExecutionSequence* AddPinNode = Cast<UK2Node_ExecutionSequence>(*It))
+		{
+			if (!AddPinNode->CanAddPin())
+			{
+				return false;
+			}
+		}
+	}
+
 	return true;
 }
 
@@ -4100,17 +4114,16 @@ void FBlueprintEditor::OnInsertExecutionPin(EPinInsertPosition Position)
 
 bool FBlueprintEditor::CanInsertExecutionPin() const
 {
-	// We likely don't need to validate here, as we validated on menu population,
-	// but better to grey out the option if it is somehow created but will
-	// not execute correctly
 	TSharedPtr<SGraphEditor> FocusedGraphEd = FocusedGraphEdPtr.Pin();
 	if (FocusedGraphEd.IsValid())
 	{
 		UEdGraphPin* SelectedPin = FocusedGraphEd->GetGraphPinForMenu();
 		if (SelectedPin)
 		{
-			UEdGraphNode* OwningNode = SelectedPin->GetOwningNode();
-			return Cast<UK2Node_ExecutionSequence>(OwningNode) != nullptr;
+			if (UK2Node_ExecutionSequence* ExecutionSequence = Cast<UK2Node_ExecutionSequence>(SelectedPin->GetOwningNode()))
+			{
+				return ExecutionSequence->CanAddPin();
+			}
 		}
 	}
 
