Don't allow properties to be used as override pins in blueprints that aren't marked as editable in blueprints.

diff --git a/Engine/Source/Editor/BlueprintGraph/Private/K2Node.cpp b/Engine/Source/Editor/BlueprintGraph/Private/K2Node.cpp
--- a/Engine/Source/Editor/BlueprintGraph/Private/K2Node.cpp
+++ b/Engine/Source/Editor/BlueprintGraph/Private/K2Node.cpp
@@ -1410,7 +1410,8 @@ void FOptionalPinManager::RebuildProperty(UProperty* TestProperty, FName Categor
 	Record->CategoryName = CategoryName;
 
 	bool bNegate = false;
-	Record->bHasOverridePin = PropertyCustomizationHelpers::GetEditConditionProperty(TestProperty, bNegate) != nullptr;
+	UProperty* OverrideProperty = PropertyCustomizationHelpers::GetEditConditionProperty(TestProperty, bNegate);
+	Record->bHasOverridePin = OverrideProperty != nullptr && OverrideProperty->HasAllPropertyFlags(CPF_BlueprintVisible) && !OverrideProperty->HasAllPropertyFlags(CPF_BlueprintReadOnly);
 	Record->bIsMarkedForAdvancedDisplay = TestProperty->HasAnyPropertyFlags(CPF_AdvancedDisplay);
 
 	// Get the defaults
diff --git a/Engine/Source/Editor/BlueprintGraph/Private/MakeStructHandler.cpp b/Engine/Source/Editor/BlueprintGraph/Private/MakeStructHandler.cpp
--- a/Engine/Source/Editor/BlueprintGraph/Private/MakeStructHandler.cpp
+++ b/Engine/Source/Editor/BlueprintGraph/Private/MakeStructHandler.cpp
@@ -134,7 +134,8 @@ void FKCHandler_MakeStruct::Compile(FKismetFunctionContext& Context, UEdGraphNod
 			// Handle injecting the override property values into the node if the property has any
 			bool bNegate = false;
 			UProperty* OverrideProperty = PropertyCustomizationHelpers::GetEditConditionProperty(BoundProperty, bNegate);
-			if (OverrideProperty)
+
+			if (OverrideProperty && OverrideProperty->HasAllPropertyFlags(CPF_BlueprintVisible) && !OverrideProperty->HasAllPropertyFlags(CPF_BlueprintReadOnly))
 			{
 				const UEdGraphSchema_K2* Schema = GetDefault<UEdGraphSchema_K2>();
 				FEdGraphPinType PinType;
