More careful scope for FSoftObjectPathSerializationScope - avoid tainting LinkerLoad->Preload(Struct);
#rnx

diff --git a/Engine/Source/Editor/BlueprintGraph/Private/K2Node.cpp b/Engine/Source/Editor/BlueprintGraph/Private/K2Node.cpp
--- a/Engine/Source/Editor/BlueprintGraph/Private/K2Node.cpp
+++ b/Engine/Source/Editor/BlueprintGraph/Private/K2Node.cpp
@@ -159,8 +159,6 @@ void UK2Node::FixupPinDefaultValues()
 
 void UK2Node::FixupPinStringDataReferences(FArchive* SavingArchive)
 {
-	FSoftObjectPathSerializationScope SetPackage(GetOutermost()->GetFName(), NAME_None, ESoftObjectPathCollectType::AlwaysCollect, ESoftObjectPathSerializeType::SkipSerializeIfArchiveHasSize);
-
 	// This code expands some pin types into their real representation, optionally serializes them, and then writes them out as strings again
 	const UEdGraphSchema_K2* K2Schema = GetDefault<UEdGraphSchema_K2>();
 	FLinkerLoad* LinkerLoad = GetLinker();
@@ -186,6 +184,8 @@ void UK2Node::FixupPinStringDataReferences(FArchive* SavingArchive)
 					{
 						LinkerLoad->Preload(Struct);
 					}
+					
+					FSoftObjectPathSerializationScope SetPackage(GetOutermost()->GetFName(), NAME_None, ESoftObjectPathCollectType::AlwaysCollect, ESoftObjectPathSerializeType::SkipSerializeIfArchiveHasSize);
 
 					TSharedPtr<FStructOnScope> StructData = MakeShareable(new FStructOnScope(Struct));
 					StructData->SetPackage(GetOutermost());
@@ -213,6 +213,8 @@ void UK2Node::FixupPinStringDataReferences(FArchive* SavingArchive)
 
 			if (Pin->PinType.PinCategory == UEdGraphSchema_K2::PC_SoftObject || Pin->PinType.PinCategory == UEdGraphSchema_K2::PC_SoftClass)
 			{
+				FSoftObjectPathSerializationScope SetPackage(GetOutermost()->GetFName(), NAME_None, ESoftObjectPathCollectType::AlwaysCollect, ESoftObjectPathSerializeType::SkipSerializeIfArchiveHasSize);
+
 				FSoftObjectPath TempRef(Pin->DefaultValue);
 
 				if (SavingArchive)
