Use DoesDefaultValueMatch instead of directly comparing default value string data
#rnx

diff --git a/Engine/Source/Editor/BlueprintGraph/Private/K2Node_FunctionEntry.cpp b/Engine/Source/Editor/BlueprintGraph/Private/K2Node_FunctionEntry.cpp
--- a/Engine/Source/Editor/BlueprintGraph/Private/K2Node_FunctionEntry.cpp
+++ b/Engine/Source/Editor/BlueprintGraph/Private/K2Node_FunctionEntry.cpp
@@ -840,6 +840,7 @@ void UK2Node_FunctionEntry::FixupPinStringDataReferences(FArchive* SavingArchive
 	Super::FixupPinStringDataReferences(SavingArchive);
 	if (SavingArchive)
 	{
+		const UEdGraphSchema_K2* K2Schema = GetDefault<UEdGraphSchema_K2>();
 		// If any of our pins got fixed up, we need to refresh our user pin default values
 		// For custom events, the Pin default values are authoritative
 		for (TSharedPtr<FUserPinInfo> PinInfo : UserDefinedPins)
@@ -848,11 +849,9 @@ void UK2Node_FunctionEntry::FixupPinStringDataReferences(FArchive* SavingArchive
 			{
 				if (Pin->Direction == PinInfo->DesiredPinDirection)
 				{
-					FString DefaultsString = Pin->GetDefaultAsString();
-
-					if (DefaultsString != PinInfo->PinDefaultValue)
+					if(!K2Schema->DoesDefaultValueMatch(*Pin, PinInfo->PinDefaultValue))
 					{
-						ModifyUserDefinedPinDefaultValue(PinInfo, DefaultsString);
+						ModifyUserDefinedPinDefaultValue(PinInfo, Pin->GetDefaultAsString());
 					}
 				}
 			}
