PR #4802: UE-47831: Execute PinConnectionListChanged to force pins to update (Contributed by [REDACTED_PERSON])

diff --git a/Engine/Source/Editor/BlueprintGraph/Private/K2Node_MakeContainer.cpp b/Engine/Source/Editor/BlueprintGraph/Private/K2Node_MakeContainer.cpp
--- a/Engine/Source/Editor/BlueprintGraph/Private/K2Node_MakeContainer.cpp
+++ b/Engine/Source/Editor/BlueprintGraph/Private/K2Node_MakeContainer.cpp
@@ -237,12 +237,13 @@ void UK2Node_MakeContainer::NotifyPinConnectionListChanged(UEdGraphPin* Pin)
 	else if (OutputPin->LinkedTo.Num() == 0)
 	{
 		// Return to wildcard if theres nothing in any of the input pins
-		auto PinsInUse = [](TArray<UEdGraphPin*>& PinsToConsider)
+		TFunction<bool(TArray<UEdGraphPin*>&)> PinsInUse = [this, &PinsInUse](TArray<UEdGraphPin*>& PinsToConsider)
 		{
 			bool bPinInUse = false;
 			for (UEdGraphPin* CurrentPin : PinsToConsider)
 			{
-				if (CurrentPin->ParentPin != nullptr || !CurrentPin->DoesDefaultValueMatchAutogenerated())
+				// Is there something in this pin?
+				if (CurrentPin->SubPins.Num() > 0 || !CurrentPin->DoesDefaultValueMatchAutogenerated())
 				{
 					bPinInUse = true;
 					break;
@@ -556,11 +557,11 @@ void UK2Node_MakeContainer::RemoveInputPin(UEdGraphPin* Pin)
 	FScopedTransaction Transaction(LOCTEXT("RemovePinTx", "RemovePin"));
 	Modify();
 	
-	TFunction<void(UEdGraphPin*)> RemovePin = [this, &RemovePin](UEdGraphPin* PinToRemove)
+	TFunction<void(UEdGraphPin*)> RemovePinLambda = [this, &RemovePinLambda](UEdGraphPin* PinToRemove)
 	{
 		for (int32 SubPinIndex = PinToRemove->SubPins.Num()-1; SubPinIndex >= 0; --SubPinIndex)
 		{
-			RemovePin(PinToRemove->SubPins[SubPinIndex]);
+			RemovePinLambda(PinToRemove->SubPins[SubPinIndex]);
 		}
 
 		int32 PinRemovalIndex = INDEX_NONE;
@@ -580,16 +581,17 @@ void UK2Node_MakeContainer::RemoveInputPin(UEdGraphPin* Pin)
 		int32 PinIndex = INDEX_NONE;
 		if (ValuePins.Find(Pin, PinIndex))
 		{
-			RemovePin(KeyPins[PinIndex]);
+			RemovePinLambda(KeyPins[PinIndex]);
 		}
 		else
 		{
 			verify(KeyPins.Find(Pin, PinIndex));
-			RemovePin(ValuePins[PinIndex]);
+			RemovePinLambda(ValuePins[PinIndex]);
 		}
 	}
 
-	RemovePin(Pin);
+	RemovePinLambda(Pin);
+	PinConnectionListChanged(Pin);
 
 	--NumInputs;
 	SyncPinNames();
