PR #5916: Clarify the tooltip of Persistent Local Variables in macros (Contributed by [REDACTED_PERSON])

diff --git a/Engine/Source/Editor/BlueprintGraph/Private/K2Node_TemporaryVariable.cpp b/Engine/Source/Editor/BlueprintGraph/Private/K2Node_TemporaryVariable.cpp
--- a/Engine/Source/Editor/BlueprintGraph/Private/K2Node_TemporaryVariable.cpp
+++ b/Engine/Source/Editor/BlueprintGraph/Private/K2Node_TemporaryVariable.cpp
@@ -53,8 +53,19 @@ FText UK2Node_TemporaryVariable::GetTooltipText() const
 	{
 		FFormatNamedArguments Args;
 		Args.Add(TEXT("VariableType"), UEdGraphSchema_K2::TypeToText(VariableType));
+
+		FText TitleFormat;
+		if (bIsPersistent)
+		{
+			TitleFormat = NSLOCTEXT("K2Node", "PersistentLocalVariable_Tooltip", "Local temporary {VariableType} variable, which is serialized in save games");
+		}
+		else
+		{
+			TitleFormat = NSLOCTEXT("K2Node", "LocalTemporaryVariable_Tooltip", "Local temporary {VariableType} variable");
+		}
+
 		// FText::Format() is slow, so we cache this to save on performance
-		CachedTooltip.SetCachedText(FText::Format(NSLOCTEXT("K2Node", "LocalTemporaryVariable_Tooltip", "Local temporary {VariableType} variable"), Args), this);
+		CachedTooltip.SetCachedText(FText::Format(TitleFormat, Args), this);
 	}
 	return CachedTooltip;
 }
