Add ability to rename function now from it's entry

diff --git a/Engine/Source/Editor/BlueprintGraph/Classes/K2Node_FunctionEntry.h b/Engine/Source/Editor/BlueprintGraph/Classes/K2Node_FunctionEntry.h
--- a/Engine/Source/Editor/BlueprintGraph/Classes/K2Node_FunctionEntry.h
+++ b/Engine/Source/Editor/BlueprintGraph/Classes/K2Node_FunctionEntry.h
@@ -41,6 +41,9 @@ class UK2Node_FunctionEntry : public UK2Node_FunctionTerminator
 	//~ Begin UEdGraphNode Interface
 	virtual void AllocateDefaultPins() override;
 	virtual FText GetNodeTitle(ENodeTitleType::Type TitleType) const override;
+	virtual void OnRenameNode(const FString& NewName) override;
+	virtual TSharedPtr<class INameValidatorInterface> MakeNameValidator() const override;
+	virtual bool GetCanRenameNode() const override;
 	virtual bool CanUserDeleteNode() const override { return false; }
 	virtual bool HasDeprecatedReference() const override;
 	virtual FEdGraphNodeDeprecationResponse GetDeprecationResponse(EEdGraphNodeDeprecationType DeprecationType) const override;
diff --git a/Engine/Source/Editor/BlueprintGraph/Private/K2Node_FunctionEntry.cpp b/Engine/Source/Editor/BlueprintGraph/Private/K2Node_FunctionEntry.cpp
--- a/Engine/Source/Editor/BlueprintGraph/Private/K2Node_FunctionEntry.cpp
+++ b/Engine/Source/Editor/BlueprintGraph/Private/K2Node_FunctionEntry.cpp
@@ -25,6 +25,7 @@
 #include "KismetCompiler.h"
 #include "Misc/OutputDeviceNull.h"
 #include "DiffResults.h"
+#include "Kismet2/Kismet2NameValidators.h"
 
 #define LOCTEXT_NAMESPACE "K2Node_FunctionEntry"
 
@@ -131,10 +132,10 @@ public:
 			{
 				UEdGraphPin* Net = FEdGraphUtilities::GetNetFromPin(Pin);
 
-				if (Context.NetMap.Find(Net) == NULL)
+				if (Context.NetMap.Find(Net) == nullptr)
 				{
 					// New net, resolve the term that will be used to construct it
-					FBPTerminal* Term = NULL;
+					FBPTerminal* Term = nullptr;
 
 					check(Net->Direction == EGPD_Output);
 
@@ -152,7 +153,7 @@ public:
 		{
 			UEdGraphPin* EntryPointPin = Node->FindPin(UEdGraphSchema_K2::PN_EntryPoint);
 			FBPTerminal** pTerm = Context.NetMap.Find(EntryPointPin);
-			if ((EntryPointPin != NULL) && (pTerm != NULL))
+			if ((EntryPointPin != nullptr) && (pTerm != nullptr))
 			{
 				FBlueprintCompiledStatement& ComputedGotoStatement = Context.AppendStatementForNode(Node);
 				ComputedGotoStatement.Type = KCST_ComputedGoto;
@@ -197,6 +198,7 @@ UK2Node_FunctionEntry::UK2Node_FunctionEntry(const FObjectInitializer& ObjectIni
 	// Enforce const-correctness by default
 	bEnforceConstCorrectness = true;
 	bUpdatedDefaultValuesOnLoad = false;
+	bCanRenameNode = bIsEditable;
 }
 
 void UK2Node_FunctionEntry::PreSave(const class ITargetPlatform* TargetPlatform)
@@ -317,6 +319,32 @@ FText UK2Node_FunctionEntry::GetNodeTitle(ENodeTitleType::Type TitleType) const
 	return DisplayInfo.DisplayName;
 }
 
+void UK2Node_FunctionEntry::OnRenameNode(const FString& NewName)
+{
+	CustomGeneratedFunctionName = FName(*NewName);
+	FBlueprintEditorUtils::RenameGraph(GetGraph(), NewName);
+}
+
+TSharedPtr<class INameValidatorInterface> UK2Node_FunctionEntry::MakeNameValidator() const
+{
+	if (CustomGeneratedFunctionName.IsNone())
+	{
+		FText TextName = GetNodeTitle(ENodeTitleType::Type::EditableTitle);
+		return MakeShareable(new FKismetNameValidator(GetBlueprint(), *TextName.ToString()));
+	}
+	else
+	{
+		return MakeShareable(new FKismetNameValidator(GetBlueprint(), CustomGeneratedFunctionName));
+	}
+}
+
+bool UK2Node_FunctionEntry::GetCanRenameNode() const
+{
+	UEdGraph* const Graph = GetGraph();
+
+	return (Graph && (Graph->bAllowDeletion || Graph->bAllowRenaming) && (bCanRenameNode || bIsEditable));
+}
+
 void UK2Node_FunctionEntry::AllocateDefaultPins()
 {
 	// Update our default values before copying them into pins
diff --git a/Engine/Source/Editor/GraphEditor/Private/SGraphNode.cpp b/Engine/Source/Editor/GraphEditor/Private/SGraphNode.cpp
--- a/Engine/Source/Editor/GraphEditor/Private/SGraphNode.cpp
+++ b/Engine/Source/Editor/GraphEditor/Private/SGraphNode.cpp
@@ -81,7 +81,7 @@ FText SNodeTitle::GetNodeTitle() const
 
 FText SNodeTitle::GetHeadTitle() const
 {
-	return (GraphNode.IsValid() && GraphNode->bCanRenameNode) ? GraphNode->GetNodeTitle(ENodeTitleType::EditableTitle) : CachedHeadTitle;
+	return (GraphNode.IsValid() && GraphNode->GetCanRenameNode()) ? GraphNode->GetNodeTitle(ENodeTitleType::EditableTitle) : CachedHeadTitle;
 }
 
 FVector2D SNodeTitle::GetTitleSize() const
@@ -1383,7 +1383,7 @@ FText SGraphNode::GetErrorMsgToolTip( ) const
 
 bool SGraphNode::IsNameReadOnly() const
 {
-	return (!GraphNode->bCanRenameNode || !IsNodeEditable());
+	return (!GraphNode->GetCanRenameNode() || !IsNodeEditable());
 }
 
 bool SGraphNode::OnVerifyNameTextChanged(const FText& InText, FText& OutErrorMessage)
@@ -1419,7 +1419,7 @@ void SGraphNode::OnNameTextCommited(const FText& InText, ETextCommit::Type Commi
 
 void SGraphNode::RequestRename()
 {
-	if ((GraphNode != NULL) && GraphNode->bCanRenameNode)
+	if ((GraphNode != nullptr) && GraphNode->GetCanRenameNode())
 	{
 		bRenameIsPending = true;
 	}
diff --git a/Engine/Source/Editor/Kismet/Private/BlueprintEditor.cpp b/Engine/Source/Editor/Kismet/Private/BlueprintEditor.cpp
--- a/Engine/Source/Editor/Kismet/Private/BlueprintEditor.cpp
+++ b/Engine/Source/Editor/Kismet/Private/BlueprintEditor.cpp
@@ -7554,7 +7554,7 @@ void FBlueprintEditor::OnRenameNode()
 		for (FGraphPanelSelectionSet::TConstIterator NodeIt(SelectedNodes); NodeIt; ++NodeIt)
 		{
 			UEdGraphNode* SelectedNode = Cast<UEdGraphNode>(*NodeIt);
-			if (SelectedNode != NULL && SelectedNode->bCanRenameNode)
+			if (SelectedNode != nullptr && SelectedNode->GetCanRenameNode())
 			{
 				FKismetEditorUtilities::BringKismetToFocusAttentionOnObject(SelectedNode, true);
 				break;
@@ -7569,7 +7569,7 @@ bool FBlueprintEditor::CanRenameNodes() const
 	{
 		if (const UEdGraphNode* SelectedNode = GetSingleSelectedNode())
 		{
-			return SelectedNode->bCanRenameNode;
+			return SelectedNode->GetCanRenameNode();
 		}
 	}
 	return false;
@@ -7579,7 +7579,7 @@ bool FBlueprintEditor::OnNodeVerifyTitleCommit(const FText& NewText, UEdGraphNod
 {
 	bool bValid(false);
 
-	if (NodeBeingChanged && NodeBeingChanged->bCanRenameNode)
+	if (NodeBeingChanged && NodeBeingChanged->GetCanRenameNode())
 	{
 		// Clear off any existing error message 
 		NodeBeingChanged->ErrorMsg.Empty();
