PR #5951: Double click on interface node call to go to implementation (Contributed by [REDACTED_PERSON])

Remove a CA_SUPPRESS to see if static analysis still catches it

diff --git a/Engine/Source/Editor/BlueprintGraph/Private/K2Node_CallFunction.cpp b/Engine/Source/Editor/BlueprintGraph/Private/K2Node_CallFunction.cpp
--- a/Engine/Source/Editor/BlueprintGraph/Private/K2Node_CallFunction.cpp
+++ b/Engine/Source/Editor/BlueprintGraph/Private/K2Node_CallFunction.cpp
@@ -2828,23 +2828,22 @@ bool UK2Node_CallFunction::HasExternalDependencies(TArray<class UStruct*>* Optio
 
 UEdGraph* UK2Node_CallFunction::GetFunctionGraph(const UEdGraphNode*& OutGraphNode) const
 {
-	OutGraphNode = NULL;
+	OutGraphNode = nullptr;
 
 	// Search for the Blueprint owner of the function graph, climbing up through the Blueprint hierarchy
 	UClass* MemberParentClass = FunctionReference.GetMemberParentClass(GetBlueprintClassFromNode());
-	if(MemberParentClass != NULL)
+	if(MemberParentClass != nullptr)
 	{
 		UBlueprintGeneratedClass* ParentClass = Cast<UBlueprintGeneratedClass>(MemberParentClass);
-		if(ParentClass != NULL && ParentClass->ClassGeneratedBy != NULL)
+		if(ParentClass != nullptr && ParentClass->ClassGeneratedBy != nullptr)
 		{
 			UBlueprint* Blueprint = Cast<UBlueprint>(ParentClass->ClassGeneratedBy);
-			while(Blueprint != NULL)
+			while(Blueprint != nullptr)
 			{
-				UEdGraph* TargetGraph = NULL;
-				FName FunctionName = FunctionReference.GetMemberName();
-				for (UEdGraph* Graph : Blueprint->FunctionGraphs) 
+				UEdGraph* TargetGraph = nullptr;
+				const FName FunctionName = FunctionReference.GetMemberName();
+				for (UEdGraph* const Graph : Blueprint->FunctionGraphs) 
 				{
-					CA_SUPPRESS(28182); // warning C28182: Dereferencing NULL pointer. 'Graph' contains the same NULL value as 'TargetGraph' did.
 					if (Graph->GetFName() == FunctionName)
 					{
 						TargetGraph = Graph;
@@ -2852,7 +2851,27 @@ UEdGraph* UK2Node_CallFunction::GetFunctionGraph(const UEdGraphNode*& OutGraphNo
 					}
 				}
 
-				if((TargetGraph != NULL) && !TargetGraph->HasAnyFlags(RF_Transient))
+				if (!TargetGraph)
+				{
+					for (const FBPInterfaceDescription& Interface : Blueprint->ImplementedInterfaces)
+					{
+						for (UEdGraph* const Graph : Interface.Graphs)
+						{
+							if (Graph->GetFName() == FunctionName)
+							{
+								TargetGraph = Graph;
+								break;
+							}
+						}
+
+						if (TargetGraph)
+						{
+							break;
+						}
+					}
+				}
+
+				if((TargetGraph != nullptr) && !TargetGraph->HasAnyFlags(RF_Transient))
 				{
 					// Found the function graph in a Blueprint, return that graph
 					return TargetGraph;
@@ -2860,12 +2879,12 @@ UEdGraph* UK2Node_CallFunction::GetFunctionGraph(const UEdGraphNode*& OutGraphNo
 				else
 				{
 					// Did not find the function call as a graph, it may be a custom event
-					UK2Node_CustomEvent* CustomEventNode = NULL;
+					UK2Node_CustomEvent* CustomEventNode = nullptr;
 
 					TArray<UK2Node_CustomEvent*> CustomEventNodes;
 					FBlueprintEditorUtils::GetAllNodesOfClass(Blueprint, CustomEventNodes);
 
-					for (UK2Node_CustomEvent* CustomEvent : CustomEventNodes)
+					for (UK2Node_CustomEvent* const CustomEvent : CustomEventNodes)
 					{
 						if(CustomEvent->CustomFunctionName == FunctionReference.GetMemberName())
 						{
@@ -2876,11 +2895,11 @@ UEdGraph* UK2Node_CallFunction::GetFunctionGraph(const UEdGraphNode*& OutGraphNo
 				}
 
 				ParentClass = Cast<UBlueprintGeneratedClass>(Blueprint->ParentClass);
-				Blueprint = ParentClass != NULL ? Cast<UBlueprint>(ParentClass->ClassGeneratedBy) : NULL;
+				Blueprint = ParentClass != nullptr ? Cast<UBlueprint>(ParentClass->ClassGeneratedBy) : nullptr;
 			}
 		}
 	}
-	return NULL;
+	return nullptr;
 }
 
 bool UK2Node_CallFunction::IsStructureWildcardProperty(const UFunction* Function, const FName PropertyName)
