PR #5962: UE-75710: Interface implementation function parameters are retrievabl. (Contributed by [REDACTED_PERSON])

diff --git a/Engine/Source/Editor/BlueprintGraph/Private/BlueprintActionDatabase.cpp b/Engine/Source/Editor/BlueprintGraph/Private/BlueprintActionDatabase.cpp
--- a/Engine/Source/Editor/BlueprintGraph/Private/BlueprintActionDatabase.cpp
+++ b/Engine/Source/Editor/BlueprintGraph/Private/BlueprintActionDatabase.cpp
@@ -766,9 +766,13 @@ static void BlueprintActionDatabaseImpl::AddBlueprintGraphActions(UBlueprint con
 		ActionListOut.Add(MakeMacroNodeSpawner(MacroGraph));
 	}
 
-	// local variables and parameters
-	for (UEdGraph* FunctionGraph : Blueprint->FunctionGraphs)
+	auto CreateEntriesForGraphLambda = [Blueprint, &ActionListOut](UEdGraph* FunctionGraph)
 	{
+		if (!FunctionGraph)
+		{
+			return;
+		}
+
 		TArray<UK2Node_FunctionEntry*> GraphEntryNodes;
 		FunctionGraph->GetNodesOfClass<UK2Node_FunctionEntry>(GraphEntryNodes);
 
@@ -804,6 +808,23 @@ static void BlueprintActionDatabaseImpl::AddBlueprintGraphActions(UBlueprint con
 				ActionListOut.Add(SetVarSpawner);
 			}
 		}
+	};
+
+	auto CreateEntriesLambda = [&](const TArray<UEdGraph*>& Graphs)
+	{
+		for (UEdGraph* const Graph : Graphs)
+		{
+			CreateEntriesForGraphLambda(Graph);
+		}
+	};
+
+	// local variables and parameters for functions
+	CreateEntriesLambda(Blueprint->FunctionGraphs);
+
+	// local variables and parameters for interfaces
+	for (const FBPInterfaceDescription& Interface : Blueprint->ImplementedInterfaces)
+	{
+		CreateEntriesLambda(Interface.Graphs);
 	}
 }
 
