Add support for object properties in GetClassDefaults as long as they are non-instanced

diff --git a/Engine/Source/Editor/BlueprintGraph/Private/K2Node_GetClassDefaults.cpp b/Engine/Source/Editor/BlueprintGraph/Private/K2Node_GetClassDefaults.cpp
--- a/Engine/Source/Editor/BlueprintGraph/Private/K2Node_GetClassDefaults.cpp
+++ b/Engine/Source/Editor/BlueprintGraph/Private/K2Node_GetClassDefaults.cpp
@@ -69,10 +69,10 @@ namespace
 					&& ((TestMapProperty->KeyProp && TestMapProperty->KeyProp->IsA<UObjectProperty>() && !TestMapProperty->KeyProp->IsA<UClassProperty>())
 						|| (TestMapProperty->ValueProp && TestMapProperty->ValueProp->IsA<UObjectProperty>() && !TestMapProperty->ValueProp->IsA<UClassProperty>())));
 			}
-
+			
 			// Don't expose object properties (except for those containing class objects).
 			// @TODO - Could potentially expose object reference values if/when we have support for 'const' input pins.
-			return !TestProperty->IsA<UObjectProperty>() || TestProperty->IsA<UClassProperty>();
+			return !TestProperty->ContainsInstancedObjectProperty() || !TestProperty->IsA<UObjectProperty>() || TestProperty->IsA<UClassProperty>();
 		}
 
 		virtual void CustomizePinData(UEdGraphPin* Pin, FName SourcePropertyName, int32 ArrayIndex, UProperty* Property = nullptr) const override
