Add multiple dynamic output paramater support. Users can now specify multiple dynamic outputs by adding "," between the property names

diff --git a/Engine/Source/Editor/BlueprintGraph/Private/K2Node_CallFunction.cpp b/Engine/Source/Editor/BlueprintGraph/Private/K2Node_CallFunction.cpp
--- a/Engine/Source/Editor/BlueprintGraph/Private/K2Node_CallFunction.cpp
+++ b/Engine/Source/Editor/BlueprintGraph/Private/K2Node_CallFunction.cpp
@@ -185,13 +185,13 @@ private:
 	bool IsTypePickerPin(UEdGraphPin* Pin) const;
 
 	/**
-	 * Retrieves the object output pin that is altered as the class input is 
-	 * changed (favors params flagged by "DynamicOutputParam" metadata).
-	 * 
-	 * @return Null if the output param cannot be altered from the class input, 
-	 *         otherwise a output pin that will mutate type as the class input is changed.
-	 */
-	static UEdGraphPin* GetDynamicOutPin(const UK2Node_CallFunction* FuncNode);
+	* Retrieves the object output pins that are altered as the class input is
+	* changed (favors params flagged by "DynamicOutputParam" metadata).
+	* 
+	* @param FuncNode 		The function node in question
+	* @param OutPins		Out array of pins that are flagged with "DynamicOutputParam" metadata
+	*/
+	static void GetDynamicOutPins(const UK2Node_CallFunction* FuncNode, TArray<UEdGraphPin*>& OutPins);
 
 	/**
 	 * Checks if the specified type is an object type that reflects the picker 
@@ -213,28 +213,17 @@ void FDynamicOutputHelper::ConformOutputType() const
 		UClass* PickedClass = GetPinClass(MutatingPin);
 		UK2Node_CallFunction* FuncNode = GetFunctionNode();
 
-		if (UEdGraphPin* DynamicOutPin = GetDynamicOutPin(FuncNode))
+		// See if there is any dynamic output pins
+		TArray<UEdGraphPin*> DynamicPins;
+		GetDynamicOutPins(FuncNode, DynamicPins);
+		
+		// Set the pins class
+		for (UEdGraphPin* Pin : DynamicPins)
 		{
-			DynamicOutPin->PinType.PinSubCategoryObject = PickedClass;
-
-			// leave the connection, and instead bring the user's attention to 
-			// it via a ValidateNodeDuringCompilation() error
-// 			const UEdGraphSchema* Schema = FuncNode->GetSchema();
-// 			for (int32 LinkIndex = 0; LinkIndex < DynamicOutPin->LinkedTo.Num();)
-// 			{
-// 				UEdGraphPin* Link = DynamicOutPin->LinkedTo[LinkIndex];
-// 				// if this can no longer be linked to the other pin, then we 
-// 				// should disconnect it (because the pin's type just changed)
-// 				if (Schema->CanCreateConnection(DynamicOutPin, Link).Response == CONNECT_RESPONSE_DISALLOW)
-// 				{
-// 					DynamicOutPin->BreakLinkTo(Link);
-// 					// @TODO: warn/notify somehow
-// 				}
-// 				else
-// 				{
-// 					++LinkIndex;
-// 				}
-// 			}
+			if (ensure(Pin != nullptr))
+			{
+				Pin->PinType.PinSubCategoryObject = PickedClass;
+			}
 		}
 	}
 }
@@ -316,15 +305,21 @@ UClass* FDynamicOutputHelper::GetPinClass(UEdGraphPin* Pin)
 
 void FDynamicOutputHelper::VerifyNode(const UK2Node_CallFunction* FuncNode, FCompilerResultsLog& MessageLog)
 {
-	if (UEdGraphPin* DynamicOutPin = GetDynamicOutPin(FuncNode))
+	TArray<UEdGraphPin*> DynamicPins;
+	GetDynamicOutPins(FuncNode, DynamicPins);
+	
+	for (UEdGraphPin* DynamicOutPin : DynamicPins)
 	{
-		const UEdGraphSchema* Schema = FuncNode->GetSchema();
-		for (UEdGraphPin* Link : DynamicOutPin->LinkedTo)
+		if (ensure(DynamicOutPin != nullptr))
 		{
-			if (Schema->CanCreateConnection(DynamicOutPin, Link).Response == CONNECT_RESPONSE_DISALLOW)
+			const UEdGraphSchema* Schema = FuncNode->GetSchema();
+			for (UEdGraphPin* Link : DynamicOutPin->LinkedTo)
 			{
-				FText const ErrorFormat = LOCTEXT("BadConnection", "Invalid pin connection from '@@' to '@@'. You may have changed the type after the connections were made.");
-				MessageLog.Error(*ErrorFormat.ToString(), DynamicOutPin, Link);
+				if (Schema->CanCreateConnection(DynamicOutPin, Link).Response == CONNECT_RESPONSE_DISALLOW)
+				{
+					FText const ErrorFormat = LOCTEXT("BadConnection", "Invalid pin connection from '@@' to '@@'. You may have changed the type after the connections were made.");
+					MessageLog.Error(*ErrorFormat.ToString(), DynamicOutPin, Link);
+				}
 			}
 		}
 	}
@@ -358,50 +353,65 @@ bool FDynamicOutputHelper::IsTypePickerPin(UEdGraphPin* Pin) const
 	return bIsTypeDeterminingPin && (bPinIsClassPicker || bPinIsObjectPicker) && (Pin->Direction == EGPD_Input);
 }
 
-UEdGraphPin* FDynamicOutputHelper::GetDynamicOutPin(const UK2Node_CallFunction* FuncNode)
+void FDynamicOutputHelper::GetDynamicOutPins(const UK2Node_CallFunction* FuncNode, TArray<UEdGraphPin*>& OutPins)
 {
-	UProperty* TaggedOutputParam = nullptr;
 	if (UFunction* Function = FuncNode->GetTargetFunction())
 	{
-		const FString& OutputPinName = Function->GetMetaData(FBlueprintMetadata::MD_DynamicOutputParam);
-		// we sort through properties, instead of pins, because the pin's type 
-		// could already be modified to some other class (for when we check CanConformPinType)
-		for (TFieldIterator<UProperty> ParamIt(Function); ParamIt && (ParamIt->PropertyFlags & CPF_Parm); ++ParamIt)
+		const FString& OutputPinNames = Function->GetMetaData(FBlueprintMetadata::MD_DynamicOutputParam);
+		
+		// There could be more than one dynamic output, so split by comma
+		TArray<FString> UserDefinedDynamicProprties;
+		OutputPinNames.ParseIntoArray(UserDefinedDynamicProprties, TEXT(","), true);
+
+		// trim up the whitespace that a user may have added
+		for (FString& Name : UserDefinedDynamicProprties)
 		{
-			if (OutputPinName.IsEmpty() && ParamIt->HasAnyPropertyFlags(CPF_ReturnParm))
-			{
-				TaggedOutputParam = *ParamIt;
-				break;
-			}
-			else if (OutputPinName == ParamIt->GetName())
-			{
-				TaggedOutputParam = *ParamIt;
-				break;
-			}
+			Name = Name.TrimStartAndEnd();
 		}
 
-		if (TaggedOutputParam != nullptr)
+		// Get the schema so we can verify the pin if we find it
+		const UEdGraphSchema_K2* K2Schema = GetDefault<UEdGraphSchema_K2>();
+
+		// Lambda to add a pin to the out pins if it is valid
+		auto AddPinToOutputLambda = [FuncNode, K2Schema](UProperty* TaggedOutputParam, TArray<UEdGraphPin*>& OutPins)
 		{
-			const UEdGraphSchema_K2* K2Schema = GetDefault<UEdGraphSchema_K2>();
+			// Ensure that this is a valid pin to make dynamic
 			FEdGraphPinType PropertyPinType;
 
-			if (!K2Schema->ConvertPropertyToPinType(TaggedOutputParam, /*out*/PropertyPinType) || !CanConformPinType(FuncNode, PropertyPinType))
+			if (TaggedOutputParam && K2Schema->ConvertPropertyToPinType(TaggedOutputParam, /*out*/PropertyPinType) || !CanConformPinType(FuncNode, PropertyPinType))
 			{
-				TaggedOutputParam = nullptr;
+				UEdGraphPin* DynamicOutPin = FuncNode->FindPin(TaggedOutputParam->GetFName());
+				if (DynamicOutPin && (DynamicOutPin->Direction == EGPD_Output))
+				{
+					OutPins.Add(DynamicOutPin);
+				}
 			}
-		}
-	}
+		};
 
-	UEdGraphPin* DynamicOutPin = nullptr;
-	if (TaggedOutputParam != nullptr)
-	{
-		DynamicOutPin = FuncNode->FindPin(TaggedOutputParam->GetFName());
-		if (DynamicOutPin && (DynamicOutPin->Direction != EGPD_Output))
+		// we sort through properties, instead of pins, because the pin's type 
+		// could already be modified to some other class (for when we check CanConformPinType)
+		for (TFieldIterator<UProperty> ParamIt(Function); ParamIt && (ParamIt->PropertyFlags & CPF_Parm); ++ParamIt)
 		{
-			DynamicOutPin = nullptr;
+			// If the user defined pins are 0 then assume we are just setting the type of a single output
+			if (UserDefinedDynamicProprties.Num() == 0 && ParamIt->HasAnyPropertyFlags(CPF_ReturnParm))
+			{
+				AddPinToOutputLambda(*ParamIt, OutPins);
+			}
+			else
+			{
+				// Check against each property that the user has specified
+				for (const FString& OutputPinName : UserDefinedDynamicProprties)
+				{
+					// If this is the return parameter of this function or the pin name matches that which the user has specified
+					if (OutputPinName == ParamIt->GetName())
+					{
+						AddPinToOutputLambda(*ParamIt, OutPins);
+						break;
+					}
+				}
+			}
 		}
 	}
-	return DynamicOutPin;
 }
 
 bool FDynamicOutputHelper::CanConformPinType(const UK2Node_CallFunction* FuncNode, const FEdGraphPinType& TypeToTest)
