Fix for static analysis warning from logical ops and Linux cross-compile error in K2Node_CallFunction.cpp

diff --git a/Engine/Source/Editor/BlueprintGraph/Private/K2Node_CallFunction.cpp b/Engine/Source/Editor/BlueprintGraph/Private/K2Node_CallFunction.cpp
--- a/Engine/Source/Editor/BlueprintGraph/Private/K2Node_CallFunction.cpp
+++ b/Engine/Source/Editor/BlueprintGraph/Private/K2Node_CallFunction.cpp
@@ -378,7 +378,7 @@ void FDynamicOutputHelper::GetDynamicOutPins(const UK2Node_CallFunction* FuncNod
 			// Ensure that this is a valid pin to make dynamic
 			FEdGraphPinType PropertyPinType;
 
-			if (TaggedOutputParam && K2Schema->ConvertPropertyToPinType(TaggedOutputParam, /*out*/PropertyPinType) || !CanConformPinType(FuncNode, PropertyPinType))
+			if (TaggedOutputParam && (K2Schema->ConvertPropertyToPinType(TaggedOutputParam, /*out*/PropertyPinType) || !CanConformPinType(FuncNode, PropertyPinType)))
 			{
 				UEdGraphPin* DynamicOutPin = FuncNode->FindPin(TaggedOutputParam->GetFName());
 				if (DynamicOutPin && (DynamicOutPin->Direction == EGPD_Output))
