Add a keyword to the Sequence flow control node to prioritize it during searching.

Typing 'seq' will match to the Sequence node

diff --git a/Engine/Source/Editor/BlueprintGraph/Classes/K2Node_ExecutionSequence.h b/Engine/Source/Editor/BlueprintGraph/Classes/K2Node_ExecutionSequence.h
--- a/Engine/Source/Editor/BlueprintGraph/Classes/K2Node_ExecutionSequence.h
+++ b/Engine/Source/Editor/BlueprintGraph/Classes/K2Node_ExecutionSequence.h
@@ -14,7 +14,7 @@
 
 class FBlueprintActionDatabaseRegistrar;
 
-UCLASS(MinimalAPI)
+UCLASS(MinimalAPI, meta = (Keywords = "sequence"))
 class UK2Node_ExecutionSequence : public UK2Node, public IK2Node_AddPinInterface
 {
 	GENERATED_UCLASS_BODY()
