Add a function to check if an event node is part of an interface

#rnx

diff --git a/Engine/Source/Editor/BlueprintGraph/Classes/K2Node_Event.h b/Engine/Source/Editor/BlueprintGraph/Classes/K2Node_Event.h
--- a/Engine/Source/Editor/BlueprintGraph/Classes/K2Node_Event.h
+++ b/Engine/Source/Editor/BlueprintGraph/Classes/K2Node_Event.h
@@ -97,7 +97,9 @@ class BLUEPRINTGRAPH_VTABLE UK2Node_Event : public UK2Node_EditablePinBase, publ
 
 	/** Checks whether the parameters for this event node are compatible with the specified function entry node */
 	BLUEPRINTGRAPH_API virtual bool IsFunctionEntryCompatible(const class UK2Node_FunctionEntry* EntryNode) const;
-	
+	/** Checks if this event node is implementing an interface event */
+	BLUEPRINTGRAPH_API bool IsInterfaceEventNode() const;
+
 	BLUEPRINTGRAPH_API UFunction* FindEventSignatureFunction();
 	BLUEPRINTGRAPH_API void UpdateDelegatePin(bool bSilent = false);
 	BLUEPRINTGRAPH_API FName GetFunctionName() const;
diff --git a/Engine/Source/Editor/BlueprintGraph/Private/K2Node_Event.cpp b/Engine/Source/Editor/BlueprintGraph/Private/K2Node_Event.cpp
--- a/Engine/Source/Editor/BlueprintGraph/Private/K2Node_Event.cpp
+++ b/Engine/Source/Editor/BlueprintGraph/Private/K2Node_Event.cpp
@@ -460,6 +460,11 @@ bool UK2Node_Event::IsFunctionEntryCompatible(const UK2Node_FunctionEntry* Entry
 	return (EventPins.Num() == 0) && (EntryPins.Num() == 0);
 }
 
+bool UK2Node_Event::IsInterfaceEventNode() const
+{
+	return (EventReference.GetMemberParentClass(GetBlueprintClassFromNode()) != nullptr) && EventReference.GetMemberParentClass(GetBlueprintClassFromNode())->IsChildOf(UInterface::StaticClass());
+}
+
 bool UK2Node_Event::IsCompatibleWithGraph(const UEdGraph* TargetGraph) const
 {
 	bool bIsCompatible = Super::IsCompatibleWithGraph(TargetGraph);
@@ -745,7 +750,7 @@ FText UK2Node_Event::GetToolTipHeading() const
 	{
 		EventHeading = LOCTEXT("ServerOnlyEvent", "Server Only");
 	}
-	else if(EventHeading.IsEmpty() && (EventReference.GetMemberParentClass(GetBlueprintClassFromNode()) != nullptr) && EventReference.GetMemberParentClass(GetBlueprintClassFromNode())->IsChildOf(UInterface::StaticClass()))
+	else if(EventHeading.IsEmpty() && IsInterfaceEventNode())
 	{
 		EventHeading = LOCTEXT("InterfaceEvent", "Interface Event");
 	}
