Do not allow the addition of a native event node if there is already an existing function graph with that name

#rnx

diff --git a/Engine/Source/Editor/BlueprintGraph/Private/BlueprintEventNodeSpawner.cpp b/Engine/Source/Editor/BlueprintGraph/Private/BlueprintEventNodeSpawner.cpp
--- a/Engine/Source/Editor/BlueprintGraph/Private/BlueprintEventNodeSpawner.cpp
+++ b/Engine/Source/Editor/BlueprintGraph/Private/BlueprintEventNodeSpawner.cpp
@@ -3,6 +3,7 @@
 #include "BlueprintEventNodeSpawner.h"
 #include "EdGraphSchema_K2.h"
 #include "K2Node_CallFunction.h"
+#include "K2Node_FunctionEntry.h"
 #include "Kismet2/BlueprintEditorUtils.h"
 #include "BlueprintNodeTemplateCache.h"
 
@@ -157,6 +158,24 @@ UEdGraphNode* UBlueprintEventNodeSpawner::Invoke(UEdGraph* ParentGraph, FBinding
 		EventName  = EventFunc->GetFName();	
 	}
 
+	// If there is a function with this name on the blueprint, focus on that instead
+	for (UEdGraph* FuncGraph : Blueprint->FunctionGraphs)
+	{
+		if (FuncGraph && FuncGraph->GetFName() == EventName)
+		{
+			TArray<UK2Node_FunctionEntry*> FunctionEntry;
+			FuncGraph->GetNodesOfClass<UK2Node_FunctionEntry>(FunctionEntry);
+			if (FunctionEntry.Num())
+			{
+				return FunctionEntry[0];
+			}
+			else
+			{
+				return nullptr;
+			}
+		}
+	}
+
 	// This Event node might already be present in the Blueprint in a disabled state, 
 	// remove it and allow the user to successfully place the node where they want it.
 	if(EventNode && EventNode->IsAutomaticallyPlacedGhostNode())
