FiB: Allow SwitchEnum node output pins for native enums to be searchable by their corresponding C++ enum name.

diff --git a/Engine/Source/Editor/BlueprintGraph/Classes/K2Node_SwitchEnum.h b/Engine/Source/Editor/BlueprintGraph/Classes/K2Node_SwitchEnum.h
--- a/Engine/Source/Editor/BlueprintGraph/Classes/K2Node_SwitchEnum.h
+++ b/Engine/Source/Editor/BlueprintGraph/Classes/K2Node_SwitchEnum.h
@@ -38,6 +38,7 @@ class UK2Node_SwitchEnum : public UK2Node_Switch, public INodeDependingOnEnumInt
 	// UEdGraphNode interface
 	virtual FText GetTooltipText() const override;
 	virtual FText GetNodeTitle(ENodeTitleType::Type TitleType) const override;
+	virtual void AddPinSearchMetaDataInfo(const UEdGraphPin* Pin, TArray<struct FSearchTagDataPair>& OutTaggedMetaData) const override;
 	// End of UEdGraphNode interface
 
 	// UK2Node interface
diff --git a/Engine/Source/Editor/BlueprintGraph/Private/K2Node_SwitchEnum.cpp b/Engine/Source/Editor/BlueprintGraph/Private/K2Node_SwitchEnum.cpp
--- a/Engine/Source/Editor/BlueprintGraph/Private/K2Node_SwitchEnum.cpp
+++ b/Engine/Source/Editor/BlueprintGraph/Private/K2Node_SwitchEnum.cpp
@@ -5,6 +5,7 @@
 #include "Kismet/KismetMathLibrary.h"
 #include "BlueprintFieldNodeSpawner.h"
 #include "BlueprintActionDatabaseRegistrar.h"
+#include "FindInBlueprintManager.h"
 
 #define LOCTEXT_NAMESPACE "K2Node"
 
@@ -74,6 +75,18 @@ FText UK2Node_SwitchEnum::GetTooltipText() const
 	return NSLOCTEXT("K2Node", "SwitchEnum_ToolTip", "Selects an output that matches the input value");
 }
 
+void UK2Node_SwitchEnum::AddPinSearchMetaDataInfo(const UEdGraphPin* Pin, TArray<struct FSearchTagDataPair>& OutTaggedMetaData) const
+{
+	Super::AddPinSearchMetaDataInfo(Pin, OutTaggedMetaData);
+
+	const UEdGraphSchema_K2* K2Schema = GetDefault<UEdGraphSchema_K2>();
+	if (K2Schema->IsExecPin(*Pin) && Pin->Direction == EGPD_Output && Enum->IsNative() && EnumEntries.Contains(Pin->GetFName()))
+	{
+		// Allow native enum switch pins to be searchable by C++ enum name
+		OutTaggedMetaData.Add(FSearchTagDataPair(FFindInBlueprintSearchTags::FiB_NativeName, FText::FromString(Pin->GetName())));
+	}
+}
+
 bool UK2Node_SwitchEnum::IsConnectionDisallowed(const UEdGraphPin* MyPin, const UEdGraphPin* OtherPin, FString& OutReason) const
 {
 	const UEnum* SubCategoryObject = Cast<UEnum>( OtherPin->PinType.PinSubCategoryObject.Get() );
diff --git a/Engine/Source/Editor/Kismet/Private/FindInBlueprintManager.cpp b/Engine/Source/Editor/Kismet/Private/FindInBlueprintManager.cpp
--- a/Engine/Source/Editor/Kismet/Private/FindInBlueprintManager.cpp
+++ b/Engine/Source/Editor/Kismet/Private/FindInBlueprintManager.cpp
@@ -997,15 +997,20 @@ namespace BlueprintSearchMetaDataHelpers
 					{
 						// Find all the pins and extract their metadata
 						InWriter->WriteArrayStart(FFindInBlueprintSearchTags::FiB_Pins);
-						for (UEdGraphPin* Pin : Node->Pins)
+						for (const UEdGraphPin* Pin : Node->Pins)
 						{
 							// Hidden pins are not searchable
 							if (Pin->bHidden == false)
 							{
 								InWriter->WriteObjectStart();
 								{
-									InWriter->WriteValue(FFindInBlueprintSearchTags::FiB_Name, Pin->GetSchema()->GetPinDisplayName(Pin));
-									InWriter->WriteValue(FFindInBlueprintSearchTags::FiB_DefaultValue, Pin->GetDefaultAsText());
+									TArray<struct FSearchTagDataPair> Tags;
+									Node->AddPinSearchMetaDataInfo(Pin, Tags);
+
+									for (const FSearchTagDataPair& SearchData : Tags)
+									{
+										InWriter->WriteValue(SearchData.Key, SearchData.Value);
+									}
 								}
 								SavePinTypeToJson(InWriter, Pin->PinType);
 								InWriter->WriteObjectEnd();
diff --git a/Engine/Source/Runtime/Engine/Classes/EdGraph/EdGraphNode.h b/Engine/Source/Runtime/Engine/Classes/EdGraph/EdGraphNode.h
--- a/Engine/Source/Runtime/Engine/Classes/EdGraph/EdGraphNode.h
+++ b/Engine/Source/Runtime/Engine/Classes/EdGraph/EdGraphNode.h
@@ -848,6 +848,14 @@ public:
 	 */
 	virtual void AddSearchMetaDataInfo(TArray<struct FSearchTagDataPair>& OutTaggedMetaData) const;
 
+	/**
+	 * Adds node pin data to the search metadata, override to collect more data that may be desirable to search for
+	 *
+	 * @param Pin					The pin for which to gather search meta data
+	 * @param OutTaggedMetaData		Built array of tagged meta data for the given pin
+	 */
+	virtual void AddPinSearchMetaDataInfo(const UEdGraphPin* Pin, TArray<struct FSearchTagDataPair>& OutTaggedMetaData) const;
+
 	/** Return the requested metadata for the pin if there is any */
 	virtual FString GetPinMetaData(FName InPinName, FName InKey) { return FString(); }
 
diff --git a/Engine/Source/Runtime/Engine/Private/EdGraph/EdGraphNode.cpp b/Engine/Source/Runtime/Engine/Private/EdGraph/EdGraphNode.cpp
--- a/Engine/Source/Runtime/Engine/Private/EdGraph/EdGraphNode.cpp
+++ b/Engine/Source/Runtime/Engine/Private/EdGraph/EdGraphNode.cpp
@@ -836,6 +836,15 @@ void UEdGraphNode::AddSearchMetaDataInfo(TArray<struct FSearchTagDataPair>& OutT
 	OutTaggedMetaData.Add(FSearchTagDataPair(FFindInBlueprintSearchTags::FiB_Comment, FText::FromString(NodeComment)));
 }
 
+void UEdGraphNode::AddPinSearchMetaDataInfo(const UEdGraphPin* Pin, TArray<struct FSearchTagDataPair>& OutTaggedMetaData) const
+{
+	// Searchable - Primary label for the item in the search results
+	OutTaggedMetaData.Add(FSearchTagDataPair(FFindInBlueprintSearchTags::FiB_Name, Pin->GetSchema()->GetPinDisplayName(Pin)));
+
+	// Searchable - The pin's default value as a text string
+	OutTaggedMetaData.Add(FSearchTagDataPair(FFindInBlueprintSearchTags::FiB_DefaultValue, Pin->GetDefaultAsText()));
+}
+
 void UEdGraphNode::OnUpdateCommentText( const FString& NewComment )
 {
 	if( !NodeComment.Equals( NewComment ))
