Fix crash when loading a blueprint that contains a Switch Enum that had its related enum class renamed

diff --git a/Engine/Source/Editor/BlueprintGraph/Private/K2Node_SwitchEnum.cpp b/Engine/Source/Editor/BlueprintGraph/Private/K2Node_SwitchEnum.cpp
--- a/Engine/Source/Editor/BlueprintGraph/Private/K2Node_SwitchEnum.cpp
+++ b/Engine/Source/Editor/BlueprintGraph/Private/K2Node_SwitchEnum.cpp
@@ -80,7 +80,7 @@ void UK2Node_SwitchEnum::AddPinSearchMetaDataInfo(const UEdGraphPin* Pin, TArray
 	Super::AddPinSearchMetaDataInfo(Pin, OutTaggedMetaData);
 
 	const UEdGraphSchema_K2* K2Schema = GetDefault<UEdGraphSchema_K2>();
-	if (K2Schema->IsExecPin(*Pin) && Pin->Direction == EGPD_Output && Enum->IsNative() && EnumEntries.Contains(Pin->GetFName()))
+	if (Enum != nullptr && K2Schema->IsExecPin(*Pin) && Pin->Direction == EGPD_Output && Enum->IsNative() && EnumEntries.Contains(Pin->GetFName()))
 	{
 		// Allow native enum switch pins to be searchable by C++ enum name
 		OutTaggedMetaData.Add(FSearchTagDataPair(FFindInBlueprintSearchTags::FiB_NativeName, FText::FromString(Pin->GetName())));
