Add a function to check if a user defined pin exists by FName to EditablePinBase

#rnx

diff --git a/Engine/Source/Editor/BlueprintGraph/Classes/K2Node_EditablePinBase.h b/Engine/Source/Editor/BlueprintGraph/Classes/K2Node_EditablePinBase.h
--- a/Engine/Source/Editor/BlueprintGraph/Classes/K2Node_EditablePinBase.h
+++ b/Engine/Source/Editor/BlueprintGraph/Classes/K2Node_EditablePinBase.h
@@ -144,6 +144,14 @@ class BLUEPRINTGRAPH_VTABLE UK2Node_EditablePinBase : public UK2Node
 	 */
 	BLUEPRINTGRAPH_API void RemoveUserDefinedPinByName(const FName PinName);
 
+	/**
+	* Check if a pin with this name exists in the user defined pin set
+	* 
+	* @param	PinName name of pin check existence of
+	* @return	True if a user defined pin with this name exists
+	*/
+	bool UserDefinedPinExists(const FName PinName) const;
+
 	/**
 	 * Creates a new pin on the node from the specified user pin info.
 	 * Must be overridden so each type of node can ensure that the pin is created in the proper direction, etc
diff --git a/Engine/Source/Editor/BlueprintGraph/Private/K2Node_EditablePinBase.cpp b/Engine/Source/Editor/BlueprintGraph/Private/K2Node_EditablePinBase.cpp
--- a/Engine/Source/Editor/BlueprintGraph/Private/K2Node_EditablePinBase.cpp
+++ b/Engine/Source/Editor/BlueprintGraph/Private/K2Node_EditablePinBase.cpp
@@ -137,6 +137,14 @@ void UK2Node_EditablePinBase::RemoveUserDefinedPinByName(const FName PinName)
 	});
 }
 
+bool UK2Node_EditablePinBase::UserDefinedPinExists(const FName PinName) const
+{
+	return UserDefinedPins.ContainsByPredicate([&PinName](const TSharedPtr<FUserPinInfo>& UDPin)
+	{
+		return UDPin.IsValid() && (UDPin->PinName == PinName);
+	});
+}
+
 void UK2Node_EditablePinBase::ExportCustomProperties(FOutputDevice& Out, uint32 Indent)
 {
 	Super::ExportCustomProperties(Out, Indent);
