Duplicating an interface function no longer duplicates the pins on the graph

When the schema duplicates a graph, have it path up all function terminators, not just the entry.
When an interface graph is duplicated, check to make sure a user defined pin exists first before adding it.

diff --git a/Engine/Source/Editor/BlueprintGraph/Private/EdGraphSchema_K2.cpp b/Engine/Source/Editor/BlueprintGraph/Private/EdGraphSchema_K2.cpp
--- a/Engine/Source/Editor/BlueprintGraph/Private/EdGraphSchema_K2.cpp
+++ b/Engine/Source/Editor/BlueprintGraph/Private/EdGraphSchema_K2.cpp
@@ -6554,18 +6554,16 @@ UEdGraph* UEdGraphSchema_K2::DuplicateGraph(UEdGraph* GraphToDuplicate) const
 			for (int32 NodeIndex = 0; NodeIndex < NewGraph->Nodes.Num(); ++NodeIndex)
 			{
 				UEdGraphNode* Node = NewGraph->Nodes[NodeIndex];
-				if (UK2Node_FunctionEntry* EntryNode = Cast<UK2Node_FunctionEntry>(Node))
+				if (UK2Node_FunctionTerminator* TerminatorNode = Cast<UK2Node_FunctionTerminator>(Node))
 				{
-					if (EntryNode->FunctionReference.GetMemberName() == GraphToDuplicate->GetFName())
+					if (TerminatorNode->FunctionReference.GetMemberName() == GraphToDuplicate->GetFName())
 					{
-						EntryNode->Modify();
+						TerminatorNode->Modify();
 						
 						// We're duplicating the graph, so fully reset the member reference (including the GUID!)
 						FMemberReference NewRef;
 						NewRef.SetMemberName(NewGraph->GetFName());
-						EntryNode->FunctionReference = NewRef;
-
-						break;
+						TerminatorNode->FunctionReference = NewRef;
 					}
 				}
 				// Rename any custom events to be unique
diff --git a/Engine/Source/Editor/BlueprintGraph/Private/K2Node.cpp b/Engine/Source/Editor/BlueprintGraph/Private/K2Node.cpp
--- a/Engine/Source/Editor/BlueprintGraph/Private/K2Node.cpp
+++ b/Engine/Source/Editor/BlueprintGraph/Private/K2Node.cpp
@@ -271,6 +271,13 @@ bool UK2Node::CreatePinsForFunctionEntryExit(const UFunction* Function, bool bFo
 	{
 		UProperty* Param = *PropIt;
 
+		// Don't create a new pin if one exists already! 
+		// @see UE-79032, UE-58390
+		if (FindPin(Param->GetFName()))
+		{
+			continue;
+		}
+
 		const bool bIsFunctionInput = !Param->HasAnyPropertyFlags(CPF_OutParm) || Param->HasAnyPropertyFlags(CPF_ReferenceParm);
 
 		if (bIsFunctionInput == bForFunctionEntry)
diff --git a/Engine/Source/Editor/BlueprintGraph/Private/K2Node_EditablePinBase.cpp b/Engine/Source/Editor/BlueprintGraph/Private/K2Node_EditablePinBase.cpp
--- a/Engine/Source/Editor/BlueprintGraph/Private/K2Node_EditablePinBase.cpp
+++ b/Engine/Source/Editor/BlueprintGraph/Private/K2Node_EditablePinBase.cpp
@@ -80,7 +80,7 @@ void UK2Node_EditablePinBase::AllocateDefaultPins()
 	for(int32 i = 0; i < UserDefinedPins.Num(); i++)
 	{
 		FText DummyErrorMsg;
-		if (!IsEditable() || CanCreateUserDefinedPin(UserDefinedPins[i]->PinType, UserDefinedPins[i]->DesiredPinDirection, DummyErrorMsg))
+		if ((!IsEditable() || CanCreateUserDefinedPin(UserDefinedPins[i]->PinType, UserDefinedPins[i]->DesiredPinDirection, DummyErrorMsg)) && !FindPin(UserDefinedPins[i]->PinName))
 		{
 			CreatePinFromUserDefinition(UserDefinedPins[i]);
 		}
diff --git a/Engine/Source/Editor/BlueprintGraph/Private/K2Node_FunctionTerminator.cpp b/Engine/Source/Editor/BlueprintGraph/Private/K2Node_FunctionTerminator.cpp
--- a/Engine/Source/Editor/BlueprintGraph/Private/K2Node_FunctionTerminator.cpp
+++ b/Engine/Source/Editor/BlueprintGraph/Private/K2Node_FunctionTerminator.cpp
@@ -109,12 +109,17 @@ void UK2Node_FunctionTerminator::PromoteFromInterfaceOverride(bool bIsPrimaryTer
 {
 	// Remove the signature class, that is not relevant.
 	FunctionReference.SetSelfMember(FunctionReference.GetMemberName());
-	TArray<UEdGraphPin*> OriginalPins = Pins;
-	for (const UEdGraphPin* Pin : OriginalPins)
+	
+	// For every pin that has been defined, make it a user defined pin if we can
+	for (UEdGraphPin* const Pin : Pins)
 	{
-		if (Pin->PinType.PinCategory != UEdGraphSchema_K2::PC_Exec)
+		if (Pin->PinType.PinCategory != UEdGraphSchema_K2::PC_Exec && !UserDefinedPinExists(Pin->PinName))
 		{
-			CreateUserDefinedPin(Pin->PinName, Pin->PinType, Pin->Direction, false);
+			TSharedPtr<FUserPinInfo> NewPinInfo = MakeShareable(new FUserPinInfo());
+			NewPinInfo->PinName = Pin->PinName;
+			NewPinInfo->PinType = Pin->PinType;
+			NewPinInfo->DesiredPinDirection = Pin->Direction;
+			UserDefinedPins.Add(NewPinInfo);
 		}
 	}
 	const UEdGraphSchema_K2* Schema = GetDefault<UEdGraphSchema_K2>();
