World subsystem non-unity

diff --git a/Engine/Source/Editor/BlueprintGraph/Private/K2Node_GetSubsystem.cpp b/Engine/Source/Editor/BlueprintGraph/Private/K2Node_GetSubsystem.cpp
--- a/Engine/Source/Editor/BlueprintGraph/Private/K2Node_GetSubsystem.cpp
+++ b/Engine/Source/Editor/BlueprintGraph/Private/K2Node_GetSubsystem.cpp
@@ -17,6 +17,7 @@
 #include "Subsystems/EditorSubsystemBlueprintLibrary.h"
 #include "GameFramework/PlayerController.h"
 #include "Kismet2/BlueprintEditorUtils.h"
+#include "Subsystems/WorldSubsystem.h"
 
 
 // ************************************************************************************
