Fix "Paste Here" action in Blueprints doing nothing by changing it to not derive from FEdGraphSchemaAction_K2NewNode

diff --git a/Engine/Source/Editor/BlueprintGraph/Classes/EdGraphSchema_K2_Actions.h b/Engine/Source/Editor/BlueprintGraph/Classes/EdGraphSchema_K2_Actions.h
--- a/Engine/Source/Editor/BlueprintGraph/Classes/EdGraphSchema_K2_Actions.h
+++ b/Engine/Source/Editor/BlueprintGraph/Classes/EdGraphSchema_K2_Actions.h
@@ -392,7 +392,7 @@ struct BLUEPRINTGRAPH_API FEdGraphSchemaAction_K2TargetNode : public FEdGraphSch
 
 /** Action to paste at this location on graph*/
 USTRUCT()
-struct BLUEPRINTGRAPH_API FEdGraphSchemaAction_K2PasteHere : public FEdGraphSchemaAction_K2NewNode
+struct BLUEPRINTGRAPH_API FEdGraphSchemaAction_K2PasteHere : public FEdGraphSchemaAction
 {
 	GENERATED_USTRUCT_BODY()
 
@@ -400,12 +400,12 @@ struct BLUEPRINTGRAPH_API FEdGraphSchemaAction_K2PasteHere : public FEdGraphSche
 	static FName StaticGetTypeId() {static FName Type("FEdGraphSchemaAction_K2PasteHere"); return Type;}
 	virtual FName GetTypeId() const override { return StaticGetTypeId(); } 
 
-	FEdGraphSchemaAction_K2PasteHere ()
-		: FEdGraphSchemaAction_K2NewNode()
+	FEdGraphSchemaAction_K2PasteHere()
+		: FEdGraphSchemaAction()
 	{}
 
 	FEdGraphSchemaAction_K2PasteHere (FText InNodeCategory, FText InMenuDesc, FText InToolTip, const int32 InGrouping)
-		: FEdGraphSchemaAction_K2NewNode(MoveTemp(InNodeCategory), MoveTemp(InMenuDesc), MoveTemp(InToolTip), InGrouping)
+		: FEdGraphSchemaAction(MoveTemp(InNodeCategory), MoveTemp(InMenuDesc), MoveTemp(InToolTip), InGrouping)
 	{}
 
 	// FEdGraphSchemaAction interface
diff --git a/Engine/Source/Editor/Kismet/Private/BlueprintActionMenuUtils.cpp b/Engine/Source/Editor/Kismet/Private/BlueprintActionMenuUtils.cpp
--- a/Engine/Source/Editor/Kismet/Private/BlueprintActionMenuUtils.cpp
+++ b/Engine/Source/Editor/Kismet/Private/BlueprintActionMenuUtils.cpp
@@ -754,7 +754,6 @@ const UK2Node* FBlueprintActionMenuUtils::ExtractNodeTemplateFromAction(TSharedP
 			ActionId == FEdGraphSchemaAction_K2AddCustomEvent::StaticGetTypeId() ||
 			ActionId == FEdGraphSchemaAction_K2AddCallOnActor::StaticGetTypeId() ||
 			ActionId == FEdGraphSchemaAction_K2TargetNode::StaticGetTypeId() ||
-			ActionId == FEdGraphSchemaAction_K2PasteHere::StaticGetTypeId() ||
 			ActionId == FEdGraphSchemaAction_K2Event::StaticGetTypeId() || 
 			ActionId == FEdGraphSchemaAction_K2AddEvent::StaticGetTypeId() ||
 			ActionId == FEdGraphSchemaAction_K2InputAction::StaticGetTypeId())
