Fix a regression introduced by the recent feature add to rename functions via the entry node.

#rnx

diff --git a/Engine/Source/Editor/BlueprintGraph/Private/K2Node_FunctionEntry.cpp b/Engine/Source/Editor/BlueprintGraph/Private/K2Node_FunctionEntry.cpp
--- a/Engine/Source/Editor/BlueprintGraph/Private/K2Node_FunctionEntry.cpp
+++ b/Engine/Source/Editor/BlueprintGraph/Private/K2Node_FunctionEntry.cpp
@@ -321,7 +321,7 @@ FText UK2Node_FunctionEntry::GetNodeTitle(ENodeTitleType::Type TitleType) const
 
 void UK2Node_FunctionEntry::OnRenameNode(const FString& NewName)
 {
-	CustomGeneratedFunctionName = FName(*NewName);
+	// Note: RenameGraph() will handle the rename operation for this node as well.
 	FBlueprintEditorUtils::RenameGraph(GetGraph(), NewName);
 }
 
