PR #5248: UE-22843: Private and protected BP function restrictions (Contributed by [REDACTED_PERSON]) - reworked by me

diff --git a/Engine/Source/Editor/BlueprintGraph/Private/EdGraphSchema_K2.cpp b/Engine/Source/Editor/BlueprintGraph/Private/EdGraphSchema_K2.cpp
--- a/Engine/Source/Editor/BlueprintGraph/Private/EdGraphSchema_K2.cpp
+++ b/Engine/Source/Editor/BlueprintGraph/Private/EdGraphSchema_K2.cpp
@@ -762,44 +762,50 @@ bool UEdGraphSchema_K2::CanFunctionBeUsedInGraph(const UClass* InClass, const UF
 			return false;
 		}
 
-		const bool bIsProtected = InFunction->GetBoolMetaData(FBlueprintMetadata::MD_Protected);
-		const bool bFuncBelongsToSubClass = InClass->IsChildOf(InFunction->GetOuterUClass());
-		if (bIsProtected)
+		const bool bIsNotNative = !FBlueprintEditorUtils::IsNativeSignature(InFunction);
+		if(!bIsNotNative)
 		{
-			const bool bAllowProtectedFuncs = (InAllowedFunctionTypes & FT_Protected) != 0;
-			if (!bAllowProtectedFuncs)
+			// Blueprint functions visibility flags can be enforced in blueprints - native functions
+			// are often using these flags to only hide functionality from other native functions:
+			const bool bIsProtected = (InFunction->FunctionFlags & FUNC_Protected) != 0;
+			const bool bFuncBelongsToSubClass = InClass->IsChildOf(InFunction->GetOuterUClass()->GetSuperStruct());
+			if (bIsProtected)
 			{
-				if(OutReason != nullptr)
+				const bool bAllowProtectedFuncs = (InAllowedFunctionTypes & FT_Protected) != 0;
+				if (!bAllowProtectedFuncs)
 				{
-					*OutReason = LOCTEXT("ProtectedFunctionsNotAllowed", "Protected functions are not allowed.");
+					if(OutReason != nullptr)
+					{
+						*OutReason = LOCTEXT("ProtectedFunctionsNotAllowed", "Protected functions are not allowed.");
+					}
+
+					return false;
 				}
 
-				return false;
+				if (!bFuncBelongsToSubClass)
+				{
+					if(OutReason != nullptr)
+					{
+						*OutReason = LOCTEXT("ProtectedFunctionInaccessible", "Function is protected and inaccessible.");
+					}
+
+					return false;
+				}
 			}
 
-			if (!bFuncBelongsToSubClass)
+			const bool bIsPrivate = (InFunction->FunctionFlags & FUNC_Private) != 0;
+			const bool bFuncBelongsToClass = bFuncBelongsToSubClass && (InClass->GetSuperStruct() == InFunction->GetOuterUClass()->GetSuperStruct());
+			if ( bIsPrivate && !bFuncBelongsToClass)
 			{
 				if(OutReason != nullptr)
 				{
-					*OutReason = LOCTEXT("ProtectedFunctionInaccessible", "Function is protected and inaccessible.");
+					*OutReason = LOCTEXT("PrivateFunctionInaccessible", "Function is private and inaccessible.");
 				}
 
 				return false;
 			}
 		}
 
-		const bool bIsPrivate = InFunction->GetBoolMetaData(FBlueprintMetadata::MD_Private);
-		const bool bFuncBelongsToClass = bFuncBelongsToSubClass && (InFunction->GetOuterUClass() == InClass);
-		if (bIsPrivate && !bFuncBelongsToClass)
-		{
-			if(OutReason != nullptr)
-			{
-				*OutReason = LOCTEXT("PrivateFunctionInaccessible", "Function is private and inaccessible.");
-			}
-
-			return false;
-		}
-
 		const bool bIsUnsafeForConstruction = InFunction->GetBoolMetaData(FBlueprintMetadata::MD_UnsafeForConstructionScripts);	
 		if (bIsUnsafeForConstruction && bIsConstructionScript)
 		{
@@ -913,8 +919,18 @@ bool UEdGraphSchema_K2::CanUserKismetCallFunction(const UFunction* Function)
 
 bool UEdGraphSchema_K2::CanKismetOverrideFunction(const UFunction* Function)
 {
-	return  Function && 
-		(Function->HasAllFunctionFlags(FUNC_BlueprintEvent) && !Function->HasAllFunctionFlags(FUNC_Delegate) && !Function->GetBoolMetaData(FBlueprintMetadata::MD_BlueprintInternalUseOnly) && (!Function->HasMetaData(FBlueprintMetadata::MD_DeprecatedFunction) || GetDefault<UBlueprintEditorSettings>()->bExposeDeprecatedFunctions));
+	return  
+		Function && 
+		(
+			Function->HasAllFunctionFlags(FUNC_BlueprintEvent)
+			&& !Function->HasAllFunctionFlags(FUNC_Delegate) && 
+			!Function->GetBoolMetaData(FBlueprintMetadata::MD_BlueprintInternalUseOnly) && 
+			(!Function->HasMetaData(FBlueprintMetadata::MD_DeprecatedFunction) || GetDefault<UBlueprintEditorSettings>()->bExposeDeprecatedFunctions)
+		) &&
+		(
+			FBlueprintEditorUtils::IsNativeSignature(Function) ||
+			(Function->FunctionFlags & FUNC_Private) == 0
+		);
 }
 
 bool UEdGraphSchema_K2::HasFunctionAnyOutputParameter(const UFunction* InFunction)
diff --git a/Engine/Source/Editor/BlueprintGraph/Private/K2Node_CallFunction.cpp b/Engine/Source/Editor/BlueprintGraph/Private/K2Node_CallFunction.cpp
--- a/Engine/Source/Editor/BlueprintGraph/Private/K2Node_CallFunction.cpp
+++ b/Engine/Source/Editor/BlueprintGraph/Private/K2Node_CallFunction.cpp
@@ -1321,6 +1321,41 @@ bool UK2Node_CallFunction::IsActionFilteredOut(FBlueprintActionFilter const& Fil
 		bIsFilteredOut |= !CanPasteHere(TargetGraph);
 	}
 
+	if(const UFunction* TargetFunction = GetTargetFunction())
+	{
+		const bool bIsProtected = (TargetFunction->FunctionFlags & FUNC_Protected) != 0;
+		const bool bIsPrivate = (TargetFunction->FunctionFlags & FUNC_Private) != 0;
+		const UClass* OwningClass = TargetFunction->GetOwnerClass();
+		if( (bIsProtected || bIsPrivate) && !FBlueprintEditorUtils::IsNativeSignature(TargetFunction) && OwningClass)
+		{
+			OwningClass = OwningClass->GetAuthoritativeClass();
+			// we can filter private and protected blueprints that are unrelated:
+			bool bAccessibleInAll = true;
+			for (const UBlueprint* Blueprint : Filter.Context.Blueprints)
+			{
+				UClass* AuthoritativeClass = Blueprint->GeneratedClass;
+				if(!AuthoritativeClass)
+				{
+					continue;
+				}
+
+				if(bIsPrivate)
+				{
+					bAccessibleInAll = bAccessibleInAll && AuthoritativeClass == OwningClass;
+				}
+				else if(bIsProtected)
+				{
+					bAccessibleInAll = bAccessibleInAll && AuthoritativeClass->IsChildOf(OwningClass);
+				}
+			}
+
+			if(!bAccessibleInAll)
+			{
+				bIsFilteredOut = true;
+			}
+		}
+	}
+
 	return bIsFilteredOut;
 }
 
@@ -2009,6 +2044,25 @@ void UK2Node_CallFunction::ValidateNodeDuringCompilation(class FCompilerResultsL
 				MessageLog.Warning(*LOCTEXT("FunctionUnsafeInContext", "Function '@@' is unsafe to call from blueprints of class '@@'.").ToString(), this, ParentClass);
 			}
 		}
+
+		if(!FBlueprintEditorUtils::IsNativeSignature(Function))
+		{
+			// enforce protected function restriction
+			const bool bIsProtected = (Function->FunctionFlags & FUNC_Protected) != 0;
+			const bool bFuncBelongsToSubClass = Blueprint->SkeletonGeneratedClass->IsChildOf(Function->GetOuterUClass());
+			if (bIsProtected && !bFuncBelongsToSubClass)
+			{
+				MessageLog.Error(*LOCTEXT("FunctionPrivateAccessed", "Function '@@' is protected and can't be accessed outside of its hierarchy.").ToString(), this);
+			}
+
+			// enforce private function restriction
+			const bool bIsPrivate = (Function->FunctionFlags & FUNC_Private) != 0;
+			const bool bFuncBelongsToClass = bFuncBelongsToSubClass && (Blueprint->SkeletonGeneratedClass == Function->GetOuterUClass());
+			if (bIsPrivate && !bFuncBelongsToClass)
+			{
+				MessageLog.Error(*LOCTEXT("FunctionPrivateAccessed", "Function '@@' is private and can't be accessed outside of its defined class '@@'.").ToString(), this, Function->GetOuterUClass());
+			}
+		}
 	}
 
 	FDynamicOutputHelper::VerifyNode(this, MessageLog);
diff --git a/Engine/Source/Editor/UnrealEd/Private/Kismet2/BlueprintEditorUtils.cpp b/Engine/Source/Editor/UnrealEd/Private/Kismet2/BlueprintEditorUtils.cpp
--- a/Engine/Source/Editor/UnrealEd/Private/Kismet2/BlueprintEditorUtils.cpp
+++ b/Engine/Source/Editor/UnrealEd/Private/Kismet2/BlueprintEditorUtils.cpp
@@ -767,6 +767,21 @@ UClass* FBlueprintEditorUtils::FindFirstNativeClass(UClass* Class)
 	return Class;
 }
 
+bool FBlueprintEditorUtils::IsNativeSignature(const UFunction* Fn)
+{
+	if(UClass* OwningClass = Fn->GetOwnerClass())
+	{
+		if(UClass* PotentialNativeOwner = FindFirstNativeClass(OwningClass))
+		{
+			return PotentialNativeOwner->FindFunctionByName(Fn->GetFName()) != nullptr;
+		}
+	}
+
+	// All UFunctions are owned by UClasses, but if we encounter one that
+	// is not lets just fall back to the flag:
+	return Fn->HasAnyFunctionFlags(FUNC_Native);
+}
+
 void FBlueprintEditorUtils::GetAllGraphNames(const UBlueprint* Blueprint, TSet<FName>& GraphNames)
 {
 	TArray< UEdGraph* > GraphList;
@@ -1763,10 +1778,7 @@ void FBlueprintEditorUtils::PostDuplicateBlueprint(UBlueprint* Blueprint, bool b
 			Blueprint->SkeletonGeneratedClass = nullptr;
 
 			// Make sure the new blueprint has a shiny new class
-			IKismetCompilerInterface& Compiler = FModuleManager::LoadModuleChecked<IKismetCompilerInterface>(KISMET_COMPILER_MODULENAME);
 			FCompilerResultsLog Results;
-			FKismetCompilerOptions CompileOptions;
-			CompileOptions.bIsDuplicationInstigated = true;
 
 			FName NewSkelClassName, NewGenClassName;
 			Blueprint->GetBlueprintClassNames(NewGenClassName, NewSkelClassName);
@@ -1779,9 +1791,6 @@ void FBlueprintEditorUtils::PostDuplicateBlueprint(UBlueprint* Blueprint, bool b
 			NewClass->SetSuperStruct(Blueprint->ParentClass);
 			Blueprint->bHasBeenRegenerated = true;		// Set to true, similar to CreateBlueprint, since we've regerated the class by duplicating it
 
-			// Since we just duplicated the generated class above, we don't need to do a full compile below
-			CompileOptions.CompileType = EKismetCompileType::SkeletonOnly;
-
 			TMap<UObject*, UObject*> OldToNewMap;
 
 			UClass* NewBPGCAsClass = Blueprint->GeneratedClass;
diff --git a/Engine/Source/Editor/UnrealEd/Public/Kismet2/BlueprintEditorUtils.h b/Engine/Source/Editor/UnrealEd/Public/Kismet2/BlueprintEditorUtils.h
--- a/Engine/Source/Editor/UnrealEd/Public/Kismet2/BlueprintEditorUtils.h
+++ b/Engine/Source/Editor/UnrealEd/Public/Kismet2/BlueprintEditorUtils.h
@@ -647,6 +647,9 @@ public:
 	/** return find first native class in the hierarchy */
 	static UClass* FindFirstNativeClass(UClass* Class);
 
+	/** returns true if this blueprints signature (inc. visibility) was determined by UHT, rather than the blueprint compiler */
+	static bool IsNativeSignature(const UFunction* Fn);
+
 	/**
 	 * Gets the names of all graphs in the Blueprint
 	 *
