Blueprint debugger now properly steps over knot nodes

diff --git a/Engine/Source/Editor/BlueprintGraph/Classes/K2Node_Knot.h b/Engine/Source/Editor/BlueprintGraph/Classes/K2Node_Knot.h
--- a/Engine/Source/Editor/BlueprintGraph/Classes/K2Node_Knot.h
+++ b/Engine/Source/Editor/BlueprintGraph/Classes/K2Node_Knot.h
@@ -52,6 +52,8 @@ public:
 		return Pins[1];
 	}
 
+	BLUEPRINTGRAPH_API UEdGraphNode* GetExecTerminal() const;
+
 private:
 	void PropagatePinType();
 	void PropagatePinTypeFromDirection(bool bFromInput);
diff --git a/Engine/Source/Editor/BlueprintGraph/Private/K2Node_Knot.cpp b/Engine/Source/Editor/BlueprintGraph/Private/K2Node_Knot.cpp
--- a/Engine/Source/Editor/BlueprintGraph/Private/K2Node_Knot.cpp
+++ b/Engine/Source/Editor/BlueprintGraph/Private/K2Node_Knot.cpp
@@ -202,6 +202,41 @@ void UK2Node_Knot::GetMenuActions(FBlueprintActionDatabaseRegistrar& ActionRegis
 	}
 }
 
+UEdGraphNode* UK2Node_Knot::GetExecTerminal() const
+{
+	// if we're not an exec wire, just bail:
+	const UK2Node_Knot* Knot = this;
+	UEdGraphPin* OutputPin = Knot->GetOutputPin();
+	if(!OutputPin)
+	{
+		return nullptr;
+	}
+
+	if(OutputPin->PinType.PinCategory != UEdGraphSchema_K2::PC_Exec)
+	{
+		return nullptr;
+	}
+
+	UEdGraphNode* Result = nullptr;
+	while(Knot)
+	{
+		UEdGraphPin* Next = Knot->GetOutputPin();
+		if(Next && Next->LinkedTo.Num() > 0)
+		{
+			// knots for exec pins can have only one connection:
+			Result = Next->LinkedTo[0]->GetOwningNode();
+			Knot = Cast<const UK2Node_Knot>(Result);
+		}
+		else
+		{
+			// dead end knot:
+			Knot = nullptr;
+			Result = nullptr;
+		}
+	}
+	return Result;
+}
+
 bool UK2Node_Knot::ShouldOverridePinNames() const
 {
 	return true;
diff --git a/Engine/Source/Editor/UnrealEd/Private/Kismet2/KismetDebugUtilities.cpp b/Engine/Source/Editor/UnrealEd/Private/Kismet2/KismetDebugUtilities.cpp
--- a/Engine/Source/Editor/UnrealEd/Private/Kismet2/KismetDebugUtilities.cpp
+++ b/Engine/Source/Editor/UnrealEd/Private/Kismet2/KismetDebugUtilities.cpp
@@ -29,6 +29,7 @@
 #include "K2Node.h"
 #include "K2Node_Tunnel.h"
 #include "K2Node_FunctionEntry.h"
+#include "K2Node_Knot.h"
 #include "K2Node_MacroInstance.h"
 #include "K2Node_Message.h"
 #include "Kismet2/KismetEditorUtilities.h"
@@ -167,7 +168,17 @@ void FKismetDebugUtilities::RequestStepOver()
 						{
 							for(UEdGraphPin* LinkedTo : Pin->LinkedTo)
 							{
-								Data.TargetGraphNodes.AddUnique(LinkedTo->GetOwningNode());
+								UEdGraphNode* GraphNode = LinkedTo->GetOwningNode();
+								if(UK2Node_Knot* Knot = Cast<UK2Node_Knot>(GraphNode))
+								{
+									// search the knot chain to find the actual node:
+									GraphNode = Knot->GetExecTerminal();
+								}
+
+								if(GraphNode)
+								{
+									Data.TargetGraphNodes.AddUnique(GraphNode);
+								}
 							}
 						}
 					}
