Fix inconsistent nullptr check introduced in 9230757
#rnx

diff --git a/Engine/Source/Editor/BlueprintGraph/Private/K2Node_CallFunction.cpp b/Engine/Source/Editor/BlueprintGraph/Private/K2Node_CallFunction.cpp
--- a/Engine/Source/Editor/BlueprintGraph/Private/K2Node_CallFunction.cpp
+++ b/Engine/Source/Editor/BlueprintGraph/Private/K2Node_CallFunction.cpp
@@ -2045,7 +2045,7 @@ void UK2Node_CallFunction::ValidateNodeDuringCompilation(class FCompilerResultsL
 			}
 		}
 
-		if(!FBlueprintEditorUtils::IsNativeSignature(Function))
+		if(Blueprint && !FBlueprintEditorUtils::IsNativeSignature(Function))
 		{
 			// enforce protected function restriction
 			const bool bIsProtected = (Function->FunctionFlags & FUNC_Protected) != 0;
