Prevent running new function visibility enforcement on old blueprints - allowing teams to fix the new errors at their leisure

diff --git a/Engine/Source/Editor/BlueprintGraph/Private/K2Node_CallFunction.cpp b/Engine/Source/Editor/BlueprintGraph/Private/K2Node_CallFunction.cpp
--- a/Engine/Source/Editor/BlueprintGraph/Private/K2Node_CallFunction.cpp
+++ b/Engine/Source/Editor/BlueprintGraph/Private/K2Node_CallFunction.cpp
@@ -4,6 +4,7 @@
 #include "BlueprintCompilationManager.h"
 #include "BlueprintEditorSettings.h"
 #include "UObject/UObjectHash.h"
+#include "UObject/FrameworkObjectVersion.h"
 #include "UObject/Interface.h"
 #include "UObject/PropertyPortFlags.h"
 #include "Kismet/BlueprintFunctionLibrary.h"
@@ -2048,11 +2049,20 @@ void UK2Node_CallFunction::ValidateNodeDuringCompilation(class FCompilerResultsL
 		if(Blueprint && !FBlueprintEditorUtils::IsNativeSignature(Function))
 		{
 			// enforce protected function restriction
+			const bool bCanTreatAsError = Blueprint->GetLinkerCustomVersion(FFrameworkObjectVersion::GUID) >= FFrameworkObjectVersion::EnforceBlueprintFunctionVisibility;
+
 			const bool bIsProtected = (Function->FunctionFlags & FUNC_Protected) != 0;
 			const bool bFuncBelongsToSubClass = Blueprint->SkeletonGeneratedClass->IsChildOf(Function->GetOuterUClass());
 			if (bIsProtected && !bFuncBelongsToSubClass)
 			{
-				MessageLog.Error(*LOCTEXT("FunctionPrivateAccessed", "Function '@@' is protected and can't be accessed outside of its hierarchy.").ToString(), this);
+				if(bCanTreatAsError)
+				{
+					MessageLog.Error(*LOCTEXT("FunctionPrivateAccessed", "Function '@@' is protected and can't be accessed outside of its hierarchy.").ToString(), this);
+				}
+				else
+				{
+					MessageLog.Note(*LOCTEXT("FunctionPrivateAccessedNote", "Function '@@' is protected and can't be accessed outside of its hierarchy - this will be an error if the asset is resaved.").ToString(), this);
+				}
 			}
 
 			// enforce private function restriction
@@ -2060,7 +2070,14 @@ void UK2Node_CallFunction::ValidateNodeDuringCompilation(class FCompilerResultsL
 			const bool bFuncBelongsToClass = bFuncBelongsToSubClass && (Blueprint->SkeletonGeneratedClass == Function->GetOuterUClass());
 			if (bIsPrivate && !bFuncBelongsToClass)
 			{
-				MessageLog.Error(*LOCTEXT("FunctionPrivateAccessed", "Function '@@' is private and can't be accessed outside of its defined class '@@'.").ToString(), this, Function->GetOuterUClass());
+				if(bCanTreatAsError)
+				{
+					MessageLog.Error(*LOCTEXT("FunctionPrivateAccessed", "Function '@@' is private and can't be accessed outside of its defined class '@@'.").ToString(), this, Function->GetOuterUClass());
+				}
+				else
+				{
+					MessageLog.Note(*LOCTEXT("FunctionPrivateAccessedNote", "Function '@@' is private and can't be accessed outside of its defined class '@@' - this will be an error if the asset is resaved.").ToString(), this, Function->GetOuterUClass());
+				}
 			}
 		}
 	}
diff --git a/Engine/Source/Runtime/Core/Public/UObject/FrameworkObjectVersion.h b/Engine/Source/Runtime/Core/Public/UObject/FrameworkObjectVersion.h
--- a/Engine/Source/Runtime/Core/Public/UObject/FrameworkObjectVersion.h
+++ b/Engine/Source/Runtime/Core/Public/UObject/FrameworkObjectVersion.h
@@ -122,6 +122,9 @@ struct CORE_API FFrameworkObjectVersion
 		// No longer serialize the legacy flag that indicates this state, as it is now implied since we don't serialize the skeleton CDO
 		BlueprintGeneratedClassIsAlwaysAuthoritative,
 
+		// Enforce visibility of blueprint functions - e.g. raise an error if calling a private function from another blueprint:
+		EnforceBlueprintFunctionVisibility,
+
 		// -----<new versions can be added above this line>-------------------------------------------------
 		VersionPlusOne,
 		LatestVersion = VersionPlusOne - 1
