Fix typo
#rnx

diff --git a/Engine/Source/Editor/BlueprintGraph/Private/K2Node_StructOperation.cpp b/Engine/Source/Editor/BlueprintGraph/Private/K2Node_StructOperation.cpp
--- a/Engine/Source/Editor/BlueprintGraph/Private/K2Node_StructOperation.cpp
+++ b/Engine/Source/Editor/BlueprintGraph/Private/K2Node_StructOperation.cpp
@@ -18,7 +18,7 @@ UK2Node_StructOperation::UK2Node_StructOperation(const FObjectInitializer& Objec
 void UK2Node_StructOperation::ValidateNodeDuringCompilation(class FCompilerResultsLog& MessageLog) const
 {
 	// Skip UK2Node_Variable's validation because it doesn't need a property (see CL# 1756451)
-	Super::Super::ValidateNodeDuringCompilation(MessageLog);
+	Super::ValidateNodeDuringCompilation(MessageLog);
 }
 
 bool UK2Node_StructOperation::HasExternalDependencies(TArray<class UStruct*>* OptionalOutput) const
