Revert change to behavior of ::CanKismetOverrideFunction
#rnx

diff --git a/Engine/Source/Editor/BlueprintGraph/Private/EdGraphSchema_K2.cpp b/Engine/Source/Editor/BlueprintGraph/Private/EdGraphSchema_K2.cpp
--- a/Engine/Source/Editor/BlueprintGraph/Private/EdGraphSchema_K2.cpp
+++ b/Engine/Source/Editor/BlueprintGraph/Private/EdGraphSchema_K2.cpp
@@ -926,10 +926,6 @@ bool UEdGraphSchema_K2::CanKismetOverrideFunction(const UFunction* Function)
 			&& !Function->HasAllFunctionFlags(FUNC_Delegate) && 
 			!Function->GetBoolMetaData(FBlueprintMetadata::MD_BlueprintInternalUseOnly) && 
 			(!Function->HasMetaData(FBlueprintMetadata::MD_DeprecatedFunction) || GetDefault<UBlueprintEditorSettings>()->bExposeDeprecatedFunctions)
-		) &&
-		(
-			FBlueprintEditorUtils::IsNativeSignature(Function) ||
-			(Function->FunctionFlags & FUNC_Private) == 0
 		);
 }
 
