Fix a memory leak in the Blueprint action database after a map load in the editor.

diff --git a/Engine/Source/Editor/BlueprintGraph/Private/BlueprintActionDatabase.cpp b/Engine/Source/Editor/BlueprintGraph/Private/BlueprintActionDatabase.cpp
--- a/Engine/Source/Editor/BlueprintGraph/Private/BlueprintActionDatabase.cpp
+++ b/Engine/Source/Editor/BlueprintGraph/Private/BlueprintActionDatabase.cpp
@@ -1277,7 +1277,21 @@ void FBlueprintActionDatabase::RefreshAll()
 	// Remove callbacks from blueprints
 	for (TObjectIterator<UBlueprint> BlueprintIt; BlueprintIt; ++BlueprintIt)
 	{
-		ClearAssetActions(*BlueprintIt);
+		UBlueprint* Blueprint = *BlueprintIt;
+
+		// Level script BPs are registered using the associated world context. This will clear any registered LSBP actions.
+		const bool bIsLevelScript = FBlueprintEditorUtils::IsLevelScriptBlueprint(Blueprint);
+		if (bIsLevelScript)
+		{
+			const ULevelScriptBlueprint* LSBP = CastChecked<ULevelScriptBlueprint>(Blueprint);
+			if (UWorld* World = LSBP->GetWorld())
+			{
+				ClearAssetActions(World);
+			}
+		}
+		
+		// Do this for all BP assets to both clear registered actions and remove callbacks. In the case of LSBPs, this will just remove callbacks.
+		ClearAssetActions(Blueprint);
 	}
 
 	ActionRegistry.Empty();
