Fix GetActorTransform not appearing in context when dragging off a transform pin

diff --git a/Engine/Source/Editor/BlueprintGraph/Private/EdGraphSchema_K2.cpp b/Engine/Source/Editor/BlueprintGraph/Private/EdGraphSchema_K2.cpp
--- a/Engine/Source/Editor/BlueprintGraph/Private/EdGraphSchema_K2.cpp
+++ b/Engine/Source/Editor/BlueprintGraph/Private/EdGraphSchema_K2.cpp
@@ -1013,7 +1013,7 @@ bool UEdGraphSchema_K2::FunctionHasParamOfType(const UFunction* InFunction, UEdG
 		if (!HiddenPins.Contains(FuncParam->GetFName()))
 		{
 			// See if this is the direction we want (input or output)
-			const bool bIsFunctionInput = !FuncParam->HasAnyPropertyFlags(CPF_OutParm) || FuncParam->HasAnyPropertyFlags(CPF_ReferenceParm);
+			const bool bIsFunctionInput = !FuncParam->HasAnyPropertyFlags(CPF_ReturnParm) && (!FuncParam->HasAnyPropertyFlags(CPF_OutParm) || FuncParam->HasAnyPropertyFlags(CPF_ReferenceParm));
 			if (bIsFunctionInput != bWantOutput)
 			{
 				// See if this pin has compatible types
