Potential fix for occasional crashes in the wild on BP action database tick.

diff --git a/Engine/Source/Editor/BlueprintGraph/Private/BlueprintActionDatabase.cpp b/Engine/Source/Editor/BlueprintGraph/Private/BlueprintActionDatabase.cpp
--- a/Engine/Source/Editor/BlueprintGraph/Private/BlueprintActionDatabase.cpp
+++ b/Engine/Source/Editor/BlueprintGraph/Private/BlueprintActionDatabase.cpp
@@ -1247,12 +1247,7 @@ void FBlueprintActionDatabase::Tick(float DeltaTime)
 	// Handle deferred removals.
 	while (ActionRemoveQueue.Num() > 0)
 	{
-		TArray<UBlueprintNodeSpawner*> NodeSpawners = ActionRegistry.FindAndRemoveChecked(ActionRemoveQueue.Pop());
-		for (UBlueprintNodeSpawner* Action : NodeSpawners)
-		{
-			check(Action);
-			Action->ClearCachedTemplateNode();
-		}
+		ClearAssetActions(ActionRemoveQueue.Pop());
 	}
 }
 
@@ -1578,31 +1573,45 @@ void FBlueprintActionDatabase::RefreshComponentActions()
 //------------------------------------------------------------------------------
 bool FBlueprintActionDatabase::ClearAssetActions(UObject* const AssetObject)
 {
-	FActionList* ActionList = ActionRegistry.Find(AssetObject);
+	FObjectKey ObjectKey(AssetObject);
+	return ClearAssetActions(ObjectKey);
+}
+
+//------------------------------------------------------------------------------
+bool FBlueprintActionDatabase::ClearAssetActions(const FObjectKey& AssetObjectKey)
+{
+	FActionList* ActionList = ActionRegistry.Find(AssetObjectKey);
 
 	bool const bHasEntry = (ActionList != nullptr);
 	if (bHasEntry)
 	{
 		for (UBlueprintNodeSpawner* Action : *ActionList)
 		{
-			// because some asserts expect everything to be cleaned up in a 
-			// single GC pass, we can't wait for the GC'd Action to release its
-			// template node from the cache
-			Action->ClearCachedTemplateNode();
+			if (Action != nullptr)
+			{
+				// because some asserts expect everything to be cleaned up in a 
+				// single GC pass, we can't wait for the GC'd Action to release its
+				// template node from the cache
+				Action->ClearCachedTemplateNode();
+			}
 		}
-		ActionRegistry.Remove(AssetObject);
+		ActionRegistry.Remove(AssetObjectKey);
 	}
 
-	if (UBlueprint* BlueprintAsset = Cast<UBlueprint>(AssetObject))
+	if (UObject* AssetObject = AssetObjectKey.ResolveObjectPtr())
 	{
-		BlueprintAsset->OnChanged().RemoveAll(this);
-		BlueprintAsset->OnCompiled().RemoveAll(this);
-	}
+		if (UBlueprint* BlueprintAsset = Cast<UBlueprint>(AssetObject))
+		{
+			BlueprintAsset->OnChanged().RemoveAll(this);
+			BlueprintAsset->OnCompiled().RemoveAll(this);
+		}
 
-	if (bHasEntry && (ActionList->Num() > 0) && !BlueprintActionDatabaseImpl::bIsInitializing)
-	{
-		EntryRemovedDelegate.Broadcast(AssetObject);
+		if (bHasEntry && (ActionList->Num() > 0) && !BlueprintActionDatabaseImpl::bIsInitializing)
+		{
+			EntryRemovedDelegate.Broadcast(AssetObject);
+		}
 	}
+	
 	return bHasEntry;
 }
 
diff --git a/Engine/Source/Editor/BlueprintGraph/Public/BlueprintActionDatabase.h b/Engine/Source/Editor/BlueprintGraph/Public/BlueprintActionDatabase.h
--- a/Engine/Source/Editor/BlueprintGraph/Public/BlueprintActionDatabase.h
+++ b/Engine/Source/Editor/BlueprintGraph/Public/BlueprintActionDatabase.h
@@ -107,13 +107,22 @@ public:
 	void RefreshComponentActions();
 
 	/**
-	 * Finds the database entry for the specified class and wipes it. The entry 
+	 * Finds the database entry for the specified object and wipes it. The entry 
 	 * won't be rebuilt, unless RefreshAssetActions() is explicitly called after.
 	 * 
 	 * @param  AssetObject	
 	 * @return True if an entry was found and removed.
 	 */
 	bool ClearAssetActions(UObject* const AssetObject);
+
+	/**
+	 * Finds the database entry for the specified object and wipes it. The entry
+	 * won't be rebuilt, unless RefreshAssetActions() is explicitly called after.
+	 *
+	 * @param  AssetObjectKey
+	 * @return True if an entry was found and removed.
+	 */
+	bool ClearAssetActions(const FObjectKey& AssetObjectKey);
 	
 	/**
 	 * Finds the database entry for the specified unloaded asset and wipes it.
