Fixed skeleton animation notifies not showing up until an animation asset is opened

Animation notifies do not show up in actor BPs any more, but unrelated to this change.

diff --git a/Engine/Source/Editor/BlueprintGraph/Private/BlueprintActionDatabase.cpp b/Engine/Source/Editor/BlueprintGraph/Private/BlueprintActionDatabase.cpp
--- a/Engine/Source/Editor/BlueprintGraph/Private/BlueprintActionDatabase.cpp
+++ b/Engine/Source/Editor/BlueprintGraph/Private/BlueprintActionDatabase.cpp
@@ -684,9 +684,9 @@ static void BlueprintActionDatabaseImpl::AddClassPropertyActions(UClass const* c
 			continue;
 		}
 
- 		bool const bIsDelegate = Property->IsA(UMulticastDelegateProperty::StaticClass());
- 		if (bIsDelegate)
- 		{
+		bool const bIsDelegate = Property->IsA(UMulticastDelegateProperty::StaticClass());
+		if (bIsDelegate)
+		{
 			UMulticastDelegateProperty* DelegateProperty = CastChecked<UMulticastDelegateProperty>(Property);
 			if (DelegateProperty->HasAnyPropertyFlags(CPF_BlueprintAssignable))
 			{
@@ -716,7 +716,7 @@ static void BlueprintActionDatabaseImpl::AddClassPropertyActions(UClass const* c
 			{
 				ActionListOut.Add(MakeActorBoundEventSpawner(DelegateProperty));
 			}
- 		}
+		}
 		else
 		{
 			UBlueprintVariableNodeSpawner* GetterSpawner = UBlueprintVariableNodeSpawner::CreateFromMemberOrParam(UK2Node_VariableGet::StaticClass(), Property);
@@ -784,9 +784,11 @@ static void BlueprintActionDatabaseImpl::AddClassCastActions(UClass* Class, FAct
 //------------------------------------------------------------------------------
 static void BlueprintActionDatabaseImpl::AddSkeletonActions(const USkeleton& Skeleton, FActionList& ActionListOut)
 {
-	for (int32 I = 0; I < Skeleton.AnimationNotifies.Num(); ++I)
+	TArray<FName> NotifyNames;
+	Skeleton.CollectAnimationNotifies(NotifyNames);
+
+	for (const FName& NotifyName : NotifyNames)
 	{
-		FName NotifyName = Skeleton.AnimationNotifies[I];
 		FString Label = NotifyName.ToString();
 
 		FString SignatureName = FString::Printf(TEXT("AnimNotify_%s"), *Label);
diff --git a/Engine/Source/Editor/SkeletonEditor/Private/EditableSkeleton.cpp b/Engine/Source/Editor/SkeletonEditor/Private/EditableSkeleton.cpp
--- a/Engine/Source/Editor/SkeletonEditor/Private/EditableSkeleton.cpp
+++ b/Engine/Source/Editor/SkeletonEditor/Private/EditableSkeleton.cpp
@@ -1121,6 +1121,8 @@ void FEditableSkeleton::AddNotify(FName NewName)
 	const FScopedTransaction Transaction(LOCTEXT("AddNewNotifyToSkeleton", "Add New Anim Notify To Skeleton"));
 	Skeleton->Modify();
 	Skeleton->AddNewAnimationNotify(NewName);
+
+	FBlueprintActionDatabase::Get().RefreshAssetActions(Skeleton);
 	OnNotifiesChanged.Broadcast();
 }
 
@@ -1174,6 +1176,7 @@ int32 FEditableSkeleton::RenameNotify(const FName NewName, const FName OldName)
 		}
 	}
 
+	FBlueprintActionDatabase::Get().RefreshAssetActions(Skeleton);
 	OnNotifiesChanged.Broadcast();
 
 	return NumAnimationsModified;
diff --git a/Engine/Source/Runtime/Engine/Classes/Animation/Skeleton.h b/Engine/Source/Runtime/Engine/Classes/Animation/Skeleton.h
--- a/Engine/Source/Runtime/Engine/Classes/Animation/Skeleton.h
+++ b/Engine/Source/Runtime/Engine/Classes/Animation/Skeleton.h
@@ -559,10 +559,19 @@ public:
 
 #if WITH_EDITORONLY_DATA
 
-	// @todo document
+	/*
+	 * Collect animation notifies that are referenced in all animations that use this skeleton (uses the asset registry).
+	 * Updates the cached AnimationNotifies array.
+	 */
 	ENGINE_API void CollectAnimationNotifies();
 
-	// @todo document
+	/*
+	 * Collect animation notifies that are referenced in all animations that use this skeleton (uses the asset registry).
+	 * @param	OutNotifies		All the notifies that were found
+	 */
+	ENGINE_API void CollectAnimationNotifies(TArray<FName>& OutNotifies) const;
+
+	// Adds a new anim notify to the cached AnimationNotifies array.
 	ENGINE_API void AddNewAnimationNotify(FName NewAnimNotifyName);
 
 	ENGINE_API USkeletalMesh* GetAssetPreviewMesh(UObject* InAsset);
diff --git a/Engine/Source/Runtime/Engine/Private/Animation/Skeleton.cpp b/Engine/Source/Runtime/Engine/Private/Animation/Skeleton.cpp
--- a/Engine/Source/Runtime/Engine/Private/Animation/Skeleton.cpp
+++ b/Engine/Source/Runtime/Engine/Private/Animation/Skeleton.cpp
@@ -878,7 +878,20 @@ int32 USkeleton::GetChildBones(int32 ParentBoneIndex, TArray<int32> & Children)
 
 void USkeleton::CollectAnimationNotifies()
 {
-	// need to verify if these pose is used by anybody else
+	CollectAnimationNotifies(AnimationNotifies);
+}
+
+void USkeleton::CollectAnimationNotifies(TArray<FName>& OutNotifies) const
+{
+	// first merge in AnimationNotifies
+	if(&AnimationNotifies != &OutNotifies)
+	{
+		for(const FName& NotifyName : AnimationNotifies)
+		{
+			OutNotifies.AddUnique(NotifyName);
+		}
+	}
+
 	FAssetRegistryModule& AssetRegistryModule = FModuleManager::LoadModuleChecked<FAssetRegistryModule>(TEXT("AssetRegistry"));
 
 	// @Todo : remove it when we know the asset registry is updated
@@ -886,7 +899,7 @@ void USkeleton::CollectAnimationNotifies()
 	//AnimationNotifies.Empty();
 	TArray<FAssetData> AssetList;
 	AssetRegistryModule.Get().GetAssetsByClass(UAnimSequenceBase::StaticClass()->GetFName(), AssetList, true);
-#if WITH_EDITOR
+
 	// do not clear AnimationNotifies. We can't remove old ones yet. 
 	FString CurrentSkeletonName = FAssetData(this).GetExportTextName();
 	for (auto Iter = AssetList.CreateConstIterator(); Iter; ++Iter)
@@ -903,12 +916,11 @@ void USkeleton::CollectAnimationNotifies()
 				for (auto NotifyIter = NotifyList.CreateConstIterator(); NotifyIter; ++NotifyIter)
 				{
 					FString NotifyName = *NotifyIter;
-					AddNewAnimationNotify(FName(*NotifyName));
+					OutNotifies.AddUnique(FName(*NotifyName));
 				}
 			}
 		}
 	}
-#endif
 }
 
 void USkeleton::AddNewAnimationNotify(FName NewAnimNotifyName)
