Fix unintentional double negation in blueprint context menu visibility controls

diff --git a/Engine/Source/Editor/BlueprintGraph/Private/EdGraphSchema_K2.cpp b/Engine/Source/Editor/BlueprintGraph/Private/EdGraphSchema_K2.cpp
--- a/Engine/Source/Editor/BlueprintGraph/Private/EdGraphSchema_K2.cpp
+++ b/Engine/Source/Editor/BlueprintGraph/Private/EdGraphSchema_K2.cpp
@@ -763,7 +763,7 @@ bool UEdGraphSchema_K2::CanFunctionBeUsedInGraph(const UClass* InClass, const UF
 		}
 
 		const bool bIsNotNative = !FBlueprintEditorUtils::IsNativeSignature(InFunction);
-		if(!bIsNotNative)
+		if(bIsNotNative)
 		{
 			// Blueprint functions visibility flags can be enforced in blueprints - native functions
 			// are often using these flags to only hide functionality from other native functions:
