Fix regression caused by an oversight in extending function call nodes from single to multiple dynamic output pin type support.

#rnx

diff --git a/Engine/Source/Editor/BlueprintGraph/Private/K2Node_CallFunction.cpp b/Engine/Source/Editor/BlueprintGraph/Private/K2Node_CallFunction.cpp
--- a/Engine/Source/Editor/BlueprintGraph/Private/K2Node_CallFunction.cpp
+++ b/Engine/Source/Editor/BlueprintGraph/Private/K2Node_CallFunction.cpp
@@ -379,7 +379,7 @@ void FDynamicOutputHelper::GetDynamicOutPins(const UK2Node_CallFunction* FuncNod
 			// Ensure that this is a valid pin to make dynamic
 			FEdGraphPinType PropertyPinType;
 
-			if (TaggedOutputParam && (K2Schema->ConvertPropertyToPinType(TaggedOutputParam, /*out*/PropertyPinType) || !CanConformPinType(FuncNode, PropertyPinType)))
+			if (TaggedOutputParam && K2Schema->ConvertPropertyToPinType(TaggedOutputParam, /*out*/PropertyPinType) && CanConformPinType(FuncNode, PropertyPinType))
 			{
 				UEdGraphPin* DynamicOutPin = FuncNode->FindPin(TaggedOutputParam->GetFName());
 				if (DynamicOutPin && (DynamicOutPin->Direction == EGPD_Output))
