Merging using //UE4/Dev-Framework_to_//UE4/Release-4.24

Member references now store information about the class that owns the reference.
Field and variable node spawners now optionally take the class that will own the nodes that they spawn.

diff --git a/Engine/Source/Editor/BlueprintGraph/Classes/BlueprintFieldNodeSpawner.h b/Engine/Source/Editor/BlueprintGraph/Classes/BlueprintFieldNodeSpawner.h
--- a/Engine/Source/Editor/BlueprintGraph/Classes/BlueprintFieldNodeSpawner.h
+++ b/Engine/Source/Editor/BlueprintGraph/Classes/BlueprintFieldNodeSpawner.h
@@ -33,9 +33,10 @@ public:
 	 * @param  NodeClass	The type of node you want the spawner to create.
 	 * @param  Field		The field you want assigned to new nodes.
 	 * @param  Outer		Optional outer for the new spawner (if left null, the transient package will be used).
+	 * @param  OwnerClass	The class that the variable is a member of or the class it is associated with if it is in a sidecar data structure.
 	 * @return A newly allocated instance of this class.
 	 */
-	static UBlueprintFieldNodeSpawner* Create(TSubclassOf<UK2Node> NodeClass, UField const* const Field, UObject* Outer = nullptr);
+	static UBlueprintFieldNodeSpawner* Create(TSubclassOf<UK2Node> NodeClass, UField const* const Field, UObject* Outer = nullptr, UClass const* OwnerClass = nullptr);
 
 	// UBlueprintNodeSpawner interface
 	virtual FBlueprintNodeSignature GetSpawnerSignature() const override;
@@ -57,4 +58,8 @@ protected:
 	/** The field to configure new nodes with. */
 	UPROPERTY()
 	UField const* Field;
+
+	/** The owning class to configure new nodes with. */
+	UPROPERTY()
+	UClass const* OwnerClass;
 };
diff --git a/Engine/Source/Editor/BlueprintGraph/Classes/BlueprintVariableNodeSpawner.h b/Engine/Source/Editor/BlueprintGraph/Classes/BlueprintVariableNodeSpawner.h
--- a/Engine/Source/Editor/BlueprintGraph/Classes/BlueprintVariableNodeSpawner.h
+++ b/Engine/Source/Editor/BlueprintGraph/Classes/BlueprintVariableNodeSpawner.h
@@ -37,10 +37,10 @@ public:
 	 * @param  NodeClass	The node type that you want the spawner to spawn.
 	 * @param  VarProperty	The property that represents the member-variable you want nodes spawned for.
 	 * @param  VarContext	The graph that the local variable belongs to.
-	 * @param  Outer		Optional outer for the new spawner (if left null, the transient package will be used).
+	 * @param  OwnerClass	The class that the variable is a member of or the class it is associated with if it is in a sidecar data structure.
 	 * @return A newly allocated instance of this class.
 	 */
-	static UBlueprintVariableNodeSpawner* CreateFromMemberOrParam(TSubclassOf<UK2Node_Variable> NodeClass, UProperty const* VarProperty, UEdGraph* VarContext = nullptr, UObject* Outer = nullptr);
+	static UBlueprintVariableNodeSpawner* CreateFromMemberOrParam(TSubclassOf<UK2Node_Variable> NodeClass, UProperty const* VarProperty, UEdGraph* VarContext = nullptr, UClass* OwnerClass = nullptr);
 
 	/**
 	 * Creates a new UBlueprintVariableNodeSpawner, charged with spawning
diff --git a/Engine/Source/Editor/BlueprintGraph/Classes/K2Node_BaseMCDelegate.h b/Engine/Source/Editor/BlueprintGraph/Classes/K2Node_BaseMCDelegate.h
--- a/Engine/Source/Editor/BlueprintGraph/Classes/K2Node_BaseMCDelegate.h
+++ b/Engine/Source/Editor/BlueprintGraph/Classes/K2Node_BaseMCDelegate.h
@@ -43,9 +43,9 @@ public:
 	virtual bool IsCompatibleWithGraph(const UEdGraph* TargetGraph) const override;
 	// End of UEdGraphNode interface
 
-	BLUEPRINTGRAPH_API void SetFromProperty(const UProperty* Property, bool bSelfContext)
+	BLUEPRINTGRAPH_API void SetFromProperty(const UProperty* Property, bool bSelfContext, UClass* OwnerClass)
 	{
-		DelegateReference.SetFromField<UProperty>(Property, bSelfContext);
+		DelegateReference.SetFromField<UProperty>(Property, bSelfContext, OwnerClass);
 	}
 
 	BLUEPRINTGRAPH_API UProperty* GetProperty() const
diff --git a/Engine/Source/Editor/BlueprintGraph/Classes/K2Node_Variable.h b/Engine/Source/Editor/BlueprintGraph/Classes/K2Node_Variable.h
--- a/Engine/Source/Editor/BlueprintGraph/Classes/K2Node_Variable.h
+++ b/Engine/Source/Editor/BlueprintGraph/Classes/K2Node_Variable.h
@@ -90,7 +90,7 @@ public:
 	//~ End K2Node Interface
 
 	/** Set up this variable node from the supplied UProperty */
-	void SetFromProperty(const UProperty* Property, bool bSelfContext);
+	void SetFromProperty(const UProperty* Property, bool bSelfContext, UClass* OwnerClass);
 
 	/** Util to get variable name as a string */
 	FString GetVarNameString() const
diff --git a/Engine/Source/Editor/BlueprintGraph/Private/BlueprintActionDatabase.cpp b/Engine/Source/Editor/BlueprintGraph/Private/BlueprintActionDatabase.cpp
--- a/Engine/Source/Editor/BlueprintGraph/Private/BlueprintActionDatabase.cpp
+++ b/Engine/Source/Editor/BlueprintGraph/Private/BlueprintActionDatabase.cpp
@@ -746,7 +746,8 @@ static void BlueprintActionDatabaseImpl::AddClassDataObjectActions(UClass const*
 				continue;
 			}
 
-			UBlueprintVariableNodeSpawner* GetterSpawner = UBlueprintVariableNodeSpawner::CreateFromMemberOrParam(UK2Node_VariableGet::StaticClass(), Property);
+			UClass* NonConstClass = const_cast<UClass*>(Class);
+			UBlueprintVariableNodeSpawner* GetterSpawner = UBlueprintVariableNodeSpawner::CreateFromMemberOrParam(UK2Node_VariableGet::StaticClass(), Property, nullptr, NonConstClass);
 			ActionListOut.Add(GetterSpawner);
 //			UBlueprintVariableNodeSpawner* SetterSpawner = UBlueprintVariableNodeSpawner::CreateFromMemberOrParam(UK2Node_VariableSet::StaticClass(), Property);
 //			ActionListOut.Add(SetterSpawner);
diff --git a/Engine/Source/Editor/BlueprintGraph/Private/BlueprintDelegateNodeSpawner.cpp b/Engine/Source/Editor/BlueprintGraph/Private/BlueprintDelegateNodeSpawner.cpp
--- a/Engine/Source/Editor/BlueprintGraph/Private/BlueprintDelegateNodeSpawner.cpp
+++ b/Engine/Source/Editor/BlueprintGraph/Private/BlueprintDelegateNodeSpawner.cpp
@@ -94,7 +94,7 @@ UBlueprintDelegateNodeSpawner* UBlueprintDelegateNodeSpawner::Create(TSubclassOf
 			UBlueprint* Blueprint = FBlueprintEditorUtils::FindBlueprintForNodeChecked(NewNode);
 			UClass* OwnerClass = MCDProperty->GetOwnerClass();
 
-			DelegateNode->SetFromProperty(MCDProperty, false);
+			DelegateNode->SetFromProperty(MCDProperty, false, OwnerClass);
 		}
 	};
 	NodeSpawner->SetNodeFieldDelegate = FSetNodeFieldDelegate::CreateStatic(SetDelegateLambda);
diff --git a/Engine/Source/Editor/BlueprintGraph/Private/BlueprintFieldNodeSpawner.cpp b/Engine/Source/Editor/BlueprintGraph/Private/BlueprintFieldNodeSpawner.cpp
--- a/Engine/Source/Editor/BlueprintGraph/Private/BlueprintFieldNodeSpawner.cpp
+++ b/Engine/Source/Editor/BlueprintGraph/Private/BlueprintFieldNodeSpawner.cpp
@@ -6,15 +6,16 @@
 #define LOCTEXT_NAMESPACE "BlueprintFieldNodeSpawner"
 
 //------------------------------------------------------------------------------
-UBlueprintFieldNodeSpawner* UBlueprintFieldNodeSpawner::Create(TSubclassOf<UK2Node> NodeClass, UField const* const Field, UObject* Outer/* = nullptr*/)
+UBlueprintFieldNodeSpawner* UBlueprintFieldNodeSpawner::Create(TSubclassOf<UK2Node> NodeClass, UField const* const Field, UObject* Outer/* = nullptr*/, UClass const* OwnerClass)
 {
 	if (Outer == nullptr)
 	{
 		Outer = GetTransientPackage();
 	}
 	UBlueprintFieldNodeSpawner* NodeSpawner = NewObject<UBlueprintFieldNodeSpawner>(Outer);
-	NodeSpawner->Field     = Field;
-	NodeSpawner->NodeClass = NodeClass;
+	NodeSpawner->Field      = Field;
+	NodeSpawner->NodeClass  = NodeClass;
+	NodeSpawner->OwnerClass = OwnerClass;
 	
 	return NodeSpawner;
 }
@@ -23,6 +24,7 @@ UBlueprintFieldNodeSpawner* UBlueprintFieldNodeSpawner::Create(TSubclassOf<UK2No
 UBlueprintFieldNodeSpawner::UBlueprintFieldNodeSpawner(FObjectInitializer const& ObjectInitializer)
 	: Super(ObjectInitializer)
 	, Field(nullptr)
+	, OwnerClass(nullptr)
 {
 }
 
diff --git a/Engine/Source/Editor/BlueprintGraph/Private/BlueprintVariableNodeSpawner.cpp b/Engine/Source/Editor/BlueprintGraph/Private/BlueprintVariableNodeSpawner.cpp
--- a/Engine/Source/Editor/BlueprintGraph/Private/BlueprintVariableNodeSpawner.cpp
+++ b/Engine/Source/Editor/BlueprintGraph/Private/BlueprintVariableNodeSpawner.cpp
@@ -19,19 +19,15 @@
  ******************************************************************************/
 
 //------------------------------------------------------------------------------
-UBlueprintVariableNodeSpawner* UBlueprintVariableNodeSpawner::CreateFromMemberOrParam(TSubclassOf<UK2Node_Variable> NodeClass, UProperty const* VarProperty, UEdGraph* VarContext, UObject* Outer)
+UBlueprintVariableNodeSpawner* UBlueprintVariableNodeSpawner::CreateFromMemberOrParam(TSubclassOf<UK2Node_Variable> NodeClass, UProperty const* VarProperty, UEdGraph* VarContext, UClass* OwnerClass)
 {
 	check(VarProperty != nullptr);
-	if (Outer == nullptr)
-	{
-		Outer = GetTransientPackage();
-	}
 
 	//--------------------------------------
 	// Constructing the Spawner
 	//--------------------------------------
 
-	UBlueprintVariableNodeSpawner* NodeSpawner = NewObject<UBlueprintVariableNodeSpawner>(Outer);
+	UBlueprintVariableNodeSpawner* NodeSpawner = NewObject<UBlueprintVariableNodeSpawner>(GetTransientPackage());
 	NodeSpawner->NodeClass = NodeClass;
 	NodeSpawner->Field = VarProperty;
 	NodeSpawner->LocalVarOuter = VarContext;
@@ -71,22 +67,22 @@ UBlueprintVariableNodeSpawner* UBlueprintVariableNodeSpawner::CreateFromMemberOr
 	// Post-Spawn Setup
 	//--------------------------------------
 
-	auto MemberVarSetupLambda = [](UEdGraphNode* NewNode, UField const* InField)
+	auto MemberVarSetupLambda = [](UEdGraphNode* NewNode, UField const* InField, UClass* OwnerClass)
 	{
 		if (UProperty const* Property = Cast<UProperty>(InField))
 		{
 			UBlueprint* Blueprint = FBlueprintEditorUtils::FindBlueprintForNodeChecked(NewNode);
-			UClass* OwnerClass = Property->GetOwnerClass();
+			OwnerClass = OwnerClass ? OwnerClass : Property->GetOwnerClass();
 
 			// We need to use a generated class instead of a skeleton class for IsChildOf, so if the OwnerClass has a Blueprint, grab the GeneratedClass
 			const bool bOwnerClassIsSelfContext = (Blueprint->SkeletonGeneratedClass->GetAuthoritativeClass() == OwnerClass) || Blueprint->SkeletonGeneratedClass->IsChildOf(OwnerClass);
 			const bool bIsFunctionVariable = Property->GetOuter()->IsA(UFunction::StaticClass());
 
 			UK2Node_Variable* VarNode = CastChecked<UK2Node_Variable>(NewNode);
-			VarNode->SetFromProperty(Property, bOwnerClassIsSelfContext && !bIsFunctionVariable);
+			VarNode->SetFromProperty(Property, bOwnerClassIsSelfContext && !bIsFunctionVariable, OwnerClass);
 		}
 	};
-	NodeSpawner->SetNodeFieldDelegate = FSetNodeFieldDelegate::CreateStatic(MemberVarSetupLambda);
+	NodeSpawner->SetNodeFieldDelegate = FSetNodeFieldDelegate::CreateStatic(MemberVarSetupLambda, OwnerClass);
 
 	return NodeSpawner;
 }
@@ -192,9 +188,10 @@ FBlueprintActionUiSpec UBlueprintVariableNodeSpawner::GetUiSpec(FBlueprintAction
 			}
 		}
 
-		auto OwnerClass = WrappedVariable->GetOwnerClass();
-		const bool bIsOwneClassValid = OwnerClass && (!Cast<const UBlueprintGeneratedClass>(OwnerClass) || OwnerClass->ClassGeneratedBy); //todo: more general validation
-		UClass const* VariableClass = bIsOwneClassValid ? OwnerClass->GetAuthoritativeClass() : NULL;
+		UClass const* EffectiveOwnerClass = WrappedVariable->GetOwnerClass();
+		EffectiveOwnerClass = EffectiveOwnerClass ? EffectiveOwnerClass : OwnerClass;
+		const bool bIsOwningClassValid = EffectiveOwnerClass && (!Cast<const UBlueprintGeneratedClass>(EffectiveOwnerClass) || EffectiveOwnerClass->ClassGeneratedBy); //todo: more general validation
+		UClass const* VariableClass = bIsOwningClassValid ? EffectiveOwnerClass->GetAuthoritativeClass() : nullptr;
 		if (VariableClass && !TargetClass->IsChildOf(VariableClass))
 		{
 			MenuSignature.Category = FEditorCategoryUtils::BuildCategoryString(FCommonEditorCategory::Class,
diff --git a/Engine/Source/Editor/BlueprintGraph/Private/K2Node_BaseAsyncTask.cpp b/Engine/Source/Editor/BlueprintGraph/Private/K2Node_BaseAsyncTask.cpp
--- a/Engine/Source/Editor/BlueprintGraph/Private/K2Node_BaseAsyncTask.cpp
+++ b/Engine/Source/Editor/BlueprintGraph/Private/K2Node_BaseAsyncTask.cpp
@@ -279,7 +279,7 @@ bool UK2Node_BaseAsyncTask::FBaseAsyncTaskHelper::HandleDelegateImplementation(
 	UK2Node_CustomEvent* CurrentCENode = CompilerContext.SpawnIntermediateEventNode<UK2Node_CustomEvent>(CurrentNode, PinForCurrentDelegateProperty, SourceGraph);
 	{
 	UK2Node_AddDelegate* AddDelegateNode = CompilerContext.SpawnIntermediateNode<UK2Node_AddDelegate>(CurrentNode, SourceGraph);
-	AddDelegateNode->SetFromProperty(CurrentProperty, false);
+	AddDelegateNode->SetFromProperty(CurrentProperty, false, CurrentProperty->GetOwnerClass());
 	AddDelegateNode->AllocateDefaultPins();
 		bIsErrorFree &= Schema->TryCreateConnection(AddDelegateNode->FindPinChecked(UEdGraphSchema_K2::PN_Self), ProxyObjectPin);
 		bIsErrorFree &= Schema->TryCreateConnection(InOutLastThenPin, AddDelegateNode->FindPinChecked(UEdGraphSchema_K2::PN_Execute));
diff --git a/Engine/Source/Editor/BlueprintGraph/Private/K2Node_FunctionEntry.cpp b/Engine/Source/Editor/BlueprintGraph/Private/K2Node_FunctionEntry.cpp
--- a/Engine/Source/Editor/BlueprintGraph/Private/K2Node_FunctionEntry.cpp
+++ b/Engine/Source/Editor/BlueprintGraph/Private/K2Node_FunctionEntry.cpp
@@ -779,7 +779,7 @@ void UK2Node_FunctionEntry::ExpandNode(class FKismetCompilerContext& CompilerCon
 					{
 						// Add a variable set node for the local variable and hook it up immediately following the entry node or the last added local variable
 						UK2Node_VariableSet* VariableSetNode = CompilerContext.SpawnIntermediateNode<UK2Node_VariableSet>(this, SourceGraph);
-						VariableSetNode->SetFromProperty(Property, false);
+						VariableSetNode->SetFromProperty(Property, false, Property->GetOwnerClass());
 						Schema->ConfigureVarNode(VariableSetNode, LocalVar.VarName, Function, CompilerContext.Blueprint);
 						VariableSetNode->AllocateDefaultPins();
 
diff --git a/Engine/Source/Editor/BlueprintGraph/Private/K2Node_Variable.cpp b/Engine/Source/Editor/BlueprintGraph/Private/K2Node_Variable.cpp
--- a/Engine/Source/Editor/BlueprintGraph/Private/K2Node_Variable.cpp
+++ b/Engine/Source/Editor/BlueprintGraph/Private/K2Node_Variable.cpp
@@ -54,10 +54,10 @@ void UK2Node_Variable::Serialize(FArchive& Ar)
 	}
 }
 
-void UK2Node_Variable::SetFromProperty(const UProperty* Property, bool bSelfContext)
+void UK2Node_Variable::SetFromProperty(const UProperty* Property, bool bSelfContext, UClass* OwnerClass)
 {
 	SelfContextInfo = bSelfContext ? ESelfContextInfo::Unspecified : ESelfContextInfo::NotSelfContext;
-	VariableReference.SetFromField<UProperty>(Property, bSelfContext);
+	VariableReference.SetFromField<UProperty>(Property, bSelfContext, OwnerClass);
 }
 
 bool UK2Node_Variable::CreatePinForVariable(EEdGraphPinDirection Direction, FName InPinName/* = NAME_None */)
@@ -111,7 +111,7 @@ void UK2Node_Variable::CreatePinForSelf()
 	if (!K2Schema->FindSelfPin(*this, EGPD_Input))
 	{
 		// Do not create a self pin for locally scoped variables or sparse class data
-		if (!VariableReference.IsLocalScope() && !VariableReference.IsSparseClassData(GetBlueprintClassFromNode()))
+		if (!VariableReference.IsLocalScope())
 		{
 			bool bSelfTarget = VariableReference.IsSelfContext() && (ESelfContextInfo::NotSelfContext != SelfContextInfo);
 			UClass* MemberParentClass = VariableReference.GetMemberParentClass(GetBlueprintClassFromNode());
@@ -121,10 +121,10 @@ void UK2Node_Variable::CreatePinForSelf()
 
 			if (VariableProperty)
 			{
-				UClass* PropertyClass = CastChecked<UClass>(VariableProperty->GetOuter());
+				UClass* PropertyClass = Cast<UClass>(VariableProperty->GetOuter());
 
 				// Fix up target class if it's not correct, this fixes cases where variables have moved within the hierarchy
-				if (PropertyClass != TargetClass)
+				if (PropertyClass && PropertyClass != TargetClass)
 				{
 					TargetClass = PropertyClass;
 				}
@@ -134,9 +134,13 @@ void UK2Node_Variable::CreatePinForSelf()
 			// so if the node is from a Macro Blueprint, it will hook up as self in any placed Blueprint
 			if (bSelfTarget)
 			{
-				if(UProperty* Property = VariableReference.ResolveMember<UProperty>(GetBlueprintClassFromNode()))
+				if (UProperty* Property = VariableReference.ResolveMember<UProperty>(GetBlueprintClassFromNode()))
 				{
-					TargetClass = Property->GetOwnerClass()->GetAuthoritativeClass();
+					UClass* OwnerClass = Property->GetOwnerClass();
+					if (OwnerClass)
+					{
+						TargetClass = OwnerClass->GetAuthoritativeClass();
+					}
 				}
 				else if (GetBlueprint()->SkeletonGeneratedClass)
 				{
@@ -379,7 +383,9 @@ UProperty* UK2Node_Variable::GetPropertyForVariable_Internal(UClass* OwningClass
 
 	// Look in the sparse class data first
 	UProperty* VariableProperty = nullptr;
-	UScriptStruct* SparseClassDataStruct = OwningClass ? OwningClass->GetSparseClassDataStruct() : nullptr;
+	// TODO: move some of this into MemberReference::ResolveMember if possible
+	UClass* Scope = VariableReference.GetScope(OwningClass);
+	UScriptStruct* SparseClassDataStruct = Scope ? Scope->GetSparseClassDataStruct() : nullptr;
 	if (SparseClassDataStruct)
 	{
 		VariableProperty = FindField<UProperty>(SparseClassDataStruct, VarName);
diff --git a/Engine/Source/Editor/Kismet/Private/BPDelegateDragDropAction.cpp b/Engine/Source/Editor/Kismet/Private/BPDelegateDragDropAction.cpp
--- a/Engine/Source/Editor/Kismet/Private/BPDelegateDragDropAction.cpp
+++ b/Engine/Source/Editor/Kismet/Private/BPDelegateDragDropAction.cpp
@@ -46,10 +46,10 @@ void FKismetDelegateDragDropAction::AssignEvent(FNodeConstructionParams Params)
 	check(Params.Graph && Params.Property);
 	const UMulticastDelegateProperty* MulticastDelegateProperty = Cast<const UMulticastDelegateProperty>(Params.Property);
 	const UFunction* SignatureFunction = MulticastDelegateProperty ? MulticastDelegateProperty->SignatureFunction : NULL;
-	if(SignatureFunction)
+	if (SignatureFunction)
 	{
 		UK2Node_AddDelegate* TemplateNode = NewObject<UK2Node_AddDelegate>();
-		TemplateNode->SetFromProperty(Params.Property, Params.bSelfContext);
+		TemplateNode->SetFromProperty(Params.Property, Params.bSelfContext, Params.Property->GetOwnerClass());
 		FEdGraphSchemaAction_K2AssignDelegate::AssignDelegate(TemplateNode, Params.Graph, NULL, Params.GraphPosition, true);
 		Params.AnalyticCallback.ExecuteIfBound();
 	}
diff --git a/Engine/Source/Editor/Kismet/Private/BPDelegateDragDropAction.h b/Engine/Source/Editor/Kismet/Private/BPDelegateDragDropAction.h
--- a/Engine/Source/Editor/Kismet/Private/BPDelegateDragDropAction.h
+++ b/Engine/Source/Editor/Kismet/Private/BPDelegateDragDropAction.h
@@ -61,7 +61,7 @@ public:
 			EK2NewNodeFlags::SelectNewNode,
 			[&Params](TNode* NewInstance)
 			{
-				NewInstance->SetFromProperty(Params.Property, Params.bSelfContext);
+				NewInstance->SetFromProperty(Params.Property, Params.bSelfContext, Params.Property->GetOwnerClass());
 			}
 		);
 		Params.AnalyticCallback.ExecuteIfBound();
diff --git a/Engine/Source/Editor/Kismet/Private/BlueprintEditor.cpp b/Engine/Source/Editor/Kismet/Private/BlueprintEditor.cpp
--- a/Engine/Source/Editor/Kismet/Private/BlueprintEditor.cpp
+++ b/Engine/Source/Editor/Kismet/Private/BlueprintEditor.cpp
@@ -6354,7 +6354,7 @@ private:
 		{
 			NewTarget = NewObject<UK2Node_VariableGet>(Graph);
 			check(NewTarget);
-			NewTarget->SetFromProperty(Property, true);
+			NewTarget->SetFromProperty(Property, true, Property->GetOwnerClass());
 			AddedTargets.Add(NewTarget);
 			const float AutoNodeOffsetX = 160.0f;
 			InitializeNewNode(NewTarget, OldTarget, OldCall->NodePosX - AutoNodeOffsetX, OldCall->NodePosY);
diff --git a/Engine/Source/Editor/Kismet/Private/SReplaceNodeReferences.cpp b/Engine/Source/Editor/Kismet/Private/SReplaceNodeReferences.cpp
--- a/Engine/Source/Editor/Kismet/Private/SReplaceNodeReferences.cpp
+++ b/Engine/Source/Editor/Kismet/Private/SReplaceNodeReferences.cpp
@@ -359,7 +359,7 @@ void SReplaceNodeReferences::GatherAllAvailableBlueprintVariables(UClass* InTarg
 			}
 
 			TSharedPtr< FTargetVariableReplaceReferences > VariableItem = MakeShareable(new FTargetVariableReplaceReferences);
-			VariableItem->VariableReference.SetFromField<UProperty>(Property, true);
+			VariableItem->VariableReference.SetFromField<UProperty>(Property, true, InTargetClass);
 
 			FEdGraphPinType Type;
 			K2Schema->ConvertPropertyToPinType(Property, VariableItem->PinType);
@@ -436,7 +436,7 @@ void SReplaceNodeReferences::OnSubmitSearchQuery(bool bFindAndReplace)
 	FString SearchTerm;
 
 	FMemberReference SourceVariableReference;
-	SourceVariableReference.SetFromField<UProperty>(SourceProperty, true);
+	SourceVariableReference.SetFromField<UProperty>(SourceProperty, true, SourceProperty->GetOwnerClass());
 	SearchTerm = SourceVariableReference.GetReferenceSearchString(SourceProperty->GetOwnerClass());
 
 	FOnSearchComplete OnSearchComplete;
diff --git a/Engine/Source/Editor/KismetCompiler/Private/AnimBlueprintCompiler.cpp b/Engine/Source/Editor/KismetCompiler/Private/AnimBlueprintCompiler.cpp
--- a/Engine/Source/Editor/KismetCompiler/Private/AnimBlueprintCompiler.cpp
+++ b/Engine/Source/Editor/KismetCompiler/Private/AnimBlueprintCompiler.cpp
@@ -2151,7 +2151,7 @@ void FAnimBlueprintCompilerContext::ProcessLinkedInputPose(UAnimGraphNode_Linked
 					{
 						// Create new node for property access
 						UK2Node_VariableGet* VariableGetNode = SpawnIntermediateNode<UK2Node_VariableGet>(InLinkedInputPose, InLinkedInputPose->GetGraph());
-						VariableGetNode->SetFromProperty(NewLinkedInputPoseProperty, true);
+						VariableGetNode->SetFromProperty(NewLinkedInputPoseProperty, true, NewLinkedInputPoseProperty->GetOwnerClass());
 						VariableGetNode->AllocateDefaultPins();
 
 						// Add pin to generated variable association, used for pin watching
diff --git a/Engine/Source/Runtime/Engine/Classes/Engine/MemberReference.h b/Engine/Source/Runtime/Engine/Classes/Engine/MemberReference.h
--- a/Engine/Source/Runtime/Engine/Classes/Engine/MemberReference.h
+++ b/Engine/Source/Runtime/Engine/Classes/Engine/MemberReference.h
@@ -104,9 +104,13 @@ public:
 
 	/** Set up this reference from a supplied field */
 	template<class TFieldType>
-	void SetFromField(const UField* InField, const bool bIsConsideredSelfContext)
+	void SetFromField(const UField* InField, const bool bIsConsideredSelfContext, UClass* OwnerClass=nullptr)
 	{
-		UClass* OwnerClass = InField->GetOwnerClass();
+		// if we didn't get an owner passed in try to figure out what it should be based on the field
+		if (!OwnerClass)
+		{
+			OwnerClass = InField->GetOwnerClass();
+		}
 		MemberParent = OwnerClass;
 
 		if (bIsConsideredSelfContext)
@@ -310,6 +314,13 @@ public:
 		return bWasDeprecated;
 	}
 
+	/** Returns the scope for the current member. This will vary based on whether or not this member uses the self context or not. */
+	UClass* GetScope(UClass* SelfScope = nullptr) const
+	{
+		return bSelfContext ? SelfScope : GetMemberParentClass();
+	}
+	
+
 	/** 
 	 *	Returns the member UProperty/UFunction this reference is pointing to, or NULL if it no longer exists 
 	 *	Derives 'self' scope from supplied Blueprint node if required
@@ -350,7 +361,7 @@ public:
 		else
 		{
 			// Look for remapped member
-			UClass* TargetScope = bSelfContext ? SelfScope : GetMemberParentClass();
+			UClass* TargetScope = GetScope(SelfScope);
 #if WITH_EDITOR
 			if( TargetScope != nullptr &&  !GIsSavingPackage )
 			{
