Fixed crash when compiling a blueprint after changing a Map variable's contained types when using split pins

diff --git a/Engine/Source/Editor/BlueprintGraph/Private/EdGraphSchema_K2.cpp b/Engine/Source/Editor/BlueprintGraph/Private/EdGraphSchema_K2.cpp
--- a/Engine/Source/Editor/BlueprintGraph/Private/EdGraphSchema_K2.cpp
+++ b/Engine/Source/Editor/BlueprintGraph/Private/EdGraphSchema_K2.cpp
@@ -6859,7 +6859,7 @@ UK2Node* UEdGraphSchema_K2::CreateSplitPinNode(UEdGraphPin* Pin, const FCreateSp
 
 	UEdGraphNode* GraphNode = Pin->GetOwningNode();
 	UEdGraph* Graph = GraphNode->GetGraph();
-	UScriptStruct* StructType = CastChecked<UScriptStruct>(Pin->PinType.PinSubCategoryObject.Get(), ECastCheckedType::NullAllowed);
+	UScriptStruct* StructType = Cast<UScriptStruct>(Pin->PinType.PinSubCategoryObject.Get());
 	if (!StructType)
 	{
 		if (Params.CompilerContext)
diff --git a/Engine/Source/Editor/BlueprintGraph/Private/K2Node_CallFunction.cpp b/Engine/Source/Editor/BlueprintGraph/Private/K2Node_CallFunction.cpp
--- a/Engine/Source/Editor/BlueprintGraph/Private/K2Node_CallFunction.cpp
+++ b/Engine/Source/Editor/BlueprintGraph/Private/K2Node_CallFunction.cpp
@@ -1197,7 +1197,10 @@ void UK2Node_CallFunction::NotifyPinConnectionListChanged(UEdGraphPin* Pin)
 	}
 
 	InvalidatePinTooltips();
-	FDynamicOutputHelper(Pin).ConformOutputType();
+	if(!Pin->IsPendingKill())
+	{
+		FDynamicOutputHelper(Pin).ConformOutputType();
+	}
 }
 
 void UK2Node_CallFunction::PinDefaultValueChanged(UEdGraphPin* Pin)
@@ -2656,7 +2659,7 @@ void UK2Node_CallFunction::ConformContainerPins()
 	{
 		if (Pin && !bOutPropagated)
 		{
-			if (Pin->LinkedTo.Num() != 0 || !Pin->DoesDefaultValueMatchAutogenerated())
+			if (Pin->HasAnyConnections() || !Pin->DoesDefaultValueMatchAutogenerated() )
 			{
 				bOutPropagated = true;
 				if (Pin->LinkedTo.Num() != 0)
@@ -2689,24 +2692,46 @@ void UK2Node_CallFunction::ConformContainerPins()
 			}
 		}
 	};
+	
+	const UEdGraphSchema_K2* Schema = CastChecked<UEdGraphSchema_K2>(GetSchema());
 
-	const auto TryPropagateType = [](UEdGraphPin* Pin, const FEdGraphTerminalType& TerminalType, bool bTypeIsAvailable)
+	const auto TryPropagateType = [Schema](UEdGraphPin* Pin, const FEdGraphTerminalType& TerminalType, bool bTypeIsAvailable)
 	{
 		if(Pin)
 		{
 			if(bTypeIsAvailable)
 			{
-				Pin->PinType.PinCategory = TerminalType.TerminalCategory;
-				Pin->PinType.PinSubCategory = TerminalType.TerminalSubCategory;
-				Pin->PinType.PinSubCategoryObject = TerminalType.TerminalSubCategoryObject;
+				if(Pin->GetPrimaryTerminalType() != TerminalType)
+				{
+					// terminal type changed:
+					if (Pin->SubPins.Num() > 0 && Pin->PinType.PinCategory != UEdGraphSchema_K2::PC_Wildcard)
+					{
+						Schema->RecombinePin(Pin->SubPins[0]);
+					}
+
+					Pin->PinType.PinCategory = TerminalType.TerminalCategory;
+					Pin->PinType.PinSubCategory = TerminalType.TerminalSubCategory;
+					Pin->PinType.PinSubCategoryObject = TerminalType.TerminalSubCategoryObject;
+				
+					// Reset default values
+					if (!Schema->IsPinDefaultValid(Pin, Pin->DefaultValue, Pin->DefaultObject, Pin->DefaultTextValue).IsEmpty())
+					{
+						Schema->ResetPinToAutogeneratedDefaultValue(Pin, false);
+					}
+				}
 			}
 			else
 			{
 				// reset to wildcard:
+				if (Pin->SubPins.Num() > 0)
+				{
+					Schema->RecombinePin(Pin->SubPins[0]);
+				}
+
 				Pin->PinType.PinCategory = UEdGraphSchema_K2::PC_Wildcard;
 				Pin->PinType.PinSubCategory = NAME_None;
 				Pin->PinType.PinSubCategoryObject = nullptr;
-				Pin->DefaultValue = FString();
+				Schema->ResetPinToAutogeneratedDefaultValue(Pin, false);
 			}
 		}
 	};
diff --git a/Engine/Source/Runtime/Engine/Classes/EdGraph/EdGraphPin.h b/Engine/Source/Runtime/Engine/Classes/EdGraph/EdGraphPin.h
--- a/Engine/Source/Runtime/Engine/Classes/EdGraph/EdGraphPin.h
+++ b/Engine/Source/Runtime/Engine/Classes/EdGraph/EdGraphPin.h
@@ -477,6 +477,9 @@ public:
 	/** Shorthand way to access the schema of the graph that owns the node that owns this pin */
 	ENGINE_API const class UEdGraphSchema* GetSchema() const;
 
+	/** Returns true if there are entries in the pin's LinkedTo array, or if it has SubPins, one of their SubPins arrays: */
+	ENGINE_API bool HasAnyConnections() const;
+
 	/** Direction flipping utility; returns the complementary direction */
 	static EEdGraphPinDirection GetComplementaryDirection(EEdGraphPinDirection InDirection)
 	{
diff --git a/Engine/Source/Runtime/Engine/Private/EdGraph/EdGraphPin.cpp b/Engine/Source/Runtime/Engine/Private/EdGraph/EdGraphPin.cpp
--- a/Engine/Source/Runtime/Engine/Private/EdGraph/EdGraphPin.cpp
+++ b/Engine/Source/Runtime/Engine/Private/EdGraph/EdGraphPin.cpp
@@ -648,6 +648,28 @@ const class UEdGraphSchema* UEdGraphPin::GetSchema() const
 #endif	//WITH_EDITOR
 }
 
+bool UEdGraphPin::HasAnyConnections() const
+{
+	bool bHasAnyConnections = false;
+	if(SubPins.Num() > 0)
+	{
+		// we have subpins, defer authority to them:
+		for(UEdGraphPin* Pin : SubPins)
+		{
+			bHasAnyConnections = Pin->HasAnyConnections();
+			if(bHasAnyConnections)
+			{
+				break;
+			}
+		}
+	}
+	else
+	{
+		bHasAnyConnections = LinkedTo.Num() > 0;
+	}
+	return bHasAnyConnections;
+}
+
 FString UEdGraphPin::GetDefaultAsString() const
 {
 	if(DefaultObject)
