The MPC setting functions now show compiler errors if you refactor a MPC to remove an option.

diff --git a/Engine/Source/Editor/BlueprintGraph/Classes/K2Node_CallMaterialParameterCollectionFunction.h b/Engine/Source/Editor/BlueprintGraph/Classes/K2Node_CallMaterialParameterCollectionFunction.h
--- a/Engine/Source/Editor/BlueprintGraph/Classes/K2Node_CallMaterialParameterCollectionFunction.h
+++ b/Engine/Source/Editor/BlueprintGraph/Classes/K2Node_CallMaterialParameterCollectionFunction.h
@@ -13,6 +13,8 @@ class UK2Node_CallMaterialParameterCollectionFunction : public UK2Node_CallFunct
 	GENERATED_UCLASS_BODY()
 
 	//~ Begin EdGraphNode Interface
+	virtual void PreloadRequiredAssets() override;
 	virtual void PinDefaultValueChanged(UEdGraphPin* Pin) override;
+	virtual void ValidateNodeDuringCompilation(class FCompilerResultsLog& MessageLog) const override;
 	//~ End EdGraphNode Interface
 };
diff --git a/Engine/Source/Editor/BlueprintGraph/Private/K2Node_CallMaterialParameterCollectionFunction.cpp b/Engine/Source/Editor/BlueprintGraph/Private/K2Node_CallMaterialParameterCollectionFunction.cpp
--- a/Engine/Source/Editor/BlueprintGraph/Private/K2Node_CallMaterialParameterCollectionFunction.cpp
+++ b/Engine/Source/Editor/BlueprintGraph/Private/K2Node_CallMaterialParameterCollectionFunction.cpp
@@ -2,7 +2,8 @@
 
 #include "K2Node_CallMaterialParameterCollectionFunction.h"
 #include "EdGraph/EdGraph.h"
-
+#include "Kismet2/CompilerResultsLog.h"
+#include "Materials/MaterialParameterCollection.h"
 
 UK2Node_CallMaterialParameterCollectionFunction::UK2Node_CallMaterialParameterCollectionFunction(const FObjectInitializer& ObjectInitializer)
 	: Super(ObjectInitializer)
@@ -20,3 +21,46 @@ void UK2Node_CallMaterialParameterCollectionFunction::PinDefaultValueChanged(UEd
 	}
 }
 
+void UK2Node_CallMaterialParameterCollectionFunction::PreloadRequiredAssets()
+{
+	if (UEdGraphPin* CollectionPin = FindPin(TEXT("Collection")))
+	{
+		if (UMaterialParameterCollection* Collection = Cast<UMaterialParameterCollection>(CollectionPin->DefaultObject))
+		{
+			PreloadObject(Collection);
+			Collection->ConditionalPostLoad();
+		}
+	}
+
+	Super::PreloadRequiredAssets();
+}
+
+void UK2Node_CallMaterialParameterCollectionFunction::ValidateNodeDuringCompilation(class FCompilerResultsLog& MessageLog) const
+{
+	Super::ValidateNodeDuringCompilation(MessageLog);
+
+	if (UEdGraphPin* ParameterNamePin = FindPin(TEXT("ParameterName")))
+	{
+		if (ParameterNamePin->LinkedTo.Num() == 0)
+		{
+			if (ParameterNamePin->DefaultValue == TEXT(""))
+			{
+				MessageLog.Error(TEXT("@@ is invalid, @@ needs to be set to a parameter."), this, ParameterNamePin);
+			}
+			else
+			{
+				if (UEdGraphPin* CollectionPin = FindPin(TEXT("Collection")))
+				{
+					if (UMaterialParameterCollection* Collection = Cast<UMaterialParameterCollection>(CollectionPin->DefaultObject))
+					{
+						const FGuid ParameterNameGuid = Collection->GetParameterId(*ParameterNamePin->DefaultValue);
+						if (!ParameterNameGuid.IsValid())
+						{
+							MessageLog.Error(TEXT("@@ is invalid, @@ needs to be set to a parameter."), this, ParameterNamePin);
+						}
+					}
+				}
+			}
+		}
+	}
+}
\ No newline at end of file
