Graph Editor Context Menus refactored as ToolMenus

diff --git a/Engine/Plugins/Editor/AssetManagerEditor/Source/AssetManagerEditor/AssetManagerEditor.Build.cs b/Engine/Plugins/Editor/AssetManagerEditor/Source/AssetManagerEditor/AssetManagerEditor.Build.cs
--- a/Engine/Plugins/Editor/AssetManagerEditor/Source/AssetManagerEditor/AssetManagerEditor.Build.cs
+++ b/Engine/Plugins/Editor/AssetManagerEditor/Source/AssetManagerEditor/AssetManagerEditor.Build.cs
@@ -43,6 +43,7 @@ public class AssetManagerEditor : ModuleRules
 				"SandboxFile",
 				"EditorWidgets",
 				"TreeMap",
+				"ToolMenus",
 			}
 		);
 	}
diff --git a/Engine/Plugins/Editor/AssetManagerEditor/Source/AssetManagerEditor/Private/ReferenceViewer/ReferenceViewerSchema.cpp b/Engine/Plugins/Editor/AssetManagerEditor/Source/AssetManagerEditor/Private/ReferenceViewer/ReferenceViewerSchema.cpp
--- a/Engine/Plugins/Editor/AssetManagerEditor/Source/AssetManagerEditor/Private/ReferenceViewer/ReferenceViewerSchema.cpp
+++ b/Engine/Plugins/Editor/AssetManagerEditor/Source/AssetManagerEditor/Private/ReferenceViewer/ReferenceViewerSchema.cpp
@@ -3,7 +3,8 @@
 #include "ReferenceViewer/ReferenceViewerSchema.h"
 #include "Textures/SlateIcon.h"
 #include "Misc/Attribute.h"
-#include "Framework/MultiBox/MultiBoxBuilder.h"
+#include "ToolMenus.h"
+#include "EdGraph/EdGraph.h"
 #include "EditorStyleSet.h"
 #include "CollectionManagerTypes.h"
 #include "AssetManagerEditorCommands.h"
@@ -51,38 +52,40 @@ UReferenceViewerSchema::UReferenceViewerSchema(const FObjectInitializer& ObjectI
 {
 }
 
-void UReferenceViewerSchema::GetContextMenuActions(const UEdGraph* CurrentGraph, const UEdGraphNode* InGraphNode, const UEdGraphPin* InGraphPin, class FMenuBuilder* MenuBuilder, bool bIsDebugging) const
+void UReferenceViewerSchema::GetContextMenuActions(UToolMenu* Menu, UGraphNodeContextMenuContext* Context) const
 {
-	MenuBuilder->BeginSection(TEXT("Asset"), NSLOCTEXT("ReferenceViewerSchema", "AssetSectionLabel", "Asset"));
 	{
-		MenuBuilder->AddMenuEntry(FGlobalEditorCommonCommands::Get().FindInContentBrowser);
-		MenuBuilder->AddMenuEntry(FAssetManagerEditorCommands::Get().OpenSelectedInAssetEditor);
+		FToolMenuSection& Section = Menu->AddSection(TEXT("Asset"), NSLOCTEXT("ReferenceViewerSchema", "AssetSectionLabel", "Asset"));
+		Section.AddMenuEntry(FGlobalEditorCommonCommands::Get().FindInContentBrowser);
+		Section.AddMenuEntry(FAssetManagerEditorCommands::Get().OpenSelectedInAssetEditor);
 	}
-	MenuBuilder->EndSection();
 
-	MenuBuilder->BeginSection(TEXT("Misc"), NSLOCTEXT("ReferenceViewerSchema", "MiscSectionLabel", "Misc"));
 	{
-		MenuBuilder->AddMenuEntry(FAssetManagerEditorCommands::Get().ZoomToFit);
-		MenuBuilder->AddMenuEntry(FAssetManagerEditorCommands::Get().ReCenterGraph);
-		MenuBuilder->AddSubMenu(
+		FToolMenuSection& Section = Menu->AddSection(TEXT("Misc"), NSLOCTEXT("ReferenceViewerSchema", "MiscSectionLabel", "Misc"));
+		Section.AddMenuEntry(FAssetManagerEditorCommands::Get().ZoomToFit);
+		Section.AddMenuEntry(FAssetManagerEditorCommands::Get().ReCenterGraph);
+		Section.AddSubMenu(
+			Menu->GetMenuName(),
+			"MakeCollectionWith",
 			NSLOCTEXT("ReferenceViewerSchema", "MakeCollectionWithTitle", "Make Collection with"),
 			NSLOCTEXT("ReferenceViewerSchema", "MakeCollectionWithTooltip", "Makes a collection with either the referencers or dependencies of the selected nodes."),
-			FNewMenuDelegate::CreateUObject(const_cast<UReferenceViewerSchema*>(this), &UReferenceViewerSchema::GetMakeCollectionWithSubMenu)
+			FNewToolMenuDelegate::CreateUObject(const_cast<UReferenceViewerSchema*>(this), &UReferenceViewerSchema::GetMakeCollectionWithSubMenu)
 		);
 	}
-	MenuBuilder->EndSection();
 
-	MenuBuilder->BeginSection(TEXT("References"), NSLOCTEXT("ReferenceViewerSchema", "ReferencesSectionLabel", "References"));
 	{
-		MenuBuilder->AddMenuEntry(FAssetManagerEditorCommands::Get().CopyReferencedObjects);
-		MenuBuilder->AddMenuEntry(FAssetManagerEditorCommands::Get().CopyReferencingObjects);
-		MenuBuilder->AddMenuEntry(FAssetManagerEditorCommands::Get().ShowReferencedObjects);
-		MenuBuilder->AddMenuEntry(FAssetManagerEditorCommands::Get().ShowReferencingObjects);
-		MenuBuilder->AddMenuEntry(FAssetManagerEditorCommands::Get().ShowReferenceTree);
-		MenuBuilder->AddMenuEntry(FAssetManagerEditorCommands::Get().ViewSizeMap);
-		MenuBuilder->AddMenuEntry(FAssetManagerEditorCommands::Get().ViewAssetAudit, TEXT("ContextMenu"));
+		FToolMenuSection& Section = Menu->AddSection(TEXT("References"), NSLOCTEXT("ReferenceViewerSchema", "ReferencesSectionLabel", "References"));
+		Section.AddMenuEntry(FAssetManagerEditorCommands::Get().CopyReferencedObjects);
+		Section.AddMenuEntry(FAssetManagerEditorCommands::Get().CopyReferencingObjects);
+		Section.AddMenuEntry(FAssetManagerEditorCommands::Get().ShowReferencedObjects);
+		Section.AddMenuEntry(FAssetManagerEditorCommands::Get().ShowReferencingObjects);
+		Section.AddMenuEntry(FAssetManagerEditorCommands::Get().ShowReferenceTree);
+		Section.AddMenuEntry(FAssetManagerEditorCommands::Get().ViewSizeMap);
+
+		FToolMenuEntry ViewAssetAuditEntry = FToolMenuEntry::InitMenuEntry(FAssetManagerEditorCommands::Get().ViewAssetAudit);
+		ViewAssetAuditEntry.Name = TEXT("ContextMenu");
+		Section.AddEntry(ViewAssetAuditEntry);
 	}
-	MenuBuilder->EndSection();
 }
 
 FLinearColor UReferenceViewerSchema::GetPinTypeColor(const FEdGraphPinType& PinType) const
@@ -137,53 +140,63 @@ void UReferenceViewerSchema::GetAssetsGraphHoverMessage(const TArray<FAssetData>
 	OutOkIcon = true;
 }
 
-void UReferenceViewerSchema::GetMakeCollectionWithSubMenu(FMenuBuilder& MenuBuilder)
+void UReferenceViewerSchema::GetMakeCollectionWithSubMenu(UToolMenu* Menu)
 {
-	MenuBuilder.AddSubMenu(NSLOCTEXT("ReferenceViewerSchema", "MakeCollectionWithReferencersTitle", "Referencers <-"),
+	FToolMenuSection& Section = Menu->AddSection("Section");
+
+	Section.AddSubMenu(
+		Menu->GetMenuName(),
+		"MakeCollectionWithReferencers",
+		NSLOCTEXT("ReferenceViewerSchema", "MakeCollectionWithReferencersTitle", "Referencers <-"),
 		NSLOCTEXT("ReferenceViewerSchema", "MakeCollectionWithReferencersTooltip", "Makes a collection with assets one connection to the left of selected nodes."),
-		FNewMenuDelegate::CreateUObject(this, &UReferenceViewerSchema::GetMakeCollectionWithReferencersOrDependenciesSubMenu, true)
+		FNewToolMenuDelegate::CreateUObject(this, &UReferenceViewerSchema::GetMakeCollectionWithReferencersOrDependenciesSubMenu, true)
 		);
 
-	MenuBuilder.AddSubMenu(NSLOCTEXT("ReferenceViewerSchema", "MakeCollectionWithDependenciesTitle", "Dependencies ->"),
+	Section.AddSubMenu(
+		Menu->GetMenuName(),
+		"MakeCollectionWithDependencies",
+		NSLOCTEXT("ReferenceViewerSchema", "MakeCollectionWithDependenciesTitle", "Dependencies ->"),
 		NSLOCTEXT("ReferenceViewerSchema", "MakeCollectionWithDependenciesTooltip", "Makes a collection with assets one connection to the right of selected nodes."),
-		FNewMenuDelegate::CreateUObject(this, &UReferenceViewerSchema::GetMakeCollectionWithReferencersOrDependenciesSubMenu, false)
+		FNewToolMenuDelegate::CreateUObject(this, &UReferenceViewerSchema::GetMakeCollectionWithReferencersOrDependenciesSubMenu, false)
 		);
 }
 
-void UReferenceViewerSchema::GetMakeCollectionWithReferencersOrDependenciesSubMenu(FMenuBuilder& MenuBuilder, bool bReferencers)
+void UReferenceViewerSchema::GetMakeCollectionWithReferencersOrDependenciesSubMenu(UToolMenu* Menu, bool bReferencers)
 {
+	FToolMenuSection& Section = Menu->AddSection("Section");
+
 	if (bReferencers)
 	{
-		MenuBuilder.AddMenuEntry(FAssetManagerEditorCommands::Get().MakeLocalCollectionWithReferencers, 
-			NAME_None, TAttribute<FText>(), 
+		Section.AddMenuEntry(FAssetManagerEditorCommands::Get().MakeLocalCollectionWithReferencers, 
+			TAttribute<FText>(),
 			ECollectionShareType::GetDescription(ECollectionShareType::CST_Local), 
 			FSlateIcon(FEditorStyle::GetStyleSetName(), ECollectionShareType::GetIconStyleName(ECollectionShareType::CST_Local))
 			);
-		MenuBuilder.AddMenuEntry(FAssetManagerEditorCommands::Get().MakePrivateCollectionWithReferencers,
-			NAME_None, TAttribute<FText>(), 
+		Section.AddMenuEntry(FAssetManagerEditorCommands::Get().MakePrivateCollectionWithReferencers,
+			TAttribute<FText>(),
 			ECollectionShareType::GetDescription(ECollectionShareType::CST_Private), 
 			FSlateIcon(FEditorStyle::GetStyleSetName(), ECollectionShareType::GetIconStyleName(ECollectionShareType::CST_Private))
 			);
-		MenuBuilder.AddMenuEntry(FAssetManagerEditorCommands::Get().MakeSharedCollectionWithReferencers,
-			NAME_None, TAttribute<FText>(), 
+		Section.AddMenuEntry(FAssetManagerEditorCommands::Get().MakeSharedCollectionWithReferencers,
+			TAttribute<FText>(),
 			ECollectionShareType::GetDescription(ECollectionShareType::CST_Shared), 
 			FSlateIcon(FEditorStyle::GetStyleSetName(), ECollectionShareType::GetIconStyleName(ECollectionShareType::CST_Shared))
 			);
 	}
 	else
 	{
-		MenuBuilder.AddMenuEntry(FAssetManagerEditorCommands::Get().MakeLocalCollectionWithDependencies, 
-			NAME_None, TAttribute<FText>(), 
+		Section.AddMenuEntry(FAssetManagerEditorCommands::Get().MakeLocalCollectionWithDependencies, 
+			TAttribute<FText>(),
 			ECollectionShareType::GetDescription(ECollectionShareType::CST_Local), 
 			FSlateIcon(FEditorStyle::GetStyleSetName(), ECollectionShareType::GetIconStyleName(ECollectionShareType::CST_Local))
 			);
-		MenuBuilder.AddMenuEntry(FAssetManagerEditorCommands::Get().MakePrivateCollectionWithDependencies,
-			NAME_None, TAttribute<FText>(), 
+		Section.AddMenuEntry(FAssetManagerEditorCommands::Get().MakePrivateCollectionWithDependencies,
+			TAttribute<FText>(),
 			ECollectionShareType::GetDescription(ECollectionShareType::CST_Private), 
 			FSlateIcon(FEditorStyle::GetStyleSetName(), ECollectionShareType::GetIconStyleName(ECollectionShareType::CST_Private))
 			);
-		MenuBuilder.AddMenuEntry(FAssetManagerEditorCommands::Get().MakeSharedCollectionWithDependencies,
-			NAME_None, TAttribute<FText>(), 
+		Section.AddMenuEntry(FAssetManagerEditorCommands::Get().MakeSharedCollectionWithDependencies,
+			TAttribute<FText>(),
 			ECollectionShareType::GetDescription(ECollectionShareType::CST_Shared), 
 			FSlateIcon(FEditorStyle::GetStyleSetName(), ECollectionShareType::GetIconStyleName(ECollectionShareType::CST_Shared))
 			);
diff --git a/Engine/Plugins/Editor/AssetManagerEditor/Source/AssetManagerEditor/Public/ReferenceViewer/ReferenceViewerSchema.h b/Engine/Plugins/Editor/AssetManagerEditor/Source/AssetManagerEditor/Public/ReferenceViewer/ReferenceViewerSchema.h
--- a/Engine/Plugins/Editor/AssetManagerEditor/Source/AssetManagerEditor/Public/ReferenceViewer/ReferenceViewerSchema.h
+++ b/Engine/Plugins/Editor/AssetManagerEditor/Source/AssetManagerEditor/Public/ReferenceViewer/ReferenceViewerSchema.h
@@ -7,7 +7,7 @@
 #include "EdGraph/EdGraphSchema.h"
 #include "ReferenceViewerSchema.generated.h"
 
-class FMenuBuilder;
+class UToolMenu;
 class FSlateRect;
 class UEdGraph;
 
@@ -18,7 +18,8 @@ class UReferenceViewerSchema : public UEdGraphSchema
 
 public:
 	// UEdGraphSchema interface
-	virtual void GetContextMenuActions(const UEdGraph* CurrentGraph, const UEdGraphNode* InGraphNode, const UEdGraphPin* InGraphPin, class FMenuBuilder* MenuBuilder, bool bIsDebugging) const override;
+	virtual void GetContextMenuActions(class UToolMenu* Menu, class UGraphNodeContextMenuContext* Context) const override;
+	virtual FName GetParentContextMenuName() const override { return NAME_None; }
 	virtual FLinearColor GetPinTypeColor(const FEdGraphPinType& PinType) const override;
 	virtual void BreakPinLinks(UEdGraphPin& TargetPin, bool bSendsNodeNotifcation) const override;
 	virtual void BreakSinglePinLink(UEdGraphPin* SourcePin, UEdGraphPin* TargetPin) const override;
@@ -32,9 +33,9 @@ public:
 private:
 
 	/** Constructs the sub-menu for Make Collection With */
-	void GetMakeCollectionWithSubMenu(FMenuBuilder& MenuBuilder);
+	void GetMakeCollectionWithSubMenu(UToolMenu* Menu);
 
 	/** Constructs the sub-menu for Make Collection With Referencers/Dependencies */
-	void GetMakeCollectionWithReferencersOrDependenciesSubMenu(class FMenuBuilder& MenuBuilder, bool bReferencers);
+	void GetMakeCollectionWithReferencersOrDependenciesSubMenu(UToolMenu* Menu, bool bReferencers);
 };
 
diff --git a/Engine/Plugins/Experimental/ControlRig/Source/ControlRigDeveloper/Private/Graph/ControlRigGraphSchema.cpp b/Engine/Plugins/Experimental/ControlRig/Source/ControlRigDeveloper/Private/Graph/ControlRigGraphSchema.cpp
--- a/Engine/Plugins/Experimental/ControlRig/Source/ControlRigDeveloper/Private/Graph/ControlRigGraphSchema.cpp
+++ b/Engine/Plugins/Experimental/ControlRig/Source/ControlRigDeveloper/Private/Graph/ControlRigGraphSchema.cpp
@@ -31,10 +31,10 @@ void UControlRigGraphSchema::GetGraphContextActions(FGraphContextMenuBuilder& Co
 
 }
 
-void UControlRigGraphSchema::GetContextMenuActions(const UEdGraph* CurrentGraph, const UEdGraphNode* InGraphNode, const UEdGraphPin* InGraphPin, FMenuBuilder* MenuBuilder, bool bIsDebugging) const
+void UControlRigGraphSchema::GetContextMenuActions(class UToolMenu* Menu, class UGraphNodeContextMenuContext* Context) const
 {
 #if WITH_EDITOR
-	return IControlRigEditorModule::Get().GetContextMenuActions(this, CurrentGraph, InGraphNode, InGraphPin, MenuBuilder, bIsDebugging);
+	return IControlRigEditorModule::Get().GetContextMenuActions(this, Menu, Context);
 #else
 	check(0);
 #endif
diff --git a/Engine/Plugins/Experimental/ControlRig/Source/ControlRigDeveloper/Public/Graph/ControlRigGraphSchema.h b/Engine/Plugins/Experimental/ControlRig/Source/ControlRigDeveloper/Public/Graph/ControlRigGraphSchema.h
--- a/Engine/Plugins/Experimental/ControlRig/Source/ControlRigDeveloper/Public/Graph/ControlRigGraphSchema.h
+++ b/Engine/Plugins/Experimental/ControlRig/Source/ControlRigDeveloper/Public/Graph/ControlRigGraphSchema.h
@@ -53,7 +53,7 @@ public:
 
 	// UEdGraphSchema interface
 	virtual void GetGraphContextActions(FGraphContextMenuBuilder& ContextMenuBuilder) const override;
-	virtual void GetContextMenuActions(const UEdGraph* CurrentGraph, const UEdGraphNode* InGraphNode, const UEdGraphPin* InGraphPin, class FMenuBuilder* MenuBuilder, bool bIsDebugging) const override;
+	virtual void GetContextMenuActions(class UToolMenu* Menu, class UGraphNodeContextMenuContext* Context) const override;
 	virtual bool TryCreateConnection(UEdGraphPin* PinA, UEdGraphPin* PinB) const override;
 	virtual const FPinConnectionResponse CanCreateConnection(const UEdGraphPin* A, const UEdGraphPin* B) const override;
 	virtual FLinearColor GetPinTypeColor(const FEdGraphPinType& PinType) const override;
diff --git a/Engine/Plugins/Experimental/ControlRig/Source/ControlRigEditor/ControlRigEditor.Build.cs b/Engine/Plugins/Experimental/ControlRig/Source/ControlRigEditor/ControlRigEditor.Build.cs
--- a/Engine/Plugins/Experimental/ControlRig/Source/ControlRigEditor/ControlRigEditor.Build.cs
+++ b/Engine/Plugins/Experimental/ControlRig/Source/ControlRigEditor/ControlRigEditor.Build.cs
@@ -52,6 +52,7 @@ namespace UnrealBuildTool.Rules
 					"WorkspaceMenuStructure",
 					"Json",
 					"DesktopPlatform",
+					"ToolMenus",
                 }
             );
         }
diff --git a/Engine/Plugins/Experimental/ControlRig/Source/ControlRigEditor/Private/ControlRigEditorModule.cpp b/Engine/Plugins/Experimental/ControlRig/Source/ControlRigEditor/Private/ControlRigEditorModule.cpp
--- a/Engine/Plugins/Experimental/ControlRig/Source/ControlRigEditor/Private/ControlRigEditorModule.cpp
+++ b/Engine/Plugins/Experimental/ControlRig/Source/ControlRigEditor/Private/ControlRigEditorModule.cpp
@@ -40,6 +40,7 @@
 #include "Materials/Material.h"
 #include "Framework/MultiBox/MultiBoxExtender.h"
 #include "Framework/MultiBox/MultiBoxBuilder.h"
+#include "ToolMenus.h"
 #include "ControlRigEditModeSettings.h"
 #include "Components/SkeletalMeshComponent.h"
 #include "ControlRigSequenceExporter.h"
@@ -908,46 +909,39 @@ void FControlRigEditorModule::GetContextMenuActions(const UControlRigGraphNode*
 	}
 }
 
-void FControlRigEditorModule::GetContextMenuActions(const UControlRigGraphSchema* Schema, const UEdGraph* CurrentGraph, const UEdGraphNode* InGraphNode, const UEdGraphPin* InGraphPin, FMenuBuilder* MenuBuilder, bool bIsDebugging)
+void FControlRigEditorModule::GetContextMenuActions(const UControlRigGraphSchema* Schema, UToolMenu* Menu, UGraphNodeContextMenuContext* Context) const
 {
-	if(MenuBuilder)
+	if (Menu && Context)
 	{
-		MenuBuilder->BeginSection("ContextMenu");
+		Schema->UEdGraphSchema::GetContextMenuActions(Menu, Context);
 
-		Schema->UEdGraphSchema::GetContextMenuActions(CurrentGraph, InGraphNode, InGraphPin, MenuBuilder, bIsDebugging);
-
-		MenuBuilder->EndSection();
-
-		if (InGraphPin != NULL)
+		if (const UEdGraphPin* InGraphPin = Context->Pin)
 		{
-			MenuBuilder->BeginSection("EdGraphSchemaPinActions", LOCTEXT("PinActionsMenuHeader", "Pin Actions"));
 			{
+				FToolMenuSection& Section = Menu->AddSection("EdGraphSchemaPinActions", LOCTEXT("PinActionsMenuHeader", "Pin Actions"));
 				// Break pin links
 				if (InGraphPin->LinkedTo.Num() > 0)
 				{
-					MenuBuilder->AddMenuEntry( FGraphEditorCommands::Get().BreakPinLinks );
+					Section.AddMenuEntry(FGraphEditorCommands::Get().BreakPinLinks);
 				}
 			}
-			MenuBuilder->EndSection();
 
 			// Add the watch pin / unwatch pin menu items
-			MenuBuilder->BeginSection("EdGraphSchemaWatches", LOCTEXT("WatchesHeader", "Watches"));
 			{
-				UBlueprint* OwnerBlueprint = FBlueprintEditorUtils::FindBlueprintForGraphChecked(CurrentGraph);
+				FToolMenuSection& Section = Menu->AddSection("EdGraphSchemaWatches", LOCTEXT("WatchesHeader", "Watches"));
+				UBlueprint* OwnerBlueprint = FBlueprintEditorUtils::FindBlueprintForGraphChecked(Context->Graph);
 				{
 					const UEdGraphPin* WatchedPin = ((InGraphPin->Direction == EGPD_Input) && (InGraphPin->LinkedTo.Num() > 0)) ? InGraphPin->LinkedTo[0] : InGraphPin;
 					if (FKismetDebugUtilities::IsPinBeingWatched(OwnerBlueprint, WatchedPin))
 					{
-						MenuBuilder->AddMenuEntry(FGraphEditorCommands::Get().StopWatchingPin);
+						Section.AddMenuEntry(FGraphEditorCommands::Get().StopWatchingPin);
 					}
 					else
 					{
-						MenuBuilder->AddMenuEntry(FGraphEditorCommands::Get().StartWatchingPin);
+						Section.AddMenuEntry(FGraphEditorCommands::Get().StartWatchingPin);
 					}
 				}
 			}
-			MenuBuilder->EndSection();
-
 		}
 	}
 }
diff --git a/Engine/Plugins/Experimental/ControlRig/Source/ControlRigEditor/Private/ControlRigEditorModule.h b/Engine/Plugins/Experimental/ControlRig/Source/ControlRigEditor/Private/ControlRigEditorModule.h
--- a/Engine/Plugins/Experimental/ControlRig/Source/ControlRigEditor/Private/ControlRigEditorModule.h
+++ b/Engine/Plugins/Experimental/ControlRig/Source/ControlRigEditor/Private/ControlRigEditorModule.h
@@ -48,7 +48,7 @@ public:
 	virtual void GetInstanceActions(const UControlRigBlueprint* CRB, FBlueprintActionDatabaseRegistrar& ActionRegistrar) override;
 	virtual FConnectionDrawingPolicy* CreateConnectionDrawingPolicy(int32 InBackLayerID, int32 InFrontLayerID, float InZoomFactor, const FSlateRect& InClippingRect, class FSlateWindowElementList& InDrawElements, class UEdGraph* InGraphObj) override;
 	virtual void GetContextMenuActions(const UControlRigGraphNode* Node, const FGraphNodeContextMenuBuilder& Context ) override;
-	virtual void GetContextMenuActions(const UControlRigGraphSchema* Schema, const UEdGraph* CurrentGraph, const UEdGraphNode* InGraphNode, const UEdGraphPin* InGraphPin, FMenuBuilder* MenuBuilder, bool bIsDebugging) override;
+	virtual void GetContextMenuActions(const UControlRigGraphSchema* Schema, class UToolMenu* Menu, class UGraphNodeContextMenuContext* Context) const override;
 
 	static TSubclassOf<URigUnitEditor_Base> GetEditorObjectByRigUnit(const FName& RigUnitClassName);
 
diff --git a/Engine/Plugins/Experimental/ControlRig/Source/ControlRigEditor/Public/IControlRigEditorModule.h b/Engine/Plugins/Experimental/ControlRig/Source/ControlRigEditor/Public/IControlRigEditorModule.h
--- a/Engine/Plugins/Experimental/ControlRig/Source/ControlRigEditor/Public/IControlRigEditorModule.h
+++ b/Engine/Plugins/Experimental/ControlRig/Source/ControlRigEditor/Public/IControlRigEditorModule.h
@@ -47,5 +47,5 @@ public:
 	virtual void GetInstanceActions(const UControlRigBlueprint* CRB, FBlueprintActionDatabaseRegistrar& ActionRegistrar) = 0;
 	virtual FConnectionDrawingPolicy* CreateConnectionDrawingPolicy(int32 InBackLayerID, int32 InFrontLayerID, float InZoomFactor, const FSlateRect& InClippingRect, class FSlateWindowElementList& InDrawElements, class UEdGraph* InGraphObj) = 0;
 	virtual void GetContextMenuActions(const UControlRigGraphNode* Node, const FGraphNodeContextMenuBuilder& Context ) = 0;
-	virtual void GetContextMenuActions(const UControlRigGraphSchema* Schema, const UEdGraph* CurrentGraph, const UEdGraphNode* InGraphNode, const UEdGraphPin* InGraphPin, FMenuBuilder* MenuBuilder, bool bIsDebugging) = 0;
+	virtual void GetContextMenuActions(const UControlRigGraphSchema* Schema, class UToolMenu* Menu, class UGraphNodeContextMenuContext* Context) const = 0;
 };
diff --git a/Engine/Plugins/FX/Niagara/Source/NiagaraEditor/Private/EdGraphSchema_Niagara.cpp b/Engine/Plugins/FX/Niagara/Source/NiagaraEditor/Private/EdGraphSchema_Niagara.cpp
--- a/Engine/Plugins/FX/Niagara/Source/NiagaraEditor/Private/EdGraphSchema_Niagara.cpp
+++ b/Engine/Plugins/FX/Niagara/Source/NiagaraEditor/Private/EdGraphSchema_Niagara.cpp
@@ -6,6 +6,7 @@
 #include "Textures/SlateIcon.h"
 #include "Framework/Commands/UIAction.h"
 #include "Framework/MultiBox/MultiBoxBuilder.h"
+#include "ToolMenus.h"
 #include "NiagaraCommon.h"
 #include "NiagaraEditorCommon.h"
 #include "INiagaraCompiler.h"
@@ -1424,8 +1425,10 @@ FNiagaraTypeDefinition UEdGraphSchema_Niagara::GetTypeDefForProperty(const UProp
 	return FNiagaraTypeDefinition::GetFloatDef();//Some invalid type?
 }
 
-void UEdGraphSchema_Niagara::GetBreakLinkToSubMenuActions(class FMenuBuilder& MenuBuilder, UEdGraphPin* InGraphPin)
+void UEdGraphSchema_Niagara::GetBreakLinkToSubMenuActions(UToolMenu* Menu, const FName SectionName, UEdGraphPin* InGraphPin)
 {
+	FToolMenuSection& Section = Menu->FindOrAddSection(SectionName);
+
 	// Make sure we have a unique name for every entry in the list
 	TMap< FString, uint32 > LinkTitleCount;
 
@@ -1462,7 +1465,7 @@ void UEdGraphSchema_Niagara::GetBreakLinkToSubMenuActions(class FMenuBuilder& Me
 			Description = FText::Format(LOCTEXT("BreakDescMulti", "Break link to {NodeTitle} ({NumberOfNodes})"), Args);
 		}
 		++Count;
-		MenuBuilder.AddMenuEntry(Description, Description, FSlateIcon(), FUIAction(
+		Section.AddMenuEntry(NAME_None, Description, Description, FSlateIcon(), FUIAction(
 			FExecuteAction::CreateUObject((UEdGraphSchema_Niagara*const)this, &UEdGraphSchema::BreakSinglePinLink, const_cast<UEdGraphPin*>(InGraphPin), *Links)));
 	}
 }
@@ -1539,8 +1542,10 @@ bool UEdGraphSchema_Niagara::CheckCircularConnection(const UEdGraphNode* InRootN
 	return false;
 }
 
-void UEdGraphSchema_Niagara::GetNumericConversionToSubMenuActions(class FMenuBuilder& MenuBuilder, UEdGraphPin* InGraphPin)
+void UEdGraphSchema_Niagara::GetNumericConversionToSubMenuActions(UToolMenu* Menu, const FName SectionName, UEdGraphPin* InGraphPin)
 {
+	FToolMenuSection& Section = Menu->FindOrAddSection(SectionName);
+
 	// Add all the types we could convert to
 	for (const FNiagaraTypeDefinition& TypeDef : FNiagaraTypeRegistry::GetNumericTypes())
 	{
@@ -1550,13 +1555,15 @@ void UEdGraphSchema_Niagara::GetNumericConversionToSubMenuActions(class FMenuBui
 		FFormatNamedArguments Args;
 		Args.Add(TEXT("TypeTitle"), Title);
 		Description = FText::Format(LOCTEXT("NumericConversionText", "{TypeTitle}"), Args);
-		MenuBuilder.AddMenuEntry(Description, Description, FSlateIcon(), FUIAction(
+		Section.AddMenuEntry(NAME_None, Description, Description, FSlateIcon(), FUIAction(
 			FExecuteAction::CreateUObject((UEdGraphSchema_Niagara*const)this, &UEdGraphSchema_Niagara::ConvertNumericPinToType, const_cast<UEdGraphPin*>(InGraphPin), FNiagaraTypeDefinition(TypeDef))));
 	}
 }
 
-void UEdGraphSchema_Niagara::GetNumericConversionToSubMenuActionsAll(class FMenuBuilder& MenuBuilder, UNiagaraNode* InNode)
+void UEdGraphSchema_Niagara::GetNumericConversionToSubMenuActionsAll(UToolMenu* Menu, const FName SectionName, UNiagaraNode* InNode)
 {
+	FToolMenuSection& Section = Menu->FindOrAddSection(SectionName);
+
 	// Add all the types we could convert to
 	for (const FNiagaraTypeDefinition& TypeDef : FNiagaraTypeRegistry::GetNumericTypes())
 	{
@@ -1566,7 +1573,7 @@ void UEdGraphSchema_Niagara::GetNumericConversionToSubMenuActionsAll(class FMenu
 		FFormatNamedArguments Args;
 		Args.Add(TEXT("TypeTitle"), Title);
 		Description = FText::Format(LOCTEXT("NumericConversionText", "{TypeTitle}"), Args);
-		MenuBuilder.AddMenuEntry(Description, Description, FSlateIcon(), FUIAction(
+		Section.AddMenuEntry(NAME_None, Description, Description, FSlateIcon(), FUIAction(
 			FExecuteAction::CreateUObject((UEdGraphSchema_Niagara*const)this, &UEdGraphSchema_Niagara::ConvertNumericPinToTypeAll, InNode, FNiagaraTypeDefinition(TypeDef))));
 	}
 }
@@ -1641,47 +1648,55 @@ bool CanResetPinToDefault(const UEdGraphSchema_Niagara* Schema, const UEdGraphPi
 	return Schema->DoesDefaultValueMatchAutogenerated(*Pin) == false;
 }
 
-void UEdGraphSchema_Niagara::GetContextMenuActions(const UEdGraph* CurrentGraph, const UEdGraphNode* InGraphNode, const UEdGraphPin* InGraphPin, class FMenuBuilder* MenuBuilder, bool bIsDebugging) const
+void UEdGraphSchema_Niagara::GetContextMenuActions(UToolMenu* Menu, UGraphNodeContextMenuContext* Context) const
 {
+	const UEdGraphNode* InGraphNode = Context->Node;
+	const UEdGraphPin* InGraphPin = Context->Pin;
 	if (InGraphPin)
 	{
-		MenuBuilder->BeginSection("EdGraphSchema_NiagaraPinActions", LOCTEXT("PinActionsMenuHeader", "Pin Actions"));
 		{
+			const FName SectionName = "EdGraphSchema_NiagaraPinActions";
+			FToolMenuSection& Section = Menu->AddSection(SectionName, LOCTEXT("PinActionsMenuHeader", "Pin Actions"));
 			if (PinToTypeDefinition(InGraphPin) == FNiagaraTypeDefinition::GetGenericNumericDef() && InGraphPin->LinkedTo.Num() == 0)
 			{
-				MenuBuilder->AddSubMenu(
+				Section.AddSubMenu(
+					Menu->GetMenuName(),
+					"ConvertNumericSpecific",
 					LOCTEXT("ConvertNumericSpecific", "Convert Numeric To..."),
 					LOCTEXT("ConvertNumericSpecificToolTip", "Convert Numeric pin to the specific typed pin."),
-				FNewMenuDelegate::CreateUObject((UEdGraphSchema_Niagara*const)this, &UEdGraphSchema_Niagara::GetNumericConversionToSubMenuActions, const_cast<UEdGraphPin*>(InGraphPin)));
+				FNewToolMenuDelegate::CreateUObject((UEdGraphSchema_Niagara*const)this, &UEdGraphSchema_Niagara::GetNumericConversionToSubMenuActions, SectionName, const_cast<UEdGraphPin*>(InGraphPin)));
 			}
 
 			// Only display the 'Break Link' option if there is a link to break!
 			if (InGraphPin->LinkedTo.Num() > 0)
 			{
-				MenuBuilder->AddMenuEntry(FGraphEditorCommands::Get().BreakPinLinks);
+				Section.AddMenuEntry(FGraphEditorCommands::Get().BreakPinLinks);
 
 				// add sub menu for break link to
 				if (InGraphPin->LinkedTo.Num() > 1)
 				{
-					MenuBuilder->AddSubMenu(
+					Section.AddSubMenu(
+						Menu->GetMenuName(),
+						"BreakLinkTo",
 						LOCTEXT("BreakLinkTo", "Break Link To..."),
 						LOCTEXT("BreakSpecificLinks", "Break a specific link..."),
-						FNewMenuDelegate::CreateUObject((UEdGraphSchema_Niagara*const)this, &UEdGraphSchema_Niagara::GetBreakLinkToSubMenuActions, const_cast<UEdGraphPin*>(InGraphPin)));
+						FNewToolMenuDelegate::CreateUObject((UEdGraphSchema_Niagara*const)this, &UEdGraphSchema_Niagara::GetBreakLinkToSubMenuActions, SectionName, const_cast<UEdGraphPin*>(InGraphPin)));
 				}
 				else
 				{
-					((UEdGraphSchema_Niagara*const)this)->GetBreakLinkToSubMenuActions(*MenuBuilder, const_cast<UEdGraphPin*>(InGraphPin));
+					((UEdGraphSchema_Niagara*const)this)->GetBreakLinkToSubMenuActions(Menu, SectionName,const_cast<UEdGraphPin*>(InGraphPin));
 				}
 			}
 
 			if (InGraphPin->Direction == EEdGraphPinDirection::EGPD_Input)
 			{
-				MenuBuilder->AddMenuEntry(LOCTEXT("PromoteToParameter", "Promote to Parameter"), LOCTEXT("PromoteToParameterTooltip", "Create a parameter argument and connect this pin to that parameter."), FSlateIcon(),
+				Section.AddMenuEntry("PromoteToParameter", LOCTEXT("PromoteToParameter", "Promote to Parameter"), LOCTEXT("PromoteToParameterTooltip", "Create a parameter argument and connect this pin to that parameter."), FSlateIcon(),
 					FUIAction(FExecuteAction::CreateUObject((UEdGraphSchema_Niagara*const)this, &UEdGraphSchema_Niagara::PromoteSinglePinToParameter, const_cast<UEdGraphPin*>(InGraphPin)),
 						FCanExecuteAction::CreateStatic(&UEdGraphSchema_Niagara::CanPromoteSinglePinToParameter, InGraphPin)));
 				if (InGraphPin->LinkedTo.Num() == 0 && InGraphPin->bDefaultValueIsIgnored == false)
 				{
-					MenuBuilder->AddMenuEntry(
+					Section.AddMenuEntry(
+						"ResetInputToDefault",
 						LOCTEXT("ResetInputToDefault", "Reset to Default"),  // TODO(mv): This is currently broken
 						LOCTEXT("ResetInputToDefaultToolTip", "Reset this input to its default value."),
 						FSlateIcon(),
@@ -1691,7 +1706,6 @@ void UEdGraphSchema_Niagara::GetContextMenuActions(const UEdGraph* CurrentGraph,
 				}
 			}
 		}
-		MenuBuilder->EndSection();
 	}
 	else if (InGraphNode)
 	{
@@ -1719,24 +1733,24 @@ void UEdGraphSchema_Niagara::GetContextMenuActions(const UEdGraph* CurrentGraph,
 		}
 		if (bHasNumerics)
 		{
-			MenuBuilder->BeginSection("EdGraphSchema_NiagaraNodeActions", LOCTEXT("PinConversionMenuHeader", "Convert Pins"));
-			MenuBuilder->AddSubMenu(
+			const FName SectionName = "EdGraphSchema_NiagaraNodeActions";
+			FToolMenuSection& Section = Menu->AddSection(SectionName, LOCTEXT("PinConversionMenuHeader", "Convert Pins"));
+			Section.AddSubMenu(
+				Menu->GetMenuName(),
+				"ConvertNumericSpecific",
 				LOCTEXT("ConvertNumericSpecific", "Convert All Numerics To..."),
 				LOCTEXT("ConvertNumericSpecificToolTip", "Convert all Numeric pins to the specific typed pin."),
-				FNewMenuDelegate::CreateUObject((UEdGraphSchema_Niagara*const)this, &UEdGraphSchema_Niagara::GetNumericConversionToSubMenuActionsAll, const_cast<UNiagaraNode*>(Node)));
-			MenuBuilder->EndSection();
+				FNewToolMenuDelegate::CreateUObject((UEdGraphSchema_Niagara*const)this, &UEdGraphSchema_Niagara::GetNumericConversionToSubMenuActionsAll, SectionName, const_cast<UNiagaraNode*>(Node)));
 		}
 		
-		MenuBuilder->BeginSection("EdGraphSchema_NiagaraNodeActions", LOCTEXT("NodeActionsMenuHeader", "Node Actions"));
-		MenuBuilder->AddMenuEntry(LOCTEXT("ToggleEnabledState", "Toggle Enabled State"), LOCTEXT("ToggleEnabledStateTooltip", "Toggle this node between Enbled (default) and Disabled (skipped from compilation)."), FSlateIcon(),
+		FToolMenuSection& Section = Menu->AddSection("EdGraphSchema_NiagaraNodeActions", LOCTEXT("NodeActionsMenuHeader", "Node Actions"));
+		Section.AddMenuEntry("ToggleEnabledState", LOCTEXT("ToggleEnabledState", "Toggle Enabled State"), LOCTEXT("ToggleEnabledStateTooltip", "Toggle this node between Enbled (default) and Disabled (skipped from compilation)."), FSlateIcon(),
 			FUIAction(FExecuteAction::CreateUObject((UEdGraphSchema_Niagara*const)this, &UEdGraphSchema_Niagara::ToggleNodeEnabledState, const_cast<UNiagaraNode*>(Node))));
-		MenuBuilder->AddMenuEntry(LOCTEXT("RefreshNode", "Refresh Node"), LOCTEXT("RefreshNodeTooltip", "Refresh this node."), FSlateIcon(),
+		Section.AddMenuEntry("RefreshNode", LOCTEXT("RefreshNode", "Refresh Node"), LOCTEXT("RefreshNodeTooltip", "Refresh this node."), FSlateIcon(),
 			FUIAction(FExecuteAction::CreateUObject((UEdGraphSchema_Niagara*const)this, &UEdGraphSchema_Niagara::RefreshNode, const_cast<UNiagaraNode*>(Node))));
-		
-		MenuBuilder->EndSection();
 	}
 
-	Super::GetContextMenuActions(CurrentGraph, InGraphNode, InGraphPin, MenuBuilder, bIsDebugging);
+	Super::GetContextMenuActions(Menu, Context);
 }
 
 FNiagaraConnectionDrawingPolicy::FNiagaraConnectionDrawingPolicy(int32 InBackLayerID, int32 InFrontLayerID, float InZoomFactor, const FSlateRect& InClippingRect, FSlateWindowElementList& InDrawElements, UEdGraph* InGraph)
diff --git a/Engine/Plugins/FX/Niagara/Source/NiagaraEditor/Private/EdGraphSchema_NiagaraSystemOverview.cpp b/Engine/Plugins/FX/Niagara/Source/NiagaraEditor/Private/EdGraphSchema_NiagaraSystemOverview.cpp
--- a/Engine/Plugins/FX/Niagara/Source/NiagaraEditor/Private/EdGraphSchema_NiagaraSystemOverview.cpp
+++ b/Engine/Plugins/FX/Niagara/Source/NiagaraEditor/Private/EdGraphSchema_NiagaraSystemOverview.cpp
@@ -12,7 +12,7 @@ UEdGraphSchema_NiagaraSystemOverview::UEdGraphSchema_NiagaraSystemOverview(const
 {
 }
 
-void UEdGraphSchema_NiagaraSystemOverview::GetContextMenuActions(const UEdGraph* CurrentGraph, const UEdGraphNode* InGraphNode, const UEdGraphPin* InGraphPin, class FMenuBuilder* MenuBuilder, bool bIsDebugging) const
+void UEdGraphSchema_NiagaraSystemOverview::GetContextMenuActions(UToolMenu* Menu, UGraphNodeContextMenuContext* Context) const
 {
 	//@TODO System Overview: enable/disable emitter, solo, etc. 
 	return;
diff --git a/Engine/Plugins/FX/Niagara/Source/NiagaraEditor/Public/EdGraphSchema_Niagara.h b/Engine/Plugins/FX/Niagara/Source/NiagaraEditor/Public/EdGraphSchema_Niagara.h
--- a/Engine/Plugins/FX/Niagara/Source/NiagaraEditor/Public/EdGraphSchema_Niagara.h
+++ b/Engine/Plugins/FX/Niagara/Source/NiagaraEditor/Public/EdGraphSchema_Niagara.h
@@ -91,7 +91,7 @@ class NIAGARAEDITOR_API UEdGraphSchema_Niagara : public UEdGraphSchema
 
 	//~ Begin EdGraphSchema Interface
 	virtual void GetGraphContextActions(FGraphContextMenuBuilder& ContextMenuBuilder) const override;
-	virtual void GetContextMenuActions(const UEdGraph* CurrentGraph, const UEdGraphNode* InGraphNode, const UEdGraphPin* InGraphPin, class FMenuBuilder* MenuBuilder, bool bIsDebugging) const override;
+	virtual void GetContextMenuActions(class UToolMenu* Menu, class UGraphNodeContextMenuContext* Context) const override;
 	virtual const FPinConnectionResponse CanCreateConnection(const UEdGraphPin* A, const UEdGraphPin* B) const override;
 	virtual FLinearColor GetPinTypeColor(const FEdGraphPinType& PinType) const override;
 	virtual bool ShouldHidePinDefaultValue(UEdGraphPin* Pin) const override;
@@ -149,9 +149,9 @@ class NIAGARAEDITOR_API UEdGraphSchema_Niagara : public UEdGraphSchema
 	static const FLinearColor NodeTitleColor_RapidIteration;
 	
 private:
-	void GetBreakLinkToSubMenuActions(class FMenuBuilder& MenuBuilder, UEdGraphPin* InGraphPin);
-	void GetNumericConversionToSubMenuActions(class FMenuBuilder& MenuBuilder, UEdGraphPin* InGraphPin);
-	void GetNumericConversionToSubMenuActionsAll(class FMenuBuilder& MenuBuilder, UNiagaraNode* InGraphPin);
+	void GetBreakLinkToSubMenuActions(class UToolMenu* Menu, const FName SectionName, UEdGraphPin* InGraphPin);
+	void GetNumericConversionToSubMenuActions(class UToolMenu* Menu, const FName SectionName, UEdGraphPin* InGraphPin);
+	void GetNumericConversionToSubMenuActionsAll(class UToolMenu* Menu, const FName SectionName, UNiagaraNode* InGraphPin);
 	void ConvertNumericPinToType(UEdGraphPin* InPin, FNiagaraTypeDefinition TypeDef);
 	void ConvertNumericPinToTypeAll(UNiagaraNode* InPin, FNiagaraTypeDefinition TypeDef);
 	static bool CheckCircularConnection(const UEdGraphNode* InRootNode, const EEdGraphPinDirection InRootPinDirection, const UEdGraphPin* InPin, int32& OutDepth);
diff --git a/Engine/Plugins/FX/Niagara/Source/NiagaraEditor/Public/EdGraphSchema_NiagaraSystemOverview.h b/Engine/Plugins/FX/Niagara/Source/NiagaraEditor/Public/EdGraphSchema_NiagaraSystemOverview.h
--- a/Engine/Plugins/FX/Niagara/Source/NiagaraEditor/Public/EdGraphSchema_NiagaraSystemOverview.h
+++ b/Engine/Plugins/FX/Niagara/Source/NiagaraEditor/Public/EdGraphSchema_NiagaraSystemOverview.h
@@ -15,7 +15,8 @@ class NIAGARAEDITOR_API UEdGraphSchema_NiagaraSystemOverview : public UEdGraphSc
 	GENERATED_UCLASS_BODY()
 
 	//~ Begin EdGraphSchema Interface
-	virtual void GetContextMenuActions(const UEdGraph* CurrentGraph, const UEdGraphNode* InGraphNode, const UEdGraphPin* InGraphPin, class FMenuBuilder* MenuBuilder, bool bIsDebugging) const override;
+	virtual void GetContextMenuActions(class UToolMenu* Menu, class UGraphNodeContextMenuContext* Context) const override;
+	virtual FName GetParentContextMenuName() const override { return NAME_None; }
 	virtual const FPinConnectionResponse CanCreateConnection(const UEdGraphPin* A, const UEdGraphPin* B) const override { return FPinConnectionResponse(CONNECT_RESPONSE_DISALLOW, FText()); }; //@TODO System Overview: write text response
 	virtual bool TryCreateConnection(UEdGraphPin* A, UEdGraphPin* B) const override { return false; };
 	virtual bool ShouldAlwaysPurgeOnModification() const override { return false; }
diff --git a/Engine/Source/Developer/ToolMenus/Private/ToolMenu.cpp b/Engine/Source/Developer/ToolMenus/Private/ToolMenu.cpp
--- a/Engine/Source/Developer/ToolMenus/Private/ToolMenu.cpp
+++ b/Engine/Source/Developer/ToolMenus/Private/ToolMenu.cpp
@@ -30,11 +30,11 @@ void UToolMenu::InitMenu(const FToolMenuOwner InOwner, FName InName, FName InPar
 	MenuType = InType;
 }
 
-void UToolMenu::InitGeneratedCopy(const UToolMenu* Source)
+void UToolMenu::InitGeneratedCopy(const UToolMenu* Source, const FName InMenuName, const FToolMenuContext* InContext)
 {
-	// Skip sections and context
+	// Skip sections
 
-	MenuName = Source->MenuName;
+	MenuName = InMenuName;
 	MenuParent = Source->MenuParent;
 	StyleName = Source->StyleName;
 	TutorialHighlightName = Source->TutorialHighlightName;
@@ -46,6 +46,11 @@ void UToolMenu::InitGeneratedCopy(const UToolMenu* Source)
 	bToolBarIsFocusable = Source->bToolBarIsFocusable;
 	bToolBarForceSmallIcons = Source->bToolBarForceSmallIcons;
 	MenuOwner = Source->MenuOwner;
+
+	if (InContext)
+	{
+		Context = *InContext;
+	}
 }
 
 int32 UToolMenu::IndexOfSection(const FName InSectionName) const
diff --git a/Engine/Source/Developer/ToolMenus/Private/ToolMenuContext.cpp b/Engine/Source/Developer/ToolMenus/Private/ToolMenuContext.cpp
--- a/Engine/Source/Developer/ToolMenus/Private/ToolMenuContext.cpp
+++ b/Engine/Source/Developer/ToolMenus/Private/ToolMenuContext.cpp
@@ -98,6 +98,22 @@ const FUIAction* FToolMenuContext::GetActionForCommand(TSharedPtr<const FUIComma
 	return nullptr;
 }
 
+const FUIAction* FToolMenuContext::GetActionForCommand(TSharedPtr<const FUICommandInfo> Command) const
+{
+	for (const TSharedPtr<FUICommandList>& CommandListIter : CommandLists)
+	{
+		if (CommandListIter.IsValid())
+		{
+			if (const FUIAction* Result = CommandListIter->GetActionForCommand(Command))
+			{
+				return Result;
+			}
+		}
+	}
+
+	return nullptr;
+}
+
 void FToolMenuContext::AddExtender(const TSharedPtr<FExtender>& InExtender)
 {
 	Extenders.AddUnique(InExtender);
diff --git a/Engine/Source/Developer/ToolMenus/Private/ToolMenuSection.cpp b/Engine/Source/Developer/ToolMenus/Private/ToolMenuSection.cpp
--- a/Engine/Source/Developer/ToolMenus/Private/ToolMenuSection.cpp
+++ b/Engine/Source/Developer/ToolMenus/Private/ToolMenuSection.cpp
@@ -103,6 +103,11 @@ FToolMenuEntry& FToolMenuSection::AddMenuSeparator(const FName InName)
 	return AddEntry(FToolMenuEntry::InitMenuSeparator(InName));
 }
 
+FToolMenuEntry& FToolMenuSection::AddSubMenu(const FName InParentMenu, const FName InName, const TAttribute<FText>& InLabel, const TAttribute<FText>& InToolTip, const FNewToolMenuChoice& InMakeMenu, bool bInOpenSubMenuOnClick, const TAttribute<FSlateIcon>& InIcon, const bool bShouldCloseWindowAfterMenuSelection)
+{
+	return AddEntry(FToolMenuEntry::InitSubMenu(InParentMenu, InName, InLabel, InToolTip, InMakeMenu, bInOpenSubMenuOnClick, InIcon, bShouldCloseWindowAfterMenuSelection));
+}
+
 int32 FToolMenuSection::IndexOfBlock(const FName InName) const
 {
 	for (int32 i=0; i < Blocks.Num(); ++i)
diff --git a/Engine/Source/Developer/ToolMenus/Private/ToolMenus.cpp b/Engine/Source/Developer/ToolMenus/Private/ToolMenus.cpp
--- a/Engine/Source/Developer/ToolMenus/Private/ToolMenus.cpp
+++ b/Engine/Source/Developer/ToolMenus/Private/ToolMenus.cpp
@@ -161,7 +161,7 @@ void UToolMenus::AssembleMenuByName(UToolMenu* GeneratedMenu, const FName InName
 {
 	if (UToolMenu* Menu = FindMenu(InName))
 	{
-		GeneratedMenu->InitGeneratedCopy(Menu);
+		GeneratedMenu->InitGeneratedCopy(Menu, Menu->MenuName);
 		AssembleMenuHierarchy(GeneratedMenu, CollectHierarchy(InName));
 	}
 }
@@ -569,8 +569,9 @@ void UToolMenus::PopulateMenuBuilder(FMenuBuilder& MenuBuilder, UToolMenu* MenuD
 		);
 	}
 
-	for (FToolMenuSection& Section : MenuData->Sections)
+	for (int i=0; i < MenuData->Sections.Num(); ++i)
 	{
+		FToolMenuSection& Section = MenuData->Sections[i];
 		if (Section.Construct.NewToolMenuDelegateLegacy.IsBound())
 		{
 			Section.Construct.NewToolMenuDelegateLegacy.Execute(MenuBuilder, MenuData);
@@ -1120,9 +1121,7 @@ UToolMenu* UToolMenus::GenerateMenu(const TArray<UToolMenu*>& Hierarchy, FToolMe
 
 	if (Hierarchy.Num() > 0)
 	{
-		GeneratedMenu->InitGeneratedCopy(Hierarchy[0]);
-		GeneratedMenu->MenuName = Hierarchy.Last()->MenuName;
-		GeneratedMenu->Context = InMenuContext;
+		GeneratedMenu->InitGeneratedCopy(Hierarchy[0], Hierarchy.Last()->MenuName, &InMenuContext);
 		AssembleMenuHierarchy(GeneratedMenu, Hierarchy);
 	}
 
diff --git a/Engine/Source/Developer/ToolMenus/Private/ToolMenusBlueprintLibrary.h b/Engine/Source/Developer/ToolMenus/Private/ToolMenusBlueprintLibrary.h
--- a/Engine/Source/Developer/ToolMenus/Private/ToolMenusBlueprintLibrary.h
+++ b/Engine/Source/Developer/ToolMenus/Private/ToolMenusBlueprintLibrary.h
@@ -10,12 +10,11 @@
 #include "ToolMenuMisc.h"
 #include "ToolMenuEntryScript.h"
 #include "Misc/Attribute.h"
-#include "Kismet/BlueprintFunctionLibrary.h"
 
 #include "ToolMenusBlueprintLibrary.generated.h"
 
 UCLASS()
-class UToolMenuContextExtensions : public UBlueprintFunctionLibrary
+class UToolMenuContextExtensions : public UObject
 {
 	GENERATED_BODY()
 
@@ -26,7 +25,7 @@ public:
 };
 
 UCLASS()
-class UToolMenuEntryExtensions : public UBlueprintFunctionLibrary
+class UToolMenuEntryExtensions : public UObject
 {
 	GENERATED_BODY()
 
@@ -73,7 +72,7 @@ public:
 };
 
 UCLASS()
-class UToolMenuSectionExtensions : public UBlueprintFunctionLibrary
+class UToolMenuSectionExtensions : public UObject
 {
 	GENERATED_BODY()
 
diff --git a/Engine/Source/Developer/ToolMenus/Public/ToolMenu.h b/Engine/Source/Developer/ToolMenus/Public/ToolMenu.h
--- a/Engine/Source/Developer/ToolMenus/Public/ToolMenu.h
+++ b/Engine/Source/Developer/ToolMenus/Public/ToolMenu.h
@@ -69,7 +69,7 @@ public:
 
 private:
 
-	void InitGeneratedCopy(const UToolMenu* Source);
+	void InitGeneratedCopy(const UToolMenu* Source, const FName InMenuName, const FToolMenuContext* InContext = nullptr);
 
 	bool FindEntry(const FName EntryName, int32& SectionIndex, int32& EntryIndex) const;
 
diff --git a/Engine/Source/Developer/ToolMenus/Public/ToolMenuContext.h b/Engine/Source/Developer/ToolMenus/Public/ToolMenuContext.h
--- a/Engine/Source/Developer/ToolMenus/Public/ToolMenuContext.h
+++ b/Engine/Source/Developer/ToolMenus/Public/ToolMenuContext.h
@@ -56,6 +56,7 @@ public:
 	void AppendCommandList(const TSharedRef<FUICommandList>& InCommandList);
 	void AppendCommandList(const TSharedPtr<FUICommandList>& InCommandList);
 	const FUIAction* GetActionForCommand(TSharedPtr<const FUICommandInfo> Command, TSharedPtr<const FUICommandList>& OutCommandList) const;
+	const FUIAction* GetActionForCommand(TSharedPtr<const FUICommandInfo> Command) const;
 
 	void AddExtender(const TSharedPtr<FExtender>& InExtender);
 	TSharedPtr<FExtender> GetAllExtenders();
diff --git a/Engine/Source/Developer/ToolMenus/Public/ToolMenuSection.h b/Engine/Source/Developer/ToolMenus/Public/ToolMenuSection.h
--- a/Engine/Source/Developer/ToolMenus/Public/ToolMenuSection.h
+++ b/Engine/Source/Developer/ToolMenus/Public/ToolMenuSection.h
@@ -42,6 +42,8 @@ public:
 	FToolMenuEntry& AddDynamicEntry(const FName InName, const FNewToolMenuDelegateLegacy& InConstruct);
 	FToolMenuEntry& AddMenuSeparator(const FName InName);
 
+	FToolMenuEntry& AddSubMenu(const FName InParentMenu, const FName InName, const TAttribute<FText>& InLabel, const TAttribute<FText>& InToolTip, const FNewToolMenuChoice& InMakeMenu, bool bInOpenSubMenuOnClick = false, const TAttribute<FSlateIcon>& InIcon = TAttribute<FSlateIcon>(), const bool bShouldCloseWindowAfterMenuSelection = true);
+
 	template <typename TContextType>
 	TContextType* FindContext() const
 	{
diff --git a/Engine/Source/Developer/ToolMenus/Public/ToolMenus.h b/Engine/Source/Developer/ToolMenus/Public/ToolMenus.h
--- a/Engine/Source/Developer/ToolMenus/Public/ToolMenus.h
+++ b/Engine/Source/Developer/ToolMenus/Public/ToolMenus.h
@@ -264,6 +264,9 @@ public:
 	/** For advanced use cases */
 	void AssembleMenuHierarchy(UToolMenu* GeneratedMenu, const TArray<UToolMenu*>& Hierarchy);
 
+	/* Returns list of menus starting with root parent */
+	TArray<UToolMenu*> CollectHierarchy(const FName Name);
+
 	/** For advanced use cases */
 	FToolMenuOwner CurrentOwner() const;
 
@@ -314,8 +317,6 @@ private:
 
 	UToolMenu* FindSubMenuToGenerateWith(const FName InParentName, const FName InChildName);
 
-	TArray<UToolMenu*> CollectHierarchy(const FName Name);
-
 	void FillMenu(FMenuBuilder& MenuBuilder, FName InMenuName, FToolMenuContext InMenuContext);
 	void FillMenuBarDropDown(FMenuBuilder& MenuBuilder, FName InParentName, FName InChildName, FToolMenuContext InMenuContext);
 	void PopulateMenuBuilder(FMenuBuilder& MenuBuilder, UToolMenu* MenuData);
diff --git a/Engine/Source/Developer/ToolMenus/ToolMenus.Build.cs b/Engine/Source/Developer/ToolMenus/ToolMenus.Build.cs
--- a/Engine/Source/Developer/ToolMenus/ToolMenus.Build.cs
+++ b/Engine/Source/Developer/ToolMenus/ToolMenus.Build.cs
@@ -13,7 +13,6 @@ namespace UnrealBuildTool.Rules
 					"ApplicationCore",
 					"Slate",
 					"SlateCore",
-					"Engine",
 				}
 			);
 		}
diff --git a/Engine/Source/Editor/AIGraph/AIGraph.Build.cs b/Engine/Source/Editor/AIGraph/AIGraph.Build.cs
--- a/Engine/Source/Editor/AIGraph/AIGraph.Build.cs
+++ b/Engine/Source/Editor/AIGraph/AIGraph.Build.cs
@@ -40,7 +40,8 @@ public class AIGraph : ModuleRules
 				"AnimGraph",
 				"BlueprintGraph",
                 "AIModule",
-				"ClassViewer"
+				"ClassViewer",
+				"ToolMenus",
 			}
         );
 
diff --git a/Engine/Source/Editor/AIGraph/Classes/AIGraphSchema.h b/Engine/Source/Editor/AIGraph/Classes/AIGraphSchema.h
--- a/Engine/Source/Editor/AIGraph/Classes/AIGraphSchema.h
+++ b/Engine/Source/Editor/AIGraph/Classes/AIGraphSchema.h
@@ -85,10 +85,10 @@ class AIGRAPH_API UAIGraphSchema : public UEdGraphSchema
 {
 	GENERATED_UCLASS_BODY()
 
-	void GetBreakLinkToSubMenuActions(class FMenuBuilder& MenuBuilder, class UEdGraphPin* InGraphPin);
+	void GetBreakLinkToSubMenuActions(class UToolMenu* Menu, class UEdGraphPin* InGraphPin);
 
 	//~ Begin EdGraphSchema Interface
-	virtual void GetContextMenuActions(const UEdGraph* CurrentGraph, const UEdGraphNode* InGraphNode, const UEdGraphPin* InGraphPin, class FMenuBuilder* MenuBuilder, bool bIsDebugging) const override;
+	virtual void GetContextMenuActions(class UToolMenu* Menu, class UGraphNodeContextMenuContext* Context) const override;
 	virtual FLinearColor GetPinTypeColor(const FEdGraphPinType& PinType) const override;
 	virtual bool ShouldHidePinDefaultValue(UEdGraphPin* Pin) const override;
 	virtual class FConnectionDrawingPolicy* CreateConnectionDrawingPolicy(int32 InBackLayerID, int32 InFrontLayerID, float InZoomFactor, const FSlateRect& InClippingRect, class FSlateWindowElementList& InDrawElements, class UEdGraph* InGraphObj) const override;
diff --git a/Engine/Source/Editor/AIGraph/Private/AIGraphSchema.cpp b/Engine/Source/Editor/AIGraph/Private/AIGraphSchema.cpp
--- a/Engine/Source/Editor/AIGraph/Private/AIGraphSchema.cpp
+++ b/Engine/Source/Editor/AIGraph/Private/AIGraphSchema.cpp
@@ -4,6 +4,7 @@
 #include "Textures/SlateIcon.h"
 #include "Framework/Commands/UIAction.h"
 #include "Framework/MultiBox/MultiBoxBuilder.h"
+#include "ToolMenus.h"
 #include "EdGraph/EdGraph.h"
 #include "AIGraphNode.h"
 #include "GraphEditorActions.h"
@@ -170,55 +171,57 @@ void UAIGraphSchema::GetGraphNodeContextActions(FGraphContextMenuBuilder& Contex
 	}
 }
 
-void UAIGraphSchema::GetContextMenuActions(const UEdGraph* CurrentGraph, const UEdGraphNode* InGraphNode, const UEdGraphPin* InGraphPin, class FMenuBuilder* MenuBuilder, bool bIsDebugging) const
+void UAIGraphSchema::GetContextMenuActions(class UToolMenu* Menu, class UGraphNodeContextMenuContext* Context) const
 {
-	if (InGraphPin)
+	if (Context->Pin)
 	{
-		MenuBuilder->BeginSection("AIGraphSchemaPinActions", LOCTEXT("PinActionsMenuHeader", "Pin Actions"));
 		{
+			FToolMenuSection& Section = Menu->AddSection("AIGraphSchemaPinActions", LOCTEXT("PinActionsMenuHeader", "Pin Actions"));
 			// Only display the 'Break Links' option if there is a link to break!
-			if (InGraphPin->LinkedTo.Num() > 0)
+			if (Context->Pin->LinkedTo.Num() > 0)
 			{
-				MenuBuilder->AddMenuEntry(FGraphEditorCommands::Get().BreakPinLinks);
+				Section.AddMenuEntry(FGraphEditorCommands::Get().BreakPinLinks);
 
 				// add sub menu for break link to
-				if (InGraphPin->LinkedTo.Num() > 1)
+				if (Context->Pin->LinkedTo.Num() > 1)
 				{
-					MenuBuilder->AddSubMenu(
+					Section.AddSubMenu(
+						Menu->GetMenuName(),
+						"BreakLinkTo",
 						LOCTEXT("BreakLinkTo", "Break Link To..."),
 						LOCTEXT("BreakSpecificLinks", "Break a specific link..."),
-						FNewMenuDelegate::CreateUObject((UAIGraphSchema*const)this, &UAIGraphSchema::GetBreakLinkToSubMenuActions, const_cast<UEdGraphPin*>(InGraphPin)));
+						FNewToolMenuDelegate::CreateUObject((UAIGraphSchema*const)this, &UAIGraphSchema::GetBreakLinkToSubMenuActions, const_cast<UEdGraphPin*>(Context->Pin)));
 				}
 				else
 				{
-					((UAIGraphSchema*const)this)->GetBreakLinkToSubMenuActions(*MenuBuilder, const_cast<UEdGraphPin*>(InGraphPin));
+					((UAIGraphSchema*const)this)->GetBreakLinkToSubMenuActions(Menu, const_cast<UEdGraphPin*>(Context->Pin));
 				}
 			}
 		}
-		MenuBuilder->EndSection();
 	}
-	else if (InGraphNode)
+	else if (Context->Node)
 	{
-		MenuBuilder->BeginSection("BehaviorTreeGraphSchemaNodeActions", LOCTEXT("ClassActionsMenuHeader", "Node Actions"));
 		{
-			MenuBuilder->AddMenuEntry(FGenericCommands::Get().Delete);
-			MenuBuilder->AddMenuEntry(FGenericCommands::Get().Cut);
-			MenuBuilder->AddMenuEntry(FGenericCommands::Get().Copy);
-			MenuBuilder->AddMenuEntry(FGenericCommands::Get().Duplicate);
+			FToolMenuSection& Section = Menu->AddSection("BehaviorTreeGraphSchemaNodeActions", LOCTEXT("ClassActionsMenuHeader", "Node Actions"));
+			Section.AddMenuEntry(FGenericCommands::Get().Delete);
+			Section.AddMenuEntry(FGenericCommands::Get().Cut);
+			Section.AddMenuEntry(FGenericCommands::Get().Copy);
+			Section.AddMenuEntry(FGenericCommands::Get().Duplicate);
 
-			MenuBuilder->AddMenuEntry(FGraphEditorCommands::Get().BreakNodeLinks);
+			Section.AddMenuEntry(FGraphEditorCommands::Get().BreakNodeLinks);
 		}
-		MenuBuilder->EndSection();
 	}
 
-	Super::GetContextMenuActions(CurrentGraph, InGraphNode, InGraphPin, MenuBuilder, bIsDebugging);
+	Super::GetContextMenuActions(Menu, Context);
 }
 
-void UAIGraphSchema::GetBreakLinkToSubMenuActions(class FMenuBuilder& MenuBuilder, UEdGraphPin* InGraphPin)
+void UAIGraphSchema::GetBreakLinkToSubMenuActions(UToolMenu* Menu, UEdGraphPin* InGraphPin)
 {
 	// Make sure we have a unique name for every entry in the list
 	TMap< FString, uint32 > LinkTitleCount;
 
+	FToolMenuSection& Section = Menu->FindOrAddSection("AIGraphSchemaPinActions");
+
 	// Add all the links we could break from
 	for (TArray<class UEdGraphPin*>::TConstIterator Links(InGraphPin->LinkedTo); Links; ++Links)
 	{
@@ -253,7 +256,7 @@ void UAIGraphSchema::GetBreakLinkToSubMenuActions(class FMenuBuilder& MenuBuilde
 		}
 		++Count;
 
-		MenuBuilder.AddMenuEntry(Description, Description, FSlateIcon(), FUIAction(
+		Section.AddMenuEntry(NAME_None, Description, Description, FSlateIcon(), FUIAction(
 			FExecuteAction::CreateUObject(this, &UAIGraphSchema::BreakSinglePinLink, const_cast< UEdGraphPin* >(InGraphPin), *Links)));
 	}
 }
diff --git a/Engine/Source/Editor/AnimGraph/AnimGraph.Build.cs b/Engine/Source/Editor/AnimGraph/AnimGraph.Build.cs
--- a/Engine/Source/Editor/AnimGraph/AnimGraph.Build.cs
+++ b/Engine/Source/Editor/AnimGraph/AnimGraph.Build.cs
@@ -34,6 +34,7 @@ public class AnimGraph : ModuleRules
 				"EditorStyle",
                 "ContentBrowser",
 				"KismetWidgets",
+				"ToolMenus",
             }
 		);
 
diff --git a/Engine/Source/Editor/AnimGraph/Classes/AnimationGraphSchema.h b/Engine/Source/Editor/AnimGraph/Classes/AnimationGraphSchema.h
--- a/Engine/Source/Editor/AnimGraph/Classes/AnimationGraphSchema.h
+++ b/Engine/Source/Editor/AnimGraph/Classes/AnimationGraphSchema.h
@@ -55,7 +55,7 @@ class UAnimationGraphSchema : public UEdGraphSchema_K2
 	virtual void GetAssetsNodeHoverMessage(const TArray<FAssetData>& Assets, const UEdGraphNode* HoverNode, FString& OutTooltipText, bool& OutOkIcon) const override;
 	virtual void GetAssetsPinHoverMessage(const TArray<FAssetData>& Assets, const UEdGraphPin* HoverPin, FString& OutTooltipText, bool& OutOkIcon) const override;
 	virtual void GetAssetsGraphHoverMessage(const TArray<FAssetData>& Assets, const UEdGraph* HoverGraph, FString& OutTooltipText, bool& OutOkIcon) const override;
-	virtual void GetContextMenuActions(const UEdGraph* CurrentGraph, const UEdGraphNode* InGraphNode, const UEdGraphPin* InGraphPin, FMenuBuilder* MenuBuilder, bool bIsDebugging) const override;
+	virtual void GetContextMenuActions(class UToolMenu* Menu, class UGraphNodeContextMenuContext* Context) const override;
 	virtual FText GetPinDisplayName(const UEdGraphPin* Pin) const override;
 	virtual bool CanDuplicateGraph(UEdGraph* InSourceGraph) const override { return InSourceGraph->GetFName() != UEdGraphSchema_K2::GN_AnimGraph; }
 	virtual bool DoesSupportEventDispatcher() const	override { return false; }
diff --git a/Engine/Source/Editor/AnimGraph/Classes/AnimationStateMachineSchema.h b/Engine/Source/Editor/AnimGraph/Classes/AnimationStateMachineSchema.h
--- a/Engine/Source/Editor/AnimGraph/Classes/AnimationStateMachineSchema.h
+++ b/Engine/Source/Editor/AnimGraph/Classes/AnimationStateMachineSchema.h
@@ -77,7 +77,7 @@ class UAnimationStateMachineSchema : public UEdGraphSchema
 	virtual bool CreateAutomaticConversionNodeAndConnections(UEdGraphPin* PinA, UEdGraphPin* PinB) const override;
 	virtual void GetGraphContextActions(FGraphContextMenuBuilder& ContextMenuBuilder) const override;
 	virtual EGraphType GetGraphType(const UEdGraph* TestEdGraph) const override;
-	virtual void GetContextMenuActions(const UEdGraph* CurrentGraph, const UEdGraphNode* InGraphNode, const UEdGraphPin* InGraphPin, FMenuBuilder* MenuBuilder, bool bIsDebugging) const override;
+	virtual void GetContextMenuActions(class UToolMenu* Menu, class UGraphNodeContextMenuContext* Context) const override;
 	virtual FLinearColor GetPinTypeColor(const FEdGraphPinType& PinType) const override;
 	virtual void GetGraphDisplayInformation(const UEdGraph& Graph, /*out*/ FGraphDisplayInfo& DisplayInfo) const override;
 	virtual void DroppedAssetsOnGraph(const TArray<FAssetData>& Assets, const FVector2D& GraphPosition, UEdGraph* Graph) const override;
diff --git a/Engine/Source/Editor/AnimGraph/Private/AnimationGraphSchema.cpp b/Engine/Source/Editor/AnimGraph/Private/AnimationGraphSchema.cpp
--- a/Engine/Source/Editor/AnimGraph/Private/AnimationGraphSchema.cpp
+++ b/Engine/Source/Editor/AnimGraph/Private/AnimationGraphSchema.cpp
@@ -8,6 +8,7 @@
 #include "Animation/AnimationAsset.h"
 #include "Animation/AnimBlueprint.h"
 #include "Framework/MultiBox/MultiBoxBuilder.h"
+#include "ToolMenus.h"
 #include "K2Node.h"
 #include "EdGraphSchema_K2_Actions.h"
 #include "Kismet2/BlueprintEditorUtils.h"
@@ -549,18 +550,17 @@ void UAnimationGraphSchema::GetAssetsGraphHoverMessage(const TArray<FAssetData>&
 	}
 }
 
-void UAnimationGraphSchema::GetContextMenuActions(const UEdGraph* CurrentGraph, const UEdGraphNode* InGraphNode, const UEdGraphPin* InGraphPin, FMenuBuilder* MenuBuilder, bool bIsDebugging) const
+void UAnimationGraphSchema::GetContextMenuActions(UToolMenu* Menu, UGraphNodeContextMenuContext* Context) const
 {
-	Super::GetContextMenuActions(CurrentGraph, InGraphNode, InGraphPin, MenuBuilder, bIsDebugging);
+	Super::GetContextMenuActions(Menu, Context);
 
-	if (const UAnimGraphNode_Base* AnimGraphNode = Cast<const UAnimGraphNode_Base>(InGraphNode))
+	if (const UAnimGraphNode_Base* AnimGraphNode = Cast<const UAnimGraphNode_Base>(Context->Node))
 	{
-		MenuBuilder->BeginSection("AnimGraphSchemaNodeActions", LOCTEXT("AnimNodeActionsMenuHeader", "Anim Node Actions"));
 		{
 			// Node contextual actions
-			MenuBuilder->AddMenuEntry(FAnimGraphCommands::Get().TogglePoseWatch);
+			FToolMenuSection& Section = Menu->AddSection("AnimGraphSchemaNodeActions", LOCTEXT("AnimNodeActionsMenuHeader", "Anim Node Actions"));
+			Section.AddMenuEntry(FAnimGraphCommands::Get().TogglePoseWatch);
 		}
-		MenuBuilder->EndSection();
 	}
 }
 
diff --git a/Engine/Source/Editor/AnimGraph/Private/AnimationStateMachineSchema.cpp b/Engine/Source/Editor/AnimGraph/Private/AnimationStateMachineSchema.cpp
--- a/Engine/Source/Editor/AnimGraph/Private/AnimationStateMachineSchema.cpp
+++ b/Engine/Source/Editor/AnimGraph/Private/AnimationStateMachineSchema.cpp
@@ -8,6 +8,7 @@
 #include "Layout/SlateRect.h"
 #include "Animation/AnimationAsset.h"
 #include "Framework/MultiBox/MultiBoxBuilder.h"
+#include "ToolMenus.h"
 #include "EdGraph/EdGraph.h"
 #include "Kismet2/BlueprintEditorUtils.h"
 #include "EdGraphNode_Comment.h"
@@ -319,34 +320,33 @@ EGraphType UAnimationStateMachineSchema::GetGraphType(const UEdGraph* TestEdGrap
 	return GT_StateMachine;
 }
 
-void UAnimationStateMachineSchema::GetContextMenuActions(const UEdGraph* CurrentGraph, const UEdGraphNode* InGraphNode, const UEdGraphPin* InGraphPin, FMenuBuilder* MenuBuilder, bool bIsDebugging) const
+void UAnimationStateMachineSchema::GetContextMenuActions(UToolMenu* Menu, UGraphNodeContextMenuContext* Context) const
 {
-	check(CurrentGraph);
-	UBlueprint* OwnerBlueprint = FBlueprintEditorUtils::FindBlueprintForGraphChecked(CurrentGraph);
+	check(Context && Context->Graph);
+	UBlueprint* OwnerBlueprint = FBlueprintEditorUtils::FindBlueprintForGraphChecked(Context->Graph);
 
-	if (InGraphNode != NULL)
+	if (Context->Node)
 	{
-		MenuBuilder->BeginSection("AnimationStateMachineNodeActions", LOCTEXT("NodeActionsMenuHeader", "Node Actions"));
 		{
-			if (!bIsDebugging)
+			FToolMenuSection& Section = Menu->AddSection("AnimationStateMachineNodeActions", LOCTEXT("NodeActionsMenuHeader", "Node Actions"));
+			if (!Context->bIsDebugging)
 			{
 				// Node contextual actions
-				MenuBuilder->AddMenuEntry( FGenericCommands::Get().Delete );
-				MenuBuilder->AddMenuEntry( FGenericCommands::Get().Cut );
-				MenuBuilder->AddMenuEntry( FGenericCommands::Get().Copy );
-				MenuBuilder->AddMenuEntry( FGenericCommands::Get().Duplicate );
-				MenuBuilder->AddMenuEntry( FGraphEditorCommands::Get().ReconstructNodes );
-				MenuBuilder->AddMenuEntry( FGraphEditorCommands::Get().BreakNodeLinks );
-				if(InGraphNode->bCanRenameNode)
+				Section.AddMenuEntry(FGenericCommands::Get().Delete);
+				Section.AddMenuEntry(FGenericCommands::Get().Cut);
+				Section.AddMenuEntry(FGenericCommands::Get().Copy);
+				Section.AddMenuEntry(FGenericCommands::Get().Duplicate);
+				Section.AddMenuEntry(FGraphEditorCommands::Get().ReconstructNodes);
+				Section.AddMenuEntry(FGraphEditorCommands::Get().BreakNodeLinks);
+				if(Context->Node->bCanRenameNode)
 				{
-					MenuBuilder->AddMenuEntry( FGenericCommands::Get().Rename );
+					Section.AddMenuEntry(FGenericCommands::Get().Rename);
 				}
 			}
 		}
-		MenuBuilder->EndSection();
 	}
 
-	Super::GetContextMenuActions(CurrentGraph, InGraphNode, InGraphPin, MenuBuilder, bIsDebugging);
+	Super::GetContextMenuActions(Menu, Context);
 }
 
 FLinearColor UAnimationStateMachineSchema::GetPinTypeColor(const FEdGraphPinType& PinType) const
diff --git a/Engine/Source/Editor/AudioEditor/Classes/SoundClassGraph/SoundClassGraphSchema.h b/Engine/Source/Editor/AudioEditor/Classes/SoundClassGraph/SoundClassGraphSchema.h
--- a/Engine/Source/Editor/AudioEditor/Classes/SoundClassGraph/SoundClassGraphSchema.h
+++ b/Engine/Source/Editor/AudioEditor/Classes/SoundClassGraph/SoundClassGraphSchema.h
@@ -45,11 +45,12 @@ class USoundClassGraphSchema : public UEdGraphSchema
 	/** Check whether connecting these pins would cause a loop */
 	bool ConnectionCausesLoop(const UEdGraphPin* InputPin, const UEdGraphPin* OutputPin) const;
 	/** Get menu for breaking links to specific nodes*/
-	void GetBreakLinkToSubMenuActions(class FMenuBuilder& MenuBuilder, class UEdGraphPin* InGraphPin);
+	void GetBreakLinkToSubMenuActions(class UToolMenu* Menu, class UEdGraphPin* InGraphPin);
 
 	//~ Begin EdGraphSchema Interface
 	virtual void GetGraphContextActions(FGraphContextMenuBuilder& ContextMenuBuilder) const override;
-	virtual void GetContextMenuActions(const UEdGraph* CurrentGraph, const UEdGraphNode* InGraphNode, const UEdGraphPin* InGraphPin, class FMenuBuilder* MenuBuilder, bool bIsDebugging) const override;
+	virtual void GetContextMenuActions(class UToolMenu* Menu, class UGraphNodeContextMenuContext* Context) const override;
+	virtual FName GetParentContextMenuName() const override { return NAME_None; }
 	virtual const FPinConnectionResponse CanCreateConnection(const UEdGraphPin* PinA, const UEdGraphPin* PinB) const override;
 	virtual bool TryCreateConnection(UEdGraphPin* PinA, UEdGraphPin* PinB) const override;
 	virtual bool ShouldHidePinDefaultValue(UEdGraphPin* Pin) const override;
diff --git a/Engine/Source/Editor/AudioEditor/Classes/SoundCueGraph/SoundCueGraphSchema.h b/Engine/Source/Editor/AudioEditor/Classes/SoundCueGraph/SoundCueGraphSchema.h
--- a/Engine/Source/Editor/AudioEditor/Classes/SoundCueGraph/SoundCueGraphSchema.h
+++ b/Engine/Source/Editor/AudioEditor/Classes/SoundCueGraph/SoundCueGraphSchema.h
@@ -114,7 +114,7 @@ class USoundCueGraphSchema : public UEdGraphSchema
 
 	//~ Begin EdGraphSchema Interface
 	virtual void GetGraphContextActions(FGraphContextMenuBuilder& ContextMenuBuilder) const override;
-	virtual void GetContextMenuActions(const UEdGraph* CurrentGraph, const UEdGraphNode* InGraphNode, const UEdGraphPin* InGraphPin, class FMenuBuilder* MenuBuilder, bool bIsDebugging) const override;
+	virtual void GetContextMenuActions(class UToolMenu* Menu, class UGraphNodeContextMenuContext* Context) const override;
 	virtual void CreateDefaultNodesForGraph(UEdGraph& Graph) const override;
 	virtual const FPinConnectionResponse CanCreateConnection(const UEdGraphPin* A, const UEdGraphPin* B) const override;
 	virtual bool TryCreateConnection(UEdGraphPin* A, UEdGraphPin* B) const override;
diff --git a/Engine/Source/Editor/AudioEditor/Classes/SoundSubmixGraph/SoundSubmixGraphSchema.h b/Engine/Source/Editor/AudioEditor/Classes/SoundSubmixGraph/SoundSubmixGraphSchema.h
--- a/Engine/Source/Editor/AudioEditor/Classes/SoundSubmixGraph/SoundSubmixGraphSchema.h
+++ b/Engine/Source/Editor/AudioEditor/Classes/SoundSubmixGraph/SoundSubmixGraphSchema.h
@@ -41,11 +41,12 @@ class USoundSubmixGraphSchema : public UEdGraphSchema
 	bool ConnectionCausesLoop(const UEdGraphPin* InputPin, const UEdGraphPin* OutputPin) const;
 
 	/** Get menu for breaking links to specific nodes*/
-	void GetBreakLinkToSubMenuActions(class FMenuBuilder& MenuBuilder, class UEdGraphPin* InGraphPin);
+	void GetBreakLinkToSubMenuActions(UToolMenu* Menu, const FName SectionName, class UEdGraphPin* InGraphPin);
 
 	//~ Begin EdGraphSchema Interface
 	virtual void GetGraphContextActions(FGraphContextMenuBuilder& ContextMenuBuilder) const override;
-	virtual void GetContextMenuActions(const UEdGraph* CurrentGraph, const UEdGraphNode* InGraphNode, const UEdGraphPin* InGraphPin, class FMenuBuilder* MenuBuilder, bool bIsDebugging) const override;
+	virtual void GetContextMenuActions(class UToolMenu* Menu, class UGraphNodeContextMenuContext* Context) const override;
+	virtual FName GetParentContextMenuName() const override { return NAME_None; }
 	virtual const FPinConnectionResponse CanCreateConnection(const UEdGraphPin* PinA, const UEdGraphPin* PinB) const override;
 	virtual bool TryCreateConnection(UEdGraphPin* PinA, UEdGraphPin* PinB) const override;
 	virtual bool ShouldHidePinDefaultValue(UEdGraphPin* Pin) const override;
diff --git a/Engine/Source/Editor/AudioEditor/Private/SoundClassGraphSchema.cpp b/Engine/Source/Editor/AudioEditor/Private/SoundClassGraphSchema.cpp
--- a/Engine/Source/Editor/AudioEditor/Private/SoundClassGraphSchema.cpp
+++ b/Engine/Source/Editor/AudioEditor/Private/SoundClassGraphSchema.cpp
@@ -4,6 +4,7 @@
 #include "Textures/SlateIcon.h"
 #include "Framework/Commands/UIAction.h"
 #include "Framework/MultiBox/MultiBoxBuilder.h"
+#include "ToolMenus.h"
 #include "SoundClassGraph/SoundClassGraph.h"
 #include "SoundClassGraph/SoundClassGraphNode.h"
 #include "Sound/SoundClass.h"
@@ -33,8 +34,10 @@ bool USoundClassGraphSchema::ConnectionCausesLoop(const UEdGraphPin* InputPin, c
 	return InputNode->SoundClass->RecurseCheckChild( OutputNode->SoundClass );
 }
 
-void USoundClassGraphSchema::GetBreakLinkToSubMenuActions( class FMenuBuilder& MenuBuilder, UEdGraphPin* InGraphPin )
+void USoundClassGraphSchema::GetBreakLinkToSubMenuActions(UToolMenu* Menu, UEdGraphPin* InGraphPin)
 {
+	FToolMenuSection& Section = Menu->FindOrAddSection("SoundClassGraphSchemaPinActions");
+
 	// Make sure we have a unique name for every entry in the list
 	TMap< FString, uint32 > LinkTitleCount;
 
@@ -72,7 +75,7 @@ void USoundClassGraphSchema::GetBreakLinkToSubMenuActions( class FMenuBuilder& M
 		}
 		++Count;
 
-		MenuBuilder.AddMenuEntry( Description, Description, FSlateIcon(), FUIAction(
+		Section.AddMenuEntry(NAME_None, Description, Description, FSlateIcon(), FUIAction(
 			FExecuteAction::CreateUObject((USoundClassGraphSchema*const)this, &USoundClassGraphSchema::BreakSinglePinLink, const_cast< UEdGraphPin* >(InGraphPin), *Links) ) );
 	}
 }
@@ -87,43 +90,43 @@ void USoundClassGraphSchema::GetGraphContextActions(FGraphContextMenuBuilder& Co
 	ContextMenuBuilder.AddAction( NewAction );
 }
 
-void USoundClassGraphSchema::GetContextMenuActions(const UEdGraph* CurrentGraph, const UEdGraphNode* InGraphNode, const UEdGraphPin* InGraphPin, class FMenuBuilder* MenuBuilder, bool bIsDebugging) const
+void USoundClassGraphSchema::GetContextMenuActions(UToolMenu* Menu, UGraphNodeContextMenuContext* Context) const
 {
-	if (InGraphPin)
+	if (Context->Pin)
 	{
-		MenuBuilder->BeginSection("SoundClassGraphSchemaPinActions", LOCTEXT("PinActionsMenuHeader", "Pin Actions"));
 		{
+			FToolMenuSection& Section = Menu->AddSection("SoundClassGraphSchemaPinActions", LOCTEXT("PinActionsMenuHeader", "Pin Actions"));
 			// Only display the 'Break Links' option if there is a link to break!
-			if (InGraphPin->LinkedTo.Num() > 0)
+			if (Context->Pin->LinkedTo.Num() > 0)
 			{
-				MenuBuilder->AddMenuEntry( FGraphEditorCommands::Get().BreakPinLinks );
+				Section.AddMenuEntry(FGraphEditorCommands::Get().BreakPinLinks);
 
 				// add sub menu for break link to
-				if(InGraphPin->LinkedTo.Num() > 1)
+				if(Context->Pin->LinkedTo.Num() > 1)
 				{
-					MenuBuilder->AddSubMenu(
+					Section.AddSubMenu(
+						Menu->GetMenuName(),
+						"BreakLinkTo",
 						LOCTEXT("BreakLinkTo", "Break Link To..." ),
 						LOCTEXT("BreakSpecificLinks", "Break a specific link..." ),
-						FNewMenuDelegate::CreateUObject( (USoundClassGraphSchema*const)this, &USoundClassGraphSchema::GetBreakLinkToSubMenuActions, const_cast<UEdGraphPin*>(InGraphPin)));
+						FNewToolMenuDelegate::CreateUObject( (USoundClassGraphSchema*const)this, &USoundClassGraphSchema::GetBreakLinkToSubMenuActions, const_cast<UEdGraphPin*>(Context->Pin)));
 				}
 				else
 				{
-					((USoundClassGraphSchema*const)this)->GetBreakLinkToSubMenuActions(*MenuBuilder, const_cast<UEdGraphPin*>(InGraphPin));
+					((USoundClassGraphSchema*const)this)->GetBreakLinkToSubMenuActions(Menu, const_cast<UEdGraphPin*>(Context->Pin));
 				}
 			}
 		}
-		MenuBuilder->EndSection();
 	}
-	else if (InGraphNode)
+	else if (Context->Node)
 	{
-		const USoundClassGraphNode* SoundGraphNode = Cast<const USoundClassGraphNode>(InGraphNode);
+		const USoundClassGraphNode* SoundGraphNode = Cast<const USoundClassGraphNode>(Context->Node);
 
-		MenuBuilder->BeginSection("SoundClassGraphSchemaNodeActions", LOCTEXT("ClassActionsMenuHeader", "SoundClass Actions"));
 		{
-			MenuBuilder->AddMenuEntry(FGraphEditorCommands::Get().BreakNodeLinks);
-			MenuBuilder->AddMenuEntry(FGenericCommands::Get().Delete);
+			FToolMenuSection& Section = Menu->AddSection("SoundClassGraphSchemaNodeActions", LOCTEXT("ClassActionsMenuHeader", "SoundClass Actions"));
+			Section.AddMenuEntry(FGraphEditorCommands::Get().BreakNodeLinks);
+			Section.AddMenuEntry(FGenericCommands::Get().Delete);
 		}
-		MenuBuilder->EndSection();
 	}
 
 	// No Super call so Node comments option is not shown
diff --git a/Engine/Source/Editor/AudioEditor/Private/SoundCueGraphSchema.cpp b/Engine/Source/Editor/AudioEditor/Private/SoundCueGraphSchema.cpp
--- a/Engine/Source/Editor/AudioEditor/Private/SoundCueGraphSchema.cpp
+++ b/Engine/Source/Editor/AudioEditor/Private/SoundCueGraphSchema.cpp
@@ -9,6 +9,7 @@
 #include "UObject/UObjectIterator.h"
 #include "Layout/SlateRect.h"
 #include "Framework/MultiBox/MultiBoxBuilder.h"
+#include "ToolMenus.h"
 #include "EdGraphNode_Comment.h"
 #include "EdGraph/EdGraph.h"
 #include "SoundCueGraph/SoundCueGraph.h"
@@ -273,32 +274,29 @@ void USoundCueGraphSchema::GetGraphContextActions(FGraphContextMenuBuilder& Cont
 	}
 }
 
-void USoundCueGraphSchema::GetContextMenuActions(const UEdGraph* CurrentGraph, const UEdGraphNode* InGraphNode, const UEdGraphPin* InGraphPin, class FMenuBuilder* MenuBuilder, bool bIsDebugging) const
+void USoundCueGraphSchema::GetContextMenuActions(class UToolMenu* Menu, class UGraphNodeContextMenuContext* Context) const
 {
-	if (InGraphPin)
+	if (Context->Pin)
 	{
-		MenuBuilder->BeginSection("SoundCueGraphSchemaPinActions", LOCTEXT("PinActionsMenuHeader", "Pin Actions"));
 		{
+			FToolMenuSection& Section = Menu->AddSection("SoundCueGraphSchemaPinActions", LOCTEXT("PinActionsMenuHeader", "Pin Actions"));
 			// Only display the 'Break Link' option if there is a link to break!
-			if (InGraphPin->LinkedTo.Num() > 0)
+			if (Context->Pin->LinkedTo.Num() > 0)
 			{
-				MenuBuilder->AddMenuEntry( FGraphEditorCommands::Get().BreakPinLinks );
+				Section.AddMenuEntry(FGraphEditorCommands::Get().BreakPinLinks);
 			}
 		}
-		MenuBuilder->EndSection();
 	}
-	else if (InGraphNode)
+	else if (Context->Node)
 	{
-		const USoundCueGraphNode* SoundGraphNode = Cast<const USoundCueGraphNode>(InGraphNode);
-
-		MenuBuilder->BeginSection("SoundCueGraphSchemaNodeActions", LOCTEXT("NodeActionsMenuHeader", "Node Actions"));
+		const USoundCueGraphNode* SoundGraphNode = Cast<const USoundCueGraphNode>(Context->Node);
 		{
-			MenuBuilder->AddMenuEntry(FGraphEditorCommands::Get().BreakNodeLinks);
+			FToolMenuSection& Section = Menu->AddSection("SoundCueGraphSchemaNodeActions", LOCTEXT("NodeActionsMenuHeader", "Node Actions"));
+			Section.AddMenuEntry(FGraphEditorCommands::Get().BreakNodeLinks);
 		}
-		MenuBuilder->EndSection();
 	}
 
-	Super::GetContextMenuActions(CurrentGraph, InGraphNode, InGraphPin, MenuBuilder, bIsDebugging);
+	Super::GetContextMenuActions(Menu, Context);
 }
 
 void USoundCueGraphSchema::CreateDefaultNodesForGraph(UEdGraph& Graph) const
diff --git a/Engine/Source/Editor/AudioEditor/Private/SoundSubmixGraphSchema.cpp b/Engine/Source/Editor/AudioEditor/Private/SoundSubmixGraphSchema.cpp
--- a/Engine/Source/Editor/AudioEditor/Private/SoundSubmixGraphSchema.cpp
+++ b/Engine/Source/Editor/AudioEditor/Private/SoundSubmixGraphSchema.cpp
@@ -10,6 +10,7 @@
 #include "Sound/SoundSubmix.h"
 #include "Framework/Commands/GenericCommands.h"
 #include "Framework/MultiBox/MultiBoxBuilder.h"
+#include "ToolMenus.h"
 
 #define LOCTEXT_NAMESPACE "SoundSubmixSchema"
 
@@ -32,8 +33,10 @@ bool USoundSubmixGraphSchema::ConnectionCausesLoop(const UEdGraphPin* InputPin,
 	return InputNode->SoundSubmix->RecurseCheckChild(OutputNode->SoundSubmix);
 }
 
-void USoundSubmixGraphSchema::GetBreakLinkToSubMenuActions(class FMenuBuilder& MenuBuilder, UEdGraphPin* InGraphPin)
+void USoundSubmixGraphSchema::GetBreakLinkToSubMenuActions(UToolMenu* Menu, const FName SectionName, UEdGraphPin* InGraphPin)
 {
+	FToolMenuSection& Section = Menu->FindOrAddSection(SectionName);
+
 	// Make sure we have a unique name for every entry in the list
 	TMap<FString, uint32> LinkTitleCount;
 
@@ -71,7 +74,7 @@ void USoundSubmixGraphSchema::GetBreakLinkToSubMenuActions(class FMenuBuilder& M
 		}
 		++Count;
 
-		MenuBuilder.AddMenuEntry( Description, Description, FSlateIcon(), FUIAction(
+		Section.AddMenuEntry(NAME_None, Description, Description, FSlateIcon(), FUIAction(
 			FExecuteAction::CreateUObject((USoundSubmixGraphSchema*const)this, &USoundSubmixGraphSchema::BreakSinglePinLink, const_cast< UEdGraphPin* >(InGraphPin), *Links)));
 	}
 }
@@ -86,43 +89,44 @@ void USoundSubmixGraphSchema::GetGraphContextActions(FGraphContextMenuBuilder& C
 	ContextMenuBuilder.AddAction(NewAction);
 }
 
-void USoundSubmixGraphSchema::GetContextMenuActions(const UEdGraph* CurrentGraph, const UEdGraphNode* InGraphNode, const UEdGraphPin* InGraphPin, class FMenuBuilder* MenuBuilder, bool bIsDebugging) const
+void USoundSubmixGraphSchema::GetContextMenuActions(UToolMenu* Menu, UGraphNodeContextMenuContext* Context) const
 {
-	if (InGraphPin)
+	if (Context->Pin)
 	{
-		MenuBuilder->BeginSection("SoundSubmixGraphSchemaPinActions", LOCTEXT("PinActionsMenuHeader", "Pin Actions"));
+		const UEdGraphPin* InGraphPin = Context->Pin;
 		{
+			const static FName SectionName = "SoundSubmixGraphSchemaPinActions";
+			FToolMenuSection& Section = Menu->AddSection(SectionName, LOCTEXT("PinActionsMenuHeader", "Pin Actions"));
 			// Only display the 'Break Links' option if there is a link to break!
 			if (InGraphPin->LinkedTo.Num() > 0)
 			{
-				MenuBuilder->AddMenuEntry(FGraphEditorCommands::Get().BreakPinLinks);
+				Section.AddMenuEntry(FGraphEditorCommands::Get().BreakPinLinks);
 
 				// add sub menu for break link to
 				if(InGraphPin->LinkedTo.Num() > 1)
 				{
-					MenuBuilder->AddSubMenu(
+					Section.AddEntry(FToolMenuEntry::InitSubMenu(
+						Menu->GetMenuName(),
+						"BreakLinkTo",
 						LOCTEXT("BreakLinkTo", "Break Link To..." ),
 						LOCTEXT("BreakSpecificLinks", "Break a specific link..." ),
-						FNewMenuDelegate::CreateUObject((USoundSubmixGraphSchema*const)this, &USoundSubmixGraphSchema::GetBreakLinkToSubMenuActions, const_cast<UEdGraphPin*>(InGraphPin)));
+						FNewToolMenuDelegate::CreateUObject((USoundSubmixGraphSchema*const)this, &USoundSubmixGraphSchema::GetBreakLinkToSubMenuActions, SectionName, const_cast<UEdGraphPin*>(InGraphPin))));
 				}
 				else
 				{
-					((USoundSubmixGraphSchema*const)this)->GetBreakLinkToSubMenuActions(*MenuBuilder, const_cast<UEdGraphPin*>(InGraphPin));
+					((USoundSubmixGraphSchema*const)this)->GetBreakLinkToSubMenuActions(Menu, SectionName, const_cast<UEdGraphPin*>(InGraphPin));
 				}
 			}
 		}
-		MenuBuilder->EndSection();
 	}
-	else if (InGraphNode)
+	else if (Context->Node)
 	{
-		const USoundSubmixGraphNode* SoundGraphNode = Cast<const USoundSubmixGraphNode>(InGraphNode);
-
-		MenuBuilder->BeginSection("SoundSubmixGraphSchemaNodeActions", LOCTEXT("ClassActionsMenuHeader", "SoundSubmix Actions"));
+		const USoundSubmixGraphNode* SoundGraphNode = Cast<const USoundSubmixGraphNode>(Context->Node);
 		{
-			MenuBuilder->AddMenuEntry(FGraphEditorCommands::Get().BreakNodeLinks);
-			MenuBuilder->AddMenuEntry(FGenericCommands::Get().Delete);
+			FToolMenuSection& Section = Menu->AddSection("SoundSubmixGraphSchemaNodeActions", LOCTEXT("ClassActionsMenuHeader", "SoundSubmix Actions"));
+			Section.AddMenuEntry(FGraphEditorCommands::Get().BreakNodeLinks);
+			Section.AddMenuEntry(FGenericCommands::Get().Delete);
 		}
-		MenuBuilder->EndSection();
 	}
 
 	// No Super call so Node comments option is not shown
diff --git a/Engine/Source/Editor/BehaviorTreeEditor/BehaviorTreeEditor.Build.cs b/Engine/Source/Editor/BehaviorTreeEditor/BehaviorTreeEditor.Build.cs
--- a/Engine/Source/Editor/BehaviorTreeEditor/BehaviorTreeEditor.Build.cs
+++ b/Engine/Source/Editor/BehaviorTreeEditor/BehaviorTreeEditor.Build.cs
@@ -45,7 +45,8 @@ public class BehaviorTreeEditor : ModuleRules
 				"BlueprintGraph",
                 "AIGraph",
                 "AIModule",
-				"ClassViewer"
+				"ClassViewer",
+				"ToolMenus",
 			}
 		);
 
diff --git a/Engine/Source/Editor/BehaviorTreeEditor/Classes/EdGraphSchema_BehaviorTree.h b/Engine/Source/Editor/BehaviorTreeEditor/Classes/EdGraphSchema_BehaviorTree.h
--- a/Engine/Source/Editor/BehaviorTreeEditor/Classes/EdGraphSchema_BehaviorTree.h
+++ b/Engine/Source/Editor/BehaviorTreeEditor/Classes/EdGraphSchema_BehaviorTree.h
@@ -54,7 +54,7 @@ class UEdGraphSchema_BehaviorTree : public UAIGraphSchema
 	//~ Begin EdGraphSchema Interface
 	virtual void CreateDefaultNodesForGraph(UEdGraph& Graph) const override;
 	virtual void GetGraphContextActions(FGraphContextMenuBuilder& ContextMenuBuilder) const override;
-	virtual void GetContextMenuActions(const UEdGraph* CurrentGraph, const UEdGraphNode* InGraphNode, const UEdGraphPin* InGraphPin, class FMenuBuilder* MenuBuilder, bool bIsDebugging) const override;
+	virtual void GetContextMenuActions(class UToolMenu* Menu, class UGraphNodeContextMenuContext* Context) const override;
 	virtual const FPinConnectionResponse CanCreateConnection(const UEdGraphPin* A, const UEdGraphPin* B) const override;
 	virtual const FPinConnectionResponse CanMergeNodes(const UEdGraphNode* A, const UEdGraphNode* B) const override;
 	virtual FLinearColor GetPinTypeColor(const FEdGraphPinType& PinType) const override;
diff --git a/Engine/Source/Editor/BehaviorTreeEditor/Classes/EdGraphSchema_BehaviorTreeDecorator.h b/Engine/Source/Editor/BehaviorTreeEditor/Classes/EdGraphSchema_BehaviorTreeDecorator.h
--- a/Engine/Source/Editor/BehaviorTreeEditor/Classes/EdGraphSchema_BehaviorTreeDecorator.h
+++ b/Engine/Source/Editor/BehaviorTreeEditor/Classes/EdGraphSchema_BehaviorTreeDecorator.h
@@ -59,7 +59,8 @@ class UEdGraphSchema_BehaviorTreeDecorator : public UEdGraphSchema
 	//~ Begin EdGraphSchema Interface
 	virtual void CreateDefaultNodesForGraph(UEdGraph& Graph) const override;
 	virtual void GetGraphContextActions(FGraphContextMenuBuilder& ContextMenuBuilder) const override;
-	virtual void GetContextMenuActions(const UEdGraph* CurrentGraph, const UEdGraphNode* InGraphNode, const UEdGraphPin* InGraphPin, class FMenuBuilder* MenuBuilder, bool bIsDebugging) const override;
+	virtual void GetContextMenuActions(class UToolMenu* Menu, class UGraphNodeContextMenuContext* Context) const override;
+	virtual FName GetParentContextMenuName() const override { return NAME_None; }
 	virtual const FPinConnectionResponse CanCreateConnection(const UEdGraphPin* A, const UEdGraphPin* B) const override;
 	virtual FLinearColor GetPinTypeColor(const FEdGraphPinType& PinType) const override;
 	virtual bool ShouldHidePinDefaultValue(UEdGraphPin* Pin) const override;
diff --git a/Engine/Source/Editor/BehaviorTreeEditor/Private/EdGraphSchema_BehaviorTree.cpp b/Engine/Source/Editor/BehaviorTreeEditor/Private/EdGraphSchema_BehaviorTree.cpp
--- a/Engine/Source/Editor/BehaviorTreeEditor/Private/EdGraphSchema_BehaviorTree.cpp
+++ b/Engine/Source/Editor/BehaviorTreeEditor/Private/EdGraphSchema_BehaviorTree.cpp
@@ -19,6 +19,7 @@
 #include "BehaviorTreeGraphNode_Service.h"
 #include "BehaviorTreeGraphNode_Task.h"
 #include "Framework/MultiBox/MultiBoxBuilder.h"
+#include "ToolMenus.h"
 #include "BehaviorTreeEditorModule.h"
 #include "IBehaviorTreeEditor.h"
 #include "BehaviorTreeDebugger.h"
@@ -209,26 +210,25 @@ void UEdGraphSchema_BehaviorTree::GetGraphContextActions(FGraphContextMenuBuilde
 	}
 }
 
-void UEdGraphSchema_BehaviorTree::GetContextMenuActions(const UEdGraph* CurrentGraph, const UEdGraphNode* InGraphNode, const UEdGraphPin* InGraphPin, class FMenuBuilder* MenuBuilder, bool bIsDebugging) const
+void UEdGraphSchema_BehaviorTree::GetContextMenuActions(UToolMenu* Menu, UGraphNodeContextMenuContext* Context) const
 {
-	if (InGraphNode && !InGraphPin)
+	if (Context->Node && !Context->Pin)
 	{
-		const UBehaviorTreeGraphNode* BTGraphNode = Cast<const UBehaviorTreeGraphNode>(InGraphNode);
+		const UBehaviorTreeGraphNode* BTGraphNode = Cast<const UBehaviorTreeGraphNode>(Context->Node);
 		if (BTGraphNode && BTGraphNode->CanPlaceBreakpoints())
 		{
-			MenuBuilder->BeginSection("EdGraphSchemaBreakpoints", LOCTEXT("BreakpointsHeader", "Breakpoints"));
+			FToolMenuSection& Section = Menu->AddSection("EdGraphSchemaBreakpoints", LOCTEXT("BreakpointsHeader", "Breakpoints"));
 			{
-				MenuBuilder->AddMenuEntry(FGraphEditorCommands::Get().ToggleBreakpoint);
-				MenuBuilder->AddMenuEntry(FGraphEditorCommands::Get().AddBreakpoint);
-				MenuBuilder->AddMenuEntry(FGraphEditorCommands::Get().RemoveBreakpoint);
-				MenuBuilder->AddMenuEntry(FGraphEditorCommands::Get().EnableBreakpoint);
-				MenuBuilder->AddMenuEntry(FGraphEditorCommands::Get().DisableBreakpoint);
+				Section.AddMenuEntry(FGraphEditorCommands::Get().ToggleBreakpoint);
+				Section.AddMenuEntry(FGraphEditorCommands::Get().AddBreakpoint);
+				Section.AddMenuEntry(FGraphEditorCommands::Get().RemoveBreakpoint);
+				Section.AddMenuEntry(FGraphEditorCommands::Get().EnableBreakpoint);
+				Section.AddMenuEntry(FGraphEditorCommands::Get().DisableBreakpoint);
 			}
-			MenuBuilder->EndSection();
 		}
 	}
 
-	Super::GetContextMenuActions(CurrentGraph, InGraphNode, InGraphPin, MenuBuilder, bIsDebugging);
+	Super::GetContextMenuActions(Menu, Context);
 }
 
 const FPinConnectionResponse UEdGraphSchema_BehaviorTree::CanCreateConnection(const UEdGraphPin* PinA, const UEdGraphPin* PinB) const
diff --git a/Engine/Source/Editor/BehaviorTreeEditor/Private/EdGraphSchema_BehaviorTreeDecorator.cpp b/Engine/Source/Editor/BehaviorTreeEditor/Private/EdGraphSchema_BehaviorTreeDecorator.cpp
--- a/Engine/Source/Editor/BehaviorTreeEditor/Private/EdGraphSchema_BehaviorTreeDecorator.cpp
+++ b/Engine/Source/Editor/BehaviorTreeEditor/Private/EdGraphSchema_BehaviorTreeDecorator.cpp
@@ -10,6 +10,7 @@
 #include "BehaviorTreeDecoratorGraphNode_Decorator.h"
 #include "BehaviorTreeDecoratorGraphNode_Logic.h"
 #include "Framework/MultiBox/MultiBoxBuilder.h"
+#include "ToolMenus.h"
 #include "BehaviorTreeEditorModule.h"
 #include "GraphEditorSettings.h"
 #include "GraphEditorActions.h"
@@ -158,33 +159,31 @@ void UEdGraphSchema_BehaviorTreeDecorator::GetGraphContextActions(FGraphContextM
 #endif
 }
 
-void UEdGraphSchema_BehaviorTreeDecorator::GetContextMenuActions(const UEdGraph* CurrentGraph, const UEdGraphNode* InGraphNode, const UEdGraphPin* InGraphPin, class FMenuBuilder* MenuBuilder, bool bIsDebugging) const
+void UEdGraphSchema_BehaviorTreeDecorator::GetContextMenuActions(UToolMenu* Menu, UGraphNodeContextMenuContext* Context) const
 {
-	const UBehaviorTreeDecoratorGraphNode_Logic* LogicNode = Cast<const UBehaviorTreeDecoratorGraphNode_Logic>(InGraphNode);
-	if (CurrentGraph && !CurrentGraph->bEditable)
+	const UBehaviorTreeDecoratorGraphNode_Logic* LogicNode = Cast<const UBehaviorTreeDecoratorGraphNode_Logic>(Context->Node);
+	if (Context->Graph && !Context->Graph->bEditable)
 	{
 		return;
 	}
 
-	if (InGraphPin)
+	if (Context->Pin)
 	{
 		// Only display the 'Break Links' option if there is a link to break!
-		if (InGraphPin->LinkedTo.Num() > 0)
+		if (Context->Pin->LinkedTo.Num() > 0)
 		{
-			MenuBuilder->BeginSection("DecoratorGraphSchemaPinActions", LOCTEXT("PinActionsMenuHeader", "Pin Actions"));
-			MenuBuilder->AddMenuEntry( FGraphEditorCommands::Get().BreakPinLinks );
-			MenuBuilder->EndSection();
+			FToolMenuSection& Section = Menu->AddSection("DecoratorGraphSchemaPinActions", LOCTEXT("PinActionsMenuHeader", "Pin Actions"));
+			Section.AddMenuEntry(FGraphEditorCommands::Get().BreakPinLinks);
 		}
 	}
-	else if (InGraphNode)
+	else if (Context->Node)
 	{
-		MenuBuilder->BeginSection("DecoratorGraphSchemaNodeActions", LOCTEXT("ClassActionsMenuHeader", "Node Actions"));
 		{
-			MenuBuilder->AddMenuEntry(FGraphEditorCommands::Get().AddExecutionPin);
-			MenuBuilder->AddMenuEntry(FGraphEditorCommands::Get().BreakNodeLinks);
-			MenuBuilder->AddMenuEntry(FGenericCommands::Get().Delete);
+			FToolMenuSection& Section = Menu->AddSection("DecoratorGraphSchemaNodeActions", LOCTEXT("ClassActionsMenuHeader", "Node Actions"));
+			Section.AddMenuEntry(FGraphEditorCommands::Get().AddExecutionPin);
+			Section.AddMenuEntry(FGraphEditorCommands::Get().BreakNodeLinks);
+			Section.AddMenuEntry(FGenericCommands::Get().Delete);
 		}
-		MenuBuilder->EndSection();
 	}
 }
 
diff --git a/Engine/Source/Editor/BlueprintGraph/BlueprintGraph.Build.cs b/Engine/Source/Editor/BlueprintGraph/BlueprintGraph.Build.cs
--- a/Engine/Source/Editor/BlueprintGraph/BlueprintGraph.Build.cs
+++ b/Engine/Source/Editor/BlueprintGraph/BlueprintGraph.Build.cs
@@ -36,6 +36,7 @@ public class BlueprintGraph : ModuleRules
                 "Kismet",
                 "KismetWidgets",
                 "PropertyEditor",
+				"ToolMenus",
 			}
 		);
 
diff --git a/Engine/Source/Editor/BlueprintGraph/Classes/EdGraphSchema_K2.h b/Engine/Source/Editor/BlueprintGraph/Classes/EdGraphSchema_K2.h
--- a/Engine/Source/Editor/BlueprintGraph/Classes/EdGraphSchema_K2.h
+++ b/Engine/Source/Editor/BlueprintGraph/Classes/EdGraphSchema_K2.h
@@ -16,6 +16,8 @@
 
 class AActor;
 class FMenuBuilder;
+class UToolMenu;
+struct FToolMenuSection;
 class UBlueprint;
 class UK2Node;
 struct FTypesDatabase;
@@ -452,7 +454,7 @@ public:
 	void SelectAllNodesInDirection(TEnumAsByte<enum EEdGraphPinDirection> InDirection, UEdGraph* Graph, UEdGraphPin* InGraphPin);
 
 	//~ Begin EdGraphSchema Interface
-	virtual void GetContextMenuActions(const UEdGraph* CurrentGraph, const UEdGraphNode* InGraphNode, const UEdGraphPin* InGraphPin, class FMenuBuilder* MenuBuilder, bool bIsDebugging) const override;
+	virtual void GetContextMenuActions(class UToolMenu* Menu, class UGraphNodeContextMenuContext* Context) const override;
 	virtual const FPinConnectionResponse CanCreateConnection(const UEdGraphPin* A, const UEdGraphPin* B) const override;
 	virtual bool TryCreateConnection(UEdGraphPin* A, UEdGraphPin* B) const override;
 	virtual bool CreateAutomaticConversionNodeAndConnections(UEdGraphPin* A, UEdGraphPin* B) const override;
@@ -771,7 +773,7 @@ public:
 	 *	@param	Reference to graph node
 	 *	@param	Reference to context menu builder
 	 */
-	void AddSelectedReplaceableNodes( UBlueprint* Blueprint, const UEdGraphNode* InGraphNode, FMenuBuilder* MenuBuilder ) const;
+	void AddSelectedReplaceableNodes(FToolMenuSection& Section, UBlueprint* Blueprint, const UEdGraphNode* InGraphNode) const;
 
 	/**
 	 *	Function to replace current graph node reference object with a new object
@@ -1004,19 +1006,19 @@ public:
 	virtual bool FindSpecializedConversionNode(const UEdGraphPin* OutputPin, const UEdGraphPin* InputPin, bool bCreateNode, /*out*/ class UK2Node*& TargetNode) const;
 
 	/** Get menu for breaking links to specific nodes*/
-	void GetBreakLinkToSubMenuActions(class FMenuBuilder& MenuBuilder, class UEdGraphPin* InGraphPin);
+	void GetBreakLinkToSubMenuActions(UToolMenu* Menu, class UEdGraphPin* InGraphPin);
 
 	/** Get menu for jumping to specific pin links */
-	void GetJumpToConnectionSubMenuActions(class FMenuBuilder& MenuBuilder, class UEdGraphPin* InGraphPin);
+	void GetJumpToConnectionSubMenuActions(UToolMenu* Menu, class UEdGraphPin* InGraphPin);
 
 	/** Get menu for straightening links to specific nodes*/
-	void GetStraightenConnectionToSubMenuActions( class FMenuBuilder& MenuBuilder, UEdGraphPin* InGraphPin ) const;
+	void GetStraightenConnectionToSubMenuActions(UToolMenu* Menu, UEdGraphPin* InGraphPin) const;
 
 	/** Get the destination pin for a straighten operation */
 	static UEdGraphPin* GetAndResetStraightenDestinationPin();
 
 	/** Create menu for variable get/set nodes which refer to a variable which does not exist. */
-	void GetNonExistentVariableMenu(const UEdGraphNode* InGraphNode, UBlueprint* OwnerBlueprint, FMenuBuilder* MenuBuilder) const;
+	void GetNonExistentVariableMenu(FToolMenuSection& Section, const UEdGraphNode* InGraphNode, UBlueprint* OwnerBlueprint) const;
 
 	/**
 	 * Create menu for variable get/set nodes which allows for the replacement of variables
@@ -1026,7 +1028,7 @@ public:
 	 * @param InMenuBuilder					MenuBuilder to place the menu items into
 	 * @param bInReplaceExistingVariable	TRUE if replacing an existing variable, will keep the variable from appearing on the list
 	 */
-	void GetReplaceVariableMenu(const UEdGraphNode* InGraphNode, UBlueprint* InOwnerBlueprint, FMenuBuilder* InMenuBuilder, bool bInReplaceExistingVariable = false) const;
+	void GetReplaceVariableMenu(FToolMenuSection& Section, const UEdGraphNode* InGraphNode, UBlueprint* InOwnerBlueprint, bool bInReplaceExistingVariable = false) const;
 
 	// Calculates an average position between the nodes owning the two specified pins
 	static FVector2D CalculateAveragePositionBetweenNodes(UEdGraphPin* InputPin, UEdGraphPin* OutputPin);
@@ -1105,7 +1107,7 @@ public:
 	static void OnReplaceVariableForVariableNode(class UK2Node_Variable* Variable, UBlueprint* OwnerBlueprint, FName VariableName, bool bIsSelfMember);
 
 	/** Create sub menu that shows all possible variables that can be used to replace the existing variable reference */
-	static void GetReplaceVariableMenu(class FMenuBuilder& MenuBuilder, class UK2Node_Variable* Variable, UBlueprint* OwnerBlueprint, bool bReplaceExistingVariable = false);
+	static void GetReplaceVariableMenu(UToolMenu* Menu, class UK2Node_Variable* Variable, UBlueprint* OwnerBlueprint, bool bReplaceExistingVariable = false);
 
 	/** Function called when the owning module is shut down */ 
 	static void Shutdown();
diff --git a/Engine/Source/Editor/BlueprintGraph/Private/EdGraphSchema_K2.cpp b/Engine/Source/Editor/BlueprintGraph/Private/EdGraphSchema_K2.cpp
--- a/Engine/Source/Editor/BlueprintGraph/Private/EdGraphSchema_K2.cpp
+++ b/Engine/Source/Editor/BlueprintGraph/Private/EdGraphSchema_K2.cpp
@@ -29,6 +29,7 @@
 #include "Framework/Commands/UIAction.h"
 #include "Framework/Commands/UICommandList.h"
 #include "Framework/MultiBox/MultiBoxBuilder.h"
+#include "ToolMenus.h"
 #include "GraphEditorSettings.h"
 #include "K2Node.h"
 #include "EdGraphSchema_K2_Actions.h"
@@ -1084,7 +1085,7 @@ void UEdGraphSchema_K2::ReplaceSelectedNode(UEdGraphNode* SourceNode, AActor* Ta
 	}
 }
 
-void UEdGraphSchema_K2::AddSelectedReplaceableNodes( UBlueprint* Blueprint, const UEdGraphNode* InGraphNode, FMenuBuilder* MenuBuilder ) const
+void UEdGraphSchema_K2::AddSelectedReplaceableNodes(FToolMenuSection& Section, UBlueprint* Blueprint, const UEdGraphNode* InGraphNode) const
 {
 	//Only allow replace object reference functionality for literal nodes
 	const UK2Node_Literal* LiteralNode = Cast<UK2Node_Literal>(InGraphNode);
@@ -1099,7 +1100,7 @@ void UEdGraphSchema_K2::AddSelectedReplaceableNodes( UBlueprint* Blueprint, cons
 			{
 				FText Description = FText::Format( LOCTEXT("ChangeToActorName", "Change to <{0}>"), FText::FromString( Actor->GetActorLabel() ) );
 				FText ToolTip = LOCTEXT("ReplaceNodeReferenceToolTip", "Replace node reference");
-				MenuBuilder->AddMenuEntry( Description, ToolTip, FSlateIcon(), FUIAction(
+				Section.AddMenuEntry(NAME_None, Description, ToolTip, FSlateIcon(), FUIAction(
 					FExecuteAction::CreateUObject((UEdGraphSchema_K2*const)this, &UEdGraphSchema_K2::ReplaceSelectedNode, const_cast< UEdGraphNode* >(InGraphNode), Actor) ) );
 			}
 		}
@@ -1410,33 +1411,38 @@ void UEdGraphSchema_K2::SelectAllNodesInDirection(TEnumAsByte<enum EEdGraphPinDi
 	}
 }
 
-void UEdGraphSchema_K2::GetContextMenuActions(const UEdGraph* CurrentGraph, const UEdGraphNode* InGraphNode, const UEdGraphPin* InGraphPin, FMenuBuilder* MenuBuilder, bool bIsDebugging) const
+void UEdGraphSchema_K2::GetContextMenuActions(UToolMenu* Menu, UGraphNodeContextMenuContext* Context) const
 {
+	const UEdGraph* CurrentGraph = Context->Graph;
+	const UEdGraphNode* InGraphNode = Context->Node;
+	const UEdGraphPin* InGraphPin = Context->Pin;
+	const bool bIsDebugging = Context->bIsDebugging;
 	check(CurrentGraph);
 	UBlueprint* OwnerBlueprint = FBlueprintEditorUtils::FindBlueprintForGraphChecked(CurrentGraph);
 
 	if (InGraphPin != NULL)
 	{
-		MenuBuilder->BeginSection("EdGraphSchemaPinActions", LOCTEXT("PinActionsMenuHeader", "Pin Actions"));
 		{
+			FToolMenuSection& Section = Menu->AddSection("EdGraphSchemaPinActions", LOCTEXT("PinActionsMenuHeader", "Pin Actions"));
 			if (!bIsDebugging)
 			{
 				// Break pin links
 				if (InGraphPin->LinkedTo.Num() > 1)
 				{
-					MenuBuilder->AddMenuEntry( FGraphEditorCommands::Get().BreakPinLinks );
+					Section.AddMenuEntry(FGraphEditorCommands::Get().BreakPinLinks);
 				}
 	
 				// Add the change pin type action, if this is a select node
 				if (InGraphNode->IsA(UK2Node_Select::StaticClass()))
 				{
-					MenuBuilder->AddMenuEntry(FGraphEditorCommands::Get().ChangePinType);
+					Section.AddMenuEntry(FGraphEditorCommands::Get().ChangePinType);
 				}
 	
 				// add sub menu for break link to
 				if (InGraphPin->LinkedTo.Num() > 0)
 				{
-					MenuBuilder->AddMenuEntry(
+					Section.AddMenuEntry(
+						"SelectAllInputNodes",
 						InGraphPin->Direction == EEdGraphPinDirection::EGPD_Input ? LOCTEXT("SelectAllInputNodes", "Select All Input Nodes") : LOCTEXT("SelectAllOutputNodes", "Select All Output Nodes"),
 						InGraphPin->Direction == EEdGraphPinDirection::EGPD_Input ? LOCTEXT("SelectAllInputNodesTooltip", "Adds all input Nodes linked to this Pin to selection") : LOCTEXT("SelectAllOutputNodesTooltip", "Adds all output Nodes linked to this Pin to selection"),
 						FSlateIcon(),
@@ -1444,25 +1450,31 @@ void UEdGraphSchema_K2::GetContextMenuActions(const UEdGraph* CurrentGraph, cons
 
 					if(InGraphPin->LinkedTo.Num() > 1)
 					{
-						MenuBuilder->AddSubMenu(
+						Section.AddEntry(FToolMenuEntry::InitSubMenu(
+							Menu->GetMenuName(),
+							"BreakLinkTo",
 							LOCTEXT("BreakLinkTo", "Break Link To..."),
 							LOCTEXT("BreakSpecificLinks", "Break a specific link..."),
-							FNewMenuDelegate::CreateUObject( (UEdGraphSchema_K2*const)this, &UEdGraphSchema_K2::GetBreakLinkToSubMenuActions, const_cast<UEdGraphPin*>(InGraphPin)));
+							FNewToolMenuDelegate::CreateUObject( (UEdGraphSchema_K2*const)this, &UEdGraphSchema_K2::GetBreakLinkToSubMenuActions, const_cast<UEdGraphPin*>(InGraphPin))));
 
-						MenuBuilder->AddSubMenu(
+						Section.AddEntry(FToolMenuEntry::InitSubMenu(
+							Menu->GetMenuName(),
+							"JumpToConnection",
 							LOCTEXT("JumpToConnection", "Jump to Connection..."),
 							LOCTEXT("JumpToSpecificConnection", "Jump to specific connection..."),
-							FNewMenuDelegate::CreateUObject( (UEdGraphSchema_K2*const)this, &UEdGraphSchema_K2::GetJumpToConnectionSubMenuActions, const_cast<UEdGraphPin*>(InGraphPin)));
+							FNewToolMenuDelegate::CreateUObject( (UEdGraphSchema_K2*const)this, &UEdGraphSchema_K2::GetJumpToConnectionSubMenuActions, const_cast<UEdGraphPin*>(InGraphPin))));
 
-						MenuBuilder->AddSubMenu(
+						Section.AddEntry(FToolMenuEntry::InitSubMenu(
+							Menu->GetMenuName(),
+							"StraightenConnection",
 							LOCTEXT("StraightenConnection", "Straighten Connection To..."),
 							LOCTEXT("StraightenConnection_Tip", "Straighten a specific connection"),
-							FNewMenuDelegate::CreateUObject( this, &UEdGraphSchema_K2::GetStraightenConnectionToSubMenuActions, const_cast<UEdGraphPin*>(InGraphPin)));
+							FNewToolMenuDelegate::CreateUObject( this, &UEdGraphSchema_K2::GetStraightenConnectionToSubMenuActions, const_cast<UEdGraphPin*>(InGraphPin))));
 					}
 					else
 					{
-						((UEdGraphSchema_K2*const)this)->GetBreakLinkToSubMenuActions(*MenuBuilder, const_cast<UEdGraphPin*>(InGraphPin));
-						((UEdGraphSchema_K2*const)this)->GetJumpToConnectionSubMenuActions(*MenuBuilder, const_cast<UEdGraphPin*>(InGraphPin));
+						((UEdGraphSchema_K2*const)this)->GetBreakLinkToSubMenuActions(Menu, const_cast<UEdGraphPin*>(InGraphPin));
+						((UEdGraphSchema_K2*const)this)->GetJumpToConnectionSubMenuActions(Menu, const_cast<UEdGraphPin*>(InGraphPin));
 						
 						UEdGraphPin* Pin = InGraphPin->LinkedTo[0];
 						FText PinName = Pin->GetDisplayName();
@@ -1481,9 +1493,8 @@ void UEdGraphSchema_K2::GetContextMenuActions(const UEdGraph* CurrentGraph, cons
 							EntryTooltip = FText::Format(LOCTEXT("StraightenDescription_SinglePin_Node_Tip", "Straighten the connection between this pin, and {0} ({1})"), NodeName, PinName);
 						}
 
-						MenuBuilder->AddMenuEntry(
+						Section.AddMenuEntry(
 							FGraphEditorCommands::Get().StraightenConnections,
-							NAME_None,
 							EntryTitle,
 							EntryTooltip,
 							FSlateIcon(NAME_None, NAME_None, NAME_None)
@@ -1494,11 +1505,11 @@ void UEdGraphSchema_K2::GetContextMenuActions(const UEdGraph* CurrentGraph, cons
 				// Conditionally add the var promotion pin if this is an output pin and it's not an exec pin
 				if (InGraphPin->PinType.PinCategory != PC_Exec)
 				{
-					MenuBuilder->AddMenuEntry( FGraphEditorCommands::Get().PromoteToVariable );
+					Section.AddMenuEntry( FGraphEditorCommands::Get().PromoteToVariable );
 
 					if (FBlueprintEditorUtils::DoesSupportLocalVariables(CurrentGraph))
 					{
-						MenuBuilder->AddMenuEntry( FGraphEditorCommands::Get().PromoteToLocalVariable );
+						Section.AddMenuEntry( FGraphEditorCommands::Get().PromoteToLocalVariable );
 					}
 				}
 	
@@ -1514,12 +1525,12 @@ void UEdGraphSchema_K2::GetContextMenuActions(const UEdGraph* CurrentGraph, cons
 					{
 						Tooltip = LOCTEXT("SplitStructPin_Error", "Cannot split the struct pin, it may be missing Blueprint exposed properties!");
 					}
-					MenuBuilder->AddMenuEntry( FGraphEditorCommands::Get().SplitStructPin, NAME_None, FGraphEditorCommands::Get().SplitStructPin->GetLabel(), Tooltip );
+					Section.AddMenuEntry( FGraphEditorCommands::Get().SplitStructPin, FGraphEditorCommands::Get().SplitStructPin->GetLabel(), Tooltip );
 				}
 
 				if (InGraphPin->ParentPin != NULL)
 				{
-					MenuBuilder->AddMenuEntry( FGraphEditorCommands::Get().RecombineStructPin );
+					Section.AddMenuEntry( FGraphEditorCommands::Get().RecombineStructPin );
 				}
 	
 				// Conditionally add the execution path pin options if this is an execution branching node
@@ -1527,167 +1538,162 @@ void UEdGraphSchema_K2::GetContextMenuActions(const UEdGraph* CurrentGraph, cons
 				{
 					if (CastChecked<UK2Node>(InGraphPin->GetOwningNode())->CanEverInsertExecutionPin())
 					{
-						MenuBuilder->AddMenuEntry(FGraphEditorCommands::Get().InsertExecutionPinBefore);
-						MenuBuilder->AddMenuEntry(FGraphEditorCommands::Get().InsertExecutionPinAfter);
+						Section.AddMenuEntry(FGraphEditorCommands::Get().InsertExecutionPinBefore);
+						Section.AddMenuEntry(FGraphEditorCommands::Get().InsertExecutionPinAfter);
 					}
 
 					if (CastChecked<UK2Node>(InGraphPin->GetOwningNode())->CanEverRemoveExecutionPin())
 					{
-						MenuBuilder->AddMenuEntry( FGraphEditorCommands::Get().RemoveExecutionPin );
+						Section.AddMenuEntry( FGraphEditorCommands::Get().RemoveExecutionPin );
 					}
 				}
 
 				if (UK2Node_SetFieldsInStruct::ShowCustomPinActions(InGraphPin, true))
 				{
-					MenuBuilder->AddMenuEntry(FGraphEditorCommands::Get().RemoveThisStructVarPin);
-					MenuBuilder->AddMenuEntry(FGraphEditorCommands::Get().RemoveOtherStructVarPins);
+					Section.AddMenuEntry(FGraphEditorCommands::Get().RemoveThisStructVarPin);
+					Section.AddMenuEntry(FGraphEditorCommands::Get().RemoveOtherStructVarPins);
 				}
 
 				if (InGraphPin->PinType.PinCategory != PC_Exec && InGraphPin->Direction == EGPD_Input && InGraphPin->LinkedTo.Num() == 0 && !ShouldHidePinDefaultValue(const_cast<UEdGraphPin*>(InGraphPin)))
 				{
-					MenuBuilder->AddMenuEntry(FGraphEditorCommands::Get().ResetPinToDefaultValue);
+					Section.AddMenuEntry(FGraphEditorCommands::Get().ResetPinToDefaultValue);
 				}
 			}
 		}
-		MenuBuilder->EndSection();
 
 		// Add the watch pin / unwatch pin menu items
-		MenuBuilder->BeginSection("EdGraphSchemaWatches", LOCTEXT("WatchesHeader", "Watches"));
 		{
+			FToolMenuSection& Section = Menu->AddSection("EdGraphSchemaWatches", LOCTEXT("WatchesHeader", "Watches"));
 			if (!IsMetaPin(*InGraphPin))
 			{
 				const UEdGraphPin* WatchedPin = ((InGraphPin->Direction == EGPD_Input) && (InGraphPin->LinkedTo.Num() > 0)) ? InGraphPin->LinkedTo[0] : InGraphPin;
 				if (FKismetDebugUtilities::IsPinBeingWatched(OwnerBlueprint, WatchedPin))
 				{
-					MenuBuilder->AddMenuEntry( FGraphEditorCommands::Get().StopWatchingPin );
+					Section.AddMenuEntry( FGraphEditorCommands::Get().StopWatchingPin );
 				}
 				else
 				{
-					MenuBuilder->AddMenuEntry( FGraphEditorCommands::Get().StartWatchingPin );
+					Section.AddMenuEntry( FGraphEditorCommands::Get().StartWatchingPin );
 				}
 			}
 		}
-		MenuBuilder->EndSection();
 	}
 	else if (InGraphNode != NULL)
 	{
 		if (IsUsingNonExistantVariable(InGraphNode, OwnerBlueprint))
 		{
-			MenuBuilder->BeginSection("EdGraphSchemaNodeActions", LOCTEXT("NodeActionsMenuHeader", "Node Actions"));
 			{
-				GetNonExistentVariableMenu(InGraphNode,OwnerBlueprint, MenuBuilder);
+				FToolMenuSection& Section = Menu->AddSection("EdGraphSchemaNodeActions", LOCTEXT("NodeActionsMenuHeader", "Node Actions"));
+				GetNonExistentVariableMenu(Section, InGraphNode, OwnerBlueprint);
 			}
-			MenuBuilder->EndSection();
 		}
 		else
 		{
-			MenuBuilder->BeginSection("EdGraphSchemaNodeActions", LOCTEXT("NodeActionsMenuHeader", "Node Actions"));
 			{
+				FToolMenuSection& Section = Menu->AddSection("EdGraphSchemaNodeActions", LOCTEXT("NodeActionsMenuHeader", "Node Actions"));
 				if (!bIsDebugging)
 				{
 					// Replaceable node display option
-					AddSelectedReplaceableNodes( OwnerBlueprint, InGraphNode, MenuBuilder );
+					AddSelectedReplaceableNodes(Section, OwnerBlueprint, InGraphNode);
 
 					// Node contextual actions
-					MenuBuilder->AddMenuEntry( FGenericCommands::Get().Delete );
-					MenuBuilder->AddMenuEntry( FGenericCommands::Get().Cut );
-					MenuBuilder->AddMenuEntry( FGenericCommands::Get().Copy );
-					MenuBuilder->AddMenuEntry( FGenericCommands::Get().Duplicate );
-					MenuBuilder->AddMenuEntry( FGraphEditorCommands::Get().ReconstructNodes );
-					MenuBuilder->AddMenuEntry( FGraphEditorCommands::Get().BreakNodeLinks );
+					Section.AddMenuEntry( FGenericCommands::Get().Delete );
+					Section.AddMenuEntry( FGenericCommands::Get().Cut );
+					Section.AddMenuEntry( FGenericCommands::Get().Copy );
+					Section.AddMenuEntry( FGenericCommands::Get().Duplicate );
+					Section.AddMenuEntry( FGraphEditorCommands::Get().ReconstructNodes );
+					Section.AddMenuEntry( FGraphEditorCommands::Get().BreakNodeLinks );
 
 					// Conditionally add the action to add an execution pin, if this is an execution node
 					if( InGraphNode->IsA(UK2Node_ExecutionSequence::StaticClass()) || InGraphNode->IsA(UK2Node_Switch::StaticClass()) )
 					{
-						MenuBuilder->AddMenuEntry( FGraphEditorCommands::Get().AddExecutionPin );
+						Section.AddMenuEntry( FGraphEditorCommands::Get().AddExecutionPin );
 					}
 
 					// Conditionally add the action to create a super function call node, if this is an event or function entry
 					if( InGraphNode->IsA(UK2Node_Event::StaticClass()) || InGraphNode->IsA(UK2Node_FunctionEntry::StaticClass()) )
 					{
-						MenuBuilder->AddMenuEntry( FGraphEditorCommands::Get().AddParentNode );
+						Section.AddMenuEntry( FGraphEditorCommands::Get().AddParentNode );
 					}
 
 					// Conditionally add the actions to add or remove an option pin, if this is a select node
 					if (InGraphNode->IsA(UK2Node_Select::StaticClass()))
 					{
-						MenuBuilder->AddMenuEntry(FGraphEditorCommands::Get().AddOptionPin);
-						MenuBuilder->AddMenuEntry(FGraphEditorCommands::Get().RemoveOptionPin);
+						Section.AddMenuEntry(FGraphEditorCommands::Get().AddOptionPin);
+						Section.AddMenuEntry(FGraphEditorCommands::Get().RemoveOptionPin);
 					}
 
 					// Don't show the "Assign selected Actor" option if more than one actor is selected
 					if (InGraphNode->IsA(UK2Node_ActorBoundEvent::StaticClass()) && GEditor->GetSelectedActorCount() == 1)
 					{
-						MenuBuilder->AddMenuEntry(FGraphEditorCommands::Get().AssignReferencedActor);
+						Section.AddMenuEntry(FGraphEditorCommands::Get().AssignReferencedActor);
 					}
 				}
 
 				// If the node has an associated definition (for some loose sense of the word), allow going to it (same action as double-clicking on a node)
 				if (InGraphNode->CanJumpToDefinition())
 				{
-					MenuBuilder->AddMenuEntry(FGraphEditorCommands::Get().GoToDefinition);
+					Section.AddMenuEntry(FGraphEditorCommands::Get().GoToDefinition);
 				}
 
 				// show search for references for everyone
-				MenuBuilder->AddMenuEntry(FGraphEditorCommands::Get().FindReferences);
+				Section.AddMenuEntry(FGraphEditorCommands::Get().FindReferences);
 
 				if (!bIsDebugging)
 				{
 					if (InGraphNode->IsA(UK2Node_Variable::StaticClass()))
 					{
-						GetReplaceVariableMenu(InGraphNode, OwnerBlueprint, MenuBuilder, true);
+						GetReplaceVariableMenu(Section, InGraphNode, OwnerBlueprint, true);
 					}
 
 					if (InGraphNode->IsA(UK2Node_SetFieldsInStruct::StaticClass()))
 					{
-						MenuBuilder->AddMenuEntry(FGraphEditorCommands::Get().RestoreAllStructVarPins);
+						Section.AddMenuEntry(FGraphEditorCommands::Get().RestoreAllStructVarPins);
 					}
 
-					MenuBuilder->AddMenuEntry(FGenericCommands::Get().Rename, NAME_None, LOCTEXT("Rename", "Rename"), LOCTEXT("Rename_Tooltip", "Renames selected function or variable in blueprint.") );
+					Section.AddMenuEntry(FGenericCommands::Get().Rename, LOCTEXT("Rename", "Rename"), LOCTEXT("Rename_Tooltip", "Renames selected function or variable in blueprint.") );
 				}
 
 				// Select referenced actors in the level
-				MenuBuilder->AddMenuEntry( FGraphEditorCommands::Get().SelectReferenceInLevel );
+				Section.AddMenuEntry(FGraphEditorCommands::Get().SelectReferenceInLevel);
 			}
-			MenuBuilder->EndSection(); //EdGraphSchemaNodeActions
 
 			if (!bIsDebugging)
 			{
 				// Collapse/expand nodes
-				MenuBuilder->BeginSection("EdGraphSchemaOrganization", LOCTEXT("OrganizationHeader", "Organization"));
 				{
-					MenuBuilder->AddMenuEntry( FGraphEditorCommands::Get().CollapseNodes );
-					MenuBuilder->AddMenuEntry( FGraphEditorCommands::Get().CollapseSelectionToFunction );
-					MenuBuilder->AddMenuEntry( FGraphEditorCommands::Get().CollapseSelectionToMacro );
-					MenuBuilder->AddMenuEntry( FGraphEditorCommands::Get().ExpandNodes );
+					FToolMenuSection& Section = Menu->AddSection("EdGraphSchemaOrganization", LOCTEXT("OrganizationHeader", "Organization"));
+					Section.AddMenuEntry( FGraphEditorCommands::Get().CollapseNodes );
+					Section.AddMenuEntry( FGraphEditorCommands::Get().CollapseSelectionToFunction );
+					Section.AddMenuEntry( FGraphEditorCommands::Get().CollapseSelectionToMacro );
+					Section.AddMenuEntry( FGraphEditorCommands::Get().ExpandNodes );
 
 					if(InGraphNode->IsA(UK2Node_Composite::StaticClass()))
 					{
-						MenuBuilder->AddMenuEntry( FGraphEditorCommands::Get().PromoteSelectionToFunction );
-						MenuBuilder->AddMenuEntry( FGraphEditorCommands::Get().PromoteSelectionToMacro );
+						Section.AddMenuEntry( FGraphEditorCommands::Get().PromoteSelectionToFunction );
+						Section.AddMenuEntry( FGraphEditorCommands::Get().PromoteSelectionToMacro );
 					}
 
-					MenuBuilder->AddSubMenu(LOCTEXT("AlignmentHeader", "Alignment"), FText(), FNewMenuDelegate::CreateLambda([](FMenuBuilder& InMenuBuilder){
-						
-						InMenuBuilder.BeginSection("EdGraphSchemaAlignment", LOCTEXT("AlignHeader", "Align"));
-						InMenuBuilder.AddMenuEntry( FGraphEditorCommands::Get().AlignNodesTop );
-						InMenuBuilder.AddMenuEntry( FGraphEditorCommands::Get().AlignNodesMiddle );
-						InMenuBuilder.AddMenuEntry( FGraphEditorCommands::Get().AlignNodesBottom );
-						InMenuBuilder.AddMenuEntry( FGraphEditorCommands::Get().AlignNodesLeft );
-						InMenuBuilder.AddMenuEntry( FGraphEditorCommands::Get().AlignNodesCenter );
-						InMenuBuilder.AddMenuEntry( FGraphEditorCommands::Get().AlignNodesRight );
-						InMenuBuilder.AddMenuEntry( FGraphEditorCommands::Get().StraightenConnections );
-						InMenuBuilder.EndSection();
-
-						InMenuBuilder.BeginSection("EdGraphSchemaDistribution", LOCTEXT("DistributionHeader", "Distribution"));
-						InMenuBuilder.AddMenuEntry( FGraphEditorCommands::Get().DistributeNodesHorizontally );
-						InMenuBuilder.AddMenuEntry( FGraphEditorCommands::Get().DistributeNodesVertically );
-						InMenuBuilder.EndSection();
-						
-					}));
+					Section.AddEntry(FToolMenuEntry::InitSubMenu(Menu->GetMenuName(), "Alignment", LOCTEXT("AlignmentHeader", "Alignment"), FText(), FNewToolMenuDelegate::CreateLambda([](UToolMenu* AlignmentMenu)
+					{
+						{
+							FToolMenuSection& InSection = AlignmentMenu->AddSection("EdGraphSchemaAlignment", LOCTEXT("AlignHeader", "Align"));
+							InSection.AddMenuEntry(FGraphEditorCommands::Get().AlignNodesTop);
+							InSection.AddMenuEntry(FGraphEditorCommands::Get().AlignNodesMiddle);
+							InSection.AddMenuEntry(FGraphEditorCommands::Get().AlignNodesBottom);
+							InSection.AddMenuEntry(FGraphEditorCommands::Get().AlignNodesLeft);
+							InSection.AddMenuEntry(FGraphEditorCommands::Get().AlignNodesCenter);
+							InSection.AddMenuEntry(FGraphEditorCommands::Get().AlignNodesRight);
+							InSection.AddMenuEntry(FGraphEditorCommands::Get().StraightenConnections);
+						}
+
+						{
+							FToolMenuSection& InSection = AlignmentMenu->AddSection("EdGraphSchemaDistribution", LOCTEXT("DistributionHeader", "Distribution"));
+							InSection.AddMenuEntry(FGraphEditorCommands::Get().DistributeNodesHorizontally);
+							InSection.AddMenuEntry(FGraphEditorCommands::Get().DistributeNodesVertically);
+						}
+					})));
 				}
-				
-				MenuBuilder->EndSection();
 			}
 
 			if (const UK2Node* K2Node = Cast<const UK2Node>(InGraphNode))
@@ -1700,71 +1706,62 @@ void UEdGraphSchema_K2::GetContextMenuActions(const UEdGraph* CurrentGraph, cons
 						if (!K2Node->IsAutomaticallyPlacedGhostNode())
 						{
 							// Add compile options
-							MenuBuilder->BeginSection("EdGraphSchemaCompileOptions", LOCTEXT("CompileOptionsHeader", "Compile Options"));
 							{
-								MenuBuilder->AddMenuEntry(
+								FToolMenuSection& Section = Menu->AddSection("EdGraphSchemaCompileOptions", LOCTEXT("CompileOptionsHeader", "Compile Options"));
+								Section.AddMenuEntry(
 									FGraphEditorCommands::Get().DisableNodes,
-									NAME_None,
 									LOCTEXT("DisableCompile", "Disable (Do Not Compile)"),
 									LOCTEXT("DisableCompileToolTip", "Selected node(s) will not be compiled."));
 
-								TSharedPtr<const FUICommandList> MenuCommandList = MenuBuilder->GetTopCommandList();
-								if(ensure(MenuCommandList.IsValid()))
 								{
-									const FUIAction* SubMenuUIAction = MenuCommandList->GetActionForCommand(FGraphEditorCommands::Get().EnableNodes);
+									const FUIAction* SubMenuUIAction = Menu->Context.GetActionForCommand(FGraphEditorCommands::Get().EnableNodes);
 									if(ensure(SubMenuUIAction))
 									{
-										MenuBuilder->AddSubMenu(
+										Section.AddEntry(FToolMenuEntry::InitSubMenu(
+											Menu->GetMenuName(),
+											"EnableCompileSubMenu",
 											LOCTEXT("EnableCompileSubMenu", "Enable Compile"),
 											LOCTEXT("EnableCompileSubMenuToolTip", "Options to enable selected node(s) for compile."),
-											FNewMenuDelegate::CreateLambda([MenuCommandList](FMenuBuilder& SubMenuBuilder)
+											FNewToolMenuDelegate::CreateLambda([](UToolMenu* SubMenu)
 											{
-												SubMenuBuilder.PushCommandList(MenuCommandList.ToSharedRef());
-
-												SubMenuBuilder.AddMenuEntry(
+												FToolMenuSection& SubMenuSection = SubMenu->AddSection("Section");
+												SubMenuSection.AddMenuEntry(
 													FGraphEditorCommands::Get().EnableNodes_Always,
-													NAME_None,
 													LOCTEXT("EnableCompileAlways", "Always"),
 													LOCTEXT("EnableCompileAlwaysToolTip", "Always compile selected node(s)."));
-												SubMenuBuilder.AddMenuEntry(
+												SubMenuSection.AddMenuEntry(
 													FGraphEditorCommands::Get().EnableNodes_DevelopmentOnly,
-													NAME_None,
 													LOCTEXT("EnableCompileDevelopmentOnly", "Development Only"),
 													LOCTEXT("EnableCompileDevelopmentOnlyToolTip", "Compile selected node(s) for development only."));
-
-												SubMenuBuilder.PopCommandList();
 											}),
 											*SubMenuUIAction,
-											NAME_None, FGraphEditorCommands::Get().EnableNodes->GetUserInterfaceType());
+											FGraphEditorCommands::Get().EnableNodes->GetUserInterfaceType()));
 									}
 								}
 							}
-							MenuBuilder->EndSection();
 						}
 					}
 					
 					// Add breakpoint actions
-					MenuBuilder->BeginSection("EdGraphSchemaBreakpoints", LOCTEXT("BreakpointsHeader", "Breakpoints"));
 					{
-						MenuBuilder->AddMenuEntry( FGraphEditorCommands::Get().ToggleBreakpoint );
-						MenuBuilder->AddMenuEntry( FGraphEditorCommands::Get().AddBreakpoint );
-						MenuBuilder->AddMenuEntry( FGraphEditorCommands::Get().RemoveBreakpoint );
-						MenuBuilder->AddMenuEntry( FGraphEditorCommands::Get().EnableBreakpoint );
-						MenuBuilder->AddMenuEntry( FGraphEditorCommands::Get().DisableBreakpoint );
+						FToolMenuSection& Section = Menu->AddSection("EdGraphSchemaBreakpoints", LOCTEXT("BreakpointsHeader", "Breakpoints"));
+						Section.AddMenuEntry( FGraphEditorCommands::Get().ToggleBreakpoint );
+						Section.AddMenuEntry( FGraphEditorCommands::Get().AddBreakpoint );
+						Section.AddMenuEntry( FGraphEditorCommands::Get().RemoveBreakpoint );
+						Section.AddMenuEntry( FGraphEditorCommands::Get().EnableBreakpoint );
+						Section.AddMenuEntry( FGraphEditorCommands::Get().DisableBreakpoint );
 					}
-					MenuBuilder->EndSection();
 				}
 			}
 
-			MenuBuilder->BeginSection("EdGraphSchemaDocumentation", LOCTEXT("DocumentationHeader", "Documentation"));
 			{
-				MenuBuilder->AddMenuEntry( FGraphEditorCommands::Get().GoToDocumentation );
+				FToolMenuSection& Section = Menu->AddSection("EdGraphSchemaDocumentation", LOCTEXT("DocumentationHeader", "Documentation"));
+				Section.AddMenuEntry( FGraphEditorCommands::Get().GoToDocumentation );
 			}
-			MenuBuilder->EndSection();
 		}
 	}
 
-	Super::GetContextMenuActions(CurrentGraph, InGraphNode, InGraphPin, MenuBuilder, bIsDebugging);
+	Super::GetContextMenuActions(Menu, Context);
 }
 
 
@@ -1836,7 +1833,7 @@ void UEdGraphSchema_K2::OnReplaceVariableForVariableNode( UK2Node_Variable* Vari
 	}
 }
 
-void UEdGraphSchema_K2::GetReplaceVariableMenu(FMenuBuilder& MenuBuilder, UK2Node_Variable* Variable,  UBlueprint* OwnerBlueprint, bool bReplaceExistingVariable/*=false*/)
+void UEdGraphSchema_K2::GetReplaceVariableMenu(UToolMenu* Menu, UK2Node_Variable* Variable,  UBlueprint* OwnerBlueprint, bool bReplaceExistingVariable/*=false*/)
 {
 	if (UEdGraphPin* Pin = Variable->FindPin(Variable->GetVarName()))
 	{
@@ -1855,23 +1852,24 @@ void UEdGraphSchema_K2::GetReplaceVariableMenu(FMenuBuilder& MenuBuilder, UK2Nod
 		TArray<FName> Variables;
 		FBlueprintEditorUtils::GetNewVariablesOfType(OwnerBlueprint, Pin->PinType, Variables);
 
-		MenuBuilder.BeginSection(NAME_None, LOCTEXT("Variables", "Variables"));
-		for (TArray<FName>::TIterator VarIt(Variables); VarIt; ++VarIt)
 		{
-			if(*VarIt != ExistingVariableName)
+			FToolMenuSection& Section = Menu->AddSection(NAME_None, LOCTEXT("Variables", "Variables"));
+			for (TArray<FName>::TIterator VarIt(Variables); VarIt; ++VarIt)
 			{
-				const FText AlternativeVar = FText::FromName( *VarIt );
+				if (*VarIt != ExistingVariableName)
+				{
+					const FText AlternativeVar = FText::FromName(*VarIt);
 
-				FFormatNamedArguments TooltipArgs;
-				TooltipArgs.Add(TEXT("OldVariable"), Variable->GetVarNameText());
-				TooltipArgs.Add(TEXT("AlternateVariable"), AlternativeVar);
-				const FText Desc = FText::Format(ReplaceVariableWithTooltipFormat, TooltipArgs);
+					FFormatNamedArguments TooltipArgs;
+					TooltipArgs.Add(TEXT("OldVariable"), Variable->GetVarNameText());
+					TooltipArgs.Add(TEXT("AlternateVariable"), AlternativeVar);
+					const FText Desc = FText::Format(ReplaceVariableWithTooltipFormat, TooltipArgs);
 
-				MenuBuilder.AddMenuEntry( AlternativeVar, Desc, FSlateIcon(), FUIAction(
-					FExecuteAction::CreateStatic(&UEdGraphSchema_K2::OnReplaceVariableForVariableNode, const_cast<UK2Node_Variable* >(Variable), OwnerBlueprint, *VarIt, /*bIsSelfMember=*/true ) ) );
+					Section.AddMenuEntry(NAME_None, AlternativeVar, Desc, FSlateIcon(), FUIAction(
+						FExecuteAction::CreateStatic(&UEdGraphSchema_K2::OnReplaceVariableForVariableNode, const_cast<UK2Node_Variable*>(Variable), OwnerBlueprint, *VarIt, /*bIsSelfMember=*/true)));
+				}
 			}
 		}
-		MenuBuilder.EndSection();
 
 		FText ReplaceLocalVariableWithTooltipFormat;
 		if(!bReplaceExistingVariable)
@@ -1886,27 +1884,28 @@ void UEdGraphSchema_K2::GetReplaceVariableMenu(FMenuBuilder& MenuBuilder, UK2Nod
 		TArray<FName> LocalVariables;
 		FBlueprintEditorUtils::GetLocalVariablesOfType(Variable->GetGraph(), Pin->PinType, LocalVariables);
 
-		MenuBuilder.BeginSection(NAME_None, LOCTEXT("LocalVariables", "LocalVariables"));
-		for (TArray<FName>::TIterator VarIt(LocalVariables); VarIt; ++VarIt)
 		{
-			if(*VarIt != ExistingVariableName)
+			FToolMenuSection& Section = Menu->AddSection(NAME_None, LOCTEXT("LocalVariables", "LocalVariables"));
+			for (TArray<FName>::TIterator VarIt(LocalVariables); VarIt; ++VarIt)
 			{
-				const FText AlternativeVar = FText::FromName( *VarIt );
+				if (*VarIt != ExistingVariableName)
+				{
+					const FText AlternativeVar = FText::FromName(*VarIt);
 
-				FFormatNamedArguments TooltipArgs;
-				TooltipArgs.Add(TEXT("OldVariable"), Variable->GetVarNameText());
-				TooltipArgs.Add(TEXT("AlternateVariable"), AlternativeVar);
-				const FText Desc = FText::Format( ReplaceLocalVariableWithTooltipFormat, TooltipArgs );
+					FFormatNamedArguments TooltipArgs;
+					TooltipArgs.Add(TEXT("OldVariable"), Variable->GetVarNameText());
+					TooltipArgs.Add(TEXT("AlternateVariable"), AlternativeVar);
+					const FText Desc = FText::Format(ReplaceLocalVariableWithTooltipFormat, TooltipArgs);
 
-				MenuBuilder.AddMenuEntry( AlternativeVar, Desc, FSlateIcon(), FUIAction(
-					FExecuteAction::CreateStatic(&UEdGraphSchema_K2::OnReplaceVariableForVariableNode, const_cast<UK2Node_Variable* >(Variable),OwnerBlueprint, *VarIt, /*bIsSelfMember=*/false ) ) );
+					Section.AddMenuEntry(NAME_None, AlternativeVar, Desc, FSlateIcon(), FUIAction(
+						FExecuteAction::CreateStatic(&UEdGraphSchema_K2::OnReplaceVariableForVariableNode, const_cast<UK2Node_Variable*>(Variable), OwnerBlueprint, *VarIt, /*bIsSelfMember=*/false)));
+				}
 			}
 		}
-		MenuBuilder.EndSection();
 	}
 }
 
-void UEdGraphSchema_K2::GetNonExistentVariableMenu( const UEdGraphNode* InGraphNode, UBlueprint* OwnerBlueprint, FMenuBuilder* MenuBuilder ) const
+void UEdGraphSchema_K2::GetNonExistentVariableMenu(FToolMenuSection& Section, const UEdGraphNode* InGraphNode, UBlueprint* OwnerBlueprint) const
 {
 
 	if (const UK2Node_Variable* Variable = Cast<const UK2Node_Variable>(InGraphNode))
@@ -1920,7 +1919,7 @@ void UEdGraphSchema_K2::GetNonExistentVariableMenu( const UEdGraphNode* InGraphN
 				// create missing variable
 				const FText Label = FText::Format( LOCTEXT("CreateNonExistentVar", "Create variable '{0}'"), Variable->GetVarNameText());
 				const FText Desc = FText::Format( LOCTEXT("CreateNonExistentVarToolTip", "Variable '{0}' does not exist, create it?"), Variable->GetVarNameText());
-				MenuBuilder->AddMenuEntry( Label, Desc, FSlateIcon(), FUIAction(
+				Section.AddMenuEntry("CreateNonExistentVar", Label, Desc, FSlateIcon(), FUIAction(
 					FExecuteAction::CreateStatic( &UEdGraphSchema_K2::OnCreateNonExistentVariable, const_cast<UK2Node_Variable* >(Variable),OwnerBlueprint) ) );
 			}
 
@@ -1929,7 +1928,7 @@ void UEdGraphSchema_K2::GetNonExistentVariableMenu( const UEdGraphNode* InGraphN
 			{
 				const FText Label = FText::Format( LOCTEXT("CreateNonExistentLocalVar", "Create local variable '{0}'"), Variable->GetVarNameText());
 				const FText Desc = FText::Format( LOCTEXT("CreateNonExistentLocalVarToolTip", "Local variable '{0}' does not exist, create it?"), Variable->GetVarNameText());
-				MenuBuilder->AddMenuEntry( Label, Desc, FSlateIcon(), FUIAction(
+				Section.AddMenuEntry("CreateNonExistentLocalVar", Label, Desc, FSlateIcon(), FUIAction(
 					FExecuteAction::CreateStatic( &UEdGraphSchema_K2::OnCreateNonExistentLocalVariable, const_cast<UK2Node_Variable* >(Variable),OwnerBlueprint) ) );
 			}
 		}
@@ -1937,14 +1936,14 @@ void UEdGraphSchema_K2::GetNonExistentVariableMenu( const UEdGraphNode* InGraphN
 		// delete this node
 		{			
 			const FText Desc = FText::Format( LOCTEXT("DeleteNonExistentVarToolTip", "Referenced variable '{0}' does not exist, delete this node?"), Variable->GetVarNameText());
-			MenuBuilder->AddMenuEntry( FGenericCommands::Get().Delete, NAME_None, FGenericCommands::Get().Delete->GetLabel(), Desc);
+			Section.AddMenuEntry(FGenericCommands::Get().Delete, FGenericCommands::Get().Delete->GetLabel(), Desc, FSlateIcon(), NAME_None, "DeleteNonExistentVar");
 		}
 
-		GetReplaceVariableMenu(InGraphNode, OwnerBlueprint, MenuBuilder);
+		GetReplaceVariableMenu(Section, InGraphNode, OwnerBlueprint);
 	}
 }
 
-void UEdGraphSchema_K2::GetReplaceVariableMenu(const UEdGraphNode* InGraphNode, UBlueprint* InOwnerBlueprint, FMenuBuilder* InMenuBuilder, bool bInReplaceExistingVariable/* = false*/) const
+void UEdGraphSchema_K2::GetReplaceVariableMenu(FToolMenuSection& Section, const UEdGraphNode* InGraphNode, UBlueprint* InOwnerBlueprint, bool bInReplaceExistingVariable/* = false*/) const
 {
 	if (const UK2Node_Variable* Variable = Cast<const UK2Node_Variable>(InGraphNode))
 	{
@@ -1973,18 +1972,22 @@ void UEdGraphSchema_K2::GetReplaceVariableMenu(const UEdGraphNode* InGraphNode,
 					ReplaceVariableWithTooltip = LOCTEXT("ReplaceMissingVariableWithToolTip", "Variable '{0}' does not exist, replace with another variable?");
 				}
 
-				InMenuBuilder->AddSubMenu(
+				Section.AddEntry(FToolMenuEntry::InitSubMenu(
+					NAME_None,
+					"ReplaceVariableWith",
 					FText::Format( LOCTEXT("ReplaceVariableWith", "Replace variable '{0}' with..."), Variable->GetVarNameText()),
 					FText::Format( ReplaceVariableWithTooltip, Variable->GetVarNameText()),
-					FNewMenuDelegate::CreateStatic( &UEdGraphSchema_K2::GetReplaceVariableMenu,
-					const_cast<UK2Node_Variable*>(Variable),InOwnerBlueprint, bInReplaceExistingVariable));
+					FNewToolMenuDelegate::CreateStatic( &UEdGraphSchema_K2::GetReplaceVariableMenu,
+						const_cast<UK2Node_Variable*>(Variable), InOwnerBlueprint, bInReplaceExistingVariable)));
 			}
 		}
 	}
 }
 
-void UEdGraphSchema_K2::GetBreakLinkToSubMenuActions( class FMenuBuilder& MenuBuilder, UEdGraphPin* InGraphPin )
+void UEdGraphSchema_K2::GetBreakLinkToSubMenuActions(UToolMenu* Menu, UEdGraphPin* InGraphPin)
 {
+	FToolMenuSection& Section = Menu->FindOrAddSection("EdGraphSchemaPinActions");
+
 	// Make sure we have a unique name for every entry in the list
 	TMap< FString, uint32 > LinkTitleCount;
 
@@ -2023,13 +2026,15 @@ void UEdGraphSchema_K2::GetBreakLinkToSubMenuActions( class FMenuBuilder& MenuBu
 		}
 		++Count;
 
-		MenuBuilder.AddMenuEntry( Description, Description, FSlateIcon(), FUIAction(
+		Section.AddMenuEntry(NAME_None, Description, Description, FSlateIcon(), FUIAction(
 			FExecuteAction::CreateUObject(this, &UEdGraphSchema_K2::BreakSinglePinLink, const_cast< UEdGraphPin* >(InGraphPin), *Links)));
 	}
 }
 
-void UEdGraphSchema_K2::GetJumpToConnectionSubMenuActions( class FMenuBuilder& MenuBuilder, UEdGraphPin* InGraphPin )
+void UEdGraphSchema_K2::GetJumpToConnectionSubMenuActions(UToolMenu* Menu, UEdGraphPin* InGraphPin)
 {
+	FToolMenuSection& Section = Menu->FindOrAddSection("EdGraphSchemaPinActions");
+
 	// Make sure we have a unique name for every entry in the list
 	TMap< FString, uint32 > LinkTitleCount;
 
@@ -2067,7 +2072,7 @@ void UEdGraphSchema_K2::GetJumpToConnectionSubMenuActions( class FMenuBuilder& M
 		}
 		++Count;
 
-		MenuBuilder.AddMenuEntry( Description, Description, FSlateIcon(), FUIAction(
+		Section.AddMenuEntry(NAME_None, Description, Description, FSlateIcon(), FUIAction(
 		FExecuteAction::CreateStatic(&FKismetEditorUtilities::BringKismetToFocusAttentionOnPin, PinLink)));
 	}
 }
@@ -2082,10 +2087,10 @@ UEdGraphPin* UEdGraphSchema_K2::GetAndResetStraightenDestinationPin()
 	return Temp;
 }
 
-void UEdGraphSchema_K2::GetStraightenConnectionToSubMenuActions( class FMenuBuilder& MenuBuilder, UEdGraphPin* InGraphPin ) const
+void UEdGraphSchema_K2::GetStraightenConnectionToSubMenuActions(UToolMenu* Menu, UEdGraphPin* InGraphPin) const
 {
-	TSharedPtr<const FUICommandList> MenuCommandList = MenuBuilder.GetTopCommandList();
-	if(!ensure(MenuCommandList.IsValid()))
+	const FUIAction* StraightenConnectionsUIAction = Menu->Context.GetActionForCommand(FGraphEditorCommands::Get().StraightenConnections);
+	if(!ensure(StraightenConnectionsUIAction))
 	{
 		return;
 	}
@@ -2103,9 +2108,10 @@ void UEdGraphSchema_K2::GetStraightenConnectionToSubMenuActions( class FMenuBuil
 			NodeToPins.FindOrAdd(Node).Add(Pin);
 		}
 	}
-
-	MenuBuilder.AddMenuEntry( FGraphEditorCommands::Get().StraightenConnections,
-		NAME_None, LOCTEXT("StraightenAllConnections", "All Connected Pins"),
+	
+	FToolMenuSection& Section = Menu->FindOrAddSection("EdGraphSchemaPinActions");
+	Section.AddMenuEntry(FGraphEditorCommands::Get().StraightenConnections,
+		LOCTEXT("StraightenAllConnections", "All Connected Pins"),
 		TAttribute<FText>(), FSlateIcon(NAME_None, NAME_None, NAME_None) );
 
 	for (const TPair<UEdGraphNode*, TArray<UEdGraphPin*>>& Pair : NodeToPins)
@@ -2142,12 +2148,13 @@ void UEdGraphSchema_K2::GetStraightenConnectionToSubMenuActions( class FMenuBuil
 			}
 			++Count;
 
-			MenuBuilder.AddMenuEntry(
+			Section.AddMenuEntry(
+				NAME_None,
 				Description,
 				Description,
 				FSlateIcon(),
-				FExecuteAction::CreateLambda([=]{
-				if (const FUIAction* UIAction = MenuCommandList->GetActionForCommand(FGraphEditorCommands::Get().StraightenConnections))
+				FToolMenuExecuteAction::CreateLambda([=](const FToolMenuContext& Context){
+				if (const FUIAction* UIAction = Context.GetActionForCommand(FGraphEditorCommands::Get().StraightenConnections))
 				{
 					StraightenDestinationPin = Pin;
 					UIAction->ExecuteAction.Execute();
diff --git a/Engine/Source/Editor/GraphEditor/GraphEditor.Build.cs b/Engine/Source/Editor/GraphEditor/GraphEditor.Build.cs
--- a/Engine/Source/Editor/GraphEditor/GraphEditor.Build.cs
+++ b/Engine/Source/Editor/GraphEditor/GraphEditor.Build.cs
@@ -49,6 +49,7 @@ public class GraphEditor : ModuleRules
 				"Documentation",
 				"RenderCore",
 				"RHI",
+				"ToolMenus",
 			}
 		);
 
diff --git a/Engine/Source/Editor/GraphEditor/Private/SGraphEditorImpl.cpp b/Engine/Source/Editor/GraphEditor/Private/SGraphEditorImpl.cpp
--- a/Engine/Source/Editor/GraphEditor/Private/SGraphEditorImpl.cpp
+++ b/Engine/Source/Editor/GraphEditor/Private/SGraphEditorImpl.cpp
@@ -10,6 +10,7 @@
 #include "Widgets/Text/STextBlock.h"
 #include "Framework/MultiBox/MultiBoxExtender.h"
 #include "Framework/MultiBox/MultiBoxBuilder.h"
+#include "ToolMenus.h"
 #include "EditorStyleSet.h"
 #include "Editor.h"
 #include "GraphEditorModule.h"
@@ -18,6 +19,7 @@
 #include "ScopedTransaction.h"
 #include "SGraphEditorActionMenu.h"
 #include "Widgets/Notifications/SNotificationList.h"
+#include "Widgets/Input/SMultiLineEditableTextBox.h"
 #include "EdGraphSchema_K2.h"
 
 #define LOCTEXT_NAMESPACE "GraphEditorModule"
@@ -498,6 +500,210 @@ void SGraphEditorImpl::OnClosedActionMenu()
 	GraphPanel->OnStopMakingConnection(/*bForceStop=*/ true);
 }
 
+void SGraphEditorImpl::RegisterContextMenuFor_EdGraphSchema()
+{
+	static const FName MenuName = UEdGraphSchema::GetContextMenuName(UEdGraphSchema::StaticClass());
+	UToolMenus* ToolMenus = UToolMenus::Get();
+	if (ToolMenus->IsMenuRegistered(MenuName))
+	{
+		return;
+	}
+
+	UToolMenu* Menu = ToolMenus->RegisterMenu(MenuName);
+
+	Menu->AddDynamicSection("GetNodeContextMenuActions", FNewToolMenuDelegate::CreateLambda([](UToolMenu* InMenu)
+	{
+		if (UGraphNodeContextMenuContext* Context = InMenu->FindContext<UGraphNodeContextMenuContext>())
+		{
+			Context->Node->GetNodeContextMenuActions(InMenu, Context);
+		}
+	}));
+
+	Menu->AddDynamicSection("GetNodeContextMenuActionsLegacy", FNewToolMenuDelegateLegacy::CreateLambda([](FMenuBuilder& InMenuBuilder, UToolMenu* InMenu)
+	{
+		if (UGraphNodeContextMenuContext* Context = InMenu->FindContext<UGraphNodeContextMenuContext>())
+		{
+			FGraphNodeContextMenuBuilder MenuBuilderContext(Context->Graph, Context->Node, Context->Pin, &InMenuBuilder, Context->bIsDebugging);
+			Context->Node->GetContextMenuActions(MenuBuilderContext);
+		}
+	}));
+
+	Menu->AddDynamicSection("EdGraphSchema", FNewToolMenuDelegate::CreateLambda([](UToolMenu* InMenu)
+	{
+		UGraphNodeContextMenuContext* Context = InMenu->FindContext<UGraphNodeContextMenuContext>();
+		if (!Context)
+		{
+			return;
+		}
+
+		const UEdGraphSchema* GraphSchema = Context->Graph->GetSchema();
+		if (!GraphSchema)
+		{
+			return;
+		}
+
+		// Helper to do the node comment editing
+		struct Local
+		{
+			// Called by the EditableText widget to get the current comment for the node
+			static FString GetNodeComment(TWeakObjectPtr<UEdGraphNode> NodeWeakPtr)
+			{
+				if (UEdGraphNode* SelectedNode = NodeWeakPtr.Get())
+				{
+					return SelectedNode->NodeComment;
+				}
+				return FString();
+			}
+
+			// Called by the EditableText widget when the user types a new comment for the selected node
+			static void OnNodeCommentTextCommitted(const FText& NewText, ETextCommit::Type CommitInfo, TWeakObjectPtr<UEdGraphNode> NodeWeakPtr)
+			{
+				// Apply the change to the selected actor
+				UEdGraphNode* SelectedNode = NodeWeakPtr.Get();
+				FString NewString = NewText.ToString();
+				if (SelectedNode && !SelectedNode->NodeComment.Equals(NewString, ESearchCase::CaseSensitive))
+				{
+					// send property changed events
+					const FScopedTransaction Transaction(LOCTEXT("EditNodeComment", "Change Node Comment"));
+					SelectedNode->Modify();
+					UProperty* NodeCommentProperty = FindField<UProperty>(SelectedNode->GetClass(), "NodeComment");
+					if (NodeCommentProperty != nullptr)
+					{
+						SelectedNode->PreEditChange(NodeCommentProperty);
+
+						SelectedNode->NodeComment = NewString;
+						SelectedNode->SetMakeCommentBubbleVisible(true);
+
+						FPropertyChangedEvent NodeCommentPropertyChangedEvent(NodeCommentProperty);
+						SelectedNode->PostEditChangeProperty(NodeCommentPropertyChangedEvent);
+					}
+				}
+
+				// Only dismiss all menus if the text was committed via some user action.
+				// ETextCommit::Default implies that focus was switched by some other means. If this is because a submenu was opened, we don't want to close all the menus as a consequence.
+				if (CommitInfo != ETextCommit::Default)
+				{
+					FSlateApplication::Get().DismissAllMenus();
+				}
+			}
+		};
+
+		if (!Context->Pin)
+		{
+			int32 SelectionCount = GraphSchema->GetNodeSelectionCount(Context->Graph);
+
+			if (SelectionCount == 1)
+			{
+				// Node comment area
+				TSharedRef<SHorizontalBox> NodeCommentBox = SNew(SHorizontalBox);
+
+				{
+					FToolMenuSection& Section = InMenu->AddSection("GraphNodeComment", LOCTEXT("NodeCommentMenuHeader", "Node Comment"));
+					Section.AddEntry(FToolMenuEntry::InitWidget("NodeCommentBox", NodeCommentBox, FText::GetEmpty()));
+				}
+				TWeakObjectPtr<UEdGraphNode> SelectedNodeWeakPtr = MakeWeakObjectPtr(const_cast<UEdGraphNode*>(Context->Node));
+
+				FText NodeCommentText;
+				if (UEdGraphNode* SelectedNode = SelectedNodeWeakPtr.Get())
+				{
+					NodeCommentText = FText::FromString(SelectedNode->NodeComment);
+				}
+
+				const FSlateBrush* NodeIcon = FCoreStyle::Get().GetDefaultBrush();//@TODO: FActorIconFinder::FindIconForActor(SelectedActors(0).Get());
+
+				// Comment label
+				NodeCommentBox->AddSlot()
+					.VAlign(VAlign_Center)
+					.FillWidth(1.0f)
+					[
+						SNew(SMultiLineEditableTextBox)
+						.Text(NodeCommentText)
+						.ToolTipText(LOCTEXT("NodeComment_ToolTip", "Comment for this node"))
+						.OnTextCommitted_Static(&Local::OnNodeCommentTextCommitted, SelectedNodeWeakPtr)
+						.SelectAllTextWhenFocused(true)
+						.RevertTextOnEscape(true)
+						.ModiferKeyForNewLine(EModifierKey::Control)
+					];
+			}
+			else if (SelectionCount > 1)
+			{
+				struct SCommentUtility
+				{
+					static void CreateComment(const UEdGraphSchema* Schema, UEdGraph* Graph)
+					{
+						if (Schema && Graph)
+						{
+							TSharedPtr<FEdGraphSchemaAction> Action = Schema->GetCreateCommentAction();
+
+							if (Action.IsValid())
+							{
+								Action->PerformAction(Graph, nullptr, FVector2D());
+							}
+						}
+					}
+				};
+
+				FToolMenuSection& Section = InMenu->AddSection("SchemaActionComment", LOCTEXT("MultiCommentHeader", "Comment Group"));
+				Section.AddMenuEntry(
+					"MultiCommentDesc",
+					LOCTEXT("MultiCommentDesc", "Create Comment from Selection"),
+					LOCTEXT("CommentToolTip", "Create a resizable comment box around selection."),
+					FSlateIcon(),
+					FExecuteAction::CreateStatic(SCommentUtility::CreateComment, GraphSchema, const_cast<UEdGraph*>(Context->Graph)
+				));
+			}
+		}
+	}));
+}
+
+UToolMenu* SGraphEditorImpl::GenerateContextMenu(const UEdGraphSchema* Schema, FToolMenuContext& MenuContext) const
+{
+	// Register UEdGraphSchema's menu here to reduce amount of editor code in engine module
+	RegisterContextMenuFor_EdGraphSchema();
+
+	// Register any missing menus otherwise there is a break in the hierarchy chain
+	UToolMenus* ToolMenus = UToolMenus::Get();
+	for (UClass* CurrentClass = Schema->GetClass(); CurrentClass && CurrentClass->IsChildOf(UEdGraphSchema::StaticClass()); CurrentClass = CurrentClass->GetSuperClass())
+	{
+		const UEdGraphSchema* CurrentSchema = CurrentClass->GetDefaultObject<UEdGraphSchema>();
+		const FName CheckMenuName = CurrentSchema->GetContextMenuName();
+		if (!ToolMenus->IsMenuRegistered(CheckMenuName))
+		{
+			ToolMenus->RegisterMenu(CheckMenuName, CurrentSchema->GetParentContextMenuName());
+		}
+
+		CurrentClass = CurrentClass->GetSuperClass();
+	}
+
+	const FName MenuName = Schema->GetContextMenuName();
+	TArray<UToolMenu*> Hierarchy = ToolMenus->CollectHierarchy(MenuName);
+
+	// Add extra menu into hierarchy for calling GetContextMenuActions()
+	const FName CommonRootMenuName = "GraphEditor.GraphContextMenu.Common";
+	UToolMenu* CommonRootMenu = nullptr;
+	if (!ToolMenus->IsMenuRegistered(CommonRootMenuName))
+	{
+		CommonRootMenu = ToolMenus->RegisterMenu(CommonRootMenuName);
+		CommonRootMenu->AddDynamicSection("GetContextMenuActions", FNewToolMenuDelegate::CreateLambda([](UToolMenu* InMenu)
+		{
+			if (UGraphNodeContextMenuContext* ContextObject = InMenu->FindContext<UGraphNodeContextMenuContext>())
+			{
+				if (const UEdGraphSchema* GraphSchema = ContextObject->Graph->GetSchema())
+				{
+					GraphSchema->GetContextMenuActions(InMenu, ContextObject);
+				}
+			}
+		}));
+	}
+	else
+	{
+		CommonRootMenu = ToolMenus->FindMenu(CommonRootMenuName);
+	}
+	Hierarchy.Insert(CommonRootMenu, 0);
+
+	return ToolMenus->GenerateMenu(Hierarchy, MenuContext);
+}
+
 FActionMenuContent SGraphEditorImpl::GraphEd_OnGetContextMenuFor(const FGraphContextMenuArguments& SpawnInfo)
 {
 	FActionMenuContent Result;
@@ -529,18 +735,24 @@ FActionMenuContent SGraphEditorImpl::GraphEd_OnGetContextMenuFor(const FGraphCon
 			}
 			TSharedPtr<FExtender> MenuExtender = FExtender::Combine(Extenders);
 
-			// Show the menu for the pin or node under the cursor
-			const bool bShouldCloseAfterAction = true;
-			FMenuBuilder MenuBuilder( bShouldCloseAfterAction, this->Commands, MenuExtender );
-
 			if (OnCreateNodeOrPinMenu.IsBound())
 			{
+				// Show the menu for the pin or node under the cursor
+				const bool bShouldCloseAfterAction = true;
+				FMenuBuilder MenuBuilder( bShouldCloseAfterAction, this->Commands, MenuExtender );
+
 				Result = OnCreateNodeOrPinMenu.Execute(EdGraphObj, SpawnInfo.GraphNode, SpawnInfo.GraphPin, &MenuBuilder, !IsEditable.Get());
 			}
 			else
 			{
-				Schema->GetContextMenuActions(EdGraphObj, SpawnInfo.GraphNode, SpawnInfo.GraphPin, &MenuBuilder, !IsEditable.Get());
-				Result = FActionMenuContent(MenuBuilder.MakeWidget());
+				UGraphNodeContextMenuContext* ContextObject = NewObject<UGraphNodeContextMenuContext>();
+				ContextObject->Init(EdGraphObj, SpawnInfo.GraphNode, SpawnInfo.GraphPin, !IsEditable.Get());
+
+				FToolMenuContext Context(this->Commands, MenuExtender, ContextObject);
+
+				UToolMenus* ToolMenus = UToolMenus::Get();
+				UToolMenu* GeneratedMenu = GenerateContextMenu(Schema, Context);
+				Result = FActionMenuContent(ToolMenus->GenerateWidget(GeneratedMenu));
 			}
 		}
 		else if (IsEditable.Get())
diff --git a/Engine/Source/Editor/GraphEditor/Private/SGraphEditorImpl.h b/Engine/Source/Editor/GraphEditor/Private/SGraphEditorImpl.h
--- a/Engine/Source/Editor/GraphEditor/Private/SGraphEditorImpl.h
+++ b/Engine/Source/Editor/GraphEditor/Private/SGraphEditorImpl.h
@@ -333,6 +333,10 @@ private:
 	bool DisplayGraphAsReadOnly() const;
 
 	bool IsLocked() const;
+	
+	class UToolMenu* GenerateContextMenu(const class UEdGraphSchema* Schema, struct FToolMenuContext& MenuContext) const;
+
+	static void RegisterContextMenuFor_EdGraphSchema();
 
 private:
 	bool bIsActiveTimerRegistered;
diff --git a/Engine/Source/Editor/PhysicsAssetEditor/PhysicsAssetEditor.Build.cs b/Engine/Source/Editor/PhysicsAssetEditor/PhysicsAssetEditor.Build.cs
--- a/Engine/Source/Editor/PhysicsAssetEditor/PhysicsAssetEditor.Build.cs
+++ b/Engine/Source/Editor/PhysicsAssetEditor/PhysicsAssetEditor.Build.cs
@@ -41,6 +41,7 @@ public class PhysicsAssetEditor : ModuleRules
                 "AdvancedPreviewScene",
                 "DetailCustomizations",
                 "PinnedCommandList",
+				"ToolMenus",
             }
         );
 
diff --git a/Engine/Source/Editor/PhysicsAssetEditor/Private/PhysicsAssetGraph/PhysicsAssetGraphSchema.cpp b/Engine/Source/Editor/PhysicsAssetEditor/Private/PhysicsAssetGraph/PhysicsAssetGraphSchema.cpp
--- a/Engine/Source/Editor/PhysicsAssetEditor/Private/PhysicsAssetGraph/PhysicsAssetGraphSchema.cpp
+++ b/Engine/Source/Editor/PhysicsAssetEditor/Private/PhysicsAssetGraph/PhysicsAssetGraphSchema.cpp
@@ -9,6 +9,7 @@
 #include "PhysicsAssetGraphNode_Constraint.h"
 #include "PhysicsAssetGraph.h"
 #include "PhysicsAssetEditor.h"
+#include "ToolMenus.h"
 
 #define LOCTEXT_NAMESPACE "PhysicsAssetGraphSchema"
 
@@ -17,24 +18,30 @@ UPhysicsAssetGraphSchema::UPhysicsAssetGraphSchema(const FObjectInitializer& Obj
 {
 }
 
-void UPhysicsAssetGraphSchema::GetContextMenuActions(const UEdGraph* CurrentGraph, const UEdGraphNode* InGraphNode, const UEdGraphPin* InGraphPin, class FMenuBuilder* MenuBuilder, bool bIsDebugging) const
+void UPhysicsAssetGraphSchema::GetContextMenuActions(UToolMenu* Menu, UGraphNodeContextMenuContext* Context) const
 {
-	if (MenuBuilder)
+	Menu->AddDynamicSection("PhysicsAssetGraphSchema", FNewToolMenuDelegateLegacy::CreateLambda([](FMenuBuilder& MenuBuilder, UToolMenu* InMenu)
 	{
-		const UPhysicsAssetGraph* PhysicsAssetGraph = CastChecked<const UPhysicsAssetGraph>(CurrentGraph);
+		UGraphNodeContextMenuContext* ContextObject = InMenu->FindContext<UGraphNodeContextMenuContext>();
+		if (!ContextObject)
+		{
+			return;
+		}
+
+		const UPhysicsAssetGraph* PhysicsAssetGraph = CastChecked<const UPhysicsAssetGraph>(ContextObject->Graph);
 		TSharedPtr<FPhysicsAssetEditorSharedData> SharedData = PhysicsAssetGraph->GetPhysicsAssetEditor()->GetSharedData();
 
-		if (const UPhysicsAssetGraphNode_Constraint* ConstraintNode = Cast<const UPhysicsAssetGraphNode_Constraint>(InGraphNode))
+		if (const UPhysicsAssetGraphNode_Constraint* ConstraintNode = Cast<const UPhysicsAssetGraphNode_Constraint>(ContextObject->Node))
 		{
-			PhysicsAssetGraph->GetPhysicsAssetEditor()->BuildMenuWidgetConstraint(*MenuBuilder);
+			PhysicsAssetGraph->GetPhysicsAssetEditor()->BuildMenuWidgetConstraint(MenuBuilder);
 		}
-		else if (const UPhysicsAssetGraphNode_Bone* BoneNode = Cast<const UPhysicsAssetGraphNode_Bone>(InGraphNode))
+		else if (const UPhysicsAssetGraphNode_Bone* BoneNode = Cast<const UPhysicsAssetGraphNode_Bone>(ContextObject->Node))
 		{
-			PhysicsAssetGraph->GetPhysicsAssetEditor()->BuildMenuWidgetBody(*MenuBuilder);
+			PhysicsAssetGraph->GetPhysicsAssetEditor()->BuildMenuWidgetBody(MenuBuilder);
 		}
 
-		PhysicsAssetGraph->GetPhysicsAssetEditor()->BuildMenuWidgetSelection(*MenuBuilder);
-	}
+		PhysicsAssetGraph->GetPhysicsAssetEditor()->BuildMenuWidgetSelection(MenuBuilder);
+	}));
 }
 
 FLinearColor UPhysicsAssetGraphSchema::GetPinTypeColor(const FEdGraphPinType& PinType) const
diff --git a/Engine/Source/Editor/PhysicsAssetEditor/Private/PhysicsAssetGraph/PhysicsAssetGraphSchema.h b/Engine/Source/Editor/PhysicsAssetEditor/Private/PhysicsAssetGraph/PhysicsAssetGraphSchema.h
--- a/Engine/Source/Editor/PhysicsAssetEditor/Private/PhysicsAssetGraph/PhysicsAssetGraphSchema.h
+++ b/Engine/Source/Editor/PhysicsAssetEditor/Private/PhysicsAssetGraph/PhysicsAssetGraphSchema.h
@@ -21,7 +21,8 @@ class UPhysicsAssetGraphSchema : public UEdGraphSchema
 
 public:
 	// UEdGraphSchema interface
-	virtual void GetContextMenuActions(const UEdGraph* CurrentGraph, const UEdGraphNode* InGraphNode, const UEdGraphPin* InGraphPin, class FMenuBuilder* MenuBuilder, bool bIsDebugging) const override;
+	virtual void GetContextMenuActions(class UToolMenu* Menu, class UGraphNodeContextMenuContext* Context) const override;
+	virtual FName GetParentContextMenuName() const override { return NAME_None; }
 	virtual FLinearColor GetPinTypeColor(const FEdGraphPinType& PinType) const override;
 	virtual void BreakPinLinks(UEdGraphPin& TargetPin, bool bSendsNodeNotifcation) const override;
 	virtual void BreakSinglePinLink(UEdGraphPin* SourcePin, UEdGraphPin* TargetPin) const override;
diff --git a/Engine/Source/Editor/UnrealEd/Classes/MaterialGraph/MaterialGraphNode.h b/Engine/Source/Editor/UnrealEd/Classes/MaterialGraph/MaterialGraphNode.h
--- a/Engine/Source/Editor/UnrealEd/Classes/MaterialGraph/MaterialGraphNode.h
+++ b/Engine/Source/Editor/UnrealEd/Classes/MaterialGraph/MaterialGraphNode.h
@@ -64,7 +64,7 @@ public:
 	virtual FLinearColor GetNodeTitleColor() const override;
 	virtual FText GetTooltipText() const override;
 	virtual void PrepareForCopying() override;
-	virtual void GetContextMenuActions(const FGraphNodeContextMenuBuilder& Context) const override;
+	virtual void GetNodeContextMenuActions(class UToolMenu* Menu, class UGraphNodeContextMenuContext* Context) const override;
 	virtual void PostPlacedNewNode() override;
 	virtual void NodeConnectionListChanged() override;
 	virtual void OnRenameNode(const FString& NewName) override;
diff --git a/Engine/Source/Editor/UnrealEd/Classes/MaterialGraph/MaterialGraphSchema.h b/Engine/Source/Editor/UnrealEd/Classes/MaterialGraph/MaterialGraphSchema.h
--- a/Engine/Source/Editor/UnrealEd/Classes/MaterialGraph/MaterialGraphSchema.h
+++ b/Engine/Source/Editor/UnrealEd/Classes/MaterialGraph/MaterialGraphSchema.h
@@ -153,10 +153,10 @@ class UMaterialGraphSchema : public UEdGraphSchema
 	/**
 	 * Get menu for breaking links to specific nodes
 	 *
-	 * @param	MenuBuilder	MenuBuilder we are populating
+	 * @param	Menu		Menu we are populating
 	 * @param	InGraphPin	Pin with links to break
 	 */
-	void GetBreakLinkToSubMenuActions(class FMenuBuilder& MenuBuilder, class UEdGraphPin* InGraphPin);
+	void GetBreakLinkToSubMenuActions(class UToolMenu* Menu, class UEdGraphPin* InGraphPin) const;
 
 	/**
 	 * Connect a pin to one of the Material Function's outputs
@@ -187,7 +187,7 @@ class UMaterialGraphSchema : public UEdGraphSchema
 
 	//~ Begin UEdGraphSchema Interface
 	virtual void GetGraphContextActions(FGraphContextMenuBuilder& ContextMenuBuilder) const override;
-	virtual void GetContextMenuActions(const UEdGraph* CurrentGraph, const UEdGraphNode* InGraphNode, const UEdGraphPin* InGraphPin, class FMenuBuilder* MenuBuilder, bool bIsDebugging) const override;
+	virtual void GetContextMenuActions(class UToolMenu* Menu, class UGraphNodeContextMenuContext* Context) const override;
 	virtual const FPinConnectionResponse CanCreateConnection(const UEdGraphPin* A, const UEdGraphPin* B) const override;
 	virtual bool TryCreateConnection(UEdGraphPin* A, UEdGraphPin* B) const override;
 	virtual bool ShouldHidePinDefaultValue(UEdGraphPin* Pin) const override { return true; }
diff --git a/Engine/Source/Editor/UnrealEd/Private/MaterialGraphNode.cpp b/Engine/Source/Editor/UnrealEd/Private/MaterialGraphNode.cpp
--- a/Engine/Source/Editor/UnrealEd/Private/MaterialGraphNode.cpp
+++ b/Engine/Source/Editor/UnrealEd/Private/MaterialGraphNode.cpp
@@ -6,6 +6,7 @@
 
 #include "MaterialGraph/MaterialGraphNode.h"
 #include "Framework/MultiBox/MultiBoxBuilder.h"
+#include "ToolMenus.h"
 #include "MaterialGraph/MaterialGraphSchema.h"
 
 #include "Materials/MaterialExpressionComponentMask.h"
@@ -324,31 +325,33 @@ void UMaterialGraphNode::PrepareForCopying()
 	}
 }
 
-void UMaterialGraphNode::GetContextMenuActions(const FGraphNodeContextMenuBuilder& Context) const
+void UMaterialGraphNode::GetNodeContextMenuActions(UToolMenu* Menu, UGraphNodeContextMenuContext* Context) const
 {
 	UMaterialGraph* MaterialGraph = CastChecked<UMaterialGraph>(GetGraph());
 
-	if (Context.Node)
+	if (Context->Node)
 	{
 		if (MaterialExpression)
 		{
 			if (MaterialExpression->IsA(UMaterialExpressionTextureBase::StaticClass()))
 			{
-				Context.MenuBuilder->AddMenuEntry(FMaterialEditorCommands::Get().UseCurrentTexture);
+				{
+					FToolMenuSection& Section = Menu->AddSection("MaterialGraphNode");
+					Section.AddMenuEntry(FMaterialEditorCommands::Get().UseCurrentTexture);
+				}
 
 				// Add a 'Convert To Texture' option for convertible types
-				Context.MenuBuilder->BeginSection("MaterialEditorMenu0");
 				{
+					FToolMenuSection& Section = Menu->AddSection("MaterialEditorMenu0");
 					if ( MaterialExpression->IsA(UMaterialExpressionTextureSample::StaticClass()) && !MaterialExpression->HasAParameterName())
 					{
-						Context.MenuBuilder->AddMenuEntry(FMaterialEditorCommands::Get().ConvertToTextureObjects);
+						Section.AddMenuEntry(FMaterialEditorCommands::Get().ConvertToTextureObjects);
 					}
 					else if ( MaterialExpression->IsA(UMaterialExpressionTextureObject::StaticClass()))
 					{
-						Context.MenuBuilder->AddMenuEntry(FMaterialEditorCommands::Get().ConvertToTextureSamples);
+						Section.AddMenuEntry(FMaterialEditorCommands::Get().ConvertToTextureSamples);
 					}
 				}
-				Context.MenuBuilder->EndSection();
 			}
 
 			// Add a 'Convert To Parameter' option for convertible types
@@ -361,11 +364,10 @@ void UMaterialGraphNode::GetContextMenuActions(const FGraphNodeContextMenuBuilde
 				|| MaterialExpression->IsA(UMaterialExpressionComponentMask::StaticClass())
 				|| MaterialExpression->IsA(UMaterialExpressionMaterialFunctionCall::StaticClass()))
 			{
-				Context.MenuBuilder->BeginSection("MaterialEditorMenu1");
 				{
-					Context.MenuBuilder->AddMenuEntry(FMaterialEditorCommands::Get().ConvertObjects);
+					FToolMenuSection& Section = Menu->AddSection("MaterialEditorMenu1");
+					Section.AddMenuEntry(FMaterialEditorCommands::Get().ConvertObjects);
 				}
-				Context.MenuBuilder->EndSection();
 			}
 
 			// Add a 'Convert To Constant' option for convertible types
@@ -373,15 +375,14 @@ void UMaterialGraphNode::GetContextMenuActions(const FGraphNodeContextMenuBuilde
 				|| MaterialExpression->IsA(UMaterialExpressionVectorParameter::StaticClass())
 				|| MaterialExpression->IsA(UMaterialExpressionTextureObjectParameter::StaticClass()))
 			{
-				Context.MenuBuilder->BeginSection("MaterialEditorMenu1");
 				{
-					Context.MenuBuilder->AddMenuEntry(FMaterialEditorCommands::Get().ConvertToConstant);
+					FToolMenuSection& Section = Menu->AddSection("MaterialEditorMenu1");
+					Section.AddMenuEntry(FMaterialEditorCommands::Get().ConvertToConstant);
 				}
-				Context.MenuBuilder->EndSection();
 			}
 
-			Context.MenuBuilder->BeginSection("MaterialEditorMenu2");
 			{
+				FToolMenuSection& Section = Menu->AddSection("MaterialEditorMenu2");
 				// Don't show preview option for bools
 				if (!MaterialExpression->IsA(UMaterialExpressionStaticBool::StaticClass())
 					&& !MaterialExpression->IsA(UMaterialExpressionStaticBoolParameter::StaticClass()))
@@ -390,83 +391,88 @@ void UMaterialGraphNode::GetContextMenuActions(const FGraphNodeContextMenuBuilde
 					if (bIsPreviewExpression)
 					{
 						// If we are already previewing the selected node, the menu option should tell the user that this will stop previewing
-						Context.MenuBuilder->AddMenuEntry(FMaterialEditorCommands::Get().StopPreviewNode);
+						Section.AddMenuEntry(FMaterialEditorCommands::Get().StopPreviewNode);
 					}
 					else
 					{
 						// The menu option should tell the user this node will be previewed.
-						Context.MenuBuilder->AddMenuEntry(FMaterialEditorCommands::Get().StartPreviewNode);
+						Section.AddMenuEntry(FMaterialEditorCommands::Get().StartPreviewNode);
 					}
 				}
 
 				if (MaterialExpression->bRealtimePreview)
 				{
-					Context.MenuBuilder->AddMenuEntry(FMaterialEditorCommands::Get().DisableRealtimePreviewNode);
+					Section.AddMenuEntry(FMaterialEditorCommands::Get().DisableRealtimePreviewNode);
 				}
 				else
 				{
-					Context.MenuBuilder->AddMenuEntry(FMaterialEditorCommands::Get().EnableRealtimePreviewNode);
+					Section.AddMenuEntry(FMaterialEditorCommands::Get().EnableRealtimePreviewNode);
 				}
 			}
-			Context.MenuBuilder->EndSection();
 		}
 
 		// Break all links
-		Context.MenuBuilder->AddMenuEntry(FGraphEditorCommands::Get().BreakNodeLinks);
+		{
+			FToolMenuSection& Section = Menu->AddSection("BreakAllLinks");
+			Section.AddMenuEntry(FGraphEditorCommands::Get().BreakNodeLinks);
+		}
 
 		// Separate the above frequently used options from the below less frequently used common options
-		Context.MenuBuilder->BeginSection("MaterialEditorMenu3");
+		
 		{
-			Context.MenuBuilder->AddMenuEntry( FGenericCommands::Get().Delete );
-			Context.MenuBuilder->AddMenuEntry( FGenericCommands::Get().Cut );
-			Context.MenuBuilder->AddMenuEntry( FGenericCommands::Get().Copy );
-			Context.MenuBuilder->AddMenuEntry( FGenericCommands::Get().Duplicate );
+			FToolMenuSection& Section = Menu->AddSection("MaterialEditorMenu3");
+			Section.AddMenuEntry( FGenericCommands::Get().Delete );
+			Section.AddMenuEntry( FGenericCommands::Get().Cut );
+			Section.AddMenuEntry( FGenericCommands::Get().Copy );
+			Section.AddMenuEntry( FGenericCommands::Get().Duplicate );
 
 			// Select upstream and downstream nodes
-			Context.MenuBuilder->AddMenuEntry(FMaterialEditorCommands::Get().SelectDownstreamNodes);
-			Context.MenuBuilder->AddMenuEntry(FMaterialEditorCommands::Get().SelectUpstreamNodes);
+			Section.AddMenuEntry(FMaterialEditorCommands::Get().SelectDownstreamNodes);
+			Section.AddMenuEntry(FMaterialEditorCommands::Get().SelectUpstreamNodes);
 		}
-		Context.MenuBuilder->EndSection();
-
-		Context.MenuBuilder->AddSubMenu(LOCTEXT("AlignmentHeader", "Alignment"), FText(), FNewMenuDelegate::CreateLambda([](FMenuBuilder& InMenuBuilder) {
-
-			InMenuBuilder.BeginSection("EdGraphSchemaAlignment", LOCTEXT("AlignHeader", "Align"));
-			InMenuBuilder.AddMenuEntry(FGraphEditorCommands::Get().AlignNodesTop);
-			InMenuBuilder.AddMenuEntry(FGraphEditorCommands::Get().AlignNodesMiddle);
-			InMenuBuilder.AddMenuEntry(FGraphEditorCommands::Get().AlignNodesBottom);
-			InMenuBuilder.AddMenuEntry(FGraphEditorCommands::Get().AlignNodesLeft);
-			InMenuBuilder.AddMenuEntry(FGraphEditorCommands::Get().AlignNodesCenter);
-			InMenuBuilder.AddMenuEntry(FGraphEditorCommands::Get().AlignNodesRight);
-			InMenuBuilder.AddMenuEntry(FGraphEditorCommands::Get().StraightenConnections);
-			InMenuBuilder.EndSection();
-
-			InMenuBuilder.BeginSection("EdGraphSchemaDistribution", LOCTEXT("DistributionHeader", "Distribution"));
-			InMenuBuilder.AddMenuEntry(FGraphEditorCommands::Get().DistributeNodesHorizontally);
-			InMenuBuilder.AddMenuEntry(FGraphEditorCommands::Get().DistributeNodesVertically);
-			InMenuBuilder.EndSection();
-		}));
-
-		Context.MenuBuilder->BeginSection("MaterialEditorMenuDocumentation");
+			
+		{
+			FToolMenuSection& Section = Menu->AddSection("Alignment");
+			Section.AddSubMenu(Menu->GetMenuName(), "Alignment", LOCTEXT("AlignmentHeader", "Alignment"), FText(), FNewToolMenuDelegate::CreateLambda([](UToolMenu* InMenu)
+			{
+				{
+					FToolMenuSection& SubMenuSection = InMenu->AddSection("EdGraphSchemaAlignment", LOCTEXT("AlignHeader", "Align"));
+					SubMenuSection.AddMenuEntry(FGraphEditorCommands::Get().AlignNodesTop);
+					SubMenuSection.AddMenuEntry(FGraphEditorCommands::Get().AlignNodesMiddle);
+					SubMenuSection.AddMenuEntry(FGraphEditorCommands::Get().AlignNodesBottom);
+					SubMenuSection.AddMenuEntry(FGraphEditorCommands::Get().AlignNodesLeft);
+					SubMenuSection.AddMenuEntry(FGraphEditorCommands::Get().AlignNodesCenter);
+					SubMenuSection.AddMenuEntry(FGraphEditorCommands::Get().AlignNodesRight);
+					SubMenuSection.AddMenuEntry(FGraphEditorCommands::Get().StraightenConnections);
+				}
+
+				{
+					FToolMenuSection& SubMenuSection = InMenu->AddSection("EdGraphSchemaDistribution", LOCTEXT("DistributionHeader", "Distribution"));
+					SubMenuSection.AddMenuEntry(FGraphEditorCommands::Get().DistributeNodesHorizontally);
+					SubMenuSection.AddMenuEntry(FGraphEditorCommands::Get().DistributeNodesVertically);
+				}
+			}));
+		}
+
 		{
-			Context.MenuBuilder->AddMenuEntry(FGraphEditorCommands::Get().GoToDocumentation);
+			FToolMenuSection& Section = Menu->AddSection("MaterialEditorMenuDocumentation");
+			Section.AddMenuEntry(FGraphEditorCommands::Get().GoToDocumentation);
 		}
-		Context.MenuBuilder->EndSection();
 
 		// Handle the favorites options
 		if (MaterialExpression)
 		{
-			Context.MenuBuilder->BeginSection("MaterialEditorMenuFavorites");
 			{
+				FToolMenuSection& Section = Menu->AddSection("MaterialEditorMenuFavorites");
 				if (FMaterialEditorUtilities::IsMaterialExpressionInFavorites(MaterialExpression))
 				{
-					Context.MenuBuilder->AddMenuEntry(FMaterialEditorCommands::Get().RemoveFromFavorites);
+					Section.AddMenuEntry(FMaterialEditorCommands::Get().RemoveFromFavorites);
 				}
 				else
 				{
-					Context.MenuBuilder->AddMenuEntry(FMaterialEditorCommands::Get().AddToFavorites);
+					Section.AddMenuEntry(FMaterialEditorCommands::Get().AddToFavorites);
 				}
 			}
-			Context.MenuBuilder->EndSection();
 		}
 	}
 }
diff --git a/Engine/Source/Editor/UnrealEd/Private/MaterialGraphSchema.cpp b/Engine/Source/Editor/UnrealEd/Private/MaterialGraphSchema.cpp
--- a/Engine/Source/Editor/UnrealEd/Private/MaterialGraphSchema.cpp
+++ b/Engine/Source/Editor/UnrealEd/Private/MaterialGraphSchema.cpp
@@ -12,6 +12,8 @@
 #include "Textures/SlateIcon.h"
 #include "Framework/Commands/UIAction.h"
 #include "Framework/MultiBox/MultiBoxBuilder.h"
+#include "ToolMenus.h"
+#include "EdGraph/EdGraph.h"
 #include "HAL/IConsoleManager.h"
 #include "Materials/MaterialExpression.h"
 #include "Materials/MaterialExpressionMaterialFunctionCall.h"
@@ -219,8 +221,10 @@ void UMaterialGraphSchema::SelectAllInputNodes(UEdGraph* Graph, UEdGraphPin* InG
 	}
 }
 
-void UMaterialGraphSchema::GetBreakLinkToSubMenuActions( class FMenuBuilder& MenuBuilder, UEdGraphPin* InGraphPin )
+void UMaterialGraphSchema::GetBreakLinkToSubMenuActions(class UToolMenu* Menu, UEdGraphPin* InGraphPin) const
 {
+	FToolMenuSection& Section = Menu->FindOrAddSection("MaterialGraphSchemaPinActions");
+
 	// Make sure we have a unique name for every entry in the list
 	TMap< FString, uint32 > LinkTitleCount;
 
@@ -257,7 +261,7 @@ void UMaterialGraphSchema::GetBreakLinkToSubMenuActions( class FMenuBuilder& Men
 			Description = FText::Format( LOCTEXT("BreakDescMulti", "Break link to {NodeTitle} ({NumberOfNodes})"), Args );
 		}
 		++Count;
-		MenuBuilder.AddMenuEntry( Description, Description, FSlateIcon(), FUIAction(
+		Section.AddMenuEntry(NAME_None, Description, Description, FSlateIcon(), FUIAction(
 			FExecuteAction::CreateUObject((UMaterialGraphSchema*const)this, &UMaterialGraphSchema::BreakSinglePinLink, const_cast< UEdGraphPin* >(InGraphPin), *Links) ) );
 	}
 }
@@ -388,55 +392,58 @@ void UMaterialGraphSchema::GetGraphContextActions(FGraphContextMenuBuilder& Cont
 	}
 }
 
-void UMaterialGraphSchema::GetContextMenuActions(const UEdGraph* CurrentGraph, const UEdGraphNode* InGraphNode, const UEdGraphPin* InGraphPin, class FMenuBuilder* MenuBuilder, bool bIsDebugging) const
+void UMaterialGraphSchema::GetContextMenuActions(UToolMenu* Menu, UGraphNodeContextMenuContext* Context) const
 {
-	if (InGraphPin)
+	if (Context && Context->Pin)
 	{
-		const UMaterialGraph* MaterialGraph = CastChecked<UMaterialGraph>(CurrentGraph);
-		MenuBuilder->BeginSection("MaterialGraphSchemaPinActions", LOCTEXT("PinActionsMenuHeader", "Pin Actions"));
+		const UEdGraphPin* InGraphPin = Context->Pin;
+		const UMaterialGraph* MaterialGraph = CastChecked<UMaterialGraph>(Context->Graph);
 		{
+			FToolMenuSection& Section = Menu->AddSection("MaterialGraphSchemaPinActions", LOCTEXT("PinActionsMenuHeader", "Pin Actions"));
 			// Only display the 'Break Link' option if there is a link to break!
 			if (InGraphPin->LinkedTo.Num() > 0)
 			{
 				if(InGraphPin->Direction == EEdGraphPinDirection::EGPD_Input)
 				{
-					MenuBuilder->AddMenuEntry(
+					Section.AddMenuEntry(
+						"SelectLinkedNodes",
 						LOCTEXT("SelectLinkedNodes", "Select Linked Nodes"),
 						LOCTEXT("SelectLinkedNodesTooltip", "Adds all input Nodes linked to this Pin to selection"),
 						FSlateIcon(),
-						FUIAction(FExecuteAction::CreateUObject((UMaterialGraphSchema*const)this, &UMaterialGraphSchema::SelectAllInputNodes, const_cast<UEdGraph*>(CurrentGraph), const_cast<UEdGraphPin*>(InGraphPin)))
+						FUIAction(FExecuteAction::CreateUObject(this, &UMaterialGraphSchema::SelectAllInputNodes, const_cast<UEdGraph*>(Context->Graph), const_cast<UEdGraphPin*>(InGraphPin)))
 						);
 				}
 
-				MenuBuilder->AddMenuEntry(FGraphEditorCommands::Get().BreakPinLinks);
+				Section.AddMenuEntry(FGraphEditorCommands::Get().BreakPinLinks);
 
 				// add sub menu for break link to
 				if(InGraphPin->LinkedTo.Num() > 1)
 				{
-					MenuBuilder->AddSubMenu(
+					Section.AddEntry(FToolMenuEntry::InitSubMenu(
+						Menu->GetMenuName(),
+						"BreakLinkTo",
 						LOCTEXT("BreakLinkTo", "Break Link To..." ),
 						LOCTEXT("BreakSpecificLinks", "Break a specific link..." ),
-						FNewMenuDelegate::CreateUObject( (UMaterialGraphSchema*const)this, &UMaterialGraphSchema::GetBreakLinkToSubMenuActions, const_cast<UEdGraphPin*>(InGraphPin)));
+						FNewToolMenuDelegate::CreateUObject(this, &UMaterialGraphSchema::GetBreakLinkToSubMenuActions, const_cast<UEdGraphPin*>(InGraphPin))));
 				}
 				else
 				{
-					((UMaterialGraphSchema*const)this)->GetBreakLinkToSubMenuActions(*MenuBuilder, const_cast<UEdGraphPin*>(InGraphPin));
+					GetBreakLinkToSubMenuActions(Menu, const_cast<UEdGraphPin*>(InGraphPin));
 				}
 			}
 
 			// Only display Promote to Parameters on input pins
 			if (InGraphPin->Direction == EEdGraphPinDirection::EGPD_Input)
 			{
-				MenuBuilder->AddMenuEntry(FMaterialEditorCommands::Get().PromoteToParameter);
+				Section.AddMenuEntry(FMaterialEditorCommands::Get().PromoteToParameter);
 			}
 		}
-		MenuBuilder->EndSection();
 
 		// add menu items to expression output for material connection
 		if ( InGraphPin->Direction == EGPD_Output )
 		{
-			MenuBuilder->BeginSection("MaterialEditorMenuConnector2");
 			{
+				FToolMenuSection& Section = Menu->AddSection("MaterialEditorMenuConnector2");
 				// If we are editing a material function, display options to connect to function outputs
 				if (MaterialGraph->MaterialFunction)
 				{
@@ -452,10 +459,12 @@ void UMaterialGraphSchema::GetContextMenuActions(const UEdGraph* CurrentGraph, c
 								Arguments.Add(TEXT("Name"), FText::FromName( FunctionOutput->OutputName ));
 								const FText Label = FText::Format( LOCTEXT( "ConnectToFunction", "Connect To {Name}" ), Arguments );
 								const FText ToolTip = FText::Format( LOCTEXT( "ConnectToFunctionTooltip", "Connects to the function output {Name}" ), Arguments );
-								MenuBuilder->AddMenuEntry(Label,
+								Section.AddMenuEntry(
+									NAME_None,
+									Label,
 									ToolTip,
 									FSlateIcon(),
-									FUIAction(FExecuteAction::CreateUObject((UMaterialGraphSchema*const)this, &UMaterialGraphSchema::OnConnectToFunctionOutput, const_cast< UEdGraphPin* >(InGraphPin), GraphNode->GetInputPin(0))));
+									FUIAction(FExecuteAction::CreateUObject(this, &UMaterialGraphSchema::OnConnectToFunctionOutput, const_cast< UEdGraphPin* >(InGraphPin), GraphNode->GetInputPin(0))));
 							}
 						}
 					}
@@ -470,23 +479,24 @@ void UMaterialGraphSchema::GetContextMenuActions(const UEdGraph* CurrentGraph, c
 							Arguments.Add(TEXT("Name"), MaterialGraph->MaterialInputs[Index].GetName());
 							const FText Label = FText::Format( LOCTEXT( "ConnectToInput", "Connect To {Name}" ), Arguments );
 							const FText ToolTip = FText::Format( LOCTEXT( "ConnectToInputTooltip", "Connects to the material input {Name}" ), Arguments );
-							MenuBuilder->AddMenuEntry(Label,
+							Section.AddMenuEntry(
+								NAME_None,
+								Label,
 								ToolTip,
 								FSlateIcon(),
-								FUIAction(FExecuteAction::CreateUObject((UMaterialGraphSchema*const)this, &UMaterialGraphSchema::OnConnectToMaterial, const_cast< UEdGraphPin* >(InGraphPin), Index)));
+								FUIAction(FExecuteAction::CreateUObject(this, &UMaterialGraphSchema::OnConnectToMaterial, const_cast< UEdGraphPin* >(InGraphPin), Index)));
 						}
 					}
 				}
 			}
-			MenuBuilder->EndSection(); //MaterialEditorMenuConnector2
 		}
 	}
-	else if (InGraphNode)
+	else
 	{
 		//Moved all functionality to relevant node classes
 	}
 
-	Super::GetContextMenuActions(CurrentGraph, InGraphNode, InGraphPin, MenuBuilder, bIsDebugging);
+	Super::GetContextMenuActions(Menu, Context);
 }
 
 static TAutoConsoleVariable<int32> CVarPreventInvalidMaterialConnections(
diff --git a/Engine/Source/Runtime/Engine/Classes/EdGraph/EdGraphNode.h b/Engine/Source/Runtime/Engine/Classes/EdGraph/EdGraphNode.h
--- a/Engine/Source/Runtime/Engine/Classes/EdGraph/EdGraphNode.h
+++ b/Engine/Source/Runtime/Engine/Classes/EdGraph/EdGraphNode.h
@@ -152,7 +152,7 @@ private:
 };
 
 /** This is the context for a GetContextMenuActions call into a specific node. */
-struct FGraphNodeContextMenuBuilder
+struct ENGINE_API FGraphNodeContextMenuBuilder
 {
 	/** The blueprint associated with this context; may be NULL for non-Kismet related graphs. */
 	const UBlueprint* Blueprint;
@@ -170,6 +170,37 @@ struct FGraphNodeContextMenuBuilder
 	FGraphNodeContextMenuBuilder(const UEdGraph* InGraph, const UEdGraphNode* InNode, const UEdGraphPin* InPin, class FMenuBuilder* InMenuBuilder, bool bInDebuggingMode);
 };
 
+UCLASS()
+class ENGINE_API UGraphNodeContextMenuContext : public UObject
+{
+	GENERATED_BODY()
+
+public:
+
+	UGraphNodeContextMenuContext();
+
+	void Init(const UEdGraph* InGraph, const UEdGraphNode* InNode, const UEdGraphPin* InPin, bool bInDebuggingMode);
+
+	/** The blueprint associated with this context; may be NULL for non-Kismet related graphs. */
+	UPROPERTY()
+	const UBlueprint* Blueprint;
+
+	/** The graph associated with this context. */
+	UPROPERTY()
+	const UEdGraph* Graph;
+
+	/** The node associated with this context. */
+	UPROPERTY()
+	const UEdGraphNode* Node;
+
+	/** The pin associated with this context; may be NULL when over a node. */
+	const UEdGraphPin* Pin;
+
+	/** Whether the graph editor is currently part of a debugging session (any non-debugging commands should be disabled). */
+	UPROPERTY()
+	bool bIsDebugging;
+};
+
 /** Deprecation types for node response. */
 enum class EEdGraphNodeDeprecationType
 {
@@ -805,6 +836,9 @@ public:
 
 	/** Gets a list of actions that can be done to this particular node */
 	virtual void GetContextMenuActions(const FGraphNodeContextMenuBuilder& Context) const {}
+	
+	/** Gets a list of actions that can be done to this particular node */
+	virtual void GetNodeContextMenuActions(class UToolMenu* Menu, class UGraphNodeContextMenuContext* Context) const {}
 
 	// Gives each visual node a chance to do final validation before it's node is harvested for use at runtime
 	virtual void ValidateNodeDuringCompilation(class FCompilerResultsLog& MessageLog) const {}
diff --git a/Engine/Source/Runtime/Engine/Classes/EdGraph/EdGraphSchema.h b/Engine/Source/Runtime/Engine/Classes/EdGraph/EdGraphSchema.h
--- a/Engine/Source/Runtime/Engine/Classes/EdGraph/EdGraphSchema.h
+++ b/Engine/Source/Runtime/Engine/Classes/EdGraph/EdGraphSchema.h
@@ -619,17 +619,23 @@ class ENGINE_API UEdGraphSchema : public UObject
 	 */
 	 virtual void GetGraphContextActions(FGraphContextMenuBuilder& ContextMenuBuilder) const;
 
+	/** Returns context menu name */
+	FName GetContextMenuName() const;
+
+	/** Returns parent context menu name */
+	virtual FName GetParentContextMenuName() const;
+
+	/** Returns context menu name for a given class */
+	static FName GetContextMenuName(UClass* InClass);
+
 	/**
 	 * Gets actions that should be added to the right-click context menu for a node or pin
 	 * 
-	 * @param	CurrentGraph		The current graph.
-	 * @param	InGraphNode			The node to get the context menu for, if any.
-	 * @param	InGraphPin			The pin clicked on, if any, to provide additional context
-	 * @param	MenuBuilder			The menu builder to append actions to.
-	 * @param	bIsDebugging		Is the graph editor currently part of a debugging session (any non-debugging commands should be disabled)
+	 * @param	Menu				The menu to append actions to.
+	 * @param	Context				The menu's context.
 	 */
-	 virtual void GetContextMenuActions(const UEdGraph* CurrentGraph, const UEdGraphNode* InGraphNode, const UEdGraphPin* InGraphPin, class FMenuBuilder* MenuBuilder, bool bIsDebugging) const;
-
+	virtual void GetContextMenuActions(class UToolMenu* Menu, class UGraphNodeContextMenuContext* Context) const;
+	
 	/**
 	 * Determine if a connection can be created between two pins.
 	 *
diff --git a/Engine/Source/Runtime/Engine/Private/EdGraph/EdGraphNode.cpp b/Engine/Source/Runtime/Engine/Private/EdGraph/EdGraphNode.cpp
--- a/Engine/Source/Runtime/Engine/Private/EdGraph/EdGraphNode.cpp
+++ b/Engine/Source/Runtime/Engine/Private/EdGraph/EdGraphNode.cpp
@@ -116,6 +116,31 @@ FGraphNodeContextMenuBuilder::FGraphNodeContextMenuBuilder(const UEdGraph* InGra
 	}
 }
 
+/////////////////////////////////////////////////////
+// UGraphNodeContextMenuContext
+
+UGraphNodeContextMenuContext::UGraphNodeContextMenuContext() :
+	Pin(nullptr)
+{
+}
+
+void UGraphNodeContextMenuContext::Init(const UEdGraph* InGraph, const UEdGraphNode* InNode, const UEdGraphPin* InPin, bool bInDebuggingMode)
+{
+	Graph = InGraph;
+	Node = InNode;
+	Pin = InPin;
+	bIsDebugging = bInDebuggingMode;
+
+#if WITH_EDITOR
+	Blueprint = FBlueprintEditorUtils::FindBlueprintForGraph(Graph);
+#endif
+
+	if (Pin)
+	{
+		Node = Pin->GetOwningNode();
+	}
+}
+
 /////////////////////////////////////////////////////
 // UEdGraphNode
 
diff --git a/Engine/Source/Runtime/Engine/Private/EdGraph/EdGraphSchema.cpp b/Engine/Source/Runtime/Engine/Private/EdGraph/EdGraphSchema.cpp
--- a/Engine/Source/Runtime/Engine/Private/EdGraph/EdGraphSchema.cpp
+++ b/Engine/Source/Runtime/Engine/Private/EdGraph/EdGraphSchema.cpp
@@ -784,130 +784,38 @@ void UEdGraphSchema::GetGraphDisplayInformation(const UEdGraph& Graph, /*out*/ F
 	DisplayInfo.DisplayName = DisplayInfo.PlainName;
 }
 
-void UEdGraphSchema::GetContextMenuActions(const UEdGraph* CurrentGraph, const UEdGraphNode* InGraphNode, const UEdGraphPin* InGraphPin, class FMenuBuilder* MenuBuilder, bool bIsDebugging) const
+void UEdGraphSchema::GetContextMenuActions(UToolMenu* Menu, UGraphNodeContextMenuContext* Context) const
 {
-#if WITH_EDITOR
-	FGraphNodeContextMenuBuilder Context(CurrentGraph, InGraphNode, InGraphPin, MenuBuilder, bIsDebugging);
-
-	if (Context.Node)
-	{
-		Context.Node->GetContextMenuActions(Context);
-	}
-
-	if (InGraphNode)
-	{
-		// Helper to do the node comment editing
-		struct Local
-		{
-			// Called by the EditableText widget to get the current comment for the node
-			static FString GetNodeComment(TWeakObjectPtr<UEdGraphNode> NodeWeakPtr)
-			{
-				if (UEdGraphNode* SelectedNode = NodeWeakPtr.Get())
-				{
-					return SelectedNode->NodeComment;
-				}
-				return FString();
-			}
-
-			// Called by the EditableText widget when the user types a new comment for the selected node
-			static void OnNodeCommentTextCommitted(const FText& NewText, ETextCommit::Type CommitInfo, TWeakObjectPtr<UEdGraphNode> NodeWeakPtr)
-			{
-				// Apply the change to the selected actor
-				UEdGraphNode* SelectedNode = NodeWeakPtr.Get();
-				FString NewString = NewText.ToString();
-				if (SelectedNode && !SelectedNode->NodeComment.Equals( NewString, ESearchCase::CaseSensitive ))
-				{
-					// send property changed events
-					const FScopedTransaction Transaction( LOCTEXT("EditNodeComment", "Change Node Comment") );
-					SelectedNode->Modify();
-					UProperty* NodeCommentProperty = FindField<UProperty>(SelectedNode->GetClass(), "NodeComment");
-					if(NodeCommentProperty != nullptr)
-					{
-						SelectedNode->PreEditChange(NodeCommentProperty);
+	// Instructions: Implement GetParentContextMenuName() with return NAME_None in classes that do not want UEdGraphSchema's menu entries
 
-						SelectedNode->NodeComment = NewString;
-						SelectedNode->SetMakeCommentBubbleVisible(true);
-
-						FPropertyChangedEvent NodeCommentPropertyChangedEvent(NodeCommentProperty);
-						SelectedNode->PostEditChangeProperty(NodeCommentPropertyChangedEvent);
-					}
-				}
+	// Note: Menu for UEdGraphSchema is registered in an editor module to reduce dependencies in the engine module
+}
 
-				// Only dismiss all menus if the text was committed via some user action.
-				// ETextCommit::Default implies that focus was switched by some other means. If this is because a submenu was opened, we don't want to close all the menus as a consequence.
-				if (CommitInfo != ETextCommit::Default)
-				{
-					FSlateApplication::Get().DismissAllMenus();
-				}
-			}
-		};
+FName UEdGraphSchema::GetContextMenuName() const
+{
+	return GetContextMenuName(GetClass());
+}
 
-		if( !InGraphPin )
+FName UEdGraphSchema::GetParentContextMenuName() const
+{
+#if WITH_EDITOR
+	if (GetClass() != UEdGraphSchema::StaticClass())
+	{
+		if (UClass* SuperClass = GetClass()->GetSuperClass())
 		{
-			int32 SelectionCount = GetNodeSelectionCount(CurrentGraph);
-		
-			if( SelectionCount == 1 )
-			{
-				// Node comment area
-				TSharedRef<SHorizontalBox> NodeCommentBox = SNew(SHorizontalBox);
-
-				MenuBuilder->BeginSection("GraphNodeComment", LOCTEXT("NodeCommentMenuHeader", "Node Comment"));
-				{
-					MenuBuilder->AddWidget(NodeCommentBox, FText::GetEmpty() );
-				}
-				TWeakObjectPtr<UEdGraphNode> SelectedNodeWeakPtr = MakeWeakObjectPtr(const_cast<UEdGraphNode*>(InGraphNode));
-
-				FText NodeCommentText;
-				if ( UEdGraphNode* SelectedNode = SelectedNodeWeakPtr.Get() )
-				{
-					NodeCommentText = FText::FromString( SelectedNode->NodeComment );
-				}
-
-			const FSlateBrush* NodeIcon = FCoreStyle::Get().GetDefaultBrush();//@TODO: FActorIconFinder::FindIconForActor(SelectedActors(0).Get());
-
-
-				// Comment label
-				NodeCommentBox->AddSlot()
-				.VAlign(VAlign_Center)
-				.FillWidth( 1.0f )
-				[
-					SNew(SMultiLineEditableTextBox)
-					.Text( NodeCommentText )
-					.ToolTipText(LOCTEXT("NodeComment_ToolTip", "Comment for this node"))
-					.OnTextCommitted_Static(&Local::OnNodeCommentTextCommitted, SelectedNodeWeakPtr)
-					.SelectAllTextWhenFocused( true )
-					.RevertTextOnEscape( true )
-					.ModiferKeyForNewLine( EModifierKey::Control )
-				];
-				MenuBuilder->EndSection();
-			}
-			else if( SelectionCount > 1 )
-			{
-				struct SCommentUtility
-				{
-					static void CreateComment(const UEdGraphSchema* Schema, UEdGraph* Graph)
-					{
-						if( Schema && Graph )
-						{
-							TSharedPtr<FEdGraphSchemaAction> Action = Schema->GetCreateCommentAction();
-
-							if( Action.IsValid() )
-							{
-								Action->PerformAction(Graph, nullptr, FVector2D());
-							}
-						}
-					}
-				};
-
-				MenuBuilder->BeginSection("SchemaActionComment", LOCTEXT("MultiCommentHeader", "Comment Group"));
-				MenuBuilder->AddMenuEntry(	LOCTEXT("MultiCommentDesc", "Create Comment from Selection"),
-											LOCTEXT("CommentToolTip", "Create a resizable comment box around selection."),
-											FSlateIcon(), 
-											FUIAction( FExecuteAction::CreateStatic( SCommentUtility::CreateComment, this, const_cast<UEdGraph*>( CurrentGraph ))));
-				MenuBuilder->EndSection();
-			}
+			return GetContextMenuName(SuperClass);
 		}
 	}
+#endif
+	return NAME_None;
+}
+
+FName UEdGraphSchema::GetContextMenuName(UClass* InClass)
+{
+#if WITH_EDITOR
+	return FName(*(FString(TEXT("GraphEditor.GraphContextMenu.")) + InClass->GetName()));
+#else
+	return NAME_None;
 #endif
 }
 
