Blueprints: Added the ability to hook into UK2Node::RenameUserDefinedPin

  - Uses a sparse object annotation to ensure minimal overhead for any K2Node instance that does not have this bound

diff --git a/Engine/Source/Editor/BlueprintGraph/Classes/K2Node.h b/Engine/Source/Editor/BlueprintGraph/Classes/K2Node.h
--- a/Engine/Source/Editor/BlueprintGraph/Classes/K2Node.h
+++ b/Engine/Source/Editor/BlueprintGraph/Classes/K2Node.h
@@ -18,6 +18,8 @@ class UEdGraph;
 class UEdGraphPin;
 class UEdGraphSchema;
 
+DECLARE_MULTICAST_DELEGATE_ThreeParams(FOnUserDefinedPinRenamed, UK2Node*, FName, FName);
+
 /** Helper structure to cache old data for optional pins so the data can be restored during reconstruction */
 struct FOldOptionalPinSettings
 {
@@ -315,15 +317,18 @@ class UK2Node : public UEdGraphNode
 	// Can this node be created under the specified schema
 	BLUEPRINTGRAPH_API virtual bool CanCreateUnderSpecifiedSchema(const UEdGraphSchema* DesiredSchema) const override;
 
-	// Renames an existing pin on the node.
-	BLUEPRINTGRAPH_API virtual ERenamePinResult RenameUserDefinedPin(const FName OldName, const FName NewName, bool bTest = false);
-
 	// Returns which dynamic binding class (if any) to use for this node
 	BLUEPRINTGRAPH_API virtual UClass* GetDynamicBindingClass() const { return NULL; }
 
 	// Puts information about this node into the dynamic binding object
 	BLUEPRINTGRAPH_API virtual void RegisterDynamicBinding(UDynamicBlueprintBinding* BindingObject) const { };
 
+	// Renames a user defined pin and broadcasts the result to any OnUserDefinedPinRenamed listeners on success (if bTest is false)
+	BLUEPRINTGRAPH_API ERenamePinResult RenameUserDefinedPin(const FName OldName, const FName NewName, bool bTest = false);
+
+	/** Retrieves a delegate that is called when a user-defined pin on this node is renamed */
+	BLUEPRINTGRAPH_API FOnUserDefinedPinRenamed& OnUserDefinedPinRenamed();
+
 	/**
 	 * Handles inserting the node between the FromPin and what the FromPin was original connected to
 	 *
@@ -426,6 +431,12 @@ protected:
 	/** Determines what the possible redirect pin names are **/
 	BLUEPRINTGRAPH_API virtual void GetRedirectPinNames(const UEdGraphPin& Pin, TArray<FString>& RedirectPinNames) const;
 
+	// Implementation function that renames an existing pin on the node. Does not broadcast notifications.
+	BLUEPRINTGRAPH_API virtual ERenamePinResult RenameUserDefinedPinImpl(const FName OldName, const FName NewName, bool bTest);
+
+	/** Attempt to broadcast an event for a user-defined pin being renamed on this node */
+	BLUEPRINTGRAPH_API void BroadcastUserDefinedPinRenamed(FName OldName, FName NewName);
+
 	/** 
 	 * Searches ParamRedirect Map and find if there is matching new param
 	 * 
diff --git a/Engine/Source/Editor/BlueprintGraph/Classes/K2Node_Tunnel.h b/Engine/Source/Editor/BlueprintGraph/Classes/K2Node_Tunnel.h
--- a/Engine/Source/Editor/BlueprintGraph/Classes/K2Node_Tunnel.h
+++ b/Engine/Source/Editor/BlueprintGraph/Classes/K2Node_Tunnel.h
@@ -60,7 +60,7 @@ class BLUEPRINTGRAPH_API UK2Node_Tunnel : public UK2Node_EditablePinBase
 	//~ Begin UK2Node_EditablePinBase Interface.
 	virtual UEdGraphPin* CreatePinFromUserDefinition(const TSharedPtr<FUserPinInfo> NewPinInfo) override;
 	virtual bool CanModifyExecutionWires() override;
-	virtual ERenamePinResult RenameUserDefinedPin(const FName OldName, const FName NewName, bool bTest = false) override;
+	virtual ERenamePinResult RenameUserDefinedPinImpl(const FName OldName, const FName NewName, bool bTest = false) override;
 	virtual bool CanUseRefParams() const override { return true; }
 	virtual bool CanCreateUserDefinedPin(const FEdGraphPinType& InPinType, EEdGraphPinDirection InDesiredDirection, FText& OutErrorMessage) override;
 	virtual bool ModifyUserDefinedPinDefaultValue(TSharedPtr<FUserPinInfo> PinInfo, const FString& NewDefaultValue) override;
diff --git a/Engine/Source/Editor/BlueprintGraph/Private/K2Node.cpp b/Engine/Source/Editor/BlueprintGraph/Private/K2Node.cpp
--- a/Engine/Source/Editor/BlueprintGraph/Private/K2Node.cpp
+++ b/Engine/Source/Editor/BlueprintGraph/Private/K2Node.cpp
@@ -24,6 +24,7 @@
 #include "PropertyCustomizationHelpers.h"
 
 #include "ObjectEditorUtils.h"
+#include "UObject/UObjectAnnotation.h"
 #include "UObject/FrameworkObjectVersion.h"
 
 #define LOCTEXT_NAMESPACE "K2Node"
@@ -31,6 +32,21 @@
 // File-Scoped Globals
 static const uint32 MaxArrayPinTooltipLineCount = 10;
 
+namespace UK2Node_Private
+{
+	struct FPinRenamedAnnotation
+	{
+		mutable FOnUserDefinedPinRenamed PinRenamedEvent;
+		bool bIsDefault = true;
+
+		bool IsDefault() const
+		{
+			return bIsDefault;
+		}
+	};
+	FUObjectAnnotationSparse<FPinRenamedAnnotation, true> GOnUserDefinedPinRenamedAnnotation;
+}
+
 /////////////////////////////////////////////////////
 // UK2Node
 
@@ -1285,6 +1301,46 @@ FLinearColor UK2Node::GetNodeTitleColor() const
 }
 
 ERenamePinResult UK2Node::RenameUserDefinedPin(const FName OldName, const FName NewName, bool bTest)
+{
+	ERenamePinResult Result = RenameUserDefinedPinImpl(OldName, NewName, bTest);
+
+	// If we actually applied the change (!bTest) and it was successful, broadcast the change notification
+	if (bTest == false && Result == ERenamePinResult_Success)
+	{
+		BroadcastUserDefinedPinRenamed(OldName, NewName);
+	}
+
+	return Result;
+}
+
+void UK2Node::BroadcastUserDefinedPinRenamed(FName OldName, FName NewName)
+{
+	using namespace UK2Node_Private;
+
+	// We're not concerned with thread safety here, but the annotation API forces thread safety by returning copies of data
+	if (const FPinRenamedAnnotation* Annotation = GOnUserDefinedPinRenamedAnnotation.GetAnnotationMap().Find(this))
+	{
+		Annotation->PinRenamedEvent.Broadcast(this, OldName, NewName);
+	}
+}
+
+FOnUserDefinedPinRenamed& UK2Node::OnUserDefinedPinRenamed()
+{
+	using namespace UK2Node_Private;
+
+	// We're not concerned with thread safety here, but the annotation API forces thread safety by returning copies of data
+	const TMap<const UObjectBase*,FPinRenamedAnnotation>& Map = GOnUserDefinedPinRenamedAnnotation.GetAnnotationMap();
+	if (!Map.Contains(this))
+	{
+		FPinRenamedAnnotation NewEvent;
+		NewEvent.bIsDefault = false;
+		GOnUserDefinedPinRenamedAnnotation.AddAnnotation(this, NewEvent);
+	}
+
+	return Map.FindChecked(this).PinRenamedEvent;
+}
+
+ERenamePinResult UK2Node::RenameUserDefinedPinImpl(const FName OldName, const FName NewName, bool bTest)
 {
 	UEdGraphPin* Pin = nullptr;
 	for (UEdGraphPin* CurrentPin : Pins)
diff --git a/Engine/Source/Editor/BlueprintGraph/Private/K2Node_Tunnel.cpp b/Engine/Source/Editor/BlueprintGraph/Private/K2Node_Tunnel.cpp
--- a/Engine/Source/Editor/BlueprintGraph/Private/K2Node_Tunnel.cpp
+++ b/Engine/Source/Editor/BlueprintGraph/Private/K2Node_Tunnel.cpp
@@ -235,9 +235,9 @@ bool UK2Node_Tunnel::CanModifyExecutionWires()
 	return true;
 }
 
-ERenamePinResult UK2Node_Tunnel::RenameUserDefinedPin(const FName OldName, const FName NewName, bool bTest)
+ERenamePinResult UK2Node_Tunnel::RenameUserDefinedPinImpl(const FName OldName, const FName NewName, bool bTest)
 {
-	const ERenamePinResult ThisNodeResult = Super::RenameUserDefinedPin(OldName, NewName, bTest);
+	const ERenamePinResult ThisNodeResult = Super::RenameUserDefinedPinImpl(OldName, NewName, bTest);
 	if(ERenamePinResult::ERenamePinResult_NameCollision == ThisNodeResult)
 	{
 		return ERenamePinResult::ERenamePinResult_NameCollision;
