Blueprints: All blueprints contained within assets are now allowable in the blueprint action manager

diff --git a/Engine/Source/Editor/BlueprintGraph/Private/BlueprintActionDatabase.cpp b/Engine/Source/Editor/BlueprintGraph/Private/BlueprintActionDatabase.cpp
--- a/Engine/Source/Editor/BlueprintGraph/Private/BlueprintActionDatabase.cpp
+++ b/Engine/Source/Editor/BlueprintGraph/Private/BlueprintActionDatabase.cpp
@@ -1035,10 +1035,19 @@ static bool BlueprintActionDatabaseImpl::IsObjectValidForDatabase(UObject const*
 	{
 		bReturn = true;
 	}
-	else if(UBlueprint const* Blueprint = Cast<UBlueprint>(Object))
+	else if(Object->IsA<UBlueprint>())
 	{
-		// Level scripts are sometimes not assets because they have not been saved yet, but they are still valid for the database.
-		bReturn = FBlueprintEditorUtils::IsLevelScriptBlueprint(Blueprint);
+		// If this is a blueprint contained within an asset, we can include it in the action database
+		UObject* PotentialAsset = Object->GetOuter();
+		while (PotentialAsset)
+		{
+			if (PotentialAsset->IsAsset())
+			{
+				bReturn = true;
+				break;
+			}
+			PotentialAsset = PotentialAsset->GetOuter();
+		}
 	}
 	else if(UWorld const* World = Cast<UWorld>(Object))
 	{
diff --git a/Engine/Source/Editor/UnrealEd/Private/PackageTools.cpp b/Engine/Source/Editor/UnrealEd/Private/PackageTools.cpp
--- a/Engine/Source/Editor/UnrealEd/Private/PackageTools.cpp
+++ b/Engine/Source/Editor/UnrealEd/Private/PackageTools.cpp
@@ -833,16 +833,12 @@ UPackageTools::UPackageTools(const FObjectInitializer& ObjectInitializer)
 		{
 			GEngine->NotifyToolsOfObjectReplacement(InPackageReloadedEvent->GetRepointedObjects());
 
-			// Notify any Blueprint assets that are about to be unloaded.
+			// Notify any Blueprints that are about to be unloaded.
 			ForEachObjectWithOuter(InPackageReloadedEvent->GetOldPackage(), [&](UObject* InObject)
 			{
-				if (InObject->IsAsset())
+				if (UBlueprint* BP = Cast<UBlueprint>(InObject))
 				{
-					// Notify about any BP assets that are about to be unloaded
-					if (UBlueprint* BP = Cast<UBlueprint>(InObject))
-					{
-						BP->ClearEditorReferences();
-					}
+					BP->ClearEditorReferences();
 				}
 			}, false, RF_Transient, EInternalObjectFlags::PendingKill);
 		}
