Don't localize internal pins names of UK2Node_GetDataTableRow

#rnx

diff --git a/Engine/Source/Editor/BlueprintGraph/Private/K2Node_GetDataTableRow.cpp b/Engine/Source/Editor/BlueprintGraph/Private/K2Node_GetDataTableRow.cpp
--- a/Engine/Source/Editor/BlueprintGraph/Private/K2Node_GetDataTableRow.cpp
+++ b/Engine/Source/Editor/BlueprintGraph/Private/K2Node_GetDataTableRow.cpp
@@ -15,16 +15,14 @@
 
 #define LOCTEXT_NAMESPACE "K2Node_GetDataTableRow"
 
-struct UK2Node_GetDataTableRowHelper
+namespace GetDataTableRowHelper
 {
-	static FName DataTablePinName;
-	static FName RowNamePinName;
-	static FName RowNotFoundPinName;
-};
 
-FName UK2Node_GetDataTableRowHelper::DataTablePinName(*LOCTEXT("DataTablePinName","DataTable").ToString());
-FName UK2Node_GetDataTableRowHelper::RowNotFoundPinName(*LOCTEXT("RowNotFoundPinName","RowNotFound").ToString());
-FName UK2Node_GetDataTableRowHelper::RowNamePinName(*LOCTEXT("RowNamePinName","RowName").ToString());
+const FName DataTablePinName = "DataTable";
+const FName RowNotFoundPinName = "RowNotFound";
+const FName RowNamePinName = "RowName";
+
+}
 
 UK2Node_GetDataTableRow::UK2Node_GetDataTableRow(const FObjectInitializer& ObjectInitializer)
 	: Super(ObjectInitializer)
@@ -40,14 +38,14 @@ void UK2Node_GetDataTableRow::AllocateDefaultPins()
 	CreatePin(EGPD_Input, UEdGraphSchema_K2::PC_Exec, UEdGraphSchema_K2::PN_Execute);
 	UEdGraphPin* RowFoundPin = CreatePin(EGPD_Output, UEdGraphSchema_K2::PC_Exec, UEdGraphSchema_K2::PN_Then);
 	RowFoundPin->PinFriendlyName = LOCTEXT("GetDataTableRow Row Found Exec pin", "Row Found");
-	CreatePin(EGPD_Output, UEdGraphSchema_K2::PC_Exec, UK2Node_GetDataTableRowHelper::RowNotFoundPinName);
+	CreatePin(EGPD_Output, UEdGraphSchema_K2::PC_Exec, GetDataTableRowHelper::RowNotFoundPinName);
 
 	// Add DataTable pin
-	UEdGraphPin* DataTablePin = CreatePin(EGPD_Input, UEdGraphSchema_K2::PC_Object, UDataTable::StaticClass(), UK2Node_GetDataTableRowHelper::DataTablePinName);
+	UEdGraphPin* DataTablePin = CreatePin(EGPD_Input, UEdGraphSchema_K2::PC_Object, UDataTable::StaticClass(), GetDataTableRowHelper::DataTablePinName);
 	SetPinToolTip(*DataTablePin, LOCTEXT("DataTablePinDescription", "The DataTable you want to retreive a row from"));
 
 	// Row Name pin
-	UEdGraphPin* RowNamePin = CreatePin(EGPD_Input, UEdGraphSchema_K2::PC_Name, UK2Node_GetDataTableRowHelper::RowNamePinName);
+	UEdGraphPin* RowNamePin = CreatePin(EGPD_Input, UEdGraphSchema_K2::PC_Name, GetDataTableRowHelper::RowNamePinName);
 	SetPinToolTip(*RowNamePin, LOCTEXT("RowNamePinDescription", "The name of the row to retrieve from the DataTable"));
 
 	// Result pin
@@ -234,7 +232,7 @@ bool UK2Node_GetDataTableRow::IsConnectionDisallowed(const UEdGraphPin* MyPin, c
 
 void UK2Node_GetDataTableRow::PinDefaultValueChanged(UEdGraphPin* ChangedPin) 
 {
-	if (ChangedPin && ChangedPin->PinName == UK2Node_GetDataTableRowHelper::DataTablePinName)
+	if (ChangedPin && ChangedPin->PinName == GetDataTableRowHelper::DataTablePinName)
 	{
 		RefreshOutputPinType();
 
@@ -276,7 +274,7 @@ UEdGraphPin* UK2Node_GetDataTableRow::GetDataTablePin(const TArray<UEdGraphPin*>
 	UEdGraphPin* Pin = nullptr;
 	for (UEdGraphPin* TestPin : *PinsToSearch)
 	{
-		if (TestPin && TestPin->PinName == UK2Node_GetDataTableRowHelper::DataTablePinName)
+		if (TestPin && TestPin->PinName == GetDataTableRowHelper::DataTablePinName)
 		{
 			Pin = TestPin;
 			break;
@@ -288,14 +286,14 @@ UEdGraphPin* UK2Node_GetDataTableRow::GetDataTablePin(const TArray<UEdGraphPin*>
 
 UEdGraphPin* UK2Node_GetDataTableRow::GetRowNamePin() const
 {
-	UEdGraphPin* Pin = FindPinChecked(UK2Node_GetDataTableRowHelper::RowNamePinName);
+	UEdGraphPin* Pin = FindPinChecked(GetDataTableRowHelper::RowNamePinName);
 	check(Pin->Direction == EGPD_Input);
 	return Pin;
 }
 
 UEdGraphPin* UK2Node_GetDataTableRow::GetRowNotFoundPin() const
 {
-	UEdGraphPin* Pin = FindPinChecked(UK2Node_GetDataTableRowHelper::RowNotFoundPinName);
+	UEdGraphPin* Pin = FindPinChecked(GetDataTableRowHelper::RowNotFoundPinName);
 	check(Pin->Direction == EGPD_Output);
 	return Pin;
 }
