Added Get/SetEditorProperty functions to Blueprints

diff --git a/Engine/Source/Editor/BlueprintGraph/Private/K2Node_EditorPropertyAccess.cpp b/Engine/Source/Editor/BlueprintGraph/Private/K2Node_EditorPropertyAccess.cpp
new file mode 100644
--- /dev/null
+++ b/Engine/Source/Editor/BlueprintGraph/Private/K2Node_EditorPropertyAccess.cpp
@@ -0,0 +1,332 @@
+// Copyright 1998-2019 Epic Games, Inc. All Rights Reserved.
+
+#include "K2Node_EditorPropertyAccess.h"
+#include "EdGraphSchema_K2.h"
+#include "K2Node_CallFunction.h"
+#include "K2Node_IfThenElse.h"
+#include "KismetCompiler.h"
+#include "Kismet2/BlueprintEditorUtils.h"
+#include "BlueprintNodeSpawner.h"
+#include "EditorCategoryUtils.h"
+#include "BlueprintActionDatabaseRegistrar.h"
+#include "Kismet/KismetSystemLibrary.h"
+
+#define LOCTEXT_NAMESPACE "K2Node_EditorPropertyAccess"
+
+namespace EditorPropertyAccessHelper
+{
+
+const FName ObjectPinName = "Object";
+const FName PropertyNamePinName = "PropertyName";
+const FName PropertyValuePinName = "PropertyValue";
+
+}
+
+void UK2Node_EditorPropertyAccessBase::AllocateDefaultPins()
+{
+	Super::AllocateDefaultPins();
+
+	// Add execution pins
+	CreatePin(EGPD_Input, UEdGraphSchema_K2::PC_Exec, UEdGraphSchema_K2::PN_Execute);
+	CreatePin(EGPD_Output, UEdGraphSchema_K2::PC_Exec, UEdGraphSchema_K2::PN_Then);
+
+	// Add Object pin
+	UEdGraphPin* ObjectPin = CreatePin(EGPD_Input, UEdGraphSchema_K2::PC_Object, UObject::StaticClass(), EditorPropertyAccessHelper::ObjectPinName);
+	SetPinToolTip(*ObjectPin, LOCTEXT("ObjectPinDescription", "The object you want to access a property value from"));
+
+	// Add Property Name pin
+	UEdGraphPin* PropertyNamePin = CreatePin(EGPD_Input, UEdGraphSchema_K2::PC_Name, EditorPropertyAccessHelper::PropertyNamePinName);
+	SetPinToolTip(*PropertyNamePin, LOCTEXT("PropertyNamePinDescription", "The name of the property to access from the object"));
+
+	// Add Property Value pin
+	AllocatePropertyValuePin();
+
+	// Add Result Pin
+	UEdGraphPin* ResultPin = CreatePin(EGPD_Output, UEdGraphSchema_K2::PC_Boolean, UEdGraphSchema_K2::PN_ReturnValue);
+	ResultPin->PinFriendlyName = LOCTEXT("ResultPinFriendlyName", "Success?");
+	SetPinToolTip(*ResultPin, LOCTEXT("ResultPinDescription", "Whether the property value was found"));
+}
+
+void UK2Node_EditorPropertyAccessBase::SetPinToolTip(UEdGraphPin& MutatablePin, const FText& PinDescription) const
+{
+	MutatablePin.PinToolTip = UEdGraphSchema_K2::TypeToText(MutatablePin.PinType).ToString();
+
+	const UEdGraphSchema_K2* K2Schema = Cast<const UEdGraphSchema_K2>(GetSchema());
+	if (K2Schema)
+	{
+		MutatablePin.PinToolTip += TEXT(" ");
+		MutatablePin.PinToolTip += K2Schema->GetPinDisplayName(&MutatablePin).ToString();
+	}
+
+	MutatablePin.PinToolTip += FString(TEXT("\n")) + PinDescription.ToString();
+}
+
+void UK2Node_EditorPropertyAccessBase::GetMenuActions(FBlueprintActionDatabaseRegistrar& ActionRegistrar) const
+{
+	// actions get registered under specific object-keys; the idea is that 
+	// actions might have to be updated (or deleted) if their object-key is  
+	// mutated (or removed)... here we use the node's class (so if the node 
+	// type disappears, then the action should go with it)
+	UClass* ActionKey = GetClass();
+	// to keep from needlessly instantiating a UBlueprintNodeSpawner, first   
+	// check to make sure that the registrar is looking for actions of this type
+	// (could be regenerating actions for a specific asset, and therefore the 
+	// registrar would only accept actions corresponding to that asset)
+	if (ActionRegistrar.IsOpenForRegistration(ActionKey))
+	{
+		UBlueprintNodeSpawner* NodeSpawner = UBlueprintNodeSpawner::Create(GetClass());
+		check(NodeSpawner != nullptr);
+
+		ActionRegistrar.AddBlueprintAction(ActionKey, NodeSpawner);
+	}
+}
+
+FText UK2Node_EditorPropertyAccessBase::GetMenuCategory() const
+{
+	return FEditorCategoryUtils::GetCommonCategory(FCommonEditorCategory::Utilities);
+}
+
+bool UK2Node_EditorPropertyAccessBase::CanPasteHere(const UEdGraph* TargetGraph) const
+{
+	bool bCanPaste = Super::CanPasteHere(TargetGraph);
+	if (bCanPaste)
+	{
+		bCanPaste &= FBlueprintEditorUtils::IsEditorUtilityBlueprint(FBlueprintEditorUtils::FindBlueprintForGraphChecked(TargetGraph));
+	}
+	return bCanPaste;
+}
+
+bool UK2Node_EditorPropertyAccessBase::IsActionFilteredOut(const FBlueprintActionFilter& Filter)
+{
+	bool bIsFilteredOut = Super::IsActionFilteredOut(Filter);
+	if (!bIsFilteredOut)
+	{
+		for (UEdGraph* TargetGraph : Filter.Context.Graphs)
+		{
+			bIsFilteredOut |= !CanPasteHere(TargetGraph);
+		}
+	}
+	return bIsFilteredOut;
+}
+
+bool UK2Node_EditorPropertyAccessBase::IsConnectionDisallowed(const UEdGraphPin* MyPin, const UEdGraphPin* OtherPin, FString& OutReason) const
+{
+	if (MyPin->PinName == EditorPropertyAccessHelper::PropertyValuePinName && MyPin->PinType.PinCategory == UEdGraphSchema_K2::PC_Wildcard)
+	{
+		if (OtherPin->PinType.PinCategory == UEdGraphSchema_K2::PC_Exec)
+		{
+			OutReason = TEXT("Cannot connect the property value to an execution pin");
+			return true;
+		}
+	}
+
+	return false;
+}
+
+UEdGraphPin* UK2Node_EditorPropertyAccessBase::GetThenPin()const
+{
+	UEdGraphPin* Pin = FindPinChecked(UEdGraphSchema_K2::PN_Then);
+	check(Pin->Direction == EGPD_Output);
+	return Pin;
+}
+
+UEdGraphPin* UK2Node_EditorPropertyAccessBase::GetObjectPin() const
+{
+	UEdGraphPin* Pin = FindPinChecked(EditorPropertyAccessHelper::ObjectPinName);
+	check(Pin->Direction == EGPD_Input);
+	return Pin;
+}
+
+UEdGraphPin* UK2Node_EditorPropertyAccessBase::GetPropertyNamePin() const
+{
+	UEdGraphPin* Pin = FindPinChecked(EditorPropertyAccessHelper::PropertyNamePinName);
+	check(Pin->Direction == EGPD_Input);
+	return Pin;
+}
+
+UEdGraphPin* UK2Node_EditorPropertyAccessBase::GetResultPin() const
+{
+	UEdGraphPin* Pin = FindPinChecked(UEdGraphSchema_K2::PN_ReturnValue);
+	check(Pin->Direction == EGPD_Output);
+	return Pin;
+}
+
+void UK2Node_EditorPropertyAccessBase::RefreshPropertyValuePin()
+{
+	UEdGraphPin* PropertyValuePin = GetPropertyValuePin();
+
+	FEdGraphPinType NewPinType;
+	if (PropertyValuePin->LinkedTo.Num() > 0)
+	{
+		NewPinType = PropertyValuePin->LinkedTo[0]->PinType;
+	}
+	else
+	{
+		NewPinType.PinCategory = UEdGraphSchema_K2::PC_Wildcard;
+	}
+
+	if (PropertyValuePin->PinType != NewPinType)
+	{
+		PropertyValuePin->PinType = NewPinType;
+
+		GetGraph()->NotifyGraphChanged();
+
+		UBlueprint* Blueprint = GetBlueprint();
+		if (!Blueprint->bBeingCompiled)
+		{
+			FBlueprintEditorUtils::MarkBlueprintAsModified(Blueprint);
+			Blueprint->BroadcastChanged();
+		}
+	}
+}
+
+FSlateIcon UK2Node_EditorPropertyAccessBase::GetIconAndTint(FLinearColor& OutColor) const
+{
+	OutColor = GetNodeTitleColor();
+	static FSlateIcon Icon("EditorStyle", "Kismet.AllClasses.FunctionIcon");
+	return Icon;
+}
+
+void UK2Node_EditorPropertyAccessBase::PostReconstructNode()
+{
+	Super::PostReconstructNode();
+
+	RefreshPropertyValuePin();
+}
+
+void UK2Node_EditorPropertyAccessBase::EarlyValidation(class FCompilerResultsLog& MessageLog) const
+{
+	Super::EarlyValidation(MessageLog);
+
+	const UEdGraphPin* ObjectPin = GetObjectPin();
+	const UEdGraphPin* PropertyNamePin = GetPropertyNamePin();
+	const UEdGraphPin* PropertyValuePin = GetPropertyValuePin();
+	if (!ObjectPin || !PropertyNamePin || !PropertyValuePin)
+	{
+		MessageLog.Error(*LOCTEXT("MissingPins", "Missing pins in @@").ToString(), this);
+		return;
+	}
+
+	if (ObjectPin->LinkedTo.Num() == 0 && !ObjectPin->DefaultObject)
+	{
+		MessageLog.Error(*LOCTEXT("UnsetObject", "No object set on @@").ToString(), this);
+	}
+
+	if (PropertyNamePin->LinkedTo.Num() == 0 && FName(*PropertyNamePin->DefaultValue).IsNone())
+	{
+		MessageLog.Error(*LOCTEXT("UnsetPropertyName", "No property name set on @@").ToString(), this);
+	}
+}
+
+void UK2Node_EditorPropertyAccessBase::NotifyPinConnectionListChanged(UEdGraphPin* Pin)
+{
+	Super::NotifyPinConnectionListChanged(Pin);
+
+	if (Pin->PinName == EditorPropertyAccessHelper::PropertyValuePinName)
+	{
+		RefreshPropertyValuePin();
+	}
+}
+
+void UK2Node_EditorPropertyAccessBase::ExpandNode(FKismetCompilerContext& CompilerContext, UEdGraph* SourceGraph)
+{
+	Super::ExpandNode(CompilerContext, SourceGraph);
+
+	const FName FunctionName = GetUnderlyingFunctionName();
+
+	// Add a CallFunction node for the underlying function
+	UK2Node_CallFunction* CallFunctionNode = CompilerContext.SpawnIntermediateNode<UK2Node_CallFunction>(this, SourceGraph);
+	CallFunctionNode->FunctionReference.SetExternalMember(FunctionName, UKismetSystemLibrary::StaticClass());
+	CallFunctionNode->AllocateDefaultPins();
+
+	// Re-wire the execution pins
+	CompilerContext.MovePinLinksToIntermediate(*GetExecPin(), *CallFunctionNode->GetExecPin());
+	CompilerContext.MovePinLinksToIntermediate(*GetThenPin(), *CallFunctionNode->GetThenPin());
+
+	// Re-wire the Object pin to the function input
+	{
+		UEdGraphPin* FunctionObjectPin = CallFunctionNode->FindPinChecked(TEXT("Object"));
+		CompilerContext.MovePinLinksToIntermediate(*GetObjectPin(), *FunctionObjectPin);
+	}
+
+	// Re-wire the PropertyName pin to the function input
+	{
+		UEdGraphPin* FunctionPropertyNamePin = CallFunctionNode->FindPinChecked(TEXT("PropertyName"));
+		CompilerContext.MovePinLinksToIntermediate(*GetPropertyNamePin(), *FunctionPropertyNamePin);
+	}
+
+	// Re-wire the ValuePtr pin to the function output, and copy its type over
+	{
+		UEdGraphPin* PropertyValuePin = GetPropertyValuePin();
+		UEdGraphPin* FunctionValuePtrPin = CallFunctionNode->FindPinChecked(TEXT("ValuePtr"));
+		CompilerContext.MovePinLinksToIntermediate(*PropertyValuePin, *FunctionValuePtrPin);
+		FunctionValuePtrPin->PinType = PropertyValuePin->PinType;
+	}
+
+	// Re-wire the result pin
+	{
+		UEdGraphPin* FunctionReturnPin = CallFunctionNode->FindPinChecked(UEdGraphSchema_K2::PN_ReturnValue);
+		CompilerContext.MovePinLinksToIntermediate(*GetResultPin(), *FunctionReturnPin);
+	}
+
+	// Disconnect this node
+	BreakAllNodeLinks();
+}
+
+void UK2Node_GetEditorProperty::AllocatePropertyValuePin()
+{
+	UEdGraphPin* PropertyValuePin = CreatePin(EGPD_Output, UEdGraphSchema_K2::PC_Wildcard, EditorPropertyAccessHelper::PropertyValuePinName);
+	SetPinToolTip(*PropertyValuePin, LOCTEXT("GetEditorProperty_PropertyValueDescription", "The returned property value, if found"));
+}
+
+UEdGraphPin* UK2Node_GetEditorProperty::GetPropertyValuePin() const
+{
+	UEdGraphPin* Pin = FindPinChecked(EditorPropertyAccessHelper::PropertyValuePinName);
+	check(Pin->Direction == EGPD_Output);
+	return Pin;
+}
+
+FText UK2Node_GetEditorProperty::GetNodeTitle(ENodeTitleType::Type TitleType) const
+{
+	return LOCTEXT("GetEditorProperty_NodeTitle", "Get Editor Property");
+}
+
+FText UK2Node_GetEditorProperty::GetTooltipText() const
+{
+	return LOCTEXT("GetEditorProperty_NodeTooltip", "Attempts to retrieve the value of a named property from the given object");
+}
+
+FName UK2Node_GetEditorProperty::GetUnderlyingFunctionName() const
+{
+	return GET_FUNCTION_NAME_CHECKED(UKismetSystemLibrary, GetEditorProperty);
+}
+
+void UK2Node_SetEditorProperty::AllocatePropertyValuePin()
+{
+	UEdGraphPin* PropertyValuePin = CreatePin(EGPD_Input, UEdGraphSchema_K2::PC_Wildcard, EditorPropertyAccessHelper::PropertyValuePinName);
+	SetPinToolTip(*PropertyValuePin, LOCTEXT("SetEditorProperty_PropertyValueDescription", "The property value to set"));
+}
+
+UEdGraphPin* UK2Node_SetEditorProperty::GetPropertyValuePin() const
+{
+	UEdGraphPin* Pin = FindPinChecked(EditorPropertyAccessHelper::PropertyValuePinName);
+	check(Pin->Direction == EGPD_Input);
+	return Pin;
+}
+
+FText UK2Node_SetEditorProperty::GetNodeTitle(ENodeTitleType::Type TitleType) const
+{
+	return LOCTEXT("SetEditorProperty_NodeTitle", "Set Editor Property");
+}
+
+FText UK2Node_SetEditorProperty::GetTooltipText() const
+{
+	return LOCTEXT("SetEditorProperty_NodeTooltip", "Attempts to set the value of a named property on the given object");
+}
+
+FName UK2Node_SetEditorProperty::GetUnderlyingFunctionName() const
+{
+	return GET_FUNCTION_NAME_CHECKED(UKismetSystemLibrary, SetEditorProperty);
+}
+
+#undef LOCTEXT_NAMESPACE
diff --git a/Engine/Source/Editor/BlueprintGraph/Private/K2Node_EditorPropertyAccess.h b/Engine/Source/Editor/BlueprintGraph/Private/K2Node_EditorPropertyAccess.h
new file mode 100644
--- /dev/null
+++ b/Engine/Source/Editor/BlueprintGraph/Private/K2Node_EditorPropertyAccess.h
@@ -0,0 +1,102 @@
+// Copyright 1998-2019 Epic Games, Inc. All Rights Reserved.
+
+
+#pragma once
+
+#include "CoreMinimal.h"
+#include "Textures/SlateIcon.h"
+#include "K2Node.h"
+#include "K2Node_EditorPropertyAccess.generated.h"
+
+class UEdGraph;
+class UK2Node_CallFunction;
+class FBlueprintActionDatabaseRegistrar;
+
+UCLASS(Abstract)
+class UK2Node_EditorPropertyAccessBase : public UK2Node
+{
+	GENERATED_BODY()
+
+public:
+	//~ Begin UEdGraphNode Interface.
+	virtual void AllocateDefaultPins() override;
+	virtual FSlateIcon GetIconAndTint(FLinearColor& OutColor) const override;
+	virtual void PostReconstructNode() override;
+	//~ End UEdGraphNode Interface.
+
+	//~ Begin UK2Node Interface
+	virtual void GetMenuActions(FBlueprintActionDatabaseRegistrar& ActionRegistrar) const override;
+	virtual FText GetMenuCategory() const override;
+	virtual bool CanPasteHere(const UEdGraph* TargetGraph) const override;
+	virtual bool IsActionFilteredOut(const class FBlueprintActionFilter& Filter) override;
+	virtual bool IsConnectionDisallowed(const UEdGraphPin* MyPin, const UEdGraphPin* OtherPin, FString& OutReason) const override;
+	virtual void EarlyValidation(class FCompilerResultsLog& MessageLog) const override;
+	virtual void NotifyPinConnectionListChanged(UEdGraphPin* Pin) override;
+	virtual void ExpandNode(class FKismetCompilerContext& CompilerContext, UEdGraph* SourceGraph) override;
+	//~ End UK2Node Interface
+
+	/** Get the then output pin */
+	UEdGraphPin* GetThenPin() const;
+	/** Get the object input pin */
+	UEdGraphPin* GetObjectPin() const;
+	/** Get the property name input pin */
+	UEdGraphPin* GetPropertyNamePin() const;
+	/** Get the property value output pin */
+	virtual UEdGraphPin* GetPropertyValuePin() const PURE_VIRTUAL(UK2Node_EditorPropertyAccessBase::GetPropertyValuePin, return nullptr;);
+	/** Get the result output pin */
+	UEdGraphPin* GetResultPin() const;
+
+protected:
+	/** Allocate the property value pin for this node */
+	virtual void AllocatePropertyValuePin() PURE_VIRTUAL(UK2Node_EditorPropertyAccessBase::AllocatePropertyValuePin, );
+
+	/** Get the function name from UKismetSystemLibrary that the CallFunction node should use for the underlying access function */
+	virtual FName GetUnderlyingFunctionName() const PURE_VIRTUAL(UK2Node_EditorPropertyAccessBase::GetUnderlyingFunctionName, return FName(););
+
+	/** Updates the type of the property value pin based on its connection type */
+	void RefreshPropertyValuePin();
+
+	/**
+	 * Takes the specified "MutatablePin" and sets its 'PinToolTip' field (according to the specified description)
+	 * 
+	 * @param   MutatablePin	The pin you want to set tool-tip text on
+	 * @param   PinDescription	A string describing the pin's purpose
+	 */
+	void SetPinToolTip(UEdGraphPin& MutatablePin, const FText& PinDescription) const;
+};
+
+UCLASS()
+class UK2Node_GetEditorProperty : public UK2Node_EditorPropertyAccessBase
+{
+	GENERATED_BODY()
+
+public:
+	//~ Begin UEdGraphNode Interface.
+	virtual FText GetNodeTitle(ENodeTitleType::Type TitleType) const override;
+	virtual FText GetTooltipText() const override;
+	//~ End UEdGraphNode Interface.
+
+	//~ Begin UK2Node_EditorPropertyAccessBase Interface.
+	virtual void AllocatePropertyValuePin() override;
+	virtual UEdGraphPin* GetPropertyValuePin() const override;
+	virtual FName GetUnderlyingFunctionName() const override;
+	//~ End UK2Node_EditorPropertyAccessBase Interface.
+};
+
+UCLASS()
+class UK2Node_SetEditorProperty : public UK2Node_EditorPropertyAccessBase
+{
+	GENERATED_BODY()
+
+public:
+	//~ Begin UEdGraphNode Interface.
+	virtual FText GetNodeTitle(ENodeTitleType::Type TitleType) const override;
+	virtual FText GetTooltipText() const override;
+	//~ End UEdGraphNode Interface.
+
+	//~ Begin UK2Node_EditorPropertyAccessBase Interface.
+	virtual void AllocatePropertyValuePin() override;
+	virtual UEdGraphPin* GetPropertyValuePin() const override;
+	virtual FName GetUnderlyingFunctionName() const override;
+	//~ End UK2Node_EditorPropertyAccessBase Interface.
+};
diff --git a/Engine/Source/Runtime/CoreUObject/Private/UObject/PropertyAccessUtil.cpp b/Engine/Source/Runtime/CoreUObject/Private/UObject/PropertyAccessUtil.cpp
--- a/Engine/Source/Runtime/CoreUObject/Private/UObject/PropertyAccessUtil.cpp
+++ b/Engine/Source/Runtime/CoreUObject/Private/UObject/PropertyAccessUtil.cpp
@@ -240,4 +240,25 @@ bool IsObjectTemplate(const UObject* InObject)
 	return InObject->IsTemplate() || InObject->IsAsset();
 }
 
+UProperty* FindPropertyByName(const FName InPropName, const UStruct* InStruct)
+{
+	UProperty* Prop = InStruct->FindPropertyByName(InPropName);
+
+	if (!Prop)
+	{
+		const FName NewPropName = UProperty::FindRedirectedPropertyName(const_cast<UStruct*>(InStruct), InPropName);
+		if (!NewPropName.IsNone())
+		{
+			Prop = InStruct->FindPropertyByName(NewPropName);
+		}
+	}
+
+	if (!Prop)
+	{
+		Prop = InStruct->CustomFindProperty(InPropName);
+	}
+
+	return Prop;
+}
+
 }
diff --git a/Engine/Source/Runtime/CoreUObject/Public/UObject/PropertyAccessUtil.h b/Engine/Source/Runtime/CoreUObject/Public/UObject/PropertyAccessUtil.h
--- a/Engine/Source/Runtime/CoreUObject/Public/UObject/PropertyAccessUtil.h
+++ b/Engine/Source/Runtime/CoreUObject/Public/UObject/PropertyAccessUtil.h
@@ -253,4 +253,14 @@ namespace PropertyAccessUtil
 	 * @return True if the object instance is considered a template.
 	 */
 	COREUOBJECT_API bool IsObjectTemplate(const UObject* InObject);
+
+	/**
+	 * Low-level function to find a property by its name, following redirectors if it cannot be found.
+	 *
+	 * @param InPropName The name of the property to find.
+	 * @param InStruct The struct that should contain the property.
+	 *
+	 * @return The found property, or null if the property cannot be found.
+	 */
+	COREUOBJECT_API UProperty* FindPropertyByName(const FName InPropName, const UStruct* InStruct);
 }
diff --git a/Engine/Source/Runtime/Engine/Classes/Kismet/KismetSystemLibrary.h b/Engine/Source/Runtime/Engine/Classes/Kismet/KismetSystemLibrary.h
--- a/Engine/Source/Runtime/Engine/Classes/Kismet/KismetSystemLibrary.h
+++ b/Engine/Source/Runtime/Engine/Classes/Kismet/KismetSystemLibrary.h
@@ -1733,6 +1733,22 @@ class ENGINE_API UKismetSystemLibrary : public UBlueprintFunctionLibrary
 	UFUNCTION(BlueprintPure, Category = "Utilities")
 	static bool IsUnattended();
 
+	// --- Property Access ---------------------------
+
+#if WITH_EDITOR
+	/** Read the value of a named property on the given object to the given value pointer */
+    UFUNCTION(BlueprintCallable, CustomThunk, Category = "Properties", meta=(CustomStructureParam="ValuePtr", BlueprintInternalUseOnly="true"))
+    static bool GetEditorProperty(UObject* Object, const FName PropertyName, int32& ValuePtr);
+	static bool Generic_GetEditorProperty(const UObject* Object, const UProperty* Property, void* ValuePtr);
+	DECLARE_FUNCTION(execGetEditorProperty);
+
+	/** Write the value of the given value pointer to a named property on the given object */
+    UFUNCTION(BlueprintCallable, CustomThunk, Category = "Properties", meta=(CustomStructureParam="ValuePtr", BlueprintInternalUseOnly="true"))
+    static bool SetEditorProperty(UObject* Object, const FName PropertyName, const int32& ValuePtr);
+	static bool Generic_SetEditorProperty(UObject* Object, const UProperty* Property, const void* ValuePtr);
+	DECLARE_FUNCTION(execSetEditorProperty);
+#endif
+
 	// --- Transactions ------------------------------
 
 	/**
diff --git a/Engine/Source/Runtime/Engine/Private/KismetSystemLibrary.cpp b/Engine/Source/Runtime/Engine/Private/KismetSystemLibrary.cpp
--- a/Engine/Source/Runtime/Engine/Private/KismetSystemLibrary.cpp
+++ b/Engine/Source/Runtime/Engine/Private/KismetSystemLibrary.cpp
@@ -10,6 +10,7 @@
 #include "Misc/Paths.h"
 #include "Misc/RuntimeErrors.h"
 #include "UObject/GCObject.h"
+#include "UObject/PropertyAccessUtil.h"
 #include "EngineGlobals.h"
 #include "Components/ActorComponent.h"
 #include "TimerManager.h"
@@ -40,9 +41,27 @@
 //////////////////////////////////////////////////////////////////////////
 // UKismetSystemLibrary
 
+#define LOCTEXT_NAMESPACE "UKismetSystemLibrary"
+
+const FName PropertyGetFailedWarning = FName("PropertyGetFailedWarning");
+const FName PropertySetFailedWarning = FName("PropertySetFailedWarning");
+
 UKismetSystemLibrary::UKismetSystemLibrary(const FObjectInitializer& ObjectInitializer)
 	: Super(ObjectInitializer)
 {
+	FBlueprintSupport::RegisterBlueprintWarning(
+		FBlueprintWarningDeclaration(
+			PropertyGetFailedWarning,
+			LOCTEXT("PropertyGetFailedWarning", "Property Get Failed")
+		)
+	);
+
+	FBlueprintSupport::RegisterBlueprintWarning(
+		FBlueprintWarningDeclaration(
+			PropertySetFailedWarning,
+			LOCTEXT("PropertySetFailedWarning", "Property Set Failed")
+		)
+	);
 }
 
 void UKismetSystemLibrary::StackTraceImpl(const FFrame& StackFrame)
@@ -2642,6 +2661,190 @@ bool UKismetSystemLibrary::IsUnattended()
 	return FApp::IsUnattended();
 }
 
+#if WITH_EDITOR
+
+bool UKismetSystemLibrary::GetEditorProperty(UObject* Object, const FName PropertyName, int32& OutValuePtr)
+{
+	// We should never hit this! Stubbed to avoid NoExport on the class.
+	check(0);
+	return false;
+}
+
+bool UKismetSystemLibrary::Generic_GetEditorProperty(const UObject* Object, const UProperty* Property, void* ValuePtr)
+{
+	const EPropertyAccessResultFlags AccessResult = PropertyAccessUtil::GetPropertyValue_Object(Property, Object, ValuePtr, INDEX_NONE);
+
+	if (EnumHasAnyFlags(AccessResult, EPropertyAccessResultFlags::PermissionDenied))
+	{
+		if (EnumHasAnyFlags(AccessResult, EPropertyAccessResultFlags::AccessProtected))
+		{
+			FFrame::KismetExecutionMessage(*FString::Printf(TEXT("Property '%s' on '%s' (%s) is protected and cannot be read"), *Property->GetName(), *Object->GetPathName(), *Object->GetClass()->GetName()), ELogVerbosity::Warning, PropertyGetFailedWarning);
+			return false;
+		}
+
+		FFrame::KismetExecutionMessage(*FString::Printf(TEXT("Property '%s' on '%s' (%s) cannot be read"), *Property->GetName(), *Object->GetPathName(), *Object->GetClass()->GetName()), ELogVerbosity::Warning, PropertyGetFailedWarning);
+		return false;
+	}
+
+	if (EnumHasAnyFlags(AccessResult, EPropertyAccessResultFlags::ConversionFailed))
+	{
+		FFrame::KismetExecutionMessage(*FString::Printf(TEXT("Property '%s' on '%s' (%s) tried to read into a property value of the incorrect type"), *Property->GetName(), *Object->GetPathName(), *Object->GetClass()->GetName()), ELogVerbosity::Warning, PropertyGetFailedWarning);
+		return nullptr;
+	}
+
+	return true;
+}
+
+DEFINE_FUNCTION(UKismetSystemLibrary::execGetEditorProperty)
+{
+	P_GET_OBJECT(UObject, Object);
+	P_GET_PROPERTY(UNameProperty, PropertyName);
+
+	Stack.StepCompiledIn<UProperty>(nullptr);
+	const UProperty* ValueProp = Stack.MostRecentProperty;
+	void* ValuePtr = Stack.MostRecentPropertyAddress;
+
+	P_FINISH;
+
+	if (!ValueProp || !ValuePtr)
+	{
+		FBlueprintExceptionInfo ExceptionInfo(
+			EBlueprintExceptionType::AccessViolation,
+			LOCTEXT("GetEditorProperty_MissingOutputProperty", "Failed to resolve the output parameter for GetEditorProperty.")
+		);
+		FBlueprintCoreDelegates::ThrowScriptException(P_THIS, Stack, ExceptionInfo);
+	}
+
+	if (!Object)
+	{
+		FBlueprintExceptionInfo ExceptionInfo(
+			EBlueprintExceptionType::AccessViolation, 
+			LOCTEXT("GetEditorProperty_AccessNone", "Accessed None attempting to call GetEditorProperty.")
+		);
+		FBlueprintCoreDelegates::ThrowScriptException(P_THIS, Stack, ExceptionInfo);
+	}
+
+	bool bResult = false;
+
+	if (Object)
+	{
+		const UProperty* ObjectProp = PropertyAccessUtil::FindPropertyByName(PropertyName, Object->GetClass());
+		if (ObjectProp && ObjectProp->GetClass() == ValueProp->GetClass()) // Compare the classes as these must be an *exact* match as the read is low-level and without property coercion
+		{
+			P_NATIVE_BEGIN;
+			bResult = Generic_GetEditorProperty(Object, ObjectProp, ValuePtr);
+			P_NATIVE_END;
+		}
+		else
+		{
+			FFrame::KismetExecutionMessage(*FString::Printf(TEXT("Property '%s' on '%s' (%s) was missing or didn't match the output value type (%s)"), *PropertyName.ToString(), *Object->GetPathName(), *Object->GetClass()->GetName(), *ValueProp->GetClass()->GetName()), ELogVerbosity::Warning, PropertyGetFailedWarning);
+		}
+	}
+
+	*(bool*)RESULT_PARAM = bResult;
+}
+
+bool UKismetSystemLibrary::SetEditorProperty(UObject* Object, const FName PropertyName, const int32& OutValuePtr)
+{
+	// We should never hit this! Stubbed to avoid NoExport on the class.
+	check(0);
+	return false;
+}
+
+bool UKismetSystemLibrary::Generic_SetEditorProperty(UObject* Object, const UProperty* Property, const void* ValuePtr)
+{
+	const EPropertyAccessResultFlags AccessResult = PropertyAccessUtil::SetPropertyValue_Object(Property, Object, ValuePtr, INDEX_NONE, PropertyAccessUtil::EditorReadOnlyFlags);
+
+	if (EnumHasAnyFlags(AccessResult, EPropertyAccessResultFlags::PermissionDenied))
+	{
+		if (EnumHasAnyFlags(AccessResult, EPropertyAccessResultFlags::AccessProtected))
+		{
+			FFrame::KismetExecutionMessage(*FString::Printf(TEXT("Property '%s' on '%s' (%s) is protected and cannot be set"), *Property->GetName(), *Object->GetPathName(), *Object->GetClass()->GetName()), ELogVerbosity::Warning, PropertySetFailedWarning);
+			return false;
+		}
+
+		if (EnumHasAnyFlags(AccessResult, EPropertyAccessResultFlags::CannotEditTemplate))
+		{
+			FFrame::KismetExecutionMessage(*FString::Printf(TEXT("Property '%s' on '%s' (%s) cannot be edited on templates"), *Property->GetName(), *Object->GetPathName(), *Object->GetClass()->GetName()), ELogVerbosity::Warning, PropertySetFailedWarning);
+			return false;
+		}
+
+		if (EnumHasAnyFlags(AccessResult, EPropertyAccessResultFlags::CannotEditInstance))
+		{
+			FFrame::KismetExecutionMessage(*FString::Printf(TEXT("Property '%s' on '%s' (%s) cannot be edited on instances"), *Property->GetName(), *Object->GetPathName(), *Object->GetClass()->GetName()), ELogVerbosity::Warning, PropertySetFailedWarning);
+			return false;
+		}
+
+		if (EnumHasAnyFlags(AccessResult, EPropertyAccessResultFlags::ReadOnly))
+		{
+			FFrame::KismetExecutionMessage(*FString::Printf(TEXT("Property '%s' on '%s' (%s) is read-only and cannot be set"), *Property->GetName(), *Object->GetPathName(), *Object->GetClass()->GetName()), ELogVerbosity::Warning, PropertySetFailedWarning);
+			return false;
+		}
+
+		FFrame::KismetExecutionMessage(*FString::Printf(TEXT("Property '%s' on '%s' (%s) cannot be read"), *Property->GetName(), *Object->GetPathName(), *Object->GetClass()->GetName()), ELogVerbosity::Warning, PropertySetFailedWarning);
+		return false;
+	}
+
+	if (EnumHasAnyFlags(AccessResult, EPropertyAccessResultFlags::ConversionFailed))
+	{
+		FFrame::KismetExecutionMessage(*FString::Printf(TEXT("Property '%s' on '%s' (%s) tried to write into a property value of the incorrect type"), *Property->GetName(), *Object->GetPathName(), *Object->GetClass()->GetName()), ELogVerbosity::Warning, PropertySetFailedWarning);
+		return nullptr;
+	}
+
+	return true;
+}
+
+DEFINE_FUNCTION(UKismetSystemLibrary::execSetEditorProperty)
+{
+	P_GET_OBJECT(UObject, Object);
+	P_GET_PROPERTY(UNameProperty, PropertyName);
+
+	Stack.StepCompiledIn<UProperty>(nullptr);
+	const UProperty* ValueProp = Stack.MostRecentProperty;
+	const void* ValuePtr = Stack.MostRecentPropertyAddress;
+
+	P_FINISH;
+
+	if (!ValueProp || !ValuePtr)
+	{
+		FBlueprintExceptionInfo ExceptionInfo(
+			EBlueprintExceptionType::AccessViolation,
+			LOCTEXT("SetEditorProperty_MissingOutputProperty", "Failed to resolve the output parameter for SetEditorProperty.")
+		);
+		FBlueprintCoreDelegates::ThrowScriptException(P_THIS, Stack, ExceptionInfo);
+	}
+
+	if (!Object)
+	{
+		FBlueprintExceptionInfo ExceptionInfo(
+			EBlueprintExceptionType::AccessViolation, 
+			LOCTEXT("SetEditorProperty_AccessNone", "Accessed None attempting to call SetEditorProperty.")
+		);
+		FBlueprintCoreDelegates::ThrowScriptException(P_THIS, Stack, ExceptionInfo);
+	}
+
+	bool bResult = false;
+
+	if (Object)
+	{
+		const UProperty* ObjectProp = PropertyAccessUtil::FindPropertyByName(PropertyName, Object->GetClass());
+		if (ObjectProp && ObjectProp->GetClass() == ValueProp->GetClass()) // Compare the classes as these must be an *exact* match as the read is low-level and without property coercion
+		{
+			P_NATIVE_BEGIN;
+			bResult = Generic_SetEditorProperty(Object, ObjectProp, ValuePtr);
+			P_NATIVE_END;
+		}
+		else
+		{
+			FFrame::KismetExecutionMessage(*FString::Printf(TEXT("Property '%s' on '%s' (%s) was missing or didn't match the output value type (%s)"), *PropertyName.ToString(), *Object->GetPathName(), *Object->GetClass()->GetName(), *ValueProp->GetClass()->GetName()), ELogVerbosity::Warning, PropertySetFailedWarning);
+		}
+	}
+
+	*(bool*)RESULT_PARAM = bResult;
+}
+
+#endif	// WITH_EDITOR
+
 int32 UKismetSystemLibrary::BeginTransaction(const FString& Context, FText Description, UObject* PrimaryObject)
 {
 #if WITH_EDITOR
@@ -2856,4 +3059,6 @@ void UKismetSystemLibrary::GetPrimaryAssetsWithBundleState(const TArray<FName>&
 	{
 		Manager->GetPrimaryAssetsWithBundleState(OutPrimaryAssetIdList, ValidTypes, RequiredBundles, ExcludedBundles, bForceCurrentState);
 	}
-}
\ No newline at end of file
+}
+
+#undef LOCTEXT_NAMESPACE
