Skip type validation of wildcard function arguments in BPs

#rnx

diff --git a/Engine/Source/Editor/BlueprintGraph/Private/CallFunctionHandler.cpp b/Engine/Source/Editor/BlueprintGraph/Private/CallFunctionHandler.cpp
--- a/Engine/Source/Editor/BlueprintGraph/Private/CallFunctionHandler.cpp
+++ b/Engine/Source/Editor/BlueprintGraph/Private/CallFunctionHandler.cpp
@@ -151,7 +151,8 @@ void FKCHandler_CallFunction::CreateFunctionCallStatement(FKismetFunctionContext
 				if (Property->GetFName() == PinMatch->PinName)
 				{
 					// Found a corresponding pin, does it match in type and direction?
-					if (FKismetCompilerUtilities::IsTypeCompatibleWithProperty(PinMatch, Property, CompilerContext.MessageLog, CompilerContext.GetSchema(), Context.NewClass))
+					if (UK2Node_CallFunction::IsStructureWildcardProperty(Function, Property->GetFName()) ||
+						FKismetCompilerUtilities::IsTypeCompatibleWithProperty(PinMatch, Property, CompilerContext.MessageLog, CompilerContext.GetSchema(), Context.NewClass))
 					{
 						UEdGraphPin* PinToTry = FEdGraphUtilities::GetNetFromPin(PinMatch);
 
