Fix crash when connecting split pins for some nodes to constant pins

diff --git a/Engine/Source/Editor/BlueprintGraph/Private/K2Node_CallFunction.cpp b/Engine/Source/Editor/BlueprintGraph/Private/K2Node_CallFunction.cpp
--- a/Engine/Source/Editor/BlueprintGraph/Private/K2Node_CallFunction.cpp
+++ b/Engine/Source/Editor/BlueprintGraph/Private/K2Node_CallFunction.cpp
@@ -2701,7 +2701,10 @@ void UK2Node_CallFunction::ConformContainerPins()
 		{
 			if(bTypeIsAvailable)
 			{
-				if(Pin->GetPrimaryTerminalType() != TerminalType)
+				const FEdGraphTerminalType PrimaryType = Pin->GetPrimaryTerminalType();
+				if( PrimaryType.TerminalCategory != TerminalType.TerminalCategory ||
+					PrimaryType.TerminalSubCategory != TerminalType.TerminalSubCategory ||
+					PrimaryType.TerminalSubCategoryObject != TerminalType.TerminalSubCategoryObject)
 				{
 					// terminal type changed:
 					if (Pin->SubPins.Num() > 0 && Pin->PinType.PinCategory != UEdGraphSchema_K2::PC_Wildcard)
