Fixup PushModel merge issues in relation to FProperty changes

diff --git a/Engine/Source/Editor/BlueprintGraph/Private/K2Node_VariableSet.cpp b/Engine/Source/Editor/BlueprintGraph/Private/K2Node_VariableSet.cpp
--- a/Engine/Source/Editor/BlueprintGraph/Private/K2Node_VariableSet.cpp
+++ b/Engine/Source/Editor/BlueprintGraph/Private/K2Node_VariableSet.cpp
@@ -308,7 +308,7 @@ bool UK2Node_VariableSet::ShouldFlushDormancyOnSet() const
 
 bool UK2Node_VariableSet::IsNetProperty() const
 {
-	UProperty* Property = GetPropertyForVariable();
+	FProperty* Property = GetPropertyForVariable();
 	return Property && (Property->PropertyFlags & CPF_Net);
 }
 
diff --git a/Engine/Source/Editor/BlueprintGraph/Private/VariableSetHandler.cpp b/Engine/Source/Editor/BlueprintGraph/Private/VariableSetHandler.cpp
--- a/Engine/Source/Editor/BlueprintGraph/Private/VariableSetHandler.cpp
+++ b/Engine/Source/Editor/BlueprintGraph/Private/VariableSetHandler.cpp
@@ -223,8 +223,8 @@ void FKCHandler_VariableSet::Transform(FKismetFunctionContext& Context, UEdGraph
 			static const FName RepIndexPinName(TEXT("RepIndex"));
 			static const FName PropertyNamePinName(TEXT("PropertyName"));
 
-			UProperty* Property = SetNotify->GetPropertyForVariable();
-			UClass* Class = Property->GetTypedOuter<UClass>();
+			FProperty* Property = SetNotify->GetPropertyForVariable();
+			UClass* Class = Property->GetOwnerClass();
 
 			if (Property && Class)
 			{
@@ -237,7 +237,7 @@ void FKCHandler_VariableSet::Transform(FKismetFunctionContext& Context, UEdGraph
 					if (Class->ClassGeneratedBy == Context.Blueprint && Context.NewClass && Context.NewClass != Class)
 					{
 						Class = Context.NewClass;
-						Property = FindFieldChecked<UProperty>(Class, Property->GetFName());
+						Property = FindFieldChecked<FProperty>(Class, Property->GetFName());
 					}
 					if (Property->GetOffset_ForGC() == 0)
 					{
@@ -246,7 +246,7 @@ void FKCHandler_VariableSet::Transform(FKismetFunctionContext& Context, UEdGraph
 							if (UClass * UseClass = Blueprint->GeneratedClass)
 							{
 								Class = UseClass;
-								Property = FindFieldChecked<UProperty>(Class, Property->GetFName());
+								Property = FindFieldChecked<FProperty>(Class, Property->GetFName());
 							}
 						}
 					}
