More CIS fixes from Push Model changes

diff --git a/Engine/Source/Editor/BlueprintGraph/Private/VariableSetHandler.cpp b/Engine/Source/Editor/BlueprintGraph/Private/VariableSetHandler.cpp
--- a/Engine/Source/Editor/BlueprintGraph/Private/VariableSetHandler.cpp
+++ b/Engine/Source/Editor/BlueprintGraph/Private/VariableSetHandler.cpp
@@ -223,80 +223,80 @@ void FKCHandler_VariableSet::Transform(FKismetFunctionContext& Context, UEdGraph
 			static const FName RepIndexPinName(TEXT("RepIndex"));
 			static const FName PropertyNamePinName(TEXT("PropertyName"));
 
-			FProperty* Property = SetNotify->GetPropertyForVariable();
-			UClass* Class = Property->GetOwnerClass();
-
-			if (Property && Class)
+			if (FProperty * Property = SetNotify->GetPropertyForVariable())
 			{
-				// We need to make sure this class already has its property offsets setup, otherwise
-				// the order of our replicated properties won't match, meaning the RepIndex will be
-				// invalid.
-				if (Property->GetOffset_ForGC() == 0)
+				if (UClass* Class = Property->GetOwnerClass())
 				{
-					// Make sure that we're using the correct class and that it has replication data set up.
-					if (Class->ClassGeneratedBy == Context.Blueprint && Context.NewClass && Context.NewClass != Class)
-					{
-						Class = Context.NewClass;
-						Property = FindFieldChecked<FProperty>(Class, Property->GetFName());
-					}
+					// We need to make sure this class already has its property offsets setup, otherwise
+					// the order of our replicated properties won't match, meaning the RepIndex will be
+					// invalid.
 					if (Property->GetOffset_ForGC() == 0)
 					{
-						if (UBlueprint * Blueprint = Cast<UBlueprint>(Class->ClassGeneratedBy))
+						// Make sure that we're using the correct class and that it has replication data set up.
+						if (Class->ClassGeneratedBy == Context.Blueprint && Context.NewClass && Context.NewClass != Class)
 						{
-							if (UClass * UseClass = Blueprint->GeneratedClass)
+							Class = Context.NewClass;
+							Property = FindFieldChecked<FProperty>(Class, Property->GetFName());
+						}
+						if (Property->GetOffset_ForGC() == 0)
+						{
+							if (UBlueprint * Blueprint = Cast<UBlueprint>(Class->ClassGeneratedBy))
 							{
-								Class = UseClass;
-								Property = FindFieldChecked<FProperty>(Class, Property->GetFName());
+								if (UClass * UseClass = Blueprint->GeneratedClass)
+								{
+									Class = UseClass;
+									Property = FindFieldChecked<FProperty>(Class, Property->GetFName());
+								}
 							}
 						}
 					}
-				}
 
-				ensureAlwaysMsgf(Property->GetOffset_ForGC() != 0,
-					TEXT("Class does not have Property Offsets setup. This will cause issues with Push Model. Blueprint=%s, Class=%s, Property=%s"),
-					*Context.Blueprint->GetPathName(), *Class->GetPathName(), *Property->GetName());
+					ensureAlwaysMsgf(Property->GetOffset_ForGC() != 0,
+						TEXT("Class does not have Property Offsets setup. This will cause issues with Push Model. Blueprint=%s, Class=%s, Property=%s"),
+						*Context.Blueprint->GetPathName(), *Class->GetPathName(), *Property->GetName());
 
-				if (!Class->HasAnyClassFlags(CLASS_ReplicationDataIsSetUp))
-				{
-					Class->SetUpRuntimeReplicationData();
-				}
+					if (!Class->HasAnyClassFlags(CLASS_ReplicationDataIsSetUp))
+					{
+						Class->SetUpRuntimeReplicationData();
+					}
 
-				UK2Node_CallFunction* CallFuncNode = Node->GetGraph()->CreateIntermediateNode<UK2Node_CallFunction>();
-				CallFuncNode->FunctionReference.SetExternalMember(MarkPropertyDirtyFuncName, UNetPushModelHelpers::StaticClass());
-				CallFuncNode->AllocateDefaultPins();
+					UK2Node_CallFunction* CallFuncNode = Node->GetGraph()->CreateIntermediateNode<UK2Node_CallFunction>();
+					CallFuncNode->FunctionReference.SetExternalMember(MarkPropertyDirtyFuncName, UNetPushModelHelpers::StaticClass());
+					CallFuncNode->AllocateDefaultPins();
 
-				// Take our old Self (Target) pin and hook it up to the Object pin for UNetPushModelHelpers::MarkPropertyDirty.
-				// If our Self pin isn't hooked up to anything, then create an intermediate Self node and use that.
-				UEdGraphPin* SelfPin = Node->FindPinChecked(UEdGraphSchema_K2::PN_Self);
-				UEdGraphPin* ObjectPin = CallFuncNode->FindPinChecked(ObjectPinName);
+					// Take our old Self (Target) pin and hook it up to the Object pin for UNetPushModelHelpers::MarkPropertyDirty.
+					// If our Self pin isn't hooked up to anything, then create an intermediate Self node and use that.
+					UEdGraphPin* SelfPin = Node->FindPinChecked(UEdGraphSchema_K2::PN_Self);
+					UEdGraphPin* ObjectPin = CallFuncNode->FindPinChecked(ObjectPinName);
 
-				if (SelfPin->LinkedTo.Num() > 0)
-				{
-					SelfPin = SelfPin->LinkedTo[0];
-				}
-				else
-				{
-					UK2Node_Self* SelfNode = Node->GetGraph()->CreateIntermediateNode<UK2Node_Self>();
-					SelfNode->AllocateDefaultPins();
-					SelfPin = SelfNode->FindPinChecked(UEdGraphSchema_K2::PN_Self);
-				}
+					if (SelfPin->LinkedTo.Num() > 0)
+					{
+						SelfPin = SelfPin->LinkedTo[0];
+					}
+					else
+					{
+						UK2Node_Self* SelfNode = Node->GetGraph()->CreateIntermediateNode<UK2Node_Self>();
+						SelfNode->AllocateDefaultPins();
+						SelfPin = SelfNode->FindPinChecked(UEdGraphSchema_K2::PN_Self);
+					}
 
-				ObjectPin->MakeLinkTo(SelfPin);
+					ObjectPin->MakeLinkTo(SelfPin);
 
-				UEdGraphPin* RepIndexPin = CallFuncNode->FindPinChecked(RepIndexPinName);
-				RepIndexPin->DefaultValue = FString::FromInt(Property->RepIndex);
+					UEdGraphPin* RepIndexPin = CallFuncNode->FindPinChecked(RepIndexPinName);
+					RepIndexPin->DefaultValue = FString::FromInt(Property->RepIndex);
 
-				UEdGraphPin* PropertyNamePin = CallFuncNode->FindPinChecked(PropertyNamePinName);
-				PropertyNamePin->DefaultValue = Property->GetFName().ToString();
+					UEdGraphPin* PropertyNamePin = CallFuncNode->FindPinChecked(PropertyNamePinName);
+					PropertyNamePin->DefaultValue = Property->GetFName().ToString();
 
-				// Hook up our exec pins.
-				UEdGraphPin* OldThenPin = Node->FindPinChecked(UEdGraphSchema_K2::PN_Then);
-				UEdGraphPin* NewThenPin = CallFuncNode->GetThenPin();
-				if (ensure(NewThenPin))
-				{
-					NewThenPin->CopyPersistentDataFromOldPin(*OldThenPin);
-					OldThenPin->BreakAllPinLinks();
-					OldThenPin->MakeLinkTo(CallFuncNode->GetExecPin());
+					// Hook up our exec pins.
+					UEdGraphPin* OldThenPin = Node->FindPinChecked(UEdGraphSchema_K2::PN_Then);
+					UEdGraphPin* NewThenPin = CallFuncNode->GetThenPin();
+					if (ensure(NewThenPin))
+					{
+						NewThenPin->CopyPersistentDataFromOldPin(*OldThenPin);
+						OldThenPin->BreakAllPinLinks();
+						OldThenPin->MakeLinkTo(CallFuncNode->GetExecPin());
+					}
 				}
 			}
 		}
diff --git a/Engine/Source/Runtime/Engine/Private/RepLayout.cpp b/Engine/Source/Runtime/Engine/Private/RepLayout.cpp
--- a/Engine/Source/Runtime/Engine/Private/RepLayout.cpp
+++ b/Engine/Source/Runtime/Engine/Private/RepLayout.cpp
@@ -5369,11 +5369,13 @@ void FRepLayout::InitFromClass(
 			}
 
 			++NumberOfLifetimeProperties;
+#if WITH_PUSH_MODEL
 			if (bIsPushModelEnabled && LifetimeProps[i].bIsPushBased)
 			{
 				++NumberOfPushModelProperties;
 				Parents[ParentIndex].Flags |= ERepParentFlags::UsePushModel;
 			}
+#endif
 		}
 		else
 		{
@@ -5478,12 +5480,14 @@ void FRepLayout::InitFromClass(
 		}
 	}	
 
+#if WITH_PUSH_MODEL
 	if (bIsPushModelEnabled && NumberOfPushModelProperties > 0)
 	{
 		Flags |= (NumberOfLifetimeProperties == NumberOfPushModelProperties) ?
 			ERepLayoutFlags::FullPushSupport :
 			ERepLayoutFlags::PartialPushSupport;
 	}
+#endif
 
 	if (!ServerConnection || EnumHasAnyFlags(CreateFlags, ECreateRepLayoutFlags::MaySendProperties))
 	{
