Output wildcard pins can now connect to container types, similarly to how input pins can be. Override DoesOutputWildcardPinAcceptContainer to disable on a per-node basis

diff --git a/Engine/Source/Editor/BlueprintGraph/Classes/K2Node.h b/Engine/Source/Editor/BlueprintGraph/Classes/K2Node.h
--- a/Engine/Source/Editor/BlueprintGraph/Classes/K2Node.h
+++ b/Engine/Source/Editor/BlueprintGraph/Classes/K2Node.h
@@ -397,6 +397,7 @@ class BLUEPRINTGRAPH_VTABLE UK2Node : public UEdGraphNode
 	BLUEPRINTGRAPH_API virtual int32 GetNodeRefreshPriority() const { return EBaseNodeRefreshPriority::Normal; }
 
 	BLUEPRINTGRAPH_API virtual bool DoesInputWildcardPinAcceptArray(const UEdGraphPin* Pin) const { return true; }
+	BLUEPRINTGRAPH_API virtual bool DoesOutputWildcardPinAcceptContainer(const UEdGraphPin* Pin) const { return true; }
 
 	/** Handle when a variable is renamed in the Blueprint Palette */
 	virtual void HandleVariableRenamed(UBlueprint* InBlueprint, UClass* InVariableClass, UEdGraph* InGraph, const FName& InOldVarName, const FName& InNewVarName) {}
diff --git a/Engine/Source/Editor/BlueprintGraph/Classes/K2Node_CallArrayFunction.h b/Engine/Source/Editor/BlueprintGraph/Classes/K2Node_CallArrayFunction.h
--- a/Engine/Source/Editor/BlueprintGraph/Classes/K2Node_CallArrayFunction.h
+++ b/Engine/Source/Editor/BlueprintGraph/Classes/K2Node_CallArrayFunction.h
@@ -32,6 +32,7 @@ class UK2Node_CallArrayFunction : public UK2Node_CallFunction
 	//~ Begin UK2Node Interface
 	virtual void NotifyPinConnectionListChanged(UEdGraphPin* Pin) override;
 	virtual bool DoesInputWildcardPinAcceptArray(const UEdGraphPin* Pin) const override { return false; }
+	virtual bool DoesOutputWildcardPinAcceptContainer(const UEdGraphPin* Pin) const override { return false; }
 	virtual void ConvertDeprecatedNode(UEdGraph* Graph, bool bOnlySafeChanges) override;
 	//~ End UK2Node Interface
 
diff --git a/Engine/Source/Editor/BlueprintGraph/Classes/K2Node_ConvertAsset.h b/Engine/Source/Editor/BlueprintGraph/Classes/K2Node_ConvertAsset.h
--- a/Engine/Source/Editor/BlueprintGraph/Classes/K2Node_ConvertAsset.h
+++ b/Engine/Source/Editor/BlueprintGraph/Classes/K2Node_ConvertAsset.h
@@ -38,6 +38,7 @@ public:
 	virtual bool ShouldDrawCompact() const override { return true; }
 	virtual int32 GetNodeRefreshPriority() const override { return EBaseNodeRefreshPriority::Low_UsesDependentWildcard; }
 	virtual bool DoesInputWildcardPinAcceptArray(const UEdGraphPin* Pin) const override { return false; }
+	virtual bool DoesOutputWildcardPinAcceptContainer(const UEdGraphPin* Pin) const override { return false; }
 
 	virtual void ExpandNode(class FKismetCompilerContext& CompilerContext, UEdGraph* SourceGraph) override;
 	virtual FText GetCompactNodeTitle() const override;
diff --git a/Engine/Source/Editor/BlueprintGraph/Private/EdGraphSchema_K2.cpp b/Engine/Source/Editor/BlueprintGraph/Private/EdGraphSchema_K2.cpp
--- a/Engine/Source/Editor/BlueprintGraph/Private/EdGraphSchema_K2.cpp
+++ b/Engine/Source/Editor/BlueprintGraph/Private/EdGraphSchema_K2.cpp
@@ -4025,6 +4025,12 @@ struct FWildcardArrayPinHelper
 			return true;
 		}
 
+		const bool bOutputWildcardPinAcceptsContainer = !OwningNode || OwningNode->DoesOutputWildcardPinAcceptContainer(OutputPin);
+		if(bOutputWildcardPinAcceptsContainer)
+		{
+			return true;
+		}
+
 		const bool bCheckInputPin = (InputPin->PinType.PinCategory == UEdGraphSchema_K2::PC_Wildcard) && !InputPin->PinType.IsArray();
 		const bool bArrayOutputPin = OutputPin && OutputPin->PinType.IsArray();
 		return !(bCheckInputPin && bArrayOutputPin);
@@ -4394,7 +4400,10 @@ bool UEdGraphSchema_K2::DefaultValueSimpleValidation(const FEdGraphPinType& PinT
 
 bool UEdGraphSchema_K2::ArePinTypesCompatible(const FEdGraphPinType& Output, const FEdGraphPinType& Input, const UClass* CallingContext, bool bIgnoreArray /*= false*/) const
 {
-	if( !bIgnoreArray && ( Output.ContainerType != Input.ContainerType ) && (Input.PinCategory != PC_Wildcard || Input.IsContainer()) )
+	if( !bIgnoreArray && 
+		( Output.ContainerType != Input.ContainerType ) && 
+		(Input.PinCategory != PC_Wildcard || Input.IsContainer()) && 
+		(Output.PinCategory != PC_Wildcard || Output.IsContainer()))
 	{
 		return false;
 	}
