Fix an editor crash caused by leaking Blueprint context menu actions for new sublevel script actors.

- Mirrored from //UE4/Dev-Framework (11177969).

diff --git a/Engine/Source/Editor/BlueprintGraph/Private/BlueprintActionDatabase.cpp b/Engine/Source/Editor/BlueprintGraph/Private/BlueprintActionDatabase.cpp
--- a/Engine/Source/Editor/BlueprintGraph/Private/BlueprintActionDatabase.cpp
+++ b/Engine/Source/Editor/BlueprintGraph/Private/BlueprintActionDatabase.cpp
@@ -530,6 +530,21 @@ namespace BlueprintActionDatabaseImpl
 
 	/** True if a RefreshAll has been requested for the next Tick */
 	bool bRefreshAllRequested = false;
+
+	/** */
+	TWeakObjectPtr<UWorld> OriginalOwningWorldOnSave;
+
+	/** */
+	static void OnPreSaveWorld(uint32 SaveFlags, UWorld* World)
+	{
+		OriginalOwningWorldOnSave = World->PersistentLevel->OwningWorld;
+	}
+
+	/** */
+	static void OnPostSaveWorld(uint32 SaveFlags, UWorld* World, bool bSuccess)
+	{
+		OriginalOwningWorldOnSave.Reset();
+	}
 }
 
 //------------------------------------------------------------------------------
@@ -1155,6 +1170,9 @@ FBlueprintActionDatabase::FBlueprintActionDatabase()
 
 	IHotReloadInterface& HotReloadSupport = FModuleManager::LoadModuleChecked<IHotReloadInterface>("HotReload");
 	HotReloadSupport.OnHotReload().AddStatic(&BlueprintActionDatabaseImpl::OnProjectHotReloaded);
+
+	FEditorDelegates::PreSaveWorld.AddStatic(&BlueprintActionDatabaseImpl::OnPreSaveWorld);
+	FEditorDelegates::PostSaveWorld.AddStatic(&BlueprintActionDatabaseImpl::OnPostSaveWorld);
 }
 
 //------------------------------------------------------------------------------
@@ -1516,6 +1534,14 @@ void FBlueprintActionDatabase::RefreshAssetActions(UObject* const AssetObject)
 	{
 		for( ULevel* Level : WorldAsset->GetLevels() )
 		{
+			// Skip adding member actions for the persistent level on save if it doesn't match what the owning world was set to prior to saving.
+			// Otherwise, a new entry will be created for actions that aren't technically valid for the world asset, since this is likely a sublevel.
+			// When the world that owns the sublevel is destroyed, these actions won't be cleaned up, preventing the level from being freed during GC.
+			if (Level == WorldAsset->PersistentLevel && OriginalOwningWorldOnSave.IsValid() && Level->OwningWorld != OriginalOwningWorldOnSave.Get())
+			{
+				continue;
+			}
+
 			UBlueprint* LevelScript = Level->GetLevelScriptBlueprint(true);
 			if (LevelScript != nullptr)
 			{
