Don't gather instanced enum entry display names from some K2 nodes

This was causing false positive HasTextWithInvalidPackageLocalizationID results.

#rnx

diff --git a/Engine/Source/Editor/BlueprintGraph/Classes/K2Node_Select.h b/Engine/Source/Editor/BlueprintGraph/Classes/K2Node_Select.h
--- a/Engine/Source/Editor/BlueprintGraph/Classes/K2Node_Select.h
+++ b/Engine/Source/Editor/BlueprintGraph/Classes/K2Node_Select.h
@@ -38,7 +38,7 @@ private:
 	TArray<FName> EnumEntries;
 
 	/** List of the current entries in the enum (Pin Friendly Names) */
-	UPROPERTY()
+	UPROPERTY(Transient)
 	TArray<FText> EnumEntryFriendlyNames;
 
 	/** Whether we need to reconstruct the node after the pins have changed */
diff --git a/Engine/Source/Editor/BlueprintGraph/Classes/K2Node_SwitchEnum.h b/Engine/Source/Editor/BlueprintGraph/Classes/K2Node_SwitchEnum.h
--- a/Engine/Source/Editor/BlueprintGraph/Classes/K2Node_SwitchEnum.h
+++ b/Engine/Source/Editor/BlueprintGraph/Classes/K2Node_SwitchEnum.h
@@ -27,7 +27,7 @@ class UK2Node_SwitchEnum : public UK2Node_Switch, public INodeDependingOnEnumInt
 	TArray<FName> EnumEntries;
 
 	/** List of the current entries in the enum */
-	UPROPERTY()
+	UPROPERTY(Transient)
 	TArray<FText> EnumFriendlyNames;
 
 	// INodeDependingOnEnumInterface
