Implemented variadic function support for Blueprints

Variadic functions are required to be a CustomThunk marked with the "Variadic" meta-data. They can then be used from a custom Blueprint node to accept arbitrary arguments at the end of their parameter list (any extra pins added to the node that aren't part of the main function definition will become the variadic payload).

Variadic arguments aren't type checked, so you need other function input to tell you how many to expect, and for a nativizied function, also what type of arguments you're dealing with.

diff --git a/Engine/Source/Developer/BlueprintCompilerCppBackend/Private/BlueprintCompilerCppBackend.cpp b/Engine/Source/Developer/BlueprintCompilerCppBackend/Private/BlueprintCompilerCppBackend.cpp
--- a/Engine/Source/Developer/BlueprintCompilerCppBackend/Private/BlueprintCompilerCppBackend.cpp
+++ b/Engine/Source/Developer/BlueprintCompilerCppBackend/Private/BlueprintCompilerCppBackend.cpp
@@ -682,6 +682,26 @@ FString FBlueprintCompilerCppBackend::EmitMethodInputParameterList(FEmitterLocal
 		}
 	}
 
+	const bool bIsVariadic = Statement.FunctionToCall->HasMetaData(FBlueprintMetadata::MD_Variadic);
+	if (bIsVariadic)
+	{
+		// Variadic functions may have extra terms they need to emit after the main set of function arguments
+		// These are all considered wildcards so no type checking will be performed on them
+		// It is the responsibility of the thunk function to interpret their types correctly based on other function input
+		for (; NumParams < Statement.RHS.Num(); ++NumParams)
+		{
+			FBPTerminal* Term = Statement.RHS[NumParams];
+			check(Term);
+
+			if (NumParams > 0)
+			{
+				Result += TEXT(", ");
+			}
+
+			Result += TermToText(EmitterContext, Term, ENativizedTermUsage::UnspecifiedOrReference);
+		}
+	}
+
 	return Result;
 }
 
diff --git a/Engine/Source/Editor/BlueprintGraph/Classes/EdGraphSchema_K2.h b/Engine/Source/Editor/BlueprintGraph/Classes/EdGraphSchema_K2.h
--- a/Engine/Source/Editor/BlueprintGraph/Classes/EdGraphSchema_K2.h
+++ b/Engine/Source/Editor/BlueprintGraph/Classes/EdGraphSchema_K2.h
@@ -98,6 +98,9 @@ public:
 	/** Marks a UFUNCTION as latent execution */
 	static const FName MD_Latent;
 
+	/** Marks a UFUNCTION as accepting variadic arguments */
+	static const FName MD_Variadic;
+
 	/** Marks a UFUNCTION as unsafe for use in the UCS, which prevents it from being called from the UCS.  Useful for things that spawn actors, etc that should never happen in the UCS */
 	static const FName MD_UnsafeForConstructionScripts;
 
diff --git a/Engine/Source/Editor/BlueprintGraph/Private/CallFunctionHandler.cpp b/Engine/Source/Editor/BlueprintGraph/Private/CallFunctionHandler.cpp
--- a/Engine/Source/Editor/BlueprintGraph/Private/CallFunctionHandler.cpp
+++ b/Engine/Source/Editor/BlueprintGraph/Private/CallFunctionHandler.cpp
@@ -72,13 +72,14 @@ void FKCHandler_CallFunction::CreateFunctionCallStatement(FKismetFunctionContext
 		// Make sure the pin mapping is sound (all pins wire up to a matching function parameter, and all function parameters match a pin)
 
 		// Remaining unmatched pins
+		// Note: Should maintain a stable order for variadic arguments
 		TArray<UEdGraphPin*> RemainingPins;
 		RemainingPins.Append(Node->Pins);
 
 		const UEdGraphSchema_K2* Schema = CompilerContext.GetSchema();
 
 		// Remove expected exec and self pins
-		RemainingPins.RemoveAllSwap([Schema](UEdGraphPin* Pin) { return (Pin->bOrphanedPin || Schema->IsMetaPin(*Pin)); }, false);
+		RemainingPins.RemoveAll([Schema](UEdGraphPin* Pin) { return (Pin->bOrphanedPin || Schema->IsMetaPin(*Pin)); });
 
 		// Check for magic pins
 		const bool bIsLatent = Function->HasMetaData(FBlueprintMetadata::MD_Latent);
@@ -271,7 +272,7 @@ void FKCHandler_CallFunction::CreateFunctionCallStatement(FKismetFunctionContext
 					}
 
 					bFoundParam = true;
-					RemainingPins.RemoveAtSwap(i);
+					RemainingPins.RemoveAt(i);
 				}
 			}
 
@@ -282,10 +283,43 @@ void FKCHandler_CallFunction::CreateFunctionCallStatement(FKismetFunctionContext
 			}
 		}
 
-		// At this point, we should have consumed all pins.  If not, there are extras that need to be removed.
-		for (int32 i = 0; i < RemainingPins.Num(); ++i)
+		// If we have pins remaining then it's either an error, or extra variadic terms that need to be emitted
+		if (RemainingPins.Num() > 0)
 		{
-			CompilerContext.MessageLog.Error(*FText::Format(LOCTEXT("PinMismatchParameter_ErrorFmt", "Pin @@ named {0} doesn't match any parameters of function {1}"), FText::FromName(RemainingPins[i]->PinName), FText::FromString(Function->GetName())).ToString(), RemainingPins[i]);
+			const bool bIsVariadic = Function->HasMetaData(FBlueprintMetadata::MD_Variadic);
+			if (bIsVariadic)
+			{
+				// Add a RHS term for every remaining pin
+				for (UEdGraphPin* RemainingPin : RemainingPins)
+				{
+					// Variadic pins are assumed to be wildcard pins that have been connected to something else
+					if (RemainingPin->LinkedTo.Num() == 0)
+					{
+						CompilerContext.MessageLog.Error(*LOCTEXT("UnlinkedVariadicPin_Error", "The variadic pin @@ must be connected. Connect something to @@.").ToString(), RemainingPin, RemainingPin->GetOwningNodeUnchecked());
+						continue;
+					}
+
+					UEdGraphPin* PinToTry = FEdGraphUtilities::GetNetFromPin(RemainingPin);
+					if (FBPTerminal** Term = Context.NetMap.Find(PinToTry))
+					{
+						FBPTerminal* RHSTerm = *Term;
+						RHSTerms.Add(RHSTerm);
+					}
+					else
+					{
+						CompilerContext.MessageLog.Error(*LOCTEXT("ResolveTermVariadic_Error", "Failed to resolve variadic term passed into @@").ToString(), RemainingPin);
+						bMatchedAllParams = false;
+					}
+				}
+			}
+			else
+			{
+				// At this point, we should have consumed all pins.  If not, there are extras that need to be removed.
+				for (const UEdGraphPin* RemainingPin : RemainingPins)
+				{
+					CompilerContext.MessageLog.Error(*FText::Format(LOCTEXT("PinMismatchParameter_ErrorFmt", "Pin @@ named {0} doesn't match any parameters of function {1}"), FText::FromName(RemainingPin->PinName), FText::FromString(Function->GetName())).ToString(), RemainingPin);
+				}
+			}
 		}
 
 		if (NumErrorsAtStart == CompilerContext.MessageLog.NumErrors)
diff --git a/Engine/Source/Editor/BlueprintGraph/Private/EdGraphSchema_K2.cpp b/Engine/Source/Editor/BlueprintGraph/Private/EdGraphSchema_K2.cpp
--- a/Engine/Source/Editor/BlueprintGraph/Private/EdGraphSchema_K2.cpp
+++ b/Engine/Source/Editor/BlueprintGraph/Private/EdGraphSchema_K2.cpp
@@ -112,6 +112,7 @@ const FName FBlueprintMetadata::MD_IgnoreCategoryKeywordsInSubclasses(TEXT("Igno
 
 const FName FBlueprintMetadata::MD_Protected(TEXT("BlueprintProtected"));
 const FName FBlueprintMetadata::MD_Latent(TEXT("Latent"));
+const FName FBlueprintMetadata::MD_Variadic(TEXT("Variadic"));
 const FName FBlueprintMetadata::MD_UnsafeForConstructionScripts(TEXT("UnsafeDuringActorConstruction"));
 const FName FBlueprintMetadata::MD_FunctionCategory(TEXT("Category"));
 const FName FBlueprintMetadata::MD_DeprecatedFunction(TEXT("DeprecatedFunction"));
diff --git a/Engine/Source/Editor/KismetCompiler/Private/KismetCompilerVMBackend.cpp b/Engine/Source/Editor/KismetCompiler/Private/KismetCompilerVMBackend.cpp
--- a/Engine/Source/Editor/KismetCompiler/Private/KismetCompilerVMBackend.cpp
+++ b/Engine/Source/Editor/KismetCompiler/Private/KismetCompilerVMBackend.cpp
@@ -1286,6 +1286,19 @@ public:
 			}
 		}
 
+		const bool bIsVariadic = FunctionToCall->HasMetaData(FBlueprintMetadata::MD_Variadic);
+		if (bIsVariadic)
+		{
+			// Variadic functions may have extra terms they need to emit after the main set of function arguments
+			// These are all considered wildcards so no type checking will be performed on them
+			for (; NumParams < Statement.RHS.Num(); ++NumParams)
+			{
+				FBPTerminal* Term = Statement.RHS[NumParams];
+				check(Term);
+				EmitTerm(Term, nullptr);
+			}
+		}
+
 		// End of parameter list
 		Writer << EX_EndFunctionParms;
 	}
