Clear reference to class object on connection to class input pin for GetClassDefault nodes.

diff --git a/Engine/Source/Editor/BlueprintGraph/Classes/K2Node_GetClassDefaults.h b/Engine/Source/Editor/BlueprintGraph/Classes/K2Node_GetClassDefaults.h
--- a/Engine/Source/Editor/BlueprintGraph/Classes/K2Node_GetClassDefaults.h
+++ b/Engine/Source/Editor/BlueprintGraph/Classes/K2Node_GetClassDefaults.h
@@ -20,6 +20,7 @@ class UK2Node_GetClassDefaults : public UK2Node
 	//~ Begin UObject Interface
 	virtual void PostEditChangeProperty(struct FPropertyChangedEvent& PropertyChangedEvent) override;
 	virtual void PreEditChange(UProperty* PropertyAboutToChange) override;
+	virtual void PostLoad() override;
 	//~ End UObject Interface
 
 	//~ Begin UEdGraphNode Interface
diff --git a/Engine/Source/Editor/BlueprintGraph/Private/K2Node_GetClassDefaults.cpp b/Engine/Source/Editor/BlueprintGraph/Private/K2Node_GetClassDefaults.cpp
--- a/Engine/Source/Editor/BlueprintGraph/Private/K2Node_GetClassDefaults.cpp
+++ b/Engine/Source/Editor/BlueprintGraph/Private/K2Node_GetClassDefaults.cpp
@@ -188,6 +188,23 @@ UK2Node_GetClassDefaults::UK2Node_GetClassDefaults(const FObjectInitializer& Obj
 {
 }
 
+void UK2Node_GetClassDefaults::PostLoad()
+{
+	Super::PostLoad();
+
+	// Ensure that a previously-set class object is no longer referenced as a default value if the class input pin is connected.
+	if (UEdGraphPin* ClassPin = FindClassPin(Pins))
+	{
+		if (ClassPin->LinkedTo.Num() > 0 && !ClassPin->DoesDefaultValueMatchAutogenerated())
+		{
+			if (const UEdGraphSchema* Schema = GetSchema())
+			{
+				Schema->ResetPinToAutogeneratedDefaultValue(ClassPin, false);
+			}
+		}
+	}
+}
+
 void UK2Node_GetClassDefaults::PreEditChange(UProperty* PropertyThatWillChange)
 {
 	Super::PreEditChange(PropertyThatWillChange);
@@ -251,6 +268,8 @@ void UK2Node_GetClassDefaults::PostPlacedNewNode()
 
 void UK2Node_GetClassDefaults::PinConnectionListChanged(UEdGraphPin* ChangedPin)
 {
+	Super::PinConnectionListChanged(ChangedPin);
+
 	if(ChangedPin != nullptr && ChangedPin->PinName == ClassPinName && ChangedPin->Direction == EGPD_Input)
 	{
 		OnClassPinChanged();
@@ -259,6 +278,8 @@ void UK2Node_GetClassDefaults::PinConnectionListChanged(UEdGraphPin* ChangedPin)
 
 void UK2Node_GetClassDefaults::PinDefaultValueChanged(UEdGraphPin* ChangedPin) 
 {
+	Super::PinDefaultValueChanged(ChangedPin);
+
 	if(ChangedPin != nullptr && ChangedPin->PinName == ClassPinName && ChangedPin->Direction == EGPD_Input)
 	{
 		OnClassPinChanged();
