New pin types for K2Node_FormatText with support for bool, string, name, and Object

PR #6454: New pin types for K2Node_FormatText (Contributed by [REDACTED_PERSON])

diff --git a/Engine/Source/Editor/BlueprintGraph/Private/K2Node_FormatText.cpp b/Engine/Source/Editor/BlueprintGraph/Private/K2Node_FormatText.cpp
--- a/Engine/Source/Editor/BlueprintGraph/Private/K2Node_FormatText.cpp
+++ b/Engine/Source/Editor/BlueprintGraph/Private/K2Node_FormatText.cpp
@@ -39,7 +39,7 @@ UK2Node_FormatText::UK2Node_FormatText(const FObjectInitializer& ObjectInitializ
 	: Super(ObjectInitializer)
 	, CachedFormatPin(NULL)
 {
-	NodeTooltip = LOCTEXT("NodeTooltip", "Builds a formatted string using available format argument values.\n  \u2022 Use {} to denote format arguments.\n  \u2022 Argument types may be Byte, Integer, Float, Text, or ETextGender.");
+	NodeTooltip = LOCTEXT("NodeTooltip", "Builds a formatted string using available format argument values.\n  \u2022 Use {} to denote format arguments.\n  \u2022 Argument types may be Byte, Integer, Float, Text, String, Name, Boolean, Object or ETextGender.");
 }
 
 void UK2Node_FormatText::AllocateDefaultPins()
@@ -349,6 +349,23 @@ void UK2Node_FormatText::ExpandNode(class FKismetCompilerContext& CompilerContex
 		{
 			const FName& ArgumentPinCategory = ArgumentPin->PinType.PinCategory;
 
+			// Adds an implicit conversion node to this argument based on its function and pin name
+			auto AddConversionNode = [&](FName FuncName, TCHAR* PinName)
+			{
+				// Set the default value if there was something passed in, or default to "Text"
+				MakeFormatArgumentDataStruct->GetSchema()->TrySetDefaultValue(*ArgumentTypePin, TEXT("Text"));
+
+				// Spawn conversion node based on the given function name
+				UK2Node_CallFunction* ToTextFunction = CompilerContext.SpawnIntermediateNode<UK2Node_CallFunction>(this, SourceGraph);
+				ToTextFunction->SetFromFunction(UKismetTextLibrary::StaticClass()->FindFunctionByName(FuncName));
+				ToTextFunction->AllocateDefaultPins();
+				CompilerContext.MessageLog.NotifyIntermediateObjectCreation(ToTextFunction, this);
+
+				CompilerContext.MovePinLinksToIntermediate(*ArgumentPin, *ToTextFunction->FindPinChecked(PinName));
+
+				ToTextFunction->FindPinChecked(TEXT("ReturnValue"))->MakeLinkTo(MakeFormatArgumentDataStruct->FindPinChecked(GET_MEMBER_NAME_STRING_CHECKED(FFormatArgumentData, ArgumentValue)));
+			};
+
 			if (ArgumentPinCategory == UEdGraphSchema_K2::PC_Int)
 			{
 				MakeFormatArgumentDataStruct->GetSchema()->TrySetDefaultValue(*ArgumentTypePin, TEXT("Int"));
@@ -389,6 +406,22 @@ void UK2Node_FormatText::ExpandNode(class FKismetCompilerContext& CompilerContex
 					CompilerContext.MovePinLinksToIntermediate(*ArgumentPin, *MakeFormatArgumentDataStruct->FindPinChecked(GET_MEMBER_NAME_STRING_CHECKED(FFormatArgumentData, ArgumentValueGender)));
 				}
 			}
+			else if (ArgumentPinCategory == UEdGraphSchema_K2::PC_Boolean)
+			{
+				AddConversionNode(GET_MEMBER_NAME_CHECKED(UKismetTextLibrary, Conv_BoolToText), TEXT("InBool"));
+			}
+			else if (ArgumentPinCategory == UEdGraphSchema_K2::PC_Name)
+			{
+				AddConversionNode(GET_MEMBER_NAME_CHECKED(UKismetTextLibrary, Conv_NameToText), TEXT("InName"));
+			}
+			else if (ArgumentPinCategory == UEdGraphSchema_K2::PC_String)
+			{
+				AddConversionNode(GET_MEMBER_NAME_CHECKED(UKismetTextLibrary, Conv_StringToText), TEXT("InString"));
+			}
+			else if (ArgumentPinCategory == UEdGraphSchema_K2::PC_Object)
+			{
+				AddConversionNode(GET_MEMBER_NAME_CHECKED(UKismetTextLibrary, Conv_ObjectToText), TEXT("InObj"));
+			}
 			else
 			{
 				// Unexpected pin type!
@@ -493,8 +526,10 @@ bool UK2Node_FormatText::IsConnectionDisallowed(const UEdGraphPin* MyPin, const
 		const FName& OtherPinCategory = OtherPin->PinType.PinCategory;
 
 		bool bIsValidType = false;
-		if (OtherPinCategory == UEdGraphSchema_K2::PC_Int || OtherPinCategory == UEdGraphSchema_K2::PC_Float || OtherPinCategory == UEdGraphSchema_K2::PC_Text || (OtherPinCategory == UEdGraphSchema_K2::PC_Byte && !OtherPin->PinType.PinSubCategoryObject.IsValid()))
-		{
+		if (OtherPinCategory == UEdGraphSchema_K2::PC_Int || OtherPinCategory == UEdGraphSchema_K2::PC_Float || OtherPinCategory == UEdGraphSchema_K2::PC_Text ||
+			(OtherPinCategory == UEdGraphSchema_K2::PC_Byte && !OtherPin->PinType.PinSubCategoryObject.IsValid()) ||
+            OtherPinCategory == UEdGraphSchema_K2::PC_Boolean || OtherPinCategory == UEdGraphSchema_K2::PC_String || OtherPinCategory == UEdGraphSchema_K2::PC_Name || OtherPinCategory == UEdGraphSchema_K2::PC_Object)
+        {
 			bIsValidType = true;
 		}
 		else if (OtherPinCategory == UEdGraphSchema_K2::PC_Byte || OtherPinCategory == UEdGraphSchema_K2::PC_Enum)
@@ -508,7 +543,7 @@ bool UK2Node_FormatText::IsConnectionDisallowed(const UEdGraphPin* MyPin, const
 
 		if (!bIsValidType)
 		{
-			OutReason = LOCTEXT("Error_InvalidArgumentType", "Format arguments may only be Byte, Integer, Float, Text, or ETextGender.").ToString();
+			OutReason = LOCTEXT("Error_InvalidArgumentType", "Format arguments may only be Byte, Integer, Float, Text, String, Name, Boolean, Object, or ETextGender.").ToString();
 			return true;
 		}
 	}
