CIS fix

#rnx

diff --git a/Engine/Source/Editor/BlueprintGraph/Private/K2Node_FormatText.cpp b/Engine/Source/Editor/BlueprintGraph/Private/K2Node_FormatText.cpp
--- a/Engine/Source/Editor/BlueprintGraph/Private/K2Node_FormatText.cpp
+++ b/Engine/Source/Editor/BlueprintGraph/Private/K2Node_FormatText.cpp
@@ -350,7 +350,7 @@ void UK2Node_FormatText::ExpandNode(class FKismetCompilerContext& CompilerContex
 			const FName& ArgumentPinCategory = ArgumentPin->PinType.PinCategory;
 
 			// Adds an implicit conversion node to this argument based on its function and pin name
-			auto AddConversionNode = [&](FName FuncName, TCHAR* PinName)
+			auto AddConversionNode = [&](const FName FuncName, const TCHAR* PinName)
 			{
 				// Set the default value if there was something passed in, or default to "Text"
 				MakeFormatArgumentDataStruct->GetSchema()->TrySetDefaultValue(*ArgumentTypePin, TEXT("Text"));
