Object => Soft Object Reference conversion now works properly with self references

diff --git a/Engine/Source/Editor/BlueprintGraph/Private/EdGraphSchema_K2.cpp b/Engine/Source/Editor/BlueprintGraph/Private/EdGraphSchema_K2.cpp
--- a/Engine/Source/Editor/BlueprintGraph/Private/EdGraphSchema_K2.cpp
+++ b/Engine/Source/Editor/BlueprintGraph/Private/EdGraphSchema_K2.cpp
@@ -2801,8 +2801,8 @@ bool UEdGraphSchema_K2::FindSpecializedConversionNode(const UEdGraphPin* OutputP
 		}
 		else
 		{
-			UClass* InputClass  = Cast<UClass>(InputType.PinSubCategoryObject.Get());
-			UClass* OutputClass = Cast<UClass>(OutputType.PinSubCategoryObject.Get());
+			UClass* InputClass  = FBlueprintEditorUtils::GetTypeForPin(*InputPin);
+			UClass* OutputClass = FBlueprintEditorUtils::GetTypeForPin(*OutputPin);
 
 			if ((OutputType.PinCategory == PC_Interface) && (InputType.PinCategory == PC_Object))
 			{
diff --git a/Engine/Source/Editor/BlueprintGraph/Private/K2Node_ConvertAsset.cpp b/Engine/Source/Editor/BlueprintGraph/Private/K2Node_ConvertAsset.cpp
--- a/Engine/Source/Editor/BlueprintGraph/Private/K2Node_ConvertAsset.cpp
+++ b/Engine/Source/Editor/BlueprintGraph/Private/K2Node_ConvertAsset.cpp
@@ -5,6 +5,7 @@
 #include "EdGraph/EdGraphPin.h"
 #include "Engine/Blueprint.h"
 #include "Kismet/KismetSystemLibrary.h"
+#include "Kismet2/BlueprintEditorUtils.h"
 #include "EdGraphSchema_K2.h"
 #include "K2Node_CallFunction.h"
 #include "K2Node_DynamicCast.h"
@@ -26,7 +27,7 @@ UClass* UK2Node_ConvertAsset::GetTargetClass() const
 	UEdGraphPin* InputPin = FindPin(UK2Node_ConvertAssetImpl::InputPinName);
 	bool bIsConnected = InputPin && InputPin->LinkedTo.Num() && InputPin->LinkedTo[0];
 	UEdGraphPin* SourcePin = bIsConnected ? InputPin->LinkedTo[0] : nullptr;
-	return SourcePin ? Cast<UClass>(SourcePin->PinType.PinSubCategoryObject.Get()) : nullptr;
+	return SourcePin ? FBlueprintEditorUtils::GetTypeForPin(*SourcePin) : nullptr;
 }
 
 bool UK2Node_ConvertAsset::IsAssetClassType() const
diff --git a/Engine/Source/Editor/UnrealEd/Private/Kismet2/BlueprintEditorUtils.cpp b/Engine/Source/Editor/UnrealEd/Private/Kismet2/BlueprintEditorUtils.cpp
--- a/Engine/Source/Editor/UnrealEd/Private/Kismet2/BlueprintEditorUtils.cpp
+++ b/Engine/Source/Editor/UnrealEd/Private/Kismet2/BlueprintEditorUtils.cpp
@@ -9393,6 +9393,19 @@ UClass* FBlueprintEditorUtils::GetNativeParent(const UBlueprint* BP)
 	return Ret;
 }
 
+UClass* FBlueprintEditorUtils::GetTypeForPin(const UEdGraphPin& Pin)
+{
+	UClass* Ret = Cast<UClass>(Pin.PinType.PinSubCategoryObject.Get());
+
+	if(Ret == nullptr && Pin.PinType.PinSubCategory == UEdGraphSchema_K2::PSC_Self)
+	{
+		UBlueprint* Blueprint = FBlueprintEditorUtils::FindBlueprintForNodeChecked(Pin.GetOwningNode());
+		Ret = (Blueprint->GeneratedClass != NULL) ? Blueprint->GeneratedClass : Blueprint->ParentClass;
+	}
+
+	return Ret;
+}
+
 bool FBlueprintEditorUtils::ImplementsGetWorld(const UBlueprint* BP)
 {
 	if(UClass* NativeParent = GetNativeParent(BP))
diff --git a/Engine/Source/Editor/UnrealEd/Public/Kismet2/BlueprintEditorUtils.h b/Engine/Source/Editor/UnrealEd/Public/Kismet2/BlueprintEditorUtils.h
--- a/Engine/Source/Editor/UnrealEd/Public/Kismet2/BlueprintEditorUtils.h
+++ b/Engine/Source/Editor/UnrealEd/Public/Kismet2/BlueprintEditorUtils.h
@@ -1719,6 +1719,9 @@ public:
 	 */
 	static UClass* GetNativeParent(const UBlueprint* BP);
 
+	/** Returns the UClass type for an object pin, if any */
+	static UClass* GetTypeForPin(const UEdGraphPin& Pin);
+
 	/**
 	 * Returns true if this BP is currently based on a type that returns true for the UObject::ImplementsGetWorld() call:
 	 */
