Remove legacy "always compile on save" code path for LSBPs.

diff --git a/Engine/Source/Editor/BlueprintGraph/Private/BlueprintActionDatabase.cpp b/Engine/Source/Editor/BlueprintGraph/Private/BlueprintActionDatabase.cpp
--- a/Engine/Source/Editor/BlueprintGraph/Private/BlueprintActionDatabase.cpp
+++ b/Engine/Source/Editor/BlueprintGraph/Private/BlueprintActionDatabase.cpp
@@ -530,21 +530,6 @@ namespace BlueprintActionDatabaseImpl
 
 	/** True if a RefreshAll has been requested for the next Tick */
 	bool bRefreshAllRequested = false;
-
-	/** */
-	TWeakObjectPtr<UWorld> OriginalOwningWorldOnSave;
-
-	/** */
-	static void OnPreSaveWorld(uint32 SaveFlags, UWorld* World)
-	{
-		OriginalOwningWorldOnSave = World->PersistentLevel->OwningWorld;
-	}
-
-	/** */
-	static void OnPostSaveWorld(uint32 SaveFlags, UWorld* World, bool bSuccess)
-	{
-		OriginalOwningWorldOnSave.Reset();
-	}
 }
 
 //------------------------------------------------------------------------------
@@ -1152,27 +1137,55 @@ FBlueprintActionDatabase* FBlueprintActionDatabase::TryGet()
 FBlueprintActionDatabase::FBlueprintActionDatabase()
 {
 	RefreshAll();
-	FCoreUObjectDelegates::OnAssetLoaded.AddStatic(&BlueprintActionDatabaseImpl::OnAssetLoaded);
+	OnAssetLoadedDelegateHandle = FCoreUObjectDelegates::OnAssetLoaded.AddStatic(&BlueprintActionDatabaseImpl::OnAssetLoaded);
 
 	IAssetRegistry& AssetRegistry = FModuleManager::LoadModuleChecked<FAssetRegistryModule>(TEXT("AssetRegistry")).Get();
-	AssetRegistry.OnAssetAdded().AddStatic(&BlueprintActionDatabaseImpl::OnAssetAdded);
-	AssetRegistry.OnAssetRemoved().AddStatic(&BlueprintActionDatabaseImpl::OnAssetRemoved);
-	AssetRegistry.OnAssetRenamed().AddStatic(&BlueprintActionDatabaseImpl::OnAssetRenamed);
+	OnAssetAddedDelegateHandle = AssetRegistry.OnAssetAdded().AddStatic(&BlueprintActionDatabaseImpl::OnAssetAdded);
+	OnAssetRemovedDelegateHandle = AssetRegistry.OnAssetRemoved().AddStatic(&BlueprintActionDatabaseImpl::OnAssetRemoved);
+	OnAssetRenamedDelegateHandle = AssetRegistry.OnAssetRenamed().AddStatic(&BlueprintActionDatabaseImpl::OnAssetRenamed);
 
-	FEditorDelegates::OnAssetsPreDelete.AddStatic(&BlueprintActionDatabaseImpl::OnAssetsPendingDelete);
-	FKismetEditorUtilities::OnBlueprintUnloaded.AddStatic(&BlueprintActionDatabaseImpl::OnBlueprintUnloaded);
+	OnAssetsPreDeleteDelegateHandle = FEditorDelegates::OnAssetsPreDelete.AddStatic(&BlueprintActionDatabaseImpl::OnAssetsPendingDelete);
+	OnBlueprintUnloadedDelegateHandle = FKismetEditorUtilities::OnBlueprintUnloaded.AddStatic(&BlueprintActionDatabaseImpl::OnBlueprintUnloaded);
 
-	GEngine->OnWorldAdded().AddStatic(&BlueprintActionDatabaseImpl::OnWorldAdded);
-	GEngine->OnWorldDestroyed().AddStatic(&BlueprintActionDatabaseImpl::OnWorldDestroyed);
-	FWorldDelegates::RefreshLevelScriptActions.AddStatic(&BlueprintActionDatabaseImpl::OnRefreshLevelScripts);
+	OnWorldAddedDelegateHandle = GEngine->OnWorldAdded().AddStatic(&BlueprintActionDatabaseImpl::OnWorldAdded);
+	OnWorldDestroyedDelegateHandle = GEngine->OnWorldDestroyed().AddStatic(&BlueprintActionDatabaseImpl::OnWorldDestroyed);
+	RefreshLevelScriptActionsDelegateHandle = FWorldDelegates::RefreshLevelScriptActions.AddStatic(&BlueprintActionDatabaseImpl::OnRefreshLevelScripts);
 
-	FModuleManager::Get().OnModulesChanged().AddStatic(&BlueprintActionDatabaseImpl::OnModulesChanged);
+	OnModulesChangedDelegateHandle = FModuleManager::Get().OnModulesChanged().AddStatic(&BlueprintActionDatabaseImpl::OnModulesChanged);
 
 	IHotReloadInterface& HotReloadSupport = FModuleManager::LoadModuleChecked<IHotReloadInterface>("HotReload");
-	HotReloadSupport.OnHotReload().AddStatic(&BlueprintActionDatabaseImpl::OnProjectHotReloaded);
+	OnHotReloadDelegateHandle = HotReloadSupport.OnHotReload().AddStatic(&BlueprintActionDatabaseImpl::OnProjectHotReloaded);
+}
+
+//------------------------------------------------------------------------------
+FBlueprintActionDatabase::~FBlueprintActionDatabase()
+{
+	FCoreUObjectDelegates::OnAssetLoaded.Remove(OnAssetLoadedDelegateHandle);
 
-	FEditorDelegates::PreSaveWorld.AddStatic(&BlueprintActionDatabaseImpl::OnPreSaveWorld);
-	FEditorDelegates::PostSaveWorld.AddStatic(&BlueprintActionDatabaseImpl::OnPostSaveWorld);
+	if (FModuleManager::Get().IsModuleLoaded(TEXT("AssetRegistry")))
+	{
+		IAssetRegistry& AssetRegistry = FModuleManager::LoadModuleChecked<FAssetRegistryModule>(TEXT("AssetRegistry")).Get();
+		AssetRegistry.OnAssetAdded().Remove(OnAssetAddedDelegateHandle);
+		AssetRegistry.OnAssetAdded().Remove(OnAssetRemovedDelegateHandle);
+		AssetRegistry.OnAssetAdded().Remove(OnAssetRenamedDelegateHandle);
+	}
+
+	FEditorDelegates::OnAssetsPreDelete.Remove(OnAssetsPreDeleteDelegateHandle);
+	FKismetEditorUtilities::OnBlueprintUnloaded.Remove(OnBlueprintUnloadedDelegateHandle);
+
+	if (GEngine)
+	{
+		GEngine->OnWorldAdded().Remove(OnWorldAddedDelegateHandle);
+		GEngine->OnWorldAdded().Remove(OnWorldDestroyedDelegateHandle);
+	}
+
+	FWorldDelegates::RefreshLevelScriptActions.Remove(RefreshLevelScriptActionsDelegateHandle);
+	FModuleManager::Get().OnModulesChanged().Remove(OnModulesChangedDelegateHandle);
+
+	if (IHotReloadInterface* HotReloadSupport = FModuleManager::GetModulePtr<IHotReloadInterface>("HotReload"))
+	{
+		HotReloadSupport->OnHotReload().Remove(OnHotReloadDelegateHandle);
+	}
 }
 
 //------------------------------------------------------------------------------
@@ -1534,14 +1547,6 @@ void FBlueprintActionDatabase::RefreshAssetActions(UObject* const AssetObject)
 	{
 		for( ULevel* Level : WorldAsset->GetLevels() )
 		{
-			// Skip adding member actions for the persistent level on save if it doesn't match what the owning world was set to prior to saving.
-			// Otherwise, a new entry will be created for actions that aren't technically valid for the world asset, since this is likely a sublevel.
-			// When the world that owns the sublevel is destroyed, these actions won't be cleaned up, preventing the level from being freed during GC.
-			if (Level == WorldAsset->PersistentLevel && OriginalOwningWorldOnSave.IsValid() && Level->OwningWorld != OriginalOwningWorldOnSave.Get())
-			{
-				continue;
-			}
-
 			UBlueprint* LevelScript = Level->GetLevelScriptBlueprint(true);
 			if (LevelScript != nullptr)
 			{
diff --git a/Engine/Source/Editor/BlueprintGraph/Public/BlueprintActionDatabase.h b/Engine/Source/Editor/BlueprintGraph/Public/BlueprintActionDatabase.h
--- a/Engine/Source/Editor/BlueprintGraph/Public/BlueprintActionDatabase.h
+++ b/Engine/Source/Editor/BlueprintGraph/Public/BlueprintActionDatabase.h
@@ -33,6 +33,9 @@ public:
 	typedef TMap<FName, TArray<UBlueprintNodeSpawner*>>	FUnloadedActionRegistry;
 
 public:
+	/** Destructor */
+	virtual ~FBlueprintActionDatabase();
+
 	/**
 	 * Getter to access the database singleton. Will populate the database first 
 	 * if this is the first time accessing it.
@@ -191,8 +194,18 @@ private:
 	FOnDatabaseEntryUpdated EntryRefreshDelegate;
 	FOnDatabaseEntryUpdated EntryRemovedDelegate;
 
-	/** Handle to the registered OnBlueprintChanged delegate. */
-	FDelegateHandle OnBlueprintChangedDelegateHandle;
+	/** Handles to registered delegates. */
+	FDelegateHandle OnAssetLoadedDelegateHandle;
+	FDelegateHandle OnAssetAddedDelegateHandle;
+	FDelegateHandle OnAssetRemovedDelegateHandle;
+	FDelegateHandle OnAssetRenamedDelegateHandle;
+	FDelegateHandle OnAssetsPreDeleteDelegateHandle;
+	FDelegateHandle OnBlueprintUnloadedDelegateHandle;
+	FDelegateHandle OnWorldAddedDelegateHandle;
+	FDelegateHandle OnWorldDestroyedDelegateHandle;
+	FDelegateHandle RefreshLevelScriptActionsDelegateHandle;
+	FDelegateHandle OnModulesChangedDelegateHandle;
+	FDelegateHandle OnHotReloadDelegateHandle;
 
 	/** Pointer to the shared list of currently existing component types */
 	const TArray<struct FComponentTypeEntry>* ComponentTypes;
diff --git a/Engine/Source/Runtime/Engine/Private/World.cpp b/Engine/Source/Runtime/Engine/Private/World.cpp
--- a/Engine/Source/Runtime/Engine/Private/World.cpp
+++ b/Engine/Source/Runtime/Engine/Private/World.cpp
@@ -90,7 +90,6 @@
 
 #if WITH_EDITOR
 	#include "DerivedDataCacheInterface.h"
-	#include "Kismet2/KismetEditorUtilities.h"
 	#include "ThumbnailRendering/WorldThumbnailInfo.h"
 	#include "Editor/UnrealEdTypes.h"
 	#include "Settings/LevelEditorPlaySettings.h"
@@ -1043,17 +1042,6 @@ void UWorld::PostLoad()
 
 bool UWorld::PreSaveRoot(const TCHAR* Filename)
 {
-#if WITH_EDITOR
-	// Rebuild all level blueprints now to ensure no stale data is stored on the actors
-	if( !IsRunningCommandlet() )
-	{
-		for (UBlueprint* Blueprint : PersistentLevel->GetLevelBlueprints())
-		{
-			FKismetEditorUtilities::CompileBlueprint(Blueprint, EBlueprintCompileOptions::SkipGarbageCollection | EBlueprintCompileOptions::SkipSave);
-		}
-	}
-#endif
-
 	// Update components and keep track off whether we need to clean them up afterwards.
 	bool bCleanupIsRequired = false;
 	if(!PersistentLevel->bAreComponentsCurrentlyRegistered)
