Make ensure more precise
#rnx

diff --git a/Engine/Source/Editor/BlueprintGraph/Classes/K2Node.h b/Engine/Source/Editor/BlueprintGraph/Classes/K2Node.h
--- a/Engine/Source/Editor/BlueprintGraph/Classes/K2Node.h
+++ b/Engine/Source/Editor/BlueprintGraph/Classes/K2Node.h
@@ -396,6 +396,7 @@ class BLUEPRINTGRAPH_VTABLE UK2Node : public UEdGraphNode
 
 	BLUEPRINTGRAPH_API virtual int32 GetNodeRefreshPriority() const { return EBaseNodeRefreshPriority::Normal; }
 
+	BLUEPRINTGRAPH_API bool DoesWildcardPinAcceptContainer(const UEdGraphPin* Pin);
 	BLUEPRINTGRAPH_API virtual bool DoesInputWildcardPinAcceptArray(const UEdGraphPin* Pin) const { return true; }
 	BLUEPRINTGRAPH_API virtual bool DoesOutputWildcardPinAcceptContainer(const UEdGraphPin* Pin) const { return true; }
 
diff --git a/Engine/Source/Editor/BlueprintGraph/Private/K2Node.cpp b/Engine/Source/Editor/BlueprintGraph/Private/K2Node.cpp
--- a/Engine/Source/Editor/BlueprintGraph/Private/K2Node.cpp
+++ b/Engine/Source/Editor/BlueprintGraph/Private/K2Node.cpp
@@ -904,6 +904,21 @@ UK2Node::ERedirectType UK2Node::DoPinsMatchForReconstruction(const UEdGraphPin*
 	return RedirectType;
 }
 
+bool UK2Node::DoesWildcardPinAcceptContainer(const UEdGraphPin* Pin)
+{
+	check(Pin);
+	if (Pin->Direction == EGPD_Input)
+	{
+		return DoesInputWildcardPinAcceptArray(Pin);
+	}
+	else if (Pin->Direction == EGPD_Output)
+	{
+		return DoesOutputWildcardPinAcceptContainer(Pin);
+	}
+
+	return false;
+}
+
 void UK2Node::ReconstructSinglePin(UEdGraphPin* NewPin, UEdGraphPin* OldPin, ERedirectType RedirectType)
 {
 	DECLARE_SCOPE_CYCLE_COUNTER(TEXT("UK2Node::ReconstructSinglePin"), STAT_LinkerLoad_ReconstructSinglePin, STATGROUP_LoadTimeVerbose);
diff --git a/Engine/Source/Editor/BlueprintGraph/Private/K2Node_CallFunction.cpp b/Engine/Source/Editor/BlueprintGraph/Private/K2Node_CallFunction.cpp
--- a/Engine/Source/Editor/BlueprintGraph/Private/K2Node_CallFunction.cpp
+++ b/Engine/Source/Editor/BlueprintGraph/Private/K2Node_CallFunction.cpp
@@ -77,12 +77,24 @@ struct FCustomStructureParamHelper
 			{
 				UEdGraphPin* LinkedTo = Pin->LinkedTo[0];
 				check(LinkedTo);
-				ensure(!LinkedTo->PinType.IsContainer());
+
+				if (UK2Node* Node = Cast<UK2Node>(Pin->GetOwningNode()))
+				{
+					ensure(
+						!LinkedTo->PinType.IsContainer() ||
+						Node->DoesWildcardPinAcceptContainer(Pin)
+					);
+				}
+				else
+				{
+					ensure( !LinkedTo->PinType.IsContainer() );
+				}
 
 				Pin->PinType = LinkedTo->PinType;
 			}
 			else
 			{
+				Pin->PinType = FEdGraphPinType();
 				Pin->PinType.PinCategory = UEdGraphSchema_K2::PC_Wildcard;
 				Pin->PinType.PinSubCategory = NAME_None;
 				Pin->PinType.PinSubCategoryObject = nullptr;
