Always generate hooks for Push Model in blueprints.
This should allow for easier testing and a safer transition to Push Model.

diff --git a/Engine/Source/Editor/BlueprintGraph/Private/CallFunctionHandler.cpp b/Engine/Source/Editor/BlueprintGraph/Private/CallFunctionHandler.cpp
--- a/Engine/Source/Editor/BlueprintGraph/Private/CallFunctionHandler.cpp
+++ b/Engine/Source/Editor/BlueprintGraph/Private/CallFunctionHandler.cpp
@@ -15,6 +15,7 @@
 #include "Engine/BlueprintGeneratedClass.h"
 #include "KismetCompiler.h"
 #include "Net/Core/PushModel/PushModelMacros.h"
+#include "PushModelHelpers.h"
 
 #if WITH_PUSH_MODEL
 #include "Net/NetPushModelHelpers.h"
@@ -639,11 +640,48 @@ UFunction* FKCHandler_CallFunction::FindFunction(FKismetFunctionContext& Context
 	return nullptr;
 }
 
+static bool FindMatchingReferencedNetPropertyAndPin(TArray<UEdGraphPin*>& RemainingPins, FProperty* FunctionProperty, FProperty*& NetProperty, UEdGraphPin*& PropertyObjectPin)
+{
+	NetProperty = nullptr;
+	PropertyObjectPin = nullptr;
+
+	if (UNLIKELY(FunctionProperty->HasAllPropertyFlags(CPF_OutParm | CPF_ReferenceParm) && !FunctionProperty->HasAnyPropertyFlags(CPF_ReturnParm | CPF_ConstParm)))
+	{
+		for (int32 i = 0; i < RemainingPins.Num(); ++i)
+		{
+			if (FunctionProperty->GetFName() == RemainingPins[i]->PinName)
+			{
+				UEdGraphPin* ParamPin = RemainingPins[i];
+				RemainingPins.RemoveAtSwap(i);
+				if (UEdGraphPin* PinToTry = FEdGraphUtilities::GetNetFromPin(ParamPin))
+				{
+					// TODO: Should we traverse pin links to find references somehow?
+					//			E.G., How are select statements that pass through references
+					//					to net properties handled?
+
+					if (UK2Node_VariableGet* GetPropertyNode = Cast<UK2Node_VariableGet>(PinToTry->GetOwningNode()))
+					{
+						FProperty* ToCheck = GetPropertyNode->GetPropertyForVariable();
+						if (UNLIKELY(ToCheck->HasAnyPropertyFlags(CPF_Net)))
+						{
+							NetProperty = ToCheck;
+							PropertyObjectPin = GetPropertyNode->FindPinChecked(UEdGraphSchema_K2::PN_Self);
+							return true;
+						}
+					}
+				}
+
+				return false;
+			}
+		}
+	}
+
+	return false;
+}
+
 void FKCHandler_CallFunction::Transform(FKismetFunctionContext& Context, UEdGraphNode* Node)
 {
 	// Add an object reference pin for this call
-
-	
 	UK2Node_CallFunction* CallFuncNode = Cast<UK2Node_CallFunction>(Node);
 	if (!CallFuncNode)
 	{
@@ -678,7 +716,8 @@ void FKCHandler_CallFunction::Transform(FKismetFunctionContext& Context, UEdGrap
 	// Find the function, starting at the parent class
 	if (UFunction* Function = FindFunction(Context, Node))
 	{
-		if (Function->HasMetaData(FBlueprintMetadata::MD_Latent))
+		const bool bIsLatent = Function->HasMetaData(FBlueprintMetadata::MD_Latent);
+		if (bIsLatent)
 		{
 			UEdGraphPin* OldOutPin = K2Schema->FindExecutionPin(*CallFuncNode, EGPD_Output);
 
@@ -709,24 +748,14 @@ void FKCHandler_CallFunction::Transform(FKismetFunctionContext& Context, UEdGrap
 			}
 		}
 
-#if WITH_PUSH_MODEL
-
 		/**
-		 * Warning!! Similar code exists in VariableSetHandler.cpp
-		 *
 		 * This code is for property dirty tracking.
 		 * It works by injecting in extra nodes while compiling that will call UNetPushModelHelpers::MarkPropertyDirtyFromRepIndex.
+		 * See FKCPushModelHelpers::ConstructMarkDirtyNodeForProperty for node generation.
 		 *
-		 * That function will be called with the Owner of the property, and the RepIndex of the property.
-		 * One of these nodes needs to be added for ever Net Property that is passed by (out) reference to the function.
-		 *
-		 *
-		 * Note, this assumes that there's no way that a native class can add or remove replicated properties without also
-		 * recompiling the blueprint. The only scenario that seems possible is cooked games with custom built binaries,
-		 * but that still seems unsafe.
-		 *
-		 * If that can happen, we can instead switch to using the property name and resorting to a FindField call at runtime,
-		 * but that will be more expensive.
+		 * If we're passing net properties as reference variables, there's no guarantee what the function will do,
+		 * so we won't actually know whether or not the value has changed. In that case we'll go ahead
+		 * and mark the property as dirty.
 		 */
 
 		// If the function is pure but won't actually be evaluated, if there are no out params,
@@ -735,142 +764,94 @@ void FKCHandler_CallFunction::Transform(FKismetFunctionContext& Context, UEdGrap
 		// have any effect.
 		if (!bIsPureAndNoUsedOutputs && Function->NumParms > 0 && Function->HasAllFunctionFlags(FUNC_HasOutParms))
 		{
-			static const FName MarkPropertyDirtyFuncName(TEXT("MarkPropertyDirtyFromRepIndex"));
-			static const FName ObjectPinName(TEXT("Object"));
-			static const FName RepIndexPinName(TEXT("RepIndex"));
-			static const FName PropertyNamePinName(TEXT("PropertyName"));
+			// We don't care about returns (non-ref out params, or const-ref out params), because those would
+			// require a Set call to change any net property, and that will already have the appropriate nodes
+			// generated.
+			//
+			// Additionally, we only need to care about pins that are actually connected to something.
+			TArray<UEdGraphPin*> RemainingPins;
+			RemainingPins.Reserve(Node->Pins.Num());
+			for (UEdGraphPin* Pin : Node->Pins)
+			{
+				if (Pin->Direction == EGPD_Input &&
+					!Pin->PinType.bIsConst &&
+					Pin->PinType.bIsReference &&
+					Pin->LinkedTo.Num() > 0)
+				{
+					RemainingPins.Add(Pin);
+				}
+			}
 
-			TArray<UEdGraphPin*> RemainingPins(Node->Pins);
-			UEdGraphPin* OldThenPin = CallFuncNode->GetThenPin();
+			if (RemainingPins.Num() > 0)
+			{
+				FProperty* NetProperty = nullptr;
+				UEdGraphPin* PropertyObjectPin = nullptr;
+				UEdGraphPin* OldThenPin = CallFuncNode->GetThenPin();
 
-			// Note: This feels like it's going to be a very hot path during compilation as it will be hit for every
-			// CallFunction node. Any optimizations that can be made here are probably worth it to not terribly regress
-			// BP Compile time for cooked games.
+				// Note: This feels like it's going to be a very hot path during compilation as it will be hit for every
+				// CallFunction node. Any optimizations that can be made here are probably worth.
 
-			// Iterate the properties looking for Out Params that are tied to Net Properties.
-			// This is similar to the loop in CreateCallFunction
-			for (TFieldIterator<FProperty> It(Function); It && (It->PropertyFlags & CPF_Parm); ++It)
-			{
-				FProperty* Property = *It;
-				if (Property->HasAllPropertyFlags(CPF_OutParm | CPF_ReferenceParm))
+				// Iterate the properties looking for Out Params that are tied to Net Properties.
+				// This is similar to the loop in CreateCallFunction
+				for (TFieldIterator<FProperty> It(Function); It; ++It)
 				{
-					for (int32 i = 0; i < RemainingPins.Num(); ++i)
+					if (FindMatchingReferencedNetPropertyAndPin(RemainingPins, *It, NetProperty, PropertyObjectPin))
 					{
-						if (Property->GetFName() == RemainingPins[i]->PinName)
+						if (UEdGraphNode* MarkPropertyDirtyNode = FKCPushModelHelpers::ConstructMarkDirtyNodeForProperty(Context, NetProperty, PropertyObjectPin))
 						{
-							UEdGraphPin* ParamPin = RemainingPins[i];
-							if (UEdGraphPin * PinToTry = FEdGraphUtilities::GetNetFromPin(ParamPin))
+							bool bWereNodesAdded = false;
+							UEdGraphPin* NewThenPin = MarkPropertyDirtyNode->FindPinChecked(UEdGraphSchema_K2::PN_Then);
+							UEdGraphPin* NewInPin = MarkPropertyDirtyNode->FindPinChecked(UEdGraphSchema_K2::PN_Execute);
+
+							if (ensure(NewThenPin) && ensure(NewInPin))
 							{
-								if (UK2Node_VariableGet* GetPropertyNode = Cast<UK2Node_VariableGet>(PinToTry->GetOwningNode()))
+								if (OldThenPin)
 								{
-									FProperty* ToCheck = GetPropertyNode->GetPropertyForVariable();
+									NewThenPin->CopyPersistentDataFromOldPin(*OldThenPin);
+									OldThenPin->BreakAllPinLinks();
+									OldThenPin->MakeLinkTo(NewInPin);
 
-									// We only need to do this work if we actually found a net property.
-									if (UNLIKELY(ToCheck->HasAnyPropertyFlags(CPF_Net)))
-									{
-										if (OldThenPin == nullptr)
-										{
-											// TODO: Unfortunately, we don't have an efficient way to determine whether or not a property
-											//			uses push model at this point. That is only known from a call to "GetLifetimeRepProperties"
-											//			on a given object, and even then that is a user defined function that could do any number
-											//			of things.
-											//			We *might* be able to use something like Object Annotations to cache off the results of these
-											//			calls and clear the cache post compile, but that seems hacky and prone to error.
-											CompilerContext.MessageLog.Warning(TEXT("Passing net properties as Out Parameters with Push Model Enabled is not supported. Function=%s, Param=%s, Property=%s"),
-												*Function->GetName(), *Property->GetName(), *ToCheck->GetName());
-										}
-										else if (UClass * OwningClass = ToCheck->GetTypedOuter<UClass>())
-										{
-											// TODO: As of CL-10376545 in //UE4/Main, we no longer sort Native Replicated Properties by GC Offset.
-											//			This this chunk of code can probably be removed.
-
-											// We need to make sure this class already has its property offsets setup, otherwise
-											// the order of our replicated properties won't match, meaning the RepIndex will be
-											// invalid.
-											if (ToCheck->GetOffset_ForGC() == 0)
-											{
-												if (UBlueprint * Blueprint = Cast<UBlueprint>(OwningClass->ClassGeneratedBy))
-												{
-													if (UClass * UseClass = Blueprint->GeneratedClass)
-													{
-														OwningClass = UseClass;
-														ToCheck = FindFieldChecked<FProperty>(OwningClass, ToCheck->GetFName());
-													}
-												}
-											}
+									OldThenPin = NewThenPin;
+									bWereNodesAdded = true;
+								}
+								else
+								{
+									// If there's no then pin, we'll instead insert the dirty nodes before the execution
+									// of function with the reference.
+									// This may do weird things with Latent Nodes, so warn about that.
 
-											ensureAlwaysMsgf(ToCheck->GetOffset_ForGC() != 0,
-												TEXT("Class does not have Property Offsets setup. This will cause issues with Push Model. Blueprint=%s, Class=%s, Property=%s"),
-												*Context.Blueprint->GetPathName(), *OwningClass->GetPathName(), *ToCheck->GetName());
+									if (bIsLatent)
+									{
+										CompilerContext.MessageLog.Warning(*NSLOCTEXT("KismetCompiler", "LatentPushModel_Warning", "@@ is a latent node with references to a net property. The property may not be marked dirty in the correct frame.").ToString(), Node);
+									}
 
-											if (!OwningClass->HasAnyClassFlags(CLASS_ReplicationDataIsSetUp))
-											{
-												OwningClass->SetUpRuntimeReplicationData();
-											}
+									UEdGraphPin* OldInPin = CallFuncNode->FindPin(UEdGraphSchema_K2::PN_Execute);
+									if (ensure(OldInPin))
+									{
+										NewInPin->CopyPersistentDataFromOldPin(*OldInPin);
+										OldInPin->BreakAllPinLinks();
 
-											// Actually insert the nodes.
-											// TODO: I don't expect this to be a very frequent occurence, but we could be more clever about this.
-											//			Instead of inserting a new node per Out Ref Net Property, we might be able to create
-											//			A separate node that takes in lists of Objects, RepIndices, and Names.
-											//			Ideally, this would work like the String Append node where we could specify the values
-											//			directly on the node.
-											{
-												// Create the node that will call MarkPropertyDirty.
-												UK2Node_CallFunction* MarkPropertyDirtyNode = Node->GetGraph()->CreateIntermediateNode<UK2Node_CallFunction>();
-												MarkPropertyDirtyNode->FunctionReference.SetExternalMember(MarkPropertyDirtyFuncName, UNetPushModelHelpers::StaticClass());
-												MarkPropertyDirtyNode->AllocateDefaultPins();
-
-												// Create the Pins for RepIndex, PropertyName, and Object.
-												UEdGraphPin* RepIndexPin = MarkPropertyDirtyNode->FindPinChecked(RepIndexPinName);
-												RepIndexPin->DefaultValue = FString::FromInt(Property->RepIndex);
-
-												UEdGraphPin* PropertyNamePin = MarkPropertyDirtyNode->FindPinChecked(PropertyNamePinName);
-												PropertyNamePin->DefaultValue = Property->GetFName().ToString();
-
-												UEdGraphPin* ObjectPin = MarkPropertyDirtyNode->FindPinChecked(ObjectPinName);
-												UEdGraphPin* PropertyOwnerPin = GetPropertyNode->FindPinChecked(UEdGraphSchema_K2::PN_Self);
-
-												// If the property is linked to some other object, go ahead and grab that.
-												// Otherwise, create an intermediate self Pin and use that.
-												if (PropertyOwnerPin->LinkedTo.Num() > 0)
-												{
-													PropertyOwnerPin = PropertyOwnerPin->LinkedTo[0];
-												}
-												else
-												{
-													UK2Node_Self* SelfNode = Node->GetGraph()->CreateIntermediateNode<UK2Node_Self>();
-													SelfNode->AllocateDefaultPins();
-													PropertyOwnerPin = SelfNode->FindPinChecked(UEdGraphSchema_K2::PN_Self);
-												}
-
-												ObjectPin->MakeLinkTo(PropertyOwnerPin);
-
-												UEdGraphPin* NewThenPin = MarkPropertyDirtyNode->GetThenPin();
-												if (ensure(NewThenPin))
-												{
-													NewThenPin->CopyPersistentDataFromOldPin(*OldThenPin);
-													OldThenPin->BreakAllPinLinks();
-													OldThenPin->MakeLinkTo(MarkPropertyDirtyNode->GetExecPin());
-
-													// Now that the connection is established, go ahead and set our OldThenPin to
-													// the NewThenPin. If we find another NetProperty later, this will save us from
-													// doing a lookup.
-													OldThenPin = NewThenPin;
-												}
-											}
-										}
+										NewThenPin->MakeLinkTo(OldInPin);
+										OldThenPin = NewThenPin;
+										bWereNodesAdded = true;
 									}
 								}
 							}
 
-							RemainingPins.RemoveAtSwap(i);
-							break;
+							if (!bWereNodesAdded)
+							{
+								CompilerContext.MessageLog.Warning(*NSLOCTEXT("KismetCompiler", "PushModelNoDirty_Warning", "@@ has reference to net properties, but we were unable to generate dirty nodes.").ToString(), Node);
+							}
 						}
 					}
+
+					if (RemainingPins.Num() == 0)
+					{
+						break;
+					}
 				}
 			}
 		}
-#endif
 	}
 }
 
diff --git a/Engine/Source/Editor/BlueprintGraph/Private/PushModelHelpers.cpp b/Engine/Source/Editor/BlueprintGraph/Private/PushModelHelpers.cpp
new file mode 100644
--- /dev/null
+++ b/Engine/Source/Editor/BlueprintGraph/Private/PushModelHelpers.cpp
@@ -0,0 +1,112 @@
+// Copyright Epic Games, Inc. All Rights Reserved.
+
+#include "PushModelHelpers.h"
+#include "KismetCompiledFunctionContext.h"
+#include "UObject/Class.h"
+#include "UObject/UnrealType.h"
+#include "Net/NetPushModelHelpers.h"
+#include "K2Node_CallFunction.h"
+#include "K2Node_Self.h"
+#include "EdGraph/EdGraph.h"
+#include "EdGraph/EdGraphPin.h"
+
+UEdGraphNode* FKCPushModelHelpers::ConstructMarkDirtyNodeForProperty(FKismetFunctionContext& Context, FProperty* RepProperty, UEdGraphPin* PropertyObjectPin)
+{
+	/**
+	 * This code is for push model property dirty tracking.
+	 * It works by injecting extra nodes while compiling that will call UNetPushModelHelpers::MarkPropertyDirtyFromRepIndex.
+	 * These nodes will be generated even when push model is disabled, but when evaluated will be NoOps if PushModel is compiled
+	 * out. This is done to prevent the need for recooking data if Push Model is enabled or disabled for evaluation / testing.
+	 * TODO: Maybe have a CVar or something else that can prevent these nodes from being added?
+	 *
+	 * That function will be called with the Object that owns the property (either Self or whatever is connected to the Target pin
+	 * of the BP Node), the RepIndex of the property, and the Property's name.
+	 *
+	 * Note, this assumes that there's no way that a native class can add or remove replicated properties without also
+	 * recompiling blueprints. The only scenario that seems possible is cooked games with custom built binaries, and seems
+	 * unlikely (only really for testing purposes).
+	 *
+	 * If that becomes a problem, we can instead switch to using the property name and resorting to a FindField call at runtime,
+	 * but that will be significantly more expensive.
+	 */
+
+	static const FName MarkPropertyDirtyFuncName(TEXT("MarkPropertyDirtyFromRepIndex"));
+	static const FName ObjectPinName(TEXT("Object"));
+	static const FName RepIndexPinName(TEXT("RepIndex"));
+	static const FName PropertyNamePinName(TEXT("PropertyName"));
+
+	if (!ensure(RepProperty) || !ensure(PropertyObjectPin))
+	{
+		return nullptr;
+	}
+
+	UClass* OwningClass = RepProperty->GetTypedOwner<UClass>();
+	if (!ensure(OwningClass))
+	{
+		return nullptr;
+	}
+
+	// We need to make sure this class already has its property offsets setup, otherwise
+	// the order of our replicated properties won't match, meaning the RepIndex will be
+	// invalid.
+	if (RepProperty->GetOffset_ForGC() == 0)
+	{
+		// Make sure that we're using the correct class and that it has replication data set up.
+		if (OwningClass->ClassGeneratedBy == Context.Blueprint && Context.NewClass && Context.NewClass != OwningClass)
+		{
+			OwningClass = Context.NewClass;
+			RepProperty = FindFieldChecked<FProperty>(OwningClass, RepProperty->GetFName());
+		}
+
+		if (RepProperty->GetOffset_ForGC() == 0)
+		{
+			if (UBlueprint* Blueprint = Cast<UBlueprint>(OwningClass->ClassGeneratedBy))
+			{
+				if (UClass* UseClass = Blueprint->GeneratedClass)
+				{
+					OwningClass = UseClass;
+					RepProperty = FindFieldChecked<FProperty>(OwningClass, RepProperty->GetFName());
+				}
+			}
+		}
+	}
+
+	ensureAlwaysMsgf(RepProperty->GetOffset_ForGC() != 0,
+		TEXT("Class does not have Property Offsets setup. This will cause issues with Push Model. Blueprint=%s, Class=%s, Property=%s"),
+		*Context.Blueprint->GetPathName(), *OwningClass->GetPathName(), *RepProperty->GetName());
+
+	if (!OwningClass->HasAnyClassFlags(CLASS_ReplicationDataIsSetUp))
+	{
+		OwningClass->SetUpRuntimeReplicationData();
+	}
+
+	// Create the node that will call MarkPropertyDirty.
+	UK2Node_CallFunction* MarkPropertyDirtyNode = Context.SourceGraph->CreateIntermediateNode<UK2Node_CallFunction>();
+	MarkPropertyDirtyNode->FunctionReference.SetExternalMember(MarkPropertyDirtyFuncName, UNetPushModelHelpers::StaticClass());
+	MarkPropertyDirtyNode->AllocateDefaultPins();
+
+	// Create the Pins for RepIndex, PropertyName, and Object.
+	UEdGraphPin* RepIndexPin = MarkPropertyDirtyNode->FindPinChecked(RepIndexPinName);
+	RepIndexPin->DefaultValue = FString::FromInt(RepProperty->RepIndex);
+
+	UEdGraphPin* PropertyNamePin = MarkPropertyDirtyNode->FindPinChecked(PropertyNamePinName);
+	PropertyNamePin->DefaultValue = RepProperty->GetFName().ToString();
+
+	// If the property is linked to some other object, go ahead and grab that.
+	// Otherwise, create an intermediate self Pin and use that.
+	if (PropertyObjectPin->LinkedTo.Num() > 0)
+	{
+		PropertyObjectPin = PropertyObjectPin->LinkedTo[0];
+	}
+	else
+	{
+		UK2Node_Self* SelfNode = Context.SourceGraph->CreateIntermediateNode<UK2Node_Self>();
+		SelfNode->AllocateDefaultPins();
+		PropertyObjectPin = SelfNode->FindPinChecked(UEdGraphSchema_K2::PN_Self);
+	}
+
+	UEdGraphPin* ObjectPin = MarkPropertyDirtyNode->FindPinChecked(ObjectPinName);
+	ObjectPin->MakeLinkTo(PropertyObjectPin);
+
+	return MarkPropertyDirtyNode;
+}
\ No newline at end of file
diff --git a/Engine/Source/Editor/BlueprintGraph/Private/PushModelHelpers.h b/Engine/Source/Editor/BlueprintGraph/Private/PushModelHelpers.h
new file mode 100644
--- /dev/null
+++ b/Engine/Source/Editor/BlueprintGraph/Private/PushModelHelpers.h
@@ -0,0 +1,20 @@
+// Copyright Epic Games, Inc. All Rights Reserved.
+
+#pragma once
+
+//////////////////////////////////////////////////////////////////////////
+// FKCPushModelHelpers
+
+struct FKCPushModelHelpers
+{
+	/**
+	 * Generates nodes that will call the necessary functions for marking Push Model Properties dirty.
+	 * These nodes should be generated and linked whenever we know there's a replicated property that
+	 * is likely to have its value changed.
+	 *
+	 * @param Context				Current compilcation context.
+	 * @param RepProperty			The Replicated Property that's being reference.
+	 * @param PropertyObjectPin		The pin we'll use to grab a reference to the Object that owns the property.
+	 */
+	static class UEdGraphNode* ConstructMarkDirtyNodeForProperty(struct FKismetFunctionContext& Context, class FProperty* RepProperty, class UEdGraphPin* PropertyObjectPin);
+};
\ No newline at end of file
diff --git a/Engine/Source/Editor/BlueprintGraph/Private/VariableSetHandler.cpp b/Engine/Source/Editor/BlueprintGraph/Private/VariableSetHandler.cpp
--- a/Engine/Source/Editor/BlueprintGraph/Private/VariableSetHandler.cpp
+++ b/Engine/Source/Editor/BlueprintGraph/Private/VariableSetHandler.cpp
@@ -7,7 +7,7 @@
 #include "K2Node_VariableSet.h"
 #include "K2Node_Self.h"
 #include "Kismet2/BlueprintEditorUtils.h"
-#include "Net/NetPushModelHelpers.h"
+#include "PushModelHelpers.h"
 
 #include "EdGraphUtilities.h"
 #include "KismetCompiler.h"
@@ -202,101 +202,22 @@ void FKCHandler_VariableSet::Transform(FKismetFunctionContext& Context, UEdGraph
 		if (SetNotify->IsNetProperty())
 		{
 			/**
-			 * Warning!! Similar code exists in CallFunctionHandler.cpp
-			 *
 			 * This code is for property dirty tracking.
 			 * It works by injecting in extra nodes while compiling that will call UNetPushModelHelpers::MarkPropertyDirtyFromRepIndex.
-			 *
-			 * That function will be called with the Owner of the property (either Self or whatever is connected to the Target pin
-			 * of the BP Node), and the RepIndex of the property.
-			 *
-			 * Note, this assumes that there's no way that a native class can add or remove replicated properties without also
-			 * recompiling the blueprint. The only scenario that seems possible is cooked games with custom built binaries,
-			 * but that still seems unsafe.
-			 *
-			 * If that can happen, we can instead switch to using the property name and resorting to a FindField call at runtime,
-			 * but that will be more expensive.
+			 * See FKCPushModelHelpers::ConstructMarkDirtyNodeForProperty for node generation.
 			 */
-
-			static const FName MarkPropertyDirtyFuncName(TEXT("MarkPropertyDirtyFromRepIndex"));
-			static const FName ObjectPinName(TEXT("Object"));
-			static const FName RepIndexPinName(TEXT("RepIndex"));
-			static const FName PropertyNamePinName(TEXT("PropertyName"));
-
 			if (FProperty * Property = SetNotify->GetPropertyForVariable())
 			{
-				if (UClass* Class = Property->GetOwnerClass())
+				if (UEdGraphNode * MarkPropertyDirtyNode = FKCPushModelHelpers::ConstructMarkDirtyNodeForProperty(Context, Property, Node->FindPinChecked(UEdGraphSchema_K2::PN_Self)))
 				{
-					// We need to make sure this class already has its property offsets setup, otherwise
-					// the order of our replicated properties won't match, meaning the RepIndex will be
-					// invalid.
-					if (Property->GetOffset_ForGC() == 0)
-					{
-						// Make sure that we're using the correct class and that it has replication data set up.
-						if (Class->ClassGeneratedBy == Context.Blueprint && Context.NewClass && Context.NewClass != Class)
-						{
-							Class = Context.NewClass;
-							Property = FindFieldChecked<FProperty>(Class, Property->GetFName());
-						}
-						if (Property->GetOffset_ForGC() == 0)
-						{
-							if (UBlueprint * Blueprint = Cast<UBlueprint>(Class->ClassGeneratedBy))
-							{
-								if (UClass * UseClass = Blueprint->GeneratedClass)
-								{
-									Class = UseClass;
-									Property = FindFieldChecked<FProperty>(Class, Property->GetFName());
-								}
-							}
-						}
-					}
-
-					ensureAlwaysMsgf(Property->GetOffset_ForGC() != 0,
-						TEXT("Class does not have Property Offsets setup. This will cause issues with Push Model. Blueprint=%s, Class=%s, Property=%s"),
-						*Context.Blueprint->GetPathName(), *Class->GetPathName(), *Property->GetName());
-
-					if (!Class->HasAnyClassFlags(CLASS_ReplicationDataIsSetUp))
-					{
-						Class->SetUpRuntimeReplicationData();
-					}
-
-					UK2Node_CallFunction* CallFuncNode = Node->GetGraph()->CreateIntermediateNode<UK2Node_CallFunction>();
-					CallFuncNode->FunctionReference.SetExternalMember(MarkPropertyDirtyFuncName, UNetPushModelHelpers::StaticClass());
-					CallFuncNode->AllocateDefaultPins();
-
-					// Take our old Self (Target) pin and hook it up to the Object pin for UNetPushModelHelpers::MarkPropertyDirty.
-					// If our Self pin isn't hooked up to anything, then create an intermediate Self node and use that.
-					UEdGraphPin* SelfPin = Node->FindPinChecked(UEdGraphSchema_K2::PN_Self);
-					UEdGraphPin* ObjectPin = CallFuncNode->FindPinChecked(ObjectPinName);
-
-					if (SelfPin->LinkedTo.Num() > 0)
-					{
-						SelfPin = SelfPin->LinkedTo[0];
-					}
-					else
-					{
-						UK2Node_Self* SelfNode = Node->GetGraph()->CreateIntermediateNode<UK2Node_Self>();
-						SelfNode->AllocateDefaultPins();
-						SelfPin = SelfNode->FindPinChecked(UEdGraphSchema_K2::PN_Self);
-					}
-
-					ObjectPin->MakeLinkTo(SelfPin);
-
-					UEdGraphPin* RepIndexPin = CallFuncNode->FindPinChecked(RepIndexPinName);
-					RepIndexPin->DefaultValue = FString::FromInt(Property->RepIndex);
-
-					UEdGraphPin* PropertyNamePin = CallFuncNode->FindPinChecked(PropertyNamePinName);
-					PropertyNamePin->DefaultValue = Property->GetFName().ToString();
-
 					// Hook up our exec pins.
 					UEdGraphPin* OldThenPin = Node->FindPinChecked(UEdGraphSchema_K2::PN_Then);
-					UEdGraphPin* NewThenPin = CallFuncNode->GetThenPin();
-					if (ensure(NewThenPin))
-					{
-						NewThenPin->CopyPersistentDataFromOldPin(*OldThenPin);
-						OldThenPin->BreakAllPinLinks();
-						OldThenPin->MakeLinkTo(CallFuncNode->GetExecPin());
-					}
+					UEdGraphPin* NewThenPin = MarkPropertyDirtyNode->FindPinChecked(UEdGraphSchema_K2::PN_Then);
+					UEdGraphPin* NewInPin = MarkPropertyDirtyNode->FindPinChecked(UEdGraphSchema_K2::PN_Execute);
+
+					NewThenPin->CopyPersistentDataFromOldPin(*OldThenPin);
+					OldThenPin->BreakAllPinLinks();
+					OldThenPin->MakeLinkTo(NewInPin);
 				}
 			}
 		}
