Fixup the Call Function Handler code to use the most up to date Function Context.
This prevents issues where functions can't be found due to stale data.

diff --git a/Engine/Source/Editor/BlueprintGraph/Private/CallFunctionHandler.cpp b/Engine/Source/Editor/BlueprintGraph/Private/CallFunctionHandler.cpp
--- a/Engine/Source/Editor/BlueprintGraph/Private/CallFunctionHandler.cpp
+++ b/Engine/Source/Editor/BlueprintGraph/Private/CallFunctionHandler.cpp
@@ -1,6 +1,8 @@
 // Copyright Epic Games, Inc. All Rights Reserved.
 
 #include "CallFunctionHandler.h"
+
+#include "BlueprintCompilationManager.h"
 #include "UObject/MetaData.h"
 #include "EdGraphSchema_K2.h"
 #include "K2Node_Event.h"
@@ -633,6 +635,13 @@ UFunction* FKCHandler_CallFunction::FindFunction(FKismetFunctionContext& Context
 
 	if (CallingContext)
 	{
+		// It may be advisable to always do this branch in GetMostUpToDateClass, but
+		// being conservative:
+		if (!FBlueprintCompilationManager::IsGeneratedClassLayoutReady())
+		{
+			CallingContext = FBlueprintEditorUtils::GetMostUpToDateClass(CallingContext);
+		}
+
 		const FName FunctionName = GetFunctionNameFromNode(Node);
 		return CallingContext->FindFunctionByName(FunctionName);
 	}
