Don't attempt to insert Push Model Dirty Nodes if we know the function is Pure.
We will still warn in that case because it can cause issues with properties not being marked dirty, but will prevent unnecessary work and ensures.

diff --git a/Engine/Source/Editor/BlueprintGraph/Private/CallFunctionHandler.cpp b/Engine/Source/Editor/BlueprintGraph/Private/CallFunctionHandler.cpp
--- a/Engine/Source/Editor/BlueprintGraph/Private/CallFunctionHandler.cpp
+++ b/Engine/Source/Editor/BlueprintGraph/Private/CallFunctionHandler.cpp
@@ -697,8 +697,9 @@ void FKCHandler_CallFunction::Transform(FKismetFunctionContext& Context, UEdGrap
 		return;
 	}
 
+	const bool bIsPure = CallFuncNode->bIsPureFunc;
 	bool bIsPureAndNoUsedOutputs = false;
-	if (CallFuncNode->bIsPureFunc)
+	if (bIsPure)
 	{
 		// Flag for removal if pure and there are no consumers of the outputs
 		//@TODO: This isn't recursive (and shouldn't be here), it'll just catch the last node in a line of pure junk
@@ -806,9 +807,17 @@ void FKCHandler_CallFunction::Transform(FKismetFunctionContext& Context, UEdGrap
 				{
 					if (FindMatchingReferencedNetPropertyAndPin(RemainingPins, *It, NetProperty, PropertyObjectPin))
 					{
+						if (bIsPure)
+						{
+							// TODO: JDN - Reenable this when we can discern Const from Non Const Pure Refs.
+							// Don't need to cause BP Spam since no one is using this right now.
+							// 	CompilerContext.MessageLog.Warning(*NSLOCTEXT("KismetCompiler", "PurePushModel_Warning", "@@ is a pure node with references to a net property. The property may not be marked dirty in the correct frame. Consider making the function impure to solve the problem.").ToString(), Node);
+							break;
+						}
+
 						if (UEdGraphNode* MarkPropertyDirtyNode = FKCPushModelHelpers::ConstructMarkDirtyNodeForProperty(Context, NetProperty, PropertyObjectPin))
 						{
-							bool bWereNodesAdded = false;
+							// bool bWereNodesAdded = false;
 							UEdGraphPin* NewThenPin = MarkPropertyDirtyNode->FindPinChecked(UEdGraphSchema_K2::PN_Then);
 							UEdGraphPin* NewInPin = MarkPropertyDirtyNode->FindPinChecked(UEdGraphSchema_K2::PN_Execute);
 
@@ -821,7 +830,7 @@ void FKCHandler_CallFunction::Transform(FKismetFunctionContext& Context, UEdGrap
 									OldThenPin->MakeLinkTo(NewInPin);
 
 									OldThenPin = NewThenPin;
-									bWereNodesAdded = true;
+									// bWereNodesAdded = true;
 								}
 								else
 								{
@@ -835,22 +844,27 @@ void FKCHandler_CallFunction::Transform(FKismetFunctionContext& Context, UEdGrap
 									}
 
 									UEdGraphPin* OldInPin = CallFuncNode->FindPin(UEdGraphSchema_K2::PN_Execute);
-									if (ensure(OldInPin))
+									if (OldInPin)
 									{
 										NewInPin->CopyPersistentDataFromOldPin(*OldInPin);
 										OldInPin->BreakAllPinLinks();
 
 										NewThenPin->MakeLinkTo(OldInPin);
 										OldThenPin = NewThenPin;
-										bWereNodesAdded = true;
+										// bWereNodesAdded = true;
 									}
 								}
 							}
 
+							/*
 							if (!bWereNodesAdded)
 							{
-								CompilerContext.MessageLog.Warning(*NSLOCTEXT("KismetCompiler", "PushModelNoDirty_Warning", "@@ has reference to net properties, but we were unable to generate dirty nodes.").ToString(), Node);
+								// TODO: JDN - Reenable this once we have other edge cases worked out.
+								// This warning is confusing / not necessarily actionable, and could contribute to spam,
+								// but no one currently relies on these features.
+								// CompilerContext.MessageLog.Warning(*NSLOCTEXT("KismetCompiler", "PushModelNoDirty_Warning", "@@ has reference to net properties, but we were unable to generate dirty nodes.").ToString(), Node);
 							}
+							*/
 						}
 					}
 
