ExpandEnumAsExecs now works with bool parameters, added ExpandBoolAsExecs synonym

diff --git a/Engine/Source/Developer/BlueprintCompilerCppBackend/Private/BlueprintCompilerCppBackendUtils.cpp b/Engine/Source/Developer/BlueprintCompilerCppBackend/Private/BlueprintCompilerCppBackendUtils.cpp
--- a/Engine/Source/Developer/BlueprintCompilerCppBackend/Private/BlueprintCompilerCppBackendUtils.cpp
+++ b/Engine/Source/Developer/BlueprintCompilerCppBackend/Private/BlueprintCompilerCppBackendUtils.cpp
@@ -613,7 +613,7 @@ bool FEmitHelper::MetaDataCanBeNative(const FName MetaDataName, FFieldVariant Fi
 	{
 		return false;
 	}
-	if (MetaDataName == TEXT("ExpandEnumAsExecs"))	// applicable to editor only
+	if (MetaDataName == TEXT("ExpandEnumAsExecs") || MetaDataName == TEXT("ExpandBoolAsExecs"))	// applicable to editor only
 	{
 		return false;
 	}
diff --git a/Engine/Source/Editor/BlueprintGraph/Classes/EdGraphSchema_K2.h b/Engine/Source/Editor/BlueprintGraph/Classes/EdGraphSchema_K2.h
--- a/Engine/Source/Editor/BlueprintGraph/Classes/EdGraphSchema_K2.h
+++ b/Engine/Source/Editor/BlueprintGraph/Classes/EdGraphSchema_K2.h
@@ -181,6 +181,9 @@ public:
 	/** Indicates that during compile we want to create multiple exec pins from an enum param */
 	static const FName MD_ExpandEnumAsExecs;
 
+	/** Synonym for MD_ExpandEnumAsExecs */
+	static const FName MD_ExpandBoolAsExecs;
+
 	static const FName MD_CommutativeAssociativeBinaryOperator;
 
 	/** Metadata string that indicates to use the MaterialParameterCollectionFunction node. */
diff --git a/Engine/Source/Editor/BlueprintGraph/Private/EdGraphSchema_K2.cpp b/Engine/Source/Editor/BlueprintGraph/Private/EdGraphSchema_K2.cpp
--- a/Engine/Source/Editor/BlueprintGraph/Private/EdGraphSchema_K2.cpp
+++ b/Engine/Source/Editor/BlueprintGraph/Private/EdGraphSchema_K2.cpp
@@ -151,6 +151,7 @@ const FName FBlueprintMetadata::MD_ProhibitedInterfaces(TEXT("ProhibitedInterfac
 const FName FBlueprintMetadata::MD_FunctionKeywords(TEXT("Keywords"));
 
 const FName FBlueprintMetadata::MD_ExpandEnumAsExecs(TEXT("ExpandEnumAsExecs"));
+const FName FBlueprintMetadata::MD_ExpandBoolAsExecs(TEXT("ExpandBoolAsExecs"));
 
 const FName FBlueprintMetadata::MD_CommutativeAssociativeBinaryOperator(TEXT("CommutativeAssociativeBinaryOperator"));
 const FName FBlueprintMetadata::MD_MaterialParameterCollectionFunction(TEXT("MaterialParameterCollectionFunction"));
diff --git a/Engine/Source/Editor/BlueprintGraph/Private/K2Node_CallFunction.cpp b/Engine/Source/Editor/BlueprintGraph/Private/K2Node_CallFunction.cpp
--- a/Engine/Source/Editor/BlueprintGraph/Private/K2Node_CallFunction.cpp
+++ b/Engine/Source/Editor/BlueprintGraph/Private/K2Node_CallFunction.cpp
@@ -449,6 +449,32 @@ bool FDynamicOutputHelper::CanConformPinType(const UK2Node_CallFunction* FuncNod
 	return bIsProperType;
 }
 
+static bool WantsExecPinsForParams(const UFunction* TargetFunction)
+{
+	check(TargetFunction);
+	return TargetFunction->HasMetaData(FBlueprintMetadata::MD_ExpandEnumAsExecs) ||
+		TargetFunction->HasMetaData(FBlueprintMetadata::MD_ExpandBoolAsExecs);
+}
+
+static FString GetAllExecParams(const UFunction* TargetFunction)
+{
+	check(TargetFunction);
+	FString Ret = TargetFunction->GetMetaData(FBlueprintMetadata::MD_ExpandEnumAsExecs);
+	const FString& ExpandBools = TargetFunction->GetMetaData(FBlueprintMetadata::MD_ExpandBoolAsExecs);
+	if (!ExpandBools.IsEmpty())
+	{
+		if(!Ret.IsEmpty())
+		{
+			Ret += ", " + ExpandBools;
+		}
+		else
+		{
+			Ret = ExpandBools;
+		}
+	}
+	return Ret;
+}
+
 /*******************************************************************************
  *  UK2Node_CallFunction
  ******************************************************************************/
@@ -786,8 +812,12 @@ void UK2Node_CallFunction::CreateExecPinsForFunctionCall(const UFunction* Functi
 					Prop = EnumProp;
 					Enum = EnumProp->GetEnum();
 				}
+				else if (FBoolProperty* BoolProp = FindField<FBoolProperty>(Function, EnumParamName))
+				{
+					Prop = BoolProp;
+				}
 
-				if (Prop != nullptr && Enum != nullptr)
+				if (Prop != nullptr)
 				{
 					const bool bIsFunctionInput = !Prop->HasAnyPropertyFlags(CPF_ReturnParm) &&
 						(!Prop->HasAnyPropertyFlags(CPF_OutParm) ||
@@ -806,46 +836,57 @@ void UK2Node_CallFunction::CreateExecPinsForFunctionCall(const UFunction* Functi
 						PreviousInput = Prop;
 					}
 
-					// yay, found it! Now create exec pin for each
-					int32 NumExecs = (Enum->NumEnums() - 1);
-					for (int32 ExecIdx = 0; ExecIdx < NumExecs; ExecIdx++)
+					if (Enum)
 					{
-						bool const bShouldBeHidden = Enum->HasMetaData(TEXT("Hidden"), ExecIdx) || Enum->HasMetaData(TEXT("Spacer"), ExecIdx);
-						if (!bShouldBeHidden)
+						// yay, found it! Now create exec pin for each
+						int32 NumExecs = (Enum->NumEnums() - 1);
+						for (int32 ExecIdx = 0; ExecIdx < NumExecs; ExecIdx++)
 						{
-							// Can't use Enum->GetNameByIndex here because it doesn't do namespace mangling
-							const FString NameStr = Enum->GetNameStringByIndex(ExecIdx);
-							
-							UEdGraphPin* CreatedPin = nullptr;
-
-							// todo: really only makes sense if there are multiple outputs
-							if (bIsFunctionInput || EnumNames.Num() == 1)
+							bool const bShouldBeHidden = Enum->HasMetaData(TEXT("Hidden"), ExecIdx) || Enum->HasMetaData(TEXT("Spacer"), ExecIdx);
+							if (!bShouldBeHidden)
 							{
-								CreatedPin = CreatePin(Direction, UEdGraphSchema_K2::PC_Exec, *NameStr);
-							}
-							else
-							{
-								CreatedPin = CreatePin(Direction, UEdGraphSchema_K2::PC_Exec, *NameStr);
-								CreatedPin->PinFriendlyName = FText::FromString(FString::Printf(TEXT("(%s) %s"), *Prop->GetDisplayNameText().ToString(), *NameStr));
-							}
+								// Can't use Enum->GetNameByIndex here because it doesn't do namespace mangling
+								const FString NameStr = Enum->GetNameStringByIndex(ExecIdx);
 
-							ExpandAsEnumPins.Add(CreatedPin);
+								UEdGraphPin* CreatedPin = nullptr;
 
-							if (Enum->HasMetaData(TEXT("Tooltip"), ExecIdx))
-							{
-								FString EnumTooltip = Enum->GetMetaData(TEXT("Tooltip"), ExecIdx);
-
-								if (const UEdGraphSchema_K2* const K2Schema = Cast<const UEdGraphSchema_K2>(GetSchema()))
+								// todo: really only makes sense if there are multiple outputs
+								if (bIsFunctionInput || EnumNames.Num() == 1)
 								{
-									K2Schema->ConstructBasicPinTooltip(*CreatedPin, FText::FromString(EnumTooltip), CreatedPin->PinToolTip);
+									CreatedPin = CreatePin(Direction, UEdGraphSchema_K2::PC_Exec, *NameStr);
 								}
 								else
 								{
-									CreatedPin->PinToolTip = EnumTooltip;
+									CreatedPin = CreatePin(Direction, UEdGraphSchema_K2::PC_Exec, *NameStr);
+									CreatedPin->PinFriendlyName = FText::FromString(FString::Printf(TEXT("(%s) %s"), *Prop->GetDisplayNameText().ToString(), *NameStr));
+								}
+
+								ExpandAsEnumPins.Add(CreatedPin);
+
+								if (Enum->HasMetaData(TEXT("Tooltip"), ExecIdx))
+								{
+									FString EnumTooltip = Enum->GetMetaData(TEXT("Tooltip"), ExecIdx);
+
+									if (const UEdGraphSchema_K2* const K2Schema = Cast<const UEdGraphSchema_K2>(GetSchema()))
+									{
+										K2Schema->ConstructBasicPinTooltip(*CreatedPin, FText::FromString(EnumTooltip), CreatedPin->PinToolTip);
+									}
+									else
+									{
+										CreatedPin->PinToolTip = EnumTooltip;
+									}
 								}
 							}
 						}
 					}
+					else
+					{
+						check(Prop->IsA<FBoolProperty>());
+						// Create a pin for true and false, note that the order here does not match the
+						// numeric order of bool, but it is more natural to put true first (e.g. to match branch node):
+						ExpandAsEnumPins.Add(CreatePin(Direction, UEdGraphSchema_K2::PC_Exec, TEXT("True")));
+						ExpandAsEnumPins.Add(CreatePin(Direction, UEdGraphSchema_K2::PC_Exec, TEXT("False")));
+					}
 
 					if (bIsFunctionInput)
 					{
@@ -883,7 +924,7 @@ void UK2Node_CallFunction::DetermineWantsEnumToExecExpansion(const UFunction* Fu
 {
 	bWantsEnumToExecExpansion = false;
 
-	if (Function->HasMetaData(FBlueprintMetadata::MD_ExpandEnumAsExecs))
+	if (WantsExecPinsForParams(Function))
 	{
 		TArray<FName> EnumNamesToCheck;
 		GetExpandEnumPinNames(Function, EnumNamesToCheck);
@@ -898,6 +939,15 @@ void UK2Node_CallFunction::DetermineWantsEnumToExecExpansion(const UFunction* Fu
 				bWantsEnumToExecExpansion = true;
 				EnumNamesToCheck.RemoveAt(i);
 			}
+			else
+			{
+				FBoolProperty* BoolProp = FindField<FBoolProperty>(Function, EnumParamName);
+				if (BoolProp)
+				{
+					bWantsEnumToExecExpansion = true;
+					EnumNamesToCheck.RemoveAt(i);
+				}
+			}
 		}
 		
 		if (bWantsEnumToExecExpansion && EnumNamesToCheck.Num() > 0 && !bHasCompilerMessage)
@@ -934,7 +984,7 @@ void UK2Node_CallFunction::GetExpandEnumPinNames(const UFunction* Function, TArr
 	EnumNamesToCheck.Reset();
 
 	// todo: use metadatacache if/when that's accepted.
-	const FString& EnumParamString = Function->GetMetaData(FBlueprintMetadata::MD_ExpandEnumAsExecs);
+	const FString EnumParamString = GetAllExecParams(Function);
 
 	TArray<FString> RawGroupings;
 	EnumParamString.ParseIntoArray(RawGroupings, TEXT(","), false);
@@ -1992,10 +2042,10 @@ void UK2Node_CallFunction::ValidateNodeDuringCompilation(class FCompilerResultsL
 		FText const WarningFormat = LOCTEXT("FunctionNotFoundFmt", "Could not find a function named \"{0}\" in '{1}'.\nMake sure '{2}' has been compiled for @@");
 		MessageLog.Error(*FText::Format(WarningFormat, FText::FromString(FunctName), FText::FromString(OwnerName), FText::FromString(OwnerName)).ToString(), this);
 	}
-	else if (Function->HasMetaData(FBlueprintMetadata::MD_ExpandEnumAsExecs) && bWantsEnumToExecExpansion == false)
+	else if (WantsExecPinsForParams(Function) && bWantsEnumToExecExpansion == false)
 	{
 		// will technically not have a properly formatted output for multiple params... but /shrug. 
-		const FString& EnumParamName = Function->GetMetaData(FBlueprintMetadata::MD_ExpandEnumAsExecs);
+		const FString EnumParamName = GetAllExecParams(Function);
 		MessageLog.Warning(*FText::Format(LOCTEXT("EnumToExecExpansionFailedFmt", "Unable to find enum parameter with name '{0}' to expand for @@"), FText::FromString(EnumParamName)).ToString(), this);
 	}
 	
@@ -2243,6 +2293,45 @@ void UK2Node_CallFunction::ExpandNode(class FKismetCompilerContext& CompilerCont
 			UK2Node_ExecutionSequence* SpawnedSequenceNode = nullptr;
 			int32 OutSequenceIndex = 0;
 
+			const auto LinkIntoOutputChain = [&OutMainExecutePin, &SpawnedSequenceNode, &OutSequenceIndex, &CompilerContext, this, SourceGraph, Schema](UK2Node* Node)
+			{
+				if (!OutMainExecutePin)
+				{
+					// Create normal exec output -- only once though.
+					OutMainExecutePin = CreatePin(EGPD_Output, UEdGraphSchema_K2::PC_Exec, UEdGraphSchema_K2::PN_Then);
+				}
+				else
+				{
+					// set up a sequence so we can call one after another.
+					if (!SpawnedSequenceNode)
+					{
+						SpawnedSequenceNode = CompilerContext.SpawnIntermediateNode<UK2Node_ExecutionSequence>(this, SourceGraph);
+						SpawnedSequenceNode->AllocateDefaultPins();
+						CompilerContext.MovePinLinksToIntermediate(*OutMainExecutePin, *SpawnedSequenceNode->GetThenPinGivenIndex(OutSequenceIndex++));
+						Schema->TryCreateConnection(OutMainExecutePin, SpawnedSequenceNode->Pins[0]);
+					}
+				}
+
+				// Hook up execution to the branch node
+				if (!SpawnedSequenceNode)
+				{
+					Schema->TryCreateConnection(OutMainExecutePin, Node->GetExecPin());
+				}
+				else
+				{
+					UEdGraphPin* SequenceOutput = SpawnedSequenceNode->GetThenPinGivenIndex(OutSequenceIndex);
+
+					if (!SequenceOutput)
+					{
+						SpawnedSequenceNode->AddInputPin();
+						SequenceOutput = SpawnedSequenceNode->GetThenPinGivenIndex(OutSequenceIndex);
+					}
+
+					Schema->TryCreateConnection(SequenceOutput, Node->GetExecPin());
+					OutSequenceIndex++;
+				}
+			};
+
 			for (const FName& EnumParamName : EnumNamesToCheck)
 			{
 				UEnum* Enum = nullptr;
@@ -2318,47 +2407,14 @@ void UK2Node_CallFunction::ExpandNode(class FKismetCompilerContext& CompilerCont
 					// Expanded as output execs pins
 					else if (EnumParamPin->Direction == EGPD_Output)
 					{
-						if (!OutMainExecutePin)
-						{
-							// Create normal exec output -- only once though.
-							OutMainExecutePin = CreatePin(EGPD_Output, UEdGraphSchema_K2::PC_Exec, UEdGraphSchema_K2::PN_Then);
-						}
-						else
-						{	
-							// set up a sequence so we can call one after another.
-							if (!SpawnedSequenceNode)
-							{
-								SpawnedSequenceNode = CompilerContext.SpawnIntermediateNode<UK2Node_ExecutionSequence>(this, SourceGraph);
-								SpawnedSequenceNode->AllocateDefaultPins();
-								CompilerContext.MovePinLinksToIntermediate(*OutMainExecutePin, *SpawnedSequenceNode->GetThenPinGivenIndex(OutSequenceIndex++));
-								Schema->TryCreateConnection(OutMainExecutePin, SpawnedSequenceNode->Pins[0]);
-							}
-						}
-
 						// Create a SwitchEnum node to switch on the output enum
 						UK2Node_SwitchEnum* SwitchEnumNode = CompilerContext.SpawnIntermediateNode<UK2Node_SwitchEnum>(this, SourceGraph);
 						UEnum* EnumObject = Cast<UEnum>(EnumParamPin->PinType.PinSubCategoryObject.Get());
 						SwitchEnumNode->SetEnum(EnumObject);
 						SwitchEnumNode->AllocateDefaultPins();
 
-						// Hook up execution to the switch node
-						if (!SpawnedSequenceNode)
-						{
-							Schema->TryCreateConnection(OutMainExecutePin, SwitchEnumNode->GetExecPin());
-						}
-						else
-						{
-							UEdGraphPin* SequenceOutput = SpawnedSequenceNode->GetThenPinGivenIndex(OutSequenceIndex);
-
-							if (!SequenceOutput)
-							{
-								SpawnedSequenceNode->AddInputPin();
-								SequenceOutput = SpawnedSequenceNode->GetThenPinGivenIndex(OutSequenceIndex);
-							}
+						LinkIntoOutputChain(SwitchEnumNode);
 
-							Schema->TryCreateConnection(SequenceOutput, SwitchEnumNode->GetExecPin());
-							OutSequenceIndex++;
-						}
 						// Connect (hidden) enum pin to switch node's selection pin
 						Schema->TryCreateConnection(EnumParamPin, SwitchEnumNode->GetSelectionPin());
 
@@ -2394,6 +2450,74 @@ void UK2Node_CallFunction::ExpandNode(class FKismetCompilerContext& CompilerCont
 						}
 					}
 				}
+				else if(EnumParamPin && !EnumParamPin->PinType.IsContainer() && EnumParamPin->PinType.PinCategory == UEdGraphSchema_K2::PC_Boolean)
+				{
+					if (EnumParamPin->Direction == EGPD_Input)
+					{
+						// Create normal exec input
+						UEdGraphPin* ExecutePin = CreatePin(EGPD_Input, UEdGraphSchema_K2::PC_Exec, UEdGraphSchema_K2::PN_Execute);
+						// Create temp bool variable
+						UK2Node_TemporaryVariable* TempBoolVarNode = CompilerContext.SpawnIntermediateNode<UK2Node_TemporaryVariable>(this, SourceGraph);
+						TempBoolVarNode->VariableType.PinCategory = UEdGraphSchema_K2::PC_Boolean;
+						TempBoolVarNode->AllocateDefaultPins();
+						// Get the output pin
+						UEdGraphPin* TempBoolVarOutput = TempBoolVarNode->GetVariablePin();
+						// Connect temp enum variable to (hidden) bool pin
+						Schema->TryCreateConnection(TempBoolVarOutput, EnumParamPin);
+						
+						// create a true entry and a false:
+						const auto CreateAssignNode = [Schema, &CompilerContext, this, SourceGraph, TempBoolVarOutput, ExecutePin](UEdGraphPin* FakePin, const TCHAR* DefaultValue)
+						{
+							UK2Node_AssignmentStatement* AssignNode = CompilerContext.SpawnIntermediateNode<UK2Node_AssignmentStatement>(this, SourceGraph);
+							AssignNode->AllocateDefaultPins();
+
+							// Move connections from fake 'enum exec' pint to this assignment node
+							CompilerContext.MovePinLinksToIntermediate(*FakePin, *AssignNode->GetExecPin());
+
+							// Connect this to out temp enum var
+							Schema->TryCreateConnection(AssignNode->GetVariablePin(), TempBoolVarOutput);
+
+							// Connect exec output to 'real' exec pin
+							Schema->TryCreateConnection(AssignNode->GetThenPin(), ExecutePin);
+
+							// set the literal enum value to set to
+							AssignNode->GetValuePin()->DefaultValue = DefaultValue;
+						};
+
+						UEdGraphPin* TruePin = FindPinChecked(TEXT("True"), EEdGraphPinDirection::EGPD_Input);
+						UEdGraphPin* FalsePin = FindPinChecked(TEXT("False"), EEdGraphPinDirection::EGPD_Input);
+
+						CreateAssignNode(TruePin, TEXT("True"));
+						CreateAssignNode(FalsePin, TEXT("False"));
+
+						// remove fake false/true nodes:
+						RemovePin(TruePin);
+						RemovePin(FalsePin);
+					}
+					else if (EnumParamPin->Direction == EGPD_Output)
+					{
+						// Create a Branch node to switch on the output bool:
+						UK2Node_IfThenElse* IfElseNode = CompilerContext.SpawnIntermediateNode<UK2Node_IfThenElse>(this, SourceGraph);
+						IfElseNode->AllocateDefaultPins();
+
+						LinkIntoOutputChain(IfElseNode);
+
+						// Connect (hidden) bool pin to branch node
+						Schema->TryCreateConnection(EnumParamPin, IfElseNode->GetConditionPin());
+
+						UEdGraphPin* TruePin = FindPinChecked(TEXT("True"), EEdGraphPinDirection::EGPD_Output);
+						UEdGraphPin* FalsePin = FindPinChecked(TEXT("False"), EEdGraphPinDirection::EGPD_Output);
+
+						// move true connection to branch node:
+						CompilerContext.MovePinLinksToIntermediate(*TruePin, *IfElseNode->GetThenPin());
+						// move false connection to branch node:
+						CompilerContext.MovePinLinksToIntermediate(*FalsePin, *IfElseNode->GetElsePin());
+
+						// remove fake false/true nodes:
+						RemovePin(TruePin);
+						RemovePin(FalsePin);
+					}
+				}
 			}
 		}
 	}
diff --git a/Engine/Source/Editor/UnrealEd/Private/Kismet2/BlueprintEditorUtils.cpp b/Engine/Source/Editor/UnrealEd/Private/Kismet2/BlueprintEditorUtils.cpp
--- a/Engine/Source/Editor/UnrealEd/Private/Kismet2/BlueprintEditorUtils.cpp
+++ b/Engine/Source/Editor/UnrealEd/Private/Kismet2/BlueprintEditorUtils.cpp
@@ -4285,7 +4285,8 @@ void FBlueprintEditorUtils::GetHiddenPinsForFunction(UEdGraph const* Graph, UFun
 			{
 				HiddenPins.Add(*It.Value());
 			}
-			else if (Key == FBlueprintMetadata::MD_ExpandEnumAsExecs)
+			else if (Key == FBlueprintMetadata::MD_ExpandEnumAsExecs ||
+					Key == FBlueprintMetadata::MD_ExpandBoolAsExecs)
 			{
 				TArray<FName> EnumPinNames;
 				UK2Node_CallFunction::GetExpandEnumPinNames(Function, EnumPinNames);
diff --git a/Engine/Source/Programs/UnrealHeaderTool/Private/HeaderParser.cpp b/Engine/Source/Programs/UnrealHeaderTool/Private/HeaderParser.cpp
--- a/Engine/Source/Programs/UnrealHeaderTool/Private/HeaderParser.cpp
+++ b/Engine/Source/Programs/UnrealHeaderTool/Private/HeaderParser.cpp
@@ -1064,6 +1064,7 @@ namespace
 			}
 			break;
 
+			case ECheckedMetadataSpecifier::ExpandBoolAsExecs:
 			case ECheckedMetadataSpecifier::ExpandEnumAsExecs:
 			{
 				if (UFunction* Function = Field.Get<UFunction>())
diff --git a/Engine/Source/Programs/UnrealHeaderTool/Private/Specifiers/CheckedMetadataSpecifiers.def b/Engine/Source/Programs/UnrealHeaderTool/Private/Specifiers/CheckedMetadataSpecifiers.def
--- a/Engine/Source/Programs/UnrealHeaderTool/Private/Specifiers/CheckedMetadataSpecifiers.def
+++ b/Engine/Source/Programs/UnrealHeaderTool/Private/Specifiers/CheckedMetadataSpecifiers.def
@@ -10,6 +10,7 @@ CHECKED_METADATA_SPECIFIER(ClampMin)
 CHECKED_METADATA_SPECIFIER(CommutativeAssociativeBinaryOperator)
 CHECKED_METADATA_SPECIFIER(DevelopmentStatus)
 CHECKED_METADATA_SPECIFIER(DocumentationPolicy)
+CHECKED_METADATA_SPECIFIER(ExpandBoolAsExecs)
 CHECKED_METADATA_SPECIFIER(ExpandEnumAsExecs)
 CHECKED_METADATA_SPECIFIER(UIMax)
 CHECKED_METADATA_SPECIFIER(UIMin)
diff --git a/Engine/Source/Runtime/CoreUObject/Public/UObject/ObjectMacros.h b/Engine/Source/Runtime/CoreUObject/Public/UObject/ObjectMacros.h
--- a/Engine/Source/Runtime/CoreUObject/Public/UObject/ObjectMacros.h
+++ b/Engine/Source/Runtime/CoreUObject/Public/UObject/ObjectMacros.h
@@ -1347,8 +1347,12 @@ namespace UM
 		// DeprecationMessage, (Commented out so as to avoid duplicate name with version in the Class section, but still show in the function section)
 
 		/// [FunctionMetadata] For BlueprintCallable functions indicates that an input/output (determined by whether it is an input/output enum) exec pin should be created for each entry in the enum specified.
+		/// Use ReturnValue to refer to the return value of the function. Also works for bools.
 		ExpandEnumAsExecs,
 
+		// Synonym for ExpandEnumAsExecs
+		ExpandBoolAsExecs,
+
 		/// [ClassMetadata] [PropertyMetadata] [FunctionMetadata] The name to display for this class, property, or function instead of auto-generating it from the name.
 		// DisplayName, (Commented out so as to avoid duplicate name with version in the Class section, but still show in the function section)
 
