Don't set autogenerated nodes to transactional unless the owning graph is also marked as transactional

diff --git a/Engine/Source/Editor/BlueprintGraph/Private/EdGraphSchema_K2_Actions.cpp b/Engine/Source/Editor/BlueprintGraph/Private/EdGraphSchema_K2_Actions.cpp
--- a/Engine/Source/Editor/BlueprintGraph/Private/EdGraphSchema_K2_Actions.cpp
+++ b/Engine/Source/Editor/BlueprintGraph/Private/EdGraphSchema_K2_Actions.cpp
@@ -254,7 +254,10 @@ UEdGraphNode* FEdGraphSchemaAction_K2NewNode::CreateNode(
 	// Duplicate template node to create new node
 	UK2Node* ResultNode = ConstructionFn(ParentGraph);
 	InitializerFn(ResultNode);
-	ResultNode->SetFlags(RF_Transactional);
+	if (ParentGraph->HasAnyFlags(RF_Transactional))
+	{
+		ResultNode->SetFlags(RF_Transactional);
+	}
 
 	ParentGraph->AddNode(ResultNode, true, bSelectNewNode);
 
