Blueprints: Merged UK2Node_LatentOnlineCall into UK2Node_AsyncAction, and moved it from the Kismet module into BlueprintGraph

diff --git a/Engine/Config/BaseEngine.ini b/Engine/Config/BaseEngine.ini
--- a/Engine/Config/BaseEngine.ini
+++ b/Engine/Config/BaseEngine.ini
@@ -759,6 +759,8 @@ bUseStreamingPause=false
 
 ; 4.25
 +EnumRedirects=(OldName="ECurveBlendOption",ValueChanges=(("ECurveBlendOption::MaxWeight", "ECurveBlendOption::Override"))
++ClassRedirects=(OldName="/Script/OnlineBlueprintSupport.K2Node_LatentOnlineCall", NewName="/Script/BlueprintGraph.K2Node_AsyncAction")
++ClassRedirects=(OldName="/Script/Kismet.K2Node_AsyncAction", NewName="/Script/BlueprintGraph.K2Node_AsyncAction")
 
 [CoreUObject.Metadata]
 ; Note: UnrealHeaderTool should be rerun after making changes to MetadataRedirects to catch any keys specified in class headers
diff --git a/Engine/Plugins/Online/OnlineSubsystemUtils/Source/OnlineBlueprintSupport/Classes/K2Node_LatentOnlineCall.h b/Engine/Plugins/Online/OnlineSubsystemUtils/Source/OnlineBlueprintSupport/Classes/K2Node_LatentOnlineCall.h
deleted file mode 100644
diff --git a/Engine/Plugins/Online/OnlineSubsystemUtils/Source/OnlineBlueprintSupport/Private/K2Node_LatentOnlineCall.cpp b/Engine/Plugins/Online/OnlineSubsystemUtils/Source/OnlineBlueprintSupport/Private/K2Node_LatentOnlineCall.cpp
deleted file mode 100644
diff --git a/Engine/Source/Editor/BlueprintGraph/Classes/K2Node_AsyncAction.h b/Engine/Source/Editor/BlueprintGraph/Classes/K2Node_AsyncAction.h
new file mode 100644
--- /dev/null
+++ b/Engine/Source/Editor/BlueprintGraph/Classes/K2Node_AsyncAction.h
@@ -0,0 +1,22 @@
+// Copyright Epic Games, Inc. All Rights Reserved.
+
+#pragma once
+
+#include "CoreMinimal.h"
+#include "UObject/ObjectMacros.h"
+#include "K2Node_BaseAsyncTask.h"
+#include "K2Node_AsyncAction.generated.h"
+
+class FBlueprintActionDatabaseRegistrar;
+
+/** !!! The proxy object should have RF_StrongRefOnFrame flag. !!! */
+
+UCLASS()
+class BLUEPRINTGRAPH_API UK2Node_AsyncAction : public UK2Node_BaseAsyncTask
+{
+	GENERATED_UCLASS_BODY()
+	
+	// UK2Node interface
+	virtual void GetMenuActions(FBlueprintActionDatabaseRegistrar& ActionRegistrar) const override;
+	// End of UK2Node interface
+};
diff --git a/Engine/Source/Editor/BlueprintGraph/Private/K2Node_AsyncAction.cpp b/Engine/Source/Editor/BlueprintGraph/Private/K2Node_AsyncAction.cpp
new file mode 100644
--- /dev/null
+++ b/Engine/Source/Editor/BlueprintGraph/Private/K2Node_AsyncAction.cpp
@@ -0,0 +1,53 @@
+// Copyright Epic Games, Inc. All Rights Reserved.
+
+#include "K2Node_AsyncAction.h"
+#include "UObject/UnrealType.h"
+
+#include "Kismet/BlueprintAsyncActionBase.h"
+#include "BlueprintNodeSpawner.h"
+#include "BlueprintFunctionNodeSpawner.h"
+#include "BlueprintActionDatabaseRegistrar.h"
+
+
+#define LOCTEXT_NAMESPACE "K2Node"
+
+UK2Node_AsyncAction::UK2Node_AsyncAction(const FObjectInitializer& ObjectInitializer)
+	: Super(ObjectInitializer)
+{
+	ProxyActivateFunctionName = GET_FUNCTION_NAME_CHECKED(UBlueprintAsyncActionBase, Activate);
+}
+
+void UK2Node_AsyncAction::GetMenuActions(FBlueprintActionDatabaseRegistrar& ActionRegistrar) const
+{
+	struct GetMenuActions_Utils
+	{
+		static void SetNodeFunc(UEdGraphNode* NewNode, bool /*bIsTemplateNode*/, TWeakObjectPtr<UFunction> FunctionPtr)
+		{
+			UK2Node_AsyncAction* AsyncTaskNode = CastChecked<UK2Node_AsyncAction>(NewNode);
+			if (FunctionPtr.IsValid())
+			{
+				UFunction* Func = FunctionPtr.Get();
+				FObjectProperty* ReturnProp = CastFieldChecked<FObjectProperty>(Func->GetReturnProperty());
+						
+				AsyncTaskNode->ProxyFactoryFunctionName = Func->GetFName();
+				AsyncTaskNode->ProxyFactoryClass        = Func->GetOuterUClass();
+				AsyncTaskNode->ProxyClass               = ReturnProp->PropertyClass;
+			}
+		}
+	};
+
+	UClass* NodeClass = GetClass();
+	ActionRegistrar.RegisterClassFactoryActions<UBlueprintAsyncActionBase>(FBlueprintActionDatabaseRegistrar::FMakeFuncSpawnerDelegate::CreateLambda([NodeClass](const UFunction* FactoryFunc)->UBlueprintNodeSpawner*
+	{
+		UBlueprintNodeSpawner* NodeSpawner = UBlueprintFunctionNodeSpawner::Create(FactoryFunc);
+		check(NodeSpawner != nullptr);
+		NodeSpawner->NodeClass = NodeClass;
+
+		TWeakObjectPtr<UFunction> FunctionPtr = MakeWeakObjectPtr(const_cast<UFunction*>(FactoryFunc));
+		NodeSpawner->CustomizeNodeDelegate = UBlueprintNodeSpawner::FCustomizeNodeDelegate::CreateStatic(GetMenuActions_Utils::SetNodeFunc, FunctionPtr);
+
+		return NodeSpawner;
+	}) );
+}
+
+#undef LOCTEXT_NAMESPACE
diff --git a/Engine/Source/Editor/Kismet/Private/Nodes/K2Node_AsyncAction.cpp b/Engine/Source/Editor/Kismet/Private/Nodes/K2Node_AsyncAction.cpp
deleted file mode 100644
diff --git a/Engine/Source/Editor/Kismet/Public/Nodes/K2Node_AsyncAction.h b/Engine/Source/Editor/Kismet/Public/Nodes/K2Node_AsyncAction.h
deleted file mode 100644
diff --git a/Engine/Source/Runtime/Engine/Private/OnlineBlueprintCallProxyBase.cpp b/Engine/Source/Runtime/Engine/Private/OnlineBlueprintCallProxyBase.cpp
--- a/Engine/Source/Runtime/Engine/Private/OnlineBlueprintCallProxyBase.cpp
+++ b/Engine/Source/Runtime/Engine/Private/OnlineBlueprintCallProxyBase.cpp
@@ -10,7 +10,3 @@ UOnlineBlueprintCallProxyBase::UOnlineBlueprintCallProxyBase(const FObjectInitia
 {
 	SetFlags(RF_StrongRefOnFrame);
 }
-
-void UOnlineBlueprintCallProxyBase::Activate()
-{
-}
diff --git a/Engine/Source/Runtime/Engine/Public/Net/OnlineBlueprintCallProxyBase.h b/Engine/Source/Runtime/Engine/Public/Net/OnlineBlueprintCallProxyBase.h
--- a/Engine/Source/Runtime/Engine/Public/Net/OnlineBlueprintCallProxyBase.h
+++ b/Engine/Source/Runtime/Engine/Public/Net/OnlineBlueprintCallProxyBase.h
@@ -5,16 +5,14 @@
 #include "UObject/ObjectMacros.h"
 #include "UObject/Object.h"
 #include "UObject/ScriptMacros.h"
+#include "Kismet/BlueprintAsyncActionBase.h"
 #include "OnlineBlueprintCallProxyBase.generated.h"
 
 DECLARE_DYNAMIC_MULTICAST_DELEGATE(FEmptyOnlineDelegate);
 
 UCLASS(MinimalAPI)
-class UOnlineBlueprintCallProxyBase : public UObject
+class UOnlineBlueprintCallProxyBase : public UBlueprintAsyncActionBase
 {
 	GENERATED_UCLASS_BODY()
 
-	// Called to trigger the actual online action once the delegates have been bound
-	UFUNCTION(BlueprintCallable, meta=(BlueprintInternalUseOnly="true"), Category="Online")
-	virtual void Activate();
 };
