Blueprints: Prevented deprecated async action factory methods from showing up in the BP context menu (unless the show deprecated functions setting is enabled)

diff --git a/Engine/Source/Editor/BlueprintGraph/Private/BlueprintActionDatabaseRegistrar.cpp b/Engine/Source/Editor/BlueprintGraph/Private/BlueprintActionDatabaseRegistrar.cpp
--- a/Engine/Source/Editor/BlueprintGraph/Private/BlueprintActionDatabaseRegistrar.cpp
+++ b/Engine/Source/Editor/BlueprintGraph/Private/BlueprintActionDatabaseRegistrar.cpp
@@ -8,6 +8,7 @@
 #include "Engine/BlueprintGeneratedClass.h"
 #include "EdGraphSchema_K2.h"
 #include "BlueprintNodeSpawnerUtils.h"
+#include "BlueprintEditorSettings.h"
 
 /*******************************************************************************
  * BlueprintActionDatabaseRegistrarImpl
@@ -298,11 +299,19 @@ int32 FBlueprintActionDatabaseRegistrar::RegisterClassFactoryActions(const UClas
 				return false;
 			}
 
-			FObjectProperty* ReturnProperty = CastField<FObjectProperty>(Function->GetReturnProperty());
-			// see if the function is a static factory method
-			bool const bIsFactoryMethod = (ReturnProperty != nullptr) && ReturnProperty->PropertyClass->IsChildOf(InTargetType);
+			if (!Function->GetOwnerClass()->HasAnyClassFlags(CLASS_Deprecated | CLASS_NewerVersionExists) &&
+				(!Function->HasMetaData(FBlueprintMetadata::MD_DeprecatedFunction) || GetDefault<UBlueprintEditorSettings>()->bExposeDeprecatedFunctions))
+			{
+				FObjectProperty* ReturnProperty = CastField<FObjectProperty>(Function->GetReturnProperty());
+				// see if the function is a static factory method
+				bool const bIsFactoryMethod = (ReturnProperty != nullptr) && ReturnProperty->PropertyClass->IsChildOf(InTargetType);
 
-			return bIsFactoryMethod;
+				return bIsFactoryMethod;
+			}
+			else
+			{
+				return false;
+			}
 		}
 	};
 
