Blueprints: Toggling the 'Expose Deprecated Functions' setting now takes effect immediately rather than requiring an editor restart

diff --git a/Engine/Source/Editor/BlueprintGraph/Private/BlueprintEditorSettings.cpp b/Engine/Source/Editor/BlueprintGraph/Private/BlueprintEditorSettings.cpp
--- a/Engine/Source/Editor/BlueprintGraph/Private/BlueprintEditorSettings.cpp
+++ b/Engine/Source/Editor/BlueprintGraph/Private/BlueprintEditorSettings.cpp
@@ -4,7 +4,7 @@
 #include "Misc/ConfigCacheIni.h"
 #include "Editor/EditorPerProjectUserSettings.h"
 #include "Settings/EditorExperimentalSettings.h"
-
+#include "BlueprintActionDatabase.h"
 #include "FindInBlueprintManager.h"
 #include "Subsystems/AssetEditorSubsystem.h"
 #include "Editor.h"
@@ -85,5 +85,17 @@ void UBlueprintEditorSettings::PostEditChangeProperty(FPropertyChangedEvent& Pro
 		FFindInBlueprintSearchManager::Get().EnableGlobalFindResults(bHostFindInBlueprintsInGlobalTab);
 	}
 
+	bool bShouldRebuildRegistry = false;
+	
+	if (PropertyName == GET_MEMBER_NAME_CHECKED(UBlueprintEditorSettings, bExposeDeprecatedFunctions))
+	{
+		bShouldRebuildRegistry = true;
+	}
+
+	if (bShouldRebuildRegistry)
+	{
+		FBlueprintActionDatabase::Get().RefreshAll();
+	}
+
 	Super::PostEditChangeProperty(PropertyChangedEvent);
 }
\ No newline at end of file
