Validate pin connections on compilation for tunnel nodes ensuring errors when the user changes to an incompatible type

diff --git a/Engine/Source/Editor/BlueprintGraph/Classes/K2Node_Tunnel.h b/Engine/Source/Editor/BlueprintGraph/Classes/K2Node_Tunnel.h
--- a/Engine/Source/Editor/BlueprintGraph/Classes/K2Node_Tunnel.h
+++ b/Engine/Source/Editor/BlueprintGraph/Classes/K2Node_Tunnel.h
@@ -55,6 +55,7 @@ class BLUEPRINTGRAPH_API UK2Node_Tunnel : public UK2Node_EditablePinBase
 	virtual bool DrawNodeAsExit() const override;
 	virtual bool NodeCausesStructuralBlueprintChange() const override { return true; }
 	virtual void ClearCachedBlueprintData(UBlueprint* Blueprint) override;
+	virtual void ValidateNodeDuringCompilation(FCompilerResultsLog& MessageLog) const override;
 	//~ End UK2Node Interface
 
 	//~ Begin UK2Node_EditablePinBase Interface.
diff --git a/Engine/Source/Editor/BlueprintGraph/Private/K2Node_Tunnel.cpp b/Engine/Source/Editor/BlueprintGraph/Private/K2Node_Tunnel.cpp
--- a/Engine/Source/Editor/BlueprintGraph/Private/K2Node_Tunnel.cpp
+++ b/Engine/Source/Editor/BlueprintGraph/Private/K2Node_Tunnel.cpp
@@ -5,7 +5,7 @@
 #include "K2Node_Composite.h"
 #include "K2Node_MacroInstance.h"
 #include "Kismet2/BlueprintEditorUtils.h"
-
+#include "Kismet2/CompilerResultsLog.h"
 
 #define LOCTEXT_NAMESPACE "K2Node"
 
@@ -17,14 +17,14 @@ UK2Node_Tunnel::UK2Node_Tunnel(const FObjectInitializer& ObjectInitializer)
 
 void UK2Node_Tunnel::DestroyNode()
 {
-	if (InputSinkNode != NULL)
+	if (InputSinkNode != nullptr)
 	{
-		InputSinkNode->OutputSourceNode = NULL;
+		InputSinkNode->OutputSourceNode = nullptr;
 	}
 
-	if (OutputSourceNode != NULL)
+	if (OutputSourceNode != nullptr)
 	{
-		OutputSourceNode->InputSinkNode = NULL;
+		OutputSourceNode->InputSinkNode = nullptr;
 	}
 
 	//@TODO: Should we remove the pins provided by this node from the twinned node(s)?
@@ -183,6 +183,13 @@ void UK2Node_Tunnel::ClearCachedBlueprintData(UBlueprint* Blueprint)
 	MetaData.HasLatentFunctions = INDEX_NONE;
 }
 
+void UK2Node_Tunnel::ValidateNodeDuringCompilation(FCompilerResultsLog& MessageLog) const
+{
+	Super::ValidateNodeDuringCompilation(MessageLog);
+	
+	FBlueprintEditorUtils::ValidatePinConnections(this, MessageLog);
+}
+
 UEdGraphPin* UK2Node_Tunnel::CreatePinFromUserDefinition(const TSharedPtr<FUserPinInfo> NewPinInfo)
 {
 	// Create the new pin
@@ -245,7 +252,7 @@ ERenamePinResult UK2Node_Tunnel::RenameUserDefinedPinImpl(const FName OldName, c
 
 	// And do the same on the twinned pin
 	ERenamePinResult TargetNodeResult = ERenamePinResult::ERenamePinResult_Success;
-	UEdGraphNode* TargetNode = ((InputSinkNode != NULL) ? InputSinkNode : OutputSourceNode);
+	UEdGraphNode* TargetNode = ((InputSinkNode != nullptr) ? InputSinkNode : OutputSourceNode);
 	if (UK2Node_Composite* CompositeNode = Cast<UK2Node_Composite>(TargetNode))
 	{
 		TargetNodeResult = CompositeNode->RenameUserDefinedPin(OldName, NewName, bTest);
@@ -262,7 +269,7 @@ UObject* UK2Node_Tunnel::GetJumpTargetForDoubleClick() const
 {
 	// Try to select the other side of a tunnel node
 	UEdGraphNode* TargetNode = GetOutputSource();
-	if (TargetNode == NULL)
+	if (TargetNode == nullptr)
 	{
 		TargetNode = GetInputSink();
 	}
diff --git a/Engine/Source/Editor/UnrealEd/Private/Kismet2/BlueprintEditorUtils.cpp b/Engine/Source/Editor/UnrealEd/Private/Kismet2/BlueprintEditorUtils.cpp
--- a/Engine/Source/Editor/UnrealEd/Private/Kismet2/BlueprintEditorUtils.cpp
+++ b/Engine/Source/Editor/UnrealEd/Private/Kismet2/BlueprintEditorUtils.cpp
@@ -4352,6 +4352,32 @@ bool FBlueprintEditorUtils::IsPinTypeValid(const FEdGraphPinType& Type)
 	return true;
 }
 
+void FBlueprintEditorUtils::ValidatePinConnections(const UEdGraphNode* Node, FCompilerResultsLog& MessageLog)
+{
+	if (Node)
+	{
+		const UEdGraphSchema* Schema = Node->GetSchema();
+		check(Schema);
+
+		// Validate that all pins with links are actually set to valid connections.
+		// This is necessary because the user could change the type of the pin 
+		for (UEdGraphPin* Pin : Node->Pins)
+		{
+			if (Pin)
+			{
+				for (UEdGraphPin* Link : Pin->LinkedTo)
+				{
+					if (Link && Link != Pin && Schema->CanCreateConnection(Pin, Link).Response == CONNECT_RESPONSE_DISALLOW)
+					{
+						FText const ErrorFormat = LOCTEXT("BadConnection", "Invalid pin connection from '@@' to '@@'. You may have changed the type after the connections were made.");
+						MessageLog.Error(*ErrorFormat.ToString(), Pin, Link);
+					}
+				}
+			}
+		}
+	}
+}
+
 void FBlueprintEditorUtils::GetClassVariableList(const UBlueprint* Blueprint, TSet<FName>& VisibleVariables, bool bIncludePrivateVars) 
 {
 	// Existing variables in the parent class and above, when using the compilation manager the previous SkeletonGeneratedClass will have been cleared when
diff --git a/Engine/Source/Editor/UnrealEd/Public/Kismet2/BlueprintEditorUtils.h b/Engine/Source/Editor/UnrealEd/Public/Kismet2/BlueprintEditorUtils.h
--- a/Engine/Source/Editor/UnrealEd/Public/Kismet2/BlueprintEditorUtils.h
+++ b/Engine/Source/Editor/UnrealEd/Public/Kismet2/BlueprintEditorUtils.h
@@ -761,6 +761,14 @@ public:
 	 */
 	static bool IsPinTypeValid(const FEdGraphPinType& Type);
 
+	/**
+	* Ensures the validity of each pin connection on the given node. Outputs compiler error if invalid
+	* 
+	* @param Node			The node to check all linked pins on
+	* @param MessageLog		BP compiler results log to output any error messages to
+	*/
+	static void ValidatePinConnections(const UEdGraphNode* Node, FCompilerResultsLog& MessageLog);
+
 	/**
 	 * Gets the visible class variable list.  This includes both variables introduced here and in all superclasses.
 	 *
