Add a cvar for disabling the pin validity check by default to help manage warnings

#rnx

diff --git a/Engine/Source/Editor/BlueprintGraph/Private/K2Node_Tunnel.cpp b/Engine/Source/Editor/BlueprintGraph/Private/K2Node_Tunnel.cpp
--- a/Engine/Source/Editor/BlueprintGraph/Private/K2Node_Tunnel.cpp
+++ b/Engine/Source/Editor/BlueprintGraph/Private/K2Node_Tunnel.cpp
@@ -9,6 +9,21 @@
 
 #define LOCTEXT_NAMESPACE "K2Node"
 
+// @TODO_BH: Remove the CVar for validity checking when we can get all the errors sorted out in a pre-flight
+// When we remove this then make sure we have a valid cook happening
+namespace PinValidityCheck
+{
+	/** 
+	* CVar controls pin validity warning which will throw when a macro graph is silently failing
+	* @see UE-90009
+	*/
+	static bool bDisplayInvalidPinWarning = true;
+	static FAutoConsoleVariableRef CVarDisplayInvalidPinWarning(
+		TEXT("bp.PinValidityCheck.bDisplayInvalidPinWarning"), bDisplayInvalidPinWarning,
+		TEXT("CVar controls pin validity warning which will throw when a macro graph is silently failing"),
+		ECVF_Default);
+}
+
 UK2Node_Tunnel::UK2Node_Tunnel(const FObjectInitializer& ObjectInitializer)
 	: Super(ObjectInitializer)
 {
@@ -187,7 +202,10 @@ void UK2Node_Tunnel::ValidateNodeDuringCompilation(FCompilerResultsLog& MessageL
 {
 	Super::ValidateNodeDuringCompilation(MessageLog);
 	
-	FBlueprintEditorUtils::ValidatePinConnections(this, MessageLog);
+	if (PinValidityCheck::bDisplayInvalidPinWarning)
+	{
+		FBlueprintEditorUtils::ValidatePinConnections(this, MessageLog);
+	}
 }
 
 UEdGraphPin* UK2Node_Tunnel::CreatePinFromUserDefinition(const TSharedPtr<FUserPinInfo> NewPinInfo)
