Correctly fix up user defined pins for function and macro result nodes, function entry  nodes already handled this

diff --git a/Engine/Source/Editor/BlueprintGraph/Classes/K2Node_EditablePinBase.h b/Engine/Source/Editor/BlueprintGraph/Classes/K2Node_EditablePinBase.h
--- a/Engine/Source/Editor/BlueprintGraph/Classes/K2Node_EditablePinBase.h
+++ b/Engine/Source/Editor/BlueprintGraph/Classes/K2Node_EditablePinBase.h
@@ -163,6 +163,9 @@ class UK2Node_EditablePinBase : public UK2Node
 	/** Modifies the default value of an existing pin on the node, this will update both the UserPinInfo and the linked editor pin */
 	BLUEPRINTGRAPH_API virtual bool ModifyUserDefinedPinDefaultValue(TSharedPtr<FUserPinInfo> PinInfo, const FString& NewDefaultValue);
 
+	/** Copies default value data from the graph pins to the user pins, returns true if any were modified */
+	BLUEPRINTGRAPH_API virtual bool UpdateUserDefinedPinDefaultValues();
+
 	/**
 	 * Creates function pins that are user defined based on a function signature.
 	 */
diff --git a/Engine/Source/Editor/BlueprintGraph/Classes/K2Node_FunctionResult.h b/Engine/Source/Editor/BlueprintGraph/Classes/K2Node_FunctionResult.h
--- a/Engine/Source/Editor/BlueprintGraph/Classes/K2Node_FunctionResult.h
+++ b/Engine/Source/Editor/BlueprintGraph/Classes/K2Node_FunctionResult.h
@@ -34,6 +34,7 @@ class UK2Node_FunctionResult : public UK2Node_FunctionTerminator
 	virtual bool CanDuplicateNode() const { return true; }
 	virtual void PostPasteNode() override;
 	virtual bool CanUserDeleteNode() const override;
+	virtual void FixupPinStringDataReferences(FArchive* SavingArchive) override;
 	//~ End UK2Node Interface
 
 	//~ Begin UK2Node_EditablePinBase Interface
diff --git a/Engine/Source/Editor/BlueprintGraph/Classes/K2Node_Tunnel.h b/Engine/Source/Editor/BlueprintGraph/Classes/K2Node_Tunnel.h
--- a/Engine/Source/Editor/BlueprintGraph/Classes/K2Node_Tunnel.h
+++ b/Engine/Source/Editor/BlueprintGraph/Classes/K2Node_Tunnel.h
@@ -56,6 +56,7 @@ class BLUEPRINTGRAPH_API UK2Node_Tunnel : public UK2Node_EditablePinBase
 	virtual bool NodeCausesStructuralBlueprintChange() const override { return true; }
 	virtual void ClearCachedBlueprintData(UBlueprint* Blueprint) override;
 	virtual void ValidateNodeDuringCompilation(FCompilerResultsLog& MessageLog) const override;
+	virtual void FixupPinStringDataReferences(FArchive* SavingArchive) override;
 	//~ End UK2Node Interface
 
 	//~ Begin UK2Node_EditablePinBase Interface.
diff --git a/Engine/Source/Editor/BlueprintGraph/Private/K2Node_CustomEvent.cpp b/Engine/Source/Editor/BlueprintGraph/Private/K2Node_CustomEvent.cpp
--- a/Engine/Source/Editor/BlueprintGraph/Private/K2Node_CustomEvent.cpp
+++ b/Engine/Source/Editor/BlueprintGraph/Private/K2Node_CustomEvent.cpp
@@ -348,23 +348,7 @@ void UK2Node_CustomEvent::FixupPinStringDataReferences(FArchive* SavingArchive)
 	Super::FixupPinStringDataReferences(SavingArchive);
 	if (SavingArchive)
 	{ 
-		// If any of our pins got fixed up, we need to refresh our user pin default values
-		// For custom events, the Pin default values are authoritative
-		for (TSharedPtr<FUserPinInfo> PinInfo : UserDefinedPins)
-		{
-			if (UEdGraphPin* Pin = FindPin(PinInfo->PinName))
-			{
-				if (Pin->Direction == PinInfo->DesiredPinDirection)
-				{
-					FString DefaultsString = Pin->GetDefaultAsString();
-
-					if (DefaultsString != PinInfo->PinDefaultValue)
-					{
-						ModifyUserDefinedPinDefaultValue(PinInfo, DefaultsString);
-					}
-				}
-			}
-		}
+		UpdateUserDefinedPinDefaultValues();
 	}
 }
 
diff --git a/Engine/Source/Editor/BlueprintGraph/Private/K2Node_EditablePinBase.cpp b/Engine/Source/Editor/BlueprintGraph/Private/K2Node_EditablePinBase.cpp
--- a/Engine/Source/Editor/BlueprintGraph/Private/K2Node_EditablePinBase.cpp
+++ b/Engine/Source/Editor/BlueprintGraph/Private/K2Node_EditablePinBase.cpp
@@ -360,6 +360,28 @@ bool UK2Node_EditablePinBase::ModifyUserDefinedPinDefaultValue(TSharedPtr<FUserP
 	return true;
 }
 
+bool UK2Node_EditablePinBase::UpdateUserDefinedPinDefaultValues()
+{
+	bool bAnyChanged = false;
+	const UEdGraphSchema_K2* K2Schema = GetDefault<UEdGraphSchema_K2>();
+	
+	// If any of our pins got fixed up, we need to refresh our user pin default values
+	for (TSharedPtr<FUserPinInfo> PinInfo : UserDefinedPins)
+	{
+		if (UEdGraphPin* Pin = FindPin(PinInfo->PinName))
+		{
+			if (Pin->Direction == PinInfo->DesiredPinDirection)
+			{
+				if (!K2Schema->DoesDefaultValueMatch(*Pin, PinInfo->PinDefaultValue))
+				{
+					bAnyChanged |= ModifyUserDefinedPinDefaultValue(PinInfo, Pin->GetDefaultAsString());
+				}
+			}
+		}
+	}
+	return bAnyChanged;
+}
+
 bool UK2Node_EditablePinBase::UpdateEdGraphPinDefaultValue(TSharedPtr<FUserPinInfo> PinInfo, FString& NewDefaultValue)
 {
 	const UEdGraphSchema_K2* K2Schema = GetDefault<UEdGraphSchema_K2>();
diff --git a/Engine/Source/Editor/BlueprintGraph/Private/K2Node_FunctionEntry.cpp b/Engine/Source/Editor/BlueprintGraph/Private/K2Node_FunctionEntry.cpp
--- a/Engine/Source/Editor/BlueprintGraph/Private/K2Node_FunctionEntry.cpp
+++ b/Engine/Source/Editor/BlueprintGraph/Private/K2Node_FunctionEntry.cpp
@@ -868,22 +868,7 @@ void UK2Node_FunctionEntry::FixupPinStringDataReferences(FArchive* SavingArchive
 	Super::FixupPinStringDataReferences(SavingArchive);
 	if (SavingArchive)
 	{
-		const UEdGraphSchema_K2* K2Schema = GetDefault<UEdGraphSchema_K2>();
-		// If any of our pins got fixed up, we need to refresh our user pin default values
-		// For custom events, the Pin default values are authoritative
-		for (TSharedPtr<FUserPinInfo> PinInfo : UserDefinedPins)
-		{
-			if (UEdGraphPin* Pin = FindPin(PinInfo->PinName))
-			{
-				if (Pin->Direction == PinInfo->DesiredPinDirection)
-				{
-					if(!K2Schema->DoesDefaultValueMatch(*Pin, PinInfo->PinDefaultValue))
-					{
-						ModifyUserDefinedPinDefaultValue(PinInfo, Pin->GetDefaultAsString());
-					}
-				}
-			}
-		}
+		UpdateUserDefinedPinDefaultValues();
 	}
 }
 
diff --git a/Engine/Source/Editor/BlueprintGraph/Private/K2Node_FunctionResult.cpp b/Engine/Source/Editor/BlueprintGraph/Private/K2Node_FunctionResult.cpp
--- a/Engine/Source/Editor/BlueprintGraph/Private/K2Node_FunctionResult.cpp
+++ b/Engine/Source/Editor/BlueprintGraph/Private/K2Node_FunctionResult.cpp
@@ -194,6 +194,15 @@ UEdGraphPin* UK2Node_FunctionResult::CreatePinFromUserDefinition(const TSharedPt
 	return Pin;
 }
 
+void UK2Node_FunctionResult::FixupPinStringDataReferences(FArchive* SavingArchive)
+{
+	Super::FixupPinStringDataReferences(SavingArchive);
+	if (SavingArchive)
+	{
+		UpdateUserDefinedPinDefaultValues();
+	}
+}
+
 FNodeHandlingFunctor* UK2Node_FunctionResult::CreateNodeHandler(FKismetCompilerContext& CompilerContext) const
 {
 	return new FKCHandler_FunctionResult(CompilerContext);
diff --git a/Engine/Source/Editor/BlueprintGraph/Private/K2Node_Tunnel.cpp b/Engine/Source/Editor/BlueprintGraph/Private/K2Node_Tunnel.cpp
--- a/Engine/Source/Editor/BlueprintGraph/Private/K2Node_Tunnel.cpp
+++ b/Engine/Source/Editor/BlueprintGraph/Private/K2Node_Tunnel.cpp
@@ -208,6 +208,15 @@ void UK2Node_Tunnel::ValidateNodeDuringCompilation(FCompilerResultsLog& MessageL
 	}
 }
 
+void UK2Node_Tunnel::FixupPinStringDataReferences(FArchive* SavingArchive)
+{
+	Super::FixupPinStringDataReferences(SavingArchive);
+	if (SavingArchive)
+	{
+		UpdateUserDefinedPinDefaultValues();
+	}
+}
+
 UEdGraphPin* UK2Node_Tunnel::CreatePinFromUserDefinition(const TSharedPtr<FUserPinInfo> NewPinInfo)
 {
 	// Create the new pin
