Add an extra newline when viewing the debug values of a container in blueprints to make it more readable

#rnx

diff --git a/Engine/Source/Editor/BlueprintGraph/Private/K2Node.cpp b/Engine/Source/Editor/BlueprintGraph/Private/K2Node.cpp
--- a/Engine/Source/Editor/BlueprintGraph/Private/K2Node.cpp
+++ b/Engine/Source/Editor/BlueprintGraph/Private/K2Node.cpp
@@ -1860,7 +1860,7 @@ void UK2Node::GetPinHoverText(const UEdGraphPin& Pin, FString& HoverTextOut) con
 	switch (WatchStatus)
 	{
 	case FKismetDebugUtilities::EWTR_Valid:
-		HoverTextOut += FString::Printf(TEXT("\nCurrent value = %s"), *WatchText); //@TODO: Print out object being debugged name?
+		HoverTextOut += FString::Printf(TEXT("\nCurrent value = \n%s"), *WatchText); //@TODO: Print out object being debugged name?
 		break;
 	case FKismetDebugUtilities::EWTR_NotInScope:
 		HoverTextOut += TEXT("\n(Variable is not in scope)");
