Allow containers in "make" nodes and structs to be empty by default, which will auto generate the default value if none is provided. This makes the experience of using custom UStructs more fluid and prevents there being a node with a compiler error by default

diff --git a/Engine/Source/Editor/BlueprintGraph/Private/K2Node_MakeStruct.cpp b/Engine/Source/Editor/BlueprintGraph/Private/K2Node_MakeStruct.cpp
--- a/Engine/Source/Editor/BlueprintGraph/Private/K2Node_MakeStruct.cpp
+++ b/Engine/Source/Editor/BlueprintGraph/Private/K2Node_MakeStruct.cpp
@@ -133,6 +133,13 @@ void UK2Node_MakeStruct::AllocateDefaultPins()
 			bHasAdvancedPins = OptionalPinManager.HasAdvancedPins();
 		}
 
+		// Set container pin types to have their default values ignored, which will in turn
+		// enable auto generation for any that are not set by the user. 
+		for(UEdGraphPin* Pin : Pins)
+		{
+			Pin->bDefaultValueIsIgnored = Pin->bDefaultValueIsIgnored || Pin->PinType.IsContainer();
+		}
+
 		// When struct has a lot of fields, mark their pins as advanced
 		if(!bHasAdvancedPins && Pins.Num() > 5)
 		{
