Fix regression introduced in 4.24 where soft object and class references on blueprint pins were not being properly detected for cooking/reference viewer
This partially reverts CL 9406864, only the widget change in that CL was actually required to fix the original bug

diff --git a/Engine/Source/Editor/BlueprintGraph/Private/K2Node.cpp b/Engine/Source/Editor/BlueprintGraph/Private/K2Node.cpp
--- a/Engine/Source/Editor/BlueprintGraph/Private/K2Node.cpp
+++ b/Engine/Source/Editor/BlueprintGraph/Private/K2Node.cpp
@@ -91,7 +91,7 @@ void UK2Node::Serialize(FArchive& Ar)
 {
 	Ar.UsingCustomVersion(FFrameworkObjectVersion::GUID);
 
-	if (Ar.IsSaving() && !GIsSavingPackage)
+	if (Ar.IsSaving())
 	{
 		if (Ar.IsObjectReferenceCollector() || Ar.Tell() < 0)
 		{
