Rename bExposeDeprecatedFunctions to bExposeDeprecatedFunctionsAndClasses

diff --git a/Engine/Source/Editor/BlueprintGraph/Private/BlueprintActionDatabaseRegistrar.cpp b/Engine/Source/Editor/BlueprintGraph/Private/BlueprintActionDatabaseRegistrar.cpp
--- a/Engine/Source/Editor/BlueprintGraph/Private/BlueprintActionDatabaseRegistrar.cpp
+++ b/Engine/Source/Editor/BlueprintGraph/Private/BlueprintActionDatabaseRegistrar.cpp
@@ -300,7 +300,7 @@ int32 FBlueprintActionDatabaseRegistrar::RegisterClassFactoryActions(const UClas
 			}
 
 			if (!Function->GetOwnerClass()->HasAnyClassFlags(CLASS_Deprecated | CLASS_NewerVersionExists) &&
-				(!Function->HasMetaData(FBlueprintMetadata::MD_DeprecatedFunction) || GetDefault<UBlueprintEditorSettings>()->bExposeDeprecatedFunctions))
+				(!Function->HasMetaData(FBlueprintMetadata::MD_DeprecatedFunction) || GetDefault<UBlueprintEditorSettings>()->bExposeDeprecatedFunctionsAndClasses))
 			{
 				FObjectProperty* ReturnProperty = CastField<FObjectProperty>(Function->GetReturnProperty());
 				// see if the function is a static factory method
diff --git a/Engine/Source/Editor/BlueprintGraph/Private/BlueprintEditorSettings.cpp b/Engine/Source/Editor/BlueprintGraph/Private/BlueprintEditorSettings.cpp
--- a/Engine/Source/Editor/BlueprintGraph/Private/BlueprintEditorSettings.cpp
+++ b/Engine/Source/Editor/BlueprintGraph/Private/BlueprintEditorSettings.cpp
@@ -18,7 +18,7 @@ UBlueprintEditorSettings::UBlueprintEditorSettings(const FObjectInitializer& Obj
 	, bSplitContextTargetSettings(true)
 	, bExposeAllMemberComponentFunctions(true)
 	, bShowContextualFavorites(false)
-	, bExposeDeprecatedFunctions(false)
+	, bExposeDeprecatedFunctionsAndClasses(false)
 	, bCompactCallOnMemberNodes(false)
 	, bFlattenFavoritesMenus(true)
 	, bFavorPureCastNodes(false)
@@ -87,7 +87,7 @@ void UBlueprintEditorSettings::PostEditChangeProperty(FPropertyChangedEvent& Pro
 
 	bool bShouldRebuildRegistry = false;
 	
-	if (PropertyName == GET_MEMBER_NAME_CHECKED(UBlueprintEditorSettings, bExposeDeprecatedFunctions))
+	if (PropertyName == GET_MEMBER_NAME_CHECKED(UBlueprintEditorSettings, bExposeDeprecatedFunctionsAndClasses))
 	{
 		bShouldRebuildRegistry = true;
 	}
diff --git a/Engine/Source/Editor/BlueprintGraph/Private/EdGraphSchema_K2.cpp b/Engine/Source/Editor/BlueprintGraph/Private/EdGraphSchema_K2.cpp
--- a/Engine/Source/Editor/BlueprintGraph/Private/EdGraphSchema_K2.cpp
+++ b/Engine/Source/Editor/BlueprintGraph/Private/EdGraphSchema_K2.cpp
@@ -918,7 +918,7 @@ UFunction* UEdGraphSchema_K2::GetCallableParentFunction(UFunction* Function)
 bool UEdGraphSchema_K2::CanUserKismetCallFunction(const UFunction* Function)
 {
 	return Function && 
-		(Function->HasAllFunctionFlags(FUNC_BlueprintCallable) && !Function->HasAllFunctionFlags(FUNC_Delegate) && !Function->GetBoolMetaData(FBlueprintMetadata::MD_BlueprintInternalUseOnly) && (!Function->HasMetaData(FBlueprintMetadata::MD_DeprecatedFunction) || GetDefault<UBlueprintEditorSettings>()->bExposeDeprecatedFunctions));
+		(Function->HasAllFunctionFlags(FUNC_BlueprintCallable) && !Function->HasAllFunctionFlags(FUNC_Delegate) && !Function->GetBoolMetaData(FBlueprintMetadata::MD_BlueprintInternalUseOnly) && (!Function->HasMetaData(FBlueprintMetadata::MD_DeprecatedFunction) || GetDefault<UBlueprintEditorSettings>()->bExposeDeprecatedFunctionsAndClasses));
 }
 
 bool UEdGraphSchema_K2::CanKismetOverrideFunction(const UFunction* Function)
@@ -929,7 +929,7 @@ bool UEdGraphSchema_K2::CanKismetOverrideFunction(const UFunction* Function)
 			Function->HasAllFunctionFlags(FUNC_BlueprintEvent)
 			&& !Function->HasAllFunctionFlags(FUNC_Delegate) && 
 			!Function->GetBoolMetaData(FBlueprintMetadata::MD_BlueprintInternalUseOnly) && 
-			(!Function->HasMetaData(FBlueprintMetadata::MD_DeprecatedFunction) || GetDefault<UBlueprintEditorSettings>()->bExposeDeprecatedFunctions)
+			(!Function->HasMetaData(FBlueprintMetadata::MD_DeprecatedFunction) || GetDefault<UBlueprintEditorSettings>()->bExposeDeprecatedFunctionsAndClasses)
 		);
 }
 
diff --git a/Engine/Source/Editor/BlueprintGraph/Public/BlueprintEditorSettings.h b/Engine/Source/Editor/BlueprintGraph/Public/BlueprintEditorSettings.h
--- a/Engine/Source/Editor/BlueprintGraph/Public/BlueprintEditorSettings.h
+++ b/Engine/Source/Editor/BlueprintGraph/Public/BlueprintEditorSettings.h
@@ -46,9 +46,9 @@ public:
 	UPROPERTY(EditAnywhere, config, Category=Workflow, meta=(DisplayName="Context Menu: Show Favorites Section"))
 	bool bShowContextualFavorites;
 
-	/** If enabled, deprecated functions will be visible in the context menu and will be available for override implementation. By default, functions marked as deprecated are not exposed in either case. */
+	/** If enabled, deprecated functions will be visible in the context menu and will be available for override implementation. Deprecated classes will be available in the class dropdown menu on nodes with class pins. */
 	UPROPERTY(EditAnywhere, config, Category = Workflow)
-	bool bExposeDeprecatedFunctions;
+	bool bExposeDeprecatedFunctionsAndClasses;
 
 	/** If enabled, then call-on-member actions will be spawned as a single node (instead of a GetMember + FunctionCall node). */
 	UPROPERTY(EditAnywhere, config, Category=Workflow)
