Add function alias "FCeil" for "Ceil" on math expression nodes

diff --git a/Engine/Source/Editor/BlueprintGraph/Private/K2Node_MathExpression.cpp b/Engine/Source/Editor/BlueprintGraph/Private/K2Node_MathExpression.cpp
--- a/Engine/Source/Editor/BlueprintGraph/Private/K2Node_MathExpression.cpp
+++ b/Engine/Source/Editor/BlueprintGraph/Private/K2Node_MathExpression.cpp
@@ -948,7 +948,12 @@ private:
 			ADD_ALIAS("POWER")
 			ADD_ALIAS("POW")
 		FUNC_ALIASES_END
-				
+		
+		FUNC_ALIASES_BEGIN("FCEIL")
+			ADD_ALIAS("FCEIL")
+			ADD_ALIAS("CEIL")
+		FUNC_ALIASES_END
+
 		FUNC_ALIASES_BEGIN("ASin")
 			// have to add "ASin" back, because this overwrites the function's
 			// name and we still want it as a viable option
