Conform FText package namespace when reconstructing nodes

diff --git a/Engine/Source/Editor/BlueprintGraph/Private/EdGraphSchema_K2.cpp b/Engine/Source/Editor/BlueprintGraph/Private/EdGraphSchema_K2.cpp
--- a/Engine/Source/Editor/BlueprintGraph/Private/EdGraphSchema_K2.cpp
+++ b/Engine/Source/Editor/BlueprintGraph/Private/EdGraphSchema_K2.cpp
@@ -5093,6 +5093,50 @@ void UEdGraphSchema_K2::SetPinAutogeneratedDefaultValueBasedOnType(UEdGraphPin*
 	SetPinAutogeneratedDefaultValue(Pin, NewValue);
 }
 
+static void ConformAutogeneratedDefaultValuePackage(
+	const FEdGraphPinType& PinType, 
+	FString& AutogeneratedDefaultValue, 
+	const FText& DefaultTextValue
+)
+{
+	if (PinType.PinCategory == UEdGraphSchema_K2::PC_Text &&
+		!PinType.IsContainer() &&
+		!AutogeneratedDefaultValue.IsEmpty())
+	{
+		FText TextValue;
+		const TCHAR* Success = FTextStringHelper::ReadFromBuffer(*AutogeneratedDefaultValue, TextValue);
+		check(Success);
+
+		const TOptional<FString> RealPackageNamespace = FTextInspector::GetNamespace(DefaultTextValue);
+		const TOptional<FString> SavedPackageNamespace = FTextInspector::GetNamespace(TextValue);
+		const TOptional<FString> CurrentKey = FTextInspector::GetKey(TextValue);
+
+		if (RealPackageNamespace.IsSet() &&
+			CurrentKey.IsSet() &&
+			DefaultTextValue.ToString() == TextValue.ToString() &&
+			(	!SavedPackageNamespace.IsSet() ||
+				*RealPackageNamespace != *SavedPackageNamespace ) )
+		{
+			const TOptional<FString> RealKey = FTextInspector::GetNamespace(DefaultTextValue);
+
+			// make new text value with same key in current namespace:
+			AutogeneratedDefaultValue.Empty();
+			if (RealKey.IsSet())
+			{
+				// maintain the key on DefaultTextValue:
+				TextValue = FText::ChangeKey(*RealPackageNamespace, *RealKey, TextValue);
+			}
+			else
+			{
+				// reset back to the key in AutogeneratedDefaultValue - generating a new one
+				// will cause determinism issues which are extremely hard to track down:
+				TextValue = FText::ChangeKey(*RealPackageNamespace, *CurrentKey, TextValue);
+			}
+			FTextStringHelper::WriteToBuffer(AutogeneratedDefaultValue, TextValue);
+		}
+	}
+}
+
 void UEdGraphSchema_K2::ResetPinToAutogeneratedDefaultValue(UEdGraphPin* Pin, bool bCallModifyCallbacks) const
 {
 	if (Pin->bOrphanedPin)
@@ -5102,7 +5146,10 @@ void UEdGraphSchema_K2::ResetPinToAutogeneratedDefaultValue(UEdGraphPin* Pin, bo
 	}
 	else
 	{
-		GetPinDefaultValuesFromString(Pin->PinType, Pin->GetOwningNodeUnchecked(), Pin->AutogeneratedDefaultValue, Pin->DefaultValue, Pin->DefaultObject, Pin->DefaultTextValue);
+		// Autogenerated value has unreliable package namespace for text value, hack fix it up now:
+		ConformAutogeneratedDefaultValuePackage(Pin->PinType, Pin->AutogeneratedDefaultValue, Pin->DefaultTextValue);
+
+		GetPinDefaultValuesFromString(Pin->PinType, Pin->GetOwningNodeUnchecked(), Pin->AutogeneratedDefaultValue, Pin->DefaultValue, Pin->DefaultObject, Pin->DefaultTextValue, false);
 
 		if (bCallModifyCallbacks)
 		{
