Fix issue where a literal math expression could cause it to regenerate the user defined pin default value every time due to it skipping the pin change callback
Protect against possible future crashes from marking modified during second half of package saving
#rnx

diff --git a/Engine/Source/Editor/BlueprintGraph/Private/EdGraphSchema_K2.cpp b/Engine/Source/Editor/BlueprintGraph/Private/EdGraphSchema_K2.cpp
--- a/Engine/Source/Editor/BlueprintGraph/Private/EdGraphSchema_K2.cpp
+++ b/Engine/Source/Editor/BlueprintGraph/Private/EdGraphSchema_K2.cpp
@@ -7408,6 +7408,12 @@ void UEdGraphSchema_K2::HandleParameterDefaultValueChanged(UK2Node* InTargetNode
 {
 	if (UK2Node_EditablePinBase* EditablePinNode = Cast<UK2Node_EditablePinBase>(InTargetNode))
 	{
+		// If this is happening during a save, it's not safe to trigger a compilation
+		if (GIsSavingPackage)
+		{
+			return;
+		}
+
 		FParamsChangedHelper ParamsChangedHelper;
 		ParamsChangedHelper.ModifiedBlueprints.Add(FBlueprintEditorUtils::FindBlueprintForNode(InTargetNode));
 		FBlueprintEditorUtils::MarkBlueprintAsStructurallyModified(FBlueprintEditorUtils::FindBlueprintForNode(InTargetNode));
diff --git a/Engine/Source/Editor/BlueprintGraph/Private/K2Node_MathExpression.cpp b/Engine/Source/Editor/BlueprintGraph/Private/K2Node_MathExpression.cpp
--- a/Engine/Source/Editor/BlueprintGraph/Private/K2Node_MathExpression.cpp
+++ b/Engine/Source/Editor/BlueprintGraph/Private/K2Node_MathExpression.cpp
@@ -1219,7 +1219,7 @@ public:
 		bool bSuccess = true;//K2Schema->ArePinTypesCompatible(GetOutputType(), InputPin->PinType);
 		if (bSuccess)
 		{
-			InputPin->DefaultValue = DefaultValue;
+			K2Schema->TrySetDefaultValue(*InputPin, DefaultValue, false);
 		}
 		else 
 		{
