Blueprints - Fix softptrs literal pins not being serialized as imported into blueprint assets.  This should restore several cases of references not being shown in the asset registry - will need to resave though.

Fred.Kimberly, Marc.Audy, Dan.Oconner, Michael.Noland

diff --git a/Engine/Source/Editor/BlueprintGraph/Classes/K2Node.h b/Engine/Source/Editor/BlueprintGraph/Classes/K2Node.h
--- a/Engine/Source/Editor/BlueprintGraph/Classes/K2Node.h
+++ b/Engine/Source/Editor/BlueprintGraph/Classes/K2Node.h
@@ -478,7 +478,7 @@ protected:
 	BLUEPRINTGRAPH_API virtual void FixupPinDefaultValues();
 
 	/** Fixes up structure/soft object ref pins, on both save and load */
-	BLUEPRINTGRAPH_API virtual void FixupPinStringDataReferences(FArchive* SavingArchive);
+	BLUEPRINTGRAPH_API virtual void FixupPinStringDataReferences(FArchive& Ar);
 
 private:
 
diff --git a/Engine/Source/Editor/BlueprintGraph/Classes/K2Node_CustomEvent.h b/Engine/Source/Editor/BlueprintGraph/Classes/K2Node_CustomEvent.h
--- a/Engine/Source/Editor/BlueprintGraph/Classes/K2Node_CustomEvent.h
+++ b/Engine/Source/Editor/BlueprintGraph/Classes/K2Node_CustomEvent.h
@@ -57,7 +57,7 @@ class UK2Node_CustomEvent : public UK2Node_Event
 	//~ Begin UK2Node Interface
 	virtual void ValidateNodeDuringCompilation(class FCompilerResultsLog& MessageLog) const override;
 	virtual void GetMenuActions(FBlueprintActionDatabaseRegistrar& ActionRegistrar) const override;
-	virtual void FixupPinStringDataReferences(FArchive* SavingArchive) override;
+	virtual void FixupPinStringDataReferences(FArchive& Ar) override;
 	virtual bool NodeCausesStructuralBlueprintChange() const override { return true; }
 	//~ End UK2Node Interface
 
diff --git a/Engine/Source/Editor/BlueprintGraph/Classes/K2Node_FunctionEntry.h b/Engine/Source/Editor/BlueprintGraph/Classes/K2Node_FunctionEntry.h
--- a/Engine/Source/Editor/BlueprintGraph/Classes/K2Node_FunctionEntry.h
+++ b/Engine/Source/Editor/BlueprintGraph/Classes/K2Node_FunctionEntry.h
@@ -58,7 +58,7 @@ class UK2Node_FunctionEntry : public UK2Node_FunctionTerminator
 	virtual void ExpandNode(class FKismetCompilerContext& CompilerContext, UEdGraph* SourceGraph) override;
 	virtual void PostReconstructNode() override;
 	virtual void ClearCachedBlueprintData(UBlueprint* Blueprint) override;
-	virtual void FixupPinStringDataReferences(FArchive* SavingArchive) override;
+	virtual void FixupPinStringDataReferences(FArchive& Ar) override;
 	//~ End UK2Node Interface
 
 	//~ Begin UK2Node_EditablePinBase Interface
diff --git a/Engine/Source/Editor/BlueprintGraph/Private/K2Node.cpp b/Engine/Source/Editor/BlueprintGraph/Private/K2Node.cpp
--- a/Engine/Source/Editor/BlueprintGraph/Private/K2Node.cpp
+++ b/Engine/Source/Editor/BlueprintGraph/Private/K2Node.cpp
@@ -93,11 +93,8 @@ void UK2Node::Serialize(FArchive& Ar)
 
 	if (Ar.IsSaving())
 	{
-		if (Ar.IsObjectReferenceCollector() || Ar.Tell() < 0)
-		{
-			// When this is a reference collector/modifier, serialize some pins as structs
-			FixupPinStringDataReferences(&Ar);
-		}
+		// When this is a reference collector/modifier, serialize some pins as structs
+		FixupPinStringDataReferences(Ar);
 	}
 
 	Super::Serialize(Ar);
@@ -110,7 +107,7 @@ void UK2Node::Serialize(FArchive& Ar)
 		if (GIsEditor)
 		{
 			// We need to serialize string data references on load in editor builds so the cooker knows about them
-			FixupPinStringDataReferences(nullptr);
+			FixupPinStringDataReferences(Ar);
 		}
 	}
 }
@@ -173,17 +170,10 @@ void UK2Node::FixupPinDefaultValues()
 	}
 }
 
-void UK2Node::FixupPinStringDataReferences(FArchive* SavingArchive)
+void UK2Node::FixupPinStringDataReferences(FArchive& Ar)
 {
 	// This code expands some pin types into their real representation, optionally serializes them, and then writes them out as strings again
 	const UEdGraphSchema_K2* K2Schema = GetDefault<UEdGraphSchema_K2>();
-	FLinkerLoad* LinkerLoad = GetLinker();
-
-	// Can't do any fixups without an archive of some sort
-	if (!SavingArchive && !LinkerLoad)
-	{
-		return;
-	}
 
 	for (UEdGraphPin* Pin : Pins)
 	{
@@ -196,9 +186,9 @@ void UK2Node::FixupPinStringDataReferences(FArchive* SavingArchive)
 				if (Struct)
 				{
 					// While loading, our user struct may not have been linked yet
-					if (Struct->HasAnyFlags(RF_NeedLoad) && LinkerLoad)
+					if (Ar.IsLoading() && Struct->HasAnyFlags(RF_NeedLoad))
 					{
-						LinkerLoad->Preload(Struct);
+						Ar.Preload(Struct);
 					}
 					
 					FSoftObjectPathSerializationScope SetPackage(GetOutermost()->GetFName(), NAME_None, ESoftObjectPathCollectType::AlwaysCollect, ESoftObjectPathSerializeType::SkipSerializeIfArchiveHasSize);
@@ -210,10 +200,10 @@ void UK2Node::FixupPinStringDataReferences(FArchive* SavingArchive)
 					FOutputDeviceNull NullOutput;
 					Struct->ImportText(*Pin->DefaultValue, StructData->GetStructMemory(), this, PPF_SerializedAsImportText, &NullOutput, Pin->PinName.ToString());
 
-					if (SavingArchive)
+					if (Ar.IsSaving())
 					{
 						// If we're saving, use the archive to do any replacements
-						Struct->SerializeItem(*SavingArchive, StructData->GetStructMemory(), nullptr);
+						Struct->SerializeItem(Ar, StructData->GetStructMemory(), nullptr);
 					}
 					
 					// Convert back to the default value string as we might have changed
@@ -233,15 +223,15 @@ void UK2Node::FixupPinStringDataReferences(FArchive* SavingArchive)
 
 				FSoftObjectPath TempRef(Pin->DefaultValue);
 
-				if (SavingArchive)
+				if (Ar.IsSaving())
 				{
 					// Serialize it directly, this won't do anything if it's a real archive like LinkerSave
-					*SavingArchive << TempRef;
+					Ar << TempRef;
 				}
-				else
+				else if (Ar.IsLoading())
 				{
 					// Transform it as strings
-					TempRef.PostLoadPath(LinkerLoad);
+					TempRef.PostLoadPath(&Ar);
 					TempRef.PreSavePath();
 				}
 				
diff --git a/Engine/Source/Editor/BlueprintGraph/Private/K2Node_CustomEvent.cpp b/Engine/Source/Editor/BlueprintGraph/Private/K2Node_CustomEvent.cpp
--- a/Engine/Source/Editor/BlueprintGraph/Private/K2Node_CustomEvent.cpp
+++ b/Engine/Source/Editor/BlueprintGraph/Private/K2Node_CustomEvent.cpp
@@ -343,10 +343,11 @@ void UK2Node_CustomEvent::GetMenuActions(FBlueprintActionDatabaseRegistrar& Acti
 	}
 }
 
-void UK2Node_CustomEvent::FixupPinStringDataReferences(FArchive* SavingArchive)
+void UK2Node_CustomEvent::FixupPinStringDataReferences(FArchive& Ar)
 {
-	Super::FixupPinStringDataReferences(SavingArchive);
-	if (SavingArchive)
+	Super::FixupPinStringDataReferences(Ar);
+	
+	if (Ar.IsSaving())
 	{ 
 		UpdateUserDefinedPinDefaultValues();
 	}
diff --git a/Engine/Source/Editor/BlueprintGraph/Private/K2Node_FunctionEntry.cpp b/Engine/Source/Editor/BlueprintGraph/Private/K2Node_FunctionEntry.cpp
--- a/Engine/Source/Editor/BlueprintGraph/Private/K2Node_FunctionEntry.cpp
+++ b/Engine/Source/Editor/BlueprintGraph/Private/K2Node_FunctionEntry.cpp
@@ -863,10 +863,11 @@ void UK2Node_FunctionEntry::PostReconstructNode()
 	Super::PostReconstructNode();
 }
 
-void UK2Node_FunctionEntry::FixupPinStringDataReferences(FArchive* SavingArchive)
+void UK2Node_FunctionEntry::FixupPinStringDataReferences(FArchive& Ar)
 {
-	Super::FixupPinStringDataReferences(SavingArchive);
-	if (SavingArchive)
+	Super::FixupPinStringDataReferences(Ar);
+	
+	if (Ar.IsSaving())
 	{
 		UpdateUserDefinedPinDefaultValues();
 	}
