Blueprint - Fixing up some API usages that changed params.

diff --git a/Engine/Source/Editor/BlueprintGraph/Classes/K2Node_FunctionResult.h b/Engine/Source/Editor/BlueprintGraph/Classes/K2Node_FunctionResult.h
--- a/Engine/Source/Editor/BlueprintGraph/Classes/K2Node_FunctionResult.h
+++ b/Engine/Source/Editor/BlueprintGraph/Classes/K2Node_FunctionResult.h
@@ -34,7 +34,7 @@ class UK2Node_FunctionResult : public UK2Node_FunctionTerminator
 	virtual bool CanDuplicateNode() const { return true; }
 	virtual void PostPasteNode() override;
 	virtual bool CanUserDeleteNode() const override;
-	virtual void FixupPinStringDataReferences(FArchive* SavingArchive) override;
+	virtual void FixupPinStringDataReferences(FArchive& Ar) override;
 	//~ End UK2Node Interface
 
 	//~ Begin UK2Node_EditablePinBase Interface
diff --git a/Engine/Source/Editor/BlueprintGraph/Classes/K2Node_Tunnel.h b/Engine/Source/Editor/BlueprintGraph/Classes/K2Node_Tunnel.h
--- a/Engine/Source/Editor/BlueprintGraph/Classes/K2Node_Tunnel.h
+++ b/Engine/Source/Editor/BlueprintGraph/Classes/K2Node_Tunnel.h
@@ -56,7 +56,7 @@ class BLUEPRINTGRAPH_API UK2Node_Tunnel : public UK2Node_EditablePinBase
 	virtual bool NodeCausesStructuralBlueprintChange() const override { return true; }
 	virtual void ClearCachedBlueprintData(UBlueprint* Blueprint) override;
 	virtual void ValidateNodeDuringCompilation(FCompilerResultsLog& MessageLog) const override;
-	virtual void FixupPinStringDataReferences(FArchive* SavingArchive) override;
+	virtual void FixupPinStringDataReferences(FArchive& Ar) override;
 	//~ End UK2Node Interface
 
 	//~ Begin UK2Node_EditablePinBase Interface.
diff --git a/Engine/Source/Editor/BlueprintGraph/Private/K2Node_FunctionResult.cpp b/Engine/Source/Editor/BlueprintGraph/Private/K2Node_FunctionResult.cpp
--- a/Engine/Source/Editor/BlueprintGraph/Private/K2Node_FunctionResult.cpp
+++ b/Engine/Source/Editor/BlueprintGraph/Private/K2Node_FunctionResult.cpp
@@ -194,10 +194,11 @@ UEdGraphPin* UK2Node_FunctionResult::CreatePinFromUserDefinition(const TSharedPt
 	return Pin;
 }
 
-void UK2Node_FunctionResult::FixupPinStringDataReferences(FArchive* SavingArchive)
+void UK2Node_FunctionResult::FixupPinStringDataReferences(FArchive& Ar)
 {
-	Super::FixupPinStringDataReferences(SavingArchive);
-	if (SavingArchive)
+	Super::FixupPinStringDataReferences(Ar);
+	
+	if (Ar.IsSaving())
 	{
 		UpdateUserDefinedPinDefaultValues();
 	}
diff --git a/Engine/Source/Editor/BlueprintGraph/Private/K2Node_Tunnel.cpp b/Engine/Source/Editor/BlueprintGraph/Private/K2Node_Tunnel.cpp
--- a/Engine/Source/Editor/BlueprintGraph/Private/K2Node_Tunnel.cpp
+++ b/Engine/Source/Editor/BlueprintGraph/Private/K2Node_Tunnel.cpp
@@ -208,10 +208,11 @@ void UK2Node_Tunnel::ValidateNodeDuringCompilation(FCompilerResultsLog& MessageL
 	}
 }
 
-void UK2Node_Tunnel::FixupPinStringDataReferences(FArchive* SavingArchive)
+void UK2Node_Tunnel::FixupPinStringDataReferences(FArchive& Ar)
 {
-	Super::FixupPinStringDataReferences(SavingArchive);
-	if (SavingArchive)
+	Super::FixupPinStringDataReferences(Ar);
+
+	if (Ar.IsSaving())
 	{
 		UpdateUserDefinedPinDefaultValues();
 	}
