Undo 12572258

diff --git a/Engine/Source/Editor/BlueprintGraph/Classes/K2Node.h b/Engine/Source/Editor/BlueprintGraph/Classes/K2Node.h
--- a/Engine/Source/Editor/BlueprintGraph/Classes/K2Node.h
+++ b/Engine/Source/Editor/BlueprintGraph/Classes/K2Node.h
@@ -478,7 +478,7 @@ protected:
 	BLUEPRINTGRAPH_API virtual void FixupPinDefaultValues();
 
 	/** Fixes up structure/soft object ref pins, on both save and load */
-	BLUEPRINTGRAPH_API virtual void FixupPinStringDataReferences(FArchive& Ar);
+	BLUEPRINTGRAPH_API virtual void FixupPinStringDataReferences(FArchive* SavingArchive);
 
 private:
 
diff --git a/Engine/Source/Editor/BlueprintGraph/Classes/K2Node_CustomEvent.h b/Engine/Source/Editor/BlueprintGraph/Classes/K2Node_CustomEvent.h
--- a/Engine/Source/Editor/BlueprintGraph/Classes/K2Node_CustomEvent.h
+++ b/Engine/Source/Editor/BlueprintGraph/Classes/K2Node_CustomEvent.h
@@ -57,7 +57,7 @@ class UK2Node_CustomEvent : public UK2Node_Event
 	//~ Begin UK2Node Interface
 	virtual void ValidateNodeDuringCompilation(class FCompilerResultsLog& MessageLog) const override;
 	virtual void GetMenuActions(FBlueprintActionDatabaseRegistrar& ActionRegistrar) const override;
-	virtual void FixupPinStringDataReferences(FArchive& Ar) override;
+	virtual void FixupPinStringDataReferences(FArchive* SavingArchive) override;
 	virtual bool NodeCausesStructuralBlueprintChange() const override { return true; }
 	//~ End UK2Node Interface
 
diff --git a/Engine/Source/Editor/BlueprintGraph/Classes/K2Node_FunctionEntry.h b/Engine/Source/Editor/BlueprintGraph/Classes/K2Node_FunctionEntry.h
--- a/Engine/Source/Editor/BlueprintGraph/Classes/K2Node_FunctionEntry.h
+++ b/Engine/Source/Editor/BlueprintGraph/Classes/K2Node_FunctionEntry.h
@@ -58,7 +58,7 @@ class UK2Node_FunctionEntry : public UK2Node_FunctionTerminator
 	virtual void ExpandNode(class FKismetCompilerContext& CompilerContext, UEdGraph* SourceGraph) override;
 	virtual void PostReconstructNode() override;
 	virtual void ClearCachedBlueprintData(UBlueprint* Blueprint) override;
-	virtual void FixupPinStringDataReferences(FArchive& Ar) override;
+	virtual void FixupPinStringDataReferences(FArchive* SavingArchive) override;
 	//~ End UK2Node Interface
 
 	//~ Begin UK2Node_EditablePinBase Interface
diff --git a/Engine/Source/Editor/BlueprintGraph/Private/K2Node.cpp b/Engine/Source/Editor/BlueprintGraph/Private/K2Node.cpp
--- a/Engine/Source/Editor/BlueprintGraph/Private/K2Node.cpp
+++ b/Engine/Source/Editor/BlueprintGraph/Private/K2Node.cpp
@@ -107,7 +107,7 @@ void UK2Node::Serialize(FArchive& Ar)
 		if (GIsEditor)
 		{
 			// We need to serialize string data references on load in editor builds so the cooker knows about them
-			FixupPinStringDataReferences(Ar);
+			FixupPinStringDataReferences(nullptr);
 		}
 	}
 }
@@ -170,10 +170,17 @@ void UK2Node::FixupPinDefaultValues()
 	}
 }
 
-void UK2Node::FixupPinStringDataReferences(FArchive& Ar)
+void UK2Node::FixupPinStringDataReferences(FArchive* SavingArchive)
 {
 	// This code expands some pin types into their real representation, optionally serializes them, and then writes them out as strings again
 	const UEdGraphSchema_K2* K2Schema = GetDefault<UEdGraphSchema_K2>();
+	FLinkerLoad* LinkerLoad = GetLinker();
+
+	// Can't do any fixups without an archive of some sort
+	if (!SavingArchive && !LinkerLoad)
+	{
+		return;
+	}
 
 	for (UEdGraphPin* Pin : Pins)
 	{
@@ -186,9 +193,9 @@ void UK2Node::FixupPinStringDataReferences(FArchive& Ar)
 				if (Struct)
 				{
 					// While loading, our user struct may not have been linked yet
-					if (Ar.IsLoading() && Struct->HasAnyFlags(RF_NeedLoad))
+					if (Struct->HasAnyFlags(RF_NeedLoad) && LinkerLoad)
 					{
-						Ar.Preload(Struct);
+						LinkerLoad->Preload(Struct);
 					}
 					
 					FSoftObjectPathSerializationScope SetPackage(GetOutermost()->GetFName(), NAME_None, ESoftObjectPathCollectType::AlwaysCollect, ESoftObjectPathSerializeType::SkipSerializeIfArchiveHasSize);
@@ -200,10 +207,10 @@ void UK2Node::FixupPinStringDataReferences(FArchive& Ar)
 					FOutputDeviceNull NullOutput;
 					Struct->ImportText(*Pin->DefaultValue, StructData->GetStructMemory(), this, PPF_SerializedAsImportText, &NullOutput, Pin->PinName.ToString());
 
-					if (Ar.IsSaving())
+					if (SavingArchive)
 					{
 						// If we're saving, use the archive to do any replacements
-						Struct->SerializeItem(Ar, StructData->GetStructMemory(), nullptr);
+						Struct->SerializeItem(*SavingArchive, StructData->GetStructMemory(), nullptr);
 					}
 					
 					// Convert back to the default value string as we might have changed
@@ -223,15 +230,15 @@ void UK2Node::FixupPinStringDataReferences(FArchive& Ar)
 
 				FSoftObjectPath TempRef(Pin->DefaultValue);
 
-				if (Ar.IsSaving())
+				if (SavingArchive)
 				{
 					// Serialize it directly, this won't do anything if it's a real archive like LinkerSave
-					Ar << TempRef;
+					*SavingArchive << TempRef;
 				}
-				else if (Ar.IsLoading())
+				else
 				{
 					// Transform it as strings
-					TempRef.PostLoadPath(&Ar);
+					TempRef.PostLoadPath(LinkerLoad);
 					TempRef.PreSavePath();
 				}
 				
diff --git a/Engine/Source/Editor/BlueprintGraph/Private/K2Node_CustomEvent.cpp b/Engine/Source/Editor/BlueprintGraph/Private/K2Node_CustomEvent.cpp
--- a/Engine/Source/Editor/BlueprintGraph/Private/K2Node_CustomEvent.cpp
+++ b/Engine/Source/Editor/BlueprintGraph/Private/K2Node_CustomEvent.cpp
@@ -343,11 +343,10 @@ void UK2Node_CustomEvent::GetMenuActions(FBlueprintActionDatabaseRegistrar& Acti
 	}
 }
 
-void UK2Node_CustomEvent::FixupPinStringDataReferences(FArchive& Ar)
+void UK2Node_CustomEvent::FixupPinStringDataReferences(FArchive* SavingArchive)
 {
-	Super::FixupPinStringDataReferences(Ar);
-	
-	if (Ar.IsSaving())
+	Super::FixupPinStringDataReferences(SavingArchive);
+	if (SavingArchive)
 	{ 
 		UpdateUserDefinedPinDefaultValues();
 	}
diff --git a/Engine/Source/Editor/BlueprintGraph/Private/K2Node_FunctionEntry.cpp b/Engine/Source/Editor/BlueprintGraph/Private/K2Node_FunctionEntry.cpp
--- a/Engine/Source/Editor/BlueprintGraph/Private/K2Node_FunctionEntry.cpp
+++ b/Engine/Source/Editor/BlueprintGraph/Private/K2Node_FunctionEntry.cpp
@@ -863,11 +863,10 @@ void UK2Node_FunctionEntry::PostReconstructNode()
 	Super::PostReconstructNode();
 }
 
-void UK2Node_FunctionEntry::FixupPinStringDataReferences(FArchive& Ar)
+void UK2Node_FunctionEntry::FixupPinStringDataReferences(FArchive* SavingArchive)
 {
-	Super::FixupPinStringDataReferences(Ar);
-	
-	if (Ar.IsSaving())
+	Super::FixupPinStringDataReferences(SavingArchive);
+	if (SavingArchive)
 	{
 		UpdateUserDefinedPinDefaultValues();
 	}
