#UE4 Finish undo of 12572258

diff --git a/Engine/Source/Editor/BlueprintGraph/Private/K2Node.cpp b/Engine/Source/Editor/BlueprintGraph/Private/K2Node.cpp
--- a/Engine/Source/Editor/BlueprintGraph/Private/K2Node.cpp
+++ b/Engine/Source/Editor/BlueprintGraph/Private/K2Node.cpp
@@ -91,10 +91,10 @@ void UK2Node::Serialize(FArchive& Ar)
 {
 	Ar.UsingCustomVersion(FFrameworkObjectVersion::GUID);
 
-	if (Ar.IsSaving())
+	if (Ar.IsObjectReferenceCollector() || Ar.Tell() < 0)
 	{
 		// When this is a reference collector/modifier, serialize some pins as structs
-		FixupPinStringDataReferences(Ar);
+		FixupPinStringDataReferences(&Ar);
 	}
 
 	Super::Serialize(Ar);
