Add null check to FindMatchingReferencedNetPropertyAndPin

#rnx

diff --git a/Engine/Source/Editor/BlueprintGraph/Private/CallFunctionHandler.cpp b/Engine/Source/Editor/BlueprintGraph/Private/CallFunctionHandler.cpp
--- a/Engine/Source/Editor/BlueprintGraph/Private/CallFunctionHandler.cpp
+++ b/Engine/Source/Editor/BlueprintGraph/Private/CallFunctionHandler.cpp
@@ -671,7 +671,7 @@ static bool FindMatchingReferencedNetPropertyAndPin(TArray<UEdGraphPin*>& Remain
 					if (UK2Node_VariableGet* GetPropertyNode = Cast<UK2Node_VariableGet>(PinToTry->GetOwningNode()))
 					{
 						FProperty* ToCheck = GetPropertyNode->GetPropertyForVariable();
-						if (UNLIKELY(ToCheck->HasAnyPropertyFlags(CPF_Net)))
+						if (UNLIKELY(ToCheck && ToCheck->HasAnyPropertyFlags(CPF_Net)))
 						{
 							NetProperty = ToCheck;
 							PropertyObjectPin = GetPropertyNode->FindPinChecked(UEdGraphSchema_K2::PN_Self);
