Set component expose on spawn properties befor OnConstruction and OnRegister lifecycle occurs

diff --git a/Engine/Source/Editor/BlueprintGraph/Classes/K2Node_AddComponent.h b/Engine/Source/Editor/BlueprintGraph/Classes/K2Node_AddComponent.h
--- a/Engine/Source/Editor/BlueprintGraph/Classes/K2Node_AddComponent.h
+++ b/Engine/Source/Editor/BlueprintGraph/Classes/K2Node_AddComponent.h
@@ -65,12 +65,12 @@ class UK2Node_AddComponent : public UK2Node_CallFunction
 
 	UEdGraphPin* GetRelativeTransformPin() const
 	{
-		return FindPinChecked(TEXT("RelativeTransform"));
+		return FindPinChecked(NAME_RelativeTransform);
 	}
 
 	UEdGraphPin* GetManualAttachmentPin() const
 	{
-		return FindPinChecked(TEXT("bManualAttachment"));
+		return FindPinChecked(NAME_ManualAttachment);
 	}
 
 	/** Tries to get a template object from this node. */
@@ -92,6 +92,9 @@ private:
 	}
 
 	UClass* GetSpawnedType() const;	
+
+	static const FName NAME_RelativeTransform;
+	static const FName NAME_ManualAttachment;
 };
 
 
diff --git a/Engine/Source/Editor/BlueprintGraph/Private/K2Node_AddComponent.cpp b/Engine/Source/Editor/BlueprintGraph/Private/K2Node_AddComponent.cpp
--- a/Engine/Source/Editor/BlueprintGraph/Private/K2Node_AddComponent.cpp
+++ b/Engine/Source/Editor/BlueprintGraph/Private/K2Node_AddComponent.cpp
@@ -20,6 +20,9 @@
 
 #define LOCTEXT_NAMESPACE "K2Node_AddComponent"
 
+const FName UK2Node_AddComponent::NAME_RelativeTransform(TEXT("RelativeTransform"));
+const FName UK2Node_AddComponent::NAME_ManualAttachment(TEXT("bManualAttachment"));
+
 //////////////////////////////////////////////////////////////////////////
 // FKCHandler_AddComponent
 
@@ -179,7 +182,6 @@ void UK2Node_AddComponent::ReallocatePinsDuringReconstruction(TArray<UEdGraphPin
 	Super::ReallocatePinsDuringReconstruction(OldPins);
 
 	UEdGraphPin* TemplateNamePin = GetTemplateNamePinChecked();
-	//UEdGraphPin* ReturnValuePin = GetReturnValuePin();
 	for(int32 OldPinIdx = 0; TemplateNamePin && OldPinIdx < OldPins.Num(); ++OldPinIdx)
 	{
 		if(TemplateNamePin && (TemplateNamePin->PinName == OldPins[OldPinIdx]->PinName))
@@ -435,7 +437,7 @@ void UK2Node_AddComponent::ExpandNode(class FKismetCompilerContext& CompilerCont
 {
 	Super::ExpandNode(CompilerContext, SourceGraph);
 
-	auto TransformPin = GetRelativeTransformPin();
+	UEdGraphPin* TransformPin = GetRelativeTransformPin();
 	if (TransformPin && !TransformPin->LinkedTo.Num())
 	{
 		FString DefaultValue;
@@ -451,7 +453,7 @@ void UK2Node_AddComponent::ExpandNode(class FKismetCompilerContext& CompilerCont
 			DefaultValue = TemplateTransform.ToString();
 		}
 
-		auto ValuePin = InnerHandleAutoCreateRef(this, TransformPin, CompilerContext, SourceGraph, !DefaultValue.IsEmpty());
+		UEdGraphPin* ValuePin = InnerHandleAutoCreateRef(this, TransformPin, CompilerContext, SourceGraph, !DefaultValue.IsEmpty());
 		if (ValuePin)
 		{
 			ValuePin->DefaultValue = DefaultValue;
@@ -460,10 +462,6 @@ void UK2Node_AddComponent::ExpandNode(class FKismetCompilerContext& CompilerCont
 
 	if (bHasExposedVariable)
 	{
-		static FString ObjectParamName = FString(TEXT("Object"));
-		static FString ValueParamName = FString(TEXT("Value"));
-		static FString PropertyNameParamName = FString(TEXT("PropertyName"));
-
 		UK2Node_AddComponent* NewNode = CompilerContext.SpawnIntermediateNode<UK2Node_AddComponent>(this, SourceGraph); 
 		NewNode->SetFromFunction(GetTargetFunction());
 		NewNode->AllocateDefaultPinsWithoutExposedVariables();
@@ -474,9 +472,14 @@ void UK2Node_AddComponent::ExpandNode(class FKismetCompilerContext& CompilerCont
 		CompilerContext.MovePinLinksToIntermediate(*GetRelativeTransformPin(), *NewNode->GetRelativeTransformPin());
 		CompilerContext.MovePinLinksToIntermediate(*GetManualAttachmentPin(), *NewNode->GetManualAttachmentPin());
 
+		static const FName DeferredFinishPinName = TEXT("bDeferredFinish");
+
+		UEdGraphPin* DeferredFinishPin = NewNode->FindPinChecked(DeferredFinishPinName);
+		DeferredFinishPin->DefaultValue = TEXT("true");
+
 		UEdGraphPin* ReturnPin = NewNode->GetReturnValuePin();
 		UEdGraphPin* OriginalReturnPin = GetReturnValuePin();
-		check((NULL != ReturnPin) && (NULL != OriginalReturnPin));
+		check(ReturnPin && OriginalReturnPin);
 		ReturnPin->PinType = OriginalReturnPin->PinType;
 		CompilerContext.MovePinLinksToIntermediate(*OriginalReturnPin, *ReturnPin);
 		// exec in
@@ -484,7 +487,23 @@ void UK2Node_AddComponent::ExpandNode(class FKismetCompilerContext& CompilerCont
 
 		UEdGraphPin* LastThen = FKismetCompilerUtilities::GenerateAssignmentNodes( CompilerContext, SourceGraph, NewNode, this, ReturnPin, GetSpawnedType() );
 
-		CompilerContext.MovePinLinksToIntermediate(*GetThenPin(), *LastThen);
+		static const FName FinishAddFuncName = GET_FUNCTION_NAME_CHECKED(AActor, FinishAddComponent);
+		static const FName ComponentParamName = TEXT("Component");
+
+		UK2Node_CallFunction* FinishAdd = CompilerContext.SpawnIntermediateNode<UK2Node_CallFunction>(this, SourceGraph);
+		FinishAdd->FunctionReference.SetExternalMember(FinishAddFuncName, AActor::StaticClass());
+		FinishAdd->AllocateDefaultPins();
+
+		UEdGraphPin* FinishAdd_ComponentPin = FinishAdd->FindPinChecked(ComponentParamName);
+		UEdGraphPin* FinishAdd_RelativeTransformPin = FinishAdd->FindPinChecked(NAME_RelativeTransform);
+		UEdGraphPin* FinishAdd_ManualAttachmentPin = FinishAdd->FindPinChecked(NAME_ManualAttachment);
+
+		FinishAdd_ComponentPin->MakeLinkTo(ReturnPin);
+		CompilerContext.CopyPinLinksToIntermediate(*NewNode->GetRelativeTransformPin(), *FinishAdd_RelativeTransformPin);
+		CompilerContext.CopyPinLinksToIntermediate(*NewNode->GetManualAttachmentPin(), *FinishAdd_ManualAttachmentPin);
+
+		LastThen->MakeLinkTo(FinishAdd->GetExecPin());
+		CompilerContext.MovePinLinksToIntermediate(*GetThenPin(), *FinishAdd->GetThenPin());
 		BreakAllNodeLinks();
 	}
 }
diff --git a/Engine/Source/Runtime/Engine/Classes/GameFramework/Actor.h b/Engine/Source/Runtime/Engine/Classes/GameFramework/Actor.h
--- a/Engine/Source/Runtime/Engine/Classes/GameFramework/Actor.h
+++ b/Engine/Source/Runtime/Engine/Classes/GameFramework/Actor.h
@@ -1260,9 +1260,23 @@ public:
 	 * @param bManualAttachment				Whether manual or automatic attachment is to be used
 	 * @param RelativeTransform				The relative transform between the new component and its attach parent (automatic only)
 	 * @param ComponentTemplateContext		Optional UBlueprintGeneratedClass reference to use to find the template in. If null (or not a BPGC), component is sought in this Actor's class
+	 * @param bDeferredFinish				Whether or not to immediately complete the creation and registration process for this component. Will be false if there are expose on spawn properties being set
 	 */
-	UFUNCTION(BlueprintCallable, meta=(ScriptNoExport, BlueprintInternalUseOnly = "true", DefaultToSelf="ComponentTemplateContext", InternalUseParam="ComponentTemplateContext"))
-	class UActorComponent* AddComponent(FName TemplateName, bool bManualAttachment, const FTransform& RelativeTransform, const UObject* ComponentTemplateContext);
+	UFUNCTION(BlueprintCallable, meta=(ScriptNoExport, BlueprintInternalUseOnly = "true", DefaultToSelf="ComponentTemplateContext", InternalUseParam="ComponentTemplateContext,bDeferredFinish"))
+	class UActorComponent* AddComponent(FName TemplateName, bool bManualAttachment, const FTransform& RelativeTransform, const UObject* ComponentTemplateContext, bool bDeferredFinish = false);
+
+	/** 
+	 * Completes the creation of a new actor component. Called either from blueprint after
+	 * expose on spawn properties are set, or directly from AddComponent
+	 *
+	 * @see UK2Node_AddComponent	DO NOT CALL MANUALLY - BLUEPRINT INTERNAL USE ONLY (for Add Component nodes)
+	 *
+	 * @param Component						The component created in AddComponent to finish creation of
+	 * @param bManualAttachment				Whether manual or automatic attachment is to be used
+	 * @param RelativeTransform				The relative transform between the new component and its attach parent (automatic only)
+	 */
+	UFUNCTION(BlueprintCallable, meta=(BlueprintInternalUseOnly="true"))
+	void FinishAddComponent(UActorComponent* Component, bool bManualAttachment, const FTransform& RelativeTransform);
 
 	UE_DEPRECATED(4.17, "Use UActorComponent::DestroyComponent() instead")
 	UFUNCTION(BlueprintCallable, meta=(DeprecatedFunction, DeprecationMessage = "Use Component.DestroyComponent instead", BlueprintProtected = "true", DisplayName = "DestroyComponent", ScriptName = "DestroyComponent"))
diff --git a/Engine/Source/Runtime/Engine/Private/ActorConstruction.cpp b/Engine/Source/Runtime/Engine/Private/ActorConstruction.cpp
--- a/Engine/Source/Runtime/Engine/Private/ActorConstruction.cpp
+++ b/Engine/Source/Runtime/Engine/Private/ActorConstruction.cpp
@@ -1079,7 +1079,7 @@ UActorComponent* AActor::CreateComponentFromTemplateData(const FBlueprintCookedC
 	return NewActorComp;
 }
 
-UActorComponent* AActor::AddComponent(FName TemplateName, bool bManualAttachment, const FTransform& RelativeTransform, const UObject* ComponentTemplateContext)
+UActorComponent* AActor::AddComponent(FName TemplateName, bool bManualAttachment, const FTransform& RelativeTransform, const UObject* ComponentTemplateContext, bool bDeferredFinish)
 {
 	UWorld* World = GetWorld();
 	if (World->bIsTearingDown)
@@ -1118,10 +1118,22 @@ UActorComponent* AActor::AddComponent(FName TemplateName, bool bManualAttachment
 		}
 	}
 
-	bool bIsSceneComponent = false;
 	UActorComponent* NewActorComp = TemplateData ? CreateComponentFromTemplateData(TemplateData) : CreateComponentFromTemplate(Template);
+
+	if (!bDeferredFinish)
+	{
+		FinishAddComponent(NewActorComp, bManualAttachment, RelativeTransform);
+	}
+
+	return NewActorComp;
+}
+
+void AActor::FinishAddComponent(UActorComponent* NewActorComp, bool bManualAttachment, const FTransform& RelativeTransform)
+{
 	if(NewActorComp != nullptr)
 	{
+		bool bIsSceneComponent = false;
+
 		// Call function to notify component it has been created
 		NewActorComp->OnComponentCreated();
 		
@@ -1153,6 +1165,7 @@ UActorComponent* AActor::AddComponent(FName TemplateName, bool bManualAttachment
 			NewActorComp->RegisterComponent();
 		}
 
+		UWorld* World = GetWorld();
 		if (!bRunningUserConstructionScript && World && bIsSceneComponent)
 		{
 			UPrimitiveComponent* NewPrimitiveComponent = Cast<UPrimitiveComponent>(NewActorComp);
@@ -1162,8 +1175,6 @@ UActorComponent* AActor::AddComponent(FName TemplateName, bool bManualAttachment
 			}
 		}
 	}
-
-	return NewActorComp;
 }
 
 void AActor::CheckComponentInstanceName(const FName InName)
