Added a new AddComponentByClass node that allows a class to be selected or passed in. Unlike existing Add Component nodes, this node may be used any place functions can be called on an Actor, however, you cannot customize the component template beyond any properties that are exposed on spawn.
Based in part on PR #6681: Added new Blueprint node AddComponentByClass (Contributed by [REDACTED_PERSON])

diff --git a/Engine/Source/Editor/BlueprintGraph/Classes/K2Node_AddComponentByClass.h b/Engine/Source/Editor/BlueprintGraph/Classes/K2Node_AddComponentByClass.h
new file mode 100644
--- /dev/null
+++ b/Engine/Source/Editor/BlueprintGraph/Classes/K2Node_AddComponentByClass.h
@@ -0,0 +1,41 @@
+// Copyright Epic Games, Inc. All Rights Reserved.
+
+#pragma once
+
+#include "CoreMinimal.h"
+#include "K2Node_ConstructObjectFromClass.h"
+#include "K2Node_AddComponentByClass.generated.h"
+
+/**
+ * Implementation of K2Node for creating a component based on a selected or passed in class
+ */
+UCLASS()
+class BLUEPRINTGRAPH_API UK2Node_AddComponentByClass : public UK2Node_ConstructObjectFromClass
+{
+	GENERATED_BODY()
+
+public:
+	UK2Node_AddComponentByClass(const FObjectInitializer& ObjectInitializer);
+	
+	//~ Begin UEdGraphNode Interface.
+	virtual void AllocateDefaultPins() override;
+	virtual void ExpandNode(FKismetCompilerContext& CompilerContext, UEdGraph* SourceGraph) override;
+	//~ End UEdGraphNode Interface.
+
+	UEdGraphPin* GetRelativeTransformPin() const;
+	UEdGraphPin* GetManualAttachmentPin() const;
+
+	/** Returns if the currently selected or linked class is known to be a scene component */
+	bool IsSceneComponent() const;
+
+	//~ Begin UK2Node_ConstructObjectFromClass Interface
+	virtual void CreatePinsForClass(UClass* InClass, TArray<UEdGraphPin*>* OutClassPins) override;
+
+protected:
+	virtual FText GetBaseNodeTitle() const override;
+	virtual FText GetDefaultNodeTitle() const override;
+	virtual FText GetNodeTitleFormat() const override;
+	virtual UClass* GetClassPinBaseClass() const override;
+	virtual bool IsSpawnVarPin(UEdGraphPin* Pin) const override;
+	//~ End UK2Node_ConstructObjectFromClass Interface
+};
diff --git a/Engine/Source/Editor/BlueprintGraph/Classes/K2Node_ConstructObjectFromClass.h b/Engine/Source/Editor/BlueprintGraph/Classes/K2Node_ConstructObjectFromClass.h
--- a/Engine/Source/Editor/BlueprintGraph/Classes/K2Node_ConstructObjectFromClass.h
+++ b/Engine/Source/Editor/BlueprintGraph/Classes/K2Node_ConstructObjectFromClass.h
@@ -39,7 +39,7 @@ class BLUEPRINTGRAPH_API UK2Node_ConstructObjectFromClass : public UK2Node
 	//~ End UK2Node Interface
 
 	/** Create new pins to show properties on archetype */
-	void CreatePinsForClass(UClass* InClass, TArray<UEdGraphPin*>* OutClassPins = nullptr);
+	virtual void CreatePinsForClass(UClass* InClass, TArray<UEdGraphPin*>* OutClassPins = nullptr);
 
 	/** See if this is a spawn variable pin, or a 'default' pin */
 	virtual bool IsSpawnVarPin(UEdGraphPin* Pin) const;
@@ -65,8 +65,10 @@ class BLUEPRINTGRAPH_API UK2Node_ConstructObjectFromClass : public UK2Node
 	virtual bool UseOuter() const { return false; }
 
 protected:
-	/** Gets the default node title when no class is selected */
+	/** Gets the node for use in lists and menus */
 	virtual FText GetBaseNodeTitle() const;
+	/** Gets the default node title when no class is selected */
+	virtual FText GetDefaultNodeTitle() const;
 	/** Gets the node title when a class has been selected. */
 	virtual FText GetNodeTitleFormat() const;
 	/** Gets base class to use for the 'class' pin.  UObject by default. */
diff --git a/Engine/Source/Editor/BlueprintGraph/Private/K2Node_AddComponentByClass.cpp b/Engine/Source/Editor/BlueprintGraph/Private/K2Node_AddComponentByClass.cpp
new file mode 100644
--- /dev/null
+++ b/Engine/Source/Editor/BlueprintGraph/Private/K2Node_AddComponentByClass.cpp
@@ -0,0 +1,235 @@
+// Copyright Epic Games, Inc. All Rights Reserved.
+
+
+#include "K2Node_AddComponentByClass.h"
+#include "EdGraphSchema_K2.h"
+#include "EditorCategoryUtils.h"
+#include "Components/ActorComponent.h"
+#include "GameFramework/Actor.h"
+#include "KismetCompilerMisc.h"
+#include "KismetCompiler.h"
+#include "K2Node_CallFunction.h"
+#include "Kismet2/BlueprintEditorUtils.h"
+
+
+#define LOCTEXT_NAMESPACE "ActorComponent"
+
+struct FK2Node_AddComponentByClassHelper
+{
+	static const FName TransformPinName;
+	static const FName ManualAttachmentPinName;
+	static const FName DeferredFinishPinName;
+	static const FName ComponentClassPinName;
+	static const FName ActorComponentPinName;
+};
+
+const FName FK2Node_AddComponentByClassHelper::TransformPinName(TEXT("RelativeTransform"));
+const FName FK2Node_AddComponentByClassHelper::ManualAttachmentPinName(TEXT("bManualAttachment"));
+const FName FK2Node_AddComponentByClassHelper::DeferredFinishPinName(TEXT("bDeferredFinish"));
+const FName FK2Node_AddComponentByClassHelper::ComponentClassPinName(TEXT("Class"));
+const FName FK2Node_AddComponentByClassHelper::ActorComponentPinName(TEXT("Component"));
+
+UK2Node_AddComponentByClass::UK2Node_AddComponentByClass(const FObjectInitializer& ObjectInitializer)
+	: Super(ObjectInitializer)
+{
+	NodeTooltip = LOCTEXT("NodeTooltip", "Adds a component to an actor");
+}
+
+void UK2Node_AddComponentByClass::AllocateDefaultPins()
+{
+	// Create a dummy AActor::AddComponentByClass node to copy pins off of
+	UK2Node_CallFunction* AddComponentByClassNode = NewObject<UK2Node_CallFunction>(GetGraph());
+	AddComponentByClassNode->SetFromFunction(AActor::StaticClass()->FindFunctionByName(GET_FUNCTION_NAME_CHECKED(AActor, AddComponentByClass)));
+	AddComponentByClassNode->AllocateDefaultPins();
+
+	auto CreatePinCopy = [&](UEdGraphPin* ProtoPin)
+	{
+		const FEdGraphPinType& ProtoPinType = ProtoPin->PinType;
+		UEdGraphNode::FCreatePinParams PinParams;
+		PinParams.ContainerType = ProtoPinType.ContainerType;
+		PinParams.ValueTerminalType = ProtoPinType.PinValueType;
+		UEdGraphPin* Pin = CreatePin(ProtoPin->Direction, ProtoPinType.PinCategory, ProtoPinType.PinSubCategory, ProtoPinType.PinSubCategoryObject.Get(), ProtoPin->PinName, PinParams);
+		Pin->PinToolTip = MoveTemp(ProtoPin->PinToolTip);
+		Pin->PinFriendlyName = MoveTemp(ProtoPin->PinFriendlyName);
+		return Pin;
+	};
+
+	// Create the self pin first because we want it to come before the class pin
+	UEdGraphPin* SelfPin = CreatePinCopy(GetDefault<UEdGraphSchema_K2>()->FindSelfPin(*AddComponentByClassNode, EGPD_Input));
+
+	Super::AllocateDefaultPins();
+
+	// Put the exec pin first in the Pins array so it appears before the self pin
+	UEdGraphPin* ExecPin = GetExecPin();
+	Pins.Remove(ExecPin);
+	Pins.Insert(ExecPin, 0);
+
+	CreatePinCopy(AddComponentByClassNode->FindPinChecked(FK2Node_AddComponentByClassHelper::ManualAttachmentPinName));
+	CreatePinCopy(AddComponentByClassNode->FindPinChecked(FK2Node_AddComponentByClassHelper::TransformPinName));
+
+	AddComponentByClassNode->DestroyNode();
+}
+
+FText UK2Node_AddComponentByClass::GetBaseNodeTitle() const
+{
+	return LOCTEXT("AddComponent_BaseTitle", "Add Component by Class");
+}
+
+FText UK2Node_AddComponentByClass::GetDefaultNodeTitle() const
+{
+	return GetBaseNodeTitle();
+}
+
+FText UK2Node_AddComponentByClass::GetNodeTitleFormat() const
+{
+	return LOCTEXT("AddComponent", "Add {ClassName}");
+}
+
+UClass* UK2Node_AddComponentByClass::GetClassPinBaseClass() const
+{
+	return UActorComponent::StaticClass();
+}
+
+void UK2Node_AddComponentByClass::CreatePinsForClass(UClass* InClass, TArray<UEdGraphPin*>* OutClassPins)
+{
+	Super::CreatePinsForClass(InClass, OutClassPins);
+
+	// Based on the supplied class we hide the transform and manual attachment pins
+	const bool bIsSceneComponent = InClass ? InClass->IsChildOf<USceneComponent>() : false;
+	
+	UEdGraphPin* ManualAttachmentPin = GetManualAttachmentPin();
+	ManualAttachmentPin->SafeSetHidden(!bIsSceneComponent);
+
+	UEdGraphPin* TransformPin = GetRelativeTransformPin();
+	TransformPin->SafeSetHidden(!bIsSceneComponent);
+}
+
+UEdGraphPin* UK2Node_AddComponentByClass::GetRelativeTransformPin() const
+{
+	return FindPinChecked(FK2Node_AddComponentByClassHelper::TransformPinName);
+}
+
+UEdGraphPin* UK2Node_AddComponentByClass::GetManualAttachmentPin() const
+{
+	return FindPinChecked(FK2Node_AddComponentByClassHelper::ManualAttachmentPinName);
+}
+
+bool UK2Node_AddComponentByClass::IsSceneComponent() const
+{
+	if (UEdGraphPin* SpawnClassPin = GetClassPin())
+	{
+		if (UClass* SpawnClass = Cast<UClass>(SpawnClassPin->DefaultObject))
+		{
+			return SpawnClass->IsChildOf<USceneComponent>();
+		}
+	}
+	return false;
+}
+
+bool UK2Node_AddComponentByClass::IsSpawnVarPin(UEdGraphPin* Pin) const
+{
+	return(Super::IsSpawnVarPin(Pin)
+	       && Pin->PinName != UEdGraphSchema_K2::PN_Self
+	       && Pin->PinName != FK2Node_AddComponentByClassHelper::TransformPinName
+	       && Pin->PinName != FK2Node_AddComponentByClassHelper::ManualAttachmentPinName);	  
+}
+
+void UK2Node_AddComponentByClass::ExpandNode(class FKismetCompilerContext& CompilerContext, UEdGraph* SourceGraph)
+{
+	Super::ExpandNode(CompilerContext, SourceGraph);
+
+
+	const UEdGraphSchema_K2* K2Schema = GetDefault<UEdGraphSchema_K2>();
+
+	UEdGraphPin* SpawnNodeExec = GetExecPin();
+	UEdGraphPin* SpawnOwnerPin = K2Schema->FindSelfPin(*this, EGPD_Input);
+	UEdGraphPin* SpawnClassPin = GetClassPin();
+	UEdGraphPin* SpawnTransformPin = GetRelativeTransformPin();
+	UEdGraphPin* SpawnManualAttachmentPin = GetManualAttachmentPin();
+
+	UEdGraphPin* SpawnNodeThen = GetThenPin();
+	UEdGraphPin* SpawnNodeResult = GetResultPin();
+
+	UClass* SpawnClass = (SpawnClassPin ? Cast<UClass>(SpawnClassPin->DefaultObject) : nullptr);
+	if (SpawnClassPin == nullptr || ((SpawnClass == nullptr) && (SpawnClassPin->LinkedTo.Num() == 0))
+	{
+		CompilerContext.MessageLog.Error(*LOCTEXT("AddComponentByClassNodeMissingClass_Error", "Spawn node @@ must have a class specified.").ToString(), this);
+		// we break exec links so this is the only error we get, don't want the AddComponentByClass node being considered and giving 'unexpected node' type warnings
+		BreakAllNodeLinks();
+		return;
+	}
+
+	UK2Node_CallFunction* CallAddComponentByClassNode = CompilerContext.SpawnIntermediateNode<UK2Node_CallFunction>(this, SourceGraph);
+	CallAddComponentByClassNode->SetFromFunction(AActor::StaticClass()->FindFunctionByName(GET_FUNCTION_NAME_CHECKED(AActor, AddComponentByClass)));
+	CallAddComponentByClassNode->AllocateDefaultPins();
+
+	// store off the class to spawn before we mutate pin connections:
+	UClass* ClassToSpawn = GetClassToSpawn();
+
+	UEdGraphPin* CallAddComponentByClassExec = CallAddComponentByClassNode->GetExecPin();
+	UEdGraphPin* CallAddComponentByClassTypePin = CallAddComponentByClassNode->FindPinChecked(FK2Node_AddComponentByClassHelper::ComponentClassPinName);
+	UEdGraphPin* CallAddComponentByClassOwnerPin = K2Schema->FindSelfPin(*CallAddComponentByClassNode, EGPD_Input);
+	UEdGraphPin* CallAddComponentByClassTransformPin = CallAddComponentByClassNode->FindPinChecked(FK2Node_AddComponentByClassHelper::TransformPinName);
+	UEdGraphPin* CallAddComponentByClassManualAttachmentPin = CallAddComponentByClassNode->FindPinChecked(FK2Node_AddComponentByClassHelper::ManualAttachmentPinName);
+	UEdGraphPin* CallAddComponentByClassResult = CallAddComponentByClassNode->GetReturnValuePin();
+
+	// set properties on relative transform pin to allow it to be unconnected
+	CallAddComponentByClassTransformPin->bDefaultValueIsIgnored = true;
+	CallAddComponentByClassTransformPin->PinType.bIsReference = false;
+
+	CompilerContext.MovePinLinksToIntermediate(*SpawnNodeExec, *CallAddComponentByClassExec);
+	CompilerContext.MovePinLinksToIntermediate(*SpawnClassPin, *CallAddComponentByClassTypePin);
+	CompilerContext.MovePinLinksToIntermediate(*SpawnOwnerPin, *CallAddComponentByClassOwnerPin);
+	CompilerContext.MovePinLinksToIntermediate(*SpawnTransformPin, *CallAddComponentByClassTransformPin);
+	CompilerContext.MovePinLinksToIntermediate(*SpawnManualAttachmentPin, *CallAddComponentByClassManualAttachmentPin);
+
+	// Move result connection from spawn node to specific AddComponent function
+	CallAddComponentByClassResult->PinType = SpawnNodeResult->PinType; // Copy type so it uses the right actor subclass
+	CompilerContext.MovePinLinksToIntermediate(*SpawnNodeResult, *CallAddComponentByClassResult);
+
+	//////////////////////////////////////////////////////////////////////////
+	// create 'set var' nodes
+
+	UEdGraphPin* LastAssignmentThen = FKismetCompilerUtilities::GenerateAssignmentNodes(CompilerContext, SourceGraph, CallAddComponentByClassNode, this, CallAddComponentByClassResult, ClassToSpawn);
+
+	if (LastAssignmentThen != CallAddComponentByClassNode->GetThenPin())
+	{
+		UEdGraphPin* CallAddComponentByClassDeferredFinishPin = CallAddComponentByClassNode->FindPinChecked(FK2Node_AddComponentByClassHelper::DeferredFinishPinName);
+		CallAddComponentByClassDeferredFinishPin->DefaultValue = TEXT("true");
+
+		UK2Node_CallFunction* CallRegisterComponentNode = CompilerContext.SpawnIntermediateNode<UK2Node_CallFunction>(this, SourceGraph);
+		CallRegisterComponentNode->SetFromFunction(AActor::StaticClass()->FindFunctionByName(GET_FUNCTION_NAME_CHECKED(AActor, FinishAddComponent)));
+		CallRegisterComponentNode->AllocateDefaultPins();
+
+		// Get last 'then' of assignments
+		UEdGraphPin* LastThen = CallRegisterComponentNode->GetThenPin();
+
+		// Links execution from last assignment to 'RegisterComponent '
+		LastAssignmentThen->MakeLinkTo(CallRegisterComponentNode->GetExecPin());
+
+		// Link the pins to RegisterComponent node
+		UEdGraphPin* CallRegisterComponentByClassOwnerPin = K2Schema->FindSelfPin(*CallRegisterComponentNode, EGPD_Input);
+		UEdGraphPin* CallRegisterComponentByClassComponentPin = CallRegisterComponentNode->FindPinChecked(FK2Node_AddComponentByClassHelper::ActorComponentPinName);
+		UEdGraphPin* CallRegisterComponentByClassTransformPin = CallRegisterComponentNode->FindPinChecked(FK2Node_AddComponentByClassHelper::TransformPinName);
+		UEdGraphPin* CallRegisterComponentByClassManualAttachmentPin = CallRegisterComponentNode->FindPinChecked(FK2Node_AddComponentByClassHelper::ManualAttachmentPinName);
+
+		// set properties on relative transform pin to allow it to be unconnected
+		CallRegisterComponentByClassTransformPin->bDefaultValueIsIgnored = true;
+		CallRegisterComponentByClassTransformPin->PinType.bIsReference = false;
+
+		CompilerContext.CopyPinLinksToIntermediate(*CallAddComponentByClassOwnerPin, *CallRegisterComponentByClassOwnerPin);
+		CompilerContext.CopyPinLinksToIntermediate(*CallAddComponentByClassTransformPin, *CallRegisterComponentByClassTransformPin);
+		CompilerContext.CopyPinLinksToIntermediate(*CallAddComponentByClassManualAttachmentPin, *CallRegisterComponentByClassManualAttachmentPin);
+
+		CallRegisterComponentByClassComponentPin->MakeLinkTo(CallAddComponentByClassResult);
+
+		// Move 'then' connection from AddComponent node to the last 'then'
+		CompilerContext.MovePinLinksToIntermediate(*SpawnNodeThen, *LastThen);
+	}
+
+	// Break any links to the expanded node
+	BreakAllNodeLinks();
+}
+
+
+#undef LOCTEXT_NAMESPACE
diff --git a/Engine/Source/Editor/BlueprintGraph/Private/K2Node_ConstructObjectFromClass.cpp b/Engine/Source/Editor/BlueprintGraph/Private/K2Node_ConstructObjectFromClass.cpp
--- a/Engine/Source/Editor/BlueprintGraph/Private/K2Node_ConstructObjectFromClass.cpp
+++ b/Engine/Source/Editor/BlueprintGraph/Private/K2Node_ConstructObjectFromClass.cpp
@@ -326,6 +326,11 @@ FText UK2Node_ConstructObjectFromClass::GetBaseNodeTitle() const
 	return NSLOCTEXT("K2Node", "ConstructObject_BaseTitle", "Construct Object from Class");
 }
 
+FText UK2Node_ConstructObjectFromClass::GetDefaultNodeTitle() const
+{
+	return NSLOCTEXT("K2Node", "ConstructObject_Title_NONE", "Construct NONE");
+}
+
 FText UK2Node_ConstructObjectFromClass::GetNodeTitleFormat() const
 {
 	return NSLOCTEXT("K2Node", "Construct", "Construct {ClassName}");
@@ -348,7 +353,7 @@ FText UK2Node_ConstructObjectFromClass::GetNodeTitle(ENodeTitleType::Type TitleT
 		}
 		return CachedNodeTitle;
 	}
-	return NSLOCTEXT("K2Node", "ConstructObject_Title_NONE", "Construct NONE");
+	return GetDefaultNodeTitle();
 }
 
 bool UK2Node_ConstructObjectFromClass::IsCompatibleWithGraph(const UEdGraph* TargetGraph) const 
diff --git a/Engine/Source/Runtime/Engine/Classes/GameFramework/Actor.h b/Engine/Source/Runtime/Engine/Classes/GameFramework/Actor.h
--- a/Engine/Source/Runtime/Engine/Classes/GameFramework/Actor.h
+++ b/Engine/Source/Runtime/Engine/Classes/GameFramework/Actor.h
@@ -1263,7 +1263,25 @@ public:
 	 * @param bDeferredFinish				Whether or not to immediately complete the creation and registration process for this component. Will be false if there are expose on spawn properties being set
 	 */
 	UFUNCTION(BlueprintCallable, meta=(ScriptNoExport, BlueprintInternalUseOnly = "true", DefaultToSelf="ComponentTemplateContext", InternalUseParam="ComponentTemplateContext,bDeferredFinish"))
-	class UActorComponent* AddComponent(FName TemplateName, bool bManualAttachment, const FTransform& RelativeTransform, const UObject* ComponentTemplateContext, bool bDeferredFinish = false);
+	UActorComponent* AddComponent(FName TemplateName, bool bManualAttachment, const FTransform& RelativeTransform, const UObject* ComponentTemplateContext, bool bDeferredFinish = false);
+
+	/**
+	 * Creates a new component and assigns ownership to the Actor this is
+	 * called for. Automatic attachment causes the first component created to
+	 * become the root, and all subsequent components to be attached under that
+	 * root. When bManualAttachment is set, automatic attachment is
+	 * skipped and it is up to the user to attach the resulting component (or
+	 * set it up as the root) themselves.
+	 *
+	 * @see UK2Node_AddComponentByClass		DO NOT CALL MANUALLY - BLUEPRINT INTERNAL USE ONLY (for Add Component nodes)
+	 *
+	 * @param Class						The class of component to create
+	 * @param bManualAttachment				Whether manual or automatic attachment is to be used
+	 * @param RelativeTransform				The relative transform between the new component and its attach parent (automatic only)
+	 * @param bDeferredFinish				Whether or not to immediately complete the creation and registration process for this component. Will be false if there are expose on spawn properties being set
+	 */
+	UFUNCTION(BlueprintCallable, meta=(ScriptNoExport, BlueprintInternalUseOnly = "true", InternalUseParam="bDeferredFinish"))
+	UActorComponent* AddComponentByClass(TSubclassOf<UActorComponent> Class, bool bManualAttachment, const FTransform& RelativeTransform, bool bDeferredFinish);
 
 	/** 
 	 * Completes the creation of a new actor component. Called either from blueprint after
diff --git a/Engine/Source/Runtime/Engine/Private/ActorConstruction.cpp b/Engine/Source/Runtime/Engine/Private/ActorConstruction.cpp
--- a/Engine/Source/Runtime/Engine/Private/ActorConstruction.cpp
+++ b/Engine/Source/Runtime/Engine/Private/ActorConstruction.cpp
@@ -1128,6 +1128,26 @@ UActorComponent* AActor::AddComponent(FName TemplateName, bool bManualAttachment
 	return NewActorComp;
 }
 
+UActorComponent* AActor::AddComponentByClass(TSubclassOf<UActorComponent> Class, bool bManualAttachment, const FTransform& RelativeTransform, bool bDeferredFinish)
+{
+	UWorld* World = GetWorld();
+	if (World->bIsTearingDown)
+	{
+		UE_LOG(LogActor, Warning, TEXT("AddComponent failed because we are in the process of tearing down the world"));
+		return nullptr;
+	}
+
+	UActorComponent* NewActorComp = NewObject<UActorComponent>(this, *Class);
+	PostCreateBlueprintComponent(NewActorComp);
+
+	if (!bDeferredFinish)
+	{
+		FinishAddComponent(NewActorComp, bManualAttachment, RelativeTransform);
+	}
+
+	return NewActorComp;
+}
+
 void AActor::FinishAddComponent(UActorComponent* NewActorComp, bool bManualAttachment, const FTransform& RelativeTransform)
 {
 	if(NewActorComp != nullptr)
