Fix missing parantheses
#rnx

diff --git a/Engine/Source/Editor/BlueprintGraph/Private/K2Node_AddComponentByClass.cpp b/Engine/Source/Editor/BlueprintGraph/Private/K2Node_AddComponentByClass.cpp
--- a/Engine/Source/Editor/BlueprintGraph/Private/K2Node_AddComponentByClass.cpp
+++ b/Engine/Source/Editor/BlueprintGraph/Private/K2Node_AddComponentByClass.cpp
@@ -151,7 +151,7 @@ void UK2Node_AddComponentByClass::ExpandNode(class FKismetCompilerContext& Compi
 	UEdGraphPin* SpawnNodeResult = GetResultPin();
 
 	UClass* SpawnClass = (SpawnClassPin ? Cast<UClass>(SpawnClassPin->DefaultObject) : nullptr);
-	if (SpawnClassPin == nullptr || ((SpawnClass == nullptr) && (SpawnClassPin->LinkedTo.Num() == 0))
+	if (SpawnClassPin == nullptr || ((SpawnClass == nullptr) && (SpawnClassPin->LinkedTo.Num() == 0)))
 	{
 		CompilerContext.MessageLog.Error(*LOCTEXT("AddComponentByClassNodeMissingClass_Error", "Spawn node @@ must have a class specified.").ToString(), this);
 		// we break exec links so this is the only error we get, don't want the AddComponentByClass node being considered and giving 'unexpected node' type warnings
