Fix crash if null class is used to add component by class

Fix scene component pins being shown on add component by class node at incorrect times (mostly when the class is null)

#rnx

diff --git a/Engine/Source/Editor/BlueprintGraph/Classes/K2Node_AddComponentByClass.h b/Engine/Source/Editor/BlueprintGraph/Classes/K2Node_AddComponentByClass.h
--- a/Engine/Source/Editor/BlueprintGraph/Classes/K2Node_AddComponentByClass.h
+++ b/Engine/Source/Editor/BlueprintGraph/Classes/K2Node_AddComponentByClass.h
@@ -20,14 +20,9 @@ public:
 	//~ Begin UEdGraphNode Interface.
 	virtual void AllocateDefaultPins() override;
 	virtual void ExpandNode(FKismetCompilerContext& CompilerContext, UEdGraph* SourceGraph) override;
+	virtual void PinDefaultValueChanged(UEdGraphPin* Pin) override;
 	//~ End UEdGraphNode Interface.
 
-	UEdGraphPin* GetRelativeTransformPin() const;
-	UEdGraphPin* GetManualAttachmentPin() const;
-
-	/** Returns if the currently selected or linked class is known to be a scene component */
-	bool IsSceneComponent() const;
-
 	//~ Begin UK2Node_ConstructObjectFromClass Interface
 	virtual void CreatePinsForClass(UClass* InClass, TArray<UEdGraphPin*>* OutClassPins) override;
 
@@ -38,4 +33,13 @@ protected:
 	virtual UClass* GetClassPinBaseClass() const override;
 	virtual bool IsSpawnVarPin(UEdGraphPin* Pin) const override;
 	//~ End UK2Node_ConstructObjectFromClass Interface
+
+	UEdGraphPin* GetRelativeTransformPin() const;
+	UEdGraphPin* GetManualAttachmentPin() const;
+
+	/** Returns true if the currently selected or linked class is known to be a scene component */
+	bool IsSceneComponent() const;
+
+	/** Utility function to set whether the scene component specific pins are hidden or not */
+	void SetSceneComponentPinsHidden(bool bHidden);
 };
diff --git a/Engine/Source/Editor/BlueprintGraph/Private/K2Node_AddComponentByClass.cpp b/Engine/Source/Editor/BlueprintGraph/Private/K2Node_AddComponentByClass.cpp
--- a/Engine/Source/Editor/BlueprintGraph/Private/K2Node_AddComponentByClass.cpp
+++ b/Engine/Source/Editor/BlueprintGraph/Private/K2Node_AddComponentByClass.cpp
@@ -16,6 +16,7 @@
 
 struct FK2Node_AddComponentByClassHelper
 {
+	static const FName ClassPinName;
 	static const FName TransformPinName;
 	static const FName ManualAttachmentPinName;
 	static const FName DeferredFinishPinName;
@@ -23,6 +24,7 @@ struct FK2Node_AddComponentByClassHelper
 	static const FName ActorComponentPinName;
 };
 
+const FName FK2Node_AddComponentByClassHelper::ClassPinName(TEXT("Class"));
 const FName FK2Node_AddComponentByClassHelper::TransformPinName(TEXT("RelativeTransform"));
 const FName FK2Node_AddComponentByClassHelper::ManualAttachmentPinName(TEXT("bManualAttachment"));
 const FName FK2Node_AddComponentByClassHelper::DeferredFinishPinName(TEXT("bDeferredFinish"));
@@ -68,6 +70,8 @@ void UK2Node_AddComponentByClass::AllocateDefaultPins()
 	CreatePinCopy(AddComponentByClassNode->FindPinChecked(FK2Node_AddComponentByClassHelper::TransformPinName));
 
 	AddComponentByClassNode->DestroyNode();
+
+	SetSceneComponentPinsHidden(true);
 }
 
 FText UK2Node_AddComponentByClass::GetBaseNodeTitle() const
@@ -90,18 +94,33 @@ UClass* UK2Node_AddComponentByClass::GetClassPinBaseClass() const
 	return UActorComponent::StaticClass();
 }
 
+void UK2Node_AddComponentByClass::SetSceneComponentPinsHidden(const bool bHidden)
+{
+	UEdGraphPin* ManualAttachmentPin = GetManualAttachmentPin();
+	ManualAttachmentPin->SafeSetHidden(bHidden);
+
+	UEdGraphPin* TransformPin = GetRelativeTransformPin();
+	TransformPin->SafeSetHidden(bHidden);
+}
+
+void UK2Node_AddComponentByClass::PinDefaultValueChanged(UEdGraphPin* ChangedPin)
+{
+	if (ChangedPin && (ChangedPin->PinName == FK2Node_AddComponentByClassHelper::ClassPinName))
+	{
+		// If the class pin changes set the scene components hidden and then CreatePinsForClass will unhide if appropriate
+		SetSceneComponentPinsHidden(true);
+	}
+
+	Super::PinDefaultValueChanged(ChangedPin);
+}
+
 void UK2Node_AddComponentByClass::CreatePinsForClass(UClass* InClass, TArray<UEdGraphPin*>* OutClassPins)
 {
 	Super::CreatePinsForClass(InClass, OutClassPins);
 
 	// Based on the supplied class we hide the transform and manual attachment pins
 	const bool bIsSceneComponent = InClass ? InClass->IsChildOf<USceneComponent>() : false;
-	
-	UEdGraphPin* ManualAttachmentPin = GetManualAttachmentPin();
-	ManualAttachmentPin->SafeSetHidden(!bIsSceneComponent);
-
-	UEdGraphPin* TransformPin = GetRelativeTransformPin();
-	TransformPin->SafeSetHidden(!bIsSceneComponent);
+	SetSceneComponentPinsHidden(!bIsSceneComponent);
 }
 
 UEdGraphPin* UK2Node_AddComponentByClass::GetRelativeTransformPin() const
diff --git a/Engine/Source/Runtime/Engine/Private/ActorConstruction.cpp b/Engine/Source/Runtime/Engine/Private/ActorConstruction.cpp
--- a/Engine/Source/Runtime/Engine/Private/ActorConstruction.cpp
+++ b/Engine/Source/Runtime/Engine/Private/ActorConstruction.cpp
@@ -1130,6 +1130,11 @@ UActorComponent* AActor::AddComponent(FName TemplateName, bool bManualAttachment
 
 UActorComponent* AActor::AddComponentByClass(TSubclassOf<UActorComponent> Class, bool bManualAttachment, const FTransform& RelativeTransform, bool bDeferredFinish)
 {
+	if (Class == nullptr)
+	{
+		return nullptr;
+	}
+
 	UWorld* World = GetWorld();
 	if (World->bIsTearingDown)
 	{
