Double property support in blueprint for UFunctions, BP Variables, and conversions to floats

diff --git a/Engine/Source/Developer/BlueprintCompilerCppBackend/Private/BlueprintCompilerCppBackendUtils.cpp b/Engine/Source/Developer/BlueprintCompilerCppBackend/Private/BlueprintCompilerCppBackendUtils.cpp
--- a/Engine/Source/Developer/BlueprintCompilerCppBackend/Private/BlueprintCompilerCppBackendUtils.cpp
+++ b/Engine/Source/Developer/BlueprintCompilerCppBackend/Private/BlueprintCompilerCppBackendUtils.cpp
@@ -1065,6 +1065,16 @@ FString FEmitHelper::FloatToString(float Value)
 	return FString::Printf(TEXT("%f"), Value);
 }
 
+FString FEmitHelper::DoubleToString(double Value)
+{
+	if (FMath::IsNaN(Value))
+	{
+		UE_LOG(LogK2Compiler, Warning, TEXT("A NotANNumber value cannot be nativized. It is changed into 0.0."));
+		return TEXT("/*The original value was NaN!*/ 0.0");
+	}
+	return FString::Printf(TEXT("%lf"), Value);
+}
+
 FString FEmitHelper::LiteralTerm(FEmitterLocalContext& EmitterContext, const FLiteralTermParams& Params)
 {
 	const UEdGraphSchema_K2* Schema = GetDefault<UEdGraphSchema_K2>();
@@ -1258,6 +1268,11 @@ FString FEmitHelper::LiteralTerm(FEmitterLocalContext& EmitterContext, const FLi
 		float Value = CustomValue.IsEmpty() ? 0.0f : FCString::Atof(*CustomValue);
 		return FString::Printf(TEXT("%s"), *FloatToString(Value));
 	}
+	else if (UEdGraphSchema_K2::PC_Double == Type.PinCategory)
+	{
+		double Value = CustomValue.IsEmpty() ? 0.0f : FCString::Atod(*CustomValue);
+		return FString::Printf(TEXT("%s"), *DoubleToString(Value));
+	}
 	else if (UEdGraphSchema_K2::PC_Int == Type.PinCategory)
 	{
 		int32 Value = CustomValue.IsEmpty() ? 0 : FCString::Atoi(*CustomValue);
@@ -1511,6 +1526,10 @@ FString FEmitHelper::PinTypeToNativeType(const FEdGraphPinType& Type)
 		{
 			return TEXT("float");
 		}
+		else if (UEdGraphSchema_K2::PC_Double == InType.PinCategory)
+		{
+			return TEXT("double");
+		}
 		else if (UEdGraphSchema_K2::PC_Name == InType.PinCategory)
 		{
 			return TEXT("FName");
diff --git a/Engine/Source/Developer/BlueprintCompilerCppBackend/Private/BlueprintCompilerCppBackendUtils.h b/Engine/Source/Developer/BlueprintCompilerCppBackend/Private/BlueprintCompilerCppBackendUtils.h
--- a/Engine/Source/Developer/BlueprintCompilerCppBackend/Private/BlueprintCompilerCppBackendUtils.h
+++ b/Engine/Source/Developer/BlueprintCompilerCppBackend/Private/BlueprintCompilerCppBackendUtils.h
@@ -253,6 +253,8 @@ struct FEmitHelper
 
 	static FString FloatToString(float Value);
 
+	static FString DoubleToString(double Value);
+
 	static bool PropertyForConstCast(const FProperty* Property);
 
 	static void ArrayToString(const TArray<FString>& Array, FString& OutString, const TCHAR* Separator);
diff --git a/Engine/Source/Developer/ScriptDisassembler/Private/ScriptDisassembler.cpp b/Engine/Source/Developer/ScriptDisassembler/Private/ScriptDisassembler.cpp
--- a/Engine/Source/Developer/ScriptDisassembler/Private/ScriptDisassembler.cpp
+++ b/Engine/Source/Developer/ScriptDisassembler/Private/ScriptDisassembler.cpp
@@ -140,6 +140,13 @@ float FKismetBytecodeDisassembler::ReadFLOAT(int32& ScriptIndex)
 	return Result.f;
 }
 
+double FKismetBytecodeDisassembler::ReadDOUBLE(int32& ScriptIndex)
+{
+	union { double d; int64 i; } Result;
+	Result.i = ReadQWORD(ScriptIndex);
+	return Result.d;
+}
+
 CodeSkipSizeType FKismetBytecodeDisassembler::ReadSkipCount(int32& ScriptIndex)
 {
 #if SCRIPT_LIMIT_BYTECODE_TO_64KB
@@ -710,6 +717,12 @@ void FKismetBytecodeDisassembler::ProcessCommon(int32& ScriptIndex, EExprToken O
 			Ar.Logf(TEXT("%s $%X: literal float %f"), *Indents, (int32)Opcode, ConstValue);
 			break;
 		}
+	case EX_DoubleConst:
+		{
+			double ConstValue = ReadDOUBLE(ScriptIndex);
+			Ar.Logf(TEXT("%s $%X: literal double %lf"), *Indents, (int32)Opcode, ConstValue);
+			break;
+		}
 	case EX_StringConst:
 		{
 			FString ConstValue = ReadString8(ScriptIndex);
diff --git a/Engine/Source/Developer/ScriptDisassembler/Public/ScriptDisassembler.h b/Engine/Source/Developer/ScriptDisassembler/Public/ScriptDisassembler.h
--- a/Engine/Source/Developer/ScriptDisassembler/Public/ScriptDisassembler.h
+++ b/Engine/Source/Developer/ScriptDisassembler/Public/ScriptDisassembler.h
@@ -50,6 +50,7 @@ private:
 	FString ReadName(int32& ScriptIndex);
 	uint16 ReadWORD(int32& ScriptIndex);
 	float ReadFLOAT(int32& ScriptIndex);
+	double ReadDOUBLE(int32& ScriptIndex);
 	CodeSkipSizeType ReadSkipCount(int32& ScriptIndex);
 	FString ReadString(int32& ScriptIndex);
 	FString ReadString8(int32& ScriptIndex);
diff --git a/Engine/Source/Editor/BlueprintGraph/Classes/EdGraphSchema_K2.h b/Engine/Source/Editor/BlueprintGraph/Classes/EdGraphSchema_K2.h
--- a/Engine/Source/Editor/BlueprintGraph/Classes/EdGraphSchema_K2.h
+++ b/Engine/Source/Editor/BlueprintGraph/Classes/EdGraphSchema_K2.h
@@ -298,6 +298,7 @@ class BLUEPRINTGRAPH_API UEdGraphSchema_K2 : public UEdGraphSchema
 	static const FName PC_Int;
 	static const FName PC_Int64;
 	static const FName PC_Float;
+	static const FName PC_Double;
 	static const FName PC_Name;
 	static const FName PC_Delegate;    // SubCategoryObject is the UFunction of the delegate signature
 	static const FName PC_MCDelegate;  // SubCategoryObject is the UFunction of the delegate signature
diff --git a/Engine/Source/Editor/BlueprintGraph/Private/EdGraphSchema_K2.cpp b/Engine/Source/Editor/BlueprintGraph/Private/EdGraphSchema_K2.cpp
--- a/Engine/Source/Editor/BlueprintGraph/Private/EdGraphSchema_K2.cpp
+++ b/Engine/Source/Editor/BlueprintGraph/Private/EdGraphSchema_K2.cpp
@@ -263,6 +263,8 @@ private:
 public:
 	static void FillLoadedTypesDatabase(FTypesDatabase& TypesDatabase, bool bIndexTypesOnly)
 	{
+		TRACE_CPUPROFILER_EVENT_SCOPE(WILD_FGatherTypesHelper::FillLoadedTypesDatabase);
+
 		// Loaded types
 		TypesDatabase.LoadedTypesMap.Reset();
 
@@ -331,6 +333,7 @@ public:
 
 	static void FillUnLoadedTypesDatabase(FTypesDatabase& TypesDatabase, bool bIndexTypesOnly)
 	{
+		TRACE_CPUPROFILER_EVENT_SCOPE(WILD_FGatherTypesHelper::FillUnLoadedTypesDatabase);
 		// Loaded types
 		TypesDatabase.UnLoadedTypesMap.Reset();
 
@@ -424,6 +427,8 @@ public:
 	 */
 	static void Gather(const FText& FriendlyName, const FName CategoryName, FTypesDatabase& TypesDatabase, TArray<FPinTypeTreeInfoPtr>& OutChildren)
 	{
+		TRACE_CPUPROFILER_EVENT_SCOPE(WILD_FGatherTypesHelper::Gather);
+
 		FEdGraphPinType LoadedPinSubtype;
 		LoadedPinSubtype.PinCategory = (CategoryName == UEdGraphSchema_K2::PC_Enum ? UEdGraphSchema_K2::PC_Byte : CategoryName);
 		LoadedPinSubtype.PinSubCategory = NAME_None;
@@ -524,6 +529,7 @@ const FEdGraphPinType& UEdGraphSchema_K2::FPinTypeTreeInfo::GetPinType(bool bFor
 
 void UEdGraphSchema_K2::FPinTypeTreeInfo::Init(const FText& InFriendlyName, const FName CategoryName, const UEdGraphSchema_K2* Schema, const FText& InTooltip, bool bInReadOnly, FTypesDatabase* TypesDatabase)
 {
+	TRACE_CPUPROFILER_EVENT_SCOPE(WILD_FPinTypeTreeInfo::Init);
 	check( !CategoryName.IsNone() );
 	check( Schema );
 
@@ -616,6 +622,7 @@ const FName UEdGraphSchema_K2::PC_Class(TEXT("class"));
 const FName UEdGraphSchema_K2::PC_Int(TEXT("int"));
 const FName UEdGraphSchema_K2::PC_Int64(TEXT("int64"));
 const FName UEdGraphSchema_K2::PC_Float(TEXT("float"));
+const FName UEdGraphSchema_K2::PC_Double(TEXT("double"));
 const FName UEdGraphSchema_K2::PC_Name(TEXT("name"));
 const FName UEdGraphSchema_K2::PC_Delegate(TEXT("delegate"));
 const FName UEdGraphSchema_K2::PC_MCDelegate(TEXT("mcdelegate"));
@@ -2514,6 +2521,8 @@ public:
 
 	void Refresh()
 	{
+		TRACE_CPUPROFILER_EVENT_SCOPE(WILD_FAutocastFunctionMap::Refresh);
+
 #ifdef SCHEMA_K2_AUTOCASTFUNCTIONMAP_LOG_TIME
 		static_assert(false, "Macro redefinition.");
 #endif
@@ -2616,6 +2625,7 @@ void UEdGraphSchema_K2::Shutdown()
 
 bool UEdGraphSchema_K2::SearchForAutocastFunction(const UEdGraphPin* OutputPin, const UEdGraphPin* InputPin, /*out*/ FName& TargetFunction, /*out*/ UClass*& FunctionOwner) const
 {
+	TRACE_CPUPROFILER_EVENT_SCOPE(WILD_UEdGraphSchema_K2::SearchForAutocastFunction);
 	// NOTE: Under no circumstances should anyone *ever* add a questionable cast to this function.
 	// If it could be at all confusing why a function is provided, to even a novice user, err on the side of do not cast!!!
 	// This includes things like string->int (does it do length, atoi, or what?) that would be autocasts in a traditional scripting language
@@ -3040,6 +3050,10 @@ FLinearColor UEdGraphSchema_K2::GetPinTypeColor(const FEdGraphPinType& PinType)
 	{
 		return Settings->FloatPinTypeColor;
 	}
+	else if (TypeName == PC_Double)
+	{
+		return Settings->DoublePinTypeColor;
+	}
 	else if (TypeName == PC_Boolean)
 	{
 		return Settings->BooleanPinTypeColor;
@@ -3484,6 +3498,10 @@ bool UEdGraphSchema_K2::GetPropertyCategoryInfo(const FProperty* TestProperty, F
 	{
 		OutCategory = PC_Float;
 	}
+	else if(TestProperty->IsA<FDoubleProperty>())
+	{
+		OutCategory = PC_Double;
+	}
 	else if (TestProperty->IsA<FInt64Property>())
 	{
 		OutCategory = PC_Int64;
@@ -3771,6 +3789,7 @@ FText UEdGraphSchema_K2::GetCategoryText(const FName Category, const bool bForMe
 		CategoryDescriptions.Add(PC_Int, LOCTEXT("IntCategory", "Integer"));
 		CategoryDescriptions.Add(PC_Int64, LOCTEXT("Int64Category", "Integer64"));
 		CategoryDescriptions.Add(PC_Float, LOCTEXT("FloatCategory","Float"));
+		CategoryDescriptions.Add(PC_Double, LOCTEXT("DoubleCategory", "Double"));
 		CategoryDescriptions.Add(PC_Name, LOCTEXT("NameCategory","Name"));
 		CategoryDescriptions.Add(PC_Delegate, LOCTEXT("DelegateCategory","Delegate"));
 		CategoryDescriptions.Add(PC_MCDelegate, LOCTEXT("MulticastDelegateCategory","Multicast Delegate"));
@@ -3960,6 +3979,7 @@ void UEdGraphSchema_K2::GetVariableTypeTree(TArray< TSharedPtr<FPinTypeTreeInfo>
 	if (!bIndexTypesOnly)
 	{
 		TypeTree.Add(MakeShareable(new FPinTypeTreeInfo(GetCategoryText(PC_Float, true), PC_Float, this, LOCTEXT("FloatType", "Floating point number"))));
+		TypeTree.Add(MakeShareable(new FPinTypeTreeInfo(GetCategoryText(PC_Double, true), PC_Double, this, LOCTEXT("DoubleType", "64 bit floating point number"))));
 		TypeTree.Add(MakeShareable(new FPinTypeTreeInfo(GetCategoryText(PC_Name, true), PC_Name, this, LOCTEXT("NameType", "A text name"))));
 		TypeTree.Add(MakeShareable(new FPinTypeTreeInfo(GetCategoryText(PC_String, true), PC_String, this, LOCTEXT("StringType", "A text string"))));
 		TypeTree.Add(MakeShareable(new FPinTypeTreeInfo(GetCategoryText(PC_Text, true), PC_Text, this, LOCTEXT("TextType", "A localizable text string"))));
@@ -4242,6 +4262,16 @@ bool UEdGraphSchema_K2::DefaultValueSimpleValidation(const FEdGraphPinType& PinT
 			}
 		}
 	}
+	else if (PinCategory == PC_Double)
+	{
+		if (!NewDefaultValue.IsEmpty())
+		{
+			if (!FDefaultValueHelper::IsStringValidFloat(NewDefaultValue))
+			{
+				DVSV_RETURN_MSG(TEXT("Expected a valid number for an double property"));
+			}
+		}
+	}
 	else if (PinCategory == PC_Int)
 	{
 		if (!NewDefaultValue.IsEmpty())
@@ -4872,6 +4902,12 @@ bool UEdGraphSchema_K2::DoesDefaultValueMatchAutogenerated(const UEdGraphPin& In
 			const float DefaultFloat = FCString::Atof(*InPin.DefaultValue);
 			return (AutogeneratedFloat == DefaultFloat);
 		}
+		else if (InPin.PinType.PinCategory == PC_Double)
+		{
+			const double AutogeneratedDouble = FCString::Atod(*InPin.AutogeneratedDefaultValue);
+			const double DefaultDouble = FCString::Atod(*InPin.DefaultValue);
+			return (AutogeneratedDouble == DefaultDouble);
+		}
 		else if (InPin.PinType.PinCategory == PC_Struct)
 		{
 			if (InPin.PinType.PinSubCategoryObject == VectorStruct)
@@ -5074,6 +5110,10 @@ void UEdGraphSchema_K2::SetPinAutogeneratedDefaultValueBasedOnType(UEdGraphPin*
 		// This is a slightly different format than is produced by PropertyValueToString, but changing it has backward compatibility issues
 		NewValue = TEXT("0.0");
 	}
+	else if (Pin->PinType.PinCategory == PC_Double)
+	{
+		NewValue = TEXT("0.0");
+	}
 	else if (Pin->PinType.PinCategory == PC_Boolean)
 	{
 		NewValue = TEXT("false");
@@ -5229,6 +5269,7 @@ namespace FSetVariableByNameFunctionNames
 	static const FName SetInt64Name(GET_FUNCTION_NAME_CHECKED(UKismetSystemLibrary, SetInt64PropertyByName));
 	static const FName SetByteName(GET_FUNCTION_NAME_CHECKED(UKismetSystemLibrary, SetBytePropertyByName));
 	static const FName SetFloatName(GET_FUNCTION_NAME_CHECKED(UKismetSystemLibrary, SetFloatPropertyByName));
+	static const FName SetDoubleName(GET_FUNCTION_NAME_CHECKED(UKismetSystemLibrary, SetDoublePropertyByName));
 	static const FName SetBoolName(GET_FUNCTION_NAME_CHECKED(UKismetSystemLibrary, SetBoolPropertyByName));
 	static const FName SetObjectName(GET_FUNCTION_NAME_CHECKED(UKismetSystemLibrary, SetObjectPropertyByName));
 	static const FName SetClassName(GET_FUNCTION_NAME_CHECKED(UKismetSystemLibrary, SetClassPropertyByName));
@@ -5296,6 +5337,10 @@ UFunction* UEdGraphSchema_K2::FindSetVariableByNameFunction(const FEdGraphPinTyp
 	{
 		SetFunctionName = FSetVariableByNameFunctionNames::SetFloatName;
 	}
+	else if(PinType.PinCategory == UEdGraphSchema_K2::PC_Double)
+	{
+		SetFunctionName = FSetVariableByNameFunctionNames::SetDoubleName;
+	}
 	else if(PinType.PinCategory == UEdGraphSchema_K2::PC_Boolean)
 	{
 		SetFunctionName = FSetVariableByNameFunctionNames::SetBoolName;
diff --git a/Engine/Source/Editor/BlueprintGraph/Private/K2Node_TemporaryVariable.cpp b/Engine/Source/Editor/BlueprintGraph/Private/K2Node_TemporaryVariable.cpp
--- a/Engine/Source/Editor/BlueprintGraph/Private/K2Node_TemporaryVariable.cpp
+++ b/Engine/Source/Editor/BlueprintGraph/Private/K2Node_TemporaryVariable.cpp
@@ -181,7 +181,9 @@ void UK2Node_TemporaryVariable::GetMenuActions(FBlueprintActionDatabaseRegistrar
 	ActionRegistrar.AddBlueprintAction(ActionKey, MakeTempVarNodeSpawner(FEdGraphPinType(UEdGraphSchema_K2::PC_Int, NAME_None, nullptr, EPinContainerType::None, /*bIsReference =*/ false, /*InValueTerminalType=*/FEdGraphTerminalType()), /*bIsPersistent =*/false));
 	ActionRegistrar.AddBlueprintAction(ActionKey, MakeTempVarNodeSpawner(FEdGraphPinType(UEdGraphSchema_K2::PC_Int, NAME_None, nullptr, EPinContainerType::Array, /*bIsReference =*/ false, /*InValueTerminalType=*/FEdGraphTerminalType()), /*bIsPersistent =*/false));
 	ActionRegistrar.AddBlueprintAction(ActionKey, MakeTempVarNodeSpawner(FEdGraphPinType(UEdGraphSchema_K2::PC_Float, NAME_None, nullptr, EPinContainerType::None, /*bIsReference =*/ false, /*InValueTerminalType=*/FEdGraphTerminalType()), /*bIsPersistent =*/false));
-	ActionRegistrar.AddBlueprintAction(ActionKey, MakeTempVarNodeSpawner(FEdGraphPinType(UEdGraphSchema_K2::PC_Float, NAME_None, nullptr, EPinContainerType::Array, /*bIsReference =*/ false, /*InValueTerminalType=*/FEdGraphTerminalType()), /*bIsPersistent =*/false));
+	ActionRegistrar.AddBlueprintAction(ActionKey, MakeTempVarNodeSpawner(FEdGraphPinType(UEdGraphSchema_K2::PC_Float, NAME_None, nullptr, EPinContainerType::Array, /*bIsReference =*/ false, /*InValueTerminalType=*/FEdGraphTerminalType()), /*bIsPersistent =*/false));	
+	ActionRegistrar.AddBlueprintAction(ActionKey, MakeTempVarNodeSpawner(FEdGraphPinType(UEdGraphSchema_K2::PC_Double, NAME_None, nullptr, EPinContainerType::None, /*bIsReference =*/ false, /*InValueTerminalType=*/FEdGraphTerminalType()), /*bIsPersistent =*/false));
+	ActionRegistrar.AddBlueprintAction(ActionKey, MakeTempVarNodeSpawner(FEdGraphPinType(UEdGraphSchema_K2::PC_Double, NAME_None, nullptr, EPinContainerType::Array, /*bIsReference =*/ false, /*InValueTerminalType=*/FEdGraphTerminalType()), /*bIsPersistent =*/false));
 	ActionRegistrar.AddBlueprintAction(ActionKey, MakeTempVarNodeSpawner(FEdGraphPinType(UEdGraphSchema_K2::PC_Boolean, NAME_None, nullptr, EPinContainerType::None, /*bIsReference =*/ false, /*InValueTerminalType=*/FEdGraphTerminalType()), /*bIsPersistent =*/false));
 	ActionRegistrar.AddBlueprintAction(ActionKey, MakeTempVarNodeSpawner(FEdGraphPinType(UEdGraphSchema_K2::PC_Boolean, NAME_None, nullptr, EPinContainerType::Array, /*bIsReference =*/ false, /*InValueTerminalType=*/FEdGraphTerminalType()), /*bIsPersistent =*/false));
 	ActionRegistrar.AddBlueprintAction(ActionKey, MakeTempVarNodeSpawner(FEdGraphPinType(UEdGraphSchema_K2::PC_String, NAME_None, nullptr, EPinContainerType::None, /*bIsReference =*/ false, /*InValueTerminalType=*/FEdGraphTerminalType()), /*bIsPersistent =*/false));
diff --git a/Engine/Source/Editor/GraphEditor/Private/GraphEditorSettings.cpp b/Engine/Source/Editor/GraphEditor/Private/GraphEditorSettings.cpp
--- a/Engine/Source/Editor/GraphEditor/Private/GraphEditorSettings.cpp
+++ b/Engine/Source/Editor/GraphEditor/Private/GraphEditorSettings.cpp
@@ -38,6 +38,7 @@ UGraphEditorSettings::UGraphEditorSettings( const FObjectInitializer& ObjectInit
 	IntPinTypeColor = FLinearColor(0.013575f, 0.770000f, 0.429609f, 1.0f);		// green-blue
 	Int64PinTypeColor = FLinearColor(0.413575f, 0.770000f, 0.429609f, 1.0f);	
 	FloatPinTypeColor = FLinearColor(0.357667f, 1.0f, 0.060000f, 1.0f);			// bright green
+	DoublePinTypeColor = FLinearColor(0.039216f, 0.666667f, 0.0f, 1.0f);		// darker green
 	NamePinTypeColor = FLinearColor(0.607717f, 0.224984f, 1.0f, 1.0f);			// lilac
 	DelegatePinTypeColor = FLinearColor(1.0f, 0.04f, 0.04f, 1.0f);				// bright red
 	ObjectPinTypeColor = FLinearColor(0.0f, 0.4f, 0.910000f, 1.0f);				// sharp blue
diff --git a/Engine/Source/Editor/GraphEditor/Private/NodeFactory.cpp b/Engine/Source/Editor/GraphEditor/Private/NodeFactory.cpp
--- a/Engine/Source/Editor/GraphEditor/Private/NodeFactory.cpp
+++ b/Engine/Source/Editor/GraphEditor/Private/NodeFactory.cpp
@@ -311,6 +311,10 @@ TSharedPtr<SGraphPin> FNodeFactory::CreateK2PinWidget(UEdGraphPin* InPin)
 	{
 		return SNew(SGraphPinNum<float>, InPin);
 	}
+	else if (InPin->PinType.PinCategory == UEdGraphSchema_K2::PC_Double)
+	{
+		return SNew(SGraphPinNum<double>, InPin);
+	}
 	else if (InPin->PinType.PinCategory == UEdGraphSchema_K2::PC_String || InPin->PinType.PinCategory == UEdGraphSchema_K2::PC_Name)
 	{
 		return SNew(SGraphPinString, InPin);
diff --git a/Engine/Source/Editor/GraphEditor/Public/GraphEditorSettings.h b/Engine/Source/Editor/GraphEditor/Public/GraphEditorSettings.h
--- a/Engine/Source/Editor/GraphEditor/Public/GraphEditorSettings.h
+++ b/Engine/Source/Editor/GraphEditor/Public/GraphEditorSettings.h
@@ -156,6 +156,10 @@ public:
 	UPROPERTY(EditAnywhere, config, Category=PinColors)
 	FLinearColor FloatPinTypeColor;
 
+	/** Double pin type color */
+	UPROPERTY(EditAnywhere, config, Category = PinColors)
+	FLinearColor DoublePinTypeColor;
+
 	/** Name pin type color */
 	UPROPERTY(EditAnywhere, config, Category=PinColors)
 	FLinearColor NamePinTypeColor;
diff --git a/Engine/Source/Editor/KismetCompiler/Private/KismetCompilerMisc.cpp b/Engine/Source/Editor/KismetCompiler/Private/KismetCompilerMisc.cpp
--- a/Engine/Source/Editor/KismetCompiler/Private/KismetCompilerMisc.cpp
+++ b/Engine/Source/Editor/KismetCompiler/Private/KismetCompilerMisc.cpp
@@ -118,6 +118,11 @@ static bool IsTypeCompatibleWithProperty(UEdGraphPin* SourcePin, const FEdGraphP
 		FFloatProperty* SpecificProperty = CastField<FFloatProperty>(TestProperty);
 		bTypeMismatch = (SpecificProperty == nullptr);
 	}
+	else if (PinCategory == UEdGraphSchema_K2::PC_Double)
+	{
+		FDoubleProperty* SpecificProperty = CastField<FDoubleProperty>(TestProperty);
+		bTypeMismatch = (SpecificProperty == nullptr);
+	}
 	else if (PinCategory == UEdGraphSchema_K2::PC_Int)
 	{
 		FIntProperty* SpecificProperty = CastField<FIntProperty>(TestProperty);
@@ -1030,6 +1035,11 @@ FProperty* FKismetCompilerUtilities::CreatePrimitiveProperty(FFieldVariant Prope
 		NewProperty = new FFloatProperty(PropertyScope, ValidatedPropertyName, ObjectFlags);
 		NewProperty->SetPropertyFlags(CPF_HasGetValueTypeHash);
 	}
+	else if (PinCategory == UEdGraphSchema_K2::PC_Double)
+	{
+		NewProperty = new FDoubleProperty(PropertyScope, ValidatedPropertyName, ObjectFlags);
+		NewProperty->SetPropertyFlags(CPF_HasGetValueTypeHash);
+	}
 	else if (PinCategory == UEdGraphSchema_K2::PC_Boolean)
 	{
 		FBoolProperty* BoolProperty = new FBoolProperty(PropertyScope, ValidatedPropertyName, ObjectFlags);
diff --git a/Engine/Source/Editor/KismetCompiler/Private/KismetCompilerVMBackend.cpp b/Engine/Source/Editor/KismetCompiler/Private/KismetCompilerVMBackend.cpp
--- a/Engine/Source/Editor/KismetCompiler/Private/KismetCompilerVMBackend.cpp
+++ b/Engine/Source/Editor/KismetCompiler/Private/KismetCompilerVMBackend.cpp
@@ -485,6 +485,15 @@ public:
 			return Type && (Type->PinCategory == UEdGraphSchema_K2::PC_Float);
 		}
 
+		static bool IsDouble(const FEdGraphPinType* Type, const FProperty* Property)
+		{
+			if (Property)
+			{
+				return Property->IsA<FDoubleProperty>();
+			}
+			return Type && (Type->PinCategory == UEdGraphSchema_K2::PC_Double);
+		}
+
 		static bool IsInt(const FEdGraphPinType* Type, const FProperty* Property)
 		{
 			if (Property)
@@ -708,6 +717,12 @@ public:
 				Writer << EX_FloatConst;
 				Writer << Value;
 			}
+			else if (FLiteralTypeHelper::IsDouble(&Term->Type, CoerceProperty))
+			{
+				double Value = FCString::Atod(*(Term->Name));
+				Writer << EX_DoubleConst;
+				Writer << Value;
+			}
 			else if (FLiteralTypeHelper::IsInt(&Term->Type, CoerceProperty))
 			{
 				// In certain cases (like UKismetArrayLibrary functions), we have
diff --git a/Engine/Source/Editor/PropertyEditor/Private/Presentation/PropertyTable/PropertyTableColumn.cpp b/Engine/Source/Editor/PropertyEditor/Private/Presentation/PropertyTable/PropertyTableColumn.cpp
--- a/Engine/Source/Editor/PropertyEditor/Private/Presentation/PropertyTable/PropertyTableColumn.cpp
+++ b/Engine/Source/Editor/PropertyEditor/Private/Presentation/PropertyTable/PropertyTableColumn.cpp
@@ -643,6 +643,19 @@ TSharedPtr<FCompareRowByColumnBase> FPropertyTableColumn::GetPropertySorter(FPro
 			return MakeShareable<FCompareRowByColumnBase>(new FCompareRowByColumnDescending<FFloatProperty>(SharedThis(this), FloatProperty));
 		}
 	}
+	else if (Property->IsA(FDoubleProperty::StaticClass()))
+	{
+		FDoubleProperty* DoubleProperty(CastField< FDoubleProperty >(Property));
+
+		if (SortMode == EColumnSortMode::Ascending)
+		{
+			return MakeShareable<FCompareRowByColumnBase>(new FCompareRowByColumnAscending<FDoubleProperty>(SharedThis(this), DoubleProperty));
+		}
+		else
+		{
+			return MakeShareable<FCompareRowByColumnBase>(new FCompareRowByColumnDescending<FDoubleProperty>(SharedThis(this), DoubleProperty));
+		}
+	}
 	else if (Property->IsA(FNameProperty::StaticClass()))
 	{
 		FNameProperty* NameProperty = CastField<FNameProperty>(Property);
diff --git a/Engine/Source/Editor/UnrealEd/Private/Kismet2/BlueprintEditorUtils.cpp b/Engine/Source/Editor/UnrealEd/Private/Kismet2/BlueprintEditorUtils.cpp
--- a/Engine/Source/Editor/UnrealEd/Private/Kismet2/BlueprintEditorUtils.cpp
+++ b/Engine/Source/Editor/UnrealEd/Private/Kismet2/BlueprintEditorUtils.cpp
@@ -8464,6 +8464,12 @@ bool FBlueprintEditorUtils::PropertyValueFromString_Direct(const FProperty* Prop
 			bParseSucceeded = FDefaultValueHelper::ParseFloat(StrValue, FloatValue);
 			CastFieldChecked<const FFloatProperty>(Property)->SetPropertyValue(DirectValue, FloatValue);
 		}
+		else if (Property->IsA(FDoubleProperty::StaticClass()))
+		{
+			double DoubleValue = 0.0;
+			bParseSucceeded = FDefaultValueHelper::ParseDouble(StrValue, DoubleValue);
+			CastFieldChecked<const FDoubleProperty>(Property)->SetPropertyValue(DirectValue, DoubleValue);
+		}
 		else if (const FByteProperty* ByteProperty = CastField<const FByteProperty>(Property))
 		{
 			int32 IntValue = 0;
diff --git a/Engine/Source/Programs/UnrealHeaderTool/Private/HeaderParser.cpp b/Engine/Source/Programs/UnrealHeaderTool/Private/HeaderParser.cpp
--- a/Engine/Source/Programs/UnrealHeaderTool/Private/HeaderParser.cpp
+++ b/Engine/Source/Programs/UnrealHeaderTool/Private/HeaderParser.cpp
@@ -1197,7 +1197,7 @@ namespace
 
 	bool IsPropertySupportedByBlueprint(const FProperty* Property, bool bMemberVariable)
 	{
-		if (Property == NULL)
+		if (Property == nullptr)
 		{
 			return false;
 		}
@@ -1225,6 +1225,7 @@ namespace
 			|| Property->IsA<FSoftObjectProperty>()
 			|| Property->IsA<FObjectProperty>()
 			|| Property->IsA<FFloatProperty>()
+			|| Property->IsA<FDoubleProperty>()
 			|| Property->IsA<FIntProperty>()
 			|| Property->IsA<FInt64Property>()
 			|| Property->IsA<FByteProperty>()
diff --git a/Engine/Source/Runtime/CoreUObject/Private/UObject/ScriptCore.cpp b/Engine/Source/Runtime/CoreUObject/Private/UObject/ScriptCore.cpp
--- a/Engine/Source/Runtime/CoreUObject/Private/UObject/ScriptCore.cpp
+++ b/Engine/Source/Runtime/CoreUObject/Private/UObject/ScriptCore.cpp
@@ -3070,6 +3070,12 @@ DEFINE_FUNCTION(UObject::execFloatConst)
 }
 IMPLEMENT_VM_FUNCTION( EX_FloatConst, execFloatConst );
 
+DEFINE_FUNCTION(UObject::execDoubleConst)
+{
+	*(double*)RESULT_PARAM = Stack.ReadInt<double>();
+}
+IMPLEMENT_VM_FUNCTION( EX_DoubleConst, execDoubleConst );
+
 DEFINE_FUNCTION(UObject::execStringConst)
 {
 	*(FString*)RESULT_PARAM = (ANSICHAR*)Stack.Code;
diff --git a/Engine/Source/Runtime/CoreUObject/Public/UObject/Object.h b/Engine/Source/Runtime/CoreUObject/Public/UObject/Object.h
--- a/Engine/Source/Runtime/CoreUObject/Public/UObject/Object.h
+++ b/Engine/Source/Runtime/CoreUObject/Public/UObject/Object.h
@@ -1351,6 +1351,7 @@ public:
 	DECLARE_FUNCTION(execUInt64Const);
 	DECLARE_FUNCTION(execSkipOffsetConst);
 	DECLARE_FUNCTION(execFloatConst);
+	DECLARE_FUNCTION(execDoubleConst);
 	DECLARE_FUNCTION(execStringConst);
 	DECLARE_FUNCTION(execUnicodeStringConst);
 	DECLARE_FUNCTION(execTextConst);
diff --git a/Engine/Source/Runtime/CoreUObject/Public/UObject/Script.h b/Engine/Source/Runtime/CoreUObject/Public/UObject/Script.h
--- a/Engine/Source/Runtime/CoreUObject/Public/UObject/Script.h
+++ b/Engine/Source/Runtime/CoreUObject/Public/UObject/Script.h
@@ -217,7 +217,7 @@ enum EExprToken
 	EX_UnicodeStringConst   = 0x34, // Unicode string constant.
 	EX_Int64Const			= 0x35,	// 64-bit integer constant.
 	EX_UInt64Const			= 0x36,	// 64-bit unsigned integer constant.
-	//						= 0x37,
+	EX_DoubleConst			= 0x37, // Double constant.
 	EX_PrimitiveCast		= 0x38,	// A casting operator for primitives which reads the type as the subsequent byte
 	EX_SetSet				= 0x39,
 	EX_EndSet				= 0x3A,
diff --git a/Engine/Source/Runtime/CoreUObject/Public/UObject/ScriptSerialization.h b/Engine/Source/Runtime/CoreUObject/Public/UObject/ScriptSerialization.h
--- a/Engine/Source/Runtime/CoreUObject/Public/UObject/ScriptSerialization.h
+++ b/Engine/Source/Runtime/CoreUObject/Public/UObject/ScriptSerialization.h
@@ -354,6 +354,11 @@
 			XFER(float);
 			break;
 		}
+		case EX_DoubleConst:
+		{
+			XFER(double);
+			break;
+		}
 		case EX_StringConst:
 		{
 			XFERSTRING();
diff --git a/Engine/Source/Runtime/Engine/Classes/Kismet/KismetMathLibrary.h b/Engine/Source/Runtime/Engine/Classes/Kismet/KismetMathLibrary.h
--- a/Engine/Source/Runtime/Engine/Classes/Kismet/KismetMathLibrary.h
+++ b/Engine/Source/Runtime/Engine/Classes/Kismet/KismetMathLibrary.h
@@ -3393,6 +3393,14 @@ class ENGINE_API UKismetMathLibrary : public UBlueprintFunctionLibrary
 	UFUNCTION(BlueprintPure, meta=(DisplayName = "ToInt (Int64)", CompactNodeTitle = "->", Keywords="cast convert", BlueprintAutocast), Category="Math|Conversions")
 	static int32 Conv_Int64ToInt(int64 InInt);
 
+	/** Converts a 64 bit floating point to a 32 bit floating point (if the float is too large, returns the low 32 bits) */
+	UFUNCTION(BlueprintPure, meta = (DisplayName = "ToFloat (Double)", CompactNodeTitle = "->", Keywords = "cast convert", BlueprintAutocast), Category = "Math|Conversions")
+	static float Conv_DoubleToFloat(double InDouble);
+
+	/** Converts a 32 bit floating point to a 64 bit floating point */
+	UFUNCTION(BlueprintPure, meta = (DisplayName = "ToDouble (Float)", CompactNodeTitle = "->", Keywords = "cast convert", BlueprintAutocast), Category = "Math|Conversions")
+	static double Conv_FloatToDouble(float InFloat);
+
 	/** Converts a 64 bit integer to a byte (if the integer is too large, returns the low 8 bits) */
 	UFUNCTION(BlueprintPure, meta=(DisplayName = "ToByte (Int64)", CompactNodeTitle = "->", Keywords="cast convert", BlueprintAutocast), Category="Math|Conversions")
 	static uint8 Conv_Int64ToByte(int64 InInt);
diff --git a/Engine/Source/Runtime/Engine/Classes/Kismet/KismetMathLibrary.inl b/Engine/Source/Runtime/Engine/Classes/Kismet/KismetMathLibrary.inl
--- a/Engine/Source/Runtime/Engine/Classes/Kismet/KismetMathLibrary.inl
+++ b/Engine/Source/Runtime/Engine/Classes/Kismet/KismetMathLibrary.inl
@@ -2585,6 +2585,18 @@ int32 UKismetMathLibrary::Conv_Int64ToInt(int64 InInt)
 	return (int32)InInt;
 }
 
+KISMET_MATH_FORCEINLINE
+float UKismetMathLibrary::Conv_DoubleToFloat(double InDouble)
+{
+	return (float)InDouble;
+}
+
+KISMET_MATH_FORCEINLINE
+double UKismetMathLibrary::Conv_FloatToDouble(float InFloat)
+{
+	return (double)InFloat;
+}
+
 KISMET_MATH_FORCEINLINE
 uint8 UKismetMathLibrary::Conv_Int64ToByte(int64 InInt)
 {
diff --git a/Engine/Source/Runtime/Engine/Classes/Kismet/KismetStringLibrary.h b/Engine/Source/Runtime/Engine/Classes/Kismet/KismetStringLibrary.h
--- a/Engine/Source/Runtime/Engine/Classes/Kismet/KismetStringLibrary.h
+++ b/Engine/Source/Runtime/Engine/Classes/Kismet/KismetStringLibrary.h
@@ -16,6 +16,10 @@ class ENGINE_API UKismetStringLibrary : public UBlueprintFunctionLibrary
 	UFUNCTION(BlueprintPure, meta=(DisplayName = "ToString (float)", CompactNodeTitle = "->", BlueprintAutocast), Category="Utilities|String")
 	static FString Conv_FloatToString(float InFloat);
 
+	/** Converts a double value to a string */
+	UFUNCTION(BlueprintPure, meta = (DisplayName = "ToString (double)", CompactNodeTitle = "->", BlueprintAutocast), Category = "Utilities|String")
+	static FString Conv_DoubleToString(double InDouble);
+
 	/** Converts an integer value to a string */
 	UFUNCTION(BlueprintPure, meta=(DisplayName = "ToString (integer)", CompactNodeTitle = "->", BlueprintAutocast), Category="Utilities|String")
 	static FString Conv_IntToString(int32 InInt);
diff --git a/Engine/Source/Runtime/Engine/Classes/Kismet/KismetSystemLibrary.h b/Engine/Source/Runtime/Engine/Classes/Kismet/KismetSystemLibrary.h
--- a/Engine/Source/Runtime/Engine/Classes/Kismet/KismetSystemLibrary.h
+++ b/Engine/Source/Runtime/Engine/Classes/Kismet/KismetSystemLibrary.h
@@ -317,6 +317,14 @@ class ENGINE_API UKismetSystemLibrary : public UBlueprintFunctionLibrary
 	UFUNCTION(BlueprintPure, Category="Math|Float", meta=(BlueprintThreadSafe))
 	static float MakeLiteralFloat(float Value);
 
+	/**
+	 * Creates a literal double
+	 * @param	Value	value to set the double to
+	 * @return	The literal double
+	 */
+	UFUNCTION(BlueprintPure, Category = "Math|Double", meta = (BlueprintThreadSafe))
+	static double MakeLiteralDouble(double Value);
+
 	/**
 	 * Creates a literal bool
 	 * @param	Value	value to set the bool to
@@ -754,6 +762,10 @@ class ENGINE_API UKismetSystemLibrary : public UBlueprintFunctionLibrary
 	UFUNCTION(BlueprintCallable, meta=(BlueprintInternalUseOnly = "true"))
 	static void SetFloatPropertyByName(UObject* Object, FName PropertyName, float Value);
 
+	/** Set a double property by name */
+	UFUNCTION(BlueprintCallable, meta = (BlueprintInternalUseOnly = "true"))
+	static void SetDoublePropertyByName(UObject* Object, FName PropertyName, double Value);
+
 	/** Set a bool property by name */
 	UFUNCTION(BlueprintCallable, meta=(BlueprintInternalUseOnly = "true"))
 	static void SetBoolPropertyByName(UObject* Object, FName PropertyName, bool Value);
@@ -1961,6 +1973,11 @@ FORCEINLINE float UKismetSystemLibrary::MakeLiteralFloat(float Value)
 	return Value;
 }
 
+FORCEINLINE double UKismetSystemLibrary::MakeLiteralDouble(double Value)
+{
+	return Value;
+}
+
 FORCEINLINE bool UKismetSystemLibrary::MakeLiteralBool(bool Value)
 {
 	return Value;
diff --git a/Engine/Source/Runtime/Engine/Private/KismetStringLibrary.cpp b/Engine/Source/Runtime/Engine/Private/KismetStringLibrary.cpp
--- a/Engine/Source/Runtime/Engine/Private/KismetStringLibrary.cpp
+++ b/Engine/Source/Runtime/Engine/Private/KismetStringLibrary.cpp
@@ -56,6 +56,11 @@ FString UKismetStringLibrary::Conv_FloatToString(float InFloat)
 	return FString::SanitizeFloat(InFloat);
 }
 
+FString UKismetStringLibrary::Conv_DoubleToString(double InDouble)
+{
+	return FString::SanitizeFloat(InDouble);
+}
+
 FString UKismetStringLibrary::Conv_IntToString(int32 InInt)
 {
 	return FString::Printf(TEXT("%d"), InInt);	
diff --git a/Engine/Source/Runtime/Engine/Private/KismetSystemLibrary.cpp b/Engine/Source/Runtime/Engine/Private/KismetSystemLibrary.cpp
--- a/Engine/Source/Runtime/Engine/Private/KismetSystemLibrary.cpp
+++ b/Engine/Source/Runtime/Engine/Private/KismetSystemLibrary.cpp
@@ -897,6 +897,18 @@ void UKismetSystemLibrary::SetFloatPropertyByName(UObject* Object, FName Propert
 	}
 }
 
+void UKismetSystemLibrary::SetDoublePropertyByName(UObject* Object, FName PropertyName, double Value)
+{
+	if (Object != nullptr)
+	{
+		FDoubleProperty* DoubleProp = FindFProperty<FDoubleProperty>(Object->GetClass(), PropertyName);
+		if (DoubleProp != nullptr)
+		{
+			DoubleProp->SetPropertyValue_InContainer(Object, Value);
+		}
+	}
+}
+
 void UKismetSystemLibrary::SetBoolPropertyByName(UObject* Object, FName PropertyName, bool Value)
 {
 	if(Object != NULL)
