Fix spelling mistake in comment on K2Node_Tunnel

#rnx

diff --git a/Engine/Source/Editor/BlueprintGraph/Classes/K2Node_Tunnel.h b/Engine/Source/Editor/BlueprintGraph/Classes/K2Node_Tunnel.h
--- a/Engine/Source/Editor/BlueprintGraph/Classes/K2Node_Tunnel.h
+++ b/Engine/Source/Editor/BlueprintGraph/Classes/K2Node_Tunnel.h
@@ -13,7 +13,7 @@ class BLUEPRINTGRAPH_API UK2Node_Tunnel : public UK2Node_EditablePinBase
 {
 	GENERATED_UCLASS_BODY()
 
-	// A tunnel node either has output pins that came from another tunnel's input pins, or vis versa
+	// A tunnel node either has output pins that came from another tunnel's input pins, or vice versa
 	// Note: OutputSourceNode might be equal to InputSinkNode
 	
 	// The output pins of this tunnel node came from the input pins of OutputSourceNode
