Editor only nodes cannot be copied into and don't show up in runtime blueprints anymore

diff --git a/Engine/Source/Editor/BlueprintGraph/Classes/K2Node_CallFunction.h b/Engine/Source/Editor/BlueprintGraph/Classes/K2Node_CallFunction.h
--- a/Engine/Source/Editor/BlueprintGraph/Classes/K2Node_CallFunction.h
+++ b/Engine/Source/Editor/BlueprintGraph/Classes/K2Node_CallFunction.h
@@ -197,6 +197,9 @@ public:
 	/** Used to determine the result of AllowMultipleSelfs() (without having a node instance) */
 	static bool CanFunctionSupportMultipleTargets(UFunction const* InFunction);
 
+	/** Checks if the input function can be called in the input object with respect to editor-only/runtime mismatch*/
+	static bool CanEditorOnlyFunctionBeCalled(const UFunction* InFunction, const UObject* InObject);
+
 	/** */
 	static FSlateIcon GetPaletteIconForFunction(UFunction const* Function, FLinearColor& OutColor);
 
diff --git a/Engine/Source/Editor/BlueprintGraph/Private/BlueprintActionFilter.cpp b/Engine/Source/Editor/BlueprintGraph/Private/BlueprintActionFilter.cpp
--- a/Engine/Source/Editor/BlueprintGraph/Private/BlueprintActionFilter.cpp
+++ b/Engine/Source/Editor/BlueprintGraph/Private/BlueprintActionFilter.cpp
@@ -1736,15 +1736,14 @@ static bool BlueprintActionFilterImpl::IsHiddenInNonEditorBlueprint(FBlueprintAc
 
 	if (Function)
 	{
+		 // Leaving this check here (even though its done in CanEditorOnlyFunctionBeCalled) as an early exit
 		const bool bIsEditorOnlyFunction = IsEditorOnlyObject(Function) || Function->HasAnyFunctionFlags(FUNC_EditorOnly);
 		
 		if (bIsEditorOnlyFunction)
 		{
 			for (const UBlueprint* Blueprint : Filter.Context.Blueprints)
 			{
-				const UClass* BlueprintClass = Blueprint->ParentClass;
-				const bool bIsEditorBlueprintClass = (BlueprintClass != nullptr) && IsEditorOnlyObject(BlueprintClass);
-				bVisible &= bIsEditorBlueprintClass;
+				bVisible &= UK2Node_CallFunction::CanEditorOnlyFunctionBeCalled(Function, Blueprint->ParentClass);
 			}
 		}
 	}
diff --git a/Engine/Source/Editor/BlueprintGraph/Private/K2Node_CallFunction.cpp b/Engine/Source/Editor/BlueprintGraph/Private/K2Node_CallFunction.cpp
--- a/Engine/Source/Editor/BlueprintGraph/Private/K2Node_CallFunction.cpp
+++ b/Engine/Source/Editor/BlueprintGraph/Private/K2Node_CallFunction.cpp
@@ -1346,11 +1346,30 @@ bool UK2Node_CallFunction::CanFunctionSupportMultipleTargets(UFunction const* Fu
 	return !bHasReturnParam && bIsImpure && !bIsLatent;
 }
 
+bool UK2Node_CallFunction::CanEditorOnlyFunctionBeCalled(const UFunction* InFunction, const UObject* InObject)
+{
+	if (InFunction && InObject &&
+		(IsEditorOnlyObject(InFunction) || InFunction->HasAnyFunctionFlags(FUNC_EditorOnly)))
+	{
+		return IsEditorOnlyObject(InObject);
+	}
+
+	return true;
+}
+
 bool UK2Node_CallFunction::CanPasteHere(const UEdGraph* TargetGraph) const
 {
 	// Basic check for graph compatibility, etc.
 	bool bCanPaste = Super::CanPasteHere(TargetGraph);
 
+	// Cannot paste editor only functions into runtime graphs
+	if (bCanPaste)
+	{
+		UFunction* TargetFunction = GetTargetFunction();
+
+		bCanPaste = CanEditorOnlyFunctionBeCalled(TargetFunction, TargetGraph);
+	}
+
 	// We check function context for placability only in the base class case; derived classes are typically bound to
 	// specific functions that should always be placeable, but may not always be explicitly callable (e.g. InternalUseOnly).
 	if(bCanPaste && GetClass() == StaticClass())
@@ -2051,17 +2070,13 @@ void UK2Node_CallFunction::ValidateNodeDuringCompilation(class FCompilerResultsL
 	
 	// Ensure that editor module BP exposed UFunctions can only be called in blueprints for which the base class is also part of an editor module
 	// Also check for functions wrapped in WITH_EDITOR 
-	if (Function && Blueprint &&
-		(IsEditorOnlyObject(Function) || Function->HasAnyFunctionFlags(FUNC_EditorOnly)))
-	{	
-		const UClass* BlueprintClass = Blueprint->ParentClass;
-		bool bIsEditorOnlyBlueprintBaseClass = !BlueprintClass || IsEditorOnlyObject(BlueprintClass);
-		if (!bIsEditorOnlyBlueprintBaseClass)
-		{
-			FString const FunctName = Function->GetName();
-			FText const WarningFormat = LOCTEXT("EditorFunctionFmt", "Cannot use the editor function \"{0}\" in this runtime Blueprint. Only for use in Editor Utility Blueprints and Blutilities.");
-			MessageLog.Error(*FText::Format(WarningFormat, FText::FromString(FunctName)).ToString(), this);
-		}
+	bool bIsEditorOnlyBlueprintBaseClass = CanEditorOnlyFunctionBeCalled(Function, Blueprint);
+
+	if (!bIsEditorOnlyBlueprintBaseClass)
+	{
+		FString const FunctName = Function->GetName();
+		FText const WarningFormat = LOCTEXT("EditorFunctionFmt", "Cannot use the editor function \"{0}\" in this runtime Blueprint. Only for use in Editor Utility Blueprints and Blutilities.");
+		MessageLog.Error(*FText::Format(WarningFormat, FText::FromString(FunctName)).ToString(), this);
 	}
 
 	if (Function)
