Use BP base class when testing if an editor only function can be called

diff --git a/Engine/Source/Editor/BlueprintGraph/Private/K2Node_CallFunction.cpp b/Engine/Source/Editor/BlueprintGraph/Private/K2Node_CallFunction.cpp
--- a/Engine/Source/Editor/BlueprintGraph/Private/K2Node_CallFunction.cpp
+++ b/Engine/Source/Editor/BlueprintGraph/Private/K2Node_CallFunction.cpp
@@ -2072,7 +2072,7 @@ void UK2Node_CallFunction::ValidateNodeDuringCompilation(class FCompilerResultsL
 	{
 		// Ensure that editor module BP exposed UFunctions can only be called in blueprints for which the base class is also part of an editor module
 	    // Also check for functions wrapped in WITH_EDITOR 
-		bool bIsEditorOnlyBlueprintBaseClass = CanEditorOnlyFunctionBeCalled(Function, Blueprint);
+		bool bIsEditorOnlyBlueprintBaseClass = CanEditorOnlyFunctionBeCalled(Function, Blueprint->ParentClass);
 
 		if (!bIsEditorOnlyBlueprintBaseClass)
 		{
