Blueprints: Type promotion operator table and comparison options. Unit tests for primitive types. Type promotion is disabled by default via the BP.TypePromo.IsEnabled config variable. This allows for two pin types to be compared and determine which is the lower type, so that we can automatically set the type to the higher one when necessary. This will specifically benefit Operators to start (addition, subtraction, greater than, etc).

#rnx

diff --git a/Engine/Source/Editor/BlueprintGraph/Classes/BlueprintTypePromotion.h b/Engine/Source/Editor/BlueprintGraph/Classes/BlueprintTypePromotion.h
new file mode 100644
--- /dev/null
+++ b/Engine/Source/Editor/BlueprintGraph/Classes/BlueprintTypePromotion.h
@@ -0,0 +1,176 @@
+// Copyright Epic Games, Inc. All Rights Reserved.
+
+#pragma once
+
+#include "CoreMinimal.h"
+#include "EdGraphSchema_K2.h"
+#include "BlueprintActionDatabase.h"
+#include "BlueprintFunctionNodeSpawner.h"
+#include "HAL/IConsoleManager.h"
+
+/** Contains behavior needed to handle type promotion in blueprints */
+class FTypePromotion : private FNoncopyable
+{
+public:
+
+	static FTypePromotion& Get();
+
+	/** 
+	* If given more than two input pins, this will just use the highest type
+	* and return the function for that operation. It's update to nodes like the 
+	* UK2Node_CommutativeAssociativeBinaryOperator to patch the input pins
+	* together on expansion
+	* 
+	* @return	A function that matches a promoted type. Nullptr if no match is found
+	*/
+	static UFunction* GetOperatorFunction(const FString& Operation, const TArray<UEdGraphPin*>& WildcardPins);
+
+	static UFunction* GetOperatorFunction(const FString& Operation, const TArray<UEdGraphPin*>& InputPins, const UEdGraphPin* OutputPin);
+
+	/** 
+	* Find the function that has this input and the lowest matching other input.
+	* Ex: Given "Add" and "Vector" this function would return the "Add_VectorFloat" function
+	*/
+	static UFunction* FindLowestMatchingFunc(const FString& Operation, const FEdGraphPinType& InputType, TArray<UFunction*>& OutPossibleFunctions);
+
+	/** Returns all functions for a specific operation. Will empty the given array and populate it with UFunction pointers */
+	static void GetAllFuncsForOp(const FString& Operation, TArray<UFunction*>& OutFuncs);
+
+	/** Get an array of the supported operator names for type promo. Ex: "Add", "Subtract", "Multiply" */
+	static const TArray<FString>& GetOpNames();
+
+	static bool GetOpNameFromFunction(UFunction const* const Func, FString& OutName);
+
+	/** Returns true if the given function is a candidate to handle type promotion */
+	static bool IsPromotableFunction(const UFunction* Function);
+
+	/**
+	* Determine what type a given set of wildcard pins would result in
+	*
+	* @return	Pin type that is the "highest" of all the given pins
+	*/
+	static FEdGraphPinType GetPromotedType(const TArray<UEdGraphPin*>& WildcardPins);
+
+	/**
+	* Attempts to promote type A to type B. Will only work if TypeB is higher than type B
+	*
+	* @return	True if the promotion was successful
+	*/
+	static bool PromotePin(FEdGraphPinType& InTypeA, const FEdGraphPinType& TypeB);
+
+	static bool IsFunctionPromotionReady(const UFunction* const FuncToConsider);
+
+	/** Represents the possible results when comparing two types for promotion */
+	enum class ETypeComparisonResult : uint8
+	{
+		TypeAHigher,
+		TypeBHigher,
+		TypesEqual,
+		InvalidComparison
+	};
+
+	/**
+	* 
+	* 
+	* 
+	* @return 
+	*/
+	static ETypeComparisonResult GetHigherType(const FEdGraphPinType& A, const FEdGraphPinType& B);
+
+	/** Returns true if A can be promoted to type B correctly, or if the types are equal */
+	static bool IsValidPromotion(const FEdGraphPinType& A, const FEdGraphPinType& B);
+
+	static bool HasStructConversion(const UEdGraphPin* InputPin, const UEdGraphPin* OutputPin);
+
+	static bool IsOperatorSpawnerRegistered(UFunction const* const Func);
+
+	static UBlueprintFunctionNodeSpawner* GetOperatorSpawner(const FString& OpName);
+
+	/** keep track of the operator that this function provides so that we dont add multiple 
+	to the BP context menu */
+	static void RegisterOperatorSpawner(const FString& OpName, UBlueprintFunctionNodeSpawner* Spawner);
+
+	static void ClearNodeSpawners();
+
+private:
+
+	explicit FTypePromotion();
+	~FTypePromotion();
+
+	/** Callback that will rebuild the op table when hot reload is triggered */
+	static void OnModulesChanged(FName ModuleThatChanged, EModuleChangeReason ReasonForChange);
+
+	UFunction* GetOperatorFunction_Internal(const FString& Operation, const TArray<UEdGraphPin*>& WildcardPins) const;
+
+	UFunction* GetOperatorFunction_Internal(const FString& Operation, const TArray<UEdGraphPin*>& InputPins, const UEdGraphPin* OutputPin) const;
+
+	bool IsFunctionPromotionReady_Internal(const UFunction* const FuncToConsider) const;
+
+	FEdGraphPinType GetPromotedType_Internal(const TArray<UEdGraphPin*>& WildcardPins) const;
+
+	UFunction* FindLowestMatchingFunc_Internal(const FString& Operation, const FEdGraphPinType& InputType, TArray<UFunction*>& OutPossibleFunctions);
+
+	void GetAllFuncsForOp_Internal(const FString& Operation, TArray<UFunction*>& OutFuncs);
+
+	bool PromotePin_Internal(FEdGraphPinType& InTypeA, const FEdGraphPinType& TypeB);
+
+	/**
+	* Determines which pin type is "higher" 
+	* 
+	* @see FTypePromotion::PromotionTable
+	*/
+	ETypeComparisonResult GetHigherType_Internal(const FEdGraphPinType& A, const FEdGraphPinType& B) const;
+
+	/** Creates a lookup table of types and operations to their appropriate UFunction */
+	void CreateOpTable();
+
+	/** Creates the table of what types can be promoted to others */
+	void CreatePromotionTable();
+
+	void AddOpFunction(const FString& OpName, UFunction* Function);
+
+	static FTypePromotion* Instance;
+
+	/** Delegate that gets called when a module is reloaded so we can rebuild the op table */
+	FDelegateHandle OnModulesChangedDelegateHandle;
+
+	/** A map of 'Type' to its 'available promotions'. See ctor for creation */
+	TMap<FName, TArray<FName>> PromotionTable;
+
+	/**
+	 * A single operator can have multiple functions associated with it; usually
+	 * for handling different types (int*int, vs. int*vector), hence this array.
+	 * This is the same implementation style as the Math Expression node.
+	 */
+	typedef TArray<UFunction*> FFunctionsList;
+
+	/**
+	 * A lookup table, mapping operator strings (like "Add", "Multiply", etc.) to a list
+	 * of associated functions.
+	 */
+	TMap<FString, FFunctionsList> OperatorTable;
+
+	/** Map of operators to their node spawner so that we can clean up the context menu */
+	TMap<FString, UBlueprintFunctionNodeSpawner*> OperatorNodeSpawnerMap;
+};
+
+namespace TypePromoDebug
+{
+	/** Enables/Disables type promotion in BP */
+	static bool bIsTypePromoEnabled = false;
+	static FAutoConsoleVariableRef CVarIsTypePromoEnabled(
+		TEXT("BP.TypePromo.IsEnabled"), bIsTypePromoEnabled,
+		TEXT("If true then type promotion inside of blueprints will be enabled"),
+		FConsoleVariableDelegate::CreateLambda([](IConsoleVariable* InVariable)
+			{
+				// Clear the node spawner so that we create the new BP actions correctly
+				FTypePromotion::ClearNodeSpawners();
+
+				// Refresh all the actions so that the context menu goes back to the normal options
+				if (FBlueprintActionDatabase* Actions = FBlueprintActionDatabase::TryGet())
+				{
+					Actions->RefreshAll();
+				}
+			}),
+		ECVF_Default);
+}
\ No newline at end of file
diff --git a/Engine/Source/Editor/BlueprintGraph/Private/BlueprintTypePromotion.cpp b/Engine/Source/Editor/BlueprintGraph/Private/BlueprintTypePromotion.cpp
new file mode 100644
--- /dev/null
+++ b/Engine/Source/Editor/BlueprintGraph/Private/BlueprintTypePromotion.cpp
@@ -0,0 +1,560 @@
+// Copyright Epic Games, Inc. All Rights Reserved.
+
+#include "BlueprintTypePromotion.h"
+#include "Modules/ModuleManager.h"
+#include "Kismet/BlueprintFunctionLibrary.h"
+#include "UObject/UObjectHash.h"
+
+FTypePromotion* FTypePromotion::Instance = nullptr;
+
+namespace OperatorNames
+{
+	static const FString NoOp		= TEXT("NO_OP");
+
+	static const FString Add		= TEXT("Add");
+	static const FString Multiply	= TEXT("Multiply");
+	static const FString Subtract	= TEXT("Subtract");
+	static const FString Divide		= TEXT("Divide");
+
+	static const FString Greater	= TEXT("Greater");
+	static const FString GreaterEq	= TEXT("GreaterEqual");
+	static const FString Less		= TEXT("Less");
+	static const FString LessEq		= TEXT("LessEqual");
+	static const FString NotEq		= TEXT("NotEqual");
+}
+
+FTypePromotion& FTypePromotion::Get()
+{
+	if (Instance == nullptr)
+	{
+		Instance = new FTypePromotion();
+	}
+	return *Instance;
+}
+
+FTypePromotion::FTypePromotion()
+{
+	CreatePromotionTable();
+	CreateOpTable();
+
+	OnModulesChangedDelegateHandle = FModuleManager::Get().OnModulesChanged().AddStatic(&OnModulesChanged);
+}
+
+FTypePromotion::~FTypePromotion()
+{
+	if(Instance)
+	{
+		delete Instance;
+		Instance = nullptr;
+	}
+
+	FModuleManager::Get().OnModulesChanged().Remove(OnModulesChangedDelegateHandle);
+}
+
+void FTypePromotion::OnModulesChanged(FName ModuleThatChanged, EModuleChangeReason ReasonForChange)
+{
+	if (Instance)
+	{
+		Instance->CreateOpTable();
+	}
+}
+
+void FTypePromotion::CreatePromotionTable()
+{
+	PromotionTable =
+	{
+		// Type_X...						Can be promoted to...
+		{ UEdGraphSchema_K2::PC_Int,		{ UEdGraphSchema_K2::PC_Float, UEdGraphSchema_K2::PC_Double, UEdGraphSchema_K2::PC_Int64 } },
+		{ UEdGraphSchema_K2::PC_Byte,		{ UEdGraphSchema_K2::PC_Int, UEdGraphSchema_K2::PC_Int64 } },
+		{ UEdGraphSchema_K2::PC_Float,		{ UEdGraphSchema_K2::PC_Double, UEdGraphSchema_K2::PC_Int64 } },
+		{ UEdGraphSchema_K2::PC_Double,		{ UEdGraphSchema_K2::PC_Int64 } },
+		{ UEdGraphSchema_K2::PC_Wildcard,	{ UEdGraphSchema_K2::PC_Int, UEdGraphSchema_K2::PC_Int64, UEdGraphSchema_K2::PC_Float, UEdGraphSchema_K2::PC_Double, UEdGraphSchema_K2::PC_Byte, UEdGraphSchema_K2::PC_Boolean } },
+	};
+}
+
+bool FTypePromotion::IsValidPromotion(const FEdGraphPinType& A, const FEdGraphPinType& B)
+{
+	// If either of these pin types is a struct, than we have to have some kind of valid
+	// conversion function, otherwise we can't possibly connect them
+	const ETypeComparisonResult Res = FTypePromotion::GetHigherType(A, B);
+
+	return Res == ETypeComparisonResult::TypeBHigher;
+}
+
+bool FTypePromotion::HasStructConversion(const UEdGraphPin* InputPin, const UEdGraphPin* OutputPin)
+{
+	const UEdGraphSchema_K2* K2Schema = GetDefault<UEdGraphSchema_K2>();
+
+	FName DummyName;
+	UClass* DummyClass = nullptr;
+	UK2Node* DummyNode = nullptr;
+
+	const bool bCanAutocast = K2Schema->SearchForAutocastFunction(OutputPin, InputPin, /*out*/ DummyName, DummyClass);
+	const bool bCanAutoConvert = K2Schema->FindSpecializedConversionNode(OutputPin, InputPin, false, /* out */ DummyNode);
+	
+	return bCanAutocast || bCanAutoConvert;
+}
+
+FTypePromotion::ETypeComparisonResult FTypePromotion::GetHigherType(const FEdGraphPinType& A, const FEdGraphPinType& B)
+{
+	return FTypePromotion::Get().GetHigherType_Internal(A, B);
+}
+
+FTypePromotion::ETypeComparisonResult FTypePromotion::GetHigherType_Internal(const FEdGraphPinType& A, const FEdGraphPinType& B) const
+{
+	if(A == B)
+	{
+		return ETypeComparisonResult::TypesEqual;
+	}
+	// Is this A promotable type?					  Can type A be promoted to type B?
+	else if(PromotionTable.Contains(A.PinCategory) && PromotionTable[A.PinCategory].Contains(B.PinCategory))
+	{
+		return ETypeComparisonResult::TypeBHigher;
+	}
+	// Can B get promoted to A?
+	else if(PromotionTable.Contains(B.PinCategory) && PromotionTable[B.PinCategory].Contains(A.PinCategory))
+	{
+		return ETypeComparisonResult::TypeAHigher;
+	}
+	// Handle "No" Pin type, the default value of FEdGraphPinType
+	else if(A.PinCategory == NAME_None && B.PinCategory != NAME_None)
+	{
+		return ETypeComparisonResult::TypeBHigher;
+	}
+	else if(B.PinCategory == NAME_None && A.PinCategory != NAME_None)
+	{
+		return ETypeComparisonResult::TypeAHigher;
+	}
+	// A is a struct and B is not a struct
+	else if(A.PinCategory == UEdGraphSchema_K2::PC_Struct && B.PinCategory != UEdGraphSchema_K2::PC_Struct)
+	{
+		return ETypeComparisonResult::TypeAHigher;
+	}
+	// A is not a struct and B is a struct
+	else if (A.PinCategory != UEdGraphSchema_K2::PC_Struct && B.PinCategory == UEdGraphSchema_K2::PC_Struct)
+	{
+		return ETypeComparisonResult::TypeBHigher;
+	}
+
+	// We couldn't find any possible promotions, so this is an invalid comparison
+	return ETypeComparisonResult::InvalidComparison;
+}
+
+bool FTypePromotion::IsFunctionPromotionReady(const UFunction* const FuncToConsider)
+{
+	return FTypePromotion::Get().IsFunctionPromotionReady_Internal(FuncToConsider);
+}
+
+bool FTypePromotion::IsFunctionPromotionReady_Internal(const UFunction* const FuncToConsider) const
+{
+	// This could be better served as just keeping all known UFunctions in a TArray or TSet, and
+	// using that index in the Map so we don't need to iterate table pairs and arrays inside of them
+	for (const TPair<FString, FFunctionsList>& Pair : OperatorTable)
+	{
+		for(const UFunction* const Func : Pair.Value)
+		{
+			if(Func == FuncToConsider)
+			{
+				return true;
+			}
+		}
+	}
+
+	return false;
+}
+
+FEdGraphPinType FTypePromotion::GetPromotedType(const TArray<UEdGraphPin*>& WildcardPins)
+{
+	return FTypePromotion::Get().GetPromotedType_Internal(WildcardPins);
+}
+
+FEdGraphPinType FTypePromotion::GetPromotedType_Internal(const TArray<UEdGraphPin*>& WildcardPins) const
+{
+	// There must be some wildcard pins in order to get the promoted type
+	TRACE_CPUPROFILER_EVENT_SCOPE(FTypePromotionTable::GetPromotedType);
+	
+	FEdGraphPinType HighestPinType = FEdGraphPinType();
+
+	for (const UEdGraphPin* CurPin : WildcardPins)
+	{
+		if(CurPin)
+		{
+			ETypeComparisonResult Res = GetHigherType(/* A */ HighestPinType, /* B */ CurPin->PinType);
+
+			// If this pin is a different type and "higher" than set our out pin type to that
+			switch(Res)
+			{
+				case ETypeComparisonResult::TypeBHigher:
+					HighestPinType = CurPin->PinType;
+				break;
+			}
+		}
+	}
+	return HighestPinType;
+}
+
+bool FTypePromotion::PromotePin(FEdGraphPinType& InTypeA, const FEdGraphPinType& TypeB)
+{
+	return FTypePromotion::Get().PromotePin_Internal(InTypeA, TypeB);
+}
+
+bool FTypePromotion::PromotePin_Internal(FEdGraphPinType& InTypeA, const FEdGraphPinType& TypeB)
+{
+	// If type B is not the higher type, than we shouldn't do anything
+	if(GetHigherType(InTypeA, TypeB) != ETypeComparisonResult::TypeBHigher)
+	{
+		return false;
+	}
+
+	InTypeA = TypeB;
+
+	return true;
+}
+
+UFunction* FTypePromotion::GetOperatorFunction(const FString& Operation, const TArray<UEdGraphPin*>& WildcardPins)
+{
+	return FTypePromotion::Get().GetOperatorFunction_Internal(Operation, WildcardPins);
+}
+
+UFunction* FTypePromotion::GetOperatorFunction(const FString& Operation, const TArray<UEdGraphPin*>& InputPins, const UEdGraphPin* OutputPin)
+{
+	return FTypePromotion::Get().GetOperatorFunction_Internal(Operation, InputPins, OutputPin);
+}
+
+UFunction* FTypePromotion::GetOperatorFunction_Internal(const FString& Operation, const TArray<UEdGraphPin*>& InputPins, const UEdGraphPin* OutputPin) const
+{
+	check(InputPins.Num() > 0 && OutputPin);
+
+	TRACE_CPUPROFILER_EVENT_SCOPE(FTypePromotionTable::GetOperatorFunction_InOut);
+	
+	UFunction* MatchingFunction = nullptr;
+	if (const FFunctionsList* FuncList = OperatorTable.Find(Operation))
+	{
+		const UEdGraphSchema_K2* Schema = GetDefault<UEdGraphSchema_K2>();
+		for (UFunction* Func : *FuncList)
+		{
+			// Do the return types match? 
+			FEdGraphPinType FuncReturnType;
+			const bool bReturnTypesMatch = Schema->ConvertPropertyToPinType(Func->GetReturnProperty(), /* out */ FuncReturnType) && FuncReturnType.PinCategory == OutputPin->PinType.PinCategory;
+			if (!bReturnTypesMatch)
+			{
+				continue;
+			}
+
+			// Do the input types match up correctly? 
+			int32 ArgumentIndex = 0;
+			int32 NumMatchingInputs = 0;
+			for (TFieldIterator<FProperty> PropIt(Func); PropIt && (PropIt->PropertyFlags & CPF_Parm); ++PropIt)
+			{
+				FProperty* Param = *PropIt;
+				if (!Param->HasAnyPropertyFlags(CPF_ReturnParm))
+				{
+					if (ArgumentIndex < InputPins.Num())
+					{
+						FEdGraphPinType ParamType;
+						if (Schema->ConvertPropertyToPinType(Param, /* out */ ParamType))
+						{
+							FEdGraphPinType const& TypeToMatch = InputPins[ArgumentIndex]->PinType;
+							if (Schema->ArePinTypesCompatible(TypeToMatch, ParamType) && GetHigherType(TypeToMatch, ParamType) != ETypeComparisonResult::InvalidComparison)
+							{
+								NumMatchingInputs++;
+							}
+						}
+						else
+						{
+							break;
+						}
+					}
+					++ArgumentIndex;
+				}
+			}
+
+			if(NumMatchingInputs >= InputPins.Num())
+			{
+				MatchingFunction = Func;
+				break;
+			}
+		}
+	}
+
+	return MatchingFunction;
+}
+
+UFunction* FTypePromotion::GetOperatorFunction_Internal(const FString& Operation, const TArray<UEdGraphPin*>& WildcardPins) const
+{
+	TRACE_CPUPROFILER_EVENT_SCOPE(FTypePromotionTable::GetOperatorFunction);
+
+	UFunction* MatchingFunction = nullptr;
+	// Do these wild cards even need to be promoted? (Can we find a function that matches their
+	// inputs in the op table already?)
+	if(const FFunctionsList* FuncList = OperatorTable.Find(Operation))
+	{
+		const UEdGraphSchema_K2* Schema = GetDefault<UEdGraphSchema_K2>();
+
+		FEdGraphPinType PromotedReturnType = GetPromotedType(WildcardPins);
+
+		for(UFunction* Func : *FuncList)
+		{
+			// If the return type of this function matches...
+			FEdGraphPinType FuncReturnType;
+			if(Schema->ConvertPropertyToPinType(Func->GetReturnProperty(), /* out */ FuncReturnType) && FuncReturnType.PinCategory == PromotedReturnType.PinCategory)
+			{
+				int32 ArgumentIndex = 0;
+				for (TFieldIterator<FProperty> PropIt(Func); PropIt && (PropIt->PropertyFlags & CPF_Parm); ++PropIt)
+				{
+					FProperty* Param = *PropIt;
+					if (!Param->HasAnyPropertyFlags(CPF_ReturnParm))
+					{
+						if(ArgumentIndex < WildcardPins.Num())
+						{
+							FEdGraphPinType ParamType;
+							if (Schema->ConvertPropertyToPinType(Param, /* out */ ParamType))
+							{
+								FEdGraphPinType const& TypeToMatch = WildcardPins[ArgumentIndex]->PinType;
+								if (!Schema->ArePinTypesCompatible(TypeToMatch, ParamType) && GetHigherType(TypeToMatch, ParamType) == ETypeComparisonResult::InvalidComparison)
+								{
+									break; // type mismatch
+								}
+							}
+							else
+							{
+								break;
+							}
+						}
+						++ArgumentIndex;
+					}
+				}
+
+				// This doesn't work if given more than 2 wildcard inputs
+				if (ArgumentIndex == WildcardPins.Num())
+				{
+					MatchingFunction = Func;
+					break;
+				}
+			}
+		}
+	}
+
+	return MatchingFunction;
+}
+
+UFunction* FTypePromotion::FindLowestMatchingFunc(const FString& Operation, const FEdGraphPinType& InputType, TArray<UFunction*>& OutPossibleFunctions)
+{
+	return FTypePromotion::Get().FindLowestMatchingFunc_Internal(Operation, InputType, OutPossibleFunctions);
+}
+
+static bool PropertyCompatibleWithPin(const FProperty* Param, FEdGraphPinType const& TypeToMatch)
+{
+	const UEdGraphSchema_K2* Schema = GetDefault<UEdGraphSchema_K2>();
+
+	FEdGraphPinType ParamType;
+	if (Schema->ConvertPropertyToPinType(Param, /* out */ ParamType))
+	{
+		if (Schema->ArePinTypesCompatible(TypeToMatch, ParamType) &&
+			FTypePromotion::GetHigherType(TypeToMatch, ParamType) != FTypePromotion::ETypeComparisonResult::InvalidComparison)
+		{
+			return true;
+		}
+	}
+	return false;
+}
+
+UFunction* FTypePromotion::FindLowestMatchingFunc_Internal(const FString& Operation, const FEdGraphPinType& InputType, TArray<UFunction*>& OutPossibleFunctions)
+{
+	TRACE_CPUPROFILER_EVENT_SCOPE(FTypePromotionTable::FindLowestMatchingFunc_Internal);
+
+	UFunction* LowestFunc = nullptr;
+	OutPossibleFunctions.Empty();
+
+	if (const FFunctionsList* FuncList = OperatorTable.Find(Operation))
+	{
+		const UEdGraphSchema_K2* Schema = GetDefault<UEdGraphSchema_K2>();
+
+		for (UFunction* Func : *FuncList)
+		{
+			// Find the function that matches the input type here, we don't care about the output type
+			for (TFieldIterator<FProperty> PropIt(Func); PropIt && (PropIt->PropertyFlags & CPF_Parm); ++PropIt)
+			{
+				// Ignore return params here, we only care about inputs
+				FProperty* Param = *PropIt;
+				if (!Param->HasAnyPropertyFlags(CPF_ReturnParm) && PropertyCompatibleWithPin(Param, InputType))
+				{
+					// If an input of this function is compatible with this type
+					// Then we need to check all the other inputs on this function
+					OutPossibleFunctions.Emplace(Func);
+					break;
+				}
+			}
+		}
+	}
+
+	const UEdGraphSchema_K2* Schema = GetDefault<UEdGraphSchema_K2>();
+	FEdGraphPinType CurLowestType = {};
+
+	if(OutPossibleFunctions.Num() > 0)
+	{
+		LowestFunc = OutPossibleFunctions[0];
+	}
+
+	// Check possible matches to see who has the lowest type
+	for(UFunction* Func : OutPossibleFunctions)
+	{
+		for (TFieldIterator<FProperty> PropIt(Func); PropIt && (PropIt->PropertyFlags & CPF_Parm); ++PropIt)
+		{
+			// Check for the other input param that is not the input type we are lookin at
+			FProperty* Param = *PropIt;
+			FEdGraphPinType ParamType;
+			if (!Param->HasAnyPropertyFlags(CPF_ReturnParm) && Schema->ConvertPropertyToPinType(Param, /* out */ ParamType))
+			{
+				ETypeComparisonResult Res = FTypePromotion::GetHigherType(ParamType, CurLowestType);
+				if(Res == ETypeComparisonResult::TypeBHigher)
+				{
+					CurLowestType = ParamType;
+					LowestFunc = Func;
+				}
+				else if(Res == ETypeComparisonResult::TypeAHigher)
+				{
+					//CurLowestType = CurLowestType;
+					LowestFunc = Func;
+				}
+			}
+		}
+	}
+
+	return LowestFunc;
+}
+
+void FTypePromotion::GetAllFuncsForOp(const FString& Operation, TArray<UFunction*>& OutFuncs)
+{
+	return FTypePromotion::Get().GetAllFuncsForOp_Internal(Operation, OutFuncs);
+}
+
+const TArray<FString>& FTypePromotion::GetOpNames()
+{
+	static const TArray<FString> OpsArray =
+	{
+		OperatorNames::Add,
+		OperatorNames::Multiply,
+		OperatorNames::Subtract,
+		OperatorNames::Divide,
+		OperatorNames::Greater,
+		OperatorNames::GreaterEq,
+		OperatorNames::Less,
+		OperatorNames::LessEq,
+		OperatorNames::NotEq
+	};
+
+	return OpsArray;
+}
+
+void FTypePromotion::GetAllFuncsForOp_Internal(const FString& Operation, TArray<UFunction*>& OutFuncs)
+{
+	OutFuncs.Empty();
+
+	OutFuncs.Append(OperatorTable[Operation]);
+}
+
+bool FTypePromotion::GetOpNameFromFunction(UFunction const* const Func, FString& OutName)
+{
+	if(!Func)
+	{
+		return false;
+	}
+
+	const FString& FuncName = Func->GetName();
+	// Get everything before the "_"
+	int32 Index = FuncName.Find(TEXT("_"));
+	FString FuncNameChopped = FuncName.Mid(0, Index);
+
+	for(const FString& OpName : GetOpNames())
+	{
+		if(FuncNameChopped == OpName)
+		{
+			OutName = OpName;
+			return true;
+		}
+	}
+
+	OutName = OperatorNames::NoOp;
+	return false;
+}
+
+void FTypePromotion::CreateOpTable()
+{
+	TRACE_CPUPROFILER_EVENT_SCOPE(FTypePromotion::CreateOpTable);
+	const UEdGraphSchema_K2* Schema = GetDefault<UEdGraphSchema_K2>();
+
+	OperatorTable.Empty();
+
+	TArray<UClass*> Libraries;
+	GetDerivedClasses(UBlueprintFunctionLibrary::StaticClass(), Libraries);
+	for (UClass* Library : Libraries)
+	{
+		// Ignore abstract libraries/classes
+		if (!Library || Library->HasAnyClassFlags(CLASS_Abstract))
+		{
+			continue;
+		}
+
+		for (UFunction* Function : TFieldRange<UFunction>(Library, EFieldIteratorFlags::ExcludeSuper, EFieldIteratorFlags::ExcludeDeprecated))
+		{
+			if(!IsPromotableFunction(Function))
+			{
+				continue;
+			}
+
+			FEdGraphPinType FuncPinType;
+			FString OpName;
+
+			if (Schema->ConvertPropertyToPinType(Function->GetReturnProperty(), /* out */ FuncPinType) &&
+				GetOpNameFromFunction(Function, /* out */ OpName)
+				)
+			{
+				AddOpFunction(OpName, Function);
+			}
+		}
+	}
+}
+
+void FTypePromotion::AddOpFunction(const FString& OpName, UFunction* Function)
+{
+	OperatorTable.FindOrAdd(OpName).Add(Function);
+}
+
+bool FTypePromotion::IsPromotableFunction(const UFunction* Function)
+{
+	return Function && Function->HasAnyFunctionFlags(FUNC_BlueprintPure) && Function->GetReturnProperty();
+}
+
+bool FTypePromotion::IsOperatorSpawnerRegistered(UFunction const* const Func)
+{
+	FString OpName;
+	FTypePromotion::GetOpNameFromFunction(Func, OpName);
+	return FTypePromotion::GetOperatorSpawner(OpName) != nullptr;
+}
+
+void FTypePromotion::RegisterOperatorSpawner(const FString& OpName, UBlueprintFunctionNodeSpawner* Spawner)
+{
+	if(Instance && !Instance->OperatorNodeSpawnerMap.Contains(OpName) && OpName != OperatorNames::NoOp)
+	{
+		Instance->OperatorNodeSpawnerMap.Add(OpName, Spawner);
+	}
+}
+
+UBlueprintFunctionNodeSpawner* FTypePromotion::GetOperatorSpawner(const FString& OpName)
+{
+	if(Instance && Instance->OperatorNodeSpawnerMap.Contains(OpName))
+	{
+		return Instance->OperatorNodeSpawnerMap[OpName];
+	}
+
+	return nullptr;
+}
+
+void FTypePromotion::ClearNodeSpawners()
+{
+	if (Instance)
+	{
+		Instance->OperatorNodeSpawnerMap.Empty();
+	}
+}
\ No newline at end of file
diff --git a/Engine/Source/Editor/BlueprintGraph/Private/Tests/TypePromotionTests.cpp b/Engine/Source/Editor/BlueprintGraph/Private/Tests/TypePromotionTests.cpp
new file mode 100644
--- /dev/null
+++ b/Engine/Source/Editor/BlueprintGraph/Private/Tests/TypePromotionTests.cpp
@@ -0,0 +1,126 @@
+// Copyright Epic Games, Inc. All Rights Reserved.
+
+#include "CoreTypes.h"
+#include "Containers/UnrealString.h"
+#include "Misc/AutomationTest.h"
+#include "BlueprintTypePromotion.h"
+#include "EdGraph/EdGraphPin.h"
+
+#if WITH_DEV_AUTOMATION_TESTS
+
+IMPLEMENT_SIMPLE_AUTOMATION_TEST(FTypePromotionTest, "Blueprints.Compiler.TypePromotion", EAutomationTestFlags::ApplicationContextMask | EAutomationTestFlags::SmokeFilter)
+
+bool FTypePromotionTest::RunTest(const FString& Parameters)
+{
+	FEdGraphPinType DoublePin = {};		DoublePin.PinCategory	= UEdGraphSchema_K2::PC_Double;
+	FEdGraphPinType FloatPin = {};		FloatPin.PinCategory	= UEdGraphSchema_K2::PC_Float;
+	FEdGraphPinType IntPin = {};		IntPin.PinCategory		= UEdGraphSchema_K2::PC_Int;
+	FEdGraphPinType Int64Pin = {};		Int64Pin.PinCategory	= UEdGraphSchema_K2::PC_Int64;
+	FEdGraphPinType BytePin = {};		BytePin.PinCategory		= UEdGraphSchema_K2::PC_Byte;
+
+	// Test promotions that should happen
+	TestEqual(TEXT("Testing float to double"), FTypePromotion::GetHigherType(FloatPin, DoublePin), FTypePromotion::ETypeComparisonResult::TypeBHigher);
+	
+	TestEqual(TEXT("Testing int to float"), FTypePromotion::GetHigherType(IntPin, FloatPin), FTypePromotion::ETypeComparisonResult::TypeBHigher);
+	TestEqual(TEXT("Testing int to double"), FTypePromotion::GetHigherType(IntPin, DoublePin), FTypePromotion::ETypeComparisonResult::TypeBHigher);
+	TestEqual(TEXT("Testing int to int64"), FTypePromotion::GetHigherType(IntPin, Int64Pin), FTypePromotion::ETypeComparisonResult::TypeBHigher);
+	
+	TestEqual(TEXT("Testing Byte to int"), FTypePromotion::GetHigherType(BytePin, IntPin), FTypePromotion::ETypeComparisonResult::TypeBHigher);
+	TestEqual(TEXT("Testing Byte to int64"), FTypePromotion::GetHigherType(BytePin, Int64Pin), FTypePromotion::ETypeComparisonResult::TypeBHigher);
+	
+	TestEqual(TEXT("Testing Double to int64"), FTypePromotion::GetHigherType(DoublePin, Int64Pin), FTypePromotion::ETypeComparisonResult::TypeBHigher);
+
+	// Test Equality of pins
+	TestEqual(TEXT("Testing Byte == Byte"), FTypePromotion::GetHigherType(BytePin, BytePin), FTypePromotion::ETypeComparisonResult::TypesEqual);
+	TestEqual(TEXT("Testing float == float"), FTypePromotion::GetHigherType(FloatPin, FloatPin), FTypePromotion::ETypeComparisonResult::TypesEqual);
+	TestEqual(TEXT("Testing double == double"), FTypePromotion::GetHigherType(DoublePin, DoublePin), FTypePromotion::ETypeComparisonResult::TypesEqual);
+	TestEqual(TEXT("Testing int == int"), FTypePromotion::GetHigherType(IntPin, IntPin), FTypePromotion::ETypeComparisonResult::TypesEqual);
+	TestEqual(TEXT("Testing int64 == int64"), FTypePromotion::GetHigherType(Int64Pin, Int64Pin), FTypePromotion::ETypeComparisonResult::TypesEqual);
+
+
+	// Test promotions that should not happen
+	TestEqual(TEXT("Testing int64 cannot go to byte"), FTypePromotion::GetHigherType(Int64Pin, BytePin), FTypePromotion::ETypeComparisonResult::TypeAHigher);
+	TestEqual(TEXT("Testing int64 cannot go to int"), FTypePromotion::GetHigherType(Int64Pin, IntPin), FTypePromotion::ETypeComparisonResult::TypeAHigher);
+	TestEqual(TEXT("Testing int64 cannot go to float"), FTypePromotion::GetHigherType(Int64Pin, FloatPin), FTypePromotion::ETypeComparisonResult::TypeAHigher);
+
+	return true;
+}
+
+
+IMPLEMENT_SIMPLE_AUTOMATION_TEST(FFindPromotedFunc, "Blueprints.Compiler.FindPromotedFunc", EAutomationTestFlags::ApplicationContextMask | EAutomationTestFlags::SmokeFilter)
+
+bool FFindPromotedFunc::RunTest(const FString& Parameters)
+{
+	TArray<UEdGraphPin*> PinTypes = {};
+
+	UEdGraphNode* TestNode = NewObject<UEdGraphNode>();
+
+#define MakeTestPin(InPinName, InPinType)		UEdGraphPin* InPinName = UEdGraphPin::CreatePin(TestNode);		\
+												InPinName->PinType.PinCategory = InPinType;						\
+												PinTypes.Add(InPinName);										
+	
+	const FTypePromotion& TypePromo = FTypePromotion::Get();
+
+	MakeTestPin(DoublePin, UEdGraphSchema_K2::PC_Double);
+	MakeTestPin(FloatPin, UEdGraphSchema_K2::PC_Float);
+	MakeTestPin(Int32Pin, UEdGraphSchema_K2::PC_Int);
+	MakeTestPin(Int64Pin, UEdGraphSchema_K2::PC_Int64);
+
+	// Add Operation
+	{
+		TArray<UEdGraphPin*> TestPins =
+		{
+			DoublePin, FloatPin
+		};
+		UFunction* AddDoubleFunc = FTypePromotion::GetOperatorFunction(TEXT("add"), TestPins);
+		TestNotNull(TEXT("Add Double Float function"), AddDoubleFunc);
+	}
+
+	// Add Operation
+	{
+		TArray<UEdGraphPin*> TestPins =
+		{
+			FloatPin, DoublePin
+		};
+		UFunction* AddDoubleFunc = FTypePromotion::GetOperatorFunction(TEXT("add"), TestPins);
+		TestNotNull(TEXT("Add Float Double function"), AddDoubleFunc);
+	}
+
+	// Multiply
+	{
+		TArray<UEdGraphPin*> TestPins =
+		{
+			FloatPin, Int32Pin
+		};
+		UFunction* ResultFunc = FTypePromotion::GetOperatorFunction(TEXT("multiply"), TestPins);
+		TestNotNull(TEXT("multiply Float Int32 function"), ResultFunc);
+	}
+
+	// Divide
+	{
+		TArray<UEdGraphPin*> TestPins =
+		{
+			FloatPin, Int32Pin
+		};
+
+		UFunction* ResultFunc = FTypePromotion::GetOperatorFunction(TEXT("divide"), TestPins);
+		TestNotNull(TEXT("divide Float Int32 function"), ResultFunc);
+	}
+
+	// Clear our test pins
+	for(UEdGraphPin* TestPin : PinTypes)
+	{
+		if(TestPin)
+		{
+			TestPin->MarkPendingKill();
+		}
+	}
+
+#undef MakeTestPin
+
+	TestNode->MarkPendingKill();
+
+	return true;
+}
+
+#endif //WITH_DEV_AUTOMATION_TESTS
\ No newline at end of file
