Make SearchForAutocastFunction take in Pin Types, not pin pointers.

#rnx

diff --git a/Engine/Source/Editor/AnimGraph/Classes/AnimationGraphSchema.h b/Engine/Source/Editor/AnimGraph/Classes/AnimationGraphSchema.h
--- a/Engine/Source/Editor/AnimGraph/Classes/AnimationGraphSchema.h
+++ b/Engine/Source/Editor/AnimGraph/Classes/AnimationGraphSchema.h
@@ -69,7 +69,7 @@ class UAnimationGraphSchema : public UEdGraphSchema_K2
 
 	//~ Begin UEdGraphSchema_K2 Interface
 	virtual const FPinConnectionResponse DetermineConnectionResponseOfCompatibleTypedPins(const UEdGraphPin* PinA, const UEdGraphPin* PinB, const UEdGraphPin* InputPin, const UEdGraphPin* OutputPin) const override;
-	virtual bool SearchForAutocastFunction(const UEdGraphPin* OutputPin, const UEdGraphPin* InputPin, /*out*/ FName& TargetFunction, /*out*/ UClass*& FunctionOwner) const override;
+	virtual bool SearchForAutocastFunction(const FEdGraphPinType& OutputPinType, const FEdGraphPinType& InputPinType, /*out*/ FName& TargetFunction, /*out*/ UClass*& FunctionOwner) const override;
 	virtual bool ArePinsCompatible(const UEdGraphPin* PinA, const UEdGraphPin* PinB, const UClass* CallingContext = NULL, bool bIgnoreArray = false) const override;
 	virtual bool DoesSupportAnimNotifyActions() const override;
 	virtual void CreateFunctionGraphTerminators(UEdGraph& Graph, UClass* Class) const override;
diff --git a/Engine/Source/Editor/AnimGraph/Private/AnimationGraphSchema.cpp b/Engine/Source/Editor/AnimGraph/Private/AnimationGraphSchema.cpp
--- a/Engine/Source/Editor/AnimGraph/Private/AnimationGraphSchema.cpp
+++ b/Engine/Source/Editor/AnimGraph/Private/AnimationGraphSchema.cpp
@@ -318,21 +318,21 @@ void UAnimationGraphSchema::CreateFunctionGraphTerminators(UEdGraph& Graph, UCla
 	}
 }
 
-bool UAnimationGraphSchema::SearchForAutocastFunction(const UEdGraphPin* OutputPin, const UEdGraphPin* InputPin, FName& TargetFunction, /*out*/ UClass*& FunctionOwner) const
+bool UAnimationGraphSchema::SearchForAutocastFunction(const FEdGraphPinType& OutputPinType, const FEdGraphPinType& InputPinType, FName& TargetFunction, /*out*/ UClass*& FunctionOwner) const
 {
-	if (IsComponentSpacePosePin(OutputPin->PinType) && IsLocalSpacePosePin(InputPin->PinType))
+	if (IsComponentSpacePosePin(OutputPinType) && IsLocalSpacePosePin(InputPinType))
 	{
 		// Insert a Component To LocalSpace conversion
 		return true;
 	}
-	else if (IsLocalSpacePosePin(OutputPin->PinType) && IsComponentSpacePosePin(InputPin->PinType))
+	else if (IsLocalSpacePosePin(OutputPinType) && IsComponentSpacePosePin(InputPinType))
 	{
 		// Insert a Local To ComponentSpace conversion
 		return true;
 	}
 	else
 	{
-		return Super::SearchForAutocastFunction(OutputPin, InputPin, TargetFunction, FunctionOwner);
+		return Super::SearchForAutocastFunction(OutputPinType, InputPinType, TargetFunction, FunctionOwner);
 	}
 }
 
diff --git a/Engine/Source/Editor/BlueprintGraph/Classes/EdGraphSchema_K2.h b/Engine/Source/Editor/BlueprintGraph/Classes/EdGraphSchema_K2.h
--- a/Engine/Source/Editor/BlueprintGraph/Classes/EdGraphSchema_K2.h
+++ b/Engine/Source/Editor/BlueprintGraph/Classes/EdGraphSchema_K2.h
@@ -1015,7 +1015,7 @@ public:
 	static UFunction* FindSetVariableByNameFunction(const FEdGraphPinType& PinType);
 
 	/** Find an appropriate function to call to perform an automatic cast operation */
-	virtual bool SearchForAutocastFunction(const UEdGraphPin* OutputPin, const UEdGraphPin* InputPin, /*out*/ FName& TargetFunction, /*out*/ UClass*& FunctionOwner) const;
+	virtual bool SearchForAutocastFunction(const FEdGraphPinType& OutputPinType, const FEdGraphPinType& InputPinType, /*out*/ FName& TargetFunction, /*out*/ UClass*& FunctionOwner) const;
 
 	/** Find an appropriate node that can convert from one pin type to another (not a cast; e.g. "MakeLiteralArray" node) */
 	virtual bool FindSpecializedConversionNode(const UEdGraphPin* OutputPin, const UEdGraphPin* InputPin, bool bCreateNode, /*out*/ class UK2Node*& TargetNode) const;
diff --git a/Engine/Source/Editor/BlueprintGraph/Private/BlueprintTypePromotion.cpp b/Engine/Source/Editor/BlueprintGraph/Private/BlueprintTypePromotion.cpp
--- a/Engine/Source/Editor/BlueprintGraph/Private/BlueprintTypePromotion.cpp
+++ b/Engine/Source/Editor/BlueprintGraph/Private/BlueprintTypePromotion.cpp
@@ -89,7 +89,7 @@ bool FTypePromotion::HasStructConversion(const UEdGraphPin* InputPin, const UEdG
 	UClass* DummyClass = nullptr;
 	UK2Node* DummyNode = nullptr;
 
-	const bool bCanAutocast = K2Schema->SearchForAutocastFunction(OutputPin, InputPin, /*out*/ DummyName, DummyClass);
+	const bool bCanAutocast = K2Schema->SearchForAutocastFunction(OutputPin->PinType, InputPin->PinType, /*out*/ DummyName, DummyClass);
 	const bool bCanAutoConvert = K2Schema->FindSpecializedConversionNode(OutputPin, InputPin, false, /* out */ DummyNode);
 	
 	return bCanAutocast || bCanAutoConvert;
diff --git a/Engine/Source/Editor/BlueprintGraph/Private/EdGraphSchema_K2.cpp b/Engine/Source/Editor/BlueprintGraph/Private/EdGraphSchema_K2.cpp
--- a/Engine/Source/Editor/BlueprintGraph/Private/EdGraphSchema_K2.cpp
+++ b/Engine/Source/Editor/BlueprintGraph/Private/EdGraphSchema_K2.cpp
@@ -2162,7 +2162,7 @@ const FPinConnectionResponse UEdGraphSchema_K2::CanCreateConnection(const UEdGra
 		UClass* DummyClass;
 		UK2Node* DummyNode;
 
-		const bool bCanAutocast = SearchForAutocastFunction(OutputPin, InputPin, /*out*/ DummyName, DummyClass);
+		const bool bCanAutocast = SearchForAutocastFunction(OutputPin->PinType, InputPin->PinType, /*out*/ DummyName, DummyClass);
 		const bool bCanAutoConvert = FindSpecializedConversionNode(OutputPin, InputPin, false, /* out */ DummyNode);
 
 		if (bCanAutocast || bCanAutoConvert)
@@ -2385,7 +2385,7 @@ void UEdGraphSchema_K2::Shutdown()
 }
 
 
-bool UEdGraphSchema_K2::SearchForAutocastFunction(const UEdGraphPin* OutputPin, const UEdGraphPin* InputPin, /*out*/ FName& TargetFunction, /*out*/ UClass*& FunctionOwner) const
+bool UEdGraphSchema_K2::SearchForAutocastFunction(const FEdGraphPinType& OutputPinType, const FEdGraphPinType& InputPinType, /*out*/ FName& TargetFunction, /*out*/ UClass*& FunctionOwner) const
 {
 	TRACE_CPUPROFILER_EVENT_SCOPE(WILD_UEdGraphSchema_K2::SearchForAutocastFunction);
 	// NOTE: Under no circumstances should anyone *ever* add a questionable cast to this function.
@@ -2395,9 +2395,9 @@ bool UEdGraphSchema_K2::SearchForAutocastFunction(const UEdGraphPin* OutputPin,
 	TargetFunction = NAME_None;
 	FunctionOwner = nullptr;
 
-	if (OutputPin->PinType.ContainerType != InputPin->PinType.ContainerType)
+	if (OutputPinType.ContainerType != InputPinType.ContainerType)
 	{
-		if (OutputPin->PinType.IsSet() && InputPin->PinType.IsArray())
+		if (OutputPinType.IsSet() && InputPinType.IsArray())
 		{
 			UFunction* Function = UBlueprintSetLibrary::StaticClass()->FindFunctionByName(GET_MEMBER_NAME_CHECKED(UBlueprintSetLibrary, Set_ToArray));
 			TargetFunction = Function->GetFName();
@@ -2410,9 +2410,9 @@ bool UEdGraphSchema_K2::SearchForAutocastFunction(const UEdGraphPin* OutputPin,
 	else
 	{
 		// SPECIAL CASES, not supported by FAutocastFunctionMap.
-		if ((OutputPin->PinType.PinCategory == PC_Interface) && (InputPin->PinType.PinCategory == PC_Object))
+		if ((OutputPinType.PinCategory == PC_Interface) && (InputPinType.PinCategory == PC_Object))
 		{
-			UClass const* InputClass = Cast<UClass const>(InputPin->PinType.PinSubCategoryObject.Get());
+			UClass const* InputClass = Cast<UClass const>(InputPinType.PinSubCategoryObject.Get());
 
 			bool const bInputIsUObject = (InputClass && (InputClass == UObject::StaticClass()));
 			if (bInputIsUObject)
@@ -2422,12 +2422,12 @@ bool UEdGraphSchema_K2::SearchForAutocastFunction(const UEdGraphPin* OutputPin,
 				FunctionOwner = Function->GetOwnerClass();
 			}
 		}
-		else if (OutputPin->PinType.PinCategory == PC_Object)
+		else if (OutputPinType.PinCategory == PC_Object)
 		{
-			UClass const* OutputClass = Cast<UClass const>(OutputPin->PinType.PinSubCategoryObject.Get());
-			if (InputPin->PinType.PinCategory == PC_Class)
+			UClass const* OutputClass = Cast<UClass const>(OutputPinType.PinSubCategoryObject.Get());
+			if (InputPinType.PinCategory == PC_Class)
 			{
-				UClass const* InputClass = Cast<UClass const>(InputPin->PinType.PinSubCategoryObject.Get());
+				UClass const* InputClass = Cast<UClass const>(InputPinType.PinSubCategoryObject.Get());
 				if ((OutputClass != nullptr) &&
 					(InputClass != nullptr) &&
 					OutputClass->IsChildOf(InputClass))
@@ -2437,29 +2437,29 @@ bool UEdGraphSchema_K2::SearchForAutocastFunction(const UEdGraphPin* OutputPin,
 					FunctionOwner = Function->GetOwnerClass();
 				}
 			}
-			else if (InputPin->PinType.PinCategory == PC_String)
+			else if (InputPinType.PinCategory == PC_String)
 			{
 				UFunction* Function = UKismetSystemLibrary::StaticClass()->FindFunctionByName(GET_MEMBER_NAME_CHECKED(UKismetSystemLibrary, GetDisplayName));
 				TargetFunction = Function->GetFName();
 				FunctionOwner = Function->GetOwnerClass();
 			}
 		}
-		else if (OutputPin->PinType.PinCategory == PC_Class)
+		else if (OutputPinType.PinCategory == PC_Class)
 		{
-			if (InputPin->PinType.PinCategory == PC_String)
+			if (InputPinType.PinCategory == PC_String)
 			{
 				UFunction* Function = UKismetSystemLibrary::StaticClass()->FindFunctionByName(GET_MEMBER_NAME_CHECKED(UKismetSystemLibrary, GetClassDisplayName));
 				TargetFunction = Function->GetFName();
 				FunctionOwner = Function->GetOwnerClass();
 			}
 		}
-		else if (OutputPin->PinType.PinCategory == PC_Struct)
+		else if (OutputPinType.PinCategory == PC_Struct)
 		{
-			const UScriptStruct* OutputStructType = Cast<const UScriptStruct>(OutputPin->PinType.PinSubCategoryObject.Get());
+			const UScriptStruct* OutputStructType = Cast<const UScriptStruct>(OutputPinType.PinSubCategoryObject.Get());
 			if (OutputStructType == TBaseStructure<FRotator>::Get())
 			{
-				const UScriptStruct* InputStructType = Cast<const UScriptStruct>(InputPin->PinType.PinSubCategoryObject.Get());
-				if ((InputPin->PinType.PinCategory == PC_Struct) && (InputStructType == TBaseStructure<FTransform>::Get()))
+				const UScriptStruct* InputStructType = Cast<const UScriptStruct>(InputPinType.PinSubCategoryObject.Get());
+				if ((InputPinType.PinCategory == PC_Struct) && (InputStructType == TBaseStructure<FTransform>::Get()))
 				{
 					UFunction* Function = UKismetMathLibrary::StaticClass()->FindFunctionByName(GET_MEMBER_NAME_CHECKED(UKismetMathLibrary, MakeTransform));
 					TargetFunction = Function->GetFName();
@@ -2473,7 +2473,7 @@ bool UEdGraphSchema_K2::SearchForAutocastFunction(const UEdGraphPin* OutputPin,
 	if (TargetFunction == NAME_None)
 	{
 		const FAutocastFunctionMap& AutocastFunctionMap = FAutocastFunctionMap::Get();
-		if (UFunction* Func = AutocastFunctionMap.Find(OutputPin->PinType, InputPin->PinType))
+		if (UFunction* Func = AutocastFunctionMap.Find(OutputPinType, InputPinType))
 		{
 			TargetFunction = Func->GetFName();
 			FunctionOwner = Func->GetOwnerClass();
@@ -2663,8 +2663,8 @@ void UEdGraphSchema_K2::AutowireConversionNode(UEdGraphPin* InputPin, UEdGraphPi
 bool UEdGraphSchema_K2::CreateAutomaticConversionNodeAndConnections(UEdGraphPin* PinA, UEdGraphPin* PinB) const
 {
 	// Determine which pin is an input and which pin is an output
-	UEdGraphPin* InputPin = NULL;
-	UEdGraphPin* OutputPin = NULL;
+	UEdGraphPin* InputPin = nullptr;
+	UEdGraphPin* OutputPin = nullptr;
 	if (!CategorizePinsByDirection(PinA, PinB, /*out*/ InputPin, /*out*/ OutputPin))
 	{
 		return false;
@@ -2674,9 +2674,9 @@ bool UEdGraphSchema_K2::CreateAutomaticConversionNodeAndConnections(UEdGraphPin*
 	UClass* ClassContainingConversionFunction = nullptr;
 	TSubclassOf<UK2Node> ConversionNodeClass;
 
-	UK2Node* TemplateConversionNode = NULL;
+	UK2Node* TemplateConversionNode = nullptr;
 
-	if (SearchForAutocastFunction(OutputPin, InputPin, /*out*/ TargetFunctionName, /*out*/ClassContainingConversionFunction))
+	if (SearchForAutocastFunction(OutputPin->PinType, InputPin->PinType, /*out*/ TargetFunctionName, /*out*/ClassContainingConversionFunction))
 	{
 		// Create a new call function node for the casting operator
 		UK2Node_CallFunction* TemplateNode = NewObject<UK2Node_CallFunction>();
