Use ConsolveVariable pointer as an accessor for Type Promo flag instead of just a bool flag. Fixes issue where the flag would be incorrectly false in other modules.

#rnx

diff --git a/Engine/Source/Editor/BlueprintGraph/Classes/BlueprintTypePromotion.h b/Engine/Source/Editor/BlueprintGraph/Classes/BlueprintTypePromotion.h
--- a/Engine/Source/Editor/BlueprintGraph/Classes/BlueprintTypePromotion.h
+++ b/Engine/Source/Editor/BlueprintGraph/Classes/BlueprintTypePromotion.h
@@ -157,20 +157,9 @@ private:
 namespace TypePromoDebug
 {
 	/** Enables/Disables type promotion in BP */
-	static bool bIsTypePromoEnabled = false;
-	static FAutoConsoleVariableRef CVarIsTypePromoEnabled(
-		TEXT("BP.TypePromo.IsEnabled"), bIsTypePromoEnabled,
-		TEXT("If true then type promotion inside of blueprints will be enabled"),
-		FConsoleVariableDelegate::CreateLambda([](IConsoleVariable* InVariable)
-			{
-				// Clear the node spawner so that we create the new BP actions correctly
-				FTypePromotion::ClearNodeSpawners();
-
-				// Refresh all the actions so that the context menu goes back to the normal options
-				if (FBlueprintActionDatabase* Actions = FBlueprintActionDatabase::TryGet())
-				{
-					Actions->RefreshAll();
-				}
-			}),
-		ECVF_Default);
+	static bool IsTypePromoEnabled()
+	{
+		static IConsoleVariable* TypePromoCVar = IConsoleManager::Get().FindConsoleVariable(TEXT("BP.TypePromo.IsEnabled"));
+		return TypePromoCVar ? TypePromoCVar->GetBool() : false;
+	}
 }
\ No newline at end of file
diff --git a/Engine/Source/Editor/BlueprintGraph/Private/BlueprintTypePromotion.cpp b/Engine/Source/Editor/BlueprintGraph/Private/BlueprintTypePromotion.cpp
--- a/Engine/Source/Editor/BlueprintGraph/Private/BlueprintTypePromotion.cpp
+++ b/Engine/Source/Editor/BlueprintGraph/Private/BlueprintTypePromotion.cpp
@@ -23,6 +23,26 @@ namespace OperatorNames
 	static const FString NotEq		= TEXT("NotEqual");
 }
 
+namespace TypePromoDebug
+{
+	static bool bIsTypePromoEnabled = false;
+	static FAutoConsoleVariableRef CVarIsTypePromoEnabled(
+		TEXT("BP.TypePromo.IsEnabled"), bIsTypePromoEnabled,
+		TEXT("If true then type promotion inside of blueprints will be enabled"),
+		FConsoleVariableDelegate::CreateLambda([](IConsoleVariable* InVariable)
+		{
+			// Clear the node spawner so that we create the new BP actions correctly
+			FTypePromotion::ClearNodeSpawners();
+
+			// Refresh all the actions so that the context menu goes back to the normal options
+			if (FBlueprintActionDatabase* Actions = FBlueprintActionDatabase::TryGet())
+			{
+				Actions->RefreshAll();
+			}
+		}),
+		ECVF_Default);
+}
+
 FTypePromotion& FTypePromotion::Get()
 {
 	if (Instance == nullptr)
@@ -65,7 +85,7 @@ void FTypePromotion::CreatePromotionTable()
 	{
 		// Type_X...						Can be promoted to...
 		{ UEdGraphSchema_K2::PC_Int,		{ UEdGraphSchema_K2::PC_Float, UEdGraphSchema_K2::PC_Double, UEdGraphSchema_K2::PC_Int64 } },
-		{ UEdGraphSchema_K2::PC_Byte,		{ UEdGraphSchema_K2::PC_Int, UEdGraphSchema_K2::PC_Int64 } },
+		{ UEdGraphSchema_K2::PC_Byte,		{ UEdGraphSchema_K2::PC_Float, UEdGraphSchema_K2::PC_Int, UEdGraphSchema_K2::PC_Int64 } },
 		{ UEdGraphSchema_K2::PC_Float,		{ UEdGraphSchema_K2::PC_Double, UEdGraphSchema_K2::PC_Int64 } },
 		{ UEdGraphSchema_K2::PC_Double,		{ UEdGraphSchema_K2::PC_Int64 } },
 		{ UEdGraphSchema_K2::PC_Wildcard,	{ UEdGraphSchema_K2::PC_Int, UEdGraphSchema_K2::PC_Int64, UEdGraphSchema_K2::PC_Float, UEdGraphSchema_K2::PC_Double, UEdGraphSchema_K2::PC_Byte, UEdGraphSchema_K2::PC_Boolean } },
@@ -77,8 +97,19 @@ bool FTypePromotion::IsValidPromotion(const FEdGraphPinType& A, const FEdGraphPi
 	// If either of these pin types is a struct, than we have to have some kind of valid
 	// conversion function, otherwise we can't possibly connect them
 	const ETypeComparisonResult Res = FTypePromotion::GetHigherType(A, B);
+	bool bCanAutocast = false;
+	if (A.PinCategory == UEdGraphSchema_K2::PC_Struct || B.PinCategory == UEdGraphSchema_K2::PC_Struct)
+	{
+		const UEdGraphSchema_K2* K2Schema = GetDefault<UEdGraphSchema_K2>();
+
+		FName DummyName;
+		UClass* DummyClass = nullptr;
+		UK2Node* DummyNode = nullptr;
+
+		bCanAutocast = K2Schema->SearchForAutocastFunction(A, B, /*out*/ DummyName, DummyClass);
+	}
 
-	return Res == ETypeComparisonResult::TypeBHigher;
+	return Res == ETypeComparisonResult::TypeBHigher || bCanAutocast;
 }
 
 bool FTypePromotion::HasStructConversion(const UEdGraphPin* InputPin, const UEdGraphPin* OutputPin)
@@ -381,7 +412,7 @@ UFunction* FTypePromotion::FindLowestMatchingFunc_Internal(const FString& Operat
 				{
 					// If an input of this function is compatible with this type
 					// Then we need to check all the other inputs on this function
-					OutPossibleFunctions.Emplace(Func);
+					OutPossibleFunctions.AddUnique(Func);
 					break;
 				}
 			}
