Blueprints: Add the CONNECT_RESPONSE_MAKE_WITH_PROMOTION to the EdGraphSchema, along with the TryCreatePromotedConnection to allow type promotion as an option.

#rnx

diff --git a/Engine/Source/Editor/BlueprintGraph/Classes/EdGraphSchema_K2.h b/Engine/Source/Editor/BlueprintGraph/Classes/EdGraphSchema_K2.h
--- a/Engine/Source/Editor/BlueprintGraph/Classes/EdGraphSchema_K2.h
+++ b/Engine/Source/Editor/BlueprintGraph/Classes/EdGraphSchema_K2.h
@@ -479,6 +479,7 @@ public:
 	virtual const FPinConnectionResponse CanCreateConnection(const UEdGraphPin* A, const UEdGraphPin* B) const override;
 	virtual bool TryCreateConnection(UEdGraphPin* A, UEdGraphPin* B) const override;
 	virtual bool CreateAutomaticConversionNodeAndConnections(UEdGraphPin* A, UEdGraphPin* B) const override;
+	virtual bool CreatePromotedConnection(UEdGraphPin* A, UEdGraphPin* B) const override;
 	virtual FString IsPinDefaultValid(const UEdGraphPin* Pin, const FString& NewDefaultValue, UObject* NewDefaultObject, const FText& InNewDefaultText) const override;
 	virtual bool DoesSupportPinWatching() const	override;
 	virtual bool IsPinBeingWatched(UEdGraphPin const* Pin) const override;
diff --git a/Engine/Source/Editor/BlueprintGraph/Private/BlueprintActionDatabase.cpp b/Engine/Source/Editor/BlueprintGraph/Private/BlueprintActionDatabase.cpp
--- a/Engine/Source/Editor/BlueprintGraph/Private/BlueprintActionDatabase.cpp
+++ b/Engine/Source/Editor/BlueprintGraph/Private/BlueprintActionDatabase.cpp
@@ -68,6 +68,8 @@
 
 #include "Misc/HotReloadInterface.h"
 
+#include "BlueprintTypePromotion.h"
+
 #define LOCTEXT_NAMESPACE "BlueprintActionDatabase"
 
 /*******************************************************************************
@@ -645,7 +647,15 @@ static void BlueprintActionDatabaseImpl::AddClassFunctionActions(UClass const* c
 				ActionListOut.Add(NodeSpawner);
 			}
 		}
-		
+
+		// If this is a promotable function, and it has already been registered
+		// than do NOT add it to the asset action database. We should
+		// probably have some better logic for this, like adding our own node spawner
+		const bool bIsRegisteredPromotionFunc =
+			TypePromoDebug::IsTypePromoEnabled() &&
+			FTypePromotion::IsFunctionPromotionReady(Function) &&
+			FTypePromotion::IsOperatorSpawnerRegistered(Function);
+
 		if (UEdGraphSchema_K2::CanUserKismetCallFunction(Function))
 		{
 			// @TODO: if this is a Blueprint, and this function is from a 
@@ -653,7 +663,13 @@ static void BlueprintActionDatabaseImpl::AddClassFunctionActions(UClass const* c
 			//        include it (the function is accounted for in from the 
 			//        interface class).
 			UBlueprintFunctionNodeSpawner* FuncSpawner = UBlueprintFunctionNodeSpawner::Create(Function);
-			ActionListOut.Add(FuncSpawner);
+			
+			// Only add this action to the list of the operator function is not already registered. Otherwise we will 
+			// get a bunch of duplicate operator actions
+			if (!bIsRegisteredPromotionFunc)
+			{
+				ActionListOut.Add(FuncSpawner);
+			}
 
 			if (FKismetEditorUtilities::IsClassABlueprintInterface(Class))
 			{
diff --git a/Engine/Source/Editor/BlueprintGraph/Private/EdGraphSchema_K2.cpp b/Engine/Source/Editor/BlueprintGraph/Private/EdGraphSchema_K2.cpp
--- a/Engine/Source/Editor/BlueprintGraph/Private/EdGraphSchema_K2.cpp
+++ b/Engine/Source/Editor/BlueprintGraph/Private/EdGraphSchema_K2.cpp
@@ -96,6 +96,7 @@
 #include "K2Node_ConvertAsset.h"
 #include "Framework/Commands/GenericCommands.h"
 
+#include "BlueprintTypePromotion.h"
 
 //////////////////////////////////////////////////////////////////////////
 // FBlueprintMetadata
@@ -2706,6 +2707,25 @@ bool UEdGraphSchema_K2::CreateAutomaticConversionNodeAndConnections(UEdGraphPin*
 	return false;
 }
 
+bool UEdGraphSchema_K2::CreatePromotedConnection(UEdGraphPin* PinA, UEdGraphPin* PinB) const
+{
+	PinA->Modify();
+	PinB->Modify();
+
+	PinB->BreakAllPinLinks(true);
+
+	PinA->MakeLinkTo(PinB);
+
+	UBlueprint* Blueprint = FBlueprintEditorUtils::FindBlueprintForNodeChecked(PinA->GetOwningNode());
+
+	if (!PinA->IsPendingKill())
+	{
+		FBlueprintEditorUtils::MarkBlueprintAsModified(Blueprint);
+	}
+
+	return true;
+}
+
 FString UEdGraphSchema_K2::IsPinDefaultValid(const UEdGraphPin* Pin, const FString& NewDefaultValue, UObject* NewDefaultObject, const FText& InNewDefaultText) const
 {
 	check(Pin);
diff --git a/Engine/Source/Editor/BlueprintGraph/Private/K2Node.cpp b/Engine/Source/Editor/BlueprintGraph/Private/K2Node.cpp
--- a/Engine/Source/Editor/BlueprintGraph/Private/K2Node.cpp
+++ b/Engine/Source/Editor/BlueprintGraph/Private/K2Node.cpp
@@ -381,14 +381,24 @@ void UK2Node::AutowireNewNode(UEdGraphPin* FromPin)
 
 				// null out the backup connection (so we don't attempt to make it 
 				// once we exit the loop... we successfully made this connection!)
-				BackupConnection = NULL;
+				BackupConnection = nullptr;
+				break;
+			}
+			else if(ConnectResponse == ECanCreateConnectionResponse::CONNECT_RESPONSE_MAKE_WITH_PROMOTION)
+			{
+				if (K2Schema->CreatePromotedConnection(FromPin, Pin))
+				{
+					NodeList.Add(FromPin->GetOwningNode());
+					NodeList.Add(this);
+				}
+				BackupConnection = nullptr;
 				break;
 			}
 		}
 
 		// if we didn't find an ideal connection, then lets connect this pin to 
 		// the BackupConnection (something, like a connection that requires a conversion node, etc.)
-		if ((BackupConnection != NULL) && K2Schema->TryCreateConnection(FromPin, BackupConnection))
+		if ((BackupConnection != nullptr) && K2Schema->TryCreateConnection(FromPin, BackupConnection))
 		{
 			NodeList.Add(FromPin->GetOwningNode());
 			NodeList.Add(this);
@@ -402,7 +412,7 @@ void UK2Node::AutowireNewNode(UEdGraphPin* FromPin)
 
 			UEdGraphPin* ToExecutePin = FindPin(UEdGraphSchema_K2::PN_Execute);
 
-			if ((FromThenPin != NULL) && (FromThenPin->LinkedTo.Num() == 0) && (ToExecutePin != NULL) && K2Schema->ArePinsCompatible(FromThenPin, ToExecutePin, NULL))
+			if ((FromThenPin != nullptr) && (FromThenPin->LinkedTo.Num() == 0) && (ToExecutePin != nullptr) && K2Schema->ArePinsCompatible(FromThenPin, ToExecutePin, NULL))
 			{
 				if (K2Schema->TryCreateConnection(FromThenPin, ToExecutePin))
 				{
diff --git a/Engine/Source/Editor/GraphEditor/Private/DragConnection.cpp b/Engine/Source/Editor/GraphEditor/Private/DragConnection.cpp
--- a/Engine/Source/Editor/GraphEditor/Private/DragConnection.cpp
+++ b/Engine/Source/Editor/GraphEditor/Private/DragConnection.cpp
@@ -141,6 +141,10 @@ void FDragConnection::HoverTargetChanged()
 				StatusSymbol = FEditorStyle::GetBrush(TEXT("Graph.ConnectorFeedback.ViaCast"));
 				break;
 
+			case CONNECT_RESPONSE_MAKE_WITH_PROMOTION:
+				StatusSymbol = FEditorStyle::GetBrush(TEXT("Graph.ConnectorFeedback.ViaCast"));
+				break;
+
 			case CONNECT_RESPONSE_DISALLOW:
 			default:
 				StatusSymbol = FEditorStyle::GetBrush(TEXT("Graph.ConnectorFeedback.Error"));
diff --git a/Engine/Source/Runtime/Engine/Classes/EdGraph/EdGraphSchema.h b/Engine/Source/Runtime/Engine/Classes/EdGraph/EdGraphSchema.h
--- a/Engine/Source/Runtime/Engine/Classes/EdGraph/EdGraphSchema.h
+++ b/Engine/Source/Runtime/Engine/Classes/EdGraph/EdGraphSchema.h
@@ -48,6 +48,9 @@ enum ECanCreateConnectionResponse
 	/** Make the connection via an intermediate cast node, or some other conversion node. */
 	CONNECT_RESPONSE_MAKE_WITH_CONVERSION_NODE,
 
+	/** Make the connection by promoting a lower type to a higher type. Ex: Connecting a Float -> Double, float should become a double */
+	CONNECT_RESPONSE_MAKE_WITH_PROMOTION,
+
 	CONNECT_RESPONSE_MAX,
 };
 
@@ -711,6 +714,16 @@ class ENGINE_API UEdGraphSchema : public UObject
 	 */
 	virtual bool CreateAutomaticConversionNodeAndConnections(UEdGraphPin* A, UEdGraphPin* B) const;
 
+	/**
+	* Try to create a promotion from one type to another in order to make a connection between two pins.
+	* 
+	* @param	A	The first pin.
+	* @param	B	The second pin.
+	*
+	* @return	True if the promotion and connection were successful; False if the connection failed.
+	*/
+	virtual bool CreatePromotedConnection(UEdGraphPin* A, UEdGraphPin* B) const;
+
 	/**
 	 * Determine if the supplied pin default values would be valid.
 	 *
diff --git a/Engine/Source/Runtime/Engine/Private/EdGraph/EdGraphSchema.cpp b/Engine/Source/Runtime/Engine/Private/EdGraph/EdGraphSchema.cpp
--- a/Engine/Source/Runtime/Engine/Private/EdGraph/EdGraphSchema.cpp
+++ b/Engine/Source/Runtime/Engine/Private/EdGraph/EdGraphSchema.cpp
@@ -440,6 +440,9 @@ bool UEdGraphSchema::TryCreateConnection(UEdGraphPin* PinA, UEdGraphPin* PinB) c
 	case CONNECT_RESPONSE_MAKE_WITH_CONVERSION_NODE:
 		bModified = CreateAutomaticConversionNodeAndConnections(PinA, PinB);
 		break;
+
+	case CONNECT_RESPONSE_MAKE_WITH_PROMOTION:
+		bModified = CreatePromotedConnection(PinA, PinB);
 		break;
 
 	case CONNECT_RESPONSE_DISALLOW:
@@ -463,6 +466,11 @@ bool UEdGraphSchema::CreateAutomaticConversionNodeAndConnections(UEdGraphPin* Pi
 	return false;
 }
 
+bool UEdGraphSchema::CreatePromotedConnection(UEdGraphPin* PinA, UEdGraphPin* PinB) const
+{
+	return false;
+}
+
 void UEdGraphSchema::TrySetDefaultValue(UEdGraphPin& Pin, const FString& NewDefaultValue, bool bMarkAsModified) const
 {
 	Pin.DefaultValue = NewDefaultValue;
