Hold operator names in a TSet instead of TArray for faster Contains check

#rnx

diff --git a/Engine/Source/Editor/BlueprintGraph/Classes/BlueprintTypePromotion.h b/Engine/Source/Editor/BlueprintGraph/Classes/BlueprintTypePromotion.h
--- a/Engine/Source/Editor/BlueprintGraph/Classes/BlueprintTypePromotion.h
+++ b/Engine/Source/Editor/BlueprintGraph/Classes/BlueprintTypePromotion.h
@@ -36,8 +36,14 @@ public:
 	/** Returns all functions for a specific operation. Will empty the given array and populate it with UFunction pointers */
 	static void GetAllFuncsForOp(const FString& Operation, TArray<UFunction*>& OutFuncs);
 
-	/** Get an array of the supported operator names for type promo. Ex: "Add", "Subtract", "Multiply" */
-	static const TArray<FString>& GetOpNames();
+	/** Get a set of the supported operator names for type promo. Ex: "Add", "Subtract", "Multiply" */
+	static const TSet<FString>& GetAllOpNames();
+	
+	/** Set of comparison operator names (GreaterThan, LessThan, etc) */
+	static const TSet<FString>& GetComparisonOpNames();
+
+	/** Returns true if the given function is a comparison operator */
+	static bool IsComparisonFunc(UFunction const* const Func);
 
 	static bool GetOpNameFromFunction(UFunction const* const Func, FString& OutName);
 
diff --git a/Engine/Source/Editor/BlueprintGraph/Private/BlueprintTypePromotion.cpp b/Engine/Source/Editor/BlueprintGraph/Private/BlueprintTypePromotion.cpp
--- a/Engine/Source/Editor/BlueprintGraph/Private/BlueprintTypePromotion.cpp
+++ b/Engine/Source/Editor/BlueprintGraph/Private/BlueprintTypePromotion.cpp
@@ -460,9 +460,9 @@ void FTypePromotion::GetAllFuncsForOp(const FString& Operation, TArray<UFunction
 	return FTypePromotion::Get().GetAllFuncsForOp_Internal(Operation, OutFuncs);
 }
 
-const TArray<FString>& FTypePromotion::GetOpNames()
+const TSet<FString>& FTypePromotion::GetAllOpNames()
 {
-	static const TArray<FString> OpsArray =
+	static const TSet<FString> OpsArray =
 	{
 		OperatorNames::Add,
 		OperatorNames::Multiply,
@@ -478,6 +478,25 @@ const TArray<FString>& FTypePromotion::GetOpNames()
 	return OpsArray;
 }
 
+const TSet<FString>& FTypePromotion::GetComparisonOpNames()
+{
+	static const TSet<FString> ComparisonOps =
+	{
+		OperatorNames::Greater,
+		OperatorNames::GreaterEq,
+		OperatorNames::Less,
+		OperatorNames::LessEq,
+		OperatorNames::NotEq
+	};
+	return ComparisonOps;
+}
+
+bool FTypePromotion::IsComparisonFunc(UFunction const* const Func)
+{
+	FString OpName;
+	return Func && GetOpNameFromFunction(Func, /* out */ OpName) && GetComparisonOpNames().Contains(OpName);
+}
+
 void FTypePromotion::GetAllFuncsForOp_Internal(const FString& Operation, TArray<UFunction*>& OutFuncs)
 {
 	OutFuncs.Empty();
@@ -496,18 +515,16 @@ bool FTypePromotion::GetOpNameFromFunction(UFunction const* const Func, FString&
 	// Get everything before the "_"
 	int32 Index = FuncName.Find(TEXT("_"));
 	FString FuncNameChopped = FuncName.Mid(0, Index);
-
-	for(const FString& OpName : GetOpNames())
+	if (GetAllOpNames().Contains(FuncNameChopped))
 	{
-		if(FuncNameChopped == OpName)
-		{
-			OutName = OpName;
-			return true;
-		}
+		OutName = FuncNameChopped;
+		return true;
+	}
+	else
+	{
+		OutName = OperatorNames::NoOp;
+		return false;
 	}
-
-	OutName = OperatorNames::NoOp;
-	return false;
 }
 
 void FTypePromotion::CreateOpTable()
