Improve syntax in FTypePromotion::IsValidPromotion

#rnx
#okforigthub public

diff --git a/Engine/Source/Editor/BlueprintGraph/Private/BlueprintTypePromotion.cpp b/Engine/Source/Editor/BlueprintGraph/Private/BlueprintTypePromotion.cpp
--- a/Engine/Source/Editor/BlueprintGraph/Private/BlueprintTypePromotion.cpp
+++ b/Engine/Source/Editor/BlueprintGraph/Private/BlueprintTypePromotion.cpp
@@ -96,8 +96,6 @@ bool FTypePromotion::IsValidPromotion(const FEdGraphPinType& A, const FEdGraphPi
 {
 	// If either of these pin types is a struct, than we have to have some kind of valid
 	// conversion function, otherwise we can't possibly connect them
-	const ETypeComparisonResult Res = FTypePromotion::GetHigherType(A, B);
-	bool bCanAutocast = false;
 	if (A.PinCategory == UEdGraphSchema_K2::PC_Struct || B.PinCategory == UEdGraphSchema_K2::PC_Struct)
 	{
 		const UEdGraphSchema_K2* K2Schema = GetDefault<UEdGraphSchema_K2>();
@@ -106,10 +104,12 @@ bool FTypePromotion::IsValidPromotion(const FEdGraphPinType& A, const FEdGraphPi
 		UClass* DummyClass = nullptr;
 		UK2Node* DummyNode = nullptr;
 
-		bCanAutocast = K2Schema->SearchForAutocastFunction(A, B, /*out*/ DummyName, DummyClass);
+		return K2Schema->SearchForAutocastFunction(A, B, /*out*/ DummyName, DummyClass);
+	}
+	else
+	{
+		return FTypePromotion::GetHigherType(A, B) == ETypeComparisonResult::TypeBHigher;
 	}
-
-	return Res == ETypeComparisonResult::TypeBHigher || bCanAutocast;
 }
 
 bool FTypePromotion::HasStructConversion(const UEdGraphPin* InputPin, const UEdGraphPin* OutputPin)
