Blueprints: Add the first pass of the Promotable Operator node for type promotion. Disabled by default, can be turned on with the CVar "BP.TypePromo.IsEnabled"

diff --git a/Engine/Source/Editor/BlueprintGraph/Classes/K2Node_PromotableOperator.h b/Engine/Source/Editor/BlueprintGraph/Classes/K2Node_PromotableOperator.h
new file mode 100644
--- /dev/null
+++ b/Engine/Source/Editor/BlueprintGraph/Classes/K2Node_PromotableOperator.h
@@ -0,0 +1,91 @@
+// Copyright Epic Games, Inc. All Rights Reserved.
+
+#pragma once
+
+#include "CoreMinimal.h"
+#include "UObject/ObjectMacros.h"
+#include "EdGraph/EdGraphPin.h"
+#include "K2Node_CallFunction.h"
+#include "K2Node_PromotableOperator.generated.h"
+
+/** The promotable operator node allows for pin types to be promoted to others, i.e. float to double */
+UCLASS(MinimalAPI)
+class UK2Node_PromotableOperator : public UK2Node_CallFunction
+{
+	GENERATED_UCLASS_BODY()
+
+public:
+
+	// UEdGraphNode interface
+	virtual void AllocateDefaultPins() override;
+	virtual void GetNodeContextMenuActions(UToolMenu* Menu, UGraphNodeContextMenuContext* Context) const override;
+	virtual FText GetTooltipText() const override;
+	// End of UEdGraphNode interface
+
+	// UK2Node interface
+	virtual void ExpandNode(FKismetCompilerContext& CompilerContext, UEdGraph* SourceGraph) override;
+	virtual void NotifyPinConnectionListChanged(UEdGraphPin* ChangedPin) override;
+	virtual void PostReconstructNode() override;
+	virtual bool IsConnectionDisallowed(const UEdGraphPin* MyPin, const UEdGraphPin* OtherPin, FString& OutReason) const override;
+	virtual void ReallocatePinsDuringReconstruction(TArray<UEdGraphPin*>& OldPins) override;
+	virtual void AutowireNewNode(UEdGraphPin* FromPin) override;
+	virtual bool IsActionFilteredOut(class FBlueprintActionFilter const& Filter) override;
+	// End of UK2Node interface
+
+private:
+
+	/** Helper function to recombine all split pins that this node may have */
+	void RecombineAllSplitPins();
+
+	/** 
+	* @return	True if this node has any connections attached to it, 
+	*			or the default values have been modified by the user
+	*/
+	bool HasAnyConnectionsOrDefaults() const;
+
+	/** Recombines all split pins and sets the node to have default values (all wildcard pins) */
+	void ResetNodeToWildcard();
+
+	/**
+	* Attempts to create a cast node and connect it to a call function node
+	* 
+	* @return	True if the intermediate connection was made successfully	
+	*/
+	bool CreateIntermediateCast(FKismetCompilerContext& CompilerContext, UEdGraph* SourceGraph, UEdGraphPin* InputPin, UEdGraphPin* OutputPin);
+
+	/** Helper to make sure we have the most up to date operation name. Returns true upon success */
+	bool UpdateOpName();
+
+	/** 
+	* Update the pins on this node based on the given function. This modifies pins, meant 
+	* for use by PinConnectionListChanged, not during node construction. 
+	*/
+	void UpdatePinsFromFunction(const UFunction* Function, UEdGraphPin* ChangedPin = nullptr);
+
+	/** Called when the user attempts conversion from the context menu. Records a transaction and calls UpdatePinsFromFunction */
+	void ConvertNodeToFunction(const UFunction* Function, UEdGraphPin* ChangedPin);
+
+	/** Updates the PossibleConversions function array based on the current pin types */
+	void UpdatePossibleConversionFuncs();
+
+	/**
+	* Returns all pins that have the EGPD_Input direction
+	* @param bIncludeLinks	If true, than this will also include all the pins that are linked to the inputs.
+	*						This is useful for gathering what the highest type may be
+	*/
+	TArray<UEdGraphPin*> GetInputPins(bool bIncludeLinks = false) const;
+
+	/** The name that this operation uses ("Add", "Multiply", etc) */
+	FString OperationName;
+
+	/** Array of functions that we could possibly convert this node to via the right-click context menu */
+	UPROPERTY(Transient)
+	TArray<UFunction*> PossibleConversions;
+
+public:
+
+	/** Returns the first pin with the EGPD_Output direction */
+	UEdGraphPin* GetOutputPin() const;
+
+	const FString& GetOperationName() const { return OperationName; }
+};
diff --git a/Engine/Source/Editor/BlueprintGraph/Private/BlueprintActionFilter.cpp b/Engine/Source/Editor/BlueprintGraph/Private/BlueprintActionFilter.cpp
--- a/Engine/Source/Editor/BlueprintGraph/Private/BlueprintActionFilter.cpp
+++ b/Engine/Source/Editor/BlueprintGraph/Private/BlueprintActionFilter.cpp
@@ -34,6 +34,7 @@
 // "impure" node types (utilized in BlueprintActionFilterImpl::IsImpure)
 #include "K2Node_MultiGate.h"
 #include "K2Node_Message.h"
+#include "K2Node_PromotableOperator.h"
 #include "EditorCategoryUtils.h"
 
 /*******************************************************************************
@@ -1354,6 +1355,11 @@ static bool BlueprintActionFilterImpl::IsFunctionMissingPinParam(FBlueprintActio
 		UEdGraphSchema_K2 const* K2Schema = GetDefault<UEdGraphSchema_K2>();
 		bool const bIsEventSpawner  = BlueprintAction.GetNodeClass()->IsChildOf<UK2Node_Event>();
 		bool const bIsArrayFunction = BlueprintAction.GetNodeClass()->IsChildOf<UK2Node_CallArrayFunction>();
+		
+		if (BlueprintAction.GetNodeClass()->IsChildOf<UK2Node_PromotableOperator>())
+		{
+			return false;
+		}
 
 		for (int32 PinIndex = 0; !bIsFilteredOut && (PinIndex < Filter.Context.Pins.Num()); ++PinIndex)
 		{
@@ -1519,7 +1525,7 @@ static bool BlueprintActionFilterImpl::IsMissingMatchingPinParam(FBlueprintActio
 	// we have a separate pin tests for function/property nodes (IsFunctionMissingPinParam/IsMissmatchedPropertyType). Note that we only skip
 	// this test for functions with bindings (because it does not handle getting templates for all binding nodes). By running this for 
 	// other functions we ensure that IsConnectionDisallowed is honored.
-	bool const bTestPinCompatibility = (BlueprintAction.GetAssociatedProperty() == nullptr) && BlueprintAction.GetBindings().Num() == 0;
+	bool const bTestPinCompatibility = (BlueprintAction.GetAssociatedProperty() == nullptr) && BlueprintAction.GetBindings().Num() == 0 && !BlueprintAction.GetNodeClass()->IsChildOf<UK2Node_PromotableOperator>();
 
 	if (bTestPinCompatibility)
 	{
diff --git a/Engine/Source/Editor/BlueprintGraph/Private/BlueprintFunctionNodeSpawner.cpp b/Engine/Source/Editor/BlueprintGraph/Private/BlueprintFunctionNodeSpawner.cpp
--- a/Engine/Source/Editor/BlueprintGraph/Private/BlueprintFunctionNodeSpawner.cpp
+++ b/Engine/Source/Editor/BlueprintGraph/Private/BlueprintFunctionNodeSpawner.cpp
@@ -18,6 +18,8 @@
 #include "BlueprintNodeSpawnerUtils.h"
 #include "BlueprintEditorSettings.h"
 #include "Kismet2/BlueprintEditorUtils.h"
+#include "BlueprintTypePromotion.h"
+#include "K2Node_PromotableOperator.h"
 
 #define LOCTEXT_NAMESPACE "BlueprintFunctionNodeSpawner"
 
@@ -187,8 +189,14 @@ UBlueprintFunctionNodeSpawner* UBlueprintFunctionNodeSpawner::Create(UFunction c
 	bool const bIsMaterialParamCollectionFunc = Function->HasMetaData(FBlueprintMetadata::MD_MaterialParameterCollectionFunction);
 	bool const bIsDataTableFunc = Function->HasMetaData(FBlueprintMetadata::MD_DataTablePin);
 
+	bool const bIsPromotableFunction = TypePromoDebug::IsTypePromoEnabled() && FTypePromotion::IsFunctionPromotionReady(Function);
+
 	TSubclassOf<UK2Node_CallFunction> NodeClass;
-	if (bIsCommutativeAssociativeBinaryOp && bIsPure)
+	if (bIsPromotableFunction)
+	{
+		NodeClass = UK2Node_PromotableOperator::StaticClass();
+	}
+	else if(bIsCommutativeAssociativeBinaryOp && bIsPure)
 	{
 		NodeClass = UK2Node_CommutativeAssociativeBinaryOperator::StaticClass();
 	}
@@ -225,6 +233,22 @@ UBlueprintFunctionNodeSpawner* UBlueprintFunctionNodeSpawner::Create(TSubclassOf
 	// Constructing the Spawner
 	//--------------------------------------
 
+	bool const bIsPromotableFunction =
+		TypePromoDebug::IsTypePromoEnabled() &&
+		FTypePromotion::IsFunctionPromotionReady(Function);
+
+	FString OpName;
+	FTypePromotion::GetOpNameFromFunction(Function, OpName);
+
+	// If a spawner for this operator has been created already, than just return that
+	if (bIsPromotableFunction && FTypePromotion::IsOperatorSpawnerRegistered(Function))
+	{
+		if (UBlueprintFunctionNodeSpawner* OpSpawner = FTypePromotion::GetOperatorSpawner(OpName))
+		{
+			return OpSpawner;
+		}
+	}
+
 	UBlueprintFunctionNodeSpawner* NodeSpawner = NewObject<UBlueprintFunctionNodeSpawner>(Outer);
 	NodeSpawner->SetField(const_cast<UFunction*>(Function));
 
@@ -242,11 +266,29 @@ UBlueprintFunctionNodeSpawner* UBlueprintFunctionNodeSpawner::Create(TSubclassOf
 	//--------------------------------------
 
 	FBlueprintActionUiSpec& MenuSignature = NodeSpawner->DefaultMenuSignature;
-	MenuSignature.MenuName = UK2Node_CallFunction::GetUserFacingFunctionName(Function);
-	MenuSignature.Category = UK2Node_CallFunction::GetDefaultCategoryForFunction(Function, FText::GetEmpty());
-	MenuSignature.Tooltip  = FText::FromString( UK2Node_CallFunction::GetDefaultTooltipForFunction(Function) );
-	// add at least one character, so that PrimeDefaultUiSpec() doesn't attempt to query the template node
-	MenuSignature.Keywords = UK2Node_CallFunction::GetKeywordsForFunction(Function);
+
+	if(bIsPromotableFunction)
+	{
+		FFormatNamedArguments Args;
+		Args.Add(TEXT("OpName"), FText::FromString(OpName));
+		MenuSignature.MenuName = FText::Format(LOCTEXT("OpName", "{OpName}"), Args);
+		MenuSignature.Category = LOCTEXT("UtilityOperatorCategory", "Utilities|Operators");
+		// Possibly generate some special tooltips for promotable operators?
+		MenuSignature.Tooltip = FText::FromString(TEXT(" "));
+		// add at least one character, so that PrimeDefaultUiSpec() doesn't attempt to query the template node
+		MenuSignature.Keywords = FText::FromString(OpName);
+		FTypePromotion::RegisterOperatorSpawner(OpName, NodeSpawner);
+	}
+	else
+	{
+		MenuSignature.MenuName = UK2Node_CallFunction::GetUserFacingFunctionName(Function);
+		MenuSignature.Category = UK2Node_CallFunction::GetDefaultCategoryForFunction(Function, FText::GetEmpty());
+		MenuSignature.Tooltip = FText::FromString(UK2Node_CallFunction::GetDefaultTooltipForFunction(Function));
+		// add at least one character, so that PrimeDefaultUiSpec() doesn't attempt to query the template node
+		MenuSignature.Keywords = UK2Node_CallFunction::GetKeywordsForFunction(Function);
+	}
+	
+	
 	if (MenuSignature.Keywords.IsEmpty())
 	{
 		MenuSignature.Keywords = FText::FromString(TEXT(" "));
@@ -259,7 +301,7 @@ UBlueprintFunctionNodeSpawner* UBlueprintFunctionNodeSpawner::Create(TSubclassOf
 		MenuSignature.Category = BlueprintFunctionNodeSpawnerImpl::FallbackCategory;
 	}
 
-	if (MenuSignature.Tooltip.IsEmpty())
+	if (MenuSignature.Tooltip.IsEmpty() && !bIsPromotableFunction)
 	{
 		MenuSignature.Tooltip = MenuSignature.MenuName;
 	}
diff --git a/Engine/Source/Editor/BlueprintGraph/Private/EdGraphSchema_K2.cpp b/Engine/Source/Editor/BlueprintGraph/Private/EdGraphSchema_K2.cpp
--- a/Engine/Source/Editor/BlueprintGraph/Private/EdGraphSchema_K2.cpp
+++ b/Engine/Source/Editor/BlueprintGraph/Private/EdGraphSchema_K2.cpp
@@ -97,6 +97,7 @@
 #include "Framework/Commands/GenericCommands.h"
 
 #include "BlueprintTypePromotion.h"
+#include "K2Node_PromotableOperator.h"
 
 //////////////////////////////////////////////////////////////////////////
 // FBlueprintMetadata
@@ -2143,6 +2144,23 @@ const FPinConnectionResponse UEdGraphSchema_K2::CanCreateConnection(const UEdGra
 	// Compare the types
 	const bool bTypesMatch = ArePinsCompatible(OutputPin, InputPin, CallingContext, bIgnoreArray);
 
+	// Promotable types in blueprints! Only if the Cvar is set and the node is of a special type. Eventually we want this for all
+	if (TypePromoDebug::IsTypePromoEnabled() && InputPin->GetOwningNode()->IsA<UK2Node_PromotableOperator>())
+	{
+		if (FTypePromotion::IsValidPromotion(PinA->PinType, PinB->PinType) || FTypePromotion::IsValidPromotion(PinB->PinType, PinA->PinType))
+		{
+			// Set the Text here correctly based on which pin type is higher
+			if (FTypePromotion::GetHigherType(PinA->PinType, PinB->PinType) == FTypePromotion::ETypeComparisonResult::TypeAHigher)
+			{
+				return FPinConnectionResponse(CONNECT_RESPONSE_MAKE_WITH_PROMOTION, FString::Printf(TEXT("Promote %s to %s"), *TypeToText(PinB->PinType).ToString(), *TypeToText(PinA->PinType).ToString()));
+			}
+			else
+			{
+				return FPinConnectionResponse(CONNECT_RESPONSE_MAKE_WITH_PROMOTION, FString::Printf(TEXT("Promote %s to %s"), *TypeToText(PinA->PinType).ToString(), *TypeToText(PinB->PinType).ToString()));
+			}
+		}
+	}
+
 	if (bTypesMatch)
 	{
 		FPinConnectionResponse ConnectionResponse = DetermineConnectionResponseOfCompatibleTypedPins(PinA, PinB, InputPin, OutputPin);
diff --git a/Engine/Source/Editor/BlueprintGraph/Private/K2Node_PromotableOperator.cpp b/Engine/Source/Editor/BlueprintGraph/Private/K2Node_PromotableOperator.cpp
new file mode 100644
--- /dev/null
+++ b/Engine/Source/Editor/BlueprintGraph/Private/K2Node_PromotableOperator.cpp
@@ -0,0 +1,645 @@
+// Copyright Epic Games, Inc. All Rights Reserved.
+
+#include "K2Node_PromotableOperator.h"
+#include "BlueprintTypePromotion.h"
+#include "EdGraphSchema_K2.h"
+#include "EdGraphUtilities.h"
+#include "K2Node_CommutativeAssociativeBinaryOperator.h"
+#include "KismetCompiler.h"
+#include "ScopedTransaction.h"
+#include "ToolMenus.h"
+#include "Framework/Commands/UIAction.h"
+#include "Kismet2/BlueprintEditorUtils.h"
+#include "Kismet2/WildcardNodeUtils.h"
+
+#define LOCTEXT_NAMESPACE "PromotableOperatorNode"
+
+///////////////////////////////////////////////////////////
+// Pin names for default construction
+
+static const FName InputPinA_Name = FName(TEXT("A"));
+static const FName InputPinB_Name = FName(TEXT("B"));
+
+///////////////////////////////////////////////////////////
+// UK2Node_PromotableOperator
+
+UK2Node_PromotableOperator::UK2Node_PromotableOperator(const FObjectInitializer& ObjectInitializer)
+	: Super(ObjectInitializer)
+{
+	UpdateOpName();
+	OrphanedPinSaveMode = ESaveOrphanPinMode::SaveAllButExec;
+}
+
+///////////////////////////////////////////////////////////
+// UEdGraphNode interface
+
+void UK2Node_PromotableOperator::AllocateDefaultPins()
+{
+	FWildcardNodeUtils::CreateWildcardPin(this, InputPinA_Name, EGPD_Input);
+	FWildcardNodeUtils::CreateWildcardPin(this, InputPinB_Name, EGPD_Input);
+
+	FWildcardNodeUtils::CreateWildcardPin(this, UEdGraphSchema_K2::PN_ReturnValue, EGPD_Output);
+}
+
+void UK2Node_PromotableOperator::GetNodeContextMenuActions(UToolMenu* Menu, UGraphNodeContextMenuContext* Context) const
+{
+	Super::GetNodeContextMenuActions(Menu, Context);
+
+	// If there are possible function conversions that can happen 
+	if (Context->Pin && PossibleConversions.Num() > 0 && !Context->bIsDebugging && HasAnyConnectionsOrDefaults())
+	{
+		FToolMenuSection& Section = Menu->AddSection("K2NodePromotableOperator", LOCTEXT("ConvFunctionHeader", "Convert Function"));
+		const UFunction* CurFunction = GetTargetFunction();
+
+		for (const UFunction* Func : PossibleConversions)
+		{
+			// Don't need to convert to a function if we are already set to it
+			if (Func == CurFunction)
+			{
+				continue;
+			}
+
+			FFormatNamedArguments Args;
+			Args.Add(TEXT("TargetName"), GetUserFacingFunctionName(Func));
+			FText ConversionName = FText::Format(LOCTEXT("CallFunction_Tooltip", "Convert node to function '{TargetName}'"), Args);
+
+			const FText Tooltip = FText::FromString(GetDefaultTooltipForFunction(Func));
+
+			Section.AddMenuEntry(
+				Func->GetFName(),
+				ConversionName,
+				Tooltip,
+				FSlateIcon(),
+				FUIAction(
+					FExecuteAction::CreateUObject(const_cast<UK2Node_PromotableOperator*>(this), &UK2Node_PromotableOperator::ConvertNodeToFunction, Func, const_cast<UEdGraphPin*>(Context->Pin))
+				)
+			);			
+		}
+	}
+}
+
+FText UK2Node_PromotableOperator::GetTooltipText() const
+{
+	// If there are no connections then just display the op name
+	if (!HasAnyConnectionsOrDefaults())
+	{
+		UFunction* Function = GetTargetFunction();
+		FString OpName;
+		FTypePromotion::GetOpNameFromFunction(Function, OpName);
+		return FText::Format(LOCTEXT("PromotableOperatorFunctionTooltip", "{0} Operator"), FText::FromString(OpName));
+	}
+
+	// Otherwise use the default one (a more specific function tooltip)
+	return Super::GetTooltipText();
+}
+
+///////////////////////////////////////////////////////////
+// UK2Node interface
+
+void UK2Node_PromotableOperator::ExpandNode(FKismetCompilerContext& CompilerContext, UEdGraph* SourceGraph)
+{
+	Super::ExpandNode(CompilerContext, SourceGraph);
+
+	const bool bValidOpName = UpdateOpName();
+	if (!bValidOpName)
+	{
+		UE_LOG(LogBlueprint, Error, TEXT("Could not find matching operation name for this function!"));
+		CompilerContext.MessageLog.Error(TEXT("Could not find matching operation on '@@'!"), this);
+		return;
+	}
+
+	TArray<UEdGraphPin*> InputPins = GetInputPins();
+	UEdGraphPin* OutputPin = GetOutputPin();
+
+	// Our operator function has been determined on pin connection change
+	const UFunction* const OpFunction = GetTargetFunction();
+
+	if (!OpFunction)
+	{
+		UE_LOG(LogBlueprint, Error, TEXT("Could not find matching op function during expansion!"));
+		CompilerContext.MessageLog.Error(TEXT("Could not find matching op function during expansion on '@@'!"), this);
+		return;
+	}
+
+	// Now to actually go through the promotion process on pins that need to be promoted to 
+	// fit our function signature! 
+
+	// Spawn an intermediate K2NodeCommunicative op node of that type
+	UK2Node_CallFunction* NewOperator = SourceGraph->CreateIntermediateNode<UK2Node_CallFunction>();
+	NewOperator->SetFromFunction(OpFunction);
+	NewOperator->AllocateDefaultPins();
+
+	// Move this node next to the thing it was linked to
+	NewOperator->NodePosY = this->NodePosY;
+	NewOperator->NodePosX = this->NodePosX + 8;
+
+	CompilerContext.MessageLog.NotifyIntermediateObjectCreation(NewOperator, this);
+
+	UEdGraphPin* NewOperatorInputA = nullptr;
+	UEdGraphPin* NewOperatorInputB = nullptr;
+	UEdGraphPin* NewOperatorOutputPin = nullptr;
+	UEdGraphPin* NewOperatorSelfPin = NewOperator->FindPin(UEdGraphSchema_K2::PN_Self);
+
+	for (UEdGraphPin* Pin : NewOperator->Pins)
+	{
+		if (Pin == NewOperatorSelfPin)
+		{
+			continue;
+		}
+
+		if (Pin->Direction == EEdGraphPinDirection::EGPD_Input)
+		{
+			if (!NewOperatorInputA) 
+			{
+				NewOperatorInputA = Pin;
+			}
+			else if (!NewOperatorInputB) 
+			{
+				NewOperatorInputB = Pin;
+			}
+		}
+		else if (Pin->Direction == EEdGraphPinDirection::EGPD_Output)
+		{
+			NewOperatorOutputPin = Pin;
+		}
+	}
+
+	check(NewOperatorOutputPin);
+
+	// Create some auto casts if they are necessary
+	const bool bPinASuccess = CreateIntermediateCast(CompilerContext, SourceGraph, InputPins[0], NewOperatorInputA);
+	const bool bPinBSuccess = CreateIntermediateCast(CompilerContext, SourceGraph, InputPins[1], NewOperatorInputB);
+
+	if (!bPinASuccess || !bPinBSuccess)
+	{
+		CompilerContext.MessageLog.Error(TEXT("'@@' could not successfuly expand pins!"), this);
+	}
+
+	// Connect intermediate node output to this nodes output
+	CompilerContext.MovePinLinksToIntermediate(*GetOutputPin(), *NewOperatorOutputPin);
+}
+
+void UK2Node_PromotableOperator::NotifyPinConnectionListChanged(UEdGraphPin* ChangedPin)
+{
+	Super::NotifyPinConnectionListChanged(ChangedPin);
+
+	UpdateOpName();
+
+	const bool bOutputPinWasChanged = (ChangedPin == GetOutputPin());
+
+	// True if the pin that has changed now has zero connections
+	const bool bWasAFullDisconnect = (ChangedPin->LinkedTo.Num() == 0);
+
+	// If we have been totally disconnected and don't have any non-default inputs, 
+	// than we just reset the node to be a regular wildcard
+	if (bWasAFullDisconnect && !HasAnyConnectionsOrDefaults())
+	{
+		ResetNodeToWildcard();
+		return;
+	}
+	// If the pin that was connected is linked to a wildcard pin, then we should make it a wildcard
+	// and do nothing else.
+	else if (ChangedPin->GetOwningNode() == this && FWildcardNodeUtils::IsLinkedToWildcard(ChangedPin))
+	{
+		ChangedPin->PinType = FWildcardNodeUtils::GetDefaultWildcardPinType();
+		return;
+	}
+
+	// Gather all pins and their links so we can determine the highest type that the user could want
+	TArray<UEdGraphPin*> InputPins;
+
+	for (UEdGraphPin* Pin : Pins)
+	{
+		// Consider inputs, maybe the output, and ignore the changed pin if it was totally disconnected.
+		if (Pin->LinkedTo.Num() > 0 || !Pin->DoesDefaultValueMatchAutogenerated())
+		{
+			InputPins.Add(Pin);
+			for (UEdGraphPin* Link : Pin->LinkedTo)
+			{
+				InputPins.Emplace(Link);
+			}
+		}
+	}
+
+	FEdGraphPinType HighestType = FTypePromotion::GetPromotedType(InputPins);
+
+	// if a pin was changed, update it if it cannot be promoted to this type
+	FEdGraphPinType NewConnectionHighestType = ChangedPin->LinkedTo.Num() > 0 ? FTypePromotion::GetPromotedType(ChangedPin->LinkedTo) : FWildcardNodeUtils::GetDefaultWildcardPinType();
+
+	// If there are ANY wildcards on this node, than we need to update the whole node accordingly. Otherwise we can 
+	// update only the Changed and output pins.
+
+	if (FWildcardNodeUtils::NodeHasAnyWildcards(this) || bOutputPinWasChanged || bWasAFullDisconnect ||
+		FTypePromotion::GetHigherType(NewConnectionHighestType, GetOutputPin()->PinType) == FTypePromotion::ETypeComparisonResult::TypeAHigher)
+	{
+		const UFunction* LowestFunc = FTypePromotion::FindLowestMatchingFunc(OperationName, HighestType, PossibleConversions);
+
+		// Store these other function options for later so that the user can convert to them later
+		UpdatePinsFromFunction(LowestFunc, ChangedPin);
+	}
+
+	// If the user connected a type that was a valid promotion, than leave it as the pin type they dragged from for a better UX
+	if (!bWasAFullDisconnect && NewConnectionHighestType.PinCategory != UEdGraphSchema_K2::PC_Wildcard &&
+		(FTypePromotion::IsValidPromotion(NewConnectionHighestType, ChangedPin->PinType) || FTypePromotion::IsValidPromotion(ChangedPin->PinType, NewConnectionHighestType)))
+	{
+		ChangedPin->PinType = NewConnectionHighestType;
+	}
+
+	// Update context menu options for this node
+	UpdatePossibleConversionFuncs();
+}
+
+void UK2Node_PromotableOperator::PostReconstructNode()
+{
+	Super::PostReconstructNode();
+
+	// We only need to set the function if we have connections, otherwise we should stick in a wildcard state
+	if (HasAnyConnectionsOrDefaults())
+	{
+		// Allocate default pins will have been called before this, which means we are reset to wildcard state
+		// We need to Update the pins to be the proper function again
+		UpdatePinsFromFunction(GetTargetFunction());
+	}
+}
+
+bool UK2Node_PromotableOperator::IsConnectionDisallowed(const UEdGraphPin* MyPin, const UEdGraphPin* OtherPin, FString& OutReason) const
+{
+	check(MyPin && OtherPin);
+
+	// #TODO_BH Just disallow containers and references for now
+	if (OtherPin->PinType.IsContainer() || OtherPin->PinType.bIsReference)
+	{
+		OutReason = LOCTEXT("NoExecPinsAllowed", "Promotable Operator nodes cannot have containers or references.").ToString();
+		return true;
+	}
+	else if (MyPin == GetOutputPin() && FTypePromotion::IsComparisonFunc(GetTargetFunction()) && OtherPin->PinType.PinCategory != UEdGraphSchema_K2::PC_Boolean)
+	{
+		OutReason = LOCTEXT("ComparisonNeedsBool", "Comparison operators must return a bool!").ToString();
+		return true;
+	}
+
+	const bool bHasStructPin = MyPin->PinType.PinCategory == UEdGraphSchema_K2::PC_Struct || OtherPin->PinType.PinCategory == UEdGraphSchema_K2::PC_Struct;
+
+	// If the other pin can be promoted to my pin type, than allow the connection
+	if (FTypePromotion::IsValidPromotion(OtherPin->PinType, MyPin->PinType))
+	{
+		if (bHasStructPin)
+		{
+			// Compare the directions
+			const UEdGraphPin* InputPin = nullptr;
+			const UEdGraphPin* OutputPin = nullptr;
+			const UEdGraphSchema_K2* K2Schema = GetDefault<UEdGraphSchema_K2>();
+
+			if (!K2Schema->CategorizePinsByDirection(MyPin, OtherPin, /*out*/ InputPin, /*out*/ OutputPin))
+			{
+				OutReason = LOCTEXT("DirectionsIncompatible", "Pin directions are not compatible!").ToString();
+				return true;
+			}
+
+			if (!FTypePromotion::HasStructConversion(InputPin, OutputPin))
+			{
+				FFormatNamedArguments Args;
+				Args.Add(TEXT("MyPinType"), K2Schema->TypeToText(MyPin->PinType));
+				Args.Add(TEXT("OtherPinType"), K2Schema->TypeToText(OtherPin->PinType));
+
+				OutReason = FText::Format(LOCTEXT("NoCompatibleStructConv", "No compatible operator functions between '{MyPinType}' and '{OtherPinType}'"), Args).ToString();
+				return true;
+			}
+		}
+		return false;
+	}
+
+	return Super::IsConnectionDisallowed(MyPin, OtherPin, OutReason);
+}
+
+void UK2Node_PromotableOperator::ReallocatePinsDuringReconstruction(TArray<UEdGraphPin*>& OldPins)
+{
+	// Allocate default pins will have been called before this, which means we are reset to wildcard state
+	// We need to Update the pins to be the proper function again
+	UpdatePinsFromFunction(GetTargetFunction(), nullptr);
+
+	Super::ReallocatePinsDuringReconstruction(OldPins);
+}
+
+void UK2Node_PromotableOperator::AutowireNewNode(UEdGraphPin* ChangedPin)
+{
+	Super::AutowireNewNode(ChangedPin);
+
+	NotifyPinConnectionListChanged(ChangedPin);
+}
+
+bool UK2Node_PromotableOperator::IsActionFilteredOut(FBlueprintActionFilter const& Filter)
+{
+	return false;
+}
+
+///////////////////////////////////////////////////////////
+// UK2Node_PromotableOperator
+
+bool UK2Node_PromotableOperator::HasAnyConnectionsOrDefaults() const
+{
+	for (UEdGraphPin* Pin : Pins)
+	{
+		if (Pin->LinkedTo.Num() > 0 || !Pin->DoesDefaultValueMatchAutogenerated())
+		{
+			return true;
+		}
+	}
+	return false;
+}
+
+bool UK2Node_PromotableOperator::UpdateOpName()
+{
+	const UFunction* Func = GetTargetFunction();
+	// If the function is null then return false, because we did not successfully update it. 
+	// This could be possible during node reconstruction/refresh, and we don't want to set the 
+	// op name to "Empty" incorrectly. 
+	return Func ? FTypePromotion::GetOpNameFromFunction(GetTargetFunction(), /* out */ OperationName) : false;
+}
+
+bool UK2Node_PromotableOperator::CreateIntermediateCast(FKismetCompilerContext& CompilerContext, UEdGraph* SourceGraph, UEdGraphPin* InputPin, UEdGraphPin* OutputPin)
+{
+	check(InputPin && OutputPin);
+
+	// If the pin types are the same, than no casts are needed and we can just connect
+	if (InputPin->PinType == OutputPin->PinType)
+	{
+		return !CompilerContext.MovePinLinksToIntermediate(*InputPin, *OutputPin).IsFatal();
+	}
+
+	const UEdGraphSchema_K2* Schema = CompilerContext.GetSchema();
+
+	UK2Node* TemplateConversionNode = nullptr;
+	FName TargetFunctionName;
+	UClass* ClassContainingConversionFunction = nullptr;
+	TSubclassOf<UK2Node> ConversionNodeClass;
+
+	if (Schema->SearchForAutocastFunction(InputPin->PinType, OutputPin->PinType, /*out*/ TargetFunctionName, /*out*/ClassContainingConversionFunction))
+	{
+		// Create a new call function node for the casting operator
+		UK2Node_CallFunction* TemplateNode = SourceGraph->CreateIntermediateNode<UK2Node_CallFunction>();
+		TemplateNode->FunctionReference.SetExternalMember(TargetFunctionName, ClassContainingConversionFunction);
+		TemplateConversionNode = TemplateNode;
+		TemplateNode->AllocateDefaultPins();
+		CompilerContext.MessageLog.NotifyIntermediateObjectCreation(TemplateNode, this);
+	}
+	else
+	{
+		Schema->FindSpecializedConversionNode(InputPin, OutputPin, true, /*out*/ TemplateConversionNode);
+	}
+
+	bool bInputSuccessful = false;
+	bool bOutputSuccessful = false;
+
+	if (TemplateConversionNode)
+	{
+		UEdGraphPin* ConversionInput = nullptr;
+
+		for (UEdGraphPin* ConvPin : TemplateConversionNode->Pins)
+		{
+			if (ConvPin && ConvPin->Direction == EGPD_Input && ConvPin->PinName != UEdGraphSchema_K2::PSC_Self)
+			{
+				ConversionInput = ConvPin;
+				break;
+			}
+		}
+		UEdGraphPin* ConversionOutput = TemplateConversionNode->FindPin(UEdGraphSchema_K2::PN_ReturnValue, EGPD_Output);
+
+		// Connect my input to the conversion node
+		if (InputPin->LinkedTo.Num() > 0)
+		{
+			bInputSuccessful = Schema->TryCreateConnection(InputPin->LinkedTo[0], ConversionInput);
+		}
+
+		// Connect conversion node output to the input of the new operator
+		bOutputSuccessful = Schema->TryCreateConnection(ConversionOutput, OutputPin);
+
+		// Move this node next to the thing it was linked to
+		TemplateConversionNode->NodePosY = this->NodePosY;
+		TemplateConversionNode->NodePosX = this->NodePosX + 4;
+	}
+	else
+	{
+		CompilerContext.MessageLog.Error(*FText::Format(LOCTEXT("NoValidPromotion", "Cannot find appropriate promotion from '{0}' to '{1}' on '@@'"),
+			Schema->TypeToText(InputPin->PinType),
+			Schema->TypeToText(OutputPin->PinType)).ToString(),
+			this
+		);
+	}
+
+	return bInputSuccessful && bOutputSuccessful;
+}
+
+void UK2Node_PromotableOperator::ResetNodeToWildcard()
+{
+	RecombineAllSplitPins();
+
+	// Reset type to wildcard
+	const FEdGraphPinType& WildType = FWildcardNodeUtils::GetDefaultWildcardPinType();
+	const UEdGraphSchema_K2* K2Schema = GetDefault<UEdGraphSchema_K2>();
+
+	for (UEdGraphPin* Pin : Pins)
+	{
+		// Ensure this pin is not a split pin
+		if (Pin && Pin->ParentPin == nullptr)
+		{
+			Pin->PinType = WildType;
+			K2Schema->ResetPinToAutogeneratedDefaultValue(Pin);
+		}
+	}
+
+	// Clear out any possible function matches, since we are removing connections
+	PossibleConversions.Empty();
+}
+
+void UK2Node_PromotableOperator::RecombineAllSplitPins()
+{
+	const UEdGraphSchema_K2* K2Schema = GetDefault<UEdGraphSchema_K2>();
+
+	// Gather what pins need to be recombined from a split pin	
+	for (int32 Index = 0; Index < Pins.Num(); ++Index) 
+	{
+		if (Pins[Index] && Pins[Index]->SubPins.Num() > 0)
+		{
+			K2Schema->RecombinePin(Pins[Index]);
+		}
+	}
+}
+
+TArray<UEdGraphPin*> UK2Node_PromotableOperator::GetInputPins(bool bIncludeLinks /** = false */) const
+{
+	TArray<UEdGraphPin*> InputPins;
+	for (UEdGraphPin* Pin : Pins)
+	{
+		// Exclude split pins from this
+		if (Pin && Pin->Direction == EEdGraphPinDirection::EGPD_Input && Pin->ParentPin == nullptr)
+		{
+			InputPins.Add(Pin);
+			if (bIncludeLinks)
+			{
+				for (UEdGraphPin* Link : Pin->LinkedTo)
+				{
+					InputPins.Emplace(Link);
+				}
+			}
+		}
+	}
+	return InputPins;
+}
+
+void UK2Node_PromotableOperator::ConvertNodeToFunction(const UFunction* Function, UEdGraphPin* ChangedPin)
+{
+	const FScopedTransaction Transaction(LOCTEXT("ConvertPromotableOpToFunction", "Change the function signiture of a promotable operator node."));
+	Modify();
+	RecombineAllSplitPins();
+	UpdatePinsFromFunction(Function, ChangedPin);
+
+	// Reconstruct this node to fix any default values that may be invalid now
+	ReconstructNode();
+}
+
+void UK2Node_PromotableOperator::UpdatePinsFromFunction(const UFunction* Function, UEdGraphPin* ChangedPin /* = nullptr */)
+{
+	if (!Function)
+	{
+		UE_LOG(LogBlueprint, Warning, TEXT("UK2Node_PromotableOperator could not update pins, function was null!"));
+		return;
+	}
+
+	const UEdGraphSchema_K2* Schema = GetDefault<UEdGraphSchema_K2>();
+
+	TMap<FString, TSet<UEdGraphPin*>> PinConnections;
+	FEdGraphUtilities::GetPinConnectionMap(this, PinConnections);
+
+	int32 ArgCount = 0;
+	for (TFieldIterator<FProperty> PropIt(Function); PropIt && (PropIt->PropertyFlags & CPF_Parm); ++PropIt)
+	{
+		FProperty* Param = *PropIt;
+		FEdGraphPinType ParamType;
+
+		// If this param type is valid and we are still within range of the number of pins that we have
+		if (Schema->ConvertPropertyToPinType(Param, /* out */ ParamType) && ArgCount >= 0 && ArgCount < Pins.Num())
+		{
+			// Get either the output pin or one of the input pins
+			UEdGraphPin* PinToChange = Param->HasAnyPropertyFlags(CPF_ReturnParm) ? GetOutputPin() : Pins[ArgCount];
+
+			check(PinToChange);
+
+			const bool bHasConnectionOrDefault = PinToChange->LinkedTo.Num() > 0 || !PinToChange->DoesDefaultValueMatchAutogenerated();
+			const bool bIsWildcard = FWildcardNodeUtils::IsWildcardPin(PinToChange);
+			const bool bIsValidPromo = !bIsWildcard && FTypePromotion::IsValidPromotion(PinToChange->PinType, ParamType);
+			const bool bTypesEqual = PinToChange->PinType == ParamType;
+			const bool bIsOutPin = PinToChange->Direction == EGPD_Output;
+
+			bool bIsLinkedToWildcard = false;
+			if (bIsWildcard && bHasConnectionOrDefault)
+			{
+				for (UEdGraphPin* LinkedPin : PinToChange->LinkedTo)
+				{
+					if (FWildcardNodeUtils::IsWildcardPin(LinkedPin))
+					{
+						bIsLinkedToWildcard = true;
+						break;
+					}
+				}
+			}
+
+			// If this is a wildcard WITH a connection, then we should just leave this pin as a wildcard and let the compiler handle it
+			if (bIsLinkedToWildcard)
+			{
+				continue;
+			}
+
+			bool bNeedsTypeUpdate = true;
+
+			// If this pin has a valid value already, than dont bother updating it.
+			if (bHasConnectionOrDefault && (bIsValidPromo || bTypesEqual))
+			{
+				bNeedsTypeUpdate = false;
+			}
+
+			// We always want to update the out pin or if we have a wildcard pin (which is the case during reconstruction)
+			if (!bTypesEqual && (bIsOutPin || bIsWildcard))
+			{
+				bNeedsTypeUpdate = true;
+			}
+
+			if (bNeedsTypeUpdate)
+			{
+				// If this is a wildcard pin than we have to reconsider the links to that 
+				if (bIsWildcard && PinToChange->LinkedTo.Num() > 0)
+				{
+					FEdGraphPinType LinkedType = FTypePromotion::GetPromotedType(PinToChange->LinkedTo);
+					if (FTypePromotion::IsValidPromotion(PinToChange->PinType, ParamType))
+					{
+						ParamType = LinkedType;
+					}
+				}
+				// If this update is coming from a changed pin, than we need to orphan the pin that had the old type
+				else if (ChangedPin && bHasConnectionOrDefault)
+				{
+					PinToChange->BreakAllPinLinks();
+				}
+
+				// Only change the type of this pin if it is necessary				
+				PinToChange->PinType = ParamType;
+			}
+
+		}
+
+		++ArgCount;
+	}
+
+	// Update the function reference and the FUNC_BlueprintPure/FUNC_Const appropriately
+	SetFromFunction(Function);
+
+	UpdatePossibleConversionFuncs();
+}
+
+void UK2Node_PromotableOperator::UpdatePossibleConversionFuncs()
+{
+	// Pins can be empty if we are doing this during reconstruction
+	if (Pins.Num() <= 0)
+	{
+		return;
+	}
+
+	bool bAllPinTypesEqual = true;
+	FEdGraphPinType CurType = Pins[0]->PinType;
+
+	for (UEdGraphPin* Pin : Pins)
+	{
+		if (Pin->PinType != CurType)
+		{
+			bAllPinTypesEqual = false;
+			break;
+		}
+	}
+
+	UpdateOpName();
+
+	// We don't want to have a menu that is full of every possible function for an operator, that is way 
+	// to overwhelming to the user. Instead, only display conversion functions that when types are not 
+	// all the same. 
+	if (!bAllPinTypesEqual)
+	{
+		// The node has changed, so lets find the lowest matching function with the newly updated types
+		FEdGraphPinType HighestType = FTypePromotion::GetPromotedType(GetInputPins());
+		FTypePromotion::FindLowestMatchingFunc(OperationName, HighestType, PossibleConversions);
+	}
+}
+
+UEdGraphPin* UK2Node_PromotableOperator::GetOutputPin() const
+{
+	for (UEdGraphPin* Pin : Pins)
+	{
+		if (Pin && Pin->Direction == EGPD_Output)
+		{
+			return Pin;
+		}
+	}
+
+	return nullptr;
+}
+
+#undef LOCTEXT_NAMESPACE	// "PromotableOperatorNode"
\ No newline at end of file
