Cleanup FName usage of "ReturnValue" and replace with UEdGraphSchema_K2::PN_ReturnValue

#rnx

diff --git a/Engine/Source/Editor/BlueprintGraph/Private/K2Node_FormatText.cpp b/Engine/Source/Editor/BlueprintGraph/Private/K2Node_FormatText.cpp
--- a/Engine/Source/Editor/BlueprintGraph/Private/K2Node_FormatText.cpp
+++ b/Engine/Source/Editor/BlueprintGraph/Private/K2Node_FormatText.cpp
@@ -276,7 +276,7 @@ void UK2Node_FormatText::PostReconstructNode()
 					CurrentPin->DefaultTextValue = FText::GetEmpty();
 
 					// Connect the new node to the existing pin
-					UEdGraphPin* LiteralReturnValuePin = MakeLiteralText->FindPinChecked(TEXT("ReturnValue"));
+					UEdGraphPin* LiteralReturnValuePin = MakeLiteralText->FindPinChecked(UEdGraphSchema_K2::PN_ReturnValue);
 					GetSchema()->TryCreateConnection(LiteralReturnValuePin, CurrentPin);
 
 					++NumPinsFixedUp;
@@ -363,7 +363,7 @@ void UK2Node_FormatText::ExpandNode(class FKismetCompilerContext& CompilerContex
 
 				CompilerContext.MovePinLinksToIntermediate(*ArgumentPin, *ToTextFunction->FindPinChecked(PinName));
 
-				ToTextFunction->FindPinChecked(TEXT("ReturnValue"))->MakeLinkTo(MakeFormatArgumentDataStruct->FindPinChecked(GET_MEMBER_NAME_STRING_CHECKED(FFormatArgumentData, ArgumentValue)));
+				ToTextFunction->FindPinChecked(UEdGraphSchema_K2::PN_ReturnValue)->MakeLinkTo(MakeFormatArgumentDataStruct->FindPinChecked(GET_MEMBER_NAME_STRING_CHECKED(FFormatArgumentData, ArgumentValue)));
 			};
 
 			if (ArgumentPinCategory == UEdGraphSchema_K2::PC_Int)
@@ -395,7 +395,7 @@ void UK2Node_FormatText::ExpandNode(class FKismetCompilerContext& CompilerContex
 				CompilerContext.MovePinLinksToIntermediate(*ArgumentPin, *CallByteToIntFunction->FindPinChecked(TEXT("InByte")));
 
 				// Connect the int output pin to the argument value
-				CallByteToIntFunction->FindPinChecked(TEXT("ReturnValue"))->MakeLinkTo(MakeFormatArgumentDataStruct->FindPinChecked(GET_MEMBER_NAME_STRING_CHECKED(FFormatArgumentData, ArgumentValueInt)));
+				CallByteToIntFunction->FindPinChecked(UEdGraphSchema_K2::PN_ReturnValue)->MakeLinkTo(MakeFormatArgumentDataStruct->FindPinChecked(GET_MEMBER_NAME_STRING_CHECKED(FFormatArgumentData, ArgumentValueInt)));
 			}
 			else if (ArgumentPinCategory == UEdGraphSchema_K2::PC_Byte || ArgumentPinCategory == UEdGraphSchema_K2::PC_Enum)
 			{
diff --git a/Engine/Source/Editor/BlueprintGraph/Private/K2Node_MakeStruct.cpp b/Engine/Source/Editor/BlueprintGraph/Private/K2Node_MakeStruct.cpp
--- a/Engine/Source/Editor/BlueprintGraph/Private/K2Node_MakeStruct.cpp
+++ b/Engine/Source/Editor/BlueprintGraph/Private/K2Node_MakeStruct.cpp
@@ -440,17 +440,17 @@ void UK2Node_MakeStruct::ConvertDeprecatedNode(UEdGraph* Graph, bool bOnlySafeCh
 		if (StructType == TBaseStructure<FRotator>::Get())
 		{
 			MakeNodeFunction = UKismetMathLibrary::StaticClass()->FindFunctionByName(GET_FUNCTION_NAME_CHECKED(UKismetMathLibrary, MakeRotator));
-			OldPinToNewPinMap.Add(TEXT("Rotator"), TEXT("ReturnValue"));
+			OldPinToNewPinMap.Add(TEXT("Rotator"), UEdGraphSchema_K2::PN_ReturnValue);
 		}
 		else if (StructType == TBaseStructure<FVector>::Get())
 		{
 			MakeNodeFunction = UKismetMathLibrary::StaticClass()->FindFunctionByName(GET_FUNCTION_NAME_CHECKED(UKismetMathLibrary, MakeVector));
-			OldPinToNewPinMap.Add(TEXT("Vector"), TEXT("ReturnValue"));
+			OldPinToNewPinMap.Add(TEXT("Vector"), UEdGraphSchema_K2::PN_ReturnValue);
 		}
 		else if (StructType == TBaseStructure<FVector2D>::Get())
 		{
 			MakeNodeFunction = UKismetMathLibrary::StaticClass()->FindFunctionByName(GET_FUNCTION_NAME_CHECKED(UKismetMathLibrary, MakeVector2D));
-			OldPinToNewPinMap.Add(TEXT("Vector2D"), TEXT("ReturnValue"));
+			OldPinToNewPinMap.Add(TEXT("Vector2D"), UEdGraphSchema_K2::PN_ReturnValue);
 		}
 		else
 		{
@@ -459,7 +459,7 @@ void UK2Node_MakeStruct::ConvertDeprecatedNode(UEdGraph* Graph, bool bOnlySafeCh
 
 			if (MakeNodeFunction)
 			{
-				OldPinToNewPinMap.Add(*StructType->GetName(), TEXT("ReturnValue"));
+				OldPinToNewPinMap.Add(*StructType->GetName(), UEdGraphSchema_K2::PN_ReturnValue);
 			}
 		}
 
diff --git a/Engine/Source/Editor/BlueprintGraph/Private/K2Node_MathExpression.cpp b/Engine/Source/Editor/BlueprintGraph/Private/K2Node_MathExpression.cpp
--- a/Engine/Source/Editor/BlueprintGraph/Private/K2Node_MathExpression.cpp
+++ b/Engine/Source/Editor/BlueprintGraph/Private/K2Node_MathExpression.cpp
@@ -1318,7 +1318,7 @@ public:
 		if (RootFragment.IsValid())
 		{
 			// connect the final node of the expression with the math-node's output
-			UEdGraphPin* ReturnPin = ExitNode->CreateUserDefinedPin(TEXT("ReturnValue"), RootFragment->GetOutputType(), EGPD_Input);
+			UEdGraphPin* ReturnPin = ExitNode->CreateUserDefinedPin(UEdGraphSchema_K2::PN_ReturnValue, RootFragment->GetOutputType(), EGPD_Input);
 			if (!RootFragment->ConnectToInput(ReturnPin, MessageLog))
 			{
 				MessageLog.Error(*LOCTEXT("ResultConnectError", "Failed to connect the generated nodes with expression's result pin: '@@'").ToString(),
diff --git a/Engine/Source/Editor/BlueprintGraph/Private/K2Node_PureAssignmentStatement.cpp b/Engine/Source/Editor/BlueprintGraph/Private/K2Node_PureAssignmentStatement.cpp
--- a/Engine/Source/Editor/BlueprintGraph/Private/K2Node_PureAssignmentStatement.cpp
+++ b/Engine/Source/Editor/BlueprintGraph/Private/K2Node_PureAssignmentStatement.cpp
@@ -67,7 +67,7 @@ public:
 
 FName UK2Node_PureAssignmentStatement::VariablePinName = TEXT("Variable");
 FName UK2Node_PureAssignmentStatement::ValuePinName = TEXT("Value");
-FName UK2Node_PureAssignmentStatement::OutputPinName = TEXT("ReturnValue");
+FName UK2Node_PureAssignmentStatement::OutputPinName = UEdGraphSchema_K2::PN_ReturnValue;
 
 
 UK2Node_PureAssignmentStatement::UK2Node_PureAssignmentStatement(const FObjectInitializer& ObjectInitializer)
diff --git a/Engine/Source/Editor/Kismet/Private/BlueprintCompilationManager.cpp b/Engine/Source/Editor/Kismet/Private/BlueprintCompilationManager.cpp
--- a/Engine/Source/Editor/Kismet/Private/BlueprintCompilationManager.cpp
+++ b/Engine/Source/Editor/Kismet/Private/BlueprintCompilationManager.cpp
@@ -2505,7 +2505,6 @@ UClass* FBlueprintCompilationManagerImpl::FastGenerateSkeletonClass(UBlueprint*
 				// Gather all input pins on these nodes, these are 
 				// the outputs of the function:
 				TSet<FName> UsedPinNames;
-				static const FName RetValName = FName(TEXT("ReturnValue"));
 				for(UK2Node_FunctionResult* Node : ReturnNodes)
 				{
 					for(UEdGraphPin* Pin : Node->Pins)
@@ -2521,7 +2520,7 @@ UClass* FBlueprintCompilationManagerImpl::FastGenerateSkeletonClass(UBlueprint*
 								{
 									Param->SetFlags(RF_Transient);
 									// we only tag things as CPF_ReturnParm if the value is named ReturnValue.... this is *terrible* behavior:
-									if(Param->GetFName() == RetValName)
+									if(Param->GetFName() == UEdGraphSchema_K2::PN_ReturnValue)
 									{
 										Param->PropertyFlags |= CPF_ReturnParm;
 									}
diff --git a/Engine/Source/Editor/KismetCompiler/Private/KismetCompiler.cpp b/Engine/Source/Editor/KismetCompiler/Private/KismetCompiler.cpp
--- a/Engine/Source/Editor/KismetCompiler/Private/KismetCompiler.cpp
+++ b/Engine/Source/Editor/KismetCompiler/Private/KismetCompiler.cpp
@@ -1161,11 +1161,10 @@ void FKismetCompilerContext::CreateLocalVariablesForFunction(FKismetFunctionCont
 
 		// Fix up the return value
 		//@todo:  Is there a better way of doing this without mangling code?
-		const FName RetValName = FName(TEXT("ReturnValue"));
 		for (TFieldIterator<FProperty> It(Context.Function); It && (It->PropertyFlags & CPF_Parm); ++It)
 		{
 			FProperty* Property = *It;
-			if ((Property->GetFName() == RetValName) && Property->HasAnyPropertyFlags(CPF_OutParm))
+			if ((Property->GetFName() == UEdGraphSchema_K2::PN_ReturnValue) && Property->HasAnyPropertyFlags(CPF_OutParm))
 			{
 				Property->SetPropertyFlags(CPF_ReturnParm);
 			}
@@ -1951,11 +1950,10 @@ void FKismetCompilerContext::PrecompileFunction(FKismetFunctionContext& Context,
 		{
 			// Fix up the return value - this used to be done by CreateLocalVariablesForFunction.
 			// This should probably be done in CreateParametersForFunction
-			const FName RetValName = FName(TEXT("ReturnValue"));
 			for (TFieldIterator<FProperty> It(Context.Function); It && (It->PropertyFlags & CPF_Parm); ++It)
 			{
 				FProperty* Property = *It;
-				if ((Property->GetFName() == RetValName) && Property->HasAnyPropertyFlags(CPF_OutParm))
+				if ((Property->GetFName() == UEdGraphSchema_K2::PN_ReturnValue) && Property->HasAnyPropertyFlags(CPF_OutParm))
 				{
 					Property->SetPropertyFlags(CPF_ReturnParm);
 				}
