Blueprints: Refactor IK2Node_AddPinInterface to have some more commonly used functions

#rnx

diff --git a/Engine/Source/Editor/BlueprintGraph/Classes/K2Node_AddPinInterface.h b/Engine/Source/Editor/BlueprintGraph/Classes/K2Node_AddPinInterface.h
--- a/Engine/Source/Editor/BlueprintGraph/Classes/K2Node_AddPinInterface.h
+++ b/Engine/Source/Editor/BlueprintGraph/Classes/K2Node_AddPinInterface.h
@@ -13,11 +13,34 @@ class BLUEPRINTGRAPH_API UK2Node_AddPinInterface : public UInterface
 	GENERATED_BODY()
 };
 
+class UEdGraphPin;
+
+/**
+* Interface for adding the small "Add Pin" symbol to a node in the bottom right hand side. 
+* Implementing this interface will provide the API needed to get the UI up and running, but 
+* the actual pin creation/naming is up to the specific node. 
+* 
+* @see UK2Node_CommutativeAssociativeBinaryOperator, UK2Node_DoOnceMultiInput
+*/
 class BLUEPRINTGRAPH_API IK2Node_AddPinInterface
 {
 	GENERATED_BODY()
 
 public:
+
+	static constexpr int32 GetMaxInputPinsNum()
+	{
+		return (TCHAR('Z') - TCHAR('A'));
+	}
+
+	static FName GetNameForAdditionalPin(int32 PinIndex);
+
 	virtual void AddInputPin() { }
 	virtual bool CanAddPin() const { return true; }
+
+	virtual void RemoveInputPin(UEdGraphPin* Pin) { } 
+
+protected:
+
+	virtual bool CanRemovePin(const UEdGraphPin* Pin) const { return true; };
 };
diff --git a/Engine/Source/Editor/BlueprintGraph/Classes/K2Node_CommutativeAssociativeBinaryOperator.h b/Engine/Source/Editor/BlueprintGraph/Classes/K2Node_CommutativeAssociativeBinaryOperator.h
--- a/Engine/Source/Editor/BlueprintGraph/Classes/K2Node_CommutativeAssociativeBinaryOperator.h
+++ b/Engine/Source/Editor/BlueprintGraph/Classes/K2Node_CommutativeAssociativeBinaryOperator.h
@@ -23,13 +23,13 @@ class UK2Node_CommutativeAssociativeBinaryOperator : public UK2Node_CallFunction
 private:
 	const static int32 BinaryOperatorInputsNum = 2;
 
-	static int32 GetMaxInputPinsNum();
-	static FName GetNameForPin(int32 PinIndex);
-
 	FEdGraphPinType GetType() const;
 
 	void AddInputPinInner(int32 AdditionalPinIndex);
-	bool CanRemovePin(const UEdGraphPin* Pin) const;
+
+protected: 
+	bool CanRemovePin(const UEdGraphPin* Pin) const override;
+
 public:
 	BLUEPRINTGRAPH_API UEdGraphPin* FindOutPin() const;
 	BLUEPRINTGRAPH_API UEdGraphPin* FindSelfPin() const;
@@ -37,7 +37,7 @@ public:
 	/** Get TRUE input type (self, etc.. are skipped) */
 	BLUEPRINTGRAPH_API UEdGraphPin* GetInputPin(int32 InputPinIndex);
 
-	BLUEPRINTGRAPH_API void RemoveInputPin(UEdGraphPin* Pin);
+	virtual void RemoveInputPin(UEdGraphPin* Pin) override;
 
 	// UEdGraphNode interface
 	virtual void AllocateDefaultPins() override;
diff --git a/Engine/Source/Editor/BlueprintGraph/Classes/K2Node_DoOnceMultiInput.h b/Engine/Source/Editor/BlueprintGraph/Classes/K2Node_DoOnceMultiInput.h
--- a/Engine/Source/Editor/BlueprintGraph/Classes/K2Node_DoOnceMultiInput.h
+++ b/Engine/Source/Editor/BlueprintGraph/Classes/K2Node_DoOnceMultiInput.h
@@ -33,14 +33,13 @@ private:
 
 	const static int32 NumBaseInputs = 1;
 
-	static int32 GetMaxInputPinsNum();
 	static FText GetNameForPin(int32 PinIndex, bool In);
 
 	FEdGraphPinType GetInType() const;
 	FEdGraphPinType GetOutType() const;
 
 	void AddPinsInner(int32 AdditionalPinIndex);
-	bool CanRemovePin(const UEdGraphPin* Pin) const;
+	bool CanRemovePin(const UEdGraphPin* Pin) const override;
 public:
 	BLUEPRINTGRAPH_API UEdGraphPin* FindOutPin() const;
 	BLUEPRINTGRAPH_API UEdGraphPin* FindSelfPin() const;
@@ -49,7 +48,7 @@ public:
 	BLUEPRINTGRAPH_API UEdGraphPin* GetInputPin(int32 InputPinIndex);
 	BLUEPRINTGRAPH_API UEdGraphPin* GetOutputPin(int32 InputPinIndex);
 
-	BLUEPRINTGRAPH_API void RemoveInputPin(UEdGraphPin* Pin);
+	virtual void RemoveInputPin(UEdGraphPin* Pin) override;
 
 	// UEdGraphNode interface
 	virtual void AllocateDefaultPins() override;
diff --git a/Engine/Source/Editor/BlueprintGraph/Private/K2Node_AddPinInterface.cpp b/Engine/Source/Editor/BlueprintGraph/Private/K2Node_AddPinInterface.cpp
new file mode 100644
--- /dev/null
+++ b/Engine/Source/Editor/BlueprintGraph/Private/K2Node_AddPinInterface.cpp
@@ -0,0 +1,10 @@
+// Copyright Epic Games, Inc. All Rights Reserved.
+
+#include "K2Node_AddPinInterface.h"
+
+FName IK2Node_AddPinInterface::GetNameForAdditionalPin(int32 PinIndex)
+{
+	check(PinIndex < GetMaxInputPinsNum());
+	const FName Name(*FString::Chr(TCHAR('A') + PinIndex));
+	return Name;
+}
\ No newline at end of file
diff --git a/Engine/Source/Editor/BlueprintGraph/Private/K2Node_CommutativeAssociativeBinaryOperator.cpp b/Engine/Source/Editor/BlueprintGraph/Private/K2Node_CommutativeAssociativeBinaryOperator.cpp
--- a/Engine/Source/Editor/BlueprintGraph/Private/K2Node_CommutativeAssociativeBinaryOperator.cpp
+++ b/Engine/Source/Editor/BlueprintGraph/Private/K2Node_CommutativeAssociativeBinaryOperator.cpp
@@ -12,18 +12,6 @@
 
 #define LOCTEXT_NAMESPACE "CommutativeAssociativeBinaryOperatorNode"
 
-int32 UK2Node_CommutativeAssociativeBinaryOperator::GetMaxInputPinsNum()
-{
-	return (TCHAR('Z') - TCHAR('A'));
-}
-
-FName UK2Node_CommutativeAssociativeBinaryOperator::GetNameForPin(int32 PinIndex)
-{
-	check(PinIndex < GetMaxInputPinsNum());
-	const FName Name(*FString::Chr(TCHAR('A') + PinIndex));
-	return Name;
-}
-
 UK2Node_CommutativeAssociativeBinaryOperator::UK2Node_CommutativeAssociativeBinaryOperator(const FObjectInitializer& ObjectInitializer)
 	: Super(ObjectInitializer)
 {
@@ -168,7 +156,7 @@ void UK2Node_CommutativeAssociativeBinaryOperator::AddInputPinInner(int32 Additi
 		InputType.PinCategory, 
 		InputType.PinSubCategory, 
 		InputType.PinSubCategoryObject.Get(), 
-		GetNameForPin(AdditionalPinIndex + BinaryOperatorInputsNum),
+		GetNameForAdditionalPin(AdditionalPinIndex + BinaryOperatorInputsNum),
 		PinParams
 	);
 }
@@ -206,7 +194,7 @@ void UK2Node_CommutativeAssociativeBinaryOperator::RemoveInputPin(UEdGraphPin* P
 				UEdGraphPin* LocalPin = Pins[PinIndex];
 				if(LocalPin && (LocalPin != OutPin) && (LocalPin != SelfPin))
 				{
-					const FName PinName = GetNameForPin(NameIndex);
+					const FName PinName = GetNameForAdditionalPin(NameIndex);
 					if (PinName != LocalPin->PinName)
 					{
 						LocalPin->Modify();
diff --git a/Engine/Source/Editor/BlueprintGraph/Private/K2Node_DoOnceMultiInput.cpp b/Engine/Source/Editor/BlueprintGraph/Private/K2Node_DoOnceMultiInput.cpp
--- a/Engine/Source/Editor/BlueprintGraph/Private/K2Node_DoOnceMultiInput.cpp
+++ b/Engine/Source/Editor/BlueprintGraph/Private/K2Node_DoOnceMultiInput.cpp
@@ -50,11 +50,6 @@ FText UK2Node_DoOnceMultiInput::GetNodeTitle(ENodeTitleType::Type TitleType) con
 	return LOCTEXT("DoOnceMultiInput", "DoOnce MultiInput");
 }
 
-int32 UK2Node_DoOnceMultiInput::GetMaxInputPinsNum()
-{
-	return (TCHAR('Z') - TCHAR('A'));
-}
-
 FText UK2Node_DoOnceMultiInput::GetNameForPin(int32 PinIndex, bool In)
 {
 	check(PinIndex < GetMaxInputPinsNum());
