Blueprints: Add a utility function "FindBestMatchingFunction" call for Blueprint Type promotion allowing users to not only find the lowest function, but also the one with the most matching arguments.

diff --git a/Engine/Source/Editor/BlueprintGraph/Classes/BlueprintTypePromotion.h b/Engine/Source/Editor/BlueprintGraph/Classes/BlueprintTypePromotion.h
--- a/Engine/Source/Editor/BlueprintGraph/Classes/BlueprintTypePromotion.h
+++ b/Engine/Source/Editor/BlueprintGraph/Classes/BlueprintTypePromotion.h
@@ -9,7 +9,7 @@
 #include "HAL/IConsoleManager.h"
 
 /** Contains behavior needed to handle type promotion in blueprints */
-class FTypePromotion : private FNoncopyable
+class BLUEPRINTGRAPH_API FTypePromotion : private FNoncopyable
 {
 public:
 
@@ -33,6 +33,12 @@ public:
 	*/
 	static UFunction* FindLowestMatchingFunc(const FString& Operation, const FEdGraphPinType& InputType, TArray<UFunction*>& OutPossibleFunctions);
 
+	/**
+	* Find the function that is the best match given the pins to consider. 
+	* Ex: Given "Add" operator and an array of two Vector pins, it will return "Add_VectorVector"
+	*/
+	static UFunction* FindBestMatchingFunc(const FString& Operation, const TArray<UEdGraphPin*>& PinsToConsider);
+
 	/** Returns all functions for a specific operation. Will empty the given array and populate it with UFunction pointers */
 	static void GetAllFuncsForOp(const FString& Operation, TArray<UFunction*>& OutFuncs);
 
@@ -115,6 +121,8 @@ private:
 	FEdGraphPinType GetPromotedType_Internal(const TArray<UEdGraphPin*>& WildcardPins) const;
 
 	UFunction* FindLowestMatchingFunc_Internal(const FString& Operation, const FEdGraphPinType& InputType, TArray<UFunction*>& OutPossibleFunctions);
+	
+	UFunction* FindBestMatchingFunc_Internal(const FString& Operation, const TArray<UEdGraphPin*>& PinsToConsider);
 
 	void GetAllFuncsForOp_Internal(const FString& Operation, TArray<UFunction*>& OutFuncs);
 
diff --git a/Engine/Source/Editor/BlueprintGraph/Private/BlueprintTypePromotion.cpp b/Engine/Source/Editor/BlueprintGraph/Private/BlueprintTypePromotion.cpp
--- a/Engine/Source/Editor/BlueprintGraph/Private/BlueprintTypePromotion.cpp
+++ b/Engine/Source/Editor/BlueprintGraph/Private/BlueprintTypePromotion.cpp
@@ -374,6 +374,11 @@ UFunction* FTypePromotion::FindLowestMatchingFunc(const FString& Operation, cons
 	return FTypePromotion::Get().FindLowestMatchingFunc_Internal(Operation, InputType, OutPossibleFunctions);
 }
 
+UFunction* FTypePromotion::FindBestMatchingFunc(const FString& Operation, const TArray<UEdGraphPin*>& PinsToConsider)
+{
+	return FTypePromotion::Get().FindBestMatchingFunc_Internal(Operation, PinsToConsider);
+}
+
 static bool PropertyCompatibleWithPin(const FProperty* Param, FEdGraphPinType const& TypeToMatch)
 {
 	const UEdGraphSchema_K2* Schema = GetDefault<UEdGraphSchema_K2>();
@@ -455,6 +460,61 @@ UFunction* FTypePromotion::FindLowestMatchingFunc_Internal(const FString& Operat
 	return LowestFunc;
 }
 
+UFunction* FTypePromotion::FindBestMatchingFunc_Internal(const FString& Operation, const TArray<UEdGraphPin*>& PinsToConsider)
+{
+	TRACE_CPUPROFILER_EVENT_SCOPE(FTypePromotionTable::FindBestMatchingFunc_Internal);
+
+	const FFunctionsList* FuncList = OperatorTable.Find(Operation);
+	if (!FuncList)
+	{
+		return nullptr;
+	}
+
+	// Track the function with the best score
+	UFunction* BestFunc = nullptr;
+	int32 BestScore = -1;
+
+	const UEdGraphSchema_K2* Schema = GetDefault<UEdGraphSchema_K2>();
+		
+	// We have to keep track of what pins we have already given points for, otherwise 
+	// we will end up giving the same pin multiple points.
+	TSet<const UEdGraphPin*> CheckedPins;
+		
+	for (UFunction* Func : *FuncList)
+	{
+		int32 FuncScore = -1;
+		CheckedPins.Reset();
+
+		// For each property in the func, see if it matches any of the given pins
+		for (TFieldIterator<FProperty> PropIt(Func); PropIt && (PropIt->PropertyFlags & CPF_Parm); ++PropIt)
+		{
+			const FProperty* Param = *PropIt;
+			FEdGraphPinType ParamType;
+			if (Schema->ConvertPropertyToPinType(Param, /* out */ ParamType))
+			{
+				for (const UEdGraphPin* Pin : PinsToConsider)
+				{
+					// Give a point for each function parameter that matches up with a pin to consider
+					if (!CheckedPins.Contains(Pin) && Schema->ArePinTypesEquivalent(ParamType, Pin->PinType))
+					{
+						++FuncScore;
+						CheckedPins.Add(Pin);
+						break;
+					}
+				}
+			}
+		}
+
+		// If this function has the best score, then update it! 
+		if (FuncScore > BestScore)
+		{
+			BestScore = FuncScore;
+			BestFunc = Func;
+		}
+	}
+	return BestFunc;
+}
+
 void FTypePromotion::GetAllFuncsForOp(const FString& Operation, TArray<UFunction*>& OutFuncs)
 {
 	return FTypePromotion::Get().GetAllFuncsForOp_Internal(Operation, OutFuncs);
