Blueprints: Add unit test for "FTypePromotino::FindBestMatchingFunc"

#rnx

diff --git a/Engine/Source/Editor/BlueprintGraph/Private/Tests/TypePromotionTests.cpp b/Engine/Source/Editor/BlueprintGraph/Private/Tests/TypePromotionTests.cpp
--- a/Engine/Source/Editor/BlueprintGraph/Private/Tests/TypePromotionTests.cpp
+++ b/Engine/Source/Editor/BlueprintGraph/Private/Tests/TypePromotionTests.cpp
@@ -49,22 +49,22 @@ bool FTypePromotionTest::RunTest(const FString& Parameters)
 
 IMPLEMENT_SIMPLE_AUTOMATION_TEST(FFindPromotedFunc, "Blueprints.Compiler.FindPromotedFunc", EAutomationTestFlags::ApplicationContextMask | EAutomationTestFlags::SmokeFilter)
 
+#define MakeTestPin(OwningNode, PinArray, InPinName, InPinType)		UEdGraphPin* InPinName = UEdGraphPin::CreatePin(OwningNode);		\
+																	InPinName->PinType.PinCategory = InPinType;							\
+																	PinArray.Add(InPinName);		
+
 bool FFindPromotedFunc::RunTest(const FString& Parameters)
 {
 	TArray<UEdGraphPin*> PinTypes = {};
 
-	UEdGraphNode* TestNode = NewObject<UEdGraphNode>();
-
-#define MakeTestPin(InPinName, InPinType)		UEdGraphPin* InPinName = UEdGraphPin::CreatePin(TestNode);		\
-												InPinName->PinType.PinCategory = InPinType;						\
-												PinTypes.Add(InPinName);										
+	UEdGraphNode* TestNode = NewObject<UEdGraphNode>();			
 	
 	const FTypePromotion& TypePromo = FTypePromotion::Get();
 
-	MakeTestPin(DoublePin, UEdGraphSchema_K2::PC_Double);
-	MakeTestPin(FloatPin, UEdGraphSchema_K2::PC_Float);
-	MakeTestPin(Int32Pin, UEdGraphSchema_K2::PC_Int);
-	MakeTestPin(Int64Pin, UEdGraphSchema_K2::PC_Int64);
+	MakeTestPin(TestNode, PinTypes, DoublePin, UEdGraphSchema_K2::PC_Double);
+	MakeTestPin(TestNode, PinTypes, FloatPin, UEdGraphSchema_K2::PC_Float);
+	MakeTestPin(TestNode, PinTypes, Int32Pin, UEdGraphSchema_K2::PC_Int);
+	MakeTestPin(TestNode, PinTypes, Int64Pin, UEdGraphSchema_K2::PC_Int64);
 
 	// Add Operation
 	{
@@ -115,12 +115,90 @@ bool FFindPromotedFunc::RunTest(const FString& Parameters)
 			TestPin->MarkPendingKill();
 		}
 	}
+	PinTypes.Empty();
+	TestNode->MarkPendingKill();
 
-#undef MakeTestPin
+	return true;
+}
+
+
+IMPLEMENT_SIMPLE_AUTOMATION_TEST(FFindBestMatchingFunc, "Blueprints.Compiler.FindBestMatchingFunc", EAutomationTestFlags::ApplicationContextMask | EAutomationTestFlags::SmokeFilter)
+bool FFindBestMatchingFunc::RunTest(const FString& Parameters)
+{
+	TArray<UEdGraphPin*> PinTypes = {};
+
+	UEdGraphNode* TestNode = NewObject<UEdGraphNode>();
+
+	MakeTestPin(TestNode, PinTypes, DoublePinA, UEdGraphSchema_K2::PC_Double);
+	MakeTestPin(TestNode, PinTypes, DoublePinB, UEdGraphSchema_K2::PC_Double);
 
+	MakeTestPin(TestNode, PinTypes, FloatPinA, UEdGraphSchema_K2::PC_Float);
+	MakeTestPin(TestNode, PinTypes, FloatPinB, UEdGraphSchema_K2::PC_Float);
+	
+	MakeTestPin(TestNode, PinTypes, IntPinA, UEdGraphSchema_K2::PC_Int);
+
+	// Add_DoubleDouble
+	{
+		TArray<UEdGraphPin*> TestPins =
+		{
+			DoublePinA, DoublePinB
+		};
+
+		const UFunction* AddDoubleFunc = FTypePromotion::FindBestMatchingFunc(TEXT("Add"), TestPins);
+		static const FName ExpectedName = TEXT("Add_DoubleDouble");
+		
+		if (TestNotNull(TEXT("Add_DoubleDouble Null check"), AddDoubleFunc))
+		{
+			TestEqual(TEXT("Add_DoubleDouble Name Check"), AddDoubleFunc->GetFName(), ExpectedName);
+		}
+	}
+
+	// Subtract_FloatFloat
+	{
+		TArray<UEdGraphPin*> TestPins =
+		{
+			FloatPinA, FloatPinB
+		};
+
+		const UFunction* SubtractFloatFunc = FTypePromotion::FindBestMatchingFunc(TEXT("Subtract"), TestPins);
+		static const FName ExpectedName = TEXT("Subtract_FloatFloat");
+
+		if (TestNotNull(TEXT("Subtract_FloatFloat null check"), SubtractFloatFunc))
+		{
+			TestEqual(TEXT("Subtract_FloatFloat Name Check"), SubtractFloatFunc->GetFName(), ExpectedName);
+		}
+	}
+
+	// Multiply_IntFloat
+	{
+		TArray<UEdGraphPin*> TestPins =
+		{
+			FloatPinA, FloatPinB, IntPinA
+		};
+
+		const UFunction* MulFloatFunc = FTypePromotion::FindBestMatchingFunc(TEXT("Multiply"), TestPins);
+		static const FName ExpectedName = TEXT("Multiply_IntFloat");
+
+		if (TestNotNull(TEXT("Multiply_IntFloat null check"), MulFloatFunc))
+		{
+			TestEqual(TEXT("Multiply_IntFloat Name Check"), MulFloatFunc->GetFName(), ExpectedName);
+		}
+	}
+
+	// Clear our test pins
+	for (UEdGraphPin* TestPin : PinTypes)
+	{
+		if (TestPin)
+		{
+			TestPin->MarkPendingKill();
+		}
+	}
+	PinTypes.Empty();
 	TestNode->MarkPendingKill();
 
 	return true;
 }
 
+#undef MakeTestPin
+
 #endif //WITH_DEV_AUTOMATION_TESTS
\ No newline at end of file
