Blueprints: Implement IK2Node_AddPinInterface in the PromotableOperator node

#rnx

diff --git a/Engine/Source/Editor/BlueprintGraph/Classes/K2Node_PromotableOperator.h b/Engine/Source/Editor/BlueprintGraph/Classes/K2Node_PromotableOperator.h
--- a/Engine/Source/Editor/BlueprintGraph/Classes/K2Node_PromotableOperator.h
+++ b/Engine/Source/Editor/BlueprintGraph/Classes/K2Node_PromotableOperator.h
@@ -3,14 +3,15 @@
 #pragma once
 
 #include "CoreMinimal.h"
-#include "UObject/ObjectMacros.h"
-#include "EdGraph/EdGraphPin.h"
+#include "K2Node_AddPinInterface.h"
 #include "K2Node_CallFunction.h"
 #include "K2Node_PromotableOperator.generated.h"
 
+class UEdGraphPin;
+
 /** The promotable operator node allows for pin types to be promoted to others, i.e. float to double */
 UCLASS(MinimalAPI)
-class UK2Node_PromotableOperator : public UK2Node_CallFunction
+class UK2Node_PromotableOperator : public UK2Node_CallFunction, public IK2Node_AddPinInterface
 {
 	GENERATED_UCLASS_BODY()
 
@@ -29,11 +30,29 @@ public:
 	virtual bool IsConnectionDisallowed(const UEdGraphPin* MyPin, const UEdGraphPin* OtherPin, FString& OutReason) const override;
 	virtual void ReallocatePinsDuringReconstruction(TArray<UEdGraphPin*>& OldPins) override;
 	virtual void AutowireNewNode(UEdGraphPin* FromPin) override;
-	virtual bool IsActionFilteredOut(class FBlueprintActionFilter const& Filter) override;
 	// End of UK2Node interface
 
+	// IK2Node_AddPinInterface interface
+	virtual void AddInputPin() override;
+	virtual bool CanAddPin() const override;
+	virtual bool CanRemovePin(const UEdGraphPin* Pin) const override;
+	virtual void RemoveInputPin(UEdGraphPin* Pin) override;
+	// end IK2Node_AddPinInterface interface
+
 private:
 
+	/**
+	 * Add an additional pin to this node based on it's index. Will create a new wildcard pin
+	 * with the appropriate name and notify that this pin has changed
+	 */
+	UEdGraphPin* AddInputPinImpl(int32 PinIndex);
+
+	/** Returns true if this pin was added via the IK2Node_AddPinInterface interface */
+	bool IsAdditionalPin(const UEdGraphPin* Pin) const;
+		
+	/** Gets the additional pin that was created at this index */
+	UEdGraphPin* GetAdditionalPin(int32 PinIndex) const;
+		
 	/** Helper function to recombine all split pins that this node may have */
 	void RecombineAllSplitPins();
 
@@ -51,7 +70,9 @@ private:
 	* 
 	* @return	True if the intermediate connection was made successfully	
 	*/
-	bool CreateIntermediateCast(FKismetCompilerContext& CompilerContext, UEdGraph* SourceGraph, UEdGraphPin* InputPin, UEdGraphPin* OutputPin);
+	bool CreateIntermediateCast(UK2Node_CallFunction* SourceNode, FKismetCompilerContext& CompilerContext, UEdGraph* SourceGraph, UEdGraphPin* InputPin, UEdGraphPin* OutputPin);
+
+	UK2Node_CallFunction* CreateIntermediateNode(UK2Node_CallFunction* PreviousNode, const UFunction* const OpFunction, FKismetCompilerContext& CompilerContext, UEdGraph* SourceGraph);
 
 	/** Helper to make sure we have the most up to date operation name. Returns true upon success */
 	bool UpdateOpName();
@@ -82,6 +103,10 @@ private:
 	UPROPERTY(Transient)
 	TArray<UFunction*> PossibleConversions;
 
+	/** The current number of additional pins on this node */
+	UPROPERTY()
+	int32 NumAdditionalInputs;
+
 public:
 
 	/** Returns the first pin with the EGPD_Output direction */
diff --git a/Engine/Source/Editor/BlueprintGraph/Private/K2Node_PromotableOperator.cpp b/Engine/Source/Editor/BlueprintGraph/Private/K2Node_PromotableOperator.cpp
--- a/Engine/Source/Editor/BlueprintGraph/Private/K2Node_PromotableOperator.cpp
+++ b/Engine/Source/Editor/BlueprintGraph/Private/K2Node_PromotableOperator.cpp
@@ -11,6 +11,7 @@
 #include "Framework/Commands/UIAction.h"
 #include "Kismet2/BlueprintEditorUtils.h"
 #include "Kismet2/WildcardNodeUtils.h"
+#include "Kismet2/CompilerResultsLog.h"
 
 #define LOCTEXT_NAMESPACE "PromotableOperatorNode"
 
@@ -19,6 +20,7 @@
 
 static const FName InputPinA_Name = FName(TEXT("A"));
 static const FName InputPinB_Name = FName(TEXT("B"));
+static const int32 NumFunctionInputs = 2;
 
 ///////////////////////////////////////////////////////////
 // UK2Node_PromotableOperator
@@ -28,6 +30,7 @@ UK2Node_PromotableOperator::UK2Node_PromotableOperator(const FObjectInitializer&
 {
 	UpdateOpName();
 	OrphanedPinSaveMode = ESaveOrphanPinMode::SaveAllButExec;
+	NumAdditionalInputs = 0;
 }
 
 ///////////////////////////////////////////////////////////
@@ -39,12 +42,50 @@ void UK2Node_PromotableOperator::AllocateDefaultPins()
 	FWildcardNodeUtils::CreateWildcardPin(this, InputPinB_Name, EGPD_Input);
 
 	FWildcardNodeUtils::CreateWildcardPin(this, UEdGraphSchema_K2::PN_ReturnValue, EGPD_Output);
+
+	// Create any additional input pin. Their appropriate type is determined in ReallocatePinsDuringReconstruction
+	// because we cannot get a promoted type with no links to the pin.
+	for (int32 i = NumFunctionInputs; i < (NumAdditionalInputs + NumFunctionInputs); ++i)
+	{
+		AddInputPinImpl(i);
+	}
 }
 
 void UK2Node_PromotableOperator::GetNodeContextMenuActions(UToolMenu* Menu, UGraphNodeContextMenuContext* Context) const
 {
 	Super::GetNodeContextMenuActions(Menu, Context);
 
+	static const FName PromotableOperatorNodeName = FName("PromotableOperator");
+	static const FText PromotableOperatorStr = LOCTEXT("PromotableOperatorNode", "Operator Node");
+
+	// Add the option to remove a pin via the context menu
+	if (CanRemovePin(Context->Pin))
+	{
+		FToolMenuSection& Section = Menu->AddSection(PromotableOperatorNodeName, PromotableOperatorStr);
+		Section.AddMenuEntry(
+			"RemovePin",
+			LOCTEXT("RemovePin", "Remove pin"),
+			LOCTEXT("RemovePinTooltip", "Remove this input pin"),
+			FSlateIcon(),
+			FUIAction(
+				FExecuteAction::CreateUObject(const_cast<UK2Node_PromotableOperator*>(this), &UK2Node_PromotableOperator::RemoveInputPin, const_cast<UEdGraphPin*>(Context->Pin))
+			)
+		);
+	}
+	else if (CanAddPin())
+	{
+		FToolMenuSection& Section = Menu->AddSection(PromotableOperatorNodeName, PromotableOperatorStr);
+		Section.AddMenuEntry(
+			"AddPin",
+			LOCTEXT("AddPin", "Add pin"),
+			LOCTEXT("AddPinTooltip", "Add another input pin"),
+			FSlateIcon(),
+			FUIAction(
+				FExecuteAction::CreateUObject(const_cast<UK2Node_PromotableOperator*>(this), &UK2Node_PromotableOperator::AddInputPin)
+			)
+		);
+	}
+
 	// If there are possible function conversions that can happen 
 	if (Context->Pin && PossibleConversions.Num() > 0 && !Context->bIsDebugging && HasAnyConnectionsOrDefaults())
 	{
@@ -108,11 +149,11 @@ void UK2Node_PromotableOperator::ExpandNode(FKismetCompilerContext& CompilerCont
 		return;
 	}
 
-	TArray<UEdGraphPin*> InputPins = GetInputPins();
-	UEdGraphPin* OutputPin = GetOutputPin();
+	UEdGraphPin* OriginalOutputPin = GetOutputPin();
+	TArray<UEdGraphPin*> OriginalInputPins = GetInputPins();
 
 	// Our operator function has been determined on pin connection change
-	const UFunction* const OpFunction = GetTargetFunction();
+	UFunction* OpFunction = GetTargetFunction();
 
 	if (!OpFunction)
 	{
@@ -120,63 +161,130 @@ void UK2Node_PromotableOperator::ExpandNode(FKismetCompilerContext& CompilerCont
 		CompilerContext.MessageLog.Error(TEXT("Could not find matching op function during expansion on '@@'!"), this);
 		return;
 	}
+	const UEdGraphSchema_K2* Schema = CompilerContext.GetSchema();
 
-	// Now to actually go through the promotion process on pins that need to be promoted to 
-	// fit our function signature! 
+	/** Helper struct to gather the necessary pins we need to create redirections */
+	struct FIntermediateCastPinHelper
+	{
+		UEdGraphPin* InputA = nullptr;
+		UEdGraphPin* InputB = nullptr;
+		UEdGraphPin* OutputPin = nullptr;
+		UEdGraphPin* SelfPin = nullptr;
 
-	// Spawn an intermediate K2NodeCommunicative op node of that type
-	UK2Node_CallFunction* NewOperator = SourceGraph->CreateIntermediateNode<UK2Node_CallFunction>();
-	NewOperator->SetFromFunction(OpFunction);
-	NewOperator->AllocateDefaultPins();
+		explicit FIntermediateCastPinHelper(UK2Node_CallFunction* NewOperator)
+		{
+			check(NewOperator);
+			SelfPin = NewOperator->FindPin(UEdGraphSchema_K2::PN_Self);
 
-	// Move this node next to the thing it was linked to
-	NewOperator->NodePosY = this->NodePosY;
-	NewOperator->NodePosX = this->NodePosX + 8;
+			// Find inputs and outputs
+			for (UEdGraphPin* Pin : NewOperator->Pins)
+			{
+				if (Pin == SelfPin)
+				{
+					continue;
+				}
 
-	CompilerContext.MessageLog.NotifyIntermediateObjectCreation(NewOperator, this);
+				if (Pin->Direction == EEdGraphPinDirection::EGPD_Input)
+				{
+					if (!InputA)
+					{
+						InputA = Pin;
+					}
+					else if (!InputB)
+					{
+						InputB = Pin;
+					}
+				}
+				else if (Pin->Direction == EEdGraphPinDirection::EGPD_Output)
+				{
+					OutputPin = Pin;
+				}
+			}
+		}
+
+		~FIntermediateCastPinHelper() = default;
+	};
 
-	UEdGraphPin* NewOperatorInputA = nullptr;
-	UEdGraphPin* NewOperatorInputB = nullptr;
-	UEdGraphPin* NewOperatorOutputPin = nullptr;
-	UEdGraphPin* NewOperatorSelfPin = NewOperator->FindPin(UEdGraphSchema_K2::PN_Self);
+	UK2Node_CallFunction* PrevIntermediateNode = nullptr;
+	UEdGraphPin* PrevOutputPin = nullptr;
 
-	for (UEdGraphPin* Pin : NewOperator->Pins)
+	// Create cast from original 2 inputs to the first intermediate node
 	{
-		if (Pin == NewOperatorSelfPin)
+		UFunction* BestFunc = OpFunction;
 		{
-			continue;
+			TArray<UEdGraphPin*> PinsToConsider =
+			{
+				OriginalInputPins[0],
+				OriginalInputPins[1],
+				OriginalOutputPin
+			};
+
+			if (UFunction* Func = FTypePromotion::FindBestMatchingFunc(OperationName, PinsToConsider))
+			{
+				BestFunc = Func;
+			}
 		}
 
-		if (Pin->Direction == EEdGraphPinDirection::EGPD_Input)
+		PrevIntermediateNode = CreateIntermediateNode(this, BestFunc, CompilerContext, SourceGraph);
+		FIntermediateCastPinHelper NewOpHelper(PrevIntermediateNode);
+		PrevOutputPin = PrevIntermediateNode->FindPin(UEdGraphSchema_K2::PN_ReturnValue, EGPD_Output);
+
+		const bool bPinASuccess = UK2Node_PromotableOperator::CreateIntermediateCast(this, CompilerContext, SourceGraph, OriginalInputPins[0], NewOpHelper.InputA);
+		const bool bPinBSuccess = UK2Node_PromotableOperator::CreateIntermediateCast(this, CompilerContext, SourceGraph, OriginalInputPins[1], NewOpHelper.InputB);
+
+		if (!bPinASuccess || !bPinBSuccess)
+		{
+			CompilerContext.MessageLog.Error(TEXT("'@@' could not successfuly expand pins!"), PrevIntermediateNode);
+		}
+	}
+	
+	// Loop through all the additional inputs, create a new node of this function and connecting inputs as necessary 
+	for (int32 i = NumFunctionInputs; i < NumAdditionalInputs + NumFunctionInputs; ++i)
+	{
+		check(i > 0 && i < OriginalInputPins.Num());
+		FIntermediateCastPinHelper PrevNodeHelper(PrevIntermediateNode);
+
+		// Find the best matching function that this intermediate node should use
+		// so that we can avoid unnecessary conversion nodes and casts
+		UFunction* BestMatchingFunc = OpFunction;
 		{
-			if (!NewOperatorInputA) 
+			TArray<UEdGraphPin*> PinsToConsider = 
 			{
-				NewOperatorInputA = Pin;
-			}
-			else if (!NewOperatorInputB) 
+				PrevNodeHelper.OutputPin,
+				OriginalInputPins[i],
+				OriginalOutputPin
+			};
+			
+			if (UFunction* Func = FTypePromotion::FindBestMatchingFunc(OperationName, PinsToConsider))
 			{
-				NewOperatorInputB = Pin;
+				BestMatchingFunc = Func;
 			}
 		}
-		else if (Pin->Direction == EEdGraphPinDirection::EGPD_Output)
+
+		UK2Node_CallFunction* NewIntermediateNode = CreateIntermediateNode(PrevIntermediateNode, BestMatchingFunc, CompilerContext, SourceGraph);
+		FIntermediateCastPinHelper NewOpHelper(NewIntermediateNode);
+
+		// Connect the output pin of the previous intermediate node, to the input of the new one
+		const bool bPinASuccess = CreateIntermediateCast(PrevIntermediateNode, CompilerContext, SourceGraph, NewOpHelper.InputA, PrevOutputPin);
+		
+		// Connect the original node's pin to the newly created intermediate node's B Pin
+		const bool bPinBSuccess = CreateIntermediateCast(this, CompilerContext, SourceGraph, OriginalInputPins[i], NewOpHelper.InputB);
+		
+		if (!bPinASuccess || !bPinBSuccess)
 		{
-			NewOperatorOutputPin = Pin;
+			CompilerContext.MessageLog.Error(TEXT("'@@' could not successfuly expand additional pins!"), PrevIntermediateNode);
 		}
-	}
 
-	check(NewOperatorOutputPin);
-
-	// Create some auto casts if they are necessary
-	const bool bPinASuccess = CreateIntermediateCast(CompilerContext, SourceGraph, InputPins[0], NewOperatorInputA);
-	const bool bPinBSuccess = CreateIntermediateCast(CompilerContext, SourceGraph, InputPins[1], NewOperatorInputB);
+		// Track what the previous node is so that we can connect it's output appropriately
+		PrevOutputPin = NewOpHelper.OutputPin;
+		PrevIntermediateNode = NewIntermediateNode;
+	}
 
-	if (!bPinASuccess || !bPinBSuccess)
+	// Make the final output connection that we need
+	if (OriginalOutputPin && PrevOutputPin)
 	{
-		CompilerContext.MessageLog.Error(TEXT("'@@' could not successfuly expand pins!"), this);
+		CompilerContext.MovePinLinksToIntermediate(*OriginalOutputPin, *PrevOutputPin);
 	}
-
-	// Connect intermediate node output to this nodes output
-	CompilerContext.MovePinLinksToIntermediate(*GetOutputPin(), *NewOperatorOutputPin);
 }
 
 void UK2Node_PromotableOperator::NotifyPinConnectionListChanged(UEdGraphPin* ChangedPin)
@@ -199,9 +307,8 @@ void UK2Node_PromotableOperator::NotifyPinConnectionListChanged(UEdGraphPin* Cha
 	}
 	// If the pin that was connected is linked to a wildcard pin, then we should make it a wildcard
 	// and do nothing else.
-	else if (ChangedPin->GetOwningNode() == this && FWildcardNodeUtils::IsLinkedToWildcard(ChangedPin))
+	else if (FWildcardNodeUtils::IsWildcardPin(ChangedPin) && ChangedPin->GetOwningNode() == this && FWildcardNodeUtils::IsLinkedToWildcard(ChangedPin))
 	{
-		ChangedPin->PinType = FWildcardNodeUtils::GetDefaultWildcardPinType();
 		return;
 	}
 
@@ -259,6 +366,15 @@ void UK2Node_PromotableOperator::PostReconstructNode()
 		// Allocate default pins will have been called before this, which means we are reset to wildcard state
 		// We need to Update the pins to be the proper function again
 		UpdatePinsFromFunction(GetTargetFunction());
+
+		for (UEdGraphPin* AddPin : Pins)
+		{
+			if (IsAdditionalPin(AddPin) && AddPin->LinkedTo.Num() > 0)
+			{
+				FEdGraphPinType TypeToSet = FTypePromotion::GetPromotedType(AddPin->LinkedTo);
+				AddPin->PinType = TypeToSet;
+			}
+		}
 	}
 }
 
@@ -319,6 +435,23 @@ void UK2Node_PromotableOperator::ReallocatePinsDuringReconstruction(TArray<UEdGr
 	UpdatePinsFromFunction(GetTargetFunction(), nullptr);
 
 	Super::ReallocatePinsDuringReconstruction(OldPins);
+
+	// We need to fix up any additional pins that may have been created as a wildcard pin
+	int32 AdditionalPinsFixed = 0;
+
+	// Additional Pin creation here? Check for orphan pins here and see if we can re-create them
+	for (UEdGraphPin* OldPin : OldPins)
+	{
+		if (IsAdditionalPin(OldPin))
+		{
+			if (UEdGraphPin* AddPin = GetAdditionalPin(AdditionalPinsFixed + NumFunctionInputs))
+			{
+				AddPin->PinType = OldPin->PinType;
+				AddPin->DefaultValue = OldPin->DefaultValue;
+				++AdditionalPinsFixed;
+			}			
+		}
+	}
 }
 
 void UK2Node_PromotableOperator::AutowireNewNode(UEdGraphPin* ChangedPin)
@@ -328,14 +461,123 @@ void UK2Node_PromotableOperator::AutowireNewNode(UEdGraphPin* ChangedPin)
 	NotifyPinConnectionListChanged(ChangedPin);
 }
 
-bool UK2Node_PromotableOperator::IsActionFilteredOut(FBlueprintActionFilter const& Filter)
+///////////////////////////////////////////////////////////
+// IK2Node_AddPinInterface
+
+void UK2Node_PromotableOperator::AddInputPin()
 {
-	return false;
+	if (CanAddPin())
+	{
+		FScopedTransaction Transaction(LOCTEXT("AddPinPromotableOperator", "AddPin"));
+		Modify();
+
+		AddInputPinImpl(NumFunctionInputs + NumAdditionalInputs);
+		++NumAdditionalInputs;
+
+		FBlueprintEditorUtils::MarkBlueprintAsStructurallyModified(GetBlueprint());
+	}
+}
+
+bool UK2Node_PromotableOperator::CanAddPin() const
+{
+	return ((NumAdditionalInputs + NumFunctionInputs) < GetMaxInputPinsNum()) &&
+		!FTypePromotion::IsComparisonFunc(GetTargetFunction());
+}
+
+bool UK2Node_PromotableOperator::CanRemovePin(const UEdGraphPin* Pin) const
+{
+	if (!Pin)
+	{
+		return false;
+	}
+
+	// We cannot remove the first two inputs from a function, because they are always there
+	const bool bIsBasePin = (Pin->PinName == InputPinA_Name || Pin->PinName == InputPinB_Name);
+
+	return (
+		!bIsBasePin &&
+		Pin->ParentPin == nullptr &&
+		NumAdditionalInputs > 0 &&
+		INDEX_NONE != Pins.IndexOfByKey(Pin) &&
+		Pin->Direction == EEdGraphPinDirection::EGPD_Input
+	);
+}
+
+void UK2Node_PromotableOperator::RemoveInputPin(UEdGraphPin* Pin)
+{
+	if (CanRemovePin(Pin))
+	{
+		FScopedTransaction Transaction(LOCTEXT("RemovePinPromotableOperator", "RemovePin"));
+		Modify();
+
+		if (RemovePin(Pin))
+		{
+			--NumAdditionalInputs;
+
+			int32 NameIndex = 0;
+			const UEdGraphPin* OutPin = GetOutputPin();
+			const UEdGraphPin* SelfPin = FindPin(UEdGraphSchema_K2::PN_Self);
+
+			for (int32 PinIndex = 0; PinIndex < Pins.Num(); ++PinIndex)
+			{
+				UEdGraphPin* LocalPin = Pins[PinIndex];
+				if (LocalPin && (LocalPin != OutPin) && (LocalPin != SelfPin))
+				{
+					const FName PinName = GetNameForAdditionalPin(NameIndex);
+					if (PinName != LocalPin->PinName)
+					{
+						LocalPin->Modify();
+						LocalPin->PinName = PinName;
+					}
+					NameIndex++;
+				}
+			}
+			FBlueprintEditorUtils::MarkBlueprintAsStructurallyModified(GetBlueprint());
+		}
+	}
+}
+
+UEdGraphPin* UK2Node_PromotableOperator::GetAdditionalPin(int32 PinIndex) const
+{
+	const FName PinToFind = GetNameForAdditionalPin(PinIndex);
+
+	for (UEdGraphPin* Pin : Pins)
+	{
+		if (Pin && Pin->PinName == PinToFind)
+		{
+			return Pin;
+		}
+	}
+
+	return nullptr;
 }
 
 ///////////////////////////////////////////////////////////
 // UK2Node_PromotableOperator
 
+UEdGraphPin* UK2Node_PromotableOperator::AddInputPinImpl(int32 PinIndex)
+{
+	const FName NewPinName = GetNameForAdditionalPin(PinIndex);
+
+	UEdGraphPin* NewPin = FWildcardNodeUtils::CreateWildcardPin(this, NewPinName, EGPD_Input);
+	check(NewPin);
+
+	// Determine a default type for this pin if we have other input connections
+	const TArray<UEdGraphPin*> InputPins = GetInputPins(/* bIncludeLinks = */ true);
+	check(InputPins.Num());
+	FEdGraphPinType PromotedType = FTypePromotion::GetPromotedType(InputPins);
+
+	NewPin->PinType = PromotedType;
+
+	return NewPin;
+}
+
+bool UK2Node_PromotableOperator::IsAdditionalPin(const UEdGraphPin* Pin) const
+{
+	// Quickly check if this input pin is one of the two default input pins
+	return Pin && Pin->Direction == EGPD_Input && Pin->PinName != InputPinA_Name && Pin->PinName != InputPinB_Name;
+}
+
 bool UK2Node_PromotableOperator::HasAnyConnectionsOrDefaults() const
 {
 	for (UEdGraphPin* Pin : Pins)
@@ -357,18 +599,43 @@ bool UK2Node_PromotableOperator::UpdateOpName()
 	return Func ? FTypePromotion::GetOpNameFromFunction(GetTargetFunction(), /* out */ OperationName) : false;
 }
 
-bool UK2Node_PromotableOperator::CreateIntermediateCast(FKismetCompilerContext& CompilerContext, UEdGraph* SourceGraph, UEdGraphPin* InputPin, UEdGraphPin* OutputPin)
+UK2Node_CallFunction* UK2Node_PromotableOperator::CreateIntermediateNode(UK2Node_CallFunction* PreviousNode, const UFunction* const OpFunction, FKismetCompilerContext& CompilerContext, UEdGraph* SourceGraph)
+{
+	// Spawn an intermediate UK2Node_CallFunction node of the function type we need
+	UK2Node_CallFunction* NewOperator = SourceGraph->CreateIntermediateNode<UK2Node_CallFunction>();
+	NewOperator->SetFromFunction(OpFunction);
+	NewOperator->AllocateDefaultPins();
+
+	// Move this node next to the thing it was linked to
+	NewOperator->NodePosY = PreviousNode->NodePosY + 50;
+	NewOperator->NodePosX = PreviousNode->NodePosX + 8;
+
+	CompilerContext.MessageLog.NotifyIntermediateObjectCreation(NewOperator, this);
+
+	return NewOperator;
+}
+
+bool UK2Node_PromotableOperator::CreateIntermediateCast(UK2Node_CallFunction* SourceNode, FKismetCompilerContext& CompilerContext, UEdGraph* SourceGraph, UEdGraphPin* InputPin, UEdGraphPin* OutputPin)
 {
 	check(InputPin && OutputPin);
+	const UEdGraphSchema_K2* Schema = CompilerContext.GetSchema();
 
 	// If the pin types are the same, than no casts are needed and we can just connect
 	if (InputPin->PinType == OutputPin->PinType)
 	{
-		return !CompilerContext.MovePinLinksToIntermediate(*InputPin, *OutputPin).IsFatal();
+		// If SourceNode is 'this' then we need to move the pin links instead of just 
+		// creating the connection because the output is not another new node, but 
+		// just the intermediate expansion node.
+		if (SourceNode == this)
+		{
+			return !CompilerContext.MovePinLinksToIntermediate(*InputPin, *OutputPin).IsFatal();
+		}
+		else
+		{			
+			return Schema->TryCreateConnection(InputPin, OutputPin);
+		}
 	}
 
-	const UEdGraphSchema_K2* Schema = CompilerContext.GetSchema();
-
 	UK2Node* TemplateConversionNode = nullptr;
 	FName TargetFunctionName;
 	UClass* ClassContainingConversionFunction = nullptr;
@@ -415,15 +682,15 @@ bool UK2Node_PromotableOperator::CreateIntermediateCast(FKismetCompilerContext&
 		bOutputSuccessful = Schema->TryCreateConnection(ConversionOutput, OutputPin);
 
 		// Move this node next to the thing it was linked to
-		TemplateConversionNode->NodePosY = this->NodePosY;
-		TemplateConversionNode->NodePosX = this->NodePosX + 4;
+		TemplateConversionNode->NodePosY = SourceNode->NodePosY;
+		TemplateConversionNode->NodePosX = SourceNode->NodePosX + 4;
 	}
 	else
 	{
 		CompilerContext.MessageLog.Error(*FText::Format(LOCTEXT("NoValidPromotion", "Cannot find appropriate promotion from '{0}' to '{1}' on '@@'"),
 			Schema->TypeToText(InputPin->PinType),
 			Schema->TypeToText(OutputPin->PinType)).ToString(),
-			this
+			SourceNode
 		);
 	}
 
@@ -489,9 +756,13 @@ TArray<UEdGraphPin*> UK2Node_PromotableOperator::GetInputPins(bool bIncludeLinks
 
 void UK2Node_PromotableOperator::ConvertNodeToFunction(const UFunction* Function, UEdGraphPin* ChangedPin)
 {
-	const FScopedTransaction Transaction(LOCTEXT("ConvertPromotableOpToFunction", "Change the function signiture of a promotable operator node."));
+	const FScopedTransaction Transaction(LOCTEXT("ConvertPromotableOpToFunction", "Change the function signature of a promotable operator node."));
 	Modify();
 	RecombineAllSplitPins();
+	
+	// If we convert the node to a function, then just get rid of the additional pins
+	NumAdditionalInputs = 0;
+
 	UpdatePinsFromFunction(Function, ChangedPin);
 
 	// Reconstruct this node to fix any default values that may be invalid now
