Blueprints: Use a local variable instead of calling the same function twice

#rnx

diff --git a/Engine/Source/Editor/BlueprintGraph/Private/K2Node_PromotableOperator.cpp b/Engine/Source/Editor/BlueprintGraph/Private/K2Node_PromotableOperator.cpp
--- a/Engine/Source/Editor/BlueprintGraph/Private/K2Node_PromotableOperator.cpp
+++ b/Engine/Source/Editor/BlueprintGraph/Private/K2Node_PromotableOperator.cpp
@@ -596,7 +596,7 @@ bool UK2Node_PromotableOperator::UpdateOpName()
 	// If the function is null then return false, because we did not successfully update it. 
 	// This could be possible during node reconstruction/refresh, and we don't want to set the 
 	// op name to "Empty" incorrectly. 
-	return Func ? FTypePromotion::GetOpNameFromFunction(GetTargetFunction(), /* out */ OperationName) : false;
+	return Func ? FTypePromotion::GetOpNameFromFunction(Func, /* out */ OperationName) : false;
 }
 
 UK2Node_CallFunction* UK2Node_PromotableOperator::CreateIntermediateNode(UK2Node_CallFunction* PreviousNode, const UFunction* const OpFunction, FKismetCompilerContext& CompilerContext, UEdGraph* SourceGraph)
