Blueprints: Fix a crash in type promotion when dragging off of a variable and auto connecting the pin

#rnx

diff --git a/Engine/Source/Editor/BlueprintGraph/Classes/K2Node_PromotableOperator.h b/Engine/Source/Editor/BlueprintGraph/Classes/K2Node_PromotableOperator.h
--- a/Engine/Source/Editor/BlueprintGraph/Classes/K2Node_PromotableOperator.h
+++ b/Engine/Source/Editor/BlueprintGraph/Classes/K2Node_PromotableOperator.h
@@ -72,8 +72,22 @@ private:
 	*/
 	bool CreateIntermediateCast(UK2Node_CallFunction* SourceNode, FKismetCompilerContext& CompilerContext, UEdGraph* SourceGraph, UEdGraphPin* InputPin, UEdGraphPin* OutputPin);
 
+	/**
+	* Spawn a new intermediate call function node with the given operator function 
+	* and allocate its default pins. Place it next to the given Previous node 
+	* 
+	* @return	An intermediate call function node placed next to the "PreviousNode"
+	*/
 	UK2Node_CallFunction* CreateIntermediateNode(UK2Node_CallFunction* PreviousNode, const UFunction* const OpFunction, FKismetCompilerContext& CompilerContext, UEdGraph* SourceGraph);
 
+	/** 
+	* Re-evaluates the in types on this node based on all the current connections
+	* and the given pin that has changed. 
+	* 
+	* @param ChangedPin		The pin that has been change
+	*/
+	void EvaluatePinsFromChange(UEdGraphPin* ChangedPin);
+
 	/** Helper to make sure we have the most up to date operation name. Returns true upon success */
 	bool UpdateOpName();
 
diff --git a/Engine/Source/Editor/BlueprintGraph/Private/K2Node_PromotableOperator.cpp b/Engine/Source/Editor/BlueprintGraph/Private/K2Node_PromotableOperator.cpp
--- a/Engine/Source/Editor/BlueprintGraph/Private/K2Node_PromotableOperator.cpp
+++ b/Engine/Source/Editor/BlueprintGraph/Private/K2Node_PromotableOperator.cpp
@@ -291,69 +291,7 @@
 {
 	Super::NotifyPinConnectionListChanged(ChangedPin);
 
-	UpdateOpName();
-
-	const bool bOutputPinWasChanged = (ChangedPin == GetOutputPin());
-
-	// True if the pin that has changed now has zero connections
-	const bool bWasAFullDisconnect = (ChangedPin->LinkedTo.Num() == 0);
-
-	// If we have been totally disconnected and don't have any non-default inputs, 
-	// than we just reset the node to be a regular wildcard
-	if (bWasAFullDisconnect && !HasAnyConnectionsOrDefaults())
-	{
-		ResetNodeToWildcard();
-		return;
-	}
-	// If the pin that was connected is linked to a wildcard pin, then we should make it a wildcard
-	// and do nothing else.
-	else if (FWildcardNodeUtils::IsWildcardPin(ChangedPin) && ChangedPin->GetOwningNode() == this && FWildcardNodeUtils::IsLinkedToWildcard(ChangedPin))
-	{
-		return;
-	}
-
-	// Gather all pins and their links so we can determine the highest type that the user could want
-	TArray<UEdGraphPin*> InputPins;
-
-	for (UEdGraphPin* Pin : Pins)
-	{
-		// Consider inputs, maybe the output, and ignore the changed pin if it was totally disconnected.
-		if (Pin->LinkedTo.Num() > 0 || !Pin->DoesDefaultValueMatchAutogenerated())
-		{
-			InputPins.Add(Pin);
-			for (UEdGraphPin* Link : Pin->LinkedTo)
-			{
-				InputPins.Emplace(Link);
-			}
-		}
-	}
-
-	FEdGraphPinType HighestType = FTypePromotion::GetPromotedType(InputPins);
-
-	// if a pin was changed, update it if it cannot be promoted to this type
-	FEdGraphPinType NewConnectionHighestType = ChangedPin->LinkedTo.Num() > 0 ? FTypePromotion::GetPromotedType(ChangedPin->LinkedTo) : FWildcardNodeUtils::GetDefaultWildcardPinType();
-
-	// If there are ANY wildcards on this node, than we need to update the whole node accordingly. Otherwise we can 
-	// update only the Changed and output pins.
-
-	if (FWildcardNodeUtils::NodeHasAnyWildcards(this) || bOutputPinWasChanged || bWasAFullDisconnect ||
-		FTypePromotion::GetHigherType(NewConnectionHighestType, GetOutputPin()->PinType) == FTypePromotion::ETypeComparisonResult::TypeAHigher)
-	{
-		const UFunction* LowestFunc = FTypePromotion::FindLowestMatchingFunc(OperationName, HighestType, PossibleConversions);
-
-		// Store these other function options for later so that the user can convert to them later
-		UpdatePinsFromFunction(LowestFunc, ChangedPin);
-	}
-
-	// If the user connected a type that was a valid promotion, than leave it as the pin type they dragged from for a better UX
-	if (!bWasAFullDisconnect && NewConnectionHighestType.PinCategory != UEdGraphSchema_K2::PC_Wildcard &&
-		(FTypePromotion::IsValidPromotion(NewConnectionHighestType, ChangedPin->PinType) || FTypePromotion::IsValidPromotion(ChangedPin->PinType, NewConnectionHighestType)))
-	{
-		ChangedPin->PinType = NewConnectionHighestType;
-	}
-
-	// Update context menu options for this node
-	UpdatePossibleConversionFuncs();
+	EvaluatePinsFromChange(ChangedPin);
 }
 
 void UK2Node_PromotableOperator::PostReconstructNode()
@@ -458,7 +396,7 @@
 {
 	Super::AutowireNewNode(ChangedPin);
 
-	NotifyPinConnectionListChanged(ChangedPin);
+	EvaluatePinsFromChange(ChangedPin);
 }
 
 ///////////////////////////////////////////////////////////
@@ -590,6 +528,73 @@
 	return false;
 }
 
+void UK2Node_PromotableOperator::EvaluatePinsFromChange(UEdGraphPin* ChangedPin)
+{
+	UpdateOpName();
+
+	const bool bOutputPinWasChanged = (ChangedPin == GetOutputPin());
+
+	// True if the pin that has changed now has zero connections
+	const bool bWasAFullDisconnect = (ChangedPin->LinkedTo.Num() == 0);
+
+	// If we have been totally disconnected and don't have any non-default inputs, 
+	// than we just reset the node to be a regular wildcard
+	if (bWasAFullDisconnect && !HasAnyConnectionsOrDefaults())
+	{
+		ResetNodeToWildcard();
+		return;
+	}
+	// If the pin that was connected is linked to a wildcard pin, then we should make it a wildcard
+	// and do nothing else.
+	else if (FWildcardNodeUtils::IsWildcardPin(ChangedPin) && ChangedPin->GetOwningNode() == this && FWildcardNodeUtils::IsLinkedToWildcard(ChangedPin))
+	{
+		return;
+	}
+
+	// Gather all pins and their links so we can determine the highest type that the user could want
+	TArray<UEdGraphPin*> InputPins;
+
+	for (UEdGraphPin* Pin : Pins)
+	{
+		// Consider inputs, maybe the output, and ignore the changed pin if it was totally disconnected.
+		if (Pin->LinkedTo.Num() > 0 || !Pin->DoesDefaultValueMatchAutogenerated())
+		{
+			InputPins.Add(Pin);
+			for (UEdGraphPin* Link : Pin->LinkedTo)
+			{
+				InputPins.Emplace(Link);
+			}
+		}
+	}
+
+	FEdGraphPinType HighestType = FTypePromotion::GetPromotedType(InputPins);
+
+	// if a pin was changed, update it if it cannot be promoted to this type
+	FEdGraphPinType NewConnectionHighestType = ChangedPin->LinkedTo.Num() > 0 ? FTypePromotion::GetPromotedType(ChangedPin->LinkedTo) : FWildcardNodeUtils::GetDefaultWildcardPinType();
+
+	// If there are ANY wildcards on this node, than we need to update the whole node accordingly. Otherwise we can 
+	// update only the Changed and output pins.
+
+	if (FWildcardNodeUtils::NodeHasAnyWildcards(this) || bOutputPinWasChanged || bWasAFullDisconnect ||
+		FTypePromotion::GetHigherType(NewConnectionHighestType, GetOutputPin()->PinType) == FTypePromotion::ETypeComparisonResult::TypeAHigher)
+	{
+		const UFunction* LowestFunc = FTypePromotion::FindLowestMatchingFunc(OperationName, HighestType, PossibleConversions);
+
+		// Store these other function options for later so that the user can convert to them later
+		UpdatePinsFromFunction(LowestFunc, ChangedPin);
+	}
+
+	// If the user connected a type that was a valid promotion, than leave it as the pin type they dragged from for a better UX
+	if (!bWasAFullDisconnect && NewConnectionHighestType.PinCategory != UEdGraphSchema_K2::PC_Wildcard &&
+		(FTypePromotion::IsValidPromotion(NewConnectionHighestType, ChangedPin->PinType) || FTypePromotion::IsValidPromotion(ChangedPin->PinType, NewConnectionHighestType)))
+	{
+		ChangedPin->PinType = NewConnectionHighestType;
+	}
+
+	// Update context menu options for this node
+	UpdatePossibleConversionFuncs();
+}
+
 bool UK2Node_PromotableOperator::UpdateOpName()
 {
 	const UFunction* Func = GetTargetFunction();
