Blueprints: Add an accessor for the number of additional inputs on the CommunicativeBinaryOp node so that we can use it in a blueprint upgrade tool. Expose GetOutputPin to the BlueprintGraphAPI as well.

#rnx

diff --git a/Engine/Source/Editor/BlueprintGraph/Classes/K2Node_CommutativeAssociativeBinaryOperator.h b/Engine/Source/Editor/BlueprintGraph/Classes/K2Node_CommutativeAssociativeBinaryOperator.h
--- a/Engine/Source/Editor/BlueprintGraph/Classes/K2Node_CommutativeAssociativeBinaryOperator.h
+++ b/Engine/Source/Editor/BlueprintGraph/Classes/K2Node_CommutativeAssociativeBinaryOperator.h
@@ -37,6 +37,9 @@ public:
 	/** Get TRUE input type (self, etc.. are skipped) */
 	BLUEPRINTGRAPH_API UEdGraphPin* GetInputPin(int32 InputPinIndex);
 
+	/** Returns the number of additional input pins that this node has */
+	BLUEPRINTGRAPH_API int32 GetNumberOfAdditionalInputs() const { return NumAdditionalInputs; }
+
 	virtual void RemoveInputPin(UEdGraphPin* Pin) override;
 
 	// UEdGraphNode interface
diff --git a/Engine/Source/Editor/BlueprintGraph/Classes/K2Node_PromotableOperator.h b/Engine/Source/Editor/BlueprintGraph/Classes/K2Node_PromotableOperator.h
--- a/Engine/Source/Editor/BlueprintGraph/Classes/K2Node_PromotableOperator.h
+++ b/Engine/Source/Editor/BlueprintGraph/Classes/K2Node_PromotableOperator.h
@@ -124,7 +124,7 @@ private:
 public:
 
 	/** Returns the first pin with the EGPD_Output direction */
-	UEdGraphPin* GetOutputPin() const;
+	BLUEPRINTGRAPH_API UEdGraphPin* GetOutputPin() const;
 
 	const FString& GetOperationName() const { return OperationName; }
 };
