Blueprints: Make Comparison operators have a bool output flag by default on Promotable operators. Misc. variable naming cleanup and removing of unused variable

#rnx

diff --git a/Engine/Source/Editor/BlueprintGraph/Private/K2Node_PromotableOperator.cpp b/Engine/Source/Editor/BlueprintGraph/Private/K2Node_PromotableOperator.cpp
--- a/Engine/Source/Editor/BlueprintGraph/Private/K2Node_PromotableOperator.cpp
+++ b/Engine/Source/Editor/BlueprintGraph/Private/K2Node_PromotableOperator.cpp
@@ -41,7 +41,16 @@ void UK2Node_PromotableOperator::AllocateDefaultPins()
 	FWildcardNodeUtils::CreateWildcardPin(this, InputPinA_Name, EGPD_Input);
 	FWildcardNodeUtils::CreateWildcardPin(this, InputPinB_Name, EGPD_Input);
 
-	FWildcardNodeUtils::CreateWildcardPin(this, UEdGraphSchema_K2::PN_ReturnValue, EGPD_Output);
+	UEdGraphPin* OutPin = FWildcardNodeUtils::CreateWildcardPin(this, UEdGraphSchema_K2::PN_ReturnValue, EGPD_Output);
+
+	// For comparison functions we always want a bool output, so make it visually so
+	if (const UFunction* Func = GetTargetFunction())
+	{
+		if (FTypePromotion::IsComparisonFunc(Func))
+		{
+			OutPin->PinType.PinCategory = UEdGraphSchema_K2::PC_Boolean;
+		}
+	}
 
 	// Create any additional input pin. Their appropriate type is determined in ReallocatePinsDuringReconstruction
 	// because we cannot get a promoted type with no links to the pin.
@@ -552,22 +561,22 @@ void UK2Node_PromotableOperator::EvaluatePinsFromChange(UEdGraphPin* ChangedPin)
 	}
 
 	// Gather all pins and their links so we can determine the highest type that the user could want
-	TArray<UEdGraphPin*> InputPins;
+	TArray<UEdGraphPin*> PinsToConsider;
 
 	for (UEdGraphPin* Pin : Pins)
 	{
 		// Consider inputs, maybe the output, and ignore the changed pin if it was totally disconnected.
 		if (Pin->LinkedTo.Num() > 0 || !Pin->DoesDefaultValueMatchAutogenerated())
 		{
-			InputPins.Add(Pin);
+			PinsToConsider.Add(Pin);
 			for (UEdGraphPin* Link : Pin->LinkedTo)
 			{
-				InputPins.Emplace(Link);
+				PinsToConsider.Emplace(Link);
 			}
 		}
 	}
 
-	FEdGraphPinType HighestType = FTypePromotion::GetPromotedType(InputPins);
+	FEdGraphPinType HighestType = FTypePromotion::GetPromotedType(PinsToConsider);
 
 	// if a pin was changed, update it if it cannot be promoted to this type
 	FEdGraphPinType NewConnectionHighestType = ChangedPin->LinkedTo.Num() > 0 ? FTypePromotion::GetPromotedType(ChangedPin->LinkedTo) : FWildcardNodeUtils::GetDefaultWildcardPinType();
@@ -677,11 +686,15 @@ bool UK2Node_PromotableOperator::CreateIntermediateCast(UK2Node_CallFunction* So
 		}
 		UEdGraphPin* ConversionOutput = TemplateConversionNode->FindPin(UEdGraphSchema_K2::PN_ReturnValue, EGPD_Output);
 
-		// Connect my input to the conversion node
+		// Connect my input to the conversion node directly if we have links, otherwise we need to move the intermediate version of it
 		if (InputPin->LinkedTo.Num() > 0)
 		{
 			bInputSuccessful = Schema->TryCreateConnection(InputPin->LinkedTo[0], ConversionInput);
 		}
+		else if(InputPin && ConversionInput)
+		{
+			bInputSuccessful = !CompilerContext.MovePinLinksToIntermediate(*InputPin, *ConversionInput).IsFatal();
+		}
 
 		// Connect conversion node output to the input of the new operator
 		bOutputSuccessful = Schema->TryCreateConnection(ConversionOutput, OutputPin);
@@ -720,6 +733,17 @@ void UK2Node_PromotableOperator::ResetNodeToWildcard()
 		}
 	}
 
+	// Set output pins to have a bool output flag by default
+	{
+		UEdGraphPin* OutPin = GetOutputPin();
+		const UFunction* Func = GetTargetFunction();
+
+		if (FTypePromotion::IsComparisonFunc(Func))
+		{
+			OutPin->PinType.PinCategory = UEdGraphSchema_K2::PC_Boolean;
+		}
+	}
+	
 	// Clear out any possible function matches, since we are removing connections
 	PossibleConversions.Empty();
 }
@@ -784,9 +808,6 @@ void UK2Node_PromotableOperator::UpdatePinsFromFunction(const UFunction* Functio
 
 	const UEdGraphSchema_K2* Schema = GetDefault<UEdGraphSchema_K2>();
 
-	TMap<FString, TSet<UEdGraphPin*>> PinConnections;
-	FEdGraphUtilities::GetPinConnectionMap(this, PinConnections);
-
 	int32 ArgCount = 0;
 	for (TFieldIterator<FProperty> PropIt(Function); PropIt && (PropIt->PropertyFlags & CPF_Parm); ++PropIt)
 	{
