Blueprints: Correctly handle the expansion of default values during type promotion. When Updating the Promotable operator node from a function, if it is a wildcard pin (like during node reconstruction) then use the linked types only.

#rnx

diff --git a/Engine/Source/Editor/BlueprintGraph/Private/BlueprintTypePromotion.cpp b/Engine/Source/Editor/BlueprintGraph/Private/BlueprintTypePromotion.cpp
--- a/Engine/Source/Editor/BlueprintGraph/Private/BlueprintTypePromotion.cpp
+++ b/Engine/Source/Editor/BlueprintGraph/Private/BlueprintTypePromotion.cpp
@@ -470,10 +470,15 @@ UFunction* FTypePromotion::FindBestMatchingFunc_Internal(const FString& Operatio
 		return nullptr;
 	}
 
-	// Track the function with the best score
+	// Track the function with the best score, input, and output types
 	UFunction* BestFunc = nullptr;
+	FEdGraphPinType BestFuncLowestInputType;
+	FEdGraphPinType BestFuncOutputType;
 	int32 BestScore = -1;
 
+	const bool bIsSinglePin = PinsToConsider.Num() == 1;
+	const bool bIsComparisonOp = GetComparisonOpNames().Contains(Operation);
+
 	const UEdGraphSchema_K2* Schema = GetDefault<UEdGraphSchema_K2>();
 		
 	// We have to keep track of what pins we have already given points for, otherwise 
@@ -485,6 +490,11 @@ UFunction* FTypePromotion::FindBestMatchingFunc_Internal(const FString& Operatio
 		int32 FuncScore = -1;
 		CheckedPins.Reset();
 
+		// Track this functions highest input and output types so that if there is a function with
+		// the same score as it we can prefer the correct one. 
+		FEdGraphPinType CurFuncHighestInputType;
+		FEdGraphPinType CurFuncOutputType;
+
 		// For each property in the func, see if it matches any of the given pins
 		for (TFieldIterator<FProperty> PropIt(Func); PropIt && (PropIt->PropertyFlags & CPF_Parm); ++PropIt)
 		{
@@ -497,18 +507,54 @@ UFunction* FTypePromotion::FindBestMatchingFunc_Internal(const FString& Operatio
 					// Give a point for each function parameter that matches up with a pin to consider
 					if (!CheckedPins.Contains(Pin) && Schema->ArePinTypesEquivalent(ParamType, Pin->PinType))
 					{
-						++FuncScore;
-						CheckedPins.Add(Pin);
+						// Are the directions compatible? 
+						// If we are a comparison or only a single pin then we don't care about the direction
+						if (bIsSinglePin || bIsComparisonOp ||
+						   (Param->HasAnyPropertyFlags(CPF_ReturnParm) && Pin->Direction == EGPD_Output) ||
+						   (!Param->HasAnyPropertyFlags(CPF_ReturnParm) && Pin->Direction == EGPD_Input))
+						{
+							++FuncScore;
+							CheckedPins.Add(Pin);
+						}
+
 						break;
 					}
 				}
+
+				// Keep track of the highest input pin type on this function
+				if (Param->HasAnyPropertyFlags(CPF_ReturnParm))
+				{
+					CurFuncOutputType = ParamType;
+				}
+				else if (CurFuncHighestInputType.PinCategory == NAME_None || FTypePromotion::GetHigherType(ParamType, CurFuncHighestInputType) == ETypeComparisonResult::TypeBHigher)
+				{
+					CurFuncHighestInputType = ParamType;
+				}
 			}
 		}
 
-		// If this function has the best score, then update it! 
-		if (FuncScore > BestScore)
+		// If the pin type has no name, then this is an invalid comparison
+		ETypeComparisonResult InputCompareRes = BestFuncLowestInputType.PinCategory != NAME_None ? FTypePromotion::GetHigherType(CurFuncHighestInputType, BestFuncLowestInputType) : ETypeComparisonResult::InvalidComparison;
+		ETypeComparisonResult OutputCompareRes = BestFuncOutputType.PinCategory != NAME_None ? FTypePromotion::GetHigherType(CurFuncOutputType, BestFuncOutputType) : ETypeComparisonResult::InvalidComparison;
+		
+		// We want to prefer a HIGHER input, and a LOWER output. 
+		const bool bHasInputOutputPreference =
+			InputCompareRes != ETypeComparisonResult::TypeBHigher &&
+			OutputCompareRes != ETypeComparisonResult::TypeAHigher;
+
+		// If the scores are equal, then prefer the LARGER input and output type because we can promote up, but we can never go back down
+		const bool bScoresEqualAndPreferred = 
+			FuncScore == BestScore && FuncScore != -1 &&
+			(bIsComparisonOp || bIsSinglePin || 
+			(InputCompareRes == ETypeComparisonResult::TypeAHigher ||
+			OutputCompareRes == ETypeComparisonResult::TypeAHigher));
+
+		// Keep track of the best function!
+		if (bScoresEqualAndPreferred || (FuncScore > BestScore && (bHasInputOutputPreference || bIsComparisonOp)))
 		{
 			BestScore = FuncScore;
+			BestFuncLowestInputType = CurFuncHighestInputType;
+			BestFuncOutputType = CurFuncOutputType;
 			BestFunc = Func;
 		}
 	}
diff --git a/Engine/Source/Editor/BlueprintGraph/Private/K2Node_PromotableOperator.cpp b/Engine/Source/Editor/BlueprintGraph/Private/K2Node_PromotableOperator.cpp
--- a/Engine/Source/Editor/BlueprintGraph/Private/K2Node_PromotableOperator.cpp
+++ b/Engine/Source/Editor/BlueprintGraph/Private/K2Node_PromotableOperator.cpp
@@ -44,12 +44,9 @@ void UK2Node_PromotableOperator::AllocateDefaultPins()
 	UEdGraphPin* OutPin = FWildcardNodeUtils::CreateWildcardPin(this, UEdGraphSchema_K2::PN_ReturnValue, EGPD_Output);
 
 	// For comparison functions we always want a bool output, so make it visually so
-	if (const UFunction* Func = GetTargetFunction())
+	if (FTypePromotion::IsComparisonFunc(GetTargetFunction()))
 	{
-		if (FTypePromotion::IsComparisonFunc(Func))
-		{
-			OutPin->PinType.PinCategory = UEdGraphSchema_K2::PC_Boolean;
-		}
+		OutPin->PinType.PinCategory = UEdGraphSchema_K2::PC_Boolean;
 	}
 
 	// Create any additional input pin. Their appropriate type is determined in ReallocatePinsDuringReconstruction
@@ -377,10 +374,6 @@ bool UK2Node_PromotableOperator::IsConnectionDisallowed(const UEdGraphPin* MyPin
 
 void UK2Node_PromotableOperator::ReallocatePinsDuringReconstruction(TArray<UEdGraphPin*>& OldPins)
 {
-	// Allocate default pins will have been called before this, which means we are reset to wildcard state
-	// We need to Update the pins to be the proper function again
-	UpdatePinsFromFunction(GetTargetFunction(), nullptr);
-
 	Super::ReallocatePinsDuringReconstruction(OldPins);
 
 	// We need to fix up any additional pins that may have been created as a wildcard pin
@@ -565,15 +558,20 @@ void UK2Node_PromotableOperator::EvaluatePinsFromChange(UEdGraphPin* ChangedPin)
 
 	for (UEdGraphPin* Pin : Pins)
 	{
-		// Consider inputs, maybe the output, and ignore the changed pin if it was totally disconnected.
-		if (Pin->LinkedTo.Num() > 0 || !Pin->DoesDefaultValueMatchAutogenerated())
+		// If this pin has links, then use those instead of the actual pin because we could be process of changing it
+		// which means that it would still have it's old pin type, and could be inaccurate
+		if (Pin->LinkedTo.Num() > 0)
 		{
-			PinsToConsider.Add(Pin);
 			for (UEdGraphPin* Link : Pin->LinkedTo)
 			{
 				PinsToConsider.Emplace(Link);
 			}
 		}
+		else if(!Pin->DoesDefaultValueMatchAutogenerated() && !FWildcardNodeUtils::IsWildcardPin(Pin))
+		{
+			// If this pin has doesn't have any links then we should account for the pin itself
+			PinsToConsider.Add(Pin);
+		}
 	}
 
 	FEdGraphPinType HighestType = FTypePromotion::GetPromotedType(PinsToConsider);
@@ -584,14 +582,10 @@ void UK2Node_PromotableOperator::EvaluatePinsFromChange(UEdGraphPin* ChangedPin)
 	// If there are ANY wildcards on this node, than we need to update the whole node accordingly. Otherwise we can 
 	// update only the Changed and output pins.
 
-	if (FWildcardNodeUtils::NodeHasAnyWildcards(this) || bOutputPinWasChanged || bWasAFullDisconnect ||
-		FTypePromotion::GetHigherType(NewConnectionHighestType, GetOutputPin()->PinType) == FTypePromotion::ETypeComparisonResult::TypeAHigher)
-	{
-		const UFunction* LowestFunc = FTypePromotion::FindLowestMatchingFunc(OperationName, HighestType, PossibleConversions);
+	const UFunction* BestMatchingFunc = FTypePromotion::FindBestMatchingFunc(OperationName, PinsToConsider);
 
-		// Store these other function options for later so that the user can convert to them later
-		UpdatePinsFromFunction(LowestFunc, ChangedPin);
-	}
+	// Store these other function options for later so that the user can convert to them later
+	UpdatePinsFromFunction(BestMatchingFunc, ChangedPin);
 
 	// If the user connected a type that was a valid promotion, than leave it as the pin type they dragged from for a better UX
 	if (!bWasAFullDisconnect && NewConnectionHighestType.PinCategory != UEdGraphSchema_K2::PC_Wildcard &&
@@ -738,7 +732,7 @@ void UK2Node_PromotableOperator::ResetNodeToWildcard()
 		UEdGraphPin* OutPin = GetOutputPin();
 		const UFunction* Func = GetTargetFunction();
 
-		if (FTypePromotion::IsComparisonFunc(Func))
+		if (OutPin && Func && FTypePromotion::IsComparisonFunc(Func))
 		{
 			OutPin->PinType.PinCategory = UEdGraphSchema_K2::PC_Boolean;
 		}
@@ -867,7 +861,8 @@ void UK2Node_PromotableOperator::UpdatePinsFromFunction(const UFunction* Functio
 				if (bIsWildcard && PinToChange->LinkedTo.Num() > 0)
 				{
 					FEdGraphPinType LinkedType = FTypePromotion::GetPromotedType(PinToChange->LinkedTo);
-					if (FTypePromotion::IsValidPromotion(PinToChange->PinType, ParamType))
+					// If this is a wildcard pin, then check the compatibility with a linked type instead of the wildcard itself
+					if (FTypePromotion::IsValidPromotion(bIsWildcard  ? LinkedType : PinToChange->PinType, ParamType))
 					{
 						ParamType = LinkedType;
 					}
