Blueprints: Add more in depth unit testing for Type Promotion with struct combinations and comparison.

#rnx

diff --git a/Engine/Source/Editor/BlueprintGraph/Private/Tests/TypePromotionTests.cpp b/Engine/Source/Editor/BlueprintGraph/Private/Tests/TypePromotionTests.cpp
--- a/Engine/Source/Editor/BlueprintGraph/Private/Tests/TypePromotionTests.cpp
+++ b/Engine/Source/Editor/BlueprintGraph/Private/Tests/TypePromotionTests.cpp
@@ -129,19 +129,93 @@ bool FFindBestMatchingFunc::RunTest(const FString& Parameters)
 
 	UEdGraphNode* TestNode = NewObject<UEdGraphNode>();
 
+	////////////////////////////////////////////////////
+	// Create test pins!
 	MakeTestPin(TestNode, PinTypes, DoublePinA, UEdGraphSchema_K2::PC_Double);
 	MakeTestPin(TestNode, PinTypes, DoublePinB, UEdGraphSchema_K2::PC_Double);
-
+	MakeTestPin(TestNode, PinTypes, DoubleOutputPin, UEdGraphSchema_K2::PC_Double);
+	DoubleOutputPin->Direction = EGPD_Output;
+	
 	MakeTestPin(TestNode, PinTypes, FloatPinA, UEdGraphSchema_K2::PC_Float);
+	FloatPinA->Direction = EGPD_Input;
 	MakeTestPin(TestNode, PinTypes, FloatPinB, UEdGraphSchema_K2::PC_Float);
-	
+	FloatPinB->Direction = EGPD_Input;
+
+	MakeTestPin(TestNode, PinTypes, FloatOutputPin, UEdGraphSchema_K2::PC_Float);
+	FloatPinB->Direction = EGPD_Output;
+
+	MakeTestPin(TestNode, PinTypes, BoolOutputPin, UEdGraphSchema_K2::PC_Boolean);
+	BoolOutputPin->Direction = EGPD_Output;
+
 	MakeTestPin(TestNode, PinTypes, IntPinA, UEdGraphSchema_K2::PC_Int);
 
+	// Structs!
+	MakeTestPin(TestNode, PinTypes, VecInputPinA, UEdGraphSchema_K2::PC_Struct);
+	VecInputPinA->Direction = EGPD_Input;
+	VecInputPinA->PinType.PinSubCategoryObject = TBaseStructure<FVector>::Get();
+
+	MakeTestPin(TestNode, PinTypes, VecInputPinB, UEdGraphSchema_K2::PC_Struct);
+	VecInputPinB->Direction = EGPD_Input;
+	VecInputPinB->PinType.PinSubCategoryObject = TBaseStructure<FVector>::Get();
+
+	MakeTestPin(TestNode, PinTypes, VecOutputPinA, UEdGraphSchema_K2::PC_Struct);
+	VecOutputPinA->Direction = EGPD_Output;
+	VecOutputPinA->PinType.PinSubCategoryObject = TBaseStructure<FVector>::Get();
+
+	// Multiply_VectorVector given A float input, vector input, and a vector output
+	{
+		TArray<UEdGraphPin*> TestPins =
+		{
+			FloatPinA, VecInputPinB, VecOutputPinA,
+		};
+
+		const UFunction* MulVecFunc = FTypePromotion::FindBestMatchingFunc(TEXT("Multiply"), TestPins);
+		static const FName ExpectedName = TEXT("Multiply_VectorVector");
+
+		if (TestNotNull(TEXT("Multiply_VectorVector Null check"), MulVecFunc))
+		{
+			TestEqual(TEXT("Multiply_VectorVector Name Check"), MulVecFunc->GetFName(), ExpectedName);
+		}
+	}
+
+	// Multiply_VectorVector given a float, vector, float
+	// Order shouldn't matter when passing these pins in, which is what we are testing here
+	{
+		TArray<UEdGraphPin*> TestPins =
+		{
+			FloatPinA, VecOutputPinA, VecInputPinA,
+		};
+
+		const UFunction* MulVecFunc = FTypePromotion::FindBestMatchingFunc(TEXT("Multiply"), TestPins);
+		static const FName ExpectedName = TEXT("Multiply_VectorVector");
+
+		if (TestNotNull(TEXT("Multiply_VectorVector Null check"), MulVecFunc))
+		{
+			TestEqual(TEXT("Multiply_VectorVector Name Check"), MulVecFunc->GetFName(), ExpectedName);
+		}
+	}
+
+	// Multiply_VectorVector given two vector inputs and a vector output
+	{
+		TArray<UEdGraphPin*> TestPins =
+		{
+			VecInputPinA, VecInputPinB, VecOutputPinA,
+		};
+
+		const UFunction* MulVecFunc = FTypePromotion::FindBestMatchingFunc(TEXT("Multiply"), TestPins);
+		static const FName ExpectedName = TEXT("Multiply_VectorVector");
+
+		if (TestNotNull(TEXT("Multiply_VectorVector Null check"), MulVecFunc))
+		{
+			TestEqual(TEXT("Multiply_VectorVector Name Check"), MulVecFunc->GetFName(), ExpectedName);
+		}
+	}
+
 	// Add_DoubleDouble
 	{
 		TArray<UEdGraphPin*> TestPins =
 		{
-			DoublePinA, DoublePinB
+			DoublePinA, DoublePinB, DoubleOutputPin
 		};
 
 		const UFunction* AddDoubleFunc = FTypePromotion::FindBestMatchingFunc(TEXT("Add"), TestPins);
@@ -153,11 +227,27 @@ bool FFindBestMatchingFunc::RunTest(const FString& Parameters)
 		}
 	}
 
+	// Add_DoubleDouble given a double and float
+	{
+		TArray<UEdGraphPin*> TestPins =
+		{
+			DoublePinA, FloatPinA, DoubleOutputPin
+		};
+
+		const UFunction* AddDoubleFunc = FTypePromotion::FindBestMatchingFunc(TEXT("Add"), TestPins);
+		static const FName ExpectedName = TEXT("Add_DoubleDouble");
+
+		if (TestNotNull(TEXT("Add_DoubleDouble Null check"), AddDoubleFunc))
+		{
+			TestEqual(TEXT("Add_DoubleDouble Name Check"), AddDoubleFunc->GetFName(), ExpectedName);
+		}
+	}
+
 	// Subtract_FloatFloat
 	{
 		TArray<UEdGraphPin*> TestPins =
 		{
-			FloatPinA, FloatPinB
+			FloatPinA, FloatPinB, FloatOutputPin
 		};
 
 		const UFunction* SubtractFloatFunc = FTypePromotion::FindBestMatchingFunc(TEXT("Subtract"), TestPins);
@@ -169,19 +259,68 @@ bool FFindBestMatchingFunc::RunTest(const FString& Parameters)
 		}
 	}
 
-	// Multiply_IntFloat
+	// Add_FloatFloat given only one float pin. This simulates the first connection being made to a 
+	// promotable operator, in which case we should default to a regular old Float + Float
+	{
+		TArray<UEdGraphPin*> TestPins =
+		{
+			FloatPinA,
+		};
+
+		const UFunction* AddFloatFunc = FTypePromotion::FindBestMatchingFunc(TEXT("Add"), TestPins);
+		static const FName ExpectedName = TEXT("Add_FloatFloat");
+
+		if (TestNotNull(TEXT("Add_FloatFloat null check"), AddFloatFunc))
+		{
+			TestEqual(TEXT("Add_FloatFloat Name Check"), AddFloatFunc->GetFName(), ExpectedName);
+		}
+	}
+
+	// Less_FloatFloat Given a Float and Double
+	{
+		TArray<UEdGraphPin*> TestPins =
+		{
+			FloatPinA, BoolOutputPin
+		};
+
+		const UFunction* LessFunc = FTypePromotion::FindBestMatchingFunc(TEXT("Less"), TestPins);
+		static const FName ExpectedName = TEXT("Less_FloatFloat");
+
+		if (TestNotNull(TEXT("Less_FloatFloat null check"), LessFunc))
+		{
+			TestEqual(TEXT("Less_FloatFloat Name Check"), LessFunc->GetFName(), ExpectedName);
+		}
+	}
+
+	// Less_FloatFloat Given just a single float
+	{
+		TArray<UEdGraphPin*> TestPins =
+		{
+			FloatPinA,
+		};
+
+		const UFunction* LessFunc = FTypePromotion::FindBestMatchingFunc(TEXT("Less"), TestPins);
+		static const FName ExpectedName = TEXT("Less_FloatFloat");
+
+		if (TestNotNull(TEXT("Less_FloatFloat null check"), LessFunc))
+		{
+			TestEqual(TEXT("Less_FloatFloat Name Check"), LessFunc->GetFName(), ExpectedName);
+		}
+	}
+
+	// Greater_DoubleDouble Given a Float and Double
 	{
 		TArray<UEdGraphPin*> TestPins =
 		{
-			FloatPinA, FloatPinB, IntPinA
+			FloatPinA, DoublePinA
 		};
 
-		const UFunction* MulFloatFunc = FTypePromotion::FindBestMatchingFunc(TEXT("Multiply"), TestPins);
-		static const FName ExpectedName = TEXT("Multiply_IntFloat");
+		const UFunction* GreaterFunc = FTypePromotion::FindBestMatchingFunc(TEXT("Greater"), TestPins);
+		static const FName ExpectedName = TEXT("Greater_DoubleDouble");
 
-		if (TestNotNull(TEXT("Multiply_IntFloat null check"), MulFloatFunc))
+		if (TestNotNull(TEXT("Greater_DoubleDouble null check"), GreaterFunc))
 		{
-			TestEqual(TEXT("Multiply_IntFloat Name Check"), MulFloatFunc->GetFName(), ExpectedName);
+			TestEqual(TEXT("Greater_DoubleDouble Name Check"), GreaterFunc->GetFName(), ExpectedName);
 		}
 	}
 
