Blueprints: Add a Shutdown method to the TypePromotion singleton and call it when the BP module gets shut down

#rnx

diff --git a/Engine/Source/Editor/BlueprintGraph/Classes/BlueprintTypePromotion.h b/Engine/Source/Editor/BlueprintGraph/Classes/BlueprintTypePromotion.h
--- a/Engine/Source/Editor/BlueprintGraph/Classes/BlueprintTypePromotion.h
+++ b/Engine/Source/Editor/BlueprintGraph/Classes/BlueprintTypePromotion.h
@@ -13,8 +13,12 @@ class BLUEPRINTGRAPH_API FTypePromotion : private FNoncopyable
 {
 public:
 
+	/** Creates a new singleton instance of TypePromotion if there isn't one and returns a reference to it */
 	static FTypePromotion& Get();
 
+	/** Deletes the singleton instance of type promotion if there is one */
+	static void Shutdown();
+
 	/** 
 	* If given more than two input pins, this will just use the highest type
 	* and return the function for that operation. It's update to nodes like the 
diff --git a/Engine/Source/Editor/BlueprintGraph/Private/BlueprintGraphModule.cpp b/Engine/Source/Editor/BlueprintGraph/Private/BlueprintGraphModule.cpp
--- a/Engine/Source/Editor/BlueprintGraph/Private/BlueprintGraphModule.cpp
+++ b/Engine/Source/Editor/BlueprintGraph/Private/BlueprintGraphModule.cpp
@@ -3,6 +3,7 @@
 #include "BlueprintGraphModule.h"
 
 #include "EdGraphSchema_K2.h"
+#include "BlueprintTypePromotion.h"
 #include "Modules/ModuleManager.h"
 
 IMPLEMENT_MODULE( FBlueprintGraphModule, BlueprintGraph );
@@ -10,5 +11,5 @@ IMPLEMENT_MODULE( FBlueprintGraphModule, BlueprintGraph );
 void FBlueprintGraphModule::ShutdownModule()
 {
 	UEdGraphSchema_K2::Shutdown();
-}
-
+	FTypePromotion::Shutdown();
+}
\ No newline at end of file
diff --git a/Engine/Source/Editor/BlueprintGraph/Private/BlueprintTypePromotion.cpp b/Engine/Source/Editor/BlueprintGraph/Private/BlueprintTypePromotion.cpp
--- a/Engine/Source/Editor/BlueprintGraph/Private/BlueprintTypePromotion.cpp
+++ b/Engine/Source/Editor/BlueprintGraph/Private/BlueprintTypePromotion.cpp
@@ -52,6 +52,15 @@ FTypePromotion& FTypePromotion::Get()
 	return *Instance;
 }
 
+void FTypePromotion::Shutdown()
+{
+	if (Instance)
+	{
+		delete Instance;
+		Instance = nullptr;
+	}
+}
+
 FTypePromotion::FTypePromotion()
 {
 	CreatePromotionTable();
@@ -62,12 +71,6 @@ FTypePromotion::FTypePromotion()
 
 FTypePromotion::~FTypePromotion()
 {
-	if(Instance)
-	{
-		delete Instance;
-		Instance = nullptr;
-	}
-
 	FModuleManager::Get().OnModulesChanged().Remove(OnModulesChangedDelegateHandle);
 }
 
