Blueprints: Add Trace timer to GetOpNameFromFunction in type promotion. Remove unnecssary callback that just wasted time on module reload.

#rnx

diff --git a/Engine/Source/Editor/BlueprintGraph/Classes/BlueprintTypePromotion.h b/Engine/Source/Editor/BlueprintGraph/Classes/BlueprintTypePromotion.h
--- a/Engine/Source/Editor/BlueprintGraph/Classes/BlueprintTypePromotion.h
+++ b/Engine/Source/Editor/BlueprintGraph/Classes/BlueprintTypePromotion.h
@@ -8,7 +8,11 @@
 #include "BlueprintFunctionNodeSpawner.h"
 #include "HAL/IConsoleManager.h"
 
-/** Contains behavior needed to handle type promotion in blueprints */
+/** 
+* Contains behavior needed to handle type promotion in blueprints. 
+* Creates a map of "Operations" to any of their matching UFunctions
+* so that we can find the best possible match given several pin types.  
+*/
 class BLUEPRINTGRAPH_API FTypePromotion : private FNoncopyable
 {
 public:
@@ -74,6 +78,7 @@ public:
 	*/
 	static bool PromotePin(FEdGraphPinType& InTypeA, const FEdGraphPinType& TypeB);
 
+	/** Returns true if the given function can be used for type promotion (it is within the operator table) */
 	static bool IsFunctionPromotionReady(const UFunction* const FuncToConsider);
 
 	/** Represents the possible results when comparing two types for promotion */
@@ -86,20 +91,21 @@ public:
 	};
 
 	/**
-	* 
-	* 
-	* 
-	* @return 
+	* Given the two pin types cehck which pin type is higher. 
+	* Given two structs it will return equal, this does NOT compare PinDefaultSubobjects
 	*/
 	static ETypeComparisonResult GetHigherType(const FEdGraphPinType& A, const FEdGraphPinType& B);
 
 	/** Returns true if A can be promoted to type B correctly, or if the types are equal */
 	static bool IsValidPromotion(const FEdGraphPinType& A, const FEdGraphPinType& B);
 
+	/** Returns true if the given input pin can correctly be converted to the output type as a struct */
 	static bool HasStructConversion(const UEdGraphPin* InputPin, const UEdGraphPin* OutputPin);
 
+	/** Returns true if the given function has a registered operator node spawner */
 	static bool IsOperatorSpawnerRegistered(UFunction const* const Func);
 
+	/** Function node spawner associated with this operation */
 	static UBlueprintFunctionNodeSpawner* GetOperatorSpawner(const FString& OpName);
 
 	/** keep track of the operator that this function provides so that we dont add multiple 
@@ -149,9 +155,6 @@ private:
 
 	static FTypePromotion* Instance;
 
-	/** Delegate that gets called when a module is reloaded so we can rebuild the op table */
-	FDelegateHandle OnModulesChangedDelegateHandle;
-
 	/** A map of 'Type' to its 'available promotions'. See ctor for creation */
 	TMap<FName, TArray<FName>> PromotionTable;
 
diff --git a/Engine/Source/Editor/BlueprintGraph/Private/BlueprintTypePromotion.cpp b/Engine/Source/Editor/BlueprintGraph/Private/BlueprintTypePromotion.cpp
--- a/Engine/Source/Editor/BlueprintGraph/Private/BlueprintTypePromotion.cpp
+++ b/Engine/Source/Editor/BlueprintGraph/Private/BlueprintTypePromotion.cpp
@@ -65,21 +65,10 @@ FTypePromotion::FTypePromotion()
 {
 	CreatePromotionTable();
 	CreateOpTable();
-
-	OnModulesChangedDelegateHandle = FModuleManager::Get().OnModulesChanged().AddStatic(&OnModulesChanged);
 }
 
 FTypePromotion::~FTypePromotion()
 {
-	FModuleManager::Get().OnModulesChanged().Remove(OnModulesChangedDelegateHandle);
-}
-
-void FTypePromotion::OnModulesChanged(FName ModuleThatChanged, EModuleChangeReason ReasonForChange)
-{
-	if (Instance)
-	{
-		Instance->CreateOpTable();
-	}
 }
 
 void FTypePromotion::CreatePromotionTable()
@@ -615,6 +604,7 @@ void FTypePromotion::GetAllFuncsForOp_Internal(const FString& Operation, TArray<
 
 bool FTypePromotion::GetOpNameFromFunction(UFunction const* const Func, FString& OutName)
 {
+	TRACE_CPUPROFILER_EVENT_SCOPE(FTypePromotion::GetOpNameFromFunction);
 	if(!Func)
 	{
 		return false;
