Blueprints: Update type promotion to use FName's instead of FString's for the operation names. This has no behavioral impact, and saves about 30% of processing time per for each FindBestMatchingFunction call

#rnx

diff --git a/Engine/Source/Editor/BlueprintGraph/Classes/BlueprintTypePromotion.h b/Engine/Source/Editor/BlueprintGraph/Classes/BlueprintTypePromotion.h
--- a/Engine/Source/Editor/BlueprintGraph/Classes/BlueprintTypePromotion.h
+++ b/Engine/Source/Editor/BlueprintGraph/Classes/BlueprintTypePromotion.h
@@ -31,35 +31,35 @@ public:
 	* 
 	* @return	A function that matches a promoted type. Nullptr if no match is found
 	*/
-	static UFunction* GetOperatorFunction(const FString& Operation, const TArray<UEdGraphPin*>& WildcardPins);
+	static UFunction* GetOperatorFunction(FName Operation, const TArray<UEdGraphPin*>& WildcardPins);
 
-	static UFunction* GetOperatorFunction(const FString& Operation, const TArray<UEdGraphPin*>& InputPins, const UEdGraphPin* OutputPin);
+	static UFunction* GetOperatorFunction(FName Operation, const TArray<UEdGraphPin*>& InputPins, const UEdGraphPin* OutputPin);
 
 	/** 
 	* Find the function that has this input and the lowest matching other input.
 	* Ex: Given "Add" and "Vector" this function would return the "Add_VectorFloat" function
 	*/
-	static UFunction* FindLowestMatchingFunc(const FString& Operation, const FEdGraphPinType& InputType, TArray<UFunction*>& OutPossibleFunctions);
+	static UFunction* FindLowestMatchingFunc(FName Operation, const FEdGraphPinType& InputType, TArray<UFunction*>& OutPossibleFunctions);
 
 	/**
 	* Find the function that is the best match given the pins to consider. 
 	* Ex: Given "Add" operator and an array of two Vector pins, it will return "Add_VectorVector"
 	*/
-	static UFunction* FindBestMatchingFunc(const FString& Operation, const TArray<UEdGraphPin*>& PinsToConsider);
+	static UFunction* FindBestMatchingFunc(FName Operation, const TArray<UEdGraphPin*>& PinsToConsider);
 
 	/** Returns all functions for a specific operation. Will empty the given array and populate it with UFunction pointers */
-	static void GetAllFuncsForOp(const FString& Operation, TArray<UFunction*>& OutFuncs);
+	static void GetAllFuncsForOp(FName Operation, TArray<UFunction*>& OutFuncs);
 
 	/** Get a set of the supported operator names for type promo. Ex: "Add", "Subtract", "Multiply" */
-	static const TSet<FString>& GetAllOpNames();
+	static const TSet<FName>& GetAllOpNames();
 	
 	/** Set of comparison operator names (GreaterThan, LessThan, etc) */
-	static const TSet<FString>& GetComparisonOpNames();
+	static const TSet<FName>& GetComparisonOpNames();
 
 	/** Returns true if the given function is a comparison operator */
 	static bool IsComparisonFunc(UFunction const* const Func);
 
-	static bool GetOpNameFromFunction(UFunction const* const Func, FString& OutName);
+	static FName GetOpNameFromFunction(UFunction const* const Func);
 
 	/** Returns true if the given function is a candidate to handle type promotion */
 	static bool IsPromotableFunction(const UFunction* Function);
@@ -106,11 +106,11 @@ public:
 	static bool IsOperatorSpawnerRegistered(UFunction const* const Func);
 
 	/** Function node spawner associated with this operation */
-	static UBlueprintFunctionNodeSpawner* GetOperatorSpawner(const FString& OpName);
+	static UBlueprintFunctionNodeSpawner* GetOperatorSpawner(FName OpName);
 
 	/** keep track of the operator that this function provides so that we dont add multiple 
 	to the BP context menu */
-	static void RegisterOperatorSpawner(const FString& OpName, UBlueprintFunctionNodeSpawner* Spawner);
+	static void RegisterOperatorSpawner(FName OpName, UBlueprintFunctionNodeSpawner* Spawner);
 
 	static void ClearNodeSpawners();
 
@@ -122,19 +122,19 @@ private:
 	/** Callback that will rebuild the op table when hot reload is triggered */
 	static void OnModulesChanged(FName ModuleThatChanged, EModuleChangeReason ReasonForChange);
 
-	UFunction* GetOperatorFunction_Internal(const FString& Operation, const TArray<UEdGraphPin*>& WildcardPins) const;
+	UFunction* GetOperatorFunction_Internal(FName Operation, const TArray<UEdGraphPin*>& WildcardPins) const;
 
-	UFunction* GetOperatorFunction_Internal(const FString& Operation, const TArray<UEdGraphPin*>& InputPins, const UEdGraphPin* OutputPin) const;
+	UFunction* GetOperatorFunction_Internal(FName Operation, const TArray<UEdGraphPin*>& InputPins, const UEdGraphPin* OutputPin) const;
 
 	bool IsFunctionPromotionReady_Internal(const UFunction* const FuncToConsider) const;
 
 	FEdGraphPinType GetPromotedType_Internal(const TArray<UEdGraphPin*>& WildcardPins) const;
 
-	UFunction* FindLowestMatchingFunc_Internal(const FString& Operation, const FEdGraphPinType& InputType, TArray<UFunction*>& OutPossibleFunctions);
+	UFunction* FindLowestMatchingFunc_Internal(FName Operation, const FEdGraphPinType& InputType, TArray<UFunction*>& OutPossibleFunctions);
 	
-	UFunction* FindBestMatchingFunc_Internal(const FString& Operation, const TArray<UEdGraphPin*>& PinsToConsider);
+	UFunction* FindBestMatchingFunc_Internal(FName Operation, const TArray<UEdGraphPin*>& PinsToConsider);
 
-	void GetAllFuncsForOp_Internal(const FString& Operation, TArray<UFunction*>& OutFuncs);
+	void GetAllFuncsForOp_Internal(FName Operation, TArray<UFunction*>& OutFuncs);
 
 	bool PromotePin_Internal(FEdGraphPinType& InTypeA, const FEdGraphPinType& TypeB);
 
@@ -151,7 +151,7 @@ private:
 	/** Creates the table of what types can be promoted to others */
 	void CreatePromotionTable();
 
-	void AddOpFunction(const FString& OpName, UFunction* Function);
+	void AddOpFunction(FName OpName, UFunction* Function);
 
 	static FTypePromotion* Instance;
 
@@ -169,10 +169,10 @@ private:
 	 * A lookup table, mapping operator strings (like "Add", "Multiply", etc.) to a list
 	 * of associated functions.
 	 */
-	TMap<FString, FFunctionsList> OperatorTable;
+	TMap<FName, FFunctionsList> OperatorTable;
 
 	/** Map of operators to their node spawner so that we can clean up the context menu */
-	TMap<FString, UBlueprintFunctionNodeSpawner*> OperatorNodeSpawnerMap;
+	TMap<FName, UBlueprintFunctionNodeSpawner*> OperatorNodeSpawnerMap;
 };
 
 namespace TypePromoDebug
diff --git a/Engine/Source/Editor/BlueprintGraph/Classes/K2Node_PromotableOperator.h b/Engine/Source/Editor/BlueprintGraph/Classes/K2Node_PromotableOperator.h
--- a/Engine/Source/Editor/BlueprintGraph/Classes/K2Node_PromotableOperator.h
+++ b/Engine/Source/Editor/BlueprintGraph/Classes/K2Node_PromotableOperator.h
@@ -111,7 +111,7 @@ private:
 	TArray<UEdGraphPin*> GetInputPins(bool bIncludeLinks = false) const;
 
 	/** The name that this operation uses ("Add", "Multiply", etc) */
-	FString OperationName;
+	FName OperationName;
 
 	/** Array of functions that we could possibly convert this node to via the right-click context menu */
 	UPROPERTY(Transient)
@@ -126,5 +126,5 @@ public:
 	/** Returns the first pin with the EGPD_Output direction */
 	BLUEPRINTGRAPH_API UEdGraphPin* GetOutputPin() const;
 
-	const FString& GetOperationName() const { return OperationName; }
+	const FName GetOperationName() const { return OperationName; }
 };
diff --git a/Engine/Source/Editor/BlueprintGraph/Private/BlueprintFunctionNodeSpawner.cpp b/Engine/Source/Editor/BlueprintGraph/Private/BlueprintFunctionNodeSpawner.cpp
--- a/Engine/Source/Editor/BlueprintGraph/Private/BlueprintFunctionNodeSpawner.cpp
+++ b/Engine/Source/Editor/BlueprintGraph/Private/BlueprintFunctionNodeSpawner.cpp
@@ -238,8 +238,7 @@ UBlueprintFunctionNodeSpawner* UBlueprintFunctionNodeSpawner::Create(TSubclassOf
 		TypePromoDebug::IsTypePromoEnabled() &&
 		FTypePromotion::IsFunctionPromotionReady(Function);
 
-	FString OpName;
-	FTypePromotion::GetOpNameFromFunction(Function, OpName);
+	FName OpName = FTypePromotion::GetOpNameFromFunction(Function);
 
 	// If a spawner for this operator has been created already, than just return that
 	if (bIsPromotableFunction && FTypePromotion::IsOperatorSpawnerRegistered(Function))
@@ -271,13 +270,13 @@ UBlueprintFunctionNodeSpawner* UBlueprintFunctionNodeSpawner::Create(TSubclassOf
 	if(bIsPromotableFunction)
 	{
 		FFormatNamedArguments Args;
-		Args.Add(TEXT("OpName"), FText::FromString(OpName));
+		Args.Add(TEXT("OpName"), FText::FromName(OpName));
 		MenuSignature.MenuName = FText::Format(LOCTEXT("OpName", "{OpName}"), Args);
 		MenuSignature.Category = LOCTEXT("UtilityOperatorCategory", "Utilities|Operators");
 		// Possibly generate some special tooltips for promotable operators?
 		MenuSignature.Tooltip = FText::FromString(TEXT(" "));
 		// add at least one character, so that PrimeDefaultUiSpec() doesn't attempt to query the template node
-		MenuSignature.Keywords = FText::FromString(OpName);
+		MenuSignature.Keywords = FText::FromString(OpName.ToString());
 		FTypePromotion::RegisterOperatorSpawner(OpName, NodeSpawner);
 	}
 	else
diff --git a/Engine/Source/Editor/BlueprintGraph/Private/BlueprintTypePromotion.cpp b/Engine/Source/Editor/BlueprintGraph/Private/BlueprintTypePromotion.cpp
--- a/Engine/Source/Editor/BlueprintGraph/Private/BlueprintTypePromotion.cpp
+++ b/Engine/Source/Editor/BlueprintGraph/Private/BlueprintTypePromotion.cpp
@@ -9,18 +9,18 @@ FTypePromotion* FTypePromotion::Instance = nullptr;
 
 namespace OperatorNames
 {
-	static const FString NoOp		= TEXT("NO_OP");
+	static const FName NoOp			= TEXT("NO_OP");
 
-	static const FString Add		= TEXT("Add");
-	static const FString Multiply	= TEXT("Multiply");
-	static const FString Subtract	= TEXT("Subtract");
-	static const FString Divide		= TEXT("Divide");
-
-	static const FString Greater	= TEXT("Greater");
-	static const FString GreaterEq	= TEXT("GreaterEqual");
-	static const FString Less		= TEXT("Less");
-	static const FString LessEq		= TEXT("LessEqual");
-	static const FString NotEq		= TEXT("NotEqual");
+	static const FName Add			= TEXT("Add");
+	static const FName Multiply		= TEXT("Multiply");
+	static const FName Subtract		= TEXT("Subtract");
+	static const FName Divide		= TEXT("Divide");
+	
+	static const FName Greater		= TEXT("Greater");
+	static const FName GreaterEq	= TEXT("GreaterEqual");
+	static const FName Less			= TEXT("Less");
+	static const FName LessEq		= TEXT("LessEqual");
+	static const FName NotEq		= TEXT("NotEqual");
 }
 
 namespace TypePromoDebug
@@ -172,7 +172,7 @@ bool FTypePromotion::IsFunctionPromotionReady_Internal(const UFunction* const Fu
 {
 	// This could be better served as just keeping all known UFunctions in a TArray or TSet, and
 	// using that index in the Map so we don't need to iterate table pairs and arrays inside of them
-	for (const TPair<FString, FFunctionsList>& Pair : OperatorTable)
+	for (const TPair<FName, FFunctionsList>& Pair : OperatorTable)
 	{
 		for(const UFunction* const Func : Pair.Value)
 		{
@@ -234,17 +234,17 @@ bool FTypePromotion::PromotePin_Internal(FEdGraphPinType& InTypeA, const FEdGrap
 	return true;
 }
 
-UFunction* FTypePromotion::GetOperatorFunction(const FString& Operation, const TArray<UEdGraphPin*>& WildcardPins)
+UFunction* FTypePromotion::GetOperatorFunction(FName Operation, const TArray<UEdGraphPin*>& WildcardPins)
 {
 	return FTypePromotion::Get().GetOperatorFunction_Internal(Operation, WildcardPins);
 }
 
-UFunction* FTypePromotion::GetOperatorFunction(const FString& Operation, const TArray<UEdGraphPin*>& InputPins, const UEdGraphPin* OutputPin)
+UFunction* FTypePromotion::GetOperatorFunction(FName Operation, const TArray<UEdGraphPin*>& InputPins, const UEdGraphPin* OutputPin)
 {
 	return FTypePromotion::Get().GetOperatorFunction_Internal(Operation, InputPins, OutputPin);
 }
 
-UFunction* FTypePromotion::GetOperatorFunction_Internal(const FString& Operation, const TArray<UEdGraphPin*>& InputPins, const UEdGraphPin* OutputPin) const
+UFunction* FTypePromotion::GetOperatorFunction_Internal(FName Operation, const TArray<UEdGraphPin*>& InputPins, const UEdGraphPin* OutputPin) const
 {
 	check(InputPins.Num() > 0 && OutputPin);
 
@@ -303,7 +303,7 @@ UFunction* FTypePromotion::GetOperatorFunction_Internal(const FString& Operation
 	return MatchingFunction;
 }
 
-UFunction* FTypePromotion::GetOperatorFunction_Internal(const FString& Operation, const TArray<UEdGraphPin*>& WildcardPins) const
+UFunction* FTypePromotion::GetOperatorFunction_Internal(FName Operation, const TArray<UEdGraphPin*>& WildcardPins) const
 {
 	TRACE_CPUPROFILER_EVENT_SCOPE(FTypePromotionTable::GetOperatorFunction);
 
@@ -361,12 +361,12 @@ UFunction* FTypePromotion::GetOperatorFunction_Internal(const FString& Operation
 	return MatchingFunction;
 }
 
-UFunction* FTypePromotion::FindLowestMatchingFunc(const FString& Operation, const FEdGraphPinType& InputType, TArray<UFunction*>& OutPossibleFunctions)
+UFunction* FTypePromotion::FindLowestMatchingFunc(FName Operation, const FEdGraphPinType& InputType, TArray<UFunction*>& OutPossibleFunctions)
 {
 	return FTypePromotion::Get().FindLowestMatchingFunc_Internal(Operation, InputType, OutPossibleFunctions);
 }
 
-UFunction* FTypePromotion::FindBestMatchingFunc(const FString& Operation, const TArray<UEdGraphPin*>& PinsToConsider)
+UFunction* FTypePromotion::FindBestMatchingFunc(FName Operation, const TArray<UEdGraphPin*>& PinsToConsider)
 {
 	return FTypePromotion::Get().FindBestMatchingFunc_Internal(Operation, PinsToConsider);
 }
@@ -387,7 +387,7 @@ static bool PropertyCompatibleWithPin(const FProperty* Param, FEdGraphPinType co
 	return false;
 }
 
-UFunction* FTypePromotion::FindLowestMatchingFunc_Internal(const FString& Operation, const FEdGraphPinType& InputType, TArray<UFunction*>& OutPossibleFunctions)
+UFunction* FTypePromotion::FindLowestMatchingFunc_Internal(FName Operation, const FEdGraphPinType& InputType, TArray<UFunction*>& OutPossibleFunctions)
 {
 	TRACE_CPUPROFILER_EVENT_SCOPE(FTypePromotionTable::FindLowestMatchingFunc_Internal);
 
@@ -452,7 +452,7 @@ UFunction* FTypePromotion::FindLowestMatchingFunc_Internal(const FString& Operat
 	return LowestFunc;
 }
 
-UFunction* FTypePromotion::FindBestMatchingFunc_Internal(const FString& Operation, const TArray<UEdGraphPin*>& PinsToConsider)
+UFunction* FTypePromotion::FindBestMatchingFunc_Internal(FName Operation, const TArray<UEdGraphPin*>& PinsToConsider)
 {
 	TRACE_CPUPROFILER_EVENT_SCOPE(FTypePromotionTable::FindBestMatchingFunc_Internal);
 
@@ -553,14 +553,14 @@ UFunction* FTypePromotion::FindBestMatchingFunc_Internal(const FString& Operatio
 	return BestFunc;
 }
 
-void FTypePromotion::GetAllFuncsForOp(const FString& Operation, TArray<UFunction*>& OutFuncs)
+void FTypePromotion::GetAllFuncsForOp(FName Operation, TArray<UFunction*>& OutFuncs)
 {
 	return FTypePromotion::Get().GetAllFuncsForOp_Internal(Operation, OutFuncs);
 }
 
-const TSet<FString>& FTypePromotion::GetAllOpNames()
+const TSet<FName>& FTypePromotion::GetAllOpNames()
 {
-	static const TSet<FString> OpsArray =
+	static const TSet<FName> OpsArray =
 	{
 		OperatorNames::Add,
 		OperatorNames::Multiply,
@@ -576,9 +576,9 @@ const TSet<FString>& FTypePromotion::GetAllOpNames()
 	return OpsArray;
 }
 
-const TSet<FString>& FTypePromotion::GetComparisonOpNames()
+const TSet<FName>& FTypePromotion::GetComparisonOpNames()
 {
-	static const TSet<FString> ComparisonOps =
+	static const TSet<FName> ComparisonOps =
 	{
 		OperatorNames::Greater,
 		OperatorNames::GreaterEq,
@@ -591,39 +591,35 @@ const TSet<FString>& FTypePromotion::GetComparisonOpNames()
 
 bool FTypePromotion::IsComparisonFunc(UFunction const* const Func)
 {
-	FString OpName;
-	return Func && GetOpNameFromFunction(Func, /* out */ OpName) && GetComparisonOpNames().Contains(OpName);
+	return Func && GetComparisonOpNames().Contains(GetOpNameFromFunction(Func));
 }
 
-void FTypePromotion::GetAllFuncsForOp_Internal(const FString& Operation, TArray<UFunction*>& OutFuncs)
+void FTypePromotion::GetAllFuncsForOp_Internal(FName Operation, TArray<UFunction*>& OutFuncs)
 {
 	OutFuncs.Empty();
 
 	OutFuncs.Append(OperatorTable[Operation]);
 }
 
-bool FTypePromotion::GetOpNameFromFunction(UFunction const* const Func, FString& OutName)
+FName FTypePromotion::GetOpNameFromFunction(UFunction const* const Func)
 {
 	TRACE_CPUPROFILER_EVENT_SCOPE(FTypePromotion::GetOpNameFromFunction);
 	if(!Func)
 	{
-		return false;
+		return OperatorNames::NoOp;
 	}
 
-	const FString& FuncName = Func->GetName();
+	FString FuncName = Func->GetName();
 	// Get everything before the "_"
 	int32 Index = FuncName.Find(TEXT("_"));
-	FString FuncNameChopped = FuncName.Mid(0, Index);
+	
+	FName FuncNameChopped(FuncName.Mid(0, Index));
 	if (GetAllOpNames().Contains(FuncNameChopped))
 	{
-		OutName = FuncNameChopped;
-		return true;
-	}
-	else
-	{
-		OutName = OperatorNames::NoOp;
-		return false;
+		return FuncNameChopped;
 	}
+
+	return OperatorNames::NoOp;
 }
 
 void FTypePromotion::CreateOpTable()
@@ -651,11 +647,9 @@ void FTypePromotion::CreateOpTable()
 			}
 
 			FEdGraphPinType FuncPinType;
-			FString OpName;
+			FName OpName = GetOpNameFromFunction(Function);
 
-			if (Schema->ConvertPropertyToPinType(Function->GetReturnProperty(), /* out */ FuncPinType) &&
-				GetOpNameFromFunction(Function, /* out */ OpName)
-				)
+			if (OpName != OperatorNames::NoOp && Schema->ConvertPropertyToPinType(Function->GetReturnProperty(), /* out */ FuncPinType))
 			{
 				AddOpFunction(OpName, Function);
 			}
@@ -663,7 +657,7 @@ void FTypePromotion::CreateOpTable()
 	}
 }
 
-void FTypePromotion::AddOpFunction(const FString& OpName, UFunction* Function)
+void FTypePromotion::AddOpFunction(FName OpName, UFunction* Function)
 {
 	OperatorTable.FindOrAdd(OpName).Add(Function);
 }
@@ -675,12 +669,10 @@ bool FTypePromotion::IsPromotableFunction(const UFunction* Function)
 
 bool FTypePromotion::IsOperatorSpawnerRegistered(UFunction const* const Func)
 {
-	FString OpName;
-	FTypePromotion::GetOpNameFromFunction(Func, OpName);
-	return FTypePromotion::GetOperatorSpawner(OpName) != nullptr;
+	return FTypePromotion::GetOperatorSpawner(FTypePromotion::GetOpNameFromFunction(Func)) != nullptr;
 }
 
-void FTypePromotion::RegisterOperatorSpawner(const FString& OpName, UBlueprintFunctionNodeSpawner* Spawner)
+void FTypePromotion::RegisterOperatorSpawner(FName OpName, UBlueprintFunctionNodeSpawner* Spawner)
 {
 	if(Instance && !Instance->OperatorNodeSpawnerMap.Contains(OpName) && OpName != OperatorNames::NoOp)
 	{
@@ -688,7 +680,7 @@ void FTypePromotion::RegisterOperatorSpawner(const FString& OpName, UBlueprintFu
 	}
 }
 
-UBlueprintFunctionNodeSpawner* FTypePromotion::GetOperatorSpawner(const FString& OpName)
+UBlueprintFunctionNodeSpawner* FTypePromotion::GetOperatorSpawner(FName OpName)
 {
 	if(Instance && Instance->OperatorNodeSpawnerMap.Contains(OpName))
 	{
diff --git a/Engine/Source/Editor/BlueprintGraph/Private/K2Node_PromotableOperator.cpp b/Engine/Source/Editor/BlueprintGraph/Private/K2Node_PromotableOperator.cpp
--- a/Engine/Source/Editor/BlueprintGraph/Private/K2Node_PromotableOperator.cpp
+++ b/Engine/Source/Editor/BlueprintGraph/Private/K2Node_PromotableOperator.cpp
@@ -120,7 +120,7 @@ void UK2Node_PromotableOperator::GetNodeContextMenuActions(UToolMenu* Menu, UGra
 				FUIAction(
 					FExecuteAction::CreateUObject(const_cast<UK2Node_PromotableOperator*>(this), &UK2Node_PromotableOperator::ConvertNodeToFunction, Func, const_cast<UEdGraphPin*>(Context->Pin))
 				)
-			);			
+			);
 		}
 	}
 }
@@ -131,9 +131,9 @@ FText UK2Node_PromotableOperator::GetTooltipText() const
 	if (!HasAnyConnectionsOrDefaults())
 	{
 		UFunction* Function = GetTargetFunction();
-		FString OpName;
-		FTypePromotion::GetOpNameFromFunction(Function, OpName);
-		return FText::Format(LOCTEXT("PromotableOperatorFunctionTooltip", "{0} Operator"), FText::FromString(OpName));
+
+		FName OpName = FTypePromotion::GetOpNameFromFunction(Function);
+		return FText::Format(LOCTEXT("PromotableOperatorFunctionTooltip", "{0} Operator"), FText::FromName(OpName));
 	}
 
 	// Otherwise use the default one (a more specific function tooltip)
@@ -243,7 +243,7 @@ void UK2Node_PromotableOperator::ExpandNode(FKismetCompilerContext& CompilerCont
 			CompilerContext.MessageLog.Error(TEXT("'@@' could not successfuly expand pins!"), PrevIntermediateNode);
 		}
 	}
-	
+
 	// Loop through all the additional inputs, create a new node of this function and connecting inputs as necessary 
 	for (int32 i = NumFunctionInputs; i < NumAdditionalInputs + NumFunctionInputs; ++i)
 	{
@@ -254,13 +254,13 @@ void UK2Node_PromotableOperator::ExpandNode(FKismetCompilerContext& CompilerCont
 		// so that we can avoid unnecessary conversion nodes and casts
 		UFunction* BestMatchingFunc = OpFunction;
 		{
-			TArray<UEdGraphPin*> PinsToConsider = 
+			TArray<UEdGraphPin*> PinsToConsider =
 			{
 				PrevNodeHelper.OutputPin,
 				OriginalInputPins[i],
 				OriginalOutputPin
 			};
-			
+
 			if (UFunction* Func = FTypePromotion::FindBestMatchingFunc(OperationName, PinsToConsider))
 			{
 				BestMatchingFunc = Func;
@@ -272,10 +272,10 @@ void UK2Node_PromotableOperator::ExpandNode(FKismetCompilerContext& CompilerCont
 
 		// Connect the output pin of the previous intermediate node, to the input of the new one
 		const bool bPinASuccess = CreateIntermediateCast(PrevIntermediateNode, CompilerContext, SourceGraph, NewOpHelper.InputA, PrevOutputPin);
-		
+
 		// Connect the original node's pin to the newly created intermediate node's B Pin
 		const bool bPinBSuccess = CreateIntermediateCast(this, CompilerContext, SourceGraph, OriginalInputPins[i], NewOpHelper.InputB);
-		
+
 		if (!bPinASuccess || !bPinBSuccess)
 		{
 			CompilerContext.MessageLog.Error(TEXT("'@@' could not successfuly expand additional pins!"), PrevIntermediateNode);
@@ -389,7 +389,7 @@ void UK2Node_PromotableOperator::ReallocatePinsDuringReconstruction(TArray<UEdGr
 				AddPin->PinType = OldPin->PinType;
 				AddPin->DefaultValue = OldPin->DefaultValue;
 				++AdditionalPinsFixed;
-			}			
+			}
 		}
 	}
 }
@@ -440,7 +440,7 @@ bool UK2Node_PromotableOperator::CanRemovePin(const UEdGraphPin* Pin) const
 		NumAdditionalInputs > 0 &&
 		INDEX_NONE != Pins.IndexOfByKey(Pin) &&
 		Pin->Direction == EEdGraphPinDirection::EGPD_Input
-	);
+		);
 }
 
 void UK2Node_PromotableOperator::RemoveInputPin(UEdGraphPin* Pin)
@@ -567,7 +567,7 @@ void UK2Node_PromotableOperator::EvaluatePinsFromChange(UEdGraphPin* ChangedPin)
 				PinsToConsider.Emplace(Link);
 			}
 		}
-		else if(!Pin->DoesDefaultValueMatchAutogenerated() && !FWildcardNodeUtils::IsWildcardPin(Pin))
+		else if (!Pin->DoesDefaultValueMatchAutogenerated() && !FWildcardNodeUtils::IsWildcardPin(Pin))
 		{
 			// If this pin has doesn't have any links then we should account for the pin itself
 			PinsToConsider.Add(Pin);
@@ -600,11 +600,16 @@ void UK2Node_PromotableOperator::EvaluatePinsFromChange(UEdGraphPin* ChangedPin)
 
 bool UK2Node_PromotableOperator::UpdateOpName()
 {
-	const UFunction* Func = GetTargetFunction();
+
 	// If the function is null then return false, because we did not successfully update it. 
 	// This could be possible during node reconstruction/refresh, and we don't want to set the 
 	// op name to "Empty" incorrectly. 
-	return Func ? FTypePromotion::GetOpNameFromFunction(Func, /* out */ OperationName) : false;
+	if (const UFunction* Func = GetTargetFunction())
+	{
+		OperationName = FTypePromotion::GetOpNameFromFunction(Func);
+		return true;
+	}
+	return false;
 }
 
 UK2Node_CallFunction* UK2Node_PromotableOperator::CreateIntermediateNode(UK2Node_CallFunction* PreviousNode, const UFunction* const OpFunction, FKismetCompilerContext& CompilerContext, UEdGraph* SourceGraph)
@@ -639,7 +644,7 @@ bool UK2Node_PromotableOperator::CreateIntermediateCast(UK2Node_CallFunction* So
 			return !CompilerContext.MovePinLinksToIntermediate(*InputPin, *OutputPin).IsFatal();
 		}
 		else
-		{			
+		{
 			return Schema->TryCreateConnection(InputPin, OutputPin);
 		}
 	}
@@ -685,7 +690,7 @@ bool UK2Node_PromotableOperator::CreateIntermediateCast(UK2Node_CallFunction* So
 		{
 			bInputSuccessful = Schema->TryCreateConnection(InputPin->LinkedTo[0], ConversionInput);
 		}
-		else if(InputPin && ConversionInput)
+		else if (InputPin && ConversionInput)
 		{
 			bInputSuccessful = !CompilerContext.MovePinLinksToIntermediate(*InputPin, *ConversionInput).IsFatal();
 		}
@@ -737,7 +742,7 @@ void UK2Node_PromotableOperator::ResetNodeToWildcard()
 			OutPin->PinType.PinCategory = UEdGraphSchema_K2::PC_Boolean;
 		}
 	}
-	
+
 	// Clear out any possible function matches, since we are removing connections
 	PossibleConversions.Empty();
 }
@@ -747,7 +752,7 @@ void UK2Node_PromotableOperator::RecombineAllSplitPins()
 	const UEdGraphSchema_K2* K2Schema = GetDefault<UEdGraphSchema_K2>();
 
 	// Gather what pins need to be recombined from a split pin	
-	for (int32 Index = 0; Index < Pins.Num(); ++Index) 
+	for (int32 Index = 0; Index < Pins.Num(); ++Index)
 	{
 		if (Pins[Index] && Pins[Index]->SubPins.Num() > 0)
 		{
@@ -782,7 +787,7 @@ void UK2Node_PromotableOperator::ConvertNodeToFunction(const UFunction* Function
 	const FScopedTransaction Transaction(LOCTEXT("ConvertPromotableOpToFunction", "Change the function signature of a promotable operator node."));
 	Modify();
 	RecombineAllSplitPins();
-	
+
 	// If we convert the node to a function, then just get rid of the additional pins
 	NumAdditionalInputs = 0;
 
@@ -862,7 +867,7 @@ void UK2Node_PromotableOperator::UpdatePinsFromFunction(const UFunction* Functio
 				{
 					FEdGraphPinType LinkedType = FTypePromotion::GetPromotedType(PinToChange->LinkedTo);
 					// If this is a wildcard pin, then check the compatibility with a linked type instead of the wildcard itself
-					if (FTypePromotion::IsValidPromotion(bIsWildcard  ? LinkedType : PinToChange->PinType, ParamType))
+					if (FTypePromotion::IsValidPromotion(bIsWildcard ? LinkedType : PinToChange->PinType, ParamType))
 					{
 						ParamType = LinkedType;
 					}
