Blueprints: Notify graph of change when reseting a wildcard node to properly remove the default value inputs. Update the functional test to have a proper example of the Additional Pins interface with this.

#rnx

diff --git a/Engine/Source/Editor/BlueprintGraph/Private/K2Node_PromotableOperator.cpp b/Engine/Source/Editor/BlueprintGraph/Private/K2Node_PromotableOperator.cpp
--- a/Engine/Source/Editor/BlueprintGraph/Private/K2Node_PromotableOperator.cpp
+++ b/Engine/Source/Editor/BlueprintGraph/Private/K2Node_PromotableOperator.cpp
@@ -594,6 +594,8 @@ void UK2Node_PromotableOperator::EvaluatePinsFromChange(UEdGraphPin* ChangedPin)
 		ChangedPin->PinType = NewConnectionHighestType;
 	}
 
+	GetGraph()->NotifyGraphChanged();
+
 	// Update context menu options for this node
 	UpdatePossibleConversionFuncs();
 }
@@ -720,7 +722,7 @@ void UK2Node_PromotableOperator::ResetNodeToWildcard()
 
 	// Reset type to wildcard
 	const FEdGraphPinType& WildType = FWildcardNodeUtils::GetDefaultWildcardPinType();
-	const UEdGraphSchema_K2* K2Schema = GetDefault<UEdGraphSchema_K2>();
+	const UEdGraphSchema* Schema = GetSchema();
 
 	for (UEdGraphPin* Pin : Pins)
 	{
@@ -728,7 +730,7 @@ void UK2Node_PromotableOperator::ResetNodeToWildcard()
 		if (Pin && Pin->ParentPin == nullptr)
 		{
 			Pin->PinType = WildType;
-			K2Schema->ResetPinToAutogeneratedDefaultValue(Pin);
+			Schema->ResetPinToAutogeneratedDefaultValue(Pin);
 		}
 	}
 
@@ -745,6 +747,7 @@ void UK2Node_PromotableOperator::ResetNodeToWildcard()
 
 	// Clear out any possible function matches, since we are removing connections
 	PossibleConversions.Empty();
+	GetGraph()->NotifyGraphChanged();
 }
 
 void UK2Node_PromotableOperator::RecombineAllSplitPins()
