Blueprints: Remove unused function in type promotion

#rnx

diff --git a/Engine/Source/Editor/BlueprintGraph/Classes/BlueprintTypePromotion.h b/Engine/Source/Editor/BlueprintGraph/Classes/BlueprintTypePromotion.h
--- a/Engine/Source/Editor/BlueprintGraph/Classes/BlueprintTypePromotion.h
+++ b/Engine/Source/Editor/BlueprintGraph/Classes/BlueprintTypePromotion.h
@@ -35,12 +35,6 @@ public:
 
 	static UFunction* GetOperatorFunction(FName Operation, const TArray<UEdGraphPin*>& InputPins, const UEdGraphPin* OutputPin);
 
-	/** 
-	* Find the function that has this input and the lowest matching other input.
-	* Ex: Given "Add" and "Vector" this function would return the "Add_VectorFloat" function
-	*/
-	static UFunction* FindLowestMatchingFunc(FName Operation, const FEdGraphPinType& InputType, TArray<UFunction*>& OutPossibleFunctions);
-
 	/**
 	* Find the function that is the best match given the pins to consider. 
 	* Ex: Given "Add" operator and an array of two Vector pins, it will return "Add_VectorVector"
@@ -129,8 +123,6 @@ private:
 	bool IsFunctionPromotionReady_Internal(const UFunction* const FuncToConsider) const;
 
 	FEdGraphPinType GetPromotedType_Internal(const TArray<UEdGraphPin*>& WildcardPins) const;
-
-	UFunction* FindLowestMatchingFunc_Internal(FName Operation, const FEdGraphPinType& InputType, TArray<UFunction*>& OutPossibleFunctions);
 	
 	UFunction* FindBestMatchingFunc_Internal(FName Operation, const TArray<UEdGraphPin*>& PinsToConsider);
 
diff --git a/Engine/Source/Editor/BlueprintGraph/Private/BlueprintTypePromotion.cpp b/Engine/Source/Editor/BlueprintGraph/Private/BlueprintTypePromotion.cpp
--- a/Engine/Source/Editor/BlueprintGraph/Private/BlueprintTypePromotion.cpp
+++ b/Engine/Source/Editor/BlueprintGraph/Private/BlueprintTypePromotion.cpp
@@ -361,11 +361,6 @@ UFunction* FTypePromotion::GetOperatorFunction_Internal(FName Operation, const T
 	return MatchingFunction;
 }
 
-UFunction* FTypePromotion::FindLowestMatchingFunc(FName Operation, const FEdGraphPinType& InputType, TArray<UFunction*>& OutPossibleFunctions)
-{
-	return FTypePromotion::Get().FindLowestMatchingFunc_Internal(Operation, InputType, OutPossibleFunctions);
-}
-
 UFunction* FTypePromotion::FindBestMatchingFunc(FName Operation, const TArray<UEdGraphPin*>& PinsToConsider)
 {
 	return FTypePromotion::Get().FindBestMatchingFunc_Internal(Operation, PinsToConsider);
@@ -387,71 +382,6 @@ static bool PropertyCompatibleWithPin(const FProperty* Param, FEdGraphPinType co
 	return false;
 }
 
-UFunction* FTypePromotion::FindLowestMatchingFunc_Internal(FName Operation, const FEdGraphPinType& InputType, TArray<UFunction*>& OutPossibleFunctions)
-{
-	TRACE_CPUPROFILER_EVENT_SCOPE(FTypePromotionTable::FindLowestMatchingFunc_Internal);
-
-	UFunction* LowestFunc = nullptr;
-	OutPossibleFunctions.Empty();
-
-	if (const FFunctionsList* FuncList = OperatorTable.Find(Operation))
-	{
-		const UEdGraphSchema_K2* Schema = GetDefault<UEdGraphSchema_K2>();
-
-		for (UFunction* Func : *FuncList)
-		{
-			// Find the function that matches the input type here, we don't care about the output type
-			for (TFieldIterator<FProperty> PropIt(Func); PropIt && (PropIt->PropertyFlags & CPF_Parm); ++PropIt)
-			{
-				// Ignore return params here, we only care about inputs
-				FProperty* Param = *PropIt;
-				if (!Param->HasAnyPropertyFlags(CPF_ReturnParm) && PropertyCompatibleWithPin(Param, InputType))
-				{
-					// If an input of this function is compatible with this type
-					// Then we need to check all the other inputs on this function
-					OutPossibleFunctions.AddUnique(Func);
-					break;
-				}
-			}
-		}
-	}
-
-	const UEdGraphSchema_K2* Schema = GetDefault<UEdGraphSchema_K2>();
-	FEdGraphPinType CurLowestType = {};
-
-	if(OutPossibleFunctions.Num() > 0)
-	{
-		LowestFunc = OutPossibleFunctions[0];
-	}
-
-	// Check possible matches to see who has the lowest type
-	for(UFunction* Func : OutPossibleFunctions)
-	{
-		for (TFieldIterator<FProperty> PropIt(Func); PropIt && (PropIt->PropertyFlags & CPF_Parm); ++PropIt)
-		{
-			// Check for the other input param that is not the input type we are lookin at
-			FProperty* Param = *PropIt;
-			FEdGraphPinType ParamType;
-			if (!Param->HasAnyPropertyFlags(CPF_ReturnParm) && Schema->ConvertPropertyToPinType(Param, /* out */ ParamType))
-			{
-				ETypeComparisonResult Res = FTypePromotion::GetHigherType(ParamType, CurLowestType);
-				if(Res == ETypeComparisonResult::TypeBHigher)
-				{
-					CurLowestType = ParamType;
-					LowestFunc = Func;
-				}
-				else if(Res == ETypeComparisonResult::TypeAHigher)
-				{
-					//CurLowestType = CurLowestType;
-					LowestFunc = Func;
-				}
-			}
-		}
-	}
-
-	return LowestFunc;
-}
-
 UFunction* FTypePromotion::FindBestMatchingFunc_Internal(FName Operation, const TArray<UEdGraphPin*>& PinsToConsider)
 {
 	TRACE_CPUPROFILER_EVENT_SCOPE(FTypePromotionTable::FindBestMatchingFunc_Internal);
diff --git a/Engine/Source/Editor/BlueprintGraph/Private/K2Node_PromotableOperator.cpp b/Engine/Source/Editor/BlueprintGraph/Private/K2Node_PromotableOperator.cpp
--- a/Engine/Source/Editor/BlueprintGraph/Private/K2Node_PromotableOperator.cpp
+++ b/Engine/Source/Editor/BlueprintGraph/Private/K2Node_PromotableOperator.cpp
@@ -579,9 +579,6 @@ void UK2Node_PromotableOperator::EvaluatePinsFromChange(UEdGraphPin* ChangedPin)
 	// if a pin was changed, update it if it cannot be promoted to this type
 	FEdGraphPinType NewConnectionHighestType = ChangedPin->LinkedTo.Num() > 0 ? FTypePromotion::GetPromotedType(ChangedPin->LinkedTo) : FWildcardNodeUtils::GetDefaultWildcardPinType();
 
-	// If there are ANY wildcards on this node, than we need to update the whole node accordingly. Otherwise we can 
-	// update only the Changed and output pins.
-
 	const UFunction* BestMatchingFunc = FTypePromotion::FindBestMatchingFunc(OperationName, PinsToConsider);
 
 	// Store these other function options for later so that the user can convert to them later
@@ -892,8 +889,6 @@ void UK2Node_PromotableOperator::UpdatePinsFromFunction(const UFunction* Functio
 
 	// Update the function reference and the FUNC_BlueprintPure/FUNC_Const appropriately
 	SetFromFunction(Function);
-
-	UpdatePossibleConversionFuncs();
 }
 
 void UK2Node_PromotableOperator::UpdatePossibleConversionFuncs()
@@ -923,9 +918,10 @@ void UK2Node_PromotableOperator::UpdatePossibleConversionFuncs()
 	// all the same. 
 	if (!bAllPinTypesEqual)
 	{
+		// #TODO_BH We need up change how these conversions work to be on a per-pin basis instead of via a function
+		// which will be a more consistent experience, and work better with additional pins
 		// The node has changed, so lets find the lowest matching function with the newly updated types
 		FEdGraphPinType HighestType = FTypePromotion::GetPromotedType(GetInputPins());
-		FTypePromotion::FindLowestMatchingFunc(OperationName, HighestType, PossibleConversions);
 	}
 }
 
