Change UEdGraphSchema_K2::ReplaceOldNodeWithNew to accept UEdGraphNode's instead of UK2Node's to make it useable from other schema's as well

#rnx

diff --git a/Engine/Source/Editor/BlueprintGraph/Classes/EdGraphSchema_K2.h b/Engine/Source/Editor/BlueprintGraph/Classes/EdGraphSchema_K2.h
--- a/Engine/Source/Editor/BlueprintGraph/Classes/EdGraphSchema_K2.h
+++ b/Engine/Source/Editor/BlueprintGraph/Classes/EdGraphSchema_K2.h
@@ -1117,7 +1117,7 @@ public:
 	void LinkDataPinFromOutputToInput(UEdGraphNode* InOutputNode, UEdGraphNode* InInputNode) const;
 
 	/** Moves all connections from the old node to the new one. Returns true and destroys OldNode on success. Fails if it cannot find a mapping from an old pin. */
-	bool ReplaceOldNodeWithNew(UK2Node* OldNode, UK2Node* NewNode, const TMap<FName, FName>& OldPinToNewPinMap) const;
+	bool ReplaceOldNodeWithNew(UEdGraphNode* OldNode, UEdGraphNode* NewNode, const TMap<FName, FName>& OldPinToNewPinMap) const;
 
 	/** Convert a deprecated node into a function call node, called from per-node ConvertDeprecatedNode */
 	UK2Node* ConvertDeprecatedNodeToFunctionCall(UK2Node* OldNode, UFunction* NewFunction, TMap<FName, FName>& OldPinToNewPinMap, UEdGraph* Graph) const;
diff --git a/Engine/Source/Editor/BlueprintGraph/Private/EdGraphSchema_K2.cpp b/Engine/Source/Editor/BlueprintGraph/Private/EdGraphSchema_K2.cpp
--- a/Engine/Source/Editor/BlueprintGraph/Private/EdGraphSchema_K2.cpp
+++ b/Engine/Source/Editor/BlueprintGraph/Private/EdGraphSchema_K2.cpp
@@ -5733,7 +5733,7 @@ bool UEdGraphSchema_K2::FadeNodeWhenDraggingOffPin(const UEdGraphNode* Node, con
 struct FBackwardCompatibilityConversionHelper
 {
 	// Re-add orphaned pins to deal with any links that were lost during converstion
-	static bool RestoreOrphanLinks(UEdGraphPin* OldPin, UEdGraphPin* NewPin, UK2Node* NewNode, const TArray<UEdGraphPin*>& OldLinks)
+	static bool RestoreOrphanLinks(UEdGraphPin* OldPin, UEdGraphPin* NewPin, UEdGraphNode* NewNode, const TArray<UEdGraphPin*>& OldLinks)
 	{
 		// See if there are any links that didn't get copied, including to orphan pins or if the newpin is null
 		TArray<UEdGraphPin*> OrphanedLinks;
@@ -6012,14 +6012,14 @@ struct FBackwardCompatibilityConversionHelper
 	}
 };
 
-bool UEdGraphSchema_K2::ReplaceOldNodeWithNew(UK2Node* OldNode, UK2Node* NewNode, const TMap<FName, FName>& OldPinToNewPinMap) const
+bool UEdGraphSchema_K2::ReplaceOldNodeWithNew(UEdGraphNode* OldNode, UEdGraphNode* NewNode, const TMap<FName, FName>& OldPinToNewPinMap) const
 {
 	if (!ensure(NewNode->GetGraph() == OldNode->GetGraph()))
 	{
 		return false;
 	}
-	const UBlueprint* Blueprint = OldNode->GetBlueprint();	
 	const UEdGraphSchema* Schema = NewNode->GetSchema();
+	const UObject* NodeOuter = NewNode->GetGraph() ? NewNode->GetGraph()->GetOuter() : nullptr;
 
 	NewNode->NodePosX = OldNode->NodePosX;
 	NewNode->NodePosY = OldNode->NodePosY;
@@ -6091,10 +6091,10 @@ bool UEdGraphSchema_K2::ReplaceOldNodeWithNew(UK2Node* OldNode, UK2Node* NewNode
 		{
 			bFailedToFindPin = true;
 
-			UE_LOG(LogBlueprint, Warning, TEXT("BackwardCompatibilityNodeConversion Error 'cannot find pin %s in node %s' in blueprint: %s"),
+			UE_LOG(LogBlueprint, Warning, TEXT("BackwardCompatibilityNodeConversion Error 'cannot find pin %s in node %s' in: %s"),
 				*OldPin->PinName.ToString(),
 				*NewNode->GetNodeTitle(ENodeTitleType::FullTitle).ToString(),
-				Blueprint ? *Blueprint->GetName() : TEXT("Unknown"));
+				NodeOuter ? *NodeOuter->GetName() : TEXT("Unknown"));
 
 			break;
 		}
@@ -6119,15 +6119,15 @@ bool UEdGraphSchema_K2::ReplaceOldNodeWithNew(UK2Node* OldNode, UK2Node* NewNode
 			}
 			else if (!Schema->MovePinLinks(*OldPin, *NewPin, false, true).CanSafeConnect())
 			{
-				UE_LOG(LogBlueprint, Warning, TEXT("BackwardCompatibilityNodeConversion Error 'cannot safely move pin %s to %s' in blueprint: %s"),
+				UE_LOG(LogBlueprint, Warning, TEXT("BackwardCompatibilityNodeConversion Error 'cannot safely move pin %s to %s' in: %s"),
 					*OldPin->PinName.ToString(),
 					*NewPin->PinName.ToString(),
-					Blueprint ? *Blueprint->GetName() : TEXT("Unknown"));
+					NodeOuter ? *NodeOuter->GetName() : TEXT("Unknown"));
 			}
-			else
+			else if(UK2Node* K2Node = Cast<UK2Node>(NewNode))
 			{
 				// for wildcard pins, which may have to react to being connected with
-				NewNode->NotifyPinConnectionListChanged(NewPin);
+				K2Node->NotifyPinConnectionListChanged(NewPin);
 			}
 
 			FBackwardCompatibilityConversionHelper::RestoreOrphanLinks(OldPin, NewPin, NewNode, OldLinks);
