Blueprints: Fix improper disconnect in type promotion pins when connecting a lower type to a node that already had an output pin connected with a valid promotion

#rnx

diff --git a/Engine/Source/Editor/BlueprintGraph/Private/K2Node_PromotableOperator.cpp b/Engine/Source/Editor/BlueprintGraph/Private/K2Node_PromotableOperator.cpp
--- a/Engine/Source/Editor/BlueprintGraph/Private/K2Node_PromotableOperator.cpp
+++ b/Engine/Source/Editor/BlueprintGraph/Private/K2Node_PromotableOperator.cpp
@@ -548,7 +548,7 @@ void UK2Node_PromotableOperator::EvaluatePinsFromChange(UEdGraphPin* ChangedPin)
 	}
 	// If the pin that was connected is linked to a wildcard pin, then we should make it a wildcard
 	// and do nothing else.
-	else if (FWildcardNodeUtils::IsWildcardPin(ChangedPin) && ChangedPin->GetOwningNode() == this && FWildcardNodeUtils::IsLinkedToWildcard(ChangedPin))
+	else if (ChangedPin->GetOwningNode() == this && FWildcardNodeUtils::IsLinkedToWildcard(ChangedPin))
 	{
 		return;
 	}
@@ -583,18 +583,6 @@ void UK2Node_PromotableOperator::EvaluatePinsFromChange(UEdGraphPin* ChangedPin)
 
 	// Store these other function options for later so that the user can convert to them later
 	UpdatePinsFromFunction(BestMatchingFunc, ChangedPin);
-
-	// If the user connected a type that was a valid promotion, than leave it as the pin type they dragged from for a better UX
-	if (!bWasAFullDisconnect && NewConnectionHighestType.PinCategory != UEdGraphSchema_K2::PC_Wildcard &&
-		(FTypePromotion::IsValidPromotion(NewConnectionHighestType, ChangedPin->PinType) || FTypePromotion::IsValidPromotion(ChangedPin->PinType, NewConnectionHighestType)))
-	{
-		ChangedPin->PinType = NewConnectionHighestType;
-	}
-
-	GetGraph()->NotifyGraphChanged();
-
-	// Update context menu options for this node
-	UpdatePossibleConversionFuncs();
 }
 
 bool UK2Node_PromotableOperator::UpdateOpName()
@@ -807,88 +795,99 @@ void UK2Node_PromotableOperator::UpdatePinsFromFunction(const UFunction* Functio
 
 	const UEdGraphSchema_K2* Schema = GetDefault<UEdGraphSchema_K2>();
 
-	int32 ArgCount = 0;
-	for (TFieldIterator<FProperty> PropIt(Function); PropIt && (PropIt->PropertyFlags & CPF_Parm); ++PropIt)
+	// Gather the pin types of the properties on the function we want to convert to
+	FEdGraphPinType FunctionReturnType;
+	FEdGraphPinType HighestFuncInputType;
+	TArray<FEdGraphPinType> FunctionInputTypes;
 	{
-		FProperty* Param = *PropIt;
-		FEdGraphPinType ParamType;
-
-		// If this param type is valid and we are still within range of the number of pins that we have
-		if (Schema->ConvertPropertyToPinType(Param, /* out */ ParamType) && ArgCount >= 0 && ArgCount < Pins.Num())
+		for (TFieldIterator<FProperty> PropIt(Function); PropIt && (PropIt->PropertyFlags & CPF_Parm); ++PropIt)
 		{
-			// Get either the output pin or one of the input pins
-			UEdGraphPin* PinToChange = Param->HasAnyPropertyFlags(CPF_ReturnParm) ? GetOutputPin() : Pins[ArgCount];
-
-			check(PinToChange);
-
-			const bool bHasConnectionOrDefault = PinToChange->LinkedTo.Num() > 0 || !PinToChange->DoesDefaultValueMatchAutogenerated();
-			const bool bIsWildcard = FWildcardNodeUtils::IsWildcardPin(PinToChange);
-			const bool bIsValidPromo = !bIsWildcard && FTypePromotion::IsValidPromotion(PinToChange->PinType, ParamType);
-			const bool bTypesEqual = PinToChange->PinType == ParamType;
-			const bool bIsOutPin = PinToChange->Direction == EGPD_Output;
+			FProperty* Param = *PropIt;
+			FEdGraphPinType ParamType;
 
-			bool bIsLinkedToWildcard = false;
-			if (bIsWildcard && bHasConnectionOrDefault)
+			if (Schema->ConvertPropertyToPinType(Param, /* out */ ParamType))
 			{
-				for (UEdGraphPin* LinkedPin : PinToChange->LinkedTo)
+				if (Param->HasAnyPropertyFlags(CPF_ReturnParm))
+				{
+					FunctionReturnType = ParamType;
+				}
+				else
 				{
-					if (FWildcardNodeUtils::IsWildcardPin(LinkedPin))
+					// Track the highest input pin type that we have
+					if (FTypePromotion::GetHigherType(HighestFuncInputType, ParamType) == FTypePromotion::ETypeComparisonResult::TypeBHigher)
 					{
-						bIsLinkedToWildcard = true;
-						break;
+						HighestFuncInputType = ParamType;
 					}
+					FunctionInputTypes.Add(ParamType);
 				}
 			}
+		}
+	}
 
-			// If this is a wildcard WITH a connection, then we should just leave this pin as a wildcard and let the compiler handle it
-			if (bIsLinkedToWildcard)
-			{
-				continue;
-			}
+	auto ConformPinLamda = [&ChangedPin](const FEdGraphPinType& FunctionPinType, UEdGraphPin* NodePin)
+	{
+		// If the pin types are already equal, then we don't have to do any work
+		// If this is linked to wildcard pins, then we can just ignore it and handle it on expansion
+		if (FWildcardNodeUtils::IsLinkedToWildcard(NodePin))
+		{
+			return;
+		}
 
-			bool bNeedsTypeUpdate = true;
+		// By default, conform to the type of the function param
+		FEdGraphPinType ConformingType = FunctionPinType;
+		const FEdGraphPinType HighestLinkedType = NodePin->LinkedTo.Num() > 0 ? FTypePromotion::GetPromotedType(NodePin->LinkedTo) : NodePin->PinType;
 
-			// If this pin has a valid value already, than dont bother updating it.
-			if (bHasConnectionOrDefault && (bIsValidPromo || bTypesEqual))
+		// If the highest type is the same as the function type, just continue on with life
+		if (NodePin->LinkedTo.Num() > 0 && HighestLinkedType.PinCategory != FunctionPinType.PinCategory)
+		{
+			// If the links cannot be promoted to the function type, then we need to break them
+			if (!FTypePromotion::IsValidPromotion(HighestLinkedType, FunctionPinType) || NodePin->Direction == EGPD_Output)
 			{
-				bNeedsTypeUpdate = false;
+				NodePin->BreakAllPinLinks();
 			}
-
-			// We always want to update the out pin or if we have a wildcard pin (which is the case during reconstruction)
-			if (!bTypesEqual && (bIsOutPin || bIsWildcard))
+			else
 			{
-				bNeedsTypeUpdate = true;
+				ConformingType = HighestLinkedType;
 			}
+		}
 
-			if (bNeedsTypeUpdate)
-			{
-				// If this is a wildcard pin than we have to reconsider the links to that 
-				if (bIsWildcard && PinToChange->LinkedTo.Num() > 0)
-				{
-					FEdGraphPinType LinkedType = FTypePromotion::GetPromotedType(PinToChange->LinkedTo);
-					// If this is a wildcard pin, then check the compatibility with a linked type instead of the wildcard itself
-					if (FTypePromotion::IsValidPromotion(bIsWildcard ? LinkedType : PinToChange->PinType, ParamType))
-					{
-						ParamType = LinkedType;
-					}
-				}
-				// If this update is coming from a changed pin, than we need to orphan the pin that had the old type
-				else if (ChangedPin && bHasConnectionOrDefault)
-				{
-					PinToChange->BreakAllPinLinks();
-				}
-
-				// Only change the type of this pin if it is necessary				
-				PinToChange->PinType = ParamType;
-			}
+		// Conform the pin type appropriately
+		NodePin->PinType = ConformingType;
+	};
 
+	int32 CurPinIndex = 0;
+	for (UEdGraphPin* CurPin : Pins)
+	{
+		// We don't want to try and conform split pin, because we will already have conformed the parent pin
+		if (CurPin->ParentPin != nullptr)
+		{
+			continue;
 		}
 
-		++ArgCount;
+		if (IsAdditionalPin(CurPin))
+		{
+			// Conform to the highest input pin on the function
+			ConformPinLamda(HighestFuncInputType, CurPin);
+		}
+		else if (CurPin->Direction == EGPD_Output)
+		{
+			// Match to the output pin
+			ConformPinLamda(FunctionReturnType, CurPin);
+		}
+		else
+		{
+			// Match to the appropriate function input type
+			ConformPinLamda(FunctionInputTypes[CurPinIndex], CurPin);
+			++CurPinIndex;
+		}
 	}
 
 	// Update the function reference and the FUNC_BlueprintPure/FUNC_Const appropriately
 	SetFromFunction(Function);
+
+	// We need to notify the graph that the node has changed to get 
+	// the correct default value text boxes on the node
+	GetGraph()->NotifyGraphChanged();
 }
 
 void UK2Node_PromotableOperator::UpdatePossibleConversionFuncs()
