Blueprints: Add a way to check if a function is a comparison operator directly from an FName, so that we don't have to get the proper FName from a UFunction every time. This optimizaes the K2PromotableOperator::CanAddPin call, which happens every frame.

#rnx
#okforgiithub public

diff --git a/Engine/Source/Editor/BlueprintGraph/Classes/BlueprintTypePromotion.h b/Engine/Source/Editor/BlueprintGraph/Classes/BlueprintTypePromotion.h
--- a/Engine/Source/Editor/BlueprintGraph/Classes/BlueprintTypePromotion.h
+++ b/Engine/Source/Editor/BlueprintGraph/Classes/BlueprintTypePromotion.h
@@ -53,6 +53,9 @@ public:
 	/** Returns true if the given function is a comparison operator */
 	static bool IsComparisonFunc(UFunction const* const Func);
 
+	/** Returns true if the given op name is a comparison operator name */
+	static bool IsComparisonOpName(const FName OpName);
+
 	static FName GetOpNameFromFunction(UFunction const* const Func);
 
 	/** Returns true if the given function is a candidate to handle type promotion */
diff --git a/Engine/Source/Editor/BlueprintGraph/Private/BlueprintTypePromotion.cpp b/Engine/Source/Editor/BlueprintGraph/Private/BlueprintTypePromotion.cpp
--- a/Engine/Source/Editor/BlueprintGraph/Private/BlueprintTypePromotion.cpp
+++ b/Engine/Source/Editor/BlueprintGraph/Private/BlueprintTypePromotion.cpp
@@ -524,6 +524,11 @@ bool FTypePromotion::IsComparisonFunc(UFunction const* const Func)
 	return Func && GetComparisonOpNames().Contains(GetOpNameFromFunction(Func));
 }
 
+bool FTypePromotion::IsComparisonOpName(const FName OpName)
+{
+	return GetComparisonOpNames().Contains(OpName);
+}
+
 void FTypePromotion::GetAllFuncsForOp_Internal(FName Operation, TArray<UFunction*>& OutFuncs)
 {
 	OutFuncs.Empty();
diff --git a/Engine/Source/Editor/BlueprintGraph/Private/K2Node_PromotableOperator.cpp b/Engine/Source/Editor/BlueprintGraph/Private/K2Node_PromotableOperator.cpp
--- a/Engine/Source/Editor/BlueprintGraph/Private/K2Node_PromotableOperator.cpp
+++ b/Engine/Source/Editor/BlueprintGraph/Private/K2Node_PromotableOperator.cpp
@@ -421,7 +421,7 @@ void UK2Node_PromotableOperator::AddInputPin()
 bool UK2Node_PromotableOperator::CanAddPin() const
 {
 	return ((NumAdditionalInputs + NumFunctionInputs) < GetMaxInputPinsNum()) &&
-		!FTypePromotion::IsComparisonFunc(GetTargetFunction());
+		!FTypePromotion::IsComparisonOpName(OperationName);
 }
 
 bool UK2Node_PromotableOperator::CanRemovePin(const UEdGraphPin* Pin) const
