Blueprints: Update the promotable operator tooltip when the function node has a new connection made.

#rnx

diff --git a/Engine/Source/Editor/BlueprintGraph/Classes/K2Node_CallFunction.h b/Engine/Source/Editor/BlueprintGraph/Classes/K2Node_CallFunction.h
--- a/Engine/Source/Editor/BlueprintGraph/Classes/K2Node_CallFunction.h
+++ b/Engine/Source/Editor/BlueprintGraph/Classes/K2Node_CallFunction.h
@@ -58,9 +58,11 @@ private:
 	UPROPERTY()
 	TSubclassOf<class UObject> CallFunctionClass_DEPRECATED;
 
+protected:
 	/** Constructing FText strings can be costly, so we cache the node's tooltip */
 	FNodeTextCache CachedTooltip;
 
+private:
 	/** Flag used to track validity of pin tooltips, when tooltips are invalid they will be refreshed before being displayed */
 	mutable bool bPinTooltipsValid;
 
diff --git a/Engine/Source/Editor/BlueprintGraph/Private/K2Node_PromotableOperator.cpp b/Engine/Source/Editor/BlueprintGraph/Private/K2Node_PromotableOperator.cpp
--- a/Engine/Source/Editor/BlueprintGraph/Private/K2Node_PromotableOperator.cpp
+++ b/Engine/Source/Editor/BlueprintGraph/Private/K2Node_PromotableOperator.cpp
@@ -885,6 +885,9 @@ void UK2Node_PromotableOperator::UpdatePinsFromFunction(const UFunction* Functio
 	// Update the function reference and the FUNC_BlueprintPure/FUNC_Const appropriately
 	SetFromFunction(Function);
 
+	// Invalidate the tooltips
+	CachedTooltip.MarkDirty();
+
 	// We need to notify the graph that the node has changed to get 
 	// the correct default value text boxes on the node
 	GetGraph()->NotifyGraphChanged();
