Blueprints: Reset the promotable operator node when a defalut value is set back to its original value. This is more in line with the existing behavior of K2Array_CallFunction nodes and other wildcard nodes in blueprints.

#rnx

diff --git a/Engine/Source/Editor/BlueprintGraph/Classes/K2Node_PromotableOperator.h b/Engine/Source/Editor/BlueprintGraph/Classes/K2Node_PromotableOperator.h
--- a/Engine/Source/Editor/BlueprintGraph/Classes/K2Node_PromotableOperator.h
+++ b/Engine/Source/Editor/BlueprintGraph/Classes/K2Node_PromotableOperator.h
@@ -21,6 +21,8 @@ public:
 	virtual void AllocateDefaultPins() override;
 	virtual void GetNodeContextMenuActions(UToolMenu* Menu, UGraphNodeContextMenuContext* Context) const override;
 	virtual FText GetTooltipText() const override;
+	virtual void PinDefaultValueChanged(UEdGraphPin* Pin) override;
+	virtual void NodeConnectionListChanged() override;
 	// End of UEdGraphNode interface
 
 	// UK2Node interface
@@ -56,6 +58,9 @@ private:
 	/** Helper function to recombine all split pins that this node may have */
 	void RecombineAllSplitPins();
 
+	/** Update the pins on this node with the function that is the best match given the current connections */
+	void UpdateFromBestMatchingFunction();
+
 	/** 
 	* @return	True if this node has any connections attached to it, 
 	*			or the default values have been modified by the user
@@ -121,6 +126,9 @@ private:
 	UPROPERTY()
 	int32 NumAdditionalInputs;
 
+	/** Guard to prevent possible recursive calls from ResetPinToAutogeneratedDefaultValue when breaking all links to this node */
+	bool bDefaultValueReentranceGuard;
+
 public:
 
 	/** Returns the first pin with the EGPD_Output direction */
diff --git a/Engine/Source/Editor/BlueprintGraph/Private/K2Node_PromotableOperator.cpp b/Engine/Source/Editor/BlueprintGraph/Private/K2Node_PromotableOperator.cpp
--- a/Engine/Source/Editor/BlueprintGraph/Private/K2Node_PromotableOperator.cpp
+++ b/Engine/Source/Editor/BlueprintGraph/Private/K2Node_PromotableOperator.cpp
@@ -140,6 +140,44 @@ FText UK2Node_PromotableOperator::GetTooltipText() const
 	return Super::GetTooltipText();
 }
 
+void UK2Node_PromotableOperator::PinDefaultValueChanged(UEdGraphPin* Pin)
+{
+	Super::PinDefaultValueChanged(Pin);
+
+	if (bDefaultValueReentranceGuard)
+	{
+		return;
+	}
+
+	// Re-entrance Guard just in case this function gets called from any notify triggers in the schema
+	// to prevent possible recursive calls from ResetPinToAutogeneratedDefaultValue when breaking
+	// all links to this node
+	TGuardValue<bool> ReentranceGuard(bDefaultValueReentranceGuard, true);
+
+	// If this default value resets to the default one on the pin, and there are no other
+	// connections or default values, then we should just reset the whole node to a wildcard
+	if (!HasAnyConnectionsOrDefaults())
+	{
+		ResetNodeToWildcard();
+	}
+}
+
+void UK2Node_PromotableOperator::NodeConnectionListChanged()
+{
+	Super::NodeConnectionListChanged();
+
+	// This will handle the case of dragging off of this node, and connecting to a node via typing
+	// in the context menu. Without updating in this case, our pins would be left as wildcards!
+	if(HasAnyConnectionsOrDefaults())
+	{
+		UpdateOpName();
+		UpdateFromBestMatchingFunction();
+
+		// Get correct default value boxes
+		GetGraph()->NotifyGraphChanged();
+	}
+}
+
 ///////////////////////////////////////////////////////////
 // UK2Node interface
 
@@ -749,6 +787,34 @@ void UK2Node_PromotableOperator::RecombineAllSplitPins()
 	}
 }
 
+void UK2Node_PromotableOperator::UpdateFromBestMatchingFunction()
+{
+	// Gather all pins and their links so we can determine the highest type that the user could want
+	TArray<UEdGraphPin*> PinsToConsider;
+
+	for (UEdGraphPin* Pin : Pins)
+	{
+		// If this pin has links, then use those instead of the actual pin because we could be process of changing it
+		// which means that it would still have it's old pin type, and could be inaccurate
+		if (Pin->LinkedTo.Num() > 0)
+		{
+			for (UEdGraphPin* Link : Pin->LinkedTo)
+			{
+				PinsToConsider.Emplace(Link);
+			}
+		}
+		else if (!Pin->DoesDefaultValueMatchAutogenerated() && !FWildcardNodeUtils::IsWildcardPin(Pin))
+		{
+			// If this pin has doesn't have any links then we should account for the pin itself
+			PinsToConsider.Add(Pin);
+		}
+	}
+	const UFunction* BestMatchingFunc = FTypePromotion::FindBestMatchingFunc(OperationName, PinsToConsider);
+
+	// We need to update the pins from our function if have a new connection
+	UpdatePinsFromFunction(BestMatchingFunc);
+}
+
 TArray<UEdGraphPin*> UK2Node_PromotableOperator::GetInputPins(bool bIncludeLinks /** = false */) const
 {
 	TArray<UEdGraphPin*> InputPins;
