Blueprints: Change the Type Promotion test to not be a smoke test, but an engine test. Add the vector pin type to the promotion unit tests.

#rnx

diff --git a/Engine/Source/Editor/BlueprintGraph/Private/Tests/TypePromotionTests.cpp b/Engine/Source/Editor/BlueprintGraph/Private/Tests/TypePromotionTests.cpp
--- a/Engine/Source/Editor/BlueprintGraph/Private/Tests/TypePromotionTests.cpp
+++ b/Engine/Source/Editor/BlueprintGraph/Private/Tests/TypePromotionTests.cpp
@@ -8,7 +8,7 @@
 
 #if WITH_DEV_AUTOMATION_TESTS
 
-IMPLEMENT_SIMPLE_AUTOMATION_TEST(FTypePromotionTest, "Blueprints.Compiler.TypePromotion", EAutomationTestFlags::ApplicationContextMask | EAutomationTestFlags::SmokeFilter)
+IMPLEMENT_SIMPLE_AUTOMATION_TEST(FTypePromotionTest, "Blueprints.Compiler.TypePromotion", EAutomationTestFlags::EditorContext | EAutomationTestFlags::SmokeFilter)
 
 bool FTypePromotionTest::RunTest(const FString& Parameters)
 {
@@ -17,10 +17,12 @@ bool FTypePromotionTest::RunTest(const FString& Parameters)
 	FEdGraphPinType IntPin = {};		IntPin.PinCategory		= UEdGraphSchema_K2::PC_Int;
 	FEdGraphPinType Int64Pin = {};		Int64Pin.PinCategory	= UEdGraphSchema_K2::PC_Int64;
 	FEdGraphPinType BytePin = {};		BytePin.PinCategory		= UEdGraphSchema_K2::PC_Byte;
+	FEdGraphPinType VecPin = {};		VecPin.PinCategory		= UEdGraphSchema_K2::PC_Struct; VecPin.PinSubCategoryObject = TBaseStructure<FVector>::Get();
 
 	// Test promotions that should happen
 	TestEqual(TEXT("Testing float to double"), FTypePromotion::GetHigherType(FloatPin, DoublePin), FTypePromotion::ETypeComparisonResult::TypeBHigher);
-	
+	TestEqual(TEXT("Testing float to vector"), FTypePromotion::GetHigherType(FloatPin, VecPin), FTypePromotion::ETypeComparisonResult::TypeBHigher);
+
 	TestEqual(TEXT("Testing int to float"), FTypePromotion::GetHigherType(IntPin, FloatPin), FTypePromotion::ETypeComparisonResult::TypeBHigher);
 	TestEqual(TEXT("Testing int to double"), FTypePromotion::GetHigherType(IntPin, DoublePin), FTypePromotion::ETypeComparisonResult::TypeBHigher);
 	TestEqual(TEXT("Testing int to int64"), FTypePromotion::GetHigherType(IntPin, Int64Pin), FTypePromotion::ETypeComparisonResult::TypeBHigher);
@@ -47,7 +49,7 @@ bool FTypePromotionTest::RunTest(const FString& Parameters)
 }
 
 
-IMPLEMENT_SIMPLE_AUTOMATION_TEST(FFindPromotedFunc, "Blueprints.Compiler.FindPromotedFunc", EAutomationTestFlags::ApplicationContextMask | EAutomationTestFlags::SmokeFilter)
+IMPLEMENT_SIMPLE_AUTOMATION_TEST(FFindPromotedFunc, "Blueprints.Compiler.FindPromotedFunc", EAutomationTestFlags::EditorContext | EAutomationTestFlags::EngineFilter)
 
 #define MakeTestPin(OwningNode, PinArray, InPinName, InPinType)		UEdGraphPin* InPinName = UEdGraphPin::CreatePin(OwningNode);		\
 																	InPinName->PinType.PinCategory = InPinType;							\
@@ -122,7 +124,7 @@ bool FFindPromotedFunc::RunTest(const FString& Parameters)
 }
 
 
-IMPLEMENT_SIMPLE_AUTOMATION_TEST(FFindBestMatchingFunc, "Blueprints.Compiler.FindBestMatchingFunc", EAutomationTestFlags::ApplicationContextMask | EAutomationTestFlags::SmokeFilter)
+IMPLEMENT_SIMPLE_AUTOMATION_TEST(FFindBestMatchingFunc, "Blueprints.Compiler.FindBestMatchingFunc", EAutomationTestFlags::EditorContext | EAutomationTestFlags::EngineFilter)
 bool FFindBestMatchingFunc::RunTest(const FString& Parameters)
 {
 	TArray<UEdGraphPin*> PinTypes = {};
