Blueprints: Refresh the operator table and no spawner list when a module is changed in type promotion. This will ensure that the actions list is properly rebuilt after a module is changed, and that the BP context menu has no duplicate operators.

Check the op name for No_OP when checking if a UFunction is a valid promotable func

#rnx

diff --git a/Engine/Source/Editor/BlueprintGraph/Classes/BlueprintTypePromotion.h b/Engine/Source/Editor/BlueprintGraph/Classes/BlueprintTypePromotion.h
--- a/Engine/Source/Editor/BlueprintGraph/Classes/BlueprintTypePromotion.h
+++ b/Engine/Source/Editor/BlueprintGraph/Classes/BlueprintTypePromotion.h
@@ -150,6 +150,9 @@ private:
 
 	static FTypePromotion* Instance;
 
+	/** Delegate to handle that will be used to refresh the op table when a module has changed */
+	FDelegateHandle OnModulesChangedDelegateHandle;
+
 	/** A map of 'Type' to its 'available promotions'. See ctor for creation */
 	TMap<FName, TArray<FName>> PromotionTable;
 
diff --git a/Engine/Source/Editor/BlueprintGraph/Private/BlueprintTypePromotion.cpp b/Engine/Source/Editor/BlueprintGraph/Private/BlueprintTypePromotion.cpp
--- a/Engine/Source/Editor/BlueprintGraph/Private/BlueprintTypePromotion.cpp
+++ b/Engine/Source/Editor/BlueprintGraph/Private/BlueprintTypePromotion.cpp
@@ -65,10 +65,24 @@ FTypePromotion::FTypePromotion()
 {
 	CreatePromotionTable();
 	CreateOpTable();
+	OnModulesChangedDelegateHandle = FModuleManager::Get().OnModulesChanged().AddStatic(&FTypePromotion::OnModulesChanged);
 }
 
 FTypePromotion::~FTypePromotion()
 {
+	FModuleManager::Get().OnModulesChanged().Remove(OnModulesChangedDelegateHandle);
+}
+
+void FTypePromotion::OnModulesChanged(FName ModuleThatChanged, EModuleChangeReason ReasonForChange)
+{
+	if (Instance)
+	{
+		// Any time a module is changed, there could possibly be new UFunctions that we 
+		// need to process, so we need to recreate the op table and clear the node spawners
+		// that we are using in order to avoid invalid duplicates in the graph action menu
+		FTypePromotion::ClearNodeSpawners();
+		Instance->CreateOpTable();
+	}
 }
 
 void FTypePromotion::CreatePromotionTable()
@@ -599,7 +613,10 @@ void FTypePromotion::AddOpFunction(FName OpName, UFunction* Function)
 
 bool FTypePromotion::IsPromotableFunction(const UFunction* Function)
 {
-	return Function && Function->HasAnyFunctionFlags(FUNC_BlueprintPure) && Function->GetReturnProperty();
+	// Ensure that we don't have an invalid OpName as well for extra safety when this function 
+	// is called outside of this class, not during the OpTable creation process
+	FName OpName = GetOpNameFromFunction(Function);
+	return Function && Function->HasAnyFunctionFlags(FUNC_BlueprintPure) && Function->GetReturnProperty() && OpName != OperatorNames::NoOp;
 }
 
 bool FTypePromotion::IsOperatorSpawnerRegistered(UFunction const* const Func)
