Blueprints: Don't break links if the pin is the one that is currently being connected. This fixed a bad UX where if the function didn't match up, then it would be impossible to connect pins to some spots on a promotable operator node.

#rnx

diff --git a/Engine/Source/Editor/BlueprintGraph/Private/K2Node_PromotableOperator.cpp b/Engine/Source/Editor/BlueprintGraph/Private/K2Node_PromotableOperator.cpp
--- a/Engine/Source/Editor/BlueprintGraph/Private/K2Node_PromotableOperator.cpp
+++ b/Engine/Source/Editor/BlueprintGraph/Private/K2Node_PromotableOperator.cpp
@@ -894,7 +894,7 @@ void UK2Node_PromotableOperator::UpdatePinsFromFunction(const UFunction* Functio
 	{
 		// If the pin types are already equal, then we don't have to do any work
 		// If this is linked to wildcard pins, then we can just ignore it and handle it on expansion
-		if (FWildcardNodeUtils::IsLinkedToWildcard(NodePin))
+		if (!NodePin || FWildcardNodeUtils::IsLinkedToWildcard(NodePin))
 		{
 			return;
 		}
@@ -907,7 +907,9 @@ void UK2Node_PromotableOperator::UpdatePinsFromFunction(const UFunction* Functio
 		if (NodePin->LinkedTo.Num() > 0 && HighestLinkedType.PinCategory != FunctionPinType.PinCategory)
 		{
 			// If the links cannot be promoted to the function type, then we need to break them
-			if (!FTypePromotion::IsValidPromotion(HighestLinkedType, FunctionPinType) || NodePin->Direction == EGPD_Output)
+			// We don't want to break the pin if it is the one that the user has dragged on to though,
+			// because that would result in the node breaking connection as soon as the user lets go
+			if ((!FTypePromotion::IsValidPromotion(HighestLinkedType, FunctionPinType) || NodePin->Direction == EGPD_Output) && (NodePin != ChangedPin))
 			{
 				NodePin->BreakAllPinLinks();
 			}
