Blueprints: Ensure that comparison ops cannot have pins added to them

#rnx

diff --git a/Engine/Source/Editor/BlueprintGraph/Private/K2Node_PromotableOperator.cpp b/Engine/Source/Editor/BlueprintGraph/Private/K2Node_PromotableOperator.cpp
--- a/Engine/Source/Editor/BlueprintGraph/Private/K2Node_PromotableOperator.cpp
+++ b/Engine/Source/Editor/BlueprintGraph/Private/K2Node_PromotableOperator.cpp
@@ -459,7 +459,7 @@ void UK2Node_PromotableOperator::AddInputPin()
 bool UK2Node_PromotableOperator::CanAddPin() const
 {
 	return ((NumAdditionalInputs + NumFunctionInputs) < GetMaxInputPinsNum()) &&
-		!FTypePromotion::IsComparisonOpName(OperationName);
+		!FTypePromotion::IsComparisonFunc(GetTargetFunction());
 }
 
 bool UK2Node_PromotableOperator::CanRemovePin(const UEdGraphPin* Pin) const
