Blueprints: Correctly match the best function when only given one pin to consider.

#rnx

diff --git a/Engine/Source/Editor/BlueprintGraph/Private/BlueprintTypePromotion.cpp b/Engine/Source/Editor/BlueprintGraph/Private/BlueprintTypePromotion.cpp
--- a/Engine/Source/Editor/BlueprintGraph/Private/BlueprintTypePromotion.cpp
+++ b/Engine/Source/Editor/BlueprintGraph/Private/BlueprintTypePromotion.cpp
@@ -441,7 +441,7 @@ UFunction* FTypePromotion::FindBestMatchingFunc_Internal(FName Operation, const
 				for (const UEdGraphPin* Pin : PinsToConsider)
 				{
 					// Give a point for each function parameter that matches up with a pin to consider
-					if (!CheckedPins.Contains(Pin) && Schema->ArePinTypesEquivalent(ParamType, Pin->PinType))
+					if ((!CheckedPins.Contains(Pin) || bIsSinglePin) && Schema->ArePinTypesEquivalent(ParamType, Pin->PinType))
 					{
 						// Are the directions compatible? 
 						// If we are a comparison or only a single pin then we don't care about the direction
