Blueprints: Update the Operation name when the function is set on a promotable operator

#rnx

diff --git a/Engine/Source/Editor/BlueprintGraph/Classes/K2Node_PromotableOperator.h b/Engine/Source/Editor/BlueprintGraph/Classes/K2Node_PromotableOperator.h
--- a/Engine/Source/Editor/BlueprintGraph/Classes/K2Node_PromotableOperator.h
+++ b/Engine/Source/Editor/BlueprintGraph/Classes/K2Node_PromotableOperator.h
@@ -39,8 +39,12 @@ public:
 	virtual bool CanAddPin() const override;
 	virtual bool CanRemovePin(const UEdGraphPin* Pin) const override;
 	virtual void RemoveInputPin(UEdGraphPin* Pin) override;
-	// end IK2Node_AddPinInterface interface
+	// End of IK2Node_AddPinInterface interface
 
+	// UK2Node_CallFunction interface
+	virtual void SetFromFunction(const UFunction* Function);
+	// End of UK2Node_CallFunction interface
+	
 private:
 
 	/**
diff --git a/Engine/Source/Editor/BlueprintGraph/Private/K2Node_PromotableOperator.cpp b/Engine/Source/Editor/BlueprintGraph/Private/K2Node_PromotableOperator.cpp
--- a/Engine/Source/Editor/BlueprintGraph/Private/K2Node_PromotableOperator.cpp
+++ b/Engine/Source/Editor/BlueprintGraph/Private/K2Node_PromotableOperator.cpp
@@ -530,6 +530,18 @@ UEdGraphPin* UK2Node_PromotableOperator::GetAdditionalPin(int32 PinIndex) const
 	return nullptr;
 }
 
+///////////////////////////////////////////////////////////
+// UK2Node_CallFunction interface
+void UK2Node_PromotableOperator::SetFromFunction(const UFunction* Function)
+{
+	Super::SetFromFunction(Function);
+
+	if(Function)
+	{
+		OperationName = FTypePromotion::GetOpNameFromFunction(Function);
+	}
+}
+
 ///////////////////////////////////////////////////////////
 // UK2Node_PromotableOperator
 
