Blueprints: Remove unused function in type promotion. We use FindBestMatchingFunc now isntead of GetOperatorFunction

#rnx

diff --git a/Engine/Source/Editor/BlueprintGraph/Classes/BlueprintTypePromotion.h b/Engine/Source/Editor/BlueprintGraph/Classes/BlueprintTypePromotion.h
--- a/Engine/Source/Editor/BlueprintGraph/Classes/BlueprintTypePromotion.h
+++ b/Engine/Source/Editor/BlueprintGraph/Classes/BlueprintTypePromotion.h
@@ -23,18 +23,6 @@ public:
 	/** Deletes the singleton instance of type promotion if there is one */
 	static void Shutdown();
 
-	/** 
-	* If given more than two input pins, this will just use the highest type
-	* and return the function for that operation. It's update to nodes like the 
-	* UK2Node_CommutativeAssociativeBinaryOperator to patch the input pins
-	* together on expansion
-	* 
-	* @return	A function that matches a promoted type. Nullptr if no match is found
-	*/
-	static UFunction* GetOperatorFunction(FName Operation, const TArray<UEdGraphPin*>& WildcardPins);
-
-	static UFunction* GetOperatorFunction(FName Operation, const TArray<UEdGraphPin*>& InputPins, const UEdGraphPin* OutputPin);
-
 	/**
 	* Find the function that is the best match given the pins to consider. 
 	* Ex: Given "Add" operator and an array of two Vector pins, it will return "Add_VectorVector"
@@ -119,10 +107,6 @@ private:
 	/** Callback that will rebuild the op table when hot reload is triggered */
 	static void OnModulesChanged(FName ModuleThatChanged, EModuleChangeReason ReasonForChange);
 
-	UFunction* GetOperatorFunction_Internal(FName Operation, const TArray<UEdGraphPin*>& WildcardPins) const;
-
-	UFunction* GetOperatorFunction_Internal(FName Operation, const TArray<UEdGraphPin*>& InputPins, const UEdGraphPin* OutputPin) const;
-
 	bool IsFunctionPromotionReady_Internal(const UFunction* const FuncToConsider) const;
 
 	FEdGraphPinType GetPromotedType_Internal(const TArray<UEdGraphPin*>& WildcardPins) const;
diff --git a/Engine/Source/Editor/BlueprintGraph/Private/BlueprintTypePromotion.cpp b/Engine/Source/Editor/BlueprintGraph/Private/BlueprintTypePromotion.cpp
--- a/Engine/Source/Editor/BlueprintGraph/Private/BlueprintTypePromotion.cpp
+++ b/Engine/Source/Editor/BlueprintGraph/Private/BlueprintTypePromotion.cpp
@@ -248,133 +248,6 @@ bool FTypePromotion::PromotePin_Internal(FEdGraphPinType& InTypeA, const FEdGrap
 	return true;
 }
 
-UFunction* FTypePromotion::GetOperatorFunction(FName Operation, const TArray<UEdGraphPin*>& WildcardPins)
-{
-	return FTypePromotion::Get().GetOperatorFunction_Internal(Operation, WildcardPins);
-}
-
-UFunction* FTypePromotion::GetOperatorFunction(FName Operation, const TArray<UEdGraphPin*>& InputPins, const UEdGraphPin* OutputPin)
-{
-	return FTypePromotion::Get().GetOperatorFunction_Internal(Operation, InputPins, OutputPin);
-}
-
-UFunction* FTypePromotion::GetOperatorFunction_Internal(FName Operation, const TArray<UEdGraphPin*>& InputPins, const UEdGraphPin* OutputPin) const
-{
-	check(InputPins.Num() > 0 && OutputPin);
-
-	TRACE_CPUPROFILER_EVENT_SCOPE(FTypePromotionTable::GetOperatorFunction_InOut);
-	
-	UFunction* MatchingFunction = nullptr;
-	if (const FFunctionsList* FuncList = OperatorTable.Find(Operation))
-	{
-		const UEdGraphSchema_K2* Schema = GetDefault<UEdGraphSchema_K2>();
-		for (UFunction* Func : *FuncList)
-		{
-			// Do the return types match? 
-			FEdGraphPinType FuncReturnType;
-			const bool bReturnTypesMatch = Schema->ConvertPropertyToPinType(Func->GetReturnProperty(), /* out */ FuncReturnType) && FuncReturnType.PinCategory == OutputPin->PinType.PinCategory;
-			if (!bReturnTypesMatch)
-			{
-				continue;
-			}
-
-			// Do the input types match up correctly? 
-			int32 ArgumentIndex = 0;
-			int32 NumMatchingInputs = 0;
-			for (TFieldIterator<FProperty> PropIt(Func); PropIt && (PropIt->PropertyFlags & CPF_Parm); ++PropIt)
-			{
-				FProperty* Param = *PropIt;
-				if (!Param->HasAnyPropertyFlags(CPF_ReturnParm))
-				{
-					if (ArgumentIndex < InputPins.Num())
-					{
-						FEdGraphPinType ParamType;
-						if (Schema->ConvertPropertyToPinType(Param, /* out */ ParamType))
-						{
-							FEdGraphPinType const& TypeToMatch = InputPins[ArgumentIndex]->PinType;
-							if (Schema->ArePinTypesCompatible(TypeToMatch, ParamType) && GetHigherType(TypeToMatch, ParamType) != ETypeComparisonResult::InvalidComparison)
-							{
-								NumMatchingInputs++;
-							}
-						}
-						else
-						{
-							break;
-						}
-					}
-					++ArgumentIndex;
-				}
-			}
-
-			if(NumMatchingInputs >= InputPins.Num())
-			{
-				MatchingFunction = Func;
-				break;
-			}
-		}
-	}
-
-	return MatchingFunction;
-}
-
-UFunction* FTypePromotion::GetOperatorFunction_Internal(FName Operation, const TArray<UEdGraphPin*>& WildcardPins) const
-{
-	TRACE_CPUPROFILER_EVENT_SCOPE(FTypePromotionTable::GetOperatorFunction);
-
-	UFunction* MatchingFunction = nullptr;
-	// Do these wild cards even need to be promoted? (Can we find a function that matches their
-	// inputs in the op table already?)
-	if(const FFunctionsList* FuncList = OperatorTable.Find(Operation))
-	{
-		const UEdGraphSchema_K2* Schema = GetDefault<UEdGraphSchema_K2>();
-
-		FEdGraphPinType PromotedReturnType = GetPromotedType(WildcardPins);
-
-		for(UFunction* Func : *FuncList)
-		{
-			// If the return type of this function matches...
-			FEdGraphPinType FuncReturnType;
-			if(Schema->ConvertPropertyToPinType(Func->GetReturnProperty(), /* out */ FuncReturnType) && FuncReturnType.PinCategory == PromotedReturnType.PinCategory)
-			{
-				int32 ArgumentIndex = 0;
-				for (TFieldIterator<FProperty> PropIt(Func); PropIt && (PropIt->PropertyFlags & CPF_Parm); ++PropIt)
-				{
-					FProperty* Param = *PropIt;
-					if (!Param->HasAnyPropertyFlags(CPF_ReturnParm))
-					{
-						if(ArgumentIndex < WildcardPins.Num())
-						{
-							FEdGraphPinType ParamType;
-							if (Schema->ConvertPropertyToPinType(Param, /* out */ ParamType))
-							{
-								FEdGraphPinType const& TypeToMatch = WildcardPins[ArgumentIndex]->PinType;
-								if (!Schema->ArePinTypesCompatible(TypeToMatch, ParamType) && GetHigherType(TypeToMatch, ParamType) == ETypeComparisonResult::InvalidComparison)
-								{
-									break; // type mismatch
-								}
-							}
-							else
-							{
-								break;
-							}
-						}
-						++ArgumentIndex;
-					}
-				}
-
-				// This doesn't work if given more than 2 wildcard inputs
-				if (ArgumentIndex == WildcardPins.Num())
-				{
-					MatchingFunction = Func;
-					break;
-				}
-			}
-		}
-	}
-
-	return MatchingFunction;
-}
-
 UFunction* FTypePromotion::FindBestMatchingFunc(FName Operation, const TArray<UEdGraphPin*>& PinsToConsider)
 {
 	return FTypePromotion::Get().FindBestMatchingFunc_Internal(Operation, PinsToConsider);
diff --git a/Engine/Source/Editor/BlueprintGraph/Private/Tests/TypePromotionTests.cpp b/Engine/Source/Editor/BlueprintGraph/Private/Tests/TypePromotionTests.cpp
--- a/Engine/Source/Editor/BlueprintGraph/Private/Tests/TypePromotionTests.cpp
+++ b/Engine/Source/Editor/BlueprintGraph/Private/Tests/TypePromotionTests.cpp
@@ -49,81 +49,10 @@ bool FTypePromotionTest::RunTest(const FString& Parameters)
 }
 
 
-IMPLEMENT_SIMPLE_AUTOMATION_TEST(FFindPromotedFunc, "Blueprints.Compiler.FindPromotedFunc", EAutomationTestFlags::EditorContext | EAutomationTestFlags::EngineFilter)
-
 #define MakeTestPin(OwningNode, PinArray, InPinName, InPinType)		UEdGraphPin* InPinName = UEdGraphPin::CreatePin(OwningNode);		\
 																	InPinName->PinType.PinCategory = InPinType;							\
 																	PinArray.Add(InPinName);		
 
-bool FFindPromotedFunc::RunTest(const FString& Parameters)
-{
-	TArray<UEdGraphPin*> PinTypes = {};
-
-	UEdGraphNode* TestNode = NewObject<UEdGraphNode>();			
-	
-	const FTypePromotion& TypePromo = FTypePromotion::Get();
-
-	MakeTestPin(TestNode, PinTypes, DoublePin, UEdGraphSchema_K2::PC_Double);
-	MakeTestPin(TestNode, PinTypes, FloatPin, UEdGraphSchema_K2::PC_Float);
-	MakeTestPin(TestNode, PinTypes, Int32Pin, UEdGraphSchema_K2::PC_Int);
-	MakeTestPin(TestNode, PinTypes, Int64Pin, UEdGraphSchema_K2::PC_Int64);
-
-	// Add Operation
-	{
-		TArray<UEdGraphPin*> TestPins =
-		{
-			DoublePin, FloatPin
-		};
-		UFunction* AddDoubleFunc = FTypePromotion::GetOperatorFunction(TEXT("add"), TestPins);
-		TestNotNull(TEXT("Add Double Float function"), AddDoubleFunc);
-	}
-
-	// Add Operation
-	{
-		TArray<UEdGraphPin*> TestPins =
-		{
-			FloatPin, DoublePin
-		};
-		UFunction* AddDoubleFunc = FTypePromotion::GetOperatorFunction(TEXT("add"), TestPins);
-		TestNotNull(TEXT("Add Float Double function"), AddDoubleFunc);
-	}
-
-	// Multiply
-	{
-		TArray<UEdGraphPin*> TestPins =
-		{
-			FloatPin, Int32Pin
-		};
-		UFunction* ResultFunc = FTypePromotion::GetOperatorFunction(TEXT("multiply"), TestPins);
-		TestNotNull(TEXT("multiply Float Int32 function"), ResultFunc);
-	}
-
-	// Divide
-	{
-		TArray<UEdGraphPin*> TestPins =
-		{
-			FloatPin, Int32Pin
-		};
-
-		UFunction* ResultFunc = FTypePromotion::GetOperatorFunction(TEXT("divide"), TestPins);
-		TestNotNull(TEXT("divide Float Int32 function"), ResultFunc);
-	}
-
-	// Clear our test pins
-	for(UEdGraphPin* TestPin : PinTypes)
-	{
-		if(TestPin)
-		{
-			TestPin->MarkPendingKill();
-		}
-	}
-	PinTypes.Empty();
-	TestNode->MarkPendingKill();
-
-	return true;
-}
-
-
 IMPLEMENT_SIMPLE_AUTOMATION_TEST(FFindBestMatchingFunc, "Blueprints.Compiler.FindBestMatchingFunc", EAutomationTestFlags::EditorContext | EAutomationTestFlags::EngineFilter)
 bool FFindBestMatchingFunc::RunTest(const FString& Parameters)
 {
