Blueprints: Allow promotion of a byte to a double

#rnx

diff --git a/Engine/Source/Editor/BlueprintGraph/Private/BlueprintTypePromotion.cpp b/Engine/Source/Editor/BlueprintGraph/Private/BlueprintTypePromotion.cpp
--- a/Engine/Source/Editor/BlueprintGraph/Private/BlueprintTypePromotion.cpp
+++ b/Engine/Source/Editor/BlueprintGraph/Private/BlueprintTypePromotion.cpp
@@ -91,7 +91,7 @@ void FTypePromotion::CreatePromotionTable()
 	{
 		// Type_X...						Can be promoted to...
 		{ UEdGraphSchema_K2::PC_Int,		{ UEdGraphSchema_K2::PC_Float, UEdGraphSchema_K2::PC_Double, UEdGraphSchema_K2::PC_Int64 } },
-		{ UEdGraphSchema_K2::PC_Byte,		{ UEdGraphSchema_K2::PC_Float, UEdGraphSchema_K2::PC_Int, UEdGraphSchema_K2::PC_Int64 } },
+		{ UEdGraphSchema_K2::PC_Byte,		{ UEdGraphSchema_K2::PC_Float, UEdGraphSchema_K2::PC_Int, UEdGraphSchema_K2::PC_Int64, UEdGraphSchema_K2::PC_Double } },
 		{ UEdGraphSchema_K2::PC_Float,		{ UEdGraphSchema_K2::PC_Double, UEdGraphSchema_K2::PC_Int64 } },
 		{ UEdGraphSchema_K2::PC_Double,		{ UEdGraphSchema_K2::PC_Int64 } },
 		{ UEdGraphSchema_K2::PC_Wildcard,	{ UEdGraphSchema_K2::PC_Int, UEdGraphSchema_K2::PC_Int64, UEdGraphSchema_K2::PC_Float, UEdGraphSchema_K2::PC_Double, UEdGraphSchema_K2::PC_Byte, UEdGraphSchema_K2::PC_Boolean } },
