Blueprints: Treat comparison operators boolean output as a normal pin instead of rejecting anything that is not a boolean.

#rnx

diff --git a/Engine/Source/Editor/BlueprintGraph/Private/K2Node_PromotableOperator.cpp b/Engine/Source/Editor/BlueprintGraph/Private/K2Node_PromotableOperator.cpp
--- a/Engine/Source/Editor/BlueprintGraph/Private/K2Node_PromotableOperator.cpp
+++ b/Engine/Source/Editor/BlueprintGraph/Private/K2Node_PromotableOperator.cpp
@@ -370,10 +370,11 @@ bool UK2Node_PromotableOperator::IsConnectionDisallowed(const UEdGraphPin* MyPin
 		OutReason = LOCTEXT("NoExecPinsAllowed", "Promotable Operator nodes cannot have containers or references.").ToString();
 		return true;
 	}
+	// The output pin on comparison operators is always a boolean, and cannot have its type changed
+	// so we need to just treat it normally as a regular K2CallFunction node would
 	else if (MyPin == GetOutputPin() && FTypePromotion::IsComparisonFunc(GetTargetFunction()) && OtherPin->PinType.PinCategory != UEdGraphSchema_K2::PC_Boolean)
 	{
-		OutReason = LOCTEXT("ComparisonNeedsBool", "Comparison operators must return a bool!").ToString();
-		return true;
+		return Super::IsConnectionDisallowed(MyPin, OtherPin, OutReason);
 	}
 
 	const bool bHasStructPin = MyPin->PinType.PinCategory == UEdGraphSchema_K2::PC_Struct || OtherPin->PinType.PinCategory == UEdGraphSchema_K2::PC_Struct;
