Expose a getter for the promotion table in type promotion to allow for better unit tests

#rnx

diff --git a/Engine/Source/Editor/BlueprintGraph/Classes/BlueprintTypePromotion.h b/Engine/Source/Editor/BlueprintGraph/Classes/BlueprintTypePromotion.h
--- a/Engine/Source/Editor/BlueprintGraph/Classes/BlueprintTypePromotion.h
+++ b/Engine/Source/Editor/BlueprintGraph/Classes/BlueprintTypePromotion.h
@@ -99,6 +99,8 @@ public:
 
 	static void ClearNodeSpawners();
 
+	static const TMap<FName, TArray<FName>>* const GetPrimativePromotionTable();
+
 private:
 
 	explicit FTypePromotion();
diff --git a/Engine/Source/Editor/BlueprintGraph/Private/BlueprintTypePromotion.cpp b/Engine/Source/Editor/BlueprintGraph/Private/BlueprintTypePromotion.cpp
--- a/Engine/Source/Editor/BlueprintGraph/Private/BlueprintTypePromotion.cpp
+++ b/Engine/Source/Editor/BlueprintGraph/Private/BlueprintTypePromotion.cpp
@@ -98,6 +98,16 @@ void FTypePromotion::CreatePromotionTable()
 	};
 }
 
+const TMap<FName, TArray<FName>>* const FTypePromotion::GetPrimativePromotionTable()
+{
+	if (Instance)
+	{
+		return &Instance->PromotionTable;
+	}
+
+	return nullptr;
+}
+
 bool FTypePromotion::IsValidPromotion(const FEdGraphPinType& A, const FEdGraphPinType& B)
 {
 	// If either of these pin types is a struct, than we have to have some kind of valid
