Blueprints: Cleanup duplicate code and add a pin name for tolerance pins

#rnx

diff --git a/Engine/Source/Editor/BlueprintGraph/Classes/K2Node_PromotableOperator.h b/Engine/Source/Editor/BlueprintGraph/Classes/K2Node_PromotableOperator.h
--- a/Engine/Source/Editor/BlueprintGraph/Classes/K2Node_PromotableOperator.h
+++ b/Engine/Source/Editor/BlueprintGraph/Classes/K2Node_PromotableOperator.h
@@ -59,6 +59,9 @@ private:
 	/** Gets the additional pin that was created at this index */
 	UEdGraphPin* GetAdditionalPin(int32 PinIndex) const;
 		
+	/** Returns true if the given pin is a tolerance pin for a comparison operator */
+	bool IsTolerancePin(const UEdGraphPin* Pin) const;
+
 	/** Helper function to recombine all split pins that this node may have */
 	void RecombineAllSplitPins();
 
@@ -119,6 +122,14 @@ private:
 	*/
 	TArray<UEdGraphPin*> GetInputPins(bool bIncludeLinks = false) const;
 
+	/**
+	* Get all pins on this node that can be considered when determining what function
+	* is the best match
+	*
+	* @param OutArray	Array to populate with the considered pins
+	*/
+	void GetPinsToConsider(TArray<UEdGraphPin*>& OutArray);
+
 	/** The name that this operation uses ("Add", "Multiply", etc) */
 	FName OperationName;
 
diff --git a/Engine/Source/Editor/BlueprintGraph/Private/K2Node_PromotableOperator.cpp b/Engine/Source/Editor/BlueprintGraph/Private/K2Node_PromotableOperator.cpp
--- a/Engine/Source/Editor/BlueprintGraph/Private/K2Node_PromotableOperator.cpp
+++ b/Engine/Source/Editor/BlueprintGraph/Private/K2Node_PromotableOperator.cpp
@@ -20,6 +20,7 @@
 
 static const FName InputPinA_Name = FName(TEXT("A"));
 static const FName InputPinB_Name = FName(TEXT("B"));
+static const FName TolerancePin_Name = FName(TEXT("Tolerance"));
 static const int32 NumFunctionInputs = 2;
 
 ///////////////////////////////////////////////////////////
@@ -546,6 +547,11 @@ void UK2Node_PromotableOperator::SetFromFunction(const UFunction* Function)
 ///////////////////////////////////////////////////////////
 // UK2Node_PromotableOperator
 
+bool UK2Node_PromotableOperator::IsTolerancePin(const UEdGraphPin* Pin) const
+{
+	return Pin && Pin->PinName == TolerancePin_Name && Pin->Direction == EGPD_Input;
+}
+
 UEdGraphPin* UK2Node_PromotableOperator::AddInputPinImpl(int32 PinIndex)
 {
 	const FName NewPinName = GetNameForAdditionalPin(PinIndex);
@@ -606,29 +612,7 @@ void UK2Node_PromotableOperator::EvaluatePinsFromChange(UEdGraphPin* ChangedPin)
 
 	// Gather all pins and their links so we can determine the highest type that the user could want
 	TArray<UEdGraphPin*> PinsToConsider;
-
-	for (UEdGraphPin* Pin : Pins)
-	{
-		// If this pin has links, then use those instead of the actual pin because we could be process of changing it
-		// which means that it would still have it's old pin type, and could be inaccurate
-		if (Pin->LinkedTo.Num() > 0)
-		{
-			for (UEdGraphPin* Link : Pin->LinkedTo)
-			{
-				PinsToConsider.Emplace(Link);
-			}
-		}
-		else if (!Pin->DoesDefaultValueMatchAutogenerated() && !FWildcardNodeUtils::IsWildcardPin(Pin))
-		{
-			// If this pin has doesn't have any links then we should account for the pin itself
-			PinsToConsider.Add(Pin);
-		}
-	}
-
-	FEdGraphPinType HighestType = FTypePromotion::GetPromotedType(PinsToConsider);
-
-	// if a pin was changed, update it if it cannot be promoted to this type
-	FEdGraphPinType NewConnectionHighestType = ChangedPin->LinkedTo.Num() > 0 ? FTypePromotion::GetPromotedType(ChangedPin->LinkedTo) : FWildcardNodeUtils::GetDefaultWildcardPinType();
+	GetPinsToConsider(PinsToConsider);
 
 	const UFunction* BestMatchingFunc = FTypePromotion::FindBestMatchingFunc(OperationName, PinsToConsider);
 
@@ -804,24 +788,8 @@ void UK2Node_PromotableOperator::UpdateFromBestMatchingFunction()
 {
 	// Gather all pins and their links so we can determine the highest type that the user could want
 	TArray<UEdGraphPin*> PinsToConsider;
+	GetPinsToConsider(PinsToConsider);
 
-	for (UEdGraphPin* Pin : Pins)
-	{
-		// If this pin has links, then use those instead of the actual pin because we could be process of changing it
-		// which means that it would still have it's old pin type, and could be inaccurate
-		if (Pin->LinkedTo.Num() > 0)
-		{
-			for (UEdGraphPin* Link : Pin->LinkedTo)
-			{
-				PinsToConsider.Emplace(Link);
-			}
-		}
-		else if (!Pin->DoesDefaultValueMatchAutogenerated() && !FWildcardNodeUtils::IsWildcardPin(Pin))
-		{
-			// If this pin has doesn't have any links then we should account for the pin itself
-			PinsToConsider.Add(Pin);
-		}
-	}
 	const UFunction* BestMatchingFunc = FTypePromotion::FindBestMatchingFunc(OperationName, PinsToConsider);
 
 	// We need to update the pins from our function if have a new connection
@@ -849,6 +817,34 @@ TArray<UEdGraphPin*> UK2Node_PromotableOperator::GetInputPins(bool bIncludeLinks
 	return InputPins;
 }
 
+void UK2Node_PromotableOperator::GetPinsToConsider(TArray<UEdGraphPin*>& OutArray)
+{
+	for (UEdGraphPin* Pin : Pins)
+	{
+		// Tolerance pins don't factor into what types we should be matching 
+		// for the function, so we should not consider them
+		if (IsTolerancePin(Pin))
+		{
+			continue;
+		}
+
+		// If this pin has links, then use those instead of the actual pin because we could be process of changing it
+		// which means that it would still have it's old pin type, and could be inaccurate
+		if (Pin->LinkedTo.Num() > 0)
+		{
+			for (UEdGraphPin* Link : Pin->LinkedTo)
+			{
+				OutArray.Emplace(Link);
+			}
+		}
+		else if (!Pin->DoesDefaultValueMatchAutogenerated() && !FWildcardNodeUtils::IsWildcardPin(Pin))
+		{
+			// If this pin has doesn't have any links then we should account for the pin itself
+			OutArray.Add(Pin);
+		}
+	}
+}
+
 void UK2Node_PromotableOperator::ConvertNodeToFunction(const UFunction* Function, UEdGraphPin* ChangedPin)
 {
 	const FScopedTransaction Transaction(LOCTEXT("ConvertPromotableOpToFunction", "Change the function signature of a promotable operator node."));
