Blueprints: Allow the equals operator to be a promotable type. Add functions to get the friendly display name of operators as well as their keywords to make then have better search results

#rnx

diff --git a/Engine/Source/Editor/BlueprintGraph/Classes/BlueprintTypePromotion.h b/Engine/Source/Editor/BlueprintGraph/Classes/BlueprintTypePromotion.h
--- a/Engine/Source/Editor/BlueprintGraph/Classes/BlueprintTypePromotion.h
+++ b/Engine/Source/Editor/BlueprintGraph/Classes/BlueprintTypePromotion.h
@@ -38,6 +38,12 @@ public:
 	/** Set of comparison operator names (GreaterThan, LessThan, etc) */
 	static const TSet<FName>& GetComparisonOpNames();
 
+	/** Get the keywords metadata for the given operator name */
+	static const FText& GetKeywordsForOperator(const FName Operator);
+
+	/** Get the user facing version of this operator name */
+	static const FText& GetUserFacingOperatorName(const FName Operator);
+
 	/** Returns true if the given function is a comparison operator */
 	static bool IsComparisonFunc(UFunction const* const Func);
 
diff --git a/Engine/Source/Editor/BlueprintGraph/Private/BlueprintFunctionNodeSpawner.cpp b/Engine/Source/Editor/BlueprintGraph/Private/BlueprintFunctionNodeSpawner.cpp
--- a/Engine/Source/Editor/BlueprintGraph/Private/BlueprintFunctionNodeSpawner.cpp
+++ b/Engine/Source/Editor/BlueprintGraph/Private/BlueprintFunctionNodeSpawner.cpp
@@ -269,14 +269,11 @@ UBlueprintFunctionNodeSpawner* UBlueprintFunctionNodeSpawner::Create(TSubclassOf
 
 	if(bIsPromotableFunction)
 	{
-		FFormatNamedArguments Args;
-		Args.Add(TEXT("OpName"), FText::FromName(OpName));
-		MenuSignature.MenuName = FText::Format(LOCTEXT("OpName", "{OpName}"), Args);
+		MenuSignature.MenuName = FTypePromotion::GetUserFacingOperatorName(OpName);
 		MenuSignature.Category = LOCTEXT("UtilityOperatorCategory", "Utilities|Operators");
 		// Possibly generate some special tooltips for promotable operators?
-		MenuSignature.Tooltip = FText::FromString(TEXT(" "));
-		// add at least one character, so that PrimeDefaultUiSpec() doesn't attempt to query the template node
-		MenuSignature.Keywords = FText::FromString(OpName.ToString());
+		MenuSignature.Tooltip = FTypePromotion::GetUserFacingOperatorName(OpName);
+		MenuSignature.Keywords = FTypePromotion::GetKeywordsForOperator(OpName);
 		FTypePromotion::RegisterOperatorSpawner(OpName, NodeSpawner);
 	}
 	else
diff --git a/Engine/Source/Editor/BlueprintGraph/Private/BlueprintTypePromotion.cpp b/Engine/Source/Editor/BlueprintGraph/Private/BlueprintTypePromotion.cpp
--- a/Engine/Source/Editor/BlueprintGraph/Private/BlueprintTypePromotion.cpp
+++ b/Engine/Source/Editor/BlueprintGraph/Private/BlueprintTypePromotion.cpp
@@ -7,6 +7,8 @@
 
 FTypePromotion* FTypePromotion::Instance = nullptr;
 
+#define LOCTEXT_NAMESPACE "TypePromotion"
+
 namespace OperatorNames
 {
 	static const FName NoOp			= TEXT("NO_OP");
@@ -21,6 +23,7 @@ namespace OperatorNames
 	static const FName Less			= TEXT("Less");
 	static const FName LessEq		= TEXT("LessEqual");
 	static const FName NotEq		= TEXT("NotEqual");
+	static const FName Equal		= TEXT("EqualEqual");
 }
 
 namespace TypePromoDebug
@@ -397,7 +400,8 @@ const TSet<FName>& FTypePromotion::GetAllOpNames()
 		OperatorNames::GreaterEq,
 		OperatorNames::Less,
 		OperatorNames::LessEq,
-		OperatorNames::NotEq
+		OperatorNames::NotEq,
+		OperatorNames::Equal
 	};
 
 	return OpsArray;
@@ -411,11 +415,54 @@ const TSet<FName>& FTypePromotion::GetComparisonOpNames()
 		OperatorNames::GreaterEq,
 		OperatorNames::Less,
 		OperatorNames::LessEq,
-		OperatorNames::NotEq
+		OperatorNames::NotEq,
+		OperatorNames::Equal
 	};
 	return ComparisonOps;
 }
 
+const FText& FTypePromotion::GetKeywordsForOperator(const FName Operator)
+{
+	static const TMap<FName, FText> OpKeywords =
+	{
+		{ OperatorNames::Add,		LOCTEXT("AddKeywords",			"+ add plus") },
+		{ OperatorNames::Multiply,  LOCTEXT("MultiplyKeywords",		"* multiply") },
+		{ OperatorNames::Subtract,  LOCTEXT("SubtractKeywords",		"- subtract minus") },
+		{ OperatorNames::Divide,	LOCTEXT("DivideKeywords",		"/ divide division") },
+		{ OperatorNames::Greater,	LOCTEXT("GreaterKeywords",		"> greater") },
+		{ OperatorNames::GreaterEq, LOCTEXT("GreaterEqKeywords",	">= greater") },
+		{ OperatorNames::Less,		LOCTEXT("LessKeywords",			"< less") },
+		{ OperatorNames::LessEq,	LOCTEXT("LessEqKeywords",		"<= less") },
+		{ OperatorNames::NotEq,		LOCTEXT("NotEqKeywords",		"!= not equal") },
+		{ OperatorNames::Equal,		LOCTEXT("EqualKeywords",		"== equal") },
+		{ OperatorNames::NoOp,		LOCTEXT("NoOpKeywords",			"") },
+	};
+	const FText* Keywords = OpKeywords.Find(Operator);
+	return Keywords ? *Keywords : OpKeywords[OperatorNames::NoOp];
+}
+
+const FText& FTypePromotion::GetUserFacingOperatorName(const FName Operator)
+{
+	static const TMap<FName, FText> OperatorDisplayNames = 
+	{
+		{ OperatorNames::Add,		LOCTEXT("AddDisplayName",			"Add") },
+		{ OperatorNames::Multiply,  LOCTEXT("MultiplyDisplayName",		"Multiply") },
+		{ OperatorNames::Subtract,  LOCTEXT("SubtractDisplayName",		"Subtract") },
+		{ OperatorNames::Divide,	LOCTEXT("DivideDisplayName",		"Division") },
+		{ OperatorNames::Greater,	LOCTEXT("GreaterDisplayName",		"Greater") },
+		{ OperatorNames::GreaterEq, LOCTEXT("GreaterEqDisplayName",		"Greater Equal") },
+		{ OperatorNames::Less,		LOCTEXT("LessDisplayName",			"Less") },
+		{ OperatorNames::LessEq,	LOCTEXT("LessEqDisplayName",		"Less Equal") },
+		{ OperatorNames::NotEq,		LOCTEXT("NotEqDisplayName",			"Not Equal") },
+		{ OperatorNames::Equal,		LOCTEXT("EqualDisplayName",			"Equal") },
+		{ OperatorNames::NoOp,		LOCTEXT("NoOpDisplayName",			"") },
+	};
+
+
+	const FText* Keywords = OperatorDisplayNames.Find(Operator);
+	return Keywords ? *Keywords : OperatorDisplayNames[OperatorNames::NoOp];
+}
+
 bool FTypePromotion::IsComparisonFunc(UFunction const* const Func)
 {
 	return Func && GetComparisonOpNames().Contains(GetOpNameFromFunction(Func));
@@ -531,4 +578,6 @@ void FTypePromotion::ClearNodeSpawners()
 	{
 		Instance->OperatorNodeSpawnerMap.Empty();
 	}
-}
\ No newline at end of file
+}
+
+#undef LOCTEXT_NAMESPACE
\ No newline at end of file
