Fixed unity build error due to shadowing of a global variable.

diff --git a/Engine/Source/Editor/BlueprintGraph/Private/K2Node_SpawnActor.cpp b/Engine/Source/Editor/BlueprintGraph/Private/K2Node_SpawnActor.cpp
--- a/Engine/Source/Editor/BlueprintGraph/Private/K2Node_SpawnActor.cpp
+++ b/Engine/Source/Editor/BlueprintGraph/Private/K2Node_SpawnActor.cpp
@@ -12,11 +12,18 @@
 #include "KismetCompilerMisc.h"
 #include "KismetCompiler.h"
 
-static const FName WorldContextPinName(TEXT("WorldContextObject"));
-static const FName BlueprintPinName(TEXT("Blueprint"));
-static const FName SpawnTransformPinName(TEXT("SpawnTransform"));
-static const FName NoCollisionFailPinName(TEXT("SpawnEvenIfColliding"));
+struct FK2Node_SpawnActorHelper
+{
+	static const FName WorldContextPinName;
+	static const FName BlueprintPinName;
+	static const FName SpawnTransformPinName;
+	static const FName NoCollisionFailPinName;
+};
 
+const FName FK2Node_SpawnActorHelper::WorldContextPinName(TEXT("WorldContextObject"));
+const FName FK2Node_SpawnActorHelper::BlueprintPinName(TEXT("Blueprint"));
+const FName FK2Node_SpawnActorHelper::SpawnTransformPinName(TEXT("SpawnTransform"));
+const FName FK2Node_SpawnActorHelper::NoCollisionFailPinName(TEXT("SpawnEvenIfColliding"));
 
 #define LOCTEXT_NAMESPACE "K2Node_SpawnActor"
 
@@ -37,20 +44,20 @@ void UK2Node_SpawnActor::AllocateDefaultPins()
 	// If required add the world context pin
 	if (GetBlueprint()->ParentClass->HasMetaDataHierarchical(FBlueprintMetadata::MD_ShowWorldContextPin))
 	{
-		CreatePin(EGPD_Input, UEdGraphSchema_K2::PC_Object, UObject::StaticClass(), WorldContextPinName);
+		CreatePin(EGPD_Input, UEdGraphSchema_K2::PC_Object, UObject::StaticClass(), FK2Node_SpawnActorHelper::WorldContextPinName);
 	}
 
 	// Add blueprint pin
-	UEdGraphPin* BlueprintPin = CreatePin(EGPD_Input, UEdGraphSchema_K2::PC_Object, UBlueprint::StaticClass(), BlueprintPinName);
+	UEdGraphPin* BlueprintPin = CreatePin(EGPD_Input, UEdGraphSchema_K2::PC_Object, UBlueprint::StaticClass(), FK2Node_SpawnActorHelper::BlueprintPinName);
 	K2Schema->ConstructBasicPinTooltip(*BlueprintPin, LOCTEXT("BlueprintPinDescription", "The blueprint Actor you want to spawn"), BlueprintPin->PinToolTip);
 
 	// Transform pin
 	UScriptStruct* TransformStruct = TBaseStructure<FTransform>::Get();
-	UEdGraphPin* TransformPin = CreatePin(EGPD_Input, UEdGraphSchema_K2::PC_Struct, TransformStruct, SpawnTransformPinName);
+	UEdGraphPin* TransformPin = CreatePin(EGPD_Input, UEdGraphSchema_K2::PC_Struct, TransformStruct, FK2Node_SpawnActorHelper::SpawnTransformPinName);
 	K2Schema->ConstructBasicPinTooltip(*TransformPin, LOCTEXT("TransformPinDescription", "The transform to spawn the Actor with"), TransformPin->PinToolTip);
 
 	// bNoCollisionFail pin
-	UEdGraphPin* NoCollisionFailPin = CreatePin(EGPD_Input, UEdGraphSchema_K2::PC_Boolean, NoCollisionFailPinName);
+	UEdGraphPin* NoCollisionFailPin = CreatePin(EGPD_Input, UEdGraphSchema_K2::PC_Boolean, FK2Node_SpawnActorHelper::NoCollisionFailPinName);
 	K2Schema->ConstructBasicPinTooltip(*NoCollisionFailPin, LOCTEXT("NoCollisionFailPinDescription", "Determines if the Actor should be spawned when the location is blocked by a collision"), NoCollisionFailPin->PinToolTip);
 
 	// Result pin
@@ -128,16 +135,16 @@ bool UK2Node_SpawnActor::IsSpawnVarPin(UEdGraphPin* Pin)
 	return(	Pin->PinName != UEdGraphSchema_K2::PN_Execute &&
 			Pin->PinName != UEdGraphSchema_K2::PN_Then &&
 			Pin->PinName != UEdGraphSchema_K2::PN_ReturnValue &&
-			Pin->PinName != BlueprintPinName &&
-			Pin->PinName != WorldContextPinName &&
-			Pin->PinName != NoCollisionFailPinName &&
-			Pin->PinName != SpawnTransformPinName );
+			Pin->PinName != FK2Node_SpawnActorHelper::BlueprintPinName &&
+			Pin->PinName != FK2Node_SpawnActorHelper::WorldContextPinName &&
+			Pin->PinName != FK2Node_SpawnActorHelper::NoCollisionFailPinName &&
+			Pin->PinName != FK2Node_SpawnActorHelper::SpawnTransformPinName );
 }
 
 
 void UK2Node_SpawnActor::PinDefaultValueChanged(UEdGraphPin* ChangedPin) 
 {
-	if (ChangedPin->PinName == BlueprintPinName)
+	if (ChangedPin->PinName == FK2Node_SpawnActorHelper::BlueprintPinName)
 	{
 		const UEdGraphSchema_K2* K2Schema = GetDefault<UEdGraphSchema_K2>();
 
@@ -195,7 +202,7 @@ UEdGraphPin* UK2Node_SpawnActor::GetBlueprintPin(const TArray<UEdGraphPin*>* InP
 	UEdGraphPin* Pin = nullptr;
 	for (UEdGraphPin* TestPin : *PinsToSearch)
 	{
-		if( TestPin && TestPin->PinName == BlueprintPinName )
+		if( TestPin && TestPin->PinName == FK2Node_SpawnActorHelper::BlueprintPinName )
 		{
 			Pin = TestPin;
 			break;
@@ -207,21 +214,21 @@ UEdGraphPin* UK2Node_SpawnActor::GetBlueprintPin(const TArray<UEdGraphPin*>* InP
 
 UEdGraphPin* UK2Node_SpawnActor::GetSpawnTransformPin()const
 {
-	UEdGraphPin* Pin = FindPinChecked(SpawnTransformPinName);
+	UEdGraphPin* Pin = FindPinChecked(FK2Node_SpawnActorHelper::SpawnTransformPinName);
 	check(Pin->Direction == EGPD_Input);
 	return Pin;
 }
 
 UEdGraphPin* UK2Node_SpawnActor::GetNoCollisionFailPin()const
 {
-	UEdGraphPin* Pin = FindPinChecked(NoCollisionFailPinName);
+	UEdGraphPin* Pin = FindPinChecked(FK2Node_SpawnActorHelper::NoCollisionFailPinName);
 	check(Pin->Direction == EGPD_Input);
 	return Pin;
 }
 
 UEdGraphPin* UK2Node_SpawnActor::GetWorldContextPin() const
 {
-	UEdGraphPin* Pin = FindPin(WorldContextPinName);
+	UEdGraphPin* Pin = FindPin(FK2Node_SpawnActorHelper::WorldContextPinName);
 	check(Pin == NULL || Pin->Direction == EGPD_Input);
 	return Pin;
 }
