Clean up helper methods within the Type Promotion unit tests

#rnx

diff --git a/Engine/Source/Editor/BlueprintGraph/Private/Tests/TypePromotionTests.cpp b/Engine/Source/Editor/BlueprintGraph/Private/Tests/TypePromotionTests.cpp
--- a/Engine/Source/Editor/BlueprintGraph/Private/Tests/TypePromotionTests.cpp
+++ b/Engine/Source/Editor/BlueprintGraph/Private/Tests/TypePromotionTests.cpp
@@ -5,19 +5,146 @@
 #include "Misc/AutomationTest.h"
 #include "BlueprintTypePromotion.h"
 #include "EdGraph/EdGraphPin.h"
+#include "K2Node_PromotableOperator.h"
+#include "Kismet2/BlueprintEditorUtils.h"
+#include "GameFramework/Actor.h"			// For making a dummy BP
+#include "Kismet2/KismetEditorUtilities.h"
+#include "Kismet2/WildcardNodeUtils.h"
 
 #if WITH_DEV_AUTOMATION_TESTS
 
-IMPLEMENT_SIMPLE_AUTOMATION_TEST(FTypePromotionTest, "Blueprints.Compiler.TypePromotion", EAutomationTestFlags::EditorContext | EAutomationTestFlags::SmokeFilter)
+/////////////////////////////////////////////////////
+// Helpers to make dummy blueprints/pins/nodes
+
+namespace TypePromoTestUtils
+{
+	/**
+	* Spawn a test promotable operator node that we can use to ensure type propegation works
+	* correctly.
+	*
+	* @param Graph		The graph to spawn this node in
+	* @param OpName		Name of the promotable op to spawn (Multiply, Add, Subtract, etc)
+	*
+	* @return UK2Node_PromotableOperator* if the op exists,
+	*/
+	static UK2Node_PromotableOperator* SpawnPromotableNode(UEdGraph* Graph, const FName OpName)
+	{
+		if (!Graph)
+		{
+			return nullptr;
+		}
+
+		// The spawner will be null if type promo isn't enabled
+		if (UBlueprintFunctionNodeSpawner* Spawner = FTypePromotion::GetOperatorSpawner(OpName))
+		{
+			// Spawn a new node!
+			IBlueprintNodeBinder::FBindingSet Bindings;
+			FVector2D SpawnLoc{};
+			UK2Node_PromotableOperator* NewOpNode = Cast<UK2Node_PromotableOperator>(Spawner->Invoke(Graph, Bindings, SpawnLoc));
+
+			return NewOpNode;
+		}
+
+		return nullptr;
+	}
+
+	/**
+	* Mark this array of spawned test pins as pending kill to ensure that
+	* they get cleaned up properly by GC
+	*
+	* @param InPins		The array of spawned test pins to mark pending kill
+	*/
+	static void CleanupTestPins(TArray<UEdGraphPin*>& InPins)
+	{
+		// Clear our test pins
+		for (UEdGraphPin* TestPin : InPins)
+		{
+			if (TestPin)
+			{
+				TestPin->MarkPendingKill();
+			}
+		}
+		InPins.Empty();
+	}
+
+
+	/**
+	* Attempts to create a connection between the two given pins and tests that the connection was valid
+	*
+	* @param OpNodePin	The pin on a promotable operator node
+	* @param OtherPin	The other pin
+	*/
+	static bool TestPromotedConnection(UEdGraphPin* OpNodePin, UEdGraphPin* OtherPin)
+	{
+		const UEdGraphSchema_K2* K2Schema = GetDefault<UEdGraphSchema_K2>();
+		const bool bConnected = K2Schema->TryCreateConnection(OpNodePin, OtherPin);
+		UK2Node_PromotableOperator* OwningNode = Cast<UK2Node_PromotableOperator>(OpNodePin->GetOwningNode());
+
+		if (bConnected && OwningNode)
+		{
+			OwningNode->NotifyPinConnectionListChanged(OpNodePin);
+		}
 
+		return bConnected;
+	}
+}
+
+#define MakeTestableBP(BPName, GraphName)												\
+		UBlueprint* BPName = FKismetEditorUtilities::CreateBlueprint(					\
+		AActor::StaticClass(),															\
+		GetTransientPackage(),															\
+		TEXT(#BPName),																	\
+		BPTYPE_Normal,																	\
+		UBlueprint::StaticClass(),														\
+		UBlueprintGeneratedClass::StaticClass(),										\
+		NAME_None																		\
+	);																					\
+	UEdGraph* GraphName = BPName ? FBlueprintEditorUtils::FindEventGraph( BPName ) : nullptr;	
+
+#define MakeTestableNode(NodeName, OwningGraph)											\
+		UK2Node* NodeName = NewObject<UK2Node_CallFunction>(/*outer*/ OwningGraph );	\
+		OwningGraph->AddNode(NodeName);
+
+#define MakeTestPin(OwningNode, PinArray, InPinName, InPinType, PinDirection)		UEdGraphPin* InPinName = UEdGraphPin::CreatePin(OwningNode);		\
+																				InPinName->PinType.PinCategory = InPinType;							\
+																				PinArray.Add(InPinName);											\
+																				InPinName->Direction = PinDirection;
+
+#define MakeTestPins(OwningNode, OutArray)																\
+	MakeTestPin(OwningNode, OutArray, DoublePinA, UEdGraphSchema_K2::PC_Double, EGPD_Output);			\
+	MakeTestPin(OwningNode, OutArray, DoublePinB, UEdGraphSchema_K2::PC_Double, EGPD_Output);			\
+	MakeTestPin(OwningNode, OutArray, Int64PinA, UEdGraphSchema_K2::PC_Int64, EGPD_Output);				\
+	MakeTestPin(OwningNode, OutArray, Int64PinB, UEdGraphSchema_K2::PC_Int64, EGPD_Input);				\
+	MakeTestPin(OwningNode, OutArray, BytePinA, UEdGraphSchema_K2::PC_Byte, EGPD_Output);				\
+	MakeTestPin(OwningNode, OutArray, WildPinA, UEdGraphSchema_K2::PC_Wildcard, EGPD_Output);			\
+	MakeTestPin(OwningNode, OutArray, WildPinB, UEdGraphSchema_K2::PC_Wildcard, EGPD_Input);			\
+	MakeTestPin(OwningNode, OutArray, BytePinB, UEdGraphSchema_K2::PC_Byte, EGPD_Input);				\
+	MakeTestPin(OwningNode, OutArray, BoolPinA, UEdGraphSchema_K2::PC_Boolean, EGPD_Output);			\
+	MakeTestPin(OwningNode, OutArray, BoolPinB, UEdGraphSchema_K2::PC_Boolean, EGPD_Input);				\
+	MakeTestPin(OwningNode, OutArray, DoubleOutputPin, UEdGraphSchema_K2::PC_Double, EGPD_Output);		\
+	MakeTestPin(OwningNode, OutArray, FloatPinA, UEdGraphSchema_K2::PC_Float, EGPD_Input);				\
+	MakeTestPin(OwningNode, OutArray, FloatPinB, UEdGraphSchema_K2::PC_Float, EGPD_Input);				\
+	MakeTestPin(OwningNode, OutArray, FloatOutputPin, UEdGraphSchema_K2::PC_Float, EGPD_Output);		\
+	MakeTestPin(OwningNode, OutArray, BoolOutputPin, UEdGraphSchema_K2::PC_Boolean, EGPD_Output);		\
+	MakeTestPin(OwningNode, OutArray, IntPinA, UEdGraphSchema_K2::PC_Int, EGPD_Output);					\
+	MakeTestPin(OwningNode, OutArray, VecInputPinA, UEdGraphSchema_K2::PC_Struct, EGPD_Input);			\
+	VecInputPinA->PinType.PinSubCategoryObject = TBaseStructure<FVector>::Get();						\
+	MakeTestPin(OwningNode, OutArray, VecInputPinB, UEdGraphSchema_K2::PC_Struct, EGPD_Input);			\
+	VecInputPinB->PinType.PinSubCategoryObject = TBaseStructure<FVector>::Get();						\
+	MakeTestPin(OwningNode, OutArray, VecOutputPinA, UEdGraphSchema_K2::PC_Struct, EGPD_Output);		\
+	VecOutputPinA->PinType.PinSubCategoryObject = TBaseStructure<FVector>::Get();						\
+	MakeTestPin(OwningNode, OutArray, Vec2DOutputPinA, UEdGraphSchema_K2::PC_Struct, EGPD_Output);		\
+	Vec2DOutputPinA->PinType.PinSubCategoryObject = TBaseStructure<FVector2D>::Get();
+
+IMPLEMENT_SIMPLE_AUTOMATION_TEST(FTypePromotionTest, "Blueprints.Compiler.TypePromotion", EAutomationTestFlags::EditorContext | EAutomationTestFlags::SmokeFilter)
 bool FTypePromotionTest::RunTest(const FString& Parameters)
 {
-	FEdGraphPinType DoublePin = {};		DoublePin.PinCategory	= UEdGraphSchema_K2::PC_Double;
-	FEdGraphPinType FloatPin = {};		FloatPin.PinCategory	= UEdGraphSchema_K2::PC_Float;
-	FEdGraphPinType IntPin = {};		IntPin.PinCategory		= UEdGraphSchema_K2::PC_Int;
-	FEdGraphPinType Int64Pin = {};		Int64Pin.PinCategory	= UEdGraphSchema_K2::PC_Int64;
-	FEdGraphPinType BytePin = {};		BytePin.PinCategory		= UEdGraphSchema_K2::PC_Byte;
-	FEdGraphPinType VecPin = {};		VecPin.PinCategory		= UEdGraphSchema_K2::PC_Struct; VecPin.PinSubCategoryObject = TBaseStructure<FVector>::Get();
+	FEdGraphPinType DoublePin = {};		DoublePin.PinCategory = UEdGraphSchema_K2::PC_Double;
+	FEdGraphPinType FloatPin = {};		FloatPin.PinCategory = UEdGraphSchema_K2::PC_Float;
+	FEdGraphPinType IntPin = {};		IntPin.PinCategory = UEdGraphSchema_K2::PC_Int;
+	FEdGraphPinType Int64Pin = {};		Int64Pin.PinCategory = UEdGraphSchema_K2::PC_Int64;
+	FEdGraphPinType BytePin = {};		BytePin.PinCategory = UEdGraphSchema_K2::PC_Byte;
+	FEdGraphPinType VecPin = {};		VecPin.PinCategory = UEdGraphSchema_K2::PC_Struct; VecPin.PinSubCategoryObject = TBaseStructure<FVector>::Get();
 
 	// Test promotions that should happen
 	TestEqual(TEXT("Testing float to double"), FTypePromotion::GetHigherType(FloatPin, DoublePin), FTypePromotion::ETypeComparisonResult::TypeBHigher);
@@ -26,10 +153,10 @@ bool FTypePromotionTest::RunTest(const FString& Parameters)
 	TestEqual(TEXT("Testing int to float"), FTypePromotion::GetHigherType(IntPin, FloatPin), FTypePromotion::ETypeComparisonResult::TypeBHigher);
 	TestEqual(TEXT("Testing int to double"), FTypePromotion::GetHigherType(IntPin, DoublePin), FTypePromotion::ETypeComparisonResult::TypeBHigher);
 	TestEqual(TEXT("Testing int to int64"), FTypePromotion::GetHigherType(IntPin, Int64Pin), FTypePromotion::ETypeComparisonResult::TypeBHigher);
-	
+
 	TestEqual(TEXT("Testing Byte to int"), FTypePromotion::GetHigherType(BytePin, IntPin), FTypePromotion::ETypeComparisonResult::TypeBHigher);
 	TestEqual(TEXT("Testing Byte to int64"), FTypePromotion::GetHigherType(BytePin, Int64Pin), FTypePromotion::ETypeComparisonResult::TypeBHigher);
-	
+
 	TestEqual(TEXT("Testing Double to int64"), FTypePromotion::GetHigherType(DoublePin, Int64Pin), FTypePromotion::ETypeComparisonResult::TypeBHigher);
 
 	// Test Equality of pins
@@ -48,50 +175,37 @@ bool FTypePromotionTest::RunTest(const FString& Parameters)
 	return true;
 }
 
-
-#define MakeTestPin(OwningNode, PinArray, InPinName, InPinType)		UEdGraphPin* InPinName = UEdGraphPin::CreatePin(OwningNode);		\
-																	InPinName->PinType.PinCategory = InPinType;							\
-																	PinArray.Add(InPinName);		
-
+// Test that when given an array of UEdGraphPins we can find the appropriate UFunction that best
+// matches them. This is the core of how the Type Promotion system works at BP compile time
 IMPLEMENT_SIMPLE_AUTOMATION_TEST(FFindBestMatchingFunc, "Blueprints.Compiler.FindBestMatchingFunc", EAutomationTestFlags::EditorContext | EAutomationTestFlags::EngineFilter)
 bool FFindBestMatchingFunc::RunTest(const FString& Parameters)
 {
-	TArray<UEdGraphPin*> PinTypes = {};
+	if (!TypePromoDebug::IsTypePromoEnabled())
+	{
+		return true;
+	}
 
 	UEdGraphNode* TestNode = NewObject<UEdGraphNode>();
 
 	////////////////////////////////////////////////////
 	// Create test pins!
-	MakeTestPin(TestNode, PinTypes, DoublePinA, UEdGraphSchema_K2::PC_Double);
-	MakeTestPin(TestNode, PinTypes, DoublePinB, UEdGraphSchema_K2::PC_Double);
-	MakeTestPin(TestNode, PinTypes, DoubleOutputPin, UEdGraphSchema_K2::PC_Double);
-	DoubleOutputPin->Direction = EGPD_Output;
-	
-	MakeTestPin(TestNode, PinTypes, FloatPinA, UEdGraphSchema_K2::PC_Float);
-	FloatPinA->Direction = EGPD_Input;
-	MakeTestPin(TestNode, PinTypes, FloatPinB, UEdGraphSchema_K2::PC_Float);
-	FloatPinB->Direction = EGPD_Input;
-
-	MakeTestPin(TestNode, PinTypes, FloatOutputPin, UEdGraphSchema_K2::PC_Float);
-	FloatPinB->Direction = EGPD_Output;
-
-	MakeTestPin(TestNode, PinTypes, BoolOutputPin, UEdGraphSchema_K2::PC_Boolean);
-	BoolOutputPin->Direction = EGPD_Output;
-
-	MakeTestPin(TestNode, PinTypes, IntPinA, UEdGraphSchema_K2::PC_Int);
-
-	// Structs!
-	MakeTestPin(TestNode, PinTypes, VecInputPinA, UEdGraphSchema_K2::PC_Struct);
-	VecInputPinA->Direction = EGPD_Input;
-	VecInputPinA->PinType.PinSubCategoryObject = TBaseStructure<FVector>::Get();
-
-	MakeTestPin(TestNode, PinTypes, VecInputPinB, UEdGraphSchema_K2::PC_Struct);
-	VecInputPinB->Direction = EGPD_Input;
-	VecInputPinB->PinType.PinSubCategoryObject = TBaseStructure<FVector>::Get();
-
-	MakeTestPin(TestNode, PinTypes, VecOutputPinA, UEdGraphSchema_K2::PC_Struct);
-	VecOutputPinA->Direction = EGPD_Output;
-	VecOutputPinA->PinType.PinSubCategoryObject = TBaseStructure<FVector>::Get();
+	TArray<UEdGraphPin*> PinTypes = {};
+	MakeTestPins(TestNode, PinTypes);
+
+	{
+		TArray<UEdGraphPin*> TestPins =
+		{
+			Vec2DOutputPinA,
+		};
+
+		const UFunction* AddVecFunc = FTypePromotion::FindBestMatchingFunc(TEXT("Add"), TestPins);
+		static const FName ExpectedName = TEXT("Add_Vector2DVector2D");
+
+		if (TestNotNull(TEXT("Add_Vector2DVector2D Null check"), AddVecFunc))
+		{
+			TestEqual(TEXT("Add_Vector2DVector2D Name Check"), AddVecFunc->GetFName(), ExpectedName);
+		}
+	}
 
 	// Multiply_VectorVector given A float input, vector input, and a vector output
 	{
