Blueprints: Add a type promotion unit test that will cover making connections to the node and allow us to check that things don't disconnect when they shouldn't.

#rnx

diff --git a/Engine/Source/Editor/BlueprintGraph/Private/Tests/TypePromotionTests.cpp b/Engine/Source/Editor/BlueprintGraph/Private/Tests/TypePromotionTests.cpp
--- a/Engine/Source/Editor/BlueprintGraph/Private/Tests/TypePromotionTests.cpp
+++ b/Engine/Source/Editor/BlueprintGraph/Private/Tests/TypePromotionTests.cpp
@@ -383,6 +383,319 @@ bool FFindBestMatchingFunc::RunTest(const FString& Parameters)
 	return true;
 }
 
+// Test the default state of all operator nodes to ensure they are correct.
+// Comparison operators (Greater Than, Less Than, etc) should have two 
+// wildcard inputs and one boolean output. All others should be all wildcards.
+// The node's set function should also match the operator correctly and 
+// it should have the 'OperationName' variable set.
+IMPLEMENT_SIMPLE_AUTOMATION_TEST(FPromotableOpDefaultState, "Blueprints.Nodes.PromotableOp.DefaultState", EAutomationTestFlags::EditorContext | EAutomationTestFlags::EngineFilter)
+bool FPromotableOpDefaultState::RunTest(const FString& Parameters)
+{
+	if (!TypePromoDebug::IsTypePromoEnabled())
+	{
+		return true;
+	}
+
+	// Refresh the actions within this test in case the editor is open but hasn't loaded BlueprintGraph yet
+	FTypePromotion::ClearNodeSpawners();
+	FBlueprintActionDatabase::Get().RefreshAll();
+
+	MakeTestableBP(WildcardStartTestBP, TestWildcardGraph);
+
+	const TSet<FName> AllOpNames = FTypePromotion::GetAllOpNames();
+
+	for (const FName& OpName : AllOpNames)
+	{
+		const bool bIsComparisonOp = FTypePromotion::GetComparisonOpNames().Contains(OpName);
+		const FString OpNameString = OpName.ToString();
+
+		UK2Node_PromotableOperator* OpNode = TypePromoTestUtils::SpawnPromotableNode(TestWildcardGraph, OpName);
+		TestNotNull(FString::Printf(TEXT("Spawning a '%s' operator node"), *OpNameString), OpNode);
+		
+		// The 'OperationName' variable is correct
+		TestTrue(FString::Printf(TEXT("Operation Name '%s' matches after spawning node"), *OpNameString), OpNode->GetOperationName() == OpName);
+
+		// The target function has been set when the node is spawned
+		const UFunction* TargetFunc = OpNode->GetTargetFunction();
+		TestNotNull(FString::Printf(TEXT("'%s' Operation function is not null"), *OpNameString), TargetFunc);
+		
+		// The target function is of the correct operation type
+		const FName TargetFunctionOpName = FTypePromotion::GetOpNameFromFunction(TargetFunc);
+		TestTrue(FString::Printf(TEXT("'%s' Operation function matches requested operation"), *OpNameString), TargetFunctionOpName == OpName);
+
+		// Test pin types
+		const UEdGraphPin* TopInputPin = OpNode->FindPin(TEXT("A"), EGPD_Input);
+		const UEdGraphPin* BottomInputPin = OpNode->FindPin(TEXT("B"), EGPD_Input);
+		const UEdGraphPin* OutputPin = OpNode->GetOutputPin();
+		
+		// Comparison operators should be spawned with two wildcard inputs and a boolean output
+		if (bIsComparisonOp)
+		{
+			TestTrue(TEXT("Top is pin wildcard"), FWildcardNodeUtils::IsWildcardPin(TopInputPin));
+			TestTrue(TEXT("Bottom Pin is wildcard"), FWildcardNodeUtils::IsWildcardPin(BottomInputPin));
+			TestTrue(TEXT("Bottom Pin is a bool"), OutputPin->PinType.PinCategory == UEdGraphSchema_K2::PC_Boolean);
+		}
+		// All other operators should be spawned with all wildcard inputs and output
+		else
+		{
+			TestTrue(TEXT("Top is pin wildcard"), FWildcardNodeUtils::IsWildcardPin(TopInputPin));
+			TestTrue(TEXT("Bottom Pin is wildcard"), FWildcardNodeUtils::IsWildcardPin(BottomInputPin));
+			TestTrue(TEXT("Bottom Pin is a wildcard"), FWildcardNodeUtils::IsWildcardPin(OutputPin));
+		}
+	}
+
+	// Cleanup test BP and graph
+	{
+		WildcardStartTestBP->MarkPendingKill();
+		TestWildcardGraph->MarkPendingKill();
+	}
+
+	return true;
+}
+
+// Test that promotable operator nodes can correctly have pins added to them
+// and that comparison operators cannot have pins added to them.
+IMPLEMENT_SIMPLE_AUTOMATION_TEST(FPromotableOpNodeAddPinInterface, "Blueprints.Nodes.PromotableOp.AddPinInterface", EAutomationTestFlags::EditorContext | EAutomationTestFlags::EngineFilter)
+bool FPromotableOpNodeAddPinInterface::RunTest(const FString& Parameters)
+{
+	if (!TypePromoDebug::IsTypePromoEnabled())
+	{
+		return true;
+	}
+
+	// Refresh the actions within this test in case the editor is open but hasn't loaded BlueprintGraph yet
+	FTypePromotion::ClearNodeSpawners();
+	FBlueprintActionDatabase::Get().RefreshAll();
+
+	MakeTestableBP(TestBP, TestGraph);
+	MakeTestableNode(TestNode, TestGraph);
+
+	// Create test pins!
+	TArray<UEdGraphPin*> PinTypes = {};
+	MakeTestPins(TestNode, PinTypes);
+
+	const TSet<FName> ComparisonOpNames = FTypePromotion::GetComparisonOpNames();
+	for (const FName& OpName : ComparisonOpNames)
+	{
+		const FString OpNameString = OpName.ToString();
+
+		UK2Node_PromotableOperator* OpNode = TypePromoTestUtils::SpawnPromotableNode(TestGraph, OpName);
+		TestNotNull(FString::Printf(TEXT("'%s' Comparison op spawned"), *OpNameString), OpNode);
+
+		TestFalse(FString::Printf(TEXT("'%s' Comparison op cannot add pin"), *OpNameString), OpNode->CanAddPin());
+	}
+
+	// Anything that is not a comparison operator can have a pin added to it
+	{
+		UK2Node_PromotableOperator* MultiplyNode = TypePromoTestUtils::SpawnPromotableNode(TestGraph, TEXT("Multiply"));
+		TestNotNull(TEXT("Multiply Node spawn"), MultiplyNode);
+		TestTrue(TEXT("Multiply can add pin"), MultiplyNode->CanAddPin());
+	}
+
+	{
+		UK2Node_PromotableOperator* MultiplyNode = TypePromoTestUtils::SpawnPromotableNode(TestGraph, TEXT("Multiply"));
+		TestNotNull(TEXT("Multiply Node spawn"), MultiplyNode);
+		TestTrue(TEXT("Multiply can add pin"), MultiplyNode->CanAddPin());
+
+		const int32 StartingPinCount = MultiplyNode->Pins.Num();
+		MultiplyNode->AddInputPin();
+		const int32 EndingPinCount = MultiplyNode->Pins.Num();
+
+		TestTrue(TEXT("Multiply node had a pin added to it"), EndingPinCount == StartingPinCount + 1);
+		const UEdGraphPin* AdditonalPin = MultiplyNode->GetAdditionalPin(EndingPinCount - 1);
+		TestNotNull(TEXT("Additional Pin is not null"), AdditonalPin);
+		TestTrue(TEXT("New Pin is wildcard"), FWildcardNodeUtils::IsWildcardPin(AdditonalPin));
+	}
+
+	// Cleanup
+	{
+		TypePromoTestUtils::CleanupTestPins(PinTypes);
+
+		TestBP->MarkPendingKill();
+		TestGraph->MarkPendingKill();
+		TestNode->MarkPendingKill();
+	}
+
+	return true;
+}
+
+// Test that making connections to a Promotable Operator node results in the correct propagation of types
+// throughout the whole node and that the node has the correct UFunction that it will expand to upon compiling.
+// This will also test that pin connections are broken if they are connected to an invalid promotion, 
+// and that pin connections are preserved if a valid promotion is occuring. 
+IMPLEMENT_SIMPLE_AUTOMATION_TEST(FPromotableOperatorConnectionChanged, "Blueprints.Nodes.PromotableOp.ConnectionChanged", EAutomationTestFlags::EditorContext | EAutomationTestFlags::EngineFilter)
+bool FPromotableOperatorConnectionChanged::RunTest(const FString& Parameters)
+{
+	if (!TypePromoDebug::IsTypePromoEnabled())
+	{
+		return true;
+	}
+
+	// Refresh the actions within this test in case the editor is open but hasn't loaded BlueprintGraph yet
+	FTypePromotion::ClearNodeSpawners();
+	FBlueprintActionDatabase::Get().RefreshAll();
+
+	MakeTestableBP(BP_ConnectionChanged, TestGraph);
+	MakeTestableNode(TestNode, TestGraph);
+
+	// Create test pins!
+	TArray<UEdGraphPin*> PinTypes = {};
+	MakeTestPins(TestNode, PinTypes);
+
+	const UEdGraphSchema_K2* K2Schema = GetDefault<UEdGraphSchema_K2>();
+
+	// Test that adding a float pin to the top input on an add node makes the whole thing a float
+	{
+		UK2Node_PromotableOperator* AddNode = TypePromoTestUtils::SpawnPromotableNode(TestGraph, TEXT("Add"));
+		TestNotNull(TEXT("AddNode Node spawn"), AddNode);
+
+		UEdGraphPin* TopInputPin = AddNode->FindPin(TEXT("A"), EGPD_Input);
+		UEdGraphPin* BottomInputPin = AddNode->FindPin(TEXT("B"), EGPD_Input);
+		UEdGraphPin* OutputPin = AddNode->GetOutputPin();
+
+		check(TopInputPin && BottomInputPin && OutputPin);
+
+		const bool bConnected = K2Schema->TryCreateConnection(TopInputPin, FloatOutputPin);
+		AddNode->NotifyPinConnectionListChanged(TopInputPin);
+
+		TestTrue(TEXT("Bottom Pin type propegates to float"), bConnected && BottomInputPin->PinType.PinCategory == FloatPinB->PinType.PinCategory);
+	}
+
+	// Connecting a vector output should make the other input be a vector as well
+	{
+		UK2Node_PromotableOperator* Node = TypePromoTestUtils::SpawnPromotableNode(TestGraph, TEXT("Multiply"));
+		TestNotNull(TEXT("Multiply Node spawn"), Node);
+
+		UEdGraphPin* TopInputPin = Node->FindPin(TEXT("A"), EGPD_Input);
+		UEdGraphPin* OutputPin = Node->GetOutputPin();
+
+		check(TopInputPin);
+
+		const bool bConnected = K2Schema->TryCreateConnection(OutputPin, VecInputPinA);
+		Node->NotifyPinConnectionListChanged(OutputPin);
+
+		TestTrue(TEXT("Bottom Pin type propegates to float"), bConnected && TopInputPin->PinType.PinCategory == VecOutputPinA->PinType.PinCategory);
+	}
+
+	// Connecting a vector output should make the other input be a vector as well
+	{
+		UK2Node_PromotableOperator* Node = TypePromoTestUtils::SpawnPromotableNode(TestGraph, TEXT("Multiply"));
+		TestNotNull(TEXT("Multiply Node spawn"), Node);
+
+		UEdGraphPin* TopInputPin = Node->FindPin(TEXT("A"), EGPD_Input);
+		UEdGraphPin* BottomInputPin = Node->FindPin(TEXT("B"), EGPD_Input);
+		UEdGraphPin* OutputPin = Node->GetOutputPin();
+
+		// Connect a float to the top pin
+		const bool bConnected = K2Schema->TryCreateConnection(OutputPin, VecInputPinA);
+		Node->NotifyPinConnectionListChanged(OutputPin);
+		TestTrue(TEXT("Bottom Pin type propegates to float"), bConnected && TopInputPin->PinType.PinCategory == VecOutputPinA->PinType.PinCategory);
+
+		// The output should be a float right now
+
+		// Connect a double to the bottom pin
+
+		// The output should be a double now
+	}
+
+	// Cleanup
+	{
+		TypePromoTestUtils::CleanupTestPins(PinTypes);
+
+		BP_ConnectionChanged->MarkPendingKill();
+		TestGraph->MarkPendingKill();
+		TestNode->MarkPendingKill();
+	}
+
+	return true;
+}
+
+// Test the connections between primitive types and ensure that each one gets
+// the correct output type pin
+IMPLEMENT_SIMPLE_AUTOMATION_TEST(FPromotableOperatorPrimitivePromotions, "Blueprints.Nodes.PromotableOp.PrimitivePromotions", EAutomationTestFlags::EditorContext | EAutomationTestFlags::EngineFilter)
+bool FPromotableOperatorPrimitivePromotions::RunTest(const FString& Parameters)
+{
+	if (!TypePromoDebug::IsTypePromoEnabled())
+	{
+		return true;
+	}
+
+	// Refresh the actions within this test in case the editor is open but hasn't loaded BlueprintGraph yet
+	FTypePromotion::ClearNodeSpawners();
+	FBlueprintActionDatabase::Get().RefreshAll();
+
+	MakeTestableBP(BP_Primative_Connections, TestGraph);
+	MakeTestableNode(TestNode, TestGraph);
+
+	// Create test pins!
+	TArray<UEdGraphPin*> PinTypes = {};
+	MakeTestPins(TestNode, PinTypes);
+
+	const UEdGraphSchema_K2* K2Schema = GetDefault<UEdGraphSchema_K2>();
+
+	const TMap<FName, TArray<FName>>* const PromoTable = FTypePromotion::GetPrimativePromotionTable();
+	TestNotNull(TEXT("Primative Promotion table exists"), PromoTable);
+	
+	if (!PromoTable)
+	{
+		return false;
+	}
+
+
+	for(TPair<FName, TArray<FName>> Pair : *PromoTable)
+	{
+		FName TypeName = Pair.Key;
+		if(TypeName == UEdGraphSchema_K2::PC_Wildcard)
+		{
+			continue;
+		}
+		
+		UEdGraphPin* TypePin = *PinTypes.FindByPredicate([&TypeName](const UEdGraphPin* Pin) -> bool { return Pin && Pin->PinType.PinCategory == TypeName && Pin->Direction == EGPD_Output; });
+		
+		for (const FName AvailablePromoType : Pair.Value)
+		{
+			UK2Node_PromotableOperator* Node = TypePromoTestUtils::SpawnPromotableNode(TestGraph, TEXT("Add"));
+			TestNotNull(TEXT("Multiply Node spawn"), Node);
+
+			UEdGraphPin* TopInputPin = Node->FindPin(TEXT("A"), EGPD_Input);
+			UEdGraphPin* BottomInputPin = Node->FindPin(TEXT("B"), EGPD_Input);
+			UEdGraphPin* OutputPin = Node->GetOutputPin();
+			UEdGraphPin* PinToConnectTo = *PinTypes.FindByPredicate([&AvailablePromoType](const UEdGraphPin* Pin) -> bool { return Pin && Pin->PinType.PinCategory == AvailablePromoType && Pin->Direction == EGPD_Output; });
+
+			// Connect to the top input pin
+			const bool bConnectedTop = TypePromoTestUtils::TestPromotedConnection(TopInputPin, TypePin);
+			FString StatusMessage = FString::Printf(TEXT("Connecting '%s' to '%s'"), *K2Schema->TypeToText(TopInputPin->PinType).ToString(), *K2Schema->TypeToText(TypePin->PinType).ToString());
+
+			TestTrue(StatusMessage, bConnectedTop);
+
+			// The other pins should now all be set to the first pin type
+			TestTrue(TEXT("Bottom Pin type propegates to new connection"), BottomInputPin->PinType.PinCategory == TypePin->PinType.PinCategory);
+			TestTrue(TEXT("Output Pin type propegates to new connection"), OutputPin->PinType.PinCategory == TypePin->PinType.PinCategory);
+
+			// Connect the bottom pin to the type that the first one can be promoted to
+			const bool bConnectedBottom = TypePromoTestUtils::TestPromotedConnection(BottomInputPin, PinToConnectTo);
+			FString BottomStatusMessage = FString::Printf(TEXT("Bottom Pin '%s' Connecting to '%s' "), *K2Schema->TypeToText(BottomInputPin->PinType).ToString(), *K2Schema->TypeToText(PinToConnectTo->PinType).ToString());
+
+			TestTrue(BottomStatusMessage, bConnectedBottom);
+			
+			// The top should be same type, and the output type should have been updated to the be new higher type
+			TestTrue(TEXT("Top Pin type propegates to new connection"), TopInputPin->PinType.PinCategory == TypePin->PinType.PinCategory);
+			TestTrue(TEXT("Output Pin type propegates to new connection"), OutputPin->PinType.PinCategory == PinToConnectTo->PinType.PinCategory);
+		}
+	}
+
+	// Cleanup
+	{
+		TypePromoTestUtils::CleanupTestPins(PinTypes);
+
+		BP_Primative_Connections->MarkPendingKill();
+		TestGraph->MarkPendingKill();
+		TestNode->MarkPendingKill();
+	}
+
+	return true;
+}
+
 #undef MakeTestPin
 
 #endif //WITH_DEV_AUTOMATION_TESTS
\ No newline at end of file
