Add more descriptive comments to the AddPinInterface and expose CanRemovePin as a public method

#rnx

diff --git a/Engine/Source/Editor/BlueprintGraph/Classes/K2Node_AddPinInterface.h b/Engine/Source/Editor/BlueprintGraph/Classes/K2Node_AddPinInterface.h
--- a/Engine/Source/Editor/BlueprintGraph/Classes/K2Node_AddPinInterface.h
+++ b/Engine/Source/Editor/BlueprintGraph/Classes/K2Node_AddPinInterface.h
@@ -28,19 +28,54 @@ class BLUEPRINTGRAPH_API IK2Node_AddPinInterface
 
 public:
 
+	/**
+	* Determine the maximum number of additional input pins this node can have.
+	* The node is limited by the number of letters in the alphabet for display
+	* purposes
+	* 
+	* @return constexpr int32
+	*/
 	static constexpr int32 GetMaxInputPinsNum()
 	{
 		return (TCHAR('Z') - TCHAR('A'));
 	}
 
+	/**
+	* GetNameForAdditionalPin
+	*
+	* @param PinIndex		The additional pin index to get the name of
+	*
+	* @return FName			Name of this additional pin
+	*/
 	static FName GetNameForAdditionalPin(int32 PinIndex);
 
-	virtual void AddInputPin() { }
+	/**
+	* Add an additional input pin to this node
+	*/
+	virtual void AddInputPin() = 0;
+
+	/**
+	* Determines if a pin can be added to this node. Typically used within   
+	* UK2Node::GetNodeContextMenuActions to determine what options there should be. 
+	*
+	* @return bool	True if a new pin can be added, false if it cannot. 
+	*/
 	virtual bool CanAddPin() const { return true; }
 
+	/**
+	* Attempt to remove the given pin from this node. Should fail if the 
+	* given pin is not an additional pin to the node
+	*
+	* @param Pin		The pin to remove
+	*/
 	virtual void RemoveInputPin(UEdGraphPin* Pin) { } 
 
-protected:
-
+	/**
+	* Returns true if the given pin can be removed from this node.
+	*
+	* @param Pin		The pin to test 
+	*
+	* @return bool		True if the pin can be removed, false otherwise.
+	*/
 	virtual bool CanRemovePin(const UEdGraphPin* Pin) const { return true; };
 };
