Change "Division" to "Divide" in the BP context menu display name

#rnx

diff --git a/Engine/Source/Editor/BlueprintGraph/Private/BlueprintTypePromotion.cpp b/Engine/Source/Editor/BlueprintGraph/Private/BlueprintTypePromotion.cpp
--- a/Engine/Source/Editor/BlueprintGraph/Private/BlueprintTypePromotion.cpp
+++ b/Engine/Source/Editor/BlueprintGraph/Private/BlueprintTypePromotion.cpp
@@ -448,7 +448,7 @@ const FText& FTypePromotion::GetUserFacingOperatorName(const FName Operator)
 		{ OperatorNames::Add,		LOCTEXT("AddDisplayName",			"Add") },
 		{ OperatorNames::Multiply,  LOCTEXT("MultiplyDisplayName",		"Multiply") },
 		{ OperatorNames::Subtract,  LOCTEXT("SubtractDisplayName",		"Subtract") },
-		{ OperatorNames::Divide,	LOCTEXT("DivideDisplayName",		"Division") },
+		{ OperatorNames::Divide,	LOCTEXT("DivideDisplayName",		"Divide") },
 		{ OperatorNames::Greater,	LOCTEXT("GreaterDisplayName",		"Greater") },
 		{ OperatorNames::GreaterEq, LOCTEXT("GreaterEqDisplayName",		"Greater Equal") },
 		{ OperatorNames::Less,		LOCTEXT("LessDisplayName",			"Less") },
