Blueprints: Allow the user to remove the default input pins on a promotable operator node to be consistent with the existing behavior of the previous math nodes

#rnx

diff --git a/Engine/Source/Editor/BlueprintGraph/Private/K2Node_PromotableOperator.cpp b/Engine/Source/Editor/BlueprintGraph/Private/K2Node_PromotableOperator.cpp
--- a/Engine/Source/Editor/BlueprintGraph/Private/K2Node_PromotableOperator.cpp
+++ b/Engine/Source/Editor/BlueprintGraph/Private/K2Node_PromotableOperator.cpp
@@ -466,17 +466,8 @@ bool UK2Node_PromotableOperator::CanAddPin() const
 
 bool UK2Node_PromotableOperator::CanRemovePin(const UEdGraphPin* Pin) const
 {
-	if (!Pin)
-	{
-		return false;
-	}
-
-	// We cannot remove the first two inputs from a function, because they are always there
-	const bool bIsBasePin = (Pin->PinName == InputPinA_Name || Pin->PinName == InputPinB_Name);
-
 	return (
-		!bIsBasePin &&
-		Pin->ParentPin == nullptr &&
+		Pin && Pin->ParentPin == nullptr &&
 		NumAdditionalInputs > 0 &&
 		INDEX_NONE != Pins.IndexOfByKey(Pin) &&
 		Pin->Direction == EEdGraphPinDirection::EGPD_Input
@@ -495,13 +486,13 @@ void UK2Node_PromotableOperator::RemoveInputPin(UEdGraphPin* Pin)
 			--NumAdditionalInputs;
 
 			int32 NameIndex = 0;
-			const UEdGraphPin* OutPin = GetOutputPin();
 			const UEdGraphPin* SelfPin = FindPin(UEdGraphSchema_K2::PN_Self);
 
 			for (int32 PinIndex = 0; PinIndex < Pins.Num(); ++PinIndex)
 			{
 				UEdGraphPin* LocalPin = Pins[PinIndex];
-				if (LocalPin && (LocalPin != OutPin) && (LocalPin != SelfPin))
+
+				if (LocalPin && (LocalPin->Direction != EGPD_Output) && (LocalPin != SelfPin) && LocalPin->ParentPin == nullptr)
 				{
 					const FName PinName = GetNameForAdditionalPin(NameIndex);
 					if (PinName != LocalPin->PinName)
@@ -515,6 +506,13 @@ void UK2Node_PromotableOperator::RemoveInputPin(UEdGraphPin* Pin)
 			FBlueprintEditorUtils::MarkBlueprintAsStructurallyModified(GetBlueprint());
 		}
 	}
+
+	// If the pin that was removed makes this node empty then we can just reset it 
+	// and have no need to reevaluate any pins
+	if (!HasAnyConnectionsOrDefaults())
+	{
+		ResetNodeToWildcard();
+	}
 }
 
 UEdGraphPin* UK2Node_PromotableOperator::GetAdditionalPin(int32 PinIndex) const
@@ -832,6 +830,13 @@ void UK2Node_PromotableOperator::GetPinsToConsider(TArray<UEdGraphPin*>& OutArra
 		// which means that it would still have it's old pin type, and could be inaccurate
 		if (Pin->LinkedTo.Num() > 0)
 		{
+			// If this is from a split pin, then we care about the parent type
+			if (Pin->ParentPin != nullptr)
+			{
+				OutArray.Add(Pin->ParentPin);
+			}
+
+			// as well as all the links to it
 			for (UEdGraphPin* Link : Pin->LinkedTo)
 			{
 				OutArray.Emplace(Link);
@@ -839,8 +844,15 @@ void UK2Node_PromotableOperator::GetPinsToConsider(TArray<UEdGraphPin*>& OutArra
 		}
 		else if (!Pin->DoesDefaultValueMatchAutogenerated() && !FWildcardNodeUtils::IsWildcardPin(Pin))
 		{
-			// If this pin has doesn't have any links then we should account for the pin itself
-			OutArray.Add(Pin);
+			// If this is from a split pin, then we care about the type of the parent, not this pin
+			if (Pin->ParentPin != nullptr)
+			{
+				OutArray.Add(Pin->ParentPin);
+			}
+			else
+			{
+				OutArray.Add(Pin);
+			}
 		}
 	}
 }
