Update type promo tests to have more AddPin interface coverage

#rnx

diff --git a/Engine/Source/Editor/BlueprintGraph/Private/Tests/TypePromotionTests.cpp b/Engine/Source/Editor/BlueprintGraph/Private/Tests/TypePromotionTests.cpp
--- a/Engine/Source/Editor/BlueprintGraph/Private/Tests/TypePromotionTests.cpp
+++ b/Engine/Source/Editor/BlueprintGraph/Private/Tests/TypePromotionTests.cpp
@@ -369,15 +369,8 @@ bool FFindBestMatchingFunc::RunTest(const FString& Parameters)
 		}
 	}
 
-	// Clear our test pins
-	for (UEdGraphPin* TestPin : PinTypes)
-	{
-		if (TestPin)
-		{
-			TestPin->MarkPendingKill();
-		}
-	}
-	PinTypes.Empty();
+	
+	TypePromoTestUtils::CleanupTestPins(PinTypes);
 	TestNode->MarkPendingKill();
 
 	return true;
@@ -447,6 +440,7 @@ bool FPromotableOpDefaultState::RunTest(const FString& Parameters)
 	// Cleanup test BP and graph
 	{
 		WildcardStartTestBP->MarkPendingKill();
+		WildcardStartTestBP->Rename(nullptr, nullptr, REN_DontCreateRedirectors);
 		TestWildcardGraph->MarkPendingKill();
 	}
 
@@ -502,9 +496,51 @@ bool FPromotableOpNodeAddPinInterface::RunTest(const FString& Parameters)
 		const int32 EndingPinCount = MultiplyNode->Pins.Num();
 
 		TestTrue(TEXT("Multiply node had a pin added to it"), EndingPinCount == StartingPinCount + 1);
-		const UEdGraphPin* AdditonalPin = MultiplyNode->GetAdditionalPin(EndingPinCount - 1);
+		const UEdGraphPin* AdditonalPin = MultiplyNode->GetAdditionalPin(EndingPinCount - StartingPinCount);
 		TestNotNull(TEXT("Additional Pin is not null"), AdditonalPin);
 		TestTrue(TEXT("New Pin is wildcard"), FWildcardNodeUtils::IsWildcardPin(AdditonalPin));
+		TestTrue(TEXT("New Pin can be removed"), MultiplyNode->CanRemovePin(AdditonalPin));
+
+		const UEdGraphPin* InputPinA = MultiplyNode->FindPin(FName("A"), EGPD_Input);
+		TestNotNull(TEXT("First input pin is not null"), InputPinA);
+		TestTrue(TEXT("First Pin can be removed"), MultiplyNode->CanRemovePin(InputPinA));
+
+		const UEdGraphPin* InputPinB = MultiplyNode->FindPin(FName("B"), EGPD_Input);
+		TestNotNull(TEXT("Second input pin is not null"), InputPinB);
+		TestTrue(TEXT("Second Pin can be removed"), MultiplyNode->CanRemovePin(InputPinB));
+	}
+
+	// Anything that is not a comparison operator can have a pin added to it
+	{
+		UK2Node_PromotableOperator* AddNode = TypePromoTestUtils::SpawnPromotableNode(TestGraph, TEXT("Add"));
+		TestNotNull(TEXT("Add Node spawn"), AddNode);
+		TestTrue(TEXT("Add can add pin"), AddNode->CanAddPin());
+		
+		const UEdGraphPin* TopInputPin = AddNode->FindPin(TEXT("A"), EGPD_Input);
+		const UEdGraphPin* BottomInputPin = AddNode->FindPin(TEXT("B"), EGPD_Input);
+		const UEdGraphPin* OutputPin = AddNode->GetOutputPin();
+		const int32 StartingPinCount = AddNode->Pins.Num();
+		AddNode->AddInputPin();
+		
+		UEdGraphPin* AdditonalPin = AddNode->GetAdditionalPin(2);
+		TestNotNull(TEXT("Additional Pin is not null"), AdditonalPin);
+		
+		// Connect a float pin to the additional input pin
+		const bool bConnected = TypePromoTestUtils::TestPromotedConnection(AdditonalPin, FloatOutputPin);
+		TestTrue(TEXT("Connection to additional pin success"), bConnected);
+
+		// The other pins have propagated correctly with this new connection 
+		TestTrue(TEXT("Top Pin type propegates to new connection"), TopInputPin->PinType.PinCategory == FloatOutputPin->PinType.PinCategory);
+		TestTrue(TEXT("Bottom Pin type propegates to new connection"), BottomInputPin->PinType.PinCategory == FloatOutputPin->PinType.PinCategory);
+		TestTrue(TEXT("Out Pin type propegates to new connection"), OutputPin->PinType.PinCategory == FloatOutputPin->PinType.PinCategory);
+
+		// Removing the only pin with a connection with reset the node to wildcard
+		AddNode->RemoveInputPin(AdditonalPin);
+		TestTrue(TEXT("Top Pin type propegates to wildcard on connection break"), FWildcardNodeUtils::IsWildcardPin(TopInputPin));
+		TestTrue(TEXT("Bottom Pin type propegates to wildcard on connection break"), FWildcardNodeUtils::IsWildcardPin(BottomInputPin));
+		TestTrue(TEXT("Out Pin type propegates to wildcard on connection break"), FWildcardNodeUtils::IsWildcardPin(OutputPin));
+		
+		TestTrue(TEXT("Additional pin was successfully removed"), StartingPinCount == AddNode->Pins.Num());
 	}
 
 	// Cleanup
@@ -512,6 +548,7 @@ bool FPromotableOpNodeAddPinInterface::RunTest(const FString& Parameters)
 		TypePromoTestUtils::CleanupTestPins(PinTypes);
 
 		TestBP->MarkPendingKill();
+		TestBP->Rename(nullptr, nullptr, REN_DontCreateRedirectors);
 		TestGraph->MarkPendingKill();
 		TestNode->MarkPendingKill();
 	}
@@ -603,6 +640,7 @@ bool FPromotableOperatorConnectionChanged::RunTest(const FString& Parameters)
 		TypePromoTestUtils::CleanupTestPins(PinTypes);
 
 		BP_ConnectionChanged->MarkPendingKill();
+		BP_ConnectionChanged->Rename(nullptr, nullptr, REN_DontCreateRedirectors);
 		TestGraph->MarkPendingKill();
 		TestNode->MarkPendingKill();
 	}
@@ -689,6 +727,7 @@ bool FPromotableOperatorPrimitivePromotions::RunTest(const FString& Parameters)
 		TypePromoTestUtils::CleanupTestPins(PinTypes);
 
 		BP_Primative_Connections->MarkPendingKill();
+		BP_Primative_Connections->Rename(nullptr, nullptr, REN_DontCreateRedirectors);
 		TestGraph->MarkPendingKill();
 		TestNode->MarkPendingKill();
 	}
