Make GetPinsToConsider a const function

#rnx

diff --git a/Engine/Source/Editor/BlueprintGraph/Classes/K2Node_PromotableOperator.h b/Engine/Source/Editor/BlueprintGraph/Classes/K2Node_PromotableOperator.h
--- a/Engine/Source/Editor/BlueprintGraph/Classes/K2Node_PromotableOperator.h
+++ b/Engine/Source/Editor/BlueprintGraph/Classes/K2Node_PromotableOperator.h
@@ -128,7 +128,7 @@ private:
 	*
 	* @param OutArray	Array to populate with the considered pins
 	*/
-	void GetPinsToConsider(TArray<UEdGraphPin*>& OutArray);
+	void GetPinsToConsider(TArray<UEdGraphPin*>& OutArray) const;
 
 	/** The name that this operation uses ("Add", "Multiply", etc) */
 	FName OperationName;
diff --git a/Engine/Source/Editor/BlueprintGraph/Private/K2Node_PromotableOperator.cpp b/Engine/Source/Editor/BlueprintGraph/Private/K2Node_PromotableOperator.cpp
--- a/Engine/Source/Editor/BlueprintGraph/Private/K2Node_PromotableOperator.cpp
+++ b/Engine/Source/Editor/BlueprintGraph/Private/K2Node_PromotableOperator.cpp
@@ -815,7 +815,7 @@ TArray<UEdGraphPin*> UK2Node_PromotableOperator::GetInputPins(bool bIncludeLinks
 	return InputPins;
 }
 
-void UK2Node_PromotableOperator::GetPinsToConsider(TArray<UEdGraphPin*>& OutArray)
+void UK2Node_PromotableOperator::GetPinsToConsider(TArray<UEdGraphPin*>& OutArray) const
 {
 	for (UEdGraphPin* Pin : Pins)
 	{
