Reset operator nodes to wildcard after paste if there are no default values on it.

#rnx

diff --git a/Engine/Source/Editor/BlueprintGraph/Classes/K2Node_PromotableOperator.h b/Engine/Source/Editor/BlueprintGraph/Classes/K2Node_PromotableOperator.h
--- a/Engine/Source/Editor/BlueprintGraph/Classes/K2Node_PromotableOperator.h
+++ b/Engine/Source/Editor/BlueprintGraph/Classes/K2Node_PromotableOperator.h
@@ -23,6 +23,7 @@ public:
 	virtual FText GetTooltipText() const override;
 	virtual void PinDefaultValueChanged(UEdGraphPin* Pin) override;
 	virtual void NodeConnectionListChanged() override;
+	virtual void PostPasteNode() override;
 	// End of UEdGraphNode interface
 
 	// UK2Node interface
diff --git a/Engine/Source/Editor/BlueprintGraph/Private/K2Node_PromotableOperator.cpp b/Engine/Source/Editor/BlueprintGraph/Private/K2Node_PromotableOperator.cpp
--- a/Engine/Source/Editor/BlueprintGraph/Private/K2Node_PromotableOperator.cpp
+++ b/Engine/Source/Editor/BlueprintGraph/Private/K2Node_PromotableOperator.cpp
@@ -179,6 +179,18 @@ void UK2Node_PromotableOperator::NodeConnectionListChanged()
 	}
 }
 
+void UK2Node_PromotableOperator::PostPasteNode()
+{
+	Super::PostPasteNode();
+
+	// If we have copied a node with additional pins then we need to make sure 
+	// they get reset to wildcard as well, otherwise their type will persist
+	if (!HasAnyConnectionsOrDefaults())
+	{
+		ResetNodeToWildcard();
+	}
+}
+
 ///////////////////////////////////////////////////////////
 // UK2Node interface
 
