Add another type promotion unit test to check that each primitive type has a correct matching function

#rnx

diff --git a/Engine/Source/Editor/BlueprintGraph/Private/Tests/TypePromotionTests.cpp b/Engine/Source/Editor/BlueprintGraph/Private/Tests/TypePromotionTests.cpp
--- a/Engine/Source/Editor/BlueprintGraph/Private/Tests/TypePromotionTests.cpp
+++ b/Engine/Source/Editor/BlueprintGraph/Private/Tests/TypePromotionTests.cpp
@@ -376,6 +376,62 @@ bool FFindBestMatchingFunc::RunTest(const FString& Parameters)
 	return true;
 }
 
+// Test that every type in the promotion table has a best matching function for each operator
+IMPLEMENT_SIMPLE_AUTOMATION_TEST(FPromotableTypeToOperator, "Blueprints.Compiler.TypeToOperator", EAutomationTestFlags::EditorContext | EAutomationTestFlags::EngineFilter)
+bool FPromotableTypeToOperator::RunTest(const FString& Parameters)
+{
+	if (!TypePromoDebug::IsTypePromoEnabled())
+	{
+		return true;
+	}
+
+	// Refresh the actions within this test in case the editor is open but hasn't loaded BlueprintGraph yet
+	FTypePromotion::ClearNodeSpawners();
+	FBlueprintActionDatabase::Get().RefreshAll();
+
+	UEdGraphNode* TestNode = NewObject<UEdGraphNode>();
+
+	const TSet<FName> AllOpNames = FTypePromotion::GetAllOpNames();
+
+	// Create test pins!
+	TArray<UEdGraphPin*> PinTypes = {};
+	MakeTestPins(TestNode, PinTypes);
+
+	const UEdGraphSchema_K2* K2Schema = GetDefault<UEdGraphSchema_K2>();
+
+	const TMap<FName, TArray<FName>>* const PromoTable = FTypePromotion::GetPrimativePromotionTable();
+	TestNotNull(TEXT("Primative Promotion table exists"), PromoTable);
+
+	for (const FName OpName : AllOpNames)
+	{
+		// Ensure that there is a best matching function for each type in the promo table
+		for (const TPair<FName, TArray<FName>>& Pair : *PromoTable)
+		{
+			const FName TypeName = Pair.Key;
+			if (TypeName == UEdGraphSchema_K2::PC_Wildcard)
+			{
+				continue;
+			}
+
+			UEdGraphPin* TypePin = *PinTypes.FindByPredicate([&TypeName](const UEdGraphPin* Pin) -> bool { return Pin && Pin->PinType.PinCategory == TypeName && Pin->Direction == EGPD_Output; });
+			TestNotNull(TEXT("Testable Pin Type Found"), TypePin);
+
+			const UFunction* BestMatchFunc = FTypePromotion::FindBestMatchingFunc(OpName, { TypePin });
+			FString BottomStatusMessage = FString::Printf(TEXT("'%s' Operator has a match with '%s' pin type."), *TypeName.ToString(), *K2Schema->TypeToText(TypePin->PinType).ToString());
+
+			TestNotNull(BottomStatusMessage, BestMatchFunc);
+		}
+	}
+
+	// Cleanup test
+	{
+		TypePromoTestUtils::CleanupTestPins(PinTypes);
+		TestNode->MarkPendingKill();
+	}
+
+	return true;
+}
+
 // Test the default state of all operator nodes to ensure they are correct.
 // Comparison operators (Greater Than, Less Than, etc) should have two 
 // wildcard inputs and one boolean output. All others should be all wildcards.
