Add an accessor for finding the tolerance pin on promotable operator nodes

#rnx

diff --git a/Engine/Source/Editor/BlueprintGraph/Classes/K2Node_PromotableOperator.h b/Engine/Source/Editor/BlueprintGraph/Classes/K2Node_PromotableOperator.h
--- a/Engine/Source/Editor/BlueprintGraph/Classes/K2Node_PromotableOperator.h
+++ b/Engine/Source/Editor/BlueprintGraph/Classes/K2Node_PromotableOperator.h
@@ -45,6 +45,9 @@ public:
 	/** Gets the additional pin that was created at this index */
 	UEdGraphPin* GetAdditionalPin(int32 PinIndex) const;
 
+	/** Attempts to find the error tolerance pin on this node. Can return nullptr. */
+	UEdGraphPin* FindTolerancePin() const;
+
 	// UK2Node_CallFunction interface
 	virtual void SetFromFunction(const UFunction* Function);
 	// End of UK2Node_CallFunction interface
diff --git a/Engine/Source/Editor/BlueprintGraph/Private/K2Node_PromotableOperator.cpp b/Engine/Source/Editor/BlueprintGraph/Private/K2Node_PromotableOperator.cpp
--- a/Engine/Source/Editor/BlueprintGraph/Private/K2Node_PromotableOperator.cpp
+++ b/Engine/Source/Editor/BlueprintGraph/Private/K2Node_PromotableOperator.cpp
@@ -542,6 +542,11 @@ UEdGraphPin* UK2Node_PromotableOperator::GetAdditionalPin(int32 PinIndex) const
 	return nullptr;
 }
 
+UEdGraphPin* UK2Node_PromotableOperator::FindTolerancePin() const
+{
+	return FindPin(TolerancePin_Name, EGPD_Input);
+}
+
 ///////////////////////////////////////////////////////////
 // UK2Node_CallFunction interface
 void UK2Node_PromotableOperator::SetFromFunction(const UFunction* Function)
