Blueprints: Add tolerance pins to promotable comparison operators

#rnx

diff --git a/Engine/Source/Editor/BlueprintGraph/Private/K2Node_PromotableOperator.cpp b/Engine/Source/Editor/BlueprintGraph/Private/K2Node_PromotableOperator.cpp
--- a/Engine/Source/Editor/BlueprintGraph/Private/K2Node_PromotableOperator.cpp
+++ b/Engine/Source/Editor/BlueprintGraph/Private/K2Node_PromotableOperator.cpp
@@ -20,7 +20,7 @@
 
 static const FName InputPinA_Name = FName(TEXT("A"));
 static const FName InputPinB_Name = FName(TEXT("B"));
-static const FName TolerancePin_Name = FName(TEXT("Tolerance"));
+static const FName TolerancePin_Name = FName(TEXT("ErrorTolerance"));
 static const int32 NumFunctionInputs = 2;
 
 ///////////////////////////////////////////////////////////
@@ -37,6 +37,37 @@ UK2Node_PromotableOperator::UK2Node_PromotableOperator(const FObjectInitializer&
 ///////////////////////////////////////////////////////////
 // UEdGraphNode interface
 
+namespace PromotableOpUtils
+{
+	bool FindTolerancePinType(const UFunction* Func, FEdGraphPinType& OutPinType)
+	{
+		const UEdGraphSchema_K2* Schema = GetDefault<UEdGraphSchema_K2>();
+		// Check if there is a tolerance field that we need to add
+		// If UFunction has a third input param 
+		int32 InputArgsFound = 0;
+		for (TFieldIterator<FProperty> PropIt(Func); PropIt && (PropIt->PropertyFlags & CPF_Parm); ++PropIt)
+		{
+			const FProperty* Param = *PropIt;
+
+			// We don't care about the return property, and we are looking for
+			// any additional input param that isn't in the normal function input range
+			if (Param->HasAnyPropertyFlags(CPF_ReturnParm) || (InputArgsFound++ < NumFunctionInputs))
+			{
+				continue;
+			}
+
+			// Found the tolerance type!
+			FEdGraphPinType ParamType;
+			if (Schema->ConvertPropertyToPinType(Param, /* out */ ParamType))
+			{
+				OutPinType = ParamType;
+				return true;
+			}
+		}
+		return false;
+	}
+}
+
 void UK2Node_PromotableOperator::AllocateDefaultPins()
 {
 	FWildcardNodeUtils::CreateWildcardPin(this, InputPinA_Name, EGPD_Input);
@@ -44,10 +75,18 @@ void UK2Node_PromotableOperator::AllocateDefaultPins()
 
 	UEdGraphPin* OutPin = FWildcardNodeUtils::CreateWildcardPin(this, UEdGraphSchema_K2::PN_ReturnValue, EGPD_Output);
 
+	const UFunction* Func = GetTargetFunction();
 	// For comparison functions we always want a bool output, so make it visually so
-	if (FTypePromotion::IsComparisonFunc(GetTargetFunction()))
+	if (Func && FTypePromotion::IsComparisonFunc(Func))
 	{
 		OutPin->PinType.PinCategory = UEdGraphSchema_K2::PC_Boolean;
+
+		// If we need a tolerance pin then we need to add it here, but not if we are in a wildcard state
+		FEdGraphPinType ToleranceType;
+		if (PromotableOpUtils::FindTolerancePinType(Func, ToleranceType))
+		{
+			UEdGraphPin* TolerancePin = CreatePin(EGPD_Input, ToleranceType, TolerancePin_Name);
+		}
 	}
 
 	// Create any additional input pin. Their appropriate type is determined in ReallocatePinsDuringReconstruction
@@ -131,10 +170,7 @@ FText UK2Node_PromotableOperator::GetTooltipText() const
 	// If there are no connections then just display the op name
 	if (!HasAnyConnectionsOrDefaults())
 	{
-		UFunction* Function = GetTargetFunction();
-
-		FName OpName = FTypePromotion::GetOpNameFromFunction(Function);
-		return FText::Format(LOCTEXT("PromotableOperatorFunctionTooltip", "{0} Operator"), FText::FromName(OpName));
+		return FTypePromotion::GetUserFacingOperatorName(OperationName);
 	}
 
 	// Otherwise use the default one (a more specific function tooltip)
@@ -227,16 +263,18 @@ void UK2Node_PromotableOperator::ExpandNode(FKismetCompilerContext& CompilerCont
 		UEdGraphPin* InputB = nullptr;
 		UEdGraphPin* OutputPin = nullptr;
 		UEdGraphPin* SelfPin = nullptr;
+		UEdGraphPin* TolerancePin = nullptr;
 
 		explicit FIntermediateCastPinHelper(UK2Node_CallFunction* NewOperator)
 		{
 			check(NewOperator);
 			SelfPin = NewOperator->FindPin(UEdGraphSchema_K2::PN_Self);
+			TolerancePin = NewOperator->FindPin(TolerancePin_Name, EGPD_Input);
 
 			// Find inputs and outputs
 			for (UEdGraphPin* Pin : NewOperator->Pins)
 			{
-				if (Pin == SelfPin)
+				if (Pin == SelfPin || Pin == TolerancePin)
 				{
 					continue;
 				}
@@ -264,6 +302,7 @@ void UK2Node_PromotableOperator::ExpandNode(FKismetCompilerContext& CompilerCont
 
 	UK2Node_CallFunction* PrevIntermediateNode = nullptr;
 	UEdGraphPin* PrevOutputPin = nullptr;
+	UEdGraphPin* MyTolerancePin = FindTolerancePin();
 
 	// Create cast from original 2 inputs to the first intermediate node
 	{
@@ -288,8 +327,10 @@ void UK2Node_PromotableOperator::ExpandNode(FKismetCompilerContext& CompilerCont
 
 		const bool bPinASuccess = UK2Node_PromotableOperator::CreateIntermediateCast(this, CompilerContext, SourceGraph, OriginalInputPins[0], NewOpHelper.InputA);
 		const bool bPinBSuccess = UK2Node_PromotableOperator::CreateIntermediateCast(this, CompilerContext, SourceGraph, OriginalInputPins[1], NewOpHelper.InputB);
+		// Attempt to connect the tolerance pin if both nodes have one
+		const bool bToleranceSuccess = MyTolerancePin && NewOpHelper.TolerancePin ? UK2Node_PromotableOperator::CreateIntermediateCast(this, CompilerContext, SourceGraph, MyTolerancePin, NewOpHelper.TolerancePin) : true;
 
-		if (!bPinASuccess || !bPinBSuccess)
+		if (!bPinASuccess || !bPinBSuccess || !bToleranceSuccess)
 		{
 			CompilerContext.MessageLog.Error(TEXT("'@@' could not successfuly expand pins!"), PrevIntermediateNode);
 		}
@@ -326,8 +367,11 @@ void UK2Node_PromotableOperator::ExpandNode(FKismetCompilerContext& CompilerCont
 
 		// Connect the original node's pin to the newly created intermediate node's B Pin
 		const bool bPinBSuccess = CreateIntermediateCast(this, CompilerContext, SourceGraph, OriginalInputPins[i], NewOpHelper.InputB);
+		
+		// Make a connection to a tolerance pin if both nodes have one
+		const bool bToleranceSuccess = MyTolerancePin && NewOpHelper.TolerancePin ? UK2Node_PromotableOperator::CreateIntermediateCast(this, CompilerContext, SourceGraph, MyTolerancePin, NewOpHelper.TolerancePin) : true;
 
-		if (!bPinASuccess || !bPinBSuccess)
+		if (!bPinASuccess || !bPinBSuccess || !bToleranceSuccess)
 		{
 			CompilerContext.MessageLog.Error(TEXT("'@@' could not successfuly expand additional pins!"), PrevIntermediateNode);
 		}
@@ -358,6 +402,12 @@ void UK2Node_PromotableOperator::PostReconstructNode()
 	// We only need to set the function if we have connections, otherwise we should stick in a wildcard state
 	if (HasAnyConnectionsOrDefaults())
 	{
+		if (UEdGraphPin* TolerancePin = FindTolerancePin())
+		{
+			FEdGraphPinType ToleranceType;
+			TolerancePin->bHidden = !PromotableOpUtils::FindTolerancePinType(GetTargetFunction(), ToleranceType);
+		}
+
 		// Allocate default pins will have been called before this, which means we are reset to wildcard state
 		// We need to Update the pins to be the proper function again
 		UpdatePinsFromFunction(GetTargetFunction());
@@ -371,6 +421,10 @@ void UK2Node_PromotableOperator::PostReconstructNode()
 			}
 		}
 	}
+	else if (UEdGraphPin* TolerancePin = FindTolerancePin())
+	{
+		TolerancePin->bHidden = true;
+	}
 }
 
 bool UK2Node_PromotableOperator::IsConnectionDisallowed(const UEdGraphPin* MyPin, const UEdGraphPin* OtherPin, FString& OutReason) const
@@ -587,7 +641,7 @@ UEdGraphPin* UK2Node_PromotableOperator::AddInputPinImpl(int32 PinIndex)
 bool UK2Node_PromotableOperator::IsAdditionalPin(const UEdGraphPin* Pin) const
 {
 	// Quickly check if this input pin is one of the two default input pins
-	return Pin && Pin->Direction == EGPD_Input && Pin->PinName != InputPinA_Name && Pin->PinName != InputPinB_Name;
+	return Pin && Pin->Direction == EGPD_Input && Pin->PinName != InputPinA_Name && Pin->PinName != InputPinB_Name && !IsTolerancePin(Pin);
 }
 
 bool UK2Node_PromotableOperator::HasAnyConnectionsOrDefaults() const
@@ -624,8 +678,19 @@ void UK2Node_PromotableOperator::EvaluatePinsFromChange(UEdGraphPin* ChangedPin)
 	{
 		return;
 	}
+	// Changing the tolerance pin doesn't effect the rest of the function signature, so don't attempt to update the func
+	else if (IsTolerancePin(ChangedPin))
+	{
+		return;
+	}
+	// If we have connected the output of our comparison operator (which is always a bool) then 
+	// there is no need to update the other pins
+	else if (FTypePromotion::IsComparisonOpName(OperationName) && ChangedPin == GetOutputPin())
+	{
+		return;
+	}
 
-	// Gather all pins and their links so we can determine the highest type that the user could want
+	// Gather all pins and their links so we can determine the highest type
 	TArray<UEdGraphPin*> PinsToConsider;
 	GetPinsToConsider(PinsToConsider);
 
@@ -769,15 +834,21 @@ void UK2Node_PromotableOperator::ResetNodeToWildcard()
 		}
 	}
 
-	// Set output pins to have a bool output flag by default
-	{
-		UEdGraphPin* OutPin = GetOutputPin();
-		const UFunction* Func = GetTargetFunction();
+	const UFunction* Func = GetTargetFunction();
 
-		if (OutPin && Func && FTypePromotion::IsComparisonFunc(Func))
+	if(Func && FTypePromotion::IsComparisonFunc(Func))
+	{
+		// Set output pins to have a bool output flag by default
+		if (UEdGraphPin* OutPin = GetOutputPin())
 		{
 			OutPin->PinType.PinCategory = UEdGraphSchema_K2::PC_Boolean;
 		}
+
+		// If we have a tolerance pin, and we reset to wildcard then we should hide it
+		if (UEdGraphPin* TolerancePin = FindTolerancePin())
+		{
+			TolerancePin->bHidden = true;
+		}
 	}
 
 	// Clear out any possible function matches, since we are removing connections
@@ -805,10 +876,11 @@ void UK2Node_PromotableOperator::UpdateFromBestMatchingFunction()
 	TArray<UEdGraphPin*> PinsToConsider;
 	GetPinsToConsider(PinsToConsider);
 
-	const UFunction* BestMatchingFunc = FTypePromotion::FindBestMatchingFunc(OperationName, PinsToConsider);
-
 	// We need to update the pins from our function if have a new connection
-	UpdatePinsFromFunction(BestMatchingFunc);
+	if (const UFunction* BestMatchingFunc = FTypePromotion::FindBestMatchingFunc(OperationName, PinsToConsider))
+	{
+		UpdatePinsFromFunction(BestMatchingFunc);
+	}
 }
 
 TArray<UEdGraphPin*> UK2Node_PromotableOperator::GetInputPins(bool bIncludeLinks /** = false */) const
@@ -834,6 +906,8 @@ TArray<UEdGraphPin*> UK2Node_PromotableOperator::GetInputPins(bool bIncludeLinks
 
 void UK2Node_PromotableOperator::GetPinsToConsider(TArray<UEdGraphPin*>& OutArray) const
 {
+	const bool bIsComparisonOp = FTypePromotion::IsComparisonOpName(OperationName);
+
 	for (UEdGraphPin* Pin : Pins)
 	{
 		// Tolerance pins don't factor into what types we should be matching 
@@ -843,6 +917,13 @@ void UK2Node_PromotableOperator::GetPinsToConsider(TArray<UEdGraphPin*>& OutArra
 			continue;
 		}
 
+		// If this is a comparison operator then we don't need to consider the boolean output
+		// pin because every comparison function will have a boolean output!
+		if (bIsComparisonOp && Pin->Direction == EGPD_Output)
+		{
+			continue;
+		}
+
 		// If this pin has links, then use those instead of the actual pin because we could be process of changing it
 		// which means that it would still have it's old pin type, and could be inaccurate
 		if (Pin->LinkedTo.Num() > 0)
@@ -902,6 +983,7 @@ void UK2Node_PromotableOperator::UpdatePinsFromFunction(const UFunction* Functio
 	// Gather the pin types of the properties on the function we want to convert to
 	FEdGraphPinType FunctionReturnType;
 	FEdGraphPinType HighestFuncInputType;
+	FEdGraphPinType ToleranceType;
 	TArray<FEdGraphPinType> FunctionInputTypes;
 	{
 		for (TFieldIterator<FProperty> PropIt(Function); PropIt && (PropIt->PropertyFlags & CPF_Parm); ++PropIt)
@@ -928,7 +1010,7 @@ void UK2Node_PromotableOperator::UpdatePinsFromFunction(const UFunction* Functio
 		}
 	}
 
-	auto ConformPinLamda = [&ChangedPin](const FEdGraphPinType& FunctionPinType, UEdGraphPin* NodePin)
+	auto ConformPinLambda = [&ChangedPin](const FEdGraphPinType& FunctionPinType, UEdGraphPin* NodePin)
 	{
 		// If the pin types are already equal, then we don't have to do any work
 		// If this is linked to wildcard pins, then we can just ignore it and handle it on expansion
@@ -961,6 +1043,21 @@ void UK2Node_PromotableOperator::UpdatePinsFromFunction(const UFunction* Functio
 		NodePin->PinType = ConformingType;
 	};
 
+	// Check if we need to add a tolerance pin or not with this new function
+	if (FTypePromotion::IsComparisonFunc(Function) && PromotableOpUtils::FindTolerancePinType(Function, ToleranceType))
+	{
+		// Set the tolerance pin to visible if it is currently on the node
+		if (UEdGraphPin* ExistingTolPin = FindTolerancePin())
+		{
+			ExistingTolPin->bHidden = false;
+		}
+		// Otherwise we need to create a tolerance pin
+		else
+		{
+			CreatePin(EGPD_Input, ToleranceType, TolerancePin_Name);
+		}
+	}
+
 	int32 CurPinIndex = 0;
 	for (UEdGraphPin* CurPin : Pins)
 	{
@@ -973,17 +1070,23 @@ void UK2Node_PromotableOperator::UpdatePinsFromFunction(const UFunction* Functio
 		if (IsAdditionalPin(CurPin))
 		{
 			// Conform to the highest input pin on the function
-			ConformPinLamda(HighestFuncInputType, CurPin);
+			ConformPinLambda(HighestFuncInputType, CurPin);
 		}
 		else if (CurPin->Direction == EGPD_Output)
 		{
 			// Match to the output pin
-			ConformPinLamda(FunctionReturnType, CurPin);
+			ConformPinLambda(FunctionReturnType, CurPin);
+		}
+		// Creation and conformation of the tolerance pin is handled before all others to 
+		// ensure that connections are not broken accidentally
+		else if (IsTolerancePin(CurPin))
+		{
+			ConformPinLambda(ToleranceType, CurPin);
 		}
 		else
 		{
 			// Match to the appropriate function input type
-			ConformPinLamda(FunctionInputTypes[CurPinIndex], CurPin);
+			ConformPinLambda(FunctionInputTypes[CurPinIndex], CurPin);
 			++CurPinIndex;
 		}
 	}
