Add type promotion tests for having the correct tolerance pin

#rnx

diff --git a/Engine/Source/Editor/BlueprintGraph/Private/Tests/TypePromotionTests.cpp b/Engine/Source/Editor/BlueprintGraph/Private/Tests/TypePromotionTests.cpp
--- a/Engine/Source/Editor/BlueprintGraph/Private/Tests/TypePromotionTests.cpp
+++ b/Engine/Source/Editor/BlueprintGraph/Private/Tests/TypePromotionTests.cpp
@@ -134,7 +134,13 @@ namespace TypePromoTestUtils
 	MakeTestPin(OwningNode, OutArray, VecOutputPinA, UEdGraphSchema_K2::PC_Struct, EGPD_Output);		\
 	VecOutputPinA->PinType.PinSubCategoryObject = TBaseStructure<FVector>::Get();						\
 	MakeTestPin(OwningNode, OutArray, Vec2DOutputPinA, UEdGraphSchema_K2::PC_Struct, EGPD_Output);		\
-	Vec2DOutputPinA->PinType.PinSubCategoryObject = TBaseStructure<FVector2D>::Get();
+	Vec2DOutputPinA->PinType.PinSubCategoryObject = TBaseStructure<FVector2D>::Get();					\
+	MakeTestPin(OwningNode, OutArray, Vec4OutPin, UEdGraphSchema_K2::PC_Struct, EGPD_Output);			\
+	Vec4OutPin->PinType.PinSubCategoryObject = TBaseStructure<FVector4>::Get();							\
+	MakeTestPin(OwningNode, OutArray, RotOutPin, UEdGraphSchema_K2::PC_Struct, EGPD_Output);			\
+	RotOutPin->PinType.PinSubCategoryObject = TBaseStructure<FRotator>::Get();							\
+	MakeTestPin(OwningNode, OutArray, QuatOutPin, UEdGraphSchema_K2::PC_Struct, EGPD_Output);			\
+	QuatOutPin->PinType.PinSubCategoryObject = TBaseStructure<FQuat>::Get();
 
 IMPLEMENT_SIMPLE_AUTOMATION_TEST(FTypePromotionTest, "Blueprints.Compiler.TypePromotion", EAutomationTestFlags::EditorContext | EAutomationTestFlags::SmokeFilter)
 bool FTypePromotionTest::RunTest(const FString& Parameters)
@@ -503,6 +509,101 @@ bool FPromotableOpDefaultState::RunTest(const FString& Parameters)
 	return true;
 }
 
+IMPLEMENT_SIMPLE_AUTOMATION_TEST(FPromotableOpTolerancePin, "Blueprints.Nodes.PromotableOp.TolerancePin", EAutomationTestFlags::EditorContext | EAutomationTestFlags::EngineFilter)
+bool FPromotableOpTolerancePin::RunTest(const FString& Parameters)
+{
+	if (!TypePromoDebug::IsTypePromoEnabled())
+	{
+		return true;
+	}
+
+	// Refresh the actions within this test in case the editor is open but hasn't loaded BlueprintGraph yet
+	FTypePromotion::ClearNodeSpawners();
+	FBlueprintActionDatabase::Get().RefreshAll();
+
+	MakeTestableBP(ToleranceTestBP, ToleranceGraph);
+	MakeTestableNode(TestNode, ToleranceGraph)
+	// Create test pins!
+	TArray<UEdGraphPin*> PinTypes = {};
+	MakeTestPins(TestNode, PinTypes);
+
+	// These functions should be spawned with a tolerance float pin
+	// EqualEqual_Vector2DVector2D
+	// NotEqual_Vector2DVector2D
+	// EqualEqual_VectorVector
+	// NotEqual_VectorVector
+	// EqualEqual_Vector4Vector4
+	// NotEqual_Vector4Vector4
+	// EqualEqual_RotatorRotator
+	// NotEqual_RotatorRotator
+	// NotEqual_QuatQuat
+	// EqualEqual_QuatQuat
+	const UEdGraphSchema_K2* K2Schema = GetDefault<UEdGraphSchema_K2>();
+	// Make TMap of the op -> tpe
+	TArray<TPair<FName, UEdGraphPin*>> ToleranceFunctions;
+	ToleranceFunctions.Add(TPair<FName, UEdGraphPin*>(TEXT("EqualEqual"), VecOutputPinA));
+	ToleranceFunctions.Add(TPair<FName, UEdGraphPin*>(TEXT("NotEqual"), VecOutputPinA));
+	ToleranceFunctions.Add(TPair<FName, UEdGraphPin*>(TEXT("EqualEqual"), Vec2DOutputPinA));
+	ToleranceFunctions.Add(TPair<FName, UEdGraphPin*>(TEXT("NotEqual"), Vec2DOutputPinA));
+	ToleranceFunctions.Add(TPair<FName, UEdGraphPin*>(TEXT("EqualEqual"), Vec4OutPin));
+	ToleranceFunctions.Add(TPair<FName, UEdGraphPin*>(TEXT("NotEqual"), Vec4OutPin));
+	ToleranceFunctions.Add(TPair<FName, UEdGraphPin*>(TEXT("EqualEqual"), RotOutPin));
+	ToleranceFunctions.Add(TPair<FName, UEdGraphPin*>(TEXT("NotEqual"), RotOutPin));
+	ToleranceFunctions.Add(TPair<FName, UEdGraphPin*>(TEXT("EqualEqual"), QuatOutPin));
+	ToleranceFunctions.Add(TPair<FName, UEdGraphPin*>(TEXT("NotEqual"), QuatOutPin));
+
+	// iterate the map
+	// for each value of the map
+	// Spawn the operator node of the op type
+	// Connect the pin type
+	// Test that there is an error tolerance pin
+	// test that the pin type is a float
+	// When disconnecting the Pin the tolerance pin should be removed, unless a default value is there
+	// when copy/pasting the node the tolerance pin should not be kept unless it has a value
+	TArray<UEdGraphPin*> PinsToConsider;
+	for(const TPair<FName, UEdGraphPin*>& Pair  : ToleranceFunctions)
+	{
+		const FName OpName = Pair.Key;
+		UEdGraphPin* TestingPin = Pair.Value;
+
+		UK2Node_PromotableOperator* OpNode = TypePromoTestUtils::SpawnPromotableNode(ToleranceGraph, OpName);
+		TestNotNull(FString::Printf(TEXT("Spawning a '%s' operator node"), *OpName.ToString()), OpNode);
+
+		check(OpNode);
+
+		// Test pin types
+		UEdGraphPin* TopInputPin = OpNode->FindPin(TEXT("A"), EGPD_Input);
+		UEdGraphPin* BottomInputPin = OpNode->FindPin(TEXT("B"), EGPD_Input);
+		UEdGraphPin* OutputPin = OpNode->GetOutputPin();
+		check(TopInputPin && BottomInputPin && OutputPin);
+		UEdGraphPin* TolerancePin = OpNode->FindTolerancePin();
+		TestNull(TEXT("No tolerance by default"), TolerancePin);
+
+		const int32 StartingPinCount = OpNode->Pins.Num();
+
+		const bool bConnected = TypePromoTestUtils::TestPromotedConnection(TopInputPin, TestingPin);
+		TestTrue(TEXT("Connection to additional pin success"), bConnected);
+		const int32 EndingPinCount = OpNode->Pins.Num();
+		TolerancePin = OpNode->FindTolerancePin();
+		const FString Message = FString::Printf(TEXT("'%s' operator node connecting to '%s' pin has a tolerance pin"), *OpName.ToString(), *K2Schema->TypeToText(TestingPin->PinType).ToString());
+		TestNotNull(Message, TolerancePin);
+
+		TestTrue(TEXT("Added a new pin"), EndingPinCount == StartingPinCount + 1);
+	}
+
+	// Cleanup test BP and graph
+	{
+		TypePromoTestUtils::CleanupTestPins(PinTypes);
+
+		ToleranceTestBP->MarkPendingKill();
+		ToleranceTestBP->Rename(nullptr, nullptr, REN_DontCreateRedirectors);
+		TestNode->MarkPendingKill();
+		ToleranceGraph->MarkPendingKill();
+	}
+
+	return true;
+}
+
 // Test that promotable operator nodes can correctly have pins added to them
 // and that comparison operators cannot have pins added to them.
 IMPLEMENT_SIMPLE_AUTOMATION_TEST(FPromotableOpNodeAddPinInterface, "Blueprints.Nodes.PromotableOp.AddPinInterface", EAutomationTestFlags::EditorContext | EAutomationTestFlags::EngineFilter)
