Clean up logging in type promotion and ensure that the operator name is up to date on promotable operator nodes. Remove unused local variable

#rnx

diff --git a/Engine/Source/Editor/BlueprintGraph/Private/K2Node_PromotableOperator.cpp b/Engine/Source/Editor/BlueprintGraph/Private/K2Node_PromotableOperator.cpp
--- a/Engine/Source/Editor/BlueprintGraph/Private/K2Node_PromotableOperator.cpp
+++ b/Engine/Source/Editor/BlueprintGraph/Private/K2Node_PromotableOperator.cpp
@@ -89,6 +89,11 @@ void UK2Node_PromotableOperator::AllocateDefaultPins()
 		}
 	}
 
+	// Update the op name so that if there is a blank wildcard node left on the graph we
+	// can ensure that it is correct
+	UpdateOpName();
+	ensureMsgf((OperationName != TEXT("NO_OP")), TEXT("Invalid operation name on Promotable Operator node!"));
+
 	// Create any additional input pin. Their appropriate type is determined in ReallocatePinsDuringReconstruction
 	// because we cannot get a promoted type with no links to the pin.
 	for (int32 i = NumFunctionInputs; i < (NumAdditionalInputs + NumFunctionInputs); ++i)
@@ -443,6 +448,11 @@ bool UK2Node_PromotableOperator::IsConnectionDisallowed(const UEdGraphPin* MyPin
 	{
 		return Super::IsConnectionDisallowed(MyPin, OtherPin, OutReason);
 	}
+	// If the pins are the same type then there is no reason to check for a promotion
+	else if (MyPin->PinType == OtherPin->PinType)
+	{
+		return Super::IsConnectionDisallowed(MyPin, OtherPin, OutReason);
+	}
 
 	const bool bHasStructPin = MyPin->PinType.PinCategory == UEdGraphSchema_K2::PC_Struct || OtherPin->PinType.PinCategory == UEdGraphSchema_K2::PC_Struct;
 
@@ -660,8 +670,6 @@ void UK2Node_PromotableOperator::EvaluatePinsFromChange(UEdGraphPin* ChangedPin)
 {
 	UpdateOpName();
 
-	const bool bOutputPinWasChanged = (ChangedPin == GetOutputPin());
-
 	// True if the pin that has changed now has zero connections
 	const bool bWasAFullDisconnect = (ChangedPin->LinkedTo.Num() == 0);
 
@@ -974,7 +982,6 @@ void UK2Node_PromotableOperator::UpdatePinsFromFunction(const UFunction* Functio
 {
 	if (!Function)
 	{
-		UE_LOG(LogBlueprint, Warning, TEXT("UK2Node_PromotableOperator could not update pins, function was null!"));
 		return;
 	}
 
