Check the whole pin type instead of just the name when attempting to conform a pin

#rnx

diff --git a/Engine/Source/Editor/BlueprintGraph/Private/K2Node_PromotableOperator.cpp b/Engine/Source/Editor/BlueprintGraph/Private/K2Node_PromotableOperator.cpp
--- a/Engine/Source/Editor/BlueprintGraph/Private/K2Node_PromotableOperator.cpp
+++ b/Engine/Source/Editor/BlueprintGraph/Private/K2Node_PromotableOperator.cpp
@@ -1031,12 +1031,12 @@ void UK2Node_PromotableOperator::UpdatePinsFromFunction(const UFunction* Functio
 		const FEdGraphPinType HighestLinkedType = NodePin->LinkedTo.Num() > 0 ? FTypePromotion::GetPromotedType(NodePin->LinkedTo) : NodePin->PinType;
 
 		// If the highest type is the same as the function type, just continue on with life
-		if (NodePin->LinkedTo.Num() > 0 && HighestLinkedType.PinCategory != FunctionPinType.PinCategory)
+		if (NodePin->LinkedTo.Num() > 0 && HighestLinkedType != FunctionPinType)
 		{
 			// If the links cannot be promoted to the function type, then we need to break them
 			// We don't want to break the pin if it is the one that the user has dragged on to though,
 			// because that would result in the node breaking connection as soon as the user lets go
-			if ((!FTypePromotion::IsValidPromotion(HighestLinkedType, FunctionPinType) || NodePin->Direction == EGPD_Output) && (NodePin != ChangedPin))
+			if ((NodePin != ChangedPin) && (!FTypePromotion::IsValidPromotion(HighestLinkedType, FunctionPinType) || NodePin->Direction == EGPD_Output))
 			{
 				NodePin->BreakAllPinLinks();
 			}
