Make the output message on the type promotion test include the pin types for better log output

#rnx

diff --git a/Engine/Source/Editor/BlueprintGraph/Private/Tests/TypePromotionTests.cpp b/Engine/Source/Editor/BlueprintGraph/Private/Tests/TypePromotionTests.cpp
--- a/Engine/Source/Editor/BlueprintGraph/Private/Tests/TypePromotionTests.cpp
+++ b/Engine/Source/Editor/BlueprintGraph/Private/Tests/TypePromotionTests.cpp
@@ -875,7 +875,9 @@ bool FPromotableOperatorPrimitivePromotions::RunTest(const FString& Parameters)
 			
 			// The top should be same type, and the output type should have been updated to the be new higher type
 			TestTrue(TEXT("Top Pin type propegates to new connection"), TopInputPin->PinType.PinCategory == TypePin->PinType.PinCategory);
-			TestTrue(TEXT("Output Pin type propegates to new connection"), OutputPin->PinType.PinCategory == PinToConnectTo->PinType.PinCategory);
+			const FString FinalStatusMessage = FString::Printf(TEXT("Output pin propegates to new con: Output: '%s' New Pin: '%s'"), *K2Schema->TypeToText(OutputPin->PinType).ToString(), *K2Schema->TypeToText(PinToConnectTo->PinType).ToString());
+
+			TestTrue(FinalStatusMessage, OutputPin->PinType.PinCategory == PinToConnectTo->PinType.PinCategory);
 		}
 	}
 
