Blueprints: Remove the more explict check on pin types as it is not necessary and caused some silent failures and pins to disconnect

#rnx

diff --git a/Engine/Source/Editor/BlueprintGraph/Private/K2Node_PromotableOperator.cpp b/Engine/Source/Editor/BlueprintGraph/Private/K2Node_PromotableOperator.cpp
--- a/Engine/Source/Editor/BlueprintGraph/Private/K2Node_PromotableOperator.cpp
+++ b/Engine/Source/Editor/BlueprintGraph/Private/K2Node_PromotableOperator.cpp
@@ -1031,7 +1031,7 @@ void UK2Node_PromotableOperator::UpdatePinsFromFunction(const UFunction* Functio
 		const FEdGraphPinType HighestLinkedType = NodePin->LinkedTo.Num() > 0 ? FTypePromotion::GetPromotedType(NodePin->LinkedTo) : NodePin->PinType;
 
 		// If the highest type is the same as the function type, just continue on with life
-		if (NodePin->LinkedTo.Num() > 0 && HighestLinkedType != FunctionPinType)
+		if (NodePin->LinkedTo.Num() > 0 && HighestLinkedType.PinCategory != FunctionPinType.PinCategory)
 		{
 			// If the links cannot be promoted to the function type, then we need to break them
 			// We don't want to break the pin if it is the one that the user has dragged on to though,
