Check the subobject type on struct pins to ensure that the they conform on promotable operator nodes

#rnx

diff --git a/Engine/Source/Editor/BlueprintGraph/Private/K2Node_PromotableOperator.cpp b/Engine/Source/Editor/BlueprintGraph/Private/K2Node_PromotableOperator.cpp
--- a/Engine/Source/Editor/BlueprintGraph/Private/K2Node_PromotableOperator.cpp
+++ b/Engine/Source/Editor/BlueprintGraph/Private/K2Node_PromotableOperator.cpp
@@ -1029,9 +1029,13 @@ void UK2Node_PromotableOperator::UpdatePinsFromFunction(const UFunction* Functio
 		// By default, conform to the type of the function param
 		FEdGraphPinType ConformingType = FunctionPinType;
 		const FEdGraphPinType HighestLinkedType = NodePin->LinkedTo.Num() > 0 ? FTypePromotion::GetPromotedType(NodePin->LinkedTo) : NodePin->PinType;
+		const bool bDifferingStructs = 
+			HighestLinkedType.PinCategory == UEdGraphSchema_K2::PC_Struct &&
+			FunctionPinType.PinCategory == UEdGraphSchema_K2::PC_Struct && 
+			HighestLinkedType.PinSubCategoryObject != FunctionPinType.PinSubCategoryObject;
 
 		// If the highest type is the same as the function type, just continue on with life
-		if (NodePin->LinkedTo.Num() > 0 && HighestLinkedType.PinCategory != FunctionPinType.PinCategory)
+		if (NodePin->LinkedTo.Num() > 0 && (HighestLinkedType.PinCategory != FunctionPinType.PinCategory || bDifferingStructs))
 		{
 			// If the links cannot be promoted to the function type, then we need to break them
 			// We don't want to break the pin if it is the one that the user has dragged on to though,
