TypePromotion: Also check for special conversion functions from pimitive types to others, not just for primitive to structs.

#rnx

diff --git a/Engine/Source/Editor/BlueprintGraph/Private/EdGraphSchema_K2.cpp b/Engine/Source/Editor/BlueprintGraph/Private/EdGraphSchema_K2.cpp
--- a/Engine/Source/Editor/BlueprintGraph/Private/EdGraphSchema_K2.cpp
+++ b/Engine/Source/Editor/BlueprintGraph/Private/EdGraphSchema_K2.cpp
@@ -2147,7 +2147,7 @@ const FPinConnectionResponse UEdGraphSchema_K2::CanCreateConnection(const UEdGra
 	}
 
 	// Find the calling context in case one of the pins is of type object and has a value of Self
-	UClass* CallingContext = NULL;
+	UClass* CallingContext = nullptr;
 	const UBlueprint* Blueprint = FBlueprintEditorUtils::FindBlueprintForNode(PinA->GetOwningNodeUnchecked());
 	if (Blueprint)
 	{
@@ -2160,17 +2160,10 @@ const FPinConnectionResponse UEdGraphSchema_K2::CanCreateConnection(const UEdGra
 	// Promotable types in blueprints! Only if the Cvar is set and the node is of a special type. Eventually we want this for all
 	if (TypePromoDebug::IsTypePromoEnabled() && InputPin->GetOwningNode()->IsA<UK2Node_PromotableOperator>())
 	{
-		if (FTypePromotion::IsValidPromotion(PinA->PinType, PinB->PinType) || FTypePromotion::IsValidPromotion(PinB->PinType, PinA->PinType))
+		if (FTypePromotion::IsValidPromotion(InputPin->PinType, OutputPin->PinType) || FTypePromotion::HasStructConversion(InputPin, OutputPin))
 		{
 			// Set the Text here correctly based on which pin type is higher
-			if (FTypePromotion::GetHigherType(PinA->PinType, PinB->PinType) == FTypePromotion::ETypeComparisonResult::TypeAHigher)
-			{
-				return FPinConnectionResponse(CONNECT_RESPONSE_MAKE_WITH_PROMOTION, FString::Printf(TEXT("Promote %s to %s"), *TypeToText(PinB->PinType).ToString(), *TypeToText(PinA->PinType).ToString()));
-			}
-			else
-			{
-				return FPinConnectionResponse(CONNECT_RESPONSE_MAKE_WITH_PROMOTION, FString::Printf(TEXT("Promote %s to %s"), *TypeToText(PinA->PinType).ToString(), *TypeToText(PinB->PinType).ToString()));
-			}
+			return FPinConnectionResponse(CONNECT_RESPONSE_MAKE_WITH_PROMOTION, FString::Printf(TEXT("Promote %s to %s"), *TypeToText(InputPin->PinType).ToString(), *TypeToText(OutputPin->PinType).ToString()));
 		}
 	}
 
