Blueprints: Do not allow connections to a promotable operator node if there is no existing function that would come of it

#rnx

diff --git a/Engine/Source/Editor/BlueprintGraph/Private/K2Node_PromotableOperator.cpp b/Engine/Source/Editor/BlueprintGraph/Private/K2Node_PromotableOperator.cpp
--- a/Engine/Source/Editor/BlueprintGraph/Private/K2Node_PromotableOperator.cpp
+++ b/Engine/Source/Editor/BlueprintGraph/Private/K2Node_PromotableOperator.cpp
@@ -436,8 +436,10 @@ bool UK2Node_PromotableOperator::IsConnectionDisallowed(const UEdGraphPin* MyPin
 {
 	check(MyPin && OtherPin);
 
-	// #TODO_BH Just disallow containers and references for now
-	if (OtherPin->PinType.IsContainer() || OtherPin->PinType.bIsReference)
+	const UEdGraphSchema_K2* K2Schema = GetDefault<UEdGraphSchema_K2>();
+
+	// #TODO_BH Just disallow containers for now
+	if (OtherPin->PinType.IsContainer())
 	{
 		OutReason = LOCTEXT("NoExecPinsAllowed", "Promotable Operator nodes cannot have containers or references.").ToString();
 		return true;
@@ -453,6 +455,21 @@ bool UK2Node_PromotableOperator::IsConnectionDisallowed(const UEdGraphPin* MyPin
 	{
 		return Super::IsConnectionDisallowed(MyPin, OtherPin, OutReason);
 	}
+	else if (FWildcardNodeUtils::IsWildcardPin(MyPin) && !FWildcardNodeUtils::IsWildcardPin(OtherPin))
+	{
+		TArray<UEdGraphPin*> PinsToConsider;
+		PinsToConsider.Add(const_cast<UEdGraphPin*>(OtherPin));
+
+		if (!FTypePromotion::FindBestMatchingFunc(OperationName, PinsToConsider))
+		{
+			FFormatNamedArguments Args;
+			Args.Add(TEXT("OtherPinType"), K2Schema->TypeToText(OtherPin->PinType));
+			Args.Add(TEXT("OpName"), FText::FromName(OperationName));
+
+			OutReason = FText::Format(LOCTEXT("NoCompatibleStructConv", "No matching '{OpName}' function for '{OtherPinType}'"), Args).ToString();
+			return true;
+		}
+	}
 
 	const bool bHasStructPin = MyPin->PinType.PinCategory == UEdGraphSchema_K2::PC_Struct || OtherPin->PinType.PinCategory == UEdGraphSchema_K2::PC_Struct;
 
@@ -464,7 +481,6 @@ bool UK2Node_PromotableOperator::IsConnectionDisallowed(const UEdGraphPin* MyPin
 			// Compare the directions
 			const UEdGraphPin* InputPin = nullptr;
 			const UEdGraphPin* OutputPin = nullptr;
-			const UEdGraphSchema_K2* K2Schema = GetDefault<UEdGraphSchema_K2>();
 
 			if (!K2Schema->CategorizePinsByDirection(MyPin, OtherPin, /*out*/ InputPin, /*out*/ OutputPin))
 			{
@@ -1037,10 +1053,15 @@ void UK2Node_PromotableOperator::UpdatePinsFromFunction(const UFunction* Functio
 		// If the highest type is the same as the function type, just continue on with life
 		if (NodePin->LinkedTo.Num() > 0 && (HighestLinkedType.PinCategory != FunctionPinType.PinCategory || bDifferingStructs))
 		{
+			NodePin->PinType = FunctionPinType;
+			const bool bValidPromo = 
+				FTypePromotion::IsValidPromotion(HighestLinkedType, FunctionPinType) || 
+				FTypePromotion::HasStructConversion(NodePin, NodePin->LinkedTo[0]);
+
 			// If the links cannot be promoted to the function type, then we need to break them
 			// We don't want to break the pin if it is the one that the user has dragged on to though,
 			// because that would result in the node breaking connection as soon as the user lets go
-			if ((NodePin != ChangedPin) && (!FTypePromotion::IsValidPromotion(HighestLinkedType, FunctionPinType) || NodePin->Direction == EGPD_Output))
+			if ((NodePin != ChangedPin) && !FWildcardNodeUtils::IsWildcardPin(NodePin) && !bValidPromo)
 			{
 				NodePin->BreakAllPinLinks();
 			}
