Blueprints: Do not allow enums to be connected to a promotable operator node, instead force the user to make a converision node to the type they want it to be considered as. This is similar to how you treat enums in C++ when doing operations on them, as you have to case an enum to another type in order to have them treated as numbers.

#rnx

diff --git a/Engine/Source/Editor/BlueprintGraph/Private/K2Node_PromotableOperator.cpp b/Engine/Source/Editor/BlueprintGraph/Private/K2Node_PromotableOperator.cpp
--- a/Engine/Source/Editor/BlueprintGraph/Private/K2Node_PromotableOperator.cpp
+++ b/Engine/Source/Editor/BlueprintGraph/Private/K2Node_PromotableOperator.cpp
@@ -455,6 +455,15 @@ bool UK2Node_PromotableOperator::IsConnectionDisallowed(const UEdGraphPin* MyPin
 	{
 		return Super::IsConnectionDisallowed(MyPin, OtherPin, OutReason);
 	}
+	// Enums need to be casted to a byte manually before we can do math with them, like in C++
+	else if (!FWildcardNodeUtils::IsWildcardPin(MyPin) && MyPin->Direction == EGPD_Input && OtherPin->PinType.PinCategory == UEdGraphSchema_K2::PC_Byte && OtherPin->PinType.PinSubCategoryObject != nullptr)
+	{
+		FFormatNamedArguments Args;
+		Args.Add(TEXT("OtherPinType"), K2Schema->TypeToText(OtherPin->PinType));
+
+		OutReason = FText::Format(LOCTEXT("NoCompatibleEnumConv", "'{OtherPinType}' must be converted to a numeric type before being connected"), Args).ToString();
+		return true;
+	}
 	else if (FWildcardNodeUtils::IsWildcardPin(MyPin) && !FWildcardNodeUtils::IsWildcardPin(OtherPin))
 	{
 		TArray<UEdGraphPin*> PinsToConsider;
