Remove unused functions in BP type promotion

#rnx

diff --git a/Engine/Source/Editor/BlueprintGraph/Classes/BlueprintTypePromotion.h b/Engine/Source/Editor/BlueprintGraph/Classes/BlueprintTypePromotion.h
--- a/Engine/Source/Editor/BlueprintGraph/Classes/BlueprintTypePromotion.h
+++ b/Engine/Source/Editor/BlueprintGraph/Classes/BlueprintTypePromotion.h
@@ -50,11 +50,15 @@ public:
 	/** Returns true if the given op name is a comparison operator name */
 	static bool IsComparisonOpName(const FName OpName);
 
+	/**
+	* Parse the name of the operator that this function matches to (Add, Subtract, etc)
+	*
+	* @param Func		The function to parse the operator name from
+	*
+	* @return FName		Name of the operator for this function, NO_OP if does not exist.
+	*/
 	static FName GetOpNameFromFunction(UFunction const* const Func);
 
-	/** Returns true if the given function is a candidate to handle type promotion */
-	static bool IsPromotableFunction(const UFunction* Function);
-
 	/**
 	* Determine what type a given set of wildcard pins would result in
 	*
@@ -62,14 +66,7 @@ public:
 	*/
 	static FEdGraphPinType GetPromotedType(const TArray<UEdGraphPin*>& WildcardPins);
 
-	/**
-	* Attempts to promote type A to type B. Will only work if TypeB is higher than type B
-	*
-	* @return	True if the promotion was successful
-	*/
-	static bool PromotePin(FEdGraphPinType& InTypeA, const FEdGraphPinType& TypeB);
-
-	/** Returns true if the given function can be used for type promotion (it is within the operator table) */
+	/** Returns true if the given function has been registered within the operator table */
 	static bool IsFunctionPromotionReady(const UFunction* const FuncToConsider);
 
 	/** Represents the possible results when comparing two types for promotion */
@@ -82,7 +79,7 @@ public:
 	};
 
 	/**
-	* Given the two pin types cehck which pin type is higher. 
+	* Given the two pin types check which pin type is higher. 
 	* Given two structs it will return equal, this does NOT compare PinDefaultSubobjects
 	*/
 	static ETypeComparisonResult GetHigherType(const FEdGraphPinType& A, const FEdGraphPinType& B);
@@ -99,12 +96,22 @@ public:
 	/** Function node spawner associated with this operation */
 	static UBlueprintFunctionNodeSpawner* GetOperatorSpawner(FName OpName);
 
-	/** keep track of the operator that this function provides so that we dont add multiple 
+	/** keep track of the operator that this function provides so that we don't add multiple 
 	to the BP context menu */
 	static void RegisterOperatorSpawner(FName OpName, UBlueprintFunctionNodeSpawner* Spawner);
 
+	/**
+	* Clear all registered node spawners for operators from the system. This is necessary
+	* during hot reload to ensure that all operator spawners are up to date.
+	*/
 	static void ClearNodeSpawners();
 
+	/**
+	* Get the "Primitive Promotion Table" which represents what base Pin types 
+	* can be promoted to others. 
+	*
+	* @return Const pointer to the primitive promotion table
+	*/
 	static const TMap<FName, TArray<FName>>* const GetPrimativePromotionTable();
 
 private:
@@ -123,7 +130,8 @@ private:
 
 	void GetAllFuncsForOp_Internal(FName Operation, TArray<UFunction*>& OutFuncs);
 
-	bool PromotePin_Internal(FEdGraphPinType& InTypeA, const FEdGraphPinType& TypeB);
+	/** Returns true if the given function is a candidate to handle type promotion */
+	static bool IsPromotableFunction(const UFunction* Function);
 
 	/**
 	* Determines which pin type is "higher" 
@@ -167,7 +175,7 @@ private:
 
 namespace TypePromoDebug
 {
-	/** Enables/Disables type promotion in BP */
+	/** Checks if the CVar for type promotion is true or false (BP.TypePromo.IsEnabled) */
 	static bool IsTypePromoEnabled()
 	{
 		static IConsoleVariable* TypePromoCVar = IConsoleManager::Get().FindConsoleVariable(TEXT("BP.TypePromo.IsEnabled"));
diff --git a/Engine/Source/Editor/BlueprintGraph/Private/BlueprintTypePromotion.cpp b/Engine/Source/Editor/BlueprintGraph/Private/BlueprintTypePromotion.cpp
--- a/Engine/Source/Editor/BlueprintGraph/Private/BlueprintTypePromotion.cpp
+++ b/Engine/Source/Editor/BlueprintGraph/Private/BlueprintTypePromotion.cpp
@@ -243,24 +243,6 @@ FEdGraphPinType FTypePromotion::GetPromotedType_Internal(const TArray<UEdGraphPi
 	return HighestPinType;
 }
 
-bool FTypePromotion::PromotePin(FEdGraphPinType& InTypeA, const FEdGraphPinType& TypeB)
-{
-	return FTypePromotion::Get().PromotePin_Internal(InTypeA, TypeB);
-}
-
-bool FTypePromotion::PromotePin_Internal(FEdGraphPinType& InTypeA, const FEdGraphPinType& TypeB)
-{
-	// If type B is not the higher type, than we shouldn't do anything
-	if(GetHigherType(InTypeA, TypeB) != ETypeComparisonResult::TypeBHigher)
-	{
-		return false;
-	}
-
-	InTypeA = TypeB;
-
-	return true;
-}
-
 UFunction* FTypePromotion::FindBestMatchingFunc(FName Operation, const TArray<UEdGraphPin*>& PinsToConsider)
 {
 	return FTypePromotion::Get().FindBestMatchingFunc_Internal(Operation, PinsToConsider);
