Remove unused conversion logic in the promotable operator node

#rnx

diff --git a/Engine/Source/Editor/BlueprintGraph/Classes/K2Node_PromotableOperator.h b/Engine/Source/Editor/BlueprintGraph/Classes/K2Node_PromotableOperator.h
--- a/Engine/Source/Editor/BlueprintGraph/Classes/K2Node_PromotableOperator.h
+++ b/Engine/Source/Editor/BlueprintGraph/Classes/K2Node_PromotableOperator.h
@@ -113,12 +113,6 @@ private:
 	*/
 	void UpdatePinsFromFunction(const UFunction* Function, UEdGraphPin* ChangedPin = nullptr);
 
-	/** Called when the user attempts conversion from the context menu. Records a transaction and calls UpdatePinsFromFunction */
-	void ConvertNodeToFunction(const UFunction* Function, UEdGraphPin* ChangedPin);
-
-	/** Updates the PossibleConversions function array based on the current pin types */
-	void UpdatePossibleConversionFuncs();
-
 	/**
 	* Returns all pins that have the EGPD_Input direction
 	* @param bIncludeLinks	If true, than this will also include all the pins that are linked to the inputs.
@@ -137,10 +131,6 @@ private:
 	/** The name that this operation uses ("Add", "Multiply", etc) */
 	FName OperationName;
 
-	/** Array of functions that we could possibly convert this node to via the right-click context menu */
-	UPROPERTY(Transient)
-	TArray<UFunction*> PossibleConversions;
-
 	/** The current number of additional pins on this node */
 	UPROPERTY()
 	int32 NumAdditionalInputs;
diff --git a/Engine/Source/Editor/BlueprintGraph/Private/K2Node_PromotableOperator.cpp b/Engine/Source/Editor/BlueprintGraph/Private/K2Node_PromotableOperator.cpp
--- a/Engine/Source/Editor/BlueprintGraph/Private/K2Node_PromotableOperator.cpp
+++ b/Engine/Source/Editor/BlueprintGraph/Private/K2Node_PromotableOperator.cpp
@@ -136,38 +136,6 @@ void UK2Node_PromotableOperator::GetNodeContextMenuActions(UToolMenu* Menu, UGra
 			)
 		);
 	}
-
-	// If there are possible function conversions that can happen 
-	if (Context->Pin && PossibleConversions.Num() > 0 && !Context->bIsDebugging && HasAnyConnectionsOrDefaults())
-	{
-		FToolMenuSection& Section = Menu->AddSection("K2NodePromotableOperator", LOCTEXT("ConvFunctionHeader", "Convert Function"));
-		const UFunction* CurFunction = GetTargetFunction();
-
-		for (const UFunction* Func : PossibleConversions)
-		{
-			// Don't need to convert to a function if we are already set to it
-			if (Func == CurFunction)
-			{
-				continue;
-			}
-
-			FFormatNamedArguments Args;
-			Args.Add(TEXT("TargetName"), GetUserFacingFunctionName(Func));
-			FText ConversionName = FText::Format(LOCTEXT("CallFunction_Tooltip", "Convert node to function '{TargetName}'"), Args);
-
-			const FText Tooltip = FText::FromString(GetDefaultTooltipForFunction(Func));
-
-			Section.AddMenuEntry(
-				Func->GetFName(),
-				ConversionName,
-				Tooltip,
-				FSlateIcon(),
-				FUIAction(
-					FExecuteAction::CreateUObject(const_cast<UK2Node_PromotableOperator*>(this), &UK2Node_PromotableOperator::ConvertNodeToFunction, Func, const_cast<UEdGraphPin*>(Context->Pin))
-				)
-			);
-		}
-	}
 }
 
 FText UK2Node_PromotableOperator::GetTooltipText() const
@@ -884,8 +852,6 @@ void UK2Node_PromotableOperator::ResetNodeToWildcard()
 		}
 	}
 
-	// Clear out any possible function matches, since we are removing connections
-	PossibleConversions.Empty();
 	GetGraph()->NotifyGraphChanged();
 }
 
@@ -988,21 +954,6 @@ void UK2Node_PromotableOperator::GetPinsToConsider(TArray<UEdGraphPin*>& OutArra
 	}
 }
 
-void UK2Node_PromotableOperator::ConvertNodeToFunction(const UFunction* Function, UEdGraphPin* ChangedPin)
-{
-	const FScopedTransaction Transaction(LOCTEXT("ConvertPromotableOpToFunction", "Change the function signature of a promotable operator node."));
-	Modify();
-	RecombineAllSplitPins();
-
-	// If we convert the node to a function, then just get rid of the additional pins
-	NumAdditionalInputs = 0;
-
-	UpdatePinsFromFunction(Function, ChangedPin);
-
-	// Reconstruct this node to fix any default values that may be invalid now
-	ReconstructNode();
-}
-
 void UK2Node_PromotableOperator::UpdatePinsFromFunction(const UFunction* Function, UEdGraphPin* ChangedPin /* = nullptr */)
 {
 	if (!Function)
@@ -1143,40 +1094,6 @@ void UK2Node_PromotableOperator::UpdatePinsFromFunction(const UFunction* Functio
 	GetGraph()->NotifyGraphChanged();
 }
 
-void UK2Node_PromotableOperator::UpdatePossibleConversionFuncs()
-{
-	// Pins can be empty if we are doing this during reconstruction
-	if (Pins.Num() <= 0)
-	{
-		return;
-	}
-
-	bool bAllPinTypesEqual = true;
-	FEdGraphPinType CurType = Pins[0]->PinType;
-
-	for (UEdGraphPin* Pin : Pins)
-	{
-		if (Pin->PinType != CurType)
-		{
-			bAllPinTypesEqual = false;
-			break;
-		}
-	}
-
-	UpdateOpName();
-
-	// We don't want to have a menu that is full of every possible function for an operator, that is way 
-	// to overwhelming to the user. Instead, only display conversion functions that when types are not 
-	// all the same. 
-	if (!bAllPinTypesEqual)
-	{
-		// #TODO_BH We need up change how these conversions work to be on a per-pin basis instead of via a function
-		// which will be a more consistent experience, and work better with additional pins
-		// The node has changed, so lets find the lowest matching function with the newly updated types
-		FEdGraphPinType HighestType = FTypePromotion::GetPromotedType(GetInputPins());
-	}
-}
-
 UEdGraphPin* UK2Node_PromotableOperator::GetOutputPin() const
 {
 	for (UEdGraphPin* Pin : Pins)
