Do not compare the pin types compatibility in a call function node if they are two different container types because that will always fail.

#rnx

diff --git a/Engine/Source/Editor/BlueprintGraph/Private/K2Node_CallFunction.cpp b/Engine/Source/Editor/BlueprintGraph/Private/K2Node_CallFunction.cpp
--- a/Engine/Source/Editor/BlueprintGraph/Private/K2Node_CallFunction.cpp
+++ b/Engine/Source/Editor/BlueprintGraph/Private/K2Node_CallFunction.cpp
@@ -3432,9 +3432,8 @@ bool UK2Node_CallFunction::IsConnectionDisallowed(const UEdGraphPin* MyPin, cons
 
 				for (UEdGraphPin* Pin : DependentPins)
 				{
-					// #TODO Maps need special validation for the key/value together, not just one pin type. 
-					// otherwise you will end up with false failures. 
-					if (OtherPin->PinType.IsMap() && !Pin->PinType.IsMap())
+					// If the pins are both containers, then ArePinTypesCompatible will fail incorrectly.
+					if (OtherPin->PinType.ContainerType != Pin->PinType.ContainerType)
 					{
 						continue;
 					}
