Preview current pose when selecting Anim Node. Also fixes the colour picker for PoseWatch and makes it choose an unused color by default. Hooks into the BlueprintEditorSettings to detect changes to the option. Originally the option was in PersonaOptions so support for notifications was added there - and left there when the option was relocated

diff --git a/Engine/Source/Editor/AnimationBlueprintEditor/Private/AnimationBlueprintEditor.cpp b/Engine/Source/Editor/AnimationBlueprintEditor/Private/AnimationBlueprintEditor.cpp
--- a/Engine/Source/Editor/AnimationBlueprintEditor/Private/AnimationBlueprintEditor.cpp
+++ b/Engine/Source/Editor/AnimationBlueprintEditor/Private/AnimationBlueprintEditor.cpp
@@ -19,7 +19,6 @@
 #include "BlueprintEditorTabs.h"
 #include "SKismetInspector.h"
 
-
 #include "EdGraphUtilities.h"
 #include "Kismet2/BlueprintEditorUtils.h"
 #include "Kismet2/DebuggerCommands.h"
@@ -74,6 +73,7 @@
 
 // Hide related nodes feature
 #include "Preferences/AnimationBlueprintEditorOptions.h"
+#include "AnimationBlueprintEditorSettings.h"
 #include "Framework/MultiBox/MultiBoxBuilder.h"
 #include "EdGraphNode_Comment.h"
 #include "AnimStateNodeBase.h"
@@ -174,6 +174,14 @@ FAnimationBlueprintEditor::FAnimationBlueprintEditor()
 
 FAnimationBlueprintEditor::~FAnimationBlueprintEditor()
 {
+	// Stop watching the settings
+	UAnimationBlueprintEditorSettings* AnimationBlueprintEditorSettings = GetMutableDefault<UAnimationBlueprintEditorSettings>();
+	AnimationBlueprintEditorSettings->UnregisterOnUpdateSettings(AnimationBlueprintEditorSettingsChangedHandle);
+
+	// Remove all Pose Watches that were created as a result of selection, otherwise if the editor options are changed 
+	// they will still be active if we get recreated even though the nodes won't be selected.
+	RemoveAllSelectionPoseWatches();
+
 	GEditor->OnBlueprintPreCompile().RemoveAll(this);
 
 	GEditor->GetEditorSubsystem<UImportSubsystem>()->OnAssetPostImport.RemoveAll(this);
@@ -184,6 +192,19 @@ FAnimationBlueprintEditor::~FAnimationBlueprintEditor()
 	SaveEditorSettings();
 }
 
+void FAnimationBlueprintEditor::HandleUpdateSettings(const UAnimationBlueprintEditorSettings* AnimationBlueprintEditorSettings, EPropertyChangeType::Type ChangeType)
+{
+	if (AnimationBlueprintEditorSettings->bPoseWatchSelectedNodes != bPreviousPoseWatchSelectedNodes)
+	{
+		bPreviousPoseWatchSelectedNodes = AnimationBlueprintEditorSettings->bPoseWatchSelectedNodes;
+		RemoveAllSelectionPoseWatches();
+		if (AnimationBlueprintEditorSettings->bPoseWatchSelectedNodes)
+		{
+			HandlePoseWatchSelectedNodes();
+		}
+	}
+}
+
 UAnimBlueprint* FAnimationBlueprintEditor::GetAnimBlueprint() const
 {
 	return Cast<UAnimBlueprint>(GetBlueprintObj());
@@ -323,6 +344,10 @@ void FAnimationBlueprintEditor::InitAnimationBlueprintEditor(const EToolkitMode:
 	{
 		NewDocument_OnClick(CGT_NewAnimationLayer);
 	}
+
+	// Register for notifications when settings change
+	AnimationBlueprintEditorSettingsChangedHandle = GetMutableDefault<UAnimationBlueprintEditorSettings>()->RegisterOnUpdateSettings(
+		UAnimationBlueprintEditorSettings::FOnUpdateSettingsMulticaster::FDelegate::CreateSP(this, &FAnimationBlueprintEditor::HandleUpdateSettings));
 }
 
 void FAnimationBlueprintEditor::BindCommands()
@@ -406,6 +431,12 @@ void FAnimationBlueprintEditor::SetDetailObject(UObject* Obj)
 /** Called when graph editor focus is changed */
 void FAnimationBlueprintEditor::OnGraphEditorFocused(const TSharedRef<class SGraphEditor>& InGraphEditor)
 {
+	// Remove pose watches now before calling the base class implementation because that will switch the focus
+	if (GetDefault<UAnimationBlueprintEditorSettings>()->bPoseWatchSelectedNodes)
+	{
+		RemoveAllSelectionPoseWatches();
+	}
+
 	// in the future, depending on which graph editor is this will act different
 	FBlueprintEditor::OnGraphEditorFocused(InGraphEditor);
 
@@ -420,6 +451,11 @@ void FAnimationBlueprintEditor::OnGraphEditorFocused(const TSharedRef<class SGra
 	{
 		ResetAllNodesUnrelatedStates();
 	}
+
+	if (GetDefault<UAnimationBlueprintEditorSettings>()->bPoseWatchSelectedNodes)
+	{
+		HandlePoseWatchSelectedNodes();
+	}
 }
 
 void FAnimationBlueprintEditor::OnGraphEditorBackgrounded(const TSharedRef<SGraphEditor>& InGraphEditor)
@@ -552,6 +588,64 @@ void FAnimationBlueprintEditor::OnRemovePosePin()
 	}
 }
 
+FColor FAnimationBlueprintEditor::ChoosePoseWatchColor() const
+{
+	TArrayView<const FColor> PoseWatchColors = AnimationEditorUtils::GetPoseWatchColorPalette();
+	const int32 NumColors = PoseWatchColors.Num();
+
+	if (NumColors <= 0)
+	{
+		return FColor::Red;
+	}
+
+	struct FColorCount 
+	{ 
+		FColor Color; 
+		int    Count;
+		bool operator<(const FColorCount& Other) const { return Count < Other.Count; }
+	};
+	TArray<FColorCount> ColorCounts;
+	ColorCounts.SetNum(NumColors);
+	for (int ColorIndex = 0 ; ColorIndex != NumColors ; ++ColorIndex)
+	{
+		ColorCounts[ColorIndex].Color = PoseWatchColors[ColorIndex];
+		ColorCounts[ColorIndex].Count = 0;
+	}
+
+	TSharedPtr<SGraphEditor> FocusedGraphEd = FocusedGraphEdPtr.Pin();
+	if (FocusedGraphEd.IsValid())
+	{
+		UAnimBlueprint* AnimBP = GetAnimBlueprint();
+		for (UEdGraphNode* Node : FocusedGraphEd->GetCurrentGraph()->Nodes)
+		{
+			UAnimGraphNode_Base* GraphNode = Cast<UAnimGraphNode_Base>(Node);
+			UPoseWatch* PoseWatch = AnimationEditorUtils::FindPoseWatchForNode(GraphNode, AnimBP);
+			if (PoseWatch)
+			{
+				FColor PoseWatchColor = PoseWatch->PoseWatchColour;
+				int32 Index = ColorCounts.IndexOfByPredicate([PoseWatchColor](const FColorCount& CC){return PoseWatchColor == CC.Color;});
+				if (Index != INDEX_NONE)
+				{
+					++ColorCounts[Index].Count;
+				}
+			}
+		}
+	}
+
+	// Pick the colour with the lowest count, preferring the original order if there's a tie (which there normally will be)
+	int BestIndex = 0;
+	for (int ColorIndex = 1 ; ColorIndex < NumColors ; ++ColorIndex)
+	{
+		if (ColorCounts[ColorIndex].Count < ColorCounts[BestIndex].Count)
+		{
+			BestIndex = ColorIndex;
+		}
+	}
+
+	return ColorCounts[BestIndex].Color;
+}
+
+
 void FAnimationBlueprintEditor::OnTogglePoseWatch()
 {
 	const FGraphPanelSelectionSet SelectedNodes = GetSelectedNodes();
@@ -568,7 +662,7 @@ void FAnimationBlueprintEditor::OnTogglePoseWatch()
 			}
 			else
 			{
-				AnimationEditorUtils::MakePoseWatchForNode(AnimBP, SelectedNode, FColor::Red);
+				AnimationEditorUtils::MakePoseWatchForNode(AnimBP, SelectedNode, ChoosePoseWatchColor());
 			}
 		}
 	}
@@ -1560,8 +1654,71 @@ void FAnimationBlueprintEditor::OnSelectedNodesChangedImpl(const TSet<class UObj
 			HideUnrelatedNodes();
 		}
 	}
+
+	if (GetDefault<UAnimationBlueprintEditorSettings>()->bPoseWatchSelectedNodes)
+	{
+		HandlePoseWatchSelectedNodes();
+	}
 }
 
+void FAnimationBlueprintEditor::HandlePoseWatchSelectedNodes()
+{
+	TSharedPtr<SGraphEditor> FocusedGraphEd = FocusedGraphEdPtr.Pin();
+	if (FocusedGraphEd.IsValid())
+	{
+		UAnimBlueprint* AnimBP = GetAnimBlueprint();
+		TArray<UEdGraphNode*> AllNodes = FocusedGraphEd->GetCurrentGraph()->Nodes;
+
+		FGraphPanelSelectionSet SelectionNodes = GetSelectedNodes();
+
+		for (UEdGraphNode* Node : AllNodes)
+		{
+			UAnimGraphNode_Base* GraphNode = Cast<UAnimGraphNode_Base>(Node);
+			UPoseWatch* PoseWatch = AnimationEditorUtils::FindPoseWatchForNode(GraphNode, AnimBP);
+			if (GraphNode)
+			{
+				if (SelectionNodes.Contains(Node))
+				{
+					if (!PoseWatch)
+					{
+						PoseWatch = AnimationEditorUtils::MakePoseWatchForNode(AnimBP, GraphNode, ChoosePoseWatchColor());
+						PoseWatch->bDeleteOnDeselection = true;
+					}
+				}
+				else
+				{
+					if (PoseWatch && PoseWatch->bDeleteOnDeselection)
+					{
+						AnimationEditorUtils::RemovePoseWatch(PoseWatch, AnimBP);
+					}
+				}
+			}
+		}
+	}
+}
+
+void FAnimationBlueprintEditor::RemoveAllSelectionPoseWatches()
+{
+	TSharedPtr<SGraphEditor> FocusedGraphEd = FocusedGraphEdPtr.Pin();
+	if (FocusedGraphEd.IsValid())
+	{
+		UAnimBlueprint* AnimBP = GetAnimBlueprint();
+		TArray<UEdGraphNode*> AllNodes = FocusedGraphEd->GetCurrentGraph()->Nodes;
+
+		for (UEdGraphNode* Node : AllNodes)
+		{
+			UAnimGraphNode_Base* GraphNode = Cast<UAnimGraphNode_Base>(Node);
+			if (GraphNode)
+			{
+				UPoseWatch* PoseWatch = AnimationEditorUtils::FindPoseWatchForNode(GraphNode, AnimBP);
+				if (PoseWatch && PoseWatch->bDeleteOnDeselection)
+				{
+					AnimationEditorUtils::RemovePoseWatch(PoseWatch, AnimBP);
+				}
+			}
+		}
+	}
+}
 void FAnimationBlueprintEditor::OnPostCompile()
 {
 	// act as if we have re-selected, so internal pointers are updated
diff --git a/Engine/Source/Editor/AnimationBlueprintEditor/Private/AnimationBlueprintEditor.h b/Engine/Source/Editor/AnimationBlueprintEditor/Private/AnimationBlueprintEditor.h
--- a/Engine/Source/Editor/AnimationBlueprintEditor/Private/AnimationBlueprintEditor.h
+++ b/Engine/Source/Editor/AnimationBlueprintEditor/Private/AnimationBlueprintEditor.h
@@ -20,6 +20,7 @@ class UAnimGraphNode_Base;
 class UEdGraph;
 class USkeletalMesh;
 class ISkeletonTreeItem;
+class UAnimationBlueprintEditorSettings;
 
 struct FToolMenuContext;
 
@@ -99,7 +100,7 @@ public:
 	/** IHasPersonaToolkit interface */
 	virtual TSharedRef<class IPersonaToolkit> GetPersonaToolkit() const { return PersonaToolkit.ToSharedRef(); }
 
-	/** FBlueprintEdi1tor interface */
+	/** FBlueprintEditor interface */
 	virtual void OnActiveTabChanged(TSharedPtr<SDockTab> PreviouslyActive, TSharedPtr<SDockTab> NewlyActivated) override;
 	virtual void OnSelectedNodesChangedImpl(const TSet<class UObject*>& NewSelection) override;
 	virtual void HandleSetObjectBeingDebugged(UObject* InObject) override;
@@ -288,6 +289,12 @@ private:
 	/** Handle the preview anim blueprint being compiled */
 	void HandlePreviewAnimBlueprintCompiled(UBlueprint* InBlueprint);
 
+	/** Enable/disable pose watch on selected nodes */
+	void HandlePoseWatchSelectedNodes();
+
+	/** Removes all pose watches created by selectionfrom the current view */
+	void RemoveAllSelectionPoseWatches();
+
     /**
 	 * Load editor settings from disk (docking state, window pos/size, option state, etc).
 	 */
@@ -300,6 +307,11 @@ private:
 
 	void HandleAnimationSequenceBrowserCreated(const TSharedRef<IAnimationSequenceBrowser>& InSequenceBrowser);
 
+	void HandleUpdateSettings(const UAnimationBlueprintEditorSettings* AnimationBlueprintEditorSettings, EPropertyChangeType::Type ChangeType);
+
+	/** Chooses a suitable pose watch colour automatically - i.e. one that isn't already in use (if possible) */
+	FColor ChoosePoseWatchColor() const;
+
 	/** The extender to pass to the level editor to extend it's window menu */
 	TSharedPtr<FExtender> MenuExtender;
 
@@ -332,4 +344,10 @@ private:
 
 	/** Cached mesh component held during compilation, used to reconnect debugger */
 	USkeletalMeshComponent* DebuggedMeshComponent;
+
+	/** Used to track wither the editor option has changed */
+	bool bPreviousPoseWatchSelectedNodes = false;
+
+	/** Delegate handle registered for when settings change */
+	FDelegateHandle AnimationBlueprintEditorSettingsChangedHandle;
 };
diff --git a/Engine/Source/Editor/AnimationBlueprintEditor/Private/AnimationBlueprintEditorModule.cpp b/Engine/Source/Editor/AnimationBlueprintEditor/Private/AnimationBlueprintEditorModule.cpp
--- a/Engine/Source/Editor/AnimationBlueprintEditor/Private/AnimationBlueprintEditorModule.cpp
+++ b/Engine/Source/Editor/AnimationBlueprintEditor/Private/AnimationBlueprintEditorModule.cpp
@@ -2,6 +2,7 @@
 
 #include "AnimationBlueprintEditorModule.h"
 #include "AnimationGraphFactory.h"
+#include "AnimationBlueprintEditorSettings.h"
 #include "Animation/AnimInstance.h"
 #include "AnimationBlueprintEditor.h"
 #include "EdGraphSchema_K2.h"
@@ -9,6 +10,7 @@
 #include "Animation/AnimBlueprint.h"
 #include "Kismet2/KismetEditorUtilities.h"
 #include "Developer/MessageLog/Public/MessageLogModule.h"
+#include "ISettingsModule.h"
 
 IMPLEMENT_MODULE( FAnimationBlueprintEditorModule, AnimationBlueprintEditor);
 
@@ -35,6 +37,13 @@ void FAnimationBlueprintEditorModule::StartupModule()
 	InitOptions.bShowFilters = true;
 	InitOptions.bShowPages = true;
 	MessageLogModule.RegisterLogListing("AnimBlueprintLog", LOCTEXT("AnimBlueprintLog", "Anim Blueprint Log"), InitOptions);
+
+	ISettingsModule& SettingsModule = FModuleManager::LoadModuleChecked<ISettingsModule>("Settings");
+	SettingsModule.RegisterSettings("Editor", "ContentEditors", "AnimationBlueprintEditor",
+		LOCTEXT("AnimationBlueprintEditorSettingsName", "Animation Blueprint Editor"),
+		LOCTEXT("AnimationBlueprintEditorSettingsDescription", "Configure the look and feel of the Animation Blueprint Editor."),
+		GetMutableDefault<UAnimationBlueprintEditorSettings>());
+
 }
 
 void FAnimationBlueprintEditorModule::ShutdownModule()
@@ -45,6 +54,9 @@ void FAnimationBlueprintEditorModule::ShutdownModule()
 	FEdGraphUtilities::UnregisterVisualPinFactory(AnimGraphPinFactory);
 	FEdGraphUtilities::UnregisterVisualPinConnectionFactory(AnimGraphPinConnectionFactory);
 
+	ISettingsModule& SettingsModule = FModuleManager::LoadModuleChecked<ISettingsModule>("Settings");
+	SettingsModule.UnregisterSettings("Editor", "ContentEditors", "AnimationBlueprintEditor");
+
 	MenuExtensibilityManager.Reset();
 	ToolBarExtensibilityManager.Reset();
 }
diff --git a/Engine/Source/Editor/AnimationBlueprintEditor/Private/AnimationBlueprintEditorSettings.cpp b/Engine/Source/Editor/AnimationBlueprintEditor/Private/AnimationBlueprintEditorSettings.cpp
new file mode 100644
--- /dev/null
+++ b/Engine/Source/Editor/AnimationBlueprintEditor/Private/AnimationBlueprintEditorSettings.cpp
@@ -0,0 +1,9 @@
+// Copyright Epic Games, Inc. All Rights Reserved.
+
+#include "AnimationBlueprintEditorSettings.h"
+
+void UAnimationBlueprintEditorSettings::PostEditChangeProperty(FPropertyChangedEvent& PropertyChangedEvent)
+{
+	Super::PostEditChangeProperty(PropertyChangedEvent);
+	OnSettingsChange.Broadcast(this, PropertyChangedEvent.ChangeType);
+}
diff --git a/Engine/Source/Editor/AnimationBlueprintEditor/Private/AnimationNodes/SAnimationGraphNode.cpp b/Engine/Source/Editor/AnimationBlueprintEditor/Private/AnimationNodes/SAnimationGraphNode.cpp
--- a/Engine/Source/Editor/AnimationBlueprintEditor/Private/AnimationNodes/SAnimationGraphNode.cpp
+++ b/Engine/Source/Editor/AnimationBlueprintEditor/Private/AnimationNodes/SAnimationGraphNode.cpp
@@ -18,6 +18,7 @@
 #include "AnimationGraphSchema.h"
 #include "SGraphPin.h"
 #include "Widgets/Layout/SWrapBox.h"
+#include "Brushes/SlateColorBrush.h"
 
 #define LOCTEXT_NAMESPACE "AnimationGraphNode"
 
@@ -33,10 +34,13 @@ public:
 	{
 		PoseWatch = InArgs._PoseWatch;
 
-		static FColor PoseWatchColours[] = { FColor::Red, FColor::Green, FColor::Blue, FColor::Cyan, FColor::Orange, FColor::Purple, FColor::Yellow, FColor::Black };
+		static const FSlateColorBrush WhiteBrush(FLinearColor::White);
+
+		TArrayView<const FColor> PoseWatchColors = AnimationEditorUtils::GetPoseWatchColorPalette();
+		const int32 NumColors = PoseWatchColors.Num();
 
 		const int32 Rows = 2;
-		const int32 Columns = 4;
+		const int32 Columns = NumColors / Rows;
 
 		TSharedPtr<SVerticalBox> Layout = SNew(SVerticalBox);
 
@@ -46,16 +50,26 @@ public:
 
 			for (int32 RowItem = 0; RowItem < Columns; ++RowItem)
 			{
-				int32 ColourIndex = RowItem + (RowIndex * Columns);
-				FColor Colour = PoseWatchColours[ColourIndex];
+				int32 ColorIndex = RowItem + (RowIndex * Columns);
+				FColor Color = PoseWatchColors[ColorIndex];
+				if (ColorIndex >= NumColors)
+				{
+					break;
+				}
 
 				Row->AddSlot()
 				.Padding(5.f, 2.f)
 				[
 					SNew(SButton)
+					.ButtonStyle(FAppStyle::Get(), "NoBorder")
 					.HAlign(HAlign_Center)
-					.OnClicked(this, &SPoseViewColourPickerPopup::NewPoseWatchColourPicked, Colour)
-					.ButtonColorAndOpacity(Colour)
+					.OnClicked(this, &SPoseViewColourPickerPopup::NewPoseWatchColourPicked, Color)
+					[
+						SNew(SImage)
+						.Image(&WhiteBrush)
+						.DesiredSizeOverride(FVector2D(24, 24))
+						.ColorAndOpacity(Color)
+					]
 				];
 
 			}
diff --git a/Engine/Source/Editor/AnimationBlueprintEditor/Public/AnimationBlueprintEditorSettings.h b/Engine/Source/Editor/AnimationBlueprintEditor/Public/AnimationBlueprintEditorSettings.h
new file mode 100644
--- /dev/null
+++ b/Engine/Source/Editor/AnimationBlueprintEditor/Public/AnimationBlueprintEditorSettings.h
@@ -0,0 +1,41 @@
+// Copyright Epic Games, Inc. All Rights Reserved.
+
+#pragma once
+
+#include "CoreMinimal.h"
+#include "UObject/ObjectMacros.h"
+#include "UObject/Object.h"
+#include "Engine/EngineBaseTypes.h"
+#include "AnimationBlueprintEditorSettings.generated.h"
+
+
+// Settings for the Animation Blueprint Editor
+UCLASS(config = EditorPerProjectUserSettings)
+class UAnimationBlueprintEditorSettings : public UObject
+{
+	GENERATED_BODY()
+
+public:
+	/** If true, automatically pose watch selected nodes. */
+	UPROPERTY(EditAnywhere, config, Category = "Debugging")
+	bool bPoseWatchSelectedNodes = false;
+
+public:
+	DECLARE_MULTICAST_DELEGATE_TwoParams(FOnUpdateSettingsMulticaster, const UAnimationBlueprintEditorSettings*, EPropertyChangeType::Type);
+	FOnUpdateSettingsMulticaster OnSettingsChange;
+
+	FDelegateHandle RegisterOnUpdateSettings(const FOnUpdateSettingsMulticaster::FDelegate& Delegate)
+	{
+		return OnSettingsChange.Add(Delegate);
+	}
+
+	void UnregisterOnUpdateSettings(FDelegateHandle Object)
+	{
+		OnSettingsChange.Remove(Object);
+	}
+
+protected:
+	//~ Begin UObject Interface
+	virtual void PostEditChangeProperty(struct FPropertyChangedEvent& PropertyChangedEvent) override;
+	//~ End UObject Interface
+};
diff --git a/Engine/Source/Editor/BlueprintGraph/Private/BlueprintEditorSettings.cpp b/Engine/Source/Editor/BlueprintGraph/Private/BlueprintEditorSettings.cpp
--- a/Engine/Source/Editor/BlueprintGraph/Private/BlueprintEditorSettings.cpp
+++ b/Engine/Source/Editor/BlueprintGraph/Private/BlueprintEditorSettings.cpp
@@ -23,20 +23,14 @@ UBlueprintEditorSettings::UBlueprintEditorSettings(const FObjectInitializer& Obj
 	, bExposeDeprecatedFunctions(false)
 	, bCompactCallOnMemberNodes(false)
 	, bFlattenFavoritesMenus(true)
-	, bFavorPureCastNodes(false)
 	, bAutoCastObjectConnections(false)
 	, bShowViewportOnSimulate(false)
-	, bShowInheritedVariables(false)
-	, bAlwaysShowInterfacesInOverrides(true)
-	, bShowParentClassInOverrides(true)
-	, bShowEmptySections(true)
-	, bShowAccessSpecifier(false)
 	, bSpawnDefaultBlueprintNodes(true)
 	, bHideConstructionScriptComponentsInDetailsView(true)
 	, bHostFindInBlueprintsInGlobalTab(true)
 	, bNavigateToNativeFunctionsFromCallNodes(true)
-	, bIncludeCommentNodesInBookmarksTab(true)
-	, bShowBookmarksForCurrentDocumentOnlyInTab(false)
+	// Experimental
+	, bFavorPureCastNodes(false)
 	// Compiler Settings
 	, SaveOnCompile(SoC_Never)
 	, bJumpToNodeErrors(false)
@@ -47,6 +41,14 @@ UBlueprintEditorSettings::UBlueprintEditorSettings(const FObjectInitializer& Obj
 	, bShowDetailedCompileResults(false)
 	, CompileEventDisplayThresholdMs(5)
 	, NodeTemplateCacheCapMB(20.f)
+	// No category
+	, bShowInheritedVariables(false)
+	, bAlwaysShowInterfacesInOverrides(true)
+	, bShowParentClassInOverrides(true)
+	, bShowEmptySections(true)
+	, bShowAccessSpecifier(false)
+	, bIncludeCommentNodesInBookmarksTab(true)
+	, bShowBookmarksForCurrentDocumentOnlyInTab(false)
 {
 	// settings that were moved out of experimental...
 	UEditorExperimentalSettings const* ExperimentalSettings = GetDefault<UEditorExperimentalSettings>();
@@ -100,4 +102,6 @@ void UBlueprintEditorSettings::PostEditChangeProperty(FPropertyChangedEvent& Pro
 	}
 
 	Super::PostEditChangeProperty(PropertyChangedEvent);
-}
\ No newline at end of file
+
+	OnSettingsChange.Broadcast(this, PropertyChangedEvent.ChangeType);
+}
diff --git a/Engine/Source/Editor/BlueprintGraph/Public/BlueprintEditorSettings.h b/Engine/Source/Editor/BlueprintGraph/Public/BlueprintEditorSettings.h
--- a/Engine/Source/Editor/BlueprintGraph/Public/BlueprintEditorSettings.h
+++ b/Engine/Source/Editor/BlueprintGraph/Public/BlueprintEditorSettings.h
@@ -67,10 +67,6 @@ public:
 	UPROPERTY(EditAnywhere, config, Category=Workflow)
 	bool bFlattenFavoritesMenus;
 
-	/** If enabled, then placed cast nodes will default to their "pure" form (meaning: without execution pins). */
-	UPROPERTY(EditAnywhere, config, AdvancedDisplay, Category=Experimental, meta=(DisplayName="Default to Using Pure Cast Nodes"))
-	bool bFavorPureCastNodes;
-
 	/** If enabled, then you'll be able to directly connect arbitrary object pins together (a pure cast node will be injected automatically). */
 	UPROPERTY(EditAnywhere, config, Category=Workflow)
 	bool bAutoCastObjectConnections;
@@ -79,6 +75,71 @@ public:
 	UPROPERTY(EditAnywhere, config, Category=Workflow)
 	bool bShowViewportOnSimulate;
 
+	/** If set will spawn default "ghost" event nodes in new Blueprints, modifiable in the [DefaultEventNodes] section of EditorPerProjectUserSettings */
+	UPROPERTY(EditAnywhere, config, Category = Workflow)
+		bool bSpawnDefaultBlueprintNodes;
+
+	/** If set will exclude components added in a Blueprint class Construction Script from the component details view */
+	UPROPERTY(EditAnywhere, config, Category = Workflow)
+		bool bHideConstructionScriptComponentsInDetailsView;
+
+	/** If set, the global Find in Blueprints command (CTRL-SHIFT-F) will be hosted in a standalone tab. This tab can remain open after the Blueprint Editor context is closed. */
+	UPROPERTY(EditAnywhere, config, Category = Workflow)
+		bool bHostFindInBlueprintsInGlobalTab;
+
+	/** If set, double clicking on a call function node will attempt to navigate an open C++ editor to the native source definition */
+	UPROPERTY(EditAnywhere, config, Category = Workflow)
+		bool bNavigateToNativeFunctionsFromCallNodes;
+
+// Experimental
+public:
+	// The list of namespaces to always expose in any Blueprint (local per-user)
+	UPROPERTY(EditAnywhere, config, Category = Experimental)
+		TArray<FString> NamespacesToAlwaysInclude;
+
+	/** If enabled, then placed cast nodes will default to their "pure" form (meaning: without execution pins). */
+	UPROPERTY(EditAnywhere, config, AdvancedDisplay, Category = Experimental, meta = (DisplayName = "Default to Using Pure Cast Nodes"))
+		bool bFavorPureCastNodes;
+
+	// Compiler Settings
+public:
+	/** Determines when to save Blueprints post-compile */
+	UPROPERTY(EditAnywhere, config, Category = Compiler)
+		TEnumAsByte<ESaveOnCompile> SaveOnCompile;
+
+	/** When enabled, if a blueprint has compiler errors, then the graph will jump and focus on the first node generating an error */
+	UPROPERTY(EditAnywhere, config, Category = Compiler)
+		bool bJumpToNodeErrors;
+
+	/** If enabled, nodes can be explicitly disabled via context menu when right-clicking on impure nodes in the Blueprint editor. Disabled nodes will not be compiled, but also will not break existing connections. */
+	UPROPERTY(EditAnywhere, config, Category = Experimental, AdvancedDisplay)
+		bool bAllowExplicitImpureNodeDisabling;
+
+	// Developer Settings
+public:
+	/** If enabled, tooltips on action menu items will show the associated action's signature id (can be used to setup custom favorites menus).*/
+	UPROPERTY(EditAnywhere, config, Category = DeveloperTools)
+		bool bShowActionMenuItemSignatures;
+
+	/** If enabled, blueprint nodes in the event graph will display with unique names rather than their display name. */
+	UPROPERTY(EditAnywhere, config, Category = DeveloperTools, meta = (DisplayName = "Display Unique Names for Blueprint Nodes"))
+		bool bBlueprintNodeUniqueNames;
+
+	// Perf Settings
+public:
+	/** If enabled, additional details will be displayed in the Compiler Results tab after compiling a blueprint. */
+	UPROPERTY(EditAnywhere, config, Category = Performance)
+		bool bShowDetailedCompileResults;
+
+	/** Minimum event time threshold used as a filter when additional details are enabled for display in the Compiler Results tab. A value of zero means that all events will be included in the final summary. */
+	UPROPERTY(EditAnywhere, AdvancedDisplay, config, Category = Performance, DisplayName = "Compile Event Results Threshold (ms)", meta = (ClampMin = "0", UIMin = "0"))
+		int32 CompileEventDisplayThresholdMs;
+
+	/** The node template cache is used to speed up blueprint menuing. This determines the peak data size for that cache. */
+	UPROPERTY(EditAnywhere, AdvancedDisplay, config, Category = Performance, DisplayName = "Node-Template Cache Cap (MB)", meta = (ClampMin = "0", UIMin = "0"))
+		float NodeTemplateCacheCapMB;
+
+public:
 	/** If set we'll show the inherited variables in the My Blueprint view. */
 	UPROPERTY(config)
 	bool bShowInheritedVariables;
@@ -99,22 +160,6 @@ public:
 	UPROPERTY(config)
 	bool bShowAccessSpecifier;
 
-	/** If set will spawn default "ghost" event nodes in new Blueprints, modifiable in the [DefaultEventNodes] section of EditorPerProjectUserSettings */
-	UPROPERTY(EditAnywhere, config, Category=Workflow)
-	bool bSpawnDefaultBlueprintNodes;
-
-	/** If set will exclude components added in a Blueprint class Construction Script from the component details view */
-	UPROPERTY(EditAnywhere, config, Category=Workflow)
-	bool bHideConstructionScriptComponentsInDetailsView;
-
-	/** If set, the global Find in Blueprints command (CTRL-SHIFT-F) will be hosted in a standalone tab. This tab can remain open after the Blueprint Editor context is closed. */
-	UPROPERTY(EditAnywhere, config, Category=Workflow)
-	bool bHostFindInBlueprintsInGlobalTab;
-	
-	/** If set, double clicking on a call function node will attempt to navigate an open C++ editor to the native source definition */
-	UPROPERTY(EditAnywhere, config, Category=Workflow)
-	bool bNavigateToNativeFunctionsFromCallNodes;
-
 	/** Blueprint bookmark database */
 	UPROPERTY(config)
 	TMap<FGuid, FEditedDocumentInfo> Bookmarks;
@@ -135,48 +180,19 @@ public:
 	UPROPERTY(config)
 	TMap<int32, FEditedDocumentInfo> GraphEditorQuickJumps;
 
-// Compiler Settings
 public:
-	/** Determines when to save Blueprints post-compile */
-	UPROPERTY(EditAnywhere, config, Category=Compiler)
-	TEnumAsByte<ESaveOnCompile> SaveOnCompile;
-
-	/** When enabled, if a blueprint has compiler errors, then the graph will jump and focus on the first node generating an error */
-	UPROPERTY(EditAnywhere, config, Category=Compiler)
-	bool bJumpToNodeErrors;
-
-	/** If enabled, nodes can be explicitly disabled via context menu when right-clicking on impure nodes in the Blueprint editor. Disabled nodes will not be compiled, but also will not break existing connections. */
-	UPROPERTY(EditAnywhere, config, Category=Experimental, AdvancedDisplay)
-	bool bAllowExplicitImpureNodeDisabling;
-
-// Developer Settings
-public:
-	/** If enabled, tooltips on action menu items will show the associated action's signature id (can be used to setup custom favorites menus).*/
-	UPROPERTY(EditAnywhere, config, Category=DeveloperTools)
-	bool bShowActionMenuItemSignatures;
-
-	/** If enabled, blueprint nodes in the event graph will display with unique names rather than their display name. */
-	UPROPERTY(EditAnywhere, config, Category = DeveloperTools, meta = (DisplayName = "Display Unique Names for Blueprint Nodes"))
-	bool bBlueprintNodeUniqueNames;
-
-public:
-	// The list of namespaces to always expose in any Blueprint (local per-user)
-	UPROPERTY(EditAnywhere, config, Category = Experimental)
-	TArray<FString> NamespacesToAlwaysInclude;
-
-// Perf Settings
-public:
-	/** If enabled, additional details will be displayed in the Compiler Results tab after compiling a blueprint. */
-	UPROPERTY(EditAnywhere, config, Category=Performance)
-	bool bShowDetailedCompileResults;
-
-	/** Minimum event time threshold used as a filter when additional details are enabled for display in the Compiler Results tab. A value of zero means that all events will be included in the final summary. */
-	UPROPERTY(EditAnywhere, AdvancedDisplay, config, Category=Performance, DisplayName="Compile Event Results Threshold (ms)", meta=(ClampMin="0", UIMin="0"))
-	int32 CompileEventDisplayThresholdMs;
-
-	/** The node template cache is used to speed up blueprint menuing. This determines the peak data size for that cache. */
-	UPROPERTY(EditAnywhere, AdvancedDisplay, config, Category=Performance, DisplayName="Node-Template Cache Cap (MB)", meta=(ClampMin="0", UIMin="0"))
-	float NodeTemplateCacheCapMB;
+	DECLARE_MULTICAST_DELEGATE_TwoParams(FOnUpdateSettingsMulticaster, const UBlueprintEditorSettings*, EPropertyChangeType::Type);
+	FOnUpdateSettingsMulticaster OnSettingsChange;
+
+	FDelegateHandle RegisterOnUpdateSettings(const FOnUpdateSettingsMulticaster::FDelegate& Delegate)
+	{
+		return OnSettingsChange.Add(Delegate);
+	}
+
+	void UnregisterOnUpdateSettings(FDelegateHandle Object)
+	{
+		OnSettingsChange.Remove(Object);
+	}
 
 protected:
 	//~ Begin UObject Interface
diff --git a/Engine/Source/Editor/UnrealEd/Classes/Preferences/PersonaOptions.h b/Engine/Source/Editor/UnrealEd/Classes/Preferences/PersonaOptions.h
--- a/Engine/Source/Editor/UnrealEd/Classes/Preferences/PersonaOptions.h
+++ b/Engine/Source/Editor/UnrealEd/Classes/Preferences/PersonaOptions.h
@@ -227,4 +227,21 @@ public:
 	void SetNotifyTimingNodeColor(const FLinearColor& InColor);
 	void SetBranchingPointTimingNodeColor(const FLinearColor& InColor);
 	FAssetEditorOptions& GetAssetEditorOptions(const FName& InContext);
+
+	DECLARE_MULTICAST_DELEGATE_TwoParams(FOnUpdateSettingsMulticaster, const UPersonaOptions*, EPropertyChangeType::Type);
+	FOnUpdateSettingsMulticaster OnSettingsChange;
+
+	FDelegateHandle RegisterOnUpdateSettings(const FOnUpdateSettingsMulticaster::FDelegate& Delegate)
+	{
+		return OnSettingsChange.Add(Delegate);
+	}
+
+	void UnregisterOnUpdateSettings(FDelegateHandle Object)
+	{
+		OnSettingsChange.Remove(Object);
+	}
+
+protected:
+	// UObject interface
+	virtual void PostEditChangeProperty(struct FPropertyChangedEvent& PropertyChangedEvent) override;
 };
diff --git a/Engine/Source/Editor/UnrealEd/Private/AnimationEditorUtils.cpp b/Engine/Source/Editor/UnrealEd/Private/AnimationEditorUtils.cpp
--- a/Engine/Source/Editor/UnrealEd/Private/AnimationEditorUtils.cpp
+++ b/Engine/Source/Editor/UnrealEd/Private/AnimationEditorUtils.cpp
@@ -882,15 +882,17 @@ namespace AnimationEditorUtils
 #endif
 	}
 
-	void MakePoseWatchForNode(UAnimBlueprint* AnimBlueprint, UEdGraphNode* Node, FColor PoseWatchColour)
+	UPoseWatch* MakePoseWatchForNode(UAnimBlueprint* AnimBlueprint, UEdGraphNode* Node, FColor PoseWatchColour)
 	{
 #if WITH_EDITORONLY_DATA
 		UPoseWatch* NewPoseWatch = NewObject<UPoseWatch>(AnimBlueprint);
 		NewPoseWatch->Node = Node;
 		NewPoseWatch->PoseWatchColour = PoseWatchColour;
 		AnimBlueprint->PoseWatches.Add(NewPoseWatch);
-
 		SetPoseWatch(NewPoseWatch, AnimBlueprint);
+		return NewPoseWatch;
+#else
+		return nullptr;
 #endif
 	}
 
@@ -915,6 +917,13 @@ namespace AnimationEditorUtils
 #endif
 	}
 
+	TArrayView<const FColor> GetPoseWatchColorPalette()
+	{
+		static const FColor PoseWatchColors[] = { FColor::Red, FColor::Green, FColor::Blue, FColor::Cyan, FColor::Orange, FColor::Purple, FColor::Yellow, FColor::Black };
+		size_t NumColors = sizeof(PoseWatchColors) / sizeof(PoseWatchColors[0]);
+		return MakeArrayView(PoseWatchColors, NumColors);
+	}
+
 	void UpdatePoseWatchColour(UPoseWatch* PoseWatch, FColor NewPoseWatchColour)
 	{
 #if WITH_EDITORONLY_DATA
diff --git a/Engine/Source/Editor/UnrealEd/Private/PreferenceStubs.cpp b/Engine/Source/Editor/UnrealEd/Private/PreferenceStubs.cpp
--- a/Engine/Source/Editor/UnrealEd/Private/PreferenceStubs.cpp
+++ b/Engine/Source/Editor/UnrealEd/Private/PreferenceStubs.cpp
@@ -274,4 +274,11 @@ FAssetEditorOptions& UPersonaOptions::GetAssetEditorOptions(const FName& InConte
 	{
 		return *FoundOptions;
 	}
-}
\ No newline at end of file
+}
+
+void UPersonaOptions::PostEditChangeProperty(struct FPropertyChangedEvent& PropertyChangedEvent)
+{
+	Super::PostEditChangeProperty(PropertyChangedEvent);
+
+	OnSettingsChange.Broadcast(this, PropertyChangedEvent.ChangeType);
+}
diff --git a/Engine/Source/Editor/UnrealEd/Public/AnimationEditorUtils.h b/Engine/Source/Editor/UnrealEd/Public/AnimationEditorUtils.h
--- a/Engine/Source/Editor/UnrealEd/Public/AnimationEditorUtils.h
+++ b/Engine/Source/Editor/UnrealEd/Public/AnimationEditorUtils.h
@@ -194,9 +194,10 @@ namespace AnimationEditorUtils
 
 	UNREALED_API void SetPoseWatch(UPoseWatch* PoseWatch, UAnimBlueprint* AnimBlueprintIfKnown = nullptr);
 	UNREALED_API UPoseWatch* FindPoseWatchForNode(const UEdGraphNode* Node, UAnimBlueprint* AnimBlueprintIfKnown=nullptr);
-	UNREALED_API void MakePoseWatchForNode(UAnimBlueprint* AnimBlueprint, UEdGraphNode* Node, FColor PoseWatchColour);
+	UNREALED_API UPoseWatch* MakePoseWatchForNode(UAnimBlueprint* AnimBlueprint, UEdGraphNode* Node, FColor PoseWatchColour);
 	UNREALED_API void RemovePoseWatch(UPoseWatch* PoseWatch, UAnimBlueprint* AnimBlueprintIfKnown=nullptr);
 	UNREALED_API void UpdatePoseWatchColour(UPoseWatch* PoseWatch, FColor NewPoseWatchColour);
+	UNREALED_API TArrayView<const FColor> GetPoseWatchColorPalette();
 
 	//////////////////////////////////////////////////////////////////////////////////////////
 
diff --git a/Engine/Source/Runtime/Engine/Classes/Engine/PoseWatch.h b/Engine/Source/Runtime/Engine/Classes/Engine/PoseWatch.h
--- a/Engine/Source/Runtime/Engine/Classes/Engine/PoseWatch.h
+++ b/Engine/Source/Runtime/Engine/Classes/Engine/PoseWatch.h
@@ -31,4 +31,8 @@ public:
 
 	UPROPERTY()
 	FColor PoseWatchColour;
+
+	/** If true then the watch will be deleted when the connected node is deleted */
+	UPROPERTY()
+	bool bDeleteOnDeselection = false;
 };
