Set the friendly display name on native variables to be the DisplayName metadata for K2 Variable Nodes

#rnx

diff --git a/Engine/Source/Editor/BlueprintGraph/Private/K2Node_Variable.cpp b/Engine/Source/Editor/BlueprintGraph/Private/K2Node_Variable.cpp
--- a/Engine/Source/Editor/BlueprintGraph/Private/K2Node_Variable.cpp
+++ b/Engine/Source/Editor/BlueprintGraph/Private/K2Node_Variable.cpp
@@ -89,9 +89,14 @@ bool UK2Node_Variable::CreatePinForVariable(EEdGraphPinDirection Direction, FNam
 
 	if (VariableProperty != nullptr)
 	{
-		const FName PinName = InPinName.IsNone()? GetVarName() : InPinName;
+		const FName PinName = InPinName.IsNone() ? GetVarName() : InPinName;
 		// Create the pin
 		UEdGraphPin* VariablePin = CreatePin(Direction, NAME_None, PinName);
+		if (VariableProperty->IsNative())
+		{
+			VariablePin->PinFriendlyName = VariableProperty->GetDisplayNameText();
+		}
+
 		K2Schema->ConvertPropertyToPinType(VariableProperty, /*out*/ VariablePin->PinType);
 		K2Schema->SetPinAutogeneratedDefaultValueBasedOnType(VariablePin);
 	}
