Expose the operator table on a per-type basis so that nodes can find what primitive types can be changed

#rnx

diff --git a/Engine/Source/Editor/BlueprintGraph/Classes/BlueprintTypePromotion.h b/Engine/Source/Editor/BlueprintGraph/Classes/BlueprintTypePromotion.h
--- a/Engine/Source/Editor/BlueprintGraph/Classes/BlueprintTypePromotion.h
+++ b/Engine/Source/Editor/BlueprintGraph/Classes/BlueprintTypePromotion.h
@@ -113,6 +113,15 @@ public:
 	* @return Const pointer to the primitive promotion table
 	*/
 	static const TMap<FName, TArray<FName>>* const GetPrimativePromotionTable();
+	
+	/**
+	* Get a pointer to an array of available promotion types to a given pin type
+	*
+	* @param Type		The pin type to find available promotions to
+	*
+	* @return const TArray<FName>*	Pointer to the available promotions, nullptr if there are none.
+	*/
+	static const TArray<FName>* GetAvailablePrimitivePromotions(const FEdGraphPinType& Type);
 
 private:
 
diff --git a/Engine/Source/Editor/BlueprintGraph/Private/BlueprintTypePromotion.cpp b/Engine/Source/Editor/BlueprintGraph/Private/BlueprintTypePromotion.cpp
--- a/Engine/Source/Editor/BlueprintGraph/Private/BlueprintTypePromotion.cpp
+++ b/Engine/Source/Editor/BlueprintGraph/Private/BlueprintTypePromotion.cpp
@@ -111,6 +111,13 @@ const TMap<FName, TArray<FName>>* const FTypePromotion::GetPrimativePromotionTab
 	return nullptr;
 }
 
+const TArray<FName>* FTypePromotion::GetAvailablePrimitivePromotions(const FEdGraphPinType& Type)
+{
+	const TMap<FName, TArray<FName>>* PromoTable = GetPrimativePromotionTable();
+
+	return PromoTable->Find(Type.PinCategory);
+}
+
 bool FTypePromotion::IsValidPromotion(const FEdGraphPinType& A, const FEdGraphPinType& B)
 {
 	// If either of these pin types is a struct, than we have to have some kind of valid
