Clean up BP type promotion unit tests to reduce code repetition

#rnx

diff --git a/Engine/Source/Editor/BlueprintGraph/Private/Tests/TypePromotionTests.cpp b/Engine/Source/Editor/BlueprintGraph/Private/Tests/TypePromotionTests.cpp
--- a/Engine/Source/Editor/BlueprintGraph/Private/Tests/TypePromotionTests.cpp
+++ b/Engine/Source/Editor/BlueprintGraph/Private/Tests/TypePromotionTests.cpp
@@ -87,6 +87,22 @@ namespace TypePromoTestUtils
 
 		return bConnected;
 	}
+
+	static FString GetPinListDisplayName(const TArray<UEdGraphPin*>& TestPins)
+	{
+		const UEdGraphSchema_K2* K2Schema = GetDefault<UEdGraphSchema_K2>();
+		FString PinTypesString = "";
+		for (int32 i = 0; i < TestPins.Num(); ++i)
+		{
+			UEdGraphPin* Pin = TestPins[i];
+			PinTypesString += K2Schema->TypeToText(Pin->PinType).ToString();
+			if (i != TestPins.Num() - 1)
+			{
+				PinTypesString += ", ";
+			}
+		}
+		return PinTypesString;
+	}
 }
 
 #define MakeTestableBP(BPName, GraphName)												\
@@ -198,19 +214,28 @@ bool FFindBestMatchingFunc::RunTest(const FString& Parameters)
 	TArray<UEdGraphPin*> PinTypes = {};
 	MakeTestPins(TestNode, PinTypes);
 
+	#define TestMatchingFunc(OpName, TestPins, ExpectedFuncName) \
+	{\
+		const UFunction* FoundFunc = FTypePromotion::FindBestMatchingFunc(OpName, TestPins);\
+		const FName ExpectedName = FName(ExpectedFuncName);\
+		const FString TestNullMessage = FString::Printf(TEXT(" Find Function '%s' null check"), *ExpectedName.ToString()); \
+		if (TestNotNull(TestNullMessage, FoundFunc)) \
+		{\
+		 	const FString PinTypesString = TypePromoTestUtils::GetPinListDisplayName(TestPins);\
+			const FString TestMessage = FString::Printf(TEXT("Given pins %s Expecting Function '%s' and got '%s'"),\
+				*PinTypesString,\
+				*ExpectedName.ToString(),\
+				*FoundFunc->GetFName().ToString());\
+				TestEqual(TestMessage, FoundFunc->GetFName(), ExpectedName);\
+		}\
+	}
+
 	{
 		TArray<UEdGraphPin*> TestPins =
 		{
 			Vec2DOutputPinA,
 		};
-
-		const UFunction* AddVecFunc = FTypePromotion::FindBestMatchingFunc(TEXT("Add"), TestPins);
-		static const FName ExpectedName = TEXT("Add_Vector2DVector2D");
-
-		if (TestNotNull(TEXT("Add_Vector2DVector2D Null check"), AddVecFunc))
-		{
-			TestEqual(TEXT("Add_Vector2DVector2D Name Check"), AddVecFunc->GetFName(), ExpectedName);
-		}
+		TestMatchingFunc(TEXT("Add"), TestPins, TEXT("Add_Vector2DVector2D"));
 	}
 
 	// Multiply_VectorVector given A float input, vector input, and a vector output
@@ -219,14 +244,7 @@ bool FFindBestMatchingFunc::RunTest(const FString& Parameters)
 		{
 			FloatPinA, VecInputPinB, VecOutputPinA,
 		};
-
-		const UFunction* MulVecFunc = FTypePromotion::FindBestMatchingFunc(TEXT("Multiply"), TestPins);
-		static const FName ExpectedName = TEXT("Multiply_VectorVector");
-
-		if (TestNotNull(TEXT("Multiply_VectorVector Null check"), MulVecFunc))
-		{
-			TestEqual(TEXT("Multiply_VectorVector Name Check"), MulVecFunc->GetFName(), ExpectedName);
-		}
+		TestMatchingFunc(TEXT("Multiply"), TestPins, TEXT("Multiply_VectorVector"));
 	}
 
 	// Multiply_VectorVector given a float, vector, float
@@ -236,14 +254,7 @@ bool FFindBestMatchingFunc::RunTest(const FString& Parameters)
 		{
 			FloatPinA, VecOutputPinA, VecInputPinA,
 		};
-
-		const UFunction* MulVecFunc = FTypePromotion::FindBestMatchingFunc(TEXT("Multiply"), TestPins);
-		static const FName ExpectedName = TEXT("Multiply_VectorVector");
-
-		if (TestNotNull(TEXT("Multiply_VectorVector Null check"), MulVecFunc))
-		{
-			TestEqual(TEXT("Multiply_VectorVector Name Check"), MulVecFunc->GetFName(), ExpectedName);
-		}
+		TestMatchingFunc(TEXT("Multiply"), TestPins, TEXT("Multiply_VectorVector"));
 	}
 
 	// Multiply_VectorVector given two vector inputs and a vector output
@@ -252,14 +263,7 @@ bool FFindBestMatchingFunc::RunTest(const FString& Parameters)
 		{
 			VecInputPinA, VecInputPinB, VecOutputPinA,
 		};
-
-		const UFunction* MulVecFunc = FTypePromotion::FindBestMatchingFunc(TEXT("Multiply"), TestPins);
-		static const FName ExpectedName = TEXT("Multiply_VectorVector");
-
-		if (TestNotNull(TEXT("Multiply_VectorVector Null check"), MulVecFunc))
-		{
-			TestEqual(TEXT("Multiply_VectorVector Name Check"), MulVecFunc->GetFName(), ExpectedName);
-		}
+		TestMatchingFunc(TEXT("Multiply"), TestPins, TEXT("Multiply_VectorVector"));
 	}
 
 	// Add_DoubleDouble
@@ -268,30 +272,15 @@ bool FFindBestMatchingFunc::RunTest(const FString& Parameters)
 		{
 			DoublePinA, DoublePinB, DoubleOutputPin
 		};
-
-		const UFunction* AddDoubleFunc = FTypePromotion::FindBestMatchingFunc(TEXT("Add"), TestPins);
-		static const FName ExpectedName = TEXT("Add_DoubleDouble");
-		
-		if (TestNotNull(TEXT("Add_DoubleDouble Null check"), AddDoubleFunc))
-		{
-			TestEqual(TEXT("Add_DoubleDouble Name Check"), AddDoubleFunc->GetFName(), ExpectedName);
-		}
+		TestMatchingFunc(TEXT("Add"), TestPins, TEXT("Add_DoubleDouble"));
 	}
 
-	// Add_DoubleDouble given a double and float
 	{
 		TArray<UEdGraphPin*> TestPins =
 		{
 			DoublePinA, FloatPinA, DoubleOutputPin
 		};
-
-		const UFunction* AddDoubleFunc = FTypePromotion::FindBestMatchingFunc(TEXT("Add"), TestPins);
-		static const FName ExpectedName = TEXT("Add_DoubleDouble");
-
-		if (TestNotNull(TEXT("Add_DoubleDouble Null check"), AddDoubleFunc))
-		{
-			TestEqual(TEXT("Add_DoubleDouble Name Check"), AddDoubleFunc->GetFName(), ExpectedName);
-		}
+		TestMatchingFunc(TEXT("Add"), TestPins, TEXT("Add_DoubleDouble"));
 	}
 
 	// Subtract_FloatFloat
@@ -300,14 +289,7 @@ bool FFindBestMatchingFunc::RunTest(const FString& Parameters)
 		{
 			FloatPinA, FloatPinB, FloatOutputPin
 		};
-
-		const UFunction* SubtractFloatFunc = FTypePromotion::FindBestMatchingFunc(TEXT("Subtract"), TestPins);
-		static const FName ExpectedName = TEXT("Subtract_FloatFloat");
-
-		if (TestNotNull(TEXT("Subtract_FloatFloat null check"), SubtractFloatFunc))
-		{
-			TestEqual(TEXT("Subtract_FloatFloat Name Check"), SubtractFloatFunc->GetFName(), ExpectedName);
-		}
+		TestMatchingFunc(TEXT("Subtract"), TestPins, TEXT("Subtract_FloatFloat"));
 	}
 
 	// Add_FloatFloat given only one float pin. This simulates the first connection being made to a 
@@ -317,14 +299,7 @@ bool FFindBestMatchingFunc::RunTest(const FString& Parameters)
 		{
 			FloatPinA,
 		};
-
-		const UFunction* AddFloatFunc = FTypePromotion::FindBestMatchingFunc(TEXT("Add"), TestPins);
-		static const FName ExpectedName = TEXT("Add_FloatFloat");
-
-		if (TestNotNull(TEXT("Add_FloatFloat null check"), AddFloatFunc))
-		{
-			TestEqual(TEXT("Add_FloatFloat Name Check"), AddFloatFunc->GetFName(), ExpectedName);
-		}
+		TestMatchingFunc(TEXT("Add"), TestPins, TEXT("Add_FloatFloat"));
 	}
 
 	// Less_FloatFloat Given a Float and Double
@@ -333,14 +308,7 @@ bool FFindBestMatchingFunc::RunTest(const FString& Parameters)
 		{
 			FloatPinA, BoolOutputPin
 		};
-
-		const UFunction* LessFunc = FTypePromotion::FindBestMatchingFunc(TEXT("Less"), TestPins);
-		static const FName ExpectedName = TEXT("Less_FloatFloat");
-
-		if (TestNotNull(TEXT("Less_FloatFloat null check"), LessFunc))
-		{
-			TestEqual(TEXT("Less_FloatFloat Name Check"), LessFunc->GetFName(), ExpectedName);
-		}
+		TestMatchingFunc(TEXT("Less"), TestPins, TEXT("Less_FloatFloat"));
 	}
 
 	// Less_FloatFloat Given just a single float
@@ -349,14 +317,7 @@ bool FFindBestMatchingFunc::RunTest(const FString& Parameters)
 		{
 			FloatPinA,
 		};
-
-		const UFunction* LessFunc = FTypePromotion::FindBestMatchingFunc(TEXT("Less"), TestPins);
-		static const FName ExpectedName = TEXT("Less_FloatFloat");
-
-		if (TestNotNull(TEXT("Less_FloatFloat null check"), LessFunc))
-		{
-			TestEqual(TEXT("Less_FloatFloat Name Check"), LessFunc->GetFName(), ExpectedName);
-		}
+		TestMatchingFunc(TEXT("Less"), TestPins, TEXT("Less_FloatFloat"));
 	}
 
 	// Greater_DoubleDouble Given a Float and Double
@@ -365,16 +326,8 @@ bool FFindBestMatchingFunc::RunTest(const FString& Parameters)
 		{
 			FloatPinA, DoublePinA
 		};
-
-		const UFunction* GreaterFunc = FTypePromotion::FindBestMatchingFunc(TEXT("Greater"), TestPins);
-		static const FName ExpectedName = TEXT("Greater_DoubleDouble");
-
-		if (TestNotNull(TEXT("Greater_DoubleDouble null check"), GreaterFunc))
-		{
-			TestEqual(TEXT("Greater_DoubleDouble Name Check"), GreaterFunc->GetFName(), ExpectedName);
-		}
+		TestMatchingFunc(TEXT("Greater"), TestPins, TEXT("Greater_DoubleDouble"));
 	}
-
 	
 	TypePromoTestUtils::CleanupTestPins(PinTypes);
 	TestNode->MarkPendingKill();
