Blueprints: When implementing an overriden function in blueprints automatically map the function entry node inputs to the parent node

#rnx

diff --git a/Engine/Source/Editor/BlueprintGraph/Private/EdGraphSchema_K2.cpp b/Engine/Source/Editor/BlueprintGraph/Private/EdGraphSchema_K2.cpp
--- a/Engine/Source/Editor/BlueprintGraph/Private/EdGraphSchema_K2.cpp
+++ b/Engine/Source/Editor/BlueprintGraph/Private/EdGraphSchema_K2.cpp
@@ -3167,6 +3167,9 @@ void UEdGraphSchema_K2::CreateFunctionGraphTerminators(UEdGraph& Graph, UClass*
 			{
 				NextExec->MakeLinkTo(ParentNodeExec);
 				NextExec = FindExecutionPin(*ParentNode, EGPD_Output);
+
+				// Link any params from the function entry node to the parent node inputs
+				LinkDataPinFromOutputToInput(EntryNode, ParentNode);
 			}
 
 			NextNode = ParentNode;
