Blueprints type promo: Only search for a new function upon expansion if we need to, otherwise the function will have already been identified upon node consturciton. This is a small optimization and a bug fix for when int + float + int are given as inputs.

Add an early out when finding the best matching function where if none of the pins you are searching have a struct, then we will immedidately skip a UFunction with a struct param because it can't possibly be a match. This will help to avoid any mis-matches with complex struct types like quats and vectors.

#rnx

diff --git a/Engine/Source/Editor/BlueprintGraph/Private/BlueprintTypePromotion.cpp b/Engine/Source/Editor/BlueprintGraph/Private/BlueprintTypePromotion.cpp
--- a/Engine/Source/Editor/BlueprintGraph/Private/BlueprintTypePromotion.cpp
+++ b/Engine/Source/Editor/BlueprintGraph/Private/BlueprintTypePromotion.cpp
@@ -289,6 +289,7 @@ UFunction* FTypePromotion::FindBestMatchingFunc_Internal(FName Operation, const
 
 	const bool bIsSinglePin = PinsToConsider.Num() == 1;
 	const bool bIsComparisonOp = GetComparisonOpNames().Contains(Operation);
+	const bool bHasStruct = PinsToConsider.ContainsByPredicate([](const UEdGraphPin* Pin) { return Pin && Pin->PinType.PinCategory == UEdGraphSchema_K2::PC_Struct; });
 
 	const UEdGraphSchema_K2* Schema = GetDefault<UEdGraphSchema_K2>();
 		
@@ -313,6 +314,12 @@ UFunction* FTypePromotion::FindBestMatchingFunc_Internal(FName Operation, const
 			FEdGraphPinType ParamType;
 			if (Schema->ConvertPropertyToPinType(Param, /* out */ ParamType))
 			{
+				// Don't bother with this function if there is a struct param, if no pins have any structs
+				if (ParamType.PinCategory == UEdGraphSchema_K2::PC_Struct && !bHasStruct)
+				{
+					break;
+				}
+
 				for (const UEdGraphPin* Pin : PinsToConsider)
 				{
 					// Give a point for each function parameter that matches up with a pin to consider
diff --git a/Engine/Source/Editor/BlueprintGraph/Private/K2Node_PromotableOperator.cpp b/Engine/Source/Editor/BlueprintGraph/Private/K2Node_PromotableOperator.cpp
--- a/Engine/Source/Editor/BlueprintGraph/Private/K2Node_PromotableOperator.cpp
+++ b/Engine/Source/Editor/BlueprintGraph/Private/K2Node_PromotableOperator.cpp
@@ -280,6 +280,8 @@ void UK2Node_PromotableOperator::ExpandNode(FKismetCompilerContext& CompilerCont
 	// Create cast from original 2 inputs to the first intermediate node
 	{
 		UFunction* BestFunc = OpFunction;
+		// Look for a best matching function if we possibly don't have one
+		if(!BestFunc)
 		{
 			TArray<UEdGraphPin*> PinsToConsider =
 			{
