Fixes problem with not being able to navigate up in event/anim graphs using PageUp.
Adds shortcut of double-click to navigate up. Disabled by default for now - see the BlueprintEditor settings to enable & send feedback

diff --git a/Engine/Source/Editor/BlueprintGraph/Private/BlueprintEditorSettings.cpp b/Engine/Source/Editor/BlueprintGraph/Private/BlueprintEditorSettings.cpp
--- a/Engine/Source/Editor/BlueprintGraph/Private/BlueprintEditorSettings.cpp
+++ b/Engine/Source/Editor/BlueprintGraph/Private/BlueprintEditorSettings.cpp
@@ -29,6 +29,7 @@ UBlueprintEditorSettings::UBlueprintEditorSettings(const FObjectInitializer& Obj
 	, bHideConstructionScriptComponentsInDetailsView(true)
 	, bHostFindInBlueprintsInGlobalTab(true)
 	, bNavigateToNativeFunctionsFromCallNodes(true)
+	, bDoubleClickNavigatesToParent(false)
 	// Experimental
 	, bFavorPureCastNodes(false)
 	// Compiler Settings
diff --git a/Engine/Source/Editor/BlueprintGraph/Public/BlueprintEditorSettings.h b/Engine/Source/Editor/BlueprintGraph/Public/BlueprintEditorSettings.h
--- a/Engine/Source/Editor/BlueprintGraph/Public/BlueprintEditorSettings.h
+++ b/Engine/Source/Editor/BlueprintGraph/Public/BlueprintEditorSettings.h
@@ -77,67 +77,71 @@ public:
 
 	/** If set will spawn default "ghost" event nodes in new Blueprints, modifiable in the [DefaultEventNodes] section of EditorPerProjectUserSettings */
 	UPROPERTY(EditAnywhere, config, Category = Workflow)
-		bool bSpawnDefaultBlueprintNodes;
+	bool bSpawnDefaultBlueprintNodes;
 
 	/** If set will exclude components added in a Blueprint class Construction Script from the component details view */
 	UPROPERTY(EditAnywhere, config, Category = Workflow)
-		bool bHideConstructionScriptComponentsInDetailsView;
+	bool bHideConstructionScriptComponentsInDetailsView;
 
 	/** If set, the global Find in Blueprints command (CTRL-SHIFT-F) will be hosted in a standalone tab. This tab can remain open after the Blueprint Editor context is closed. */
 	UPROPERTY(EditAnywhere, config, Category = Workflow)
-		bool bHostFindInBlueprintsInGlobalTab;
+	bool bHostFindInBlueprintsInGlobalTab;
 
 	/** If set, double clicking on a call function node will attempt to navigate an open C++ editor to the native source definition */
 	UPROPERTY(EditAnywhere, config, Category = Workflow)
-		bool bNavigateToNativeFunctionsFromCallNodes;
+	bool bNavigateToNativeFunctionsFromCallNodes;
+
+	/** Double click to navigate up to the parent graph */
+	UPROPERTY(config, EditAnywhere, Category = Workflow)
+	bool bDoubleClickNavigatesToParent;
 
 // Experimental
 public:
 	// The list of namespaces to always expose in any Blueprint (local per-user)
 	UPROPERTY(EditAnywhere, config, Category = Experimental)
-		TArray<FString> NamespacesToAlwaysInclude;
+	TArray<FString> NamespacesToAlwaysInclude;
 
 	/** If enabled, then placed cast nodes will default to their "pure" form (meaning: without execution pins). */
 	UPROPERTY(EditAnywhere, config, AdvancedDisplay, Category = Experimental, meta = (DisplayName = "Default to Using Pure Cast Nodes"))
-		bool bFavorPureCastNodes;
+	bool bFavorPureCastNodes;
 
 	// Compiler Settings
 public:
 	/** Determines when to save Blueprints post-compile */
 	UPROPERTY(EditAnywhere, config, Category = Compiler)
-		TEnumAsByte<ESaveOnCompile> SaveOnCompile;
+	TEnumAsByte<ESaveOnCompile> SaveOnCompile;
 
 	/** When enabled, if a blueprint has compiler errors, then the graph will jump and focus on the first node generating an error */
 	UPROPERTY(EditAnywhere, config, Category = Compiler)
-		bool bJumpToNodeErrors;
+	bool bJumpToNodeErrors;
 
 	/** If enabled, nodes can be explicitly disabled via context menu when right-clicking on impure nodes in the Blueprint editor. Disabled nodes will not be compiled, but also will not break existing connections. */
 	UPROPERTY(EditAnywhere, config, Category = Experimental, AdvancedDisplay)
-		bool bAllowExplicitImpureNodeDisabling;
+	bool bAllowExplicitImpureNodeDisabling;
 
 	// Developer Settings
 public:
 	/** If enabled, tooltips on action menu items will show the associated action's signature id (can be used to setup custom favorites menus).*/
 	UPROPERTY(EditAnywhere, config, Category = DeveloperTools)
-		bool bShowActionMenuItemSignatures;
+	bool bShowActionMenuItemSignatures;
 
 	/** If enabled, blueprint nodes in the event graph will display with unique names rather than their display name. */
 	UPROPERTY(EditAnywhere, config, Category = DeveloperTools, meta = (DisplayName = "Display Unique Names for Blueprint Nodes"))
-		bool bBlueprintNodeUniqueNames;
+	bool bBlueprintNodeUniqueNames;
 
 	// Perf Settings
 public:
 	/** If enabled, additional details will be displayed in the Compiler Results tab after compiling a blueprint. */
 	UPROPERTY(EditAnywhere, config, Category = Performance)
-		bool bShowDetailedCompileResults;
+	bool bShowDetailedCompileResults;
 
 	/** Minimum event time threshold used as a filter when additional details are enabled for display in the Compiler Results tab. A value of zero means that all events will be included in the final summary. */
 	UPROPERTY(EditAnywhere, AdvancedDisplay, config, Category = Performance, DisplayName = "Compile Event Results Threshold (ms)", meta = (ClampMin = "0", UIMin = "0"))
-		int32 CompileEventDisplayThresholdMs;
+	int32 CompileEventDisplayThresholdMs;
 
 	/** The node template cache is used to speed up blueprint menuing. This determines the peak data size for that cache. */
 	UPROPERTY(EditAnywhere, AdvancedDisplay, config, Category = Performance, DisplayName = "Node-Template Cache Cap (MB)", meta = (ClampMin = "0", UIMin = "0"))
-		float NodeTemplateCacheCapMB;
+	float NodeTemplateCacheCapMB;
 
 public:
 	/** If set we'll show the inherited variables in the My Blueprint view. */
diff --git a/Engine/Source/Editor/GraphEditor/Private/ConnectionDrawingPolicy.cpp b/Engine/Source/Editor/GraphEditor/Private/ConnectionDrawingPolicy.cpp
--- a/Engine/Source/Editor/GraphEditor/Private/ConnectionDrawingPolicy.cpp
+++ b/Engine/Source/Editor/GraphEditor/Private/ConnectionDrawingPolicy.cpp
@@ -242,8 +242,9 @@ void FConnectionDrawingPolicy::DrawConnection(int32 LayerId, const FVector2D& St
 		// Distance to consider as an overlap
 		const float QueryDistanceTriggerThresholdSquared = FMath::Square(Settings->SplineHoverTolerance + Params.WireThickness * 0.5f);
 
-		// Distance to pass the bounding box cull test (may want to expand this later on if we want to do 'closest pin' actions that don't require an exact hit)
-		const float QueryDistanceToBoundingBoxSquared = QueryDistanceTriggerThresholdSquared;
+		// Distance to pass the bounding box cull test. This is used for the bCloseToSpline output that can be used as a
+		// dead zone to avoid mistakes caused by missing a double-click on a connection.
+		const float QueryDistanceForCloseSquared = FMath::Square(FMath::Sqrt(QueryDistanceTriggerThresholdSquared) + Settings->SplineCloseTolerance);
 
 		bool bCloseToSpline = false;
 		{
@@ -257,7 +258,7 @@ void FConnectionDrawingPolicy::DrawConnection(int32 LayerId, const FVector2D& St
 			Bounds += FVector2D(P1);
 			Bounds += FVector2D(P1 - MaximumTangentContribution * P1Tangent);
 
-			bCloseToSpline = Bounds.ComputeSquaredDistanceToPoint(LocalMousePosition) < QueryDistanceToBoundingBoxSquared;
+			bCloseToSpline = Bounds.ComputeSquaredDistanceToPoint(LocalMousePosition) < QueryDistanceForCloseSquared;
 
 			// Draw the bounding box for debugging
 #if 0
@@ -313,9 +314,13 @@ void FConnectionDrawingPolicy::DrawConnection(int32 LayerId, const FVector2D& St
 					const float SquaredDistToPin1 = (Params.AssociatedPin1 != nullptr) ? (P0 - ClosestPoint).SizeSquared() : FLT_MAX;
 					const float SquaredDistToPin2 = (Params.AssociatedPin2 != nullptr) ? (P1 - ClosestPoint).SizeSquared() : FLT_MAX;
 
-					SplineOverlapResult = FGraphSplineOverlapResult(Params.AssociatedPin1, Params.AssociatedPin2, ClosestDistanceSquared, SquaredDistToPin1, SquaredDistToPin2);
+					SplineOverlapResult = FGraphSplineOverlapResult(Params.AssociatedPin1, Params.AssociatedPin2, ClosestDistanceSquared, SquaredDistToPin1, SquaredDistToPin2, true);
 				}
 			}
+			else if (ClosestDistanceSquared < QueryDistanceForCloseSquared)
+			{
+				SplineOverlapResult.SetCloseToSpline(true);
+			}
 		}
 	}
 
diff --git a/Engine/Source/Editor/GraphEditor/Private/GraphEditorSettings.cpp b/Engine/Source/Editor/GraphEditor/Private/GraphEditorSettings.cpp
--- a/Engine/Source/Editor/GraphEditor/Private/GraphEditorSettings.cpp
+++ b/Engine/Source/Editor/GraphEditor/Private/GraphEditorSettings.cpp
@@ -18,6 +18,7 @@ UGraphEditorSettings::UGraphEditorSettings( const FObjectInitializer& ObjectInit
 	, PaddingTowardsNodeEdge(10.0f)
 	, bTreatSplinesLikePins(true)
 	, SplineHoverTolerance(2.0f)
+	, SplineCloseTolerance(6.0f)
 	, ForwardSplineHorizontalDeltaRange(1000.0f)
 	, ForwardSplineVerticalDeltaRange(1000.0f)
 	, ForwardSplineTangentFromHorizontalDelta(1.0f, 0.0f)
diff --git a/Engine/Source/Editor/GraphEditor/Private/SGraphEditorImpl.cpp b/Engine/Source/Editor/GraphEditor/Private/SGraphEditorImpl.cpp
--- a/Engine/Source/Editor/GraphEditor/Private/SGraphEditorImpl.cpp
+++ b/Engine/Source/Editor/GraphEditor/Private/SGraphEditorImpl.cpp
@@ -681,6 +681,7 @@ void SGraphEditorImpl::Construct( const FArguments& InArgs )
 			//.OnUpdateGraphPanel( this, &SGraphEditorImpl::GraphEd_OnPanelUpdated )
 			.OnDisallowedPinConnection( InArgs._GraphEvents.OnDisallowedPinConnection )
 			.ShowGraphStateOverlay(InArgs._ShowGraphStateOverlay)
+			.OnDoubleClicked(InArgs._GraphEvents.OnDoubleClicked)
 		]
 
 		// Indicator of current zoom level
diff --git a/Engine/Source/Editor/GraphEditor/Private/SGraphPanel.cpp b/Engine/Source/Editor/GraphEditor/Private/SGraphPanel.cpp
--- a/Engine/Source/Editor/GraphEditor/Private/SGraphPanel.cpp
+++ b/Engine/Source/Editor/GraphEditor/Private/SGraphPanel.cpp
@@ -57,6 +57,7 @@ void SGraphPanel::Construct( const SGraphPanel::FArguments& InArgs )
 	this->OnSpawnNodeByShortcut = InArgs._OnSpawnNodeByShortcut;
 	this->OnUpdateGraphPanel = InArgs._OnUpdateGraphPanel;
 	this->OnDisallowedPinConnection = InArgs._OnDisallowedPinConnection;
+	this->OnDoubleClicked = InArgs._OnDoubleClicked;
 
 	this->bPreservePinPreviewConnection = false;
 	this->PinVisibility = SGraphEditor::Pin_Show;
@@ -442,7 +443,7 @@ int32 SGraphPanel::OnPaint( const FPaintArgs& Args, const FGeometry& AllottedGeo
 					if (CommentNode == nullptr)
 					{
 						// Wasn't a comment node, disallow the spline interaction
-						OverlapData = FGraphSplineOverlapResult();
+						OverlapData = FGraphSplineOverlapResult(OverlapData.GetCloseToSpline());
 					}
 				}
 			}
@@ -667,6 +668,10 @@ FReply SGraphPanel::OnMouseButtonDoubleClick(const FGeometry& MyGeometry, const
 		const UEdGraphSchema* Schema = GraphObj->GetSchema();
 		Schema->OnPinConnectionDoubleCicked(Pin1, Pin2, DoubleClickPositionInGraphSpace);
 	}
+	else if (!PreviousFrameSplineOverlap.GetCloseToSpline())
+	{
+		OnDoubleClicked.ExecuteIfBound();
+	}
 
 	return SNodePanel::OnMouseButtonDoubleClick(MyGeometry, MouseEvent);
 }
diff --git a/Engine/Source/Editor/GraphEditor/Public/GraphEditorSettings.h b/Engine/Source/Editor/GraphEditor/Public/GraphEditorSettings.h
--- a/Engine/Source/Editor/GraphEditor/Public/GraphEditorSettings.h
+++ b/Engine/Source/Editor/GraphEditor/Public/GraphEditorSettings.h
@@ -44,7 +44,7 @@ public:
 
 	/** Switch between using the right and middle mouse button for panning (defaults to right) */
 	UPROPERTY(config, EditAnywhere, Category = GeneralStyle)
-		EGraphPanningMouseButton PanningMouseButton;
+	EGraphPanningMouseButton PanningMouseButton;
 
 	/** The amount of padding above a pin (defaults to 4) */
 	UPROPERTY(config, EditAnywhere, Category=GeneralStyle, AdvancedDisplay)
@@ -88,6 +88,10 @@ public:
 	UPROPERTY(EditAnywhere, config, Category=Splines, meta=(EditCondition=bTreatSplinesLikePins))
 	float SplineHoverTolerance;
 
+	/** The additional distance around the spline to count as close, preventing other actions if the user just misses the spline when clicking*/
+	UPROPERTY(EditAnywhere, config, Category = Splines, meta = (EditCondition = bTreatSplinesLikePins))
+	float SplineCloseTolerance;
+
 	/** The maximum value to clamp the absolute value of the horizontal distance between endpoints when calculating tangents (when the wire is moving forward) */
 	UPROPERTY(config, EditAnywhere, Category=Splines, AdvancedDisplay)
 	float ForwardSplineHorizontalDeltaRange;
diff --git a/Engine/Source/Editor/GraphEditor/Public/GraphSplineOverlapResult.h b/Engine/Source/Editor/GraphEditor/Public/GraphSplineOverlapResult.h
--- a/Engine/Source/Editor/GraphEditor/Public/GraphSplineOverlapResult.h
+++ b/Engine/Source/Editor/GraphEditor/Public/GraphSplineOverlapResult.h
@@ -11,7 +11,7 @@
 struct GRAPHEDITOR_API FGraphSplineOverlapResult
 {
 public:
-	FGraphSplineOverlapResult()
+	FGraphSplineOverlapResult(bool InCloseToSpline = false)
 		: Pin1Handle(nullptr)
 		, Pin2Handle(nullptr)
 		, BestPinHandle(nullptr)
@@ -20,10 +20,11 @@ public:
 		, DistanceSquared(FLT_MAX)
 		, DistanceSquaredToPin1(FLT_MAX)
 		, DistanceSquaredToPin2(FLT_MAX)
+		, bCloseToSpline(InCloseToSpline)
 	{
 	}
 
-	FGraphSplineOverlapResult(UEdGraphPin* InPin1, UEdGraphPin* InPin2, float InDistanceSquared, float InDistanceSquaredToPin1, float InDistanceSquaredToPin2)
+	FGraphSplineOverlapResult(UEdGraphPin* InPin1, UEdGraphPin* InPin2, float InDistanceSquared, float InDistanceSquaredToPin1, float InDistanceSquaredToPin2, bool InCloseToSpline)
 		: Pin1Handle(InPin1)
 		, Pin2Handle(InPin2)
 		, BestPinHandle(nullptr)
@@ -32,6 +33,7 @@ public:
 		, DistanceSquared(InDistanceSquared)
 		, DistanceSquaredToPin1(InDistanceSquaredToPin1)
 		, DistanceSquaredToPin2(InDistanceSquaredToPin2)
+		, bCloseToSpline(InCloseToSpline)
 	{
 	}
 
@@ -47,6 +49,16 @@ public:
 		return DistanceSquared;
 	}
 
+	bool GetCloseToSpline() const
+	{
+		return bCloseToSpline;
+	}
+
+	void SetCloseToSpline(bool InCloseToSpline) 
+	{
+		bCloseToSpline = InCloseToSpline;
+	}
+
 	TSharedPtr<class SGraphPin> GetBestPinWidget(const class SGraphPanel& InGraphPanel) const
 	{
 		TSharedPtr<class SGraphPin> Result;
@@ -68,4 +80,5 @@ protected:
 	float DistanceSquared;
 	float DistanceSquaredToPin1;
 	float DistanceSquaredToPin2;
+	bool  bCloseToSpline;
 };
diff --git a/Engine/Source/Editor/GraphEditor/Public/SGraphPanel.h b/Engine/Source/Editor/GraphEditor/Public/SGraphPanel.h
--- a/Engine/Source/Editor/GraphEditor/Public/SGraphPanel.h
+++ b/Engine/Source/Editor/GraphEditor/Public/SGraphPanel.h
@@ -81,6 +81,7 @@ public:
 		SLATE_EVENT( SGraphEditor::FOnSpawnNodeByShortcut, OnSpawnNodeByShortcut )
 		SLATE_EVENT( FOnUpdateGraphPanel, OnUpdateGraphPanel )
 		SLATE_EVENT( SGraphEditor::FOnDisallowedPinConnection, OnDisallowedPinConnection )
+		SLATE_EVENT( SGraphEditor::FOnDoubleClicked, OnDoubleClicked )
 		//SLATE_ATTRIBUTE( FGraphAppearanceInfo, Appearance )
 	SLATE_END_ARGS()
 
@@ -277,6 +278,9 @@ protected:
 	/** Called when the user generates a warning tooltip because a connection was invalid */
 	SGraphEditor::FOnDisallowedPinConnection OnDisallowedPinConnection;
 
+	/** Called when the graph itself is double clicked */
+	SGraphEditor::FOnDoubleClicked OnDoubleClicked;
+
 	/** Whether to draw the overlay indicating we're in PIE */
 	bool bShowPIENotification;
 
diff --git a/Engine/Source/Editor/Kismet/Private/BlueprintEditor.cpp b/Engine/Source/Editor/Kismet/Private/BlueprintEditor.cpp
--- a/Engine/Source/Editor/Kismet/Private/BlueprintEditor.cpp
+++ b/Engine/Source/Editor/Kismet/Private/BlueprintEditor.cpp
@@ -1488,6 +1488,7 @@ void FBlueprintEditor::SetupGraphEditorEvents(UEdGraph* InGraph, SGraphEditor::F
 	InEvents.OnSpawnNodeByShortcut = SGraphEditor::FOnSpawnNodeByShortcut::CreateSP(this, &FBlueprintEditor::OnSpawnGraphNodeByShortcut, InGraph);
 	InEvents.OnNodeSpawnedByKeymap = SGraphEditor::FOnNodeSpawnedByKeymap::CreateSP(this, &FBlueprintEditor::OnNodeSpawnedByKeymap );
 	InEvents.OnDisallowedPinConnection = SGraphEditor::FOnDisallowedPinConnection::CreateSP(this, &FBlueprintEditor::OnDisallowedPinConnection);
+	InEvents.OnDoubleClicked = SGraphEditor::FOnDoubleClicked::CreateSP(this, &FBlueprintEditor::NavigateToParentGraphByDoubleClick);
 
 	// Custom menu for K2 schemas
 	if(InGraph->Schema != NULL && InGraph->Schema->IsChildOf(UEdGraphSchema_K2::StaticClass()))
@@ -2735,17 +2736,17 @@ void FBlueprintEditor::CreateDefaultCommands()
 		);
 
 	ToolkitCommands->MapAction( FBlueprintEditorCommands::Get().NavigateToParent,
-		FExecuteAction::CreateSP( this, &FBlueprintEditor::NavigateToParentGraph_Clicked ),
+		FExecuteAction::CreateSP( this, &FBlueprintEditor::NavigateToParentGraph ),
 		FCanExecuteAction::CreateSP( this, &FBlueprintEditor::CanNavigateToParentGraph )
 		);
 
 	ToolkitCommands->MapAction( FBlueprintEditorCommands::Get().NavigateToParentBackspace,
-		FExecuteAction::CreateSP( this, &FBlueprintEditor::NavigateToParentGraph_Clicked ),
+		FExecuteAction::CreateSP( this, &FBlueprintEditor::NavigateToParentGraph ),
 		FCanExecuteAction::CreateSP( this, &FBlueprintEditor::CanNavigateToParentGraph )
 		);
 
 	ToolkitCommands->MapAction( FBlueprintEditorCommands::Get().NavigateToChild,
-		FExecuteAction::CreateSP( this, &FBlueprintEditor::NavigateToChildGraph_Clicked ),
+		FExecuteAction::CreateSP( this, &FBlueprintEditor::NavigateToChildGraph ),
 		FCanExecuteAction::CreateSP( this, &FBlueprintEditor::CanNavigateToChildGraph )
 		);
 
@@ -3173,11 +3174,20 @@ bool FBlueprintEditor::CanZoomToSelection() const
 }
 
 // Navigating into/out of graphs
-void FBlueprintEditor::NavigateToParentGraph_Clicked()
+void FBlueprintEditor::NavigateToParentGraphByDoubleClick()
+{
+	UBlueprintEditorSettings const* Settings = GetDefault<UBlueprintEditorSettings>();
+	if (Settings->bDoubleClickNavigatesToParent)
+	{
+		NavigateToParentGraph();
+	}
+}
+
+void FBlueprintEditor::NavigateToParentGraph()
 {
 	if (FocusedGraphEdPtr.IsValid())
 	{
-		if (UEdGraph* ParentGraph = Cast<UEdGraph>(FocusedGraphEdPtr.Pin()->GetCurrentGraph()->GetOuter()))
+		if (UEdGraph* ParentGraph = UEdGraph::GetOuterGraph(FocusedGraphEdPtr.Pin()->GetCurrentGraph()))
 		{
 			OpenDocument(ParentGraph, FDocumentTracker::NavigatingCurrentDocument);
 		}
@@ -3186,10 +3196,10 @@ void FBlueprintEditor::NavigateToParentGraph_Clicked()
 
 bool FBlueprintEditor::CanNavigateToParentGraph() const
 {
-	return FocusedGraphEdPtr.IsValid() && (FocusedGraphEdPtr.Pin()->GetCurrentGraph()->GetOuter()->IsA(UEdGraph::StaticClass()));
+	return FocusedGraphEdPtr.IsValid() && UEdGraph::GetOuterGraph(FocusedGraphEdPtr.Pin()->GetCurrentGraph());
 }
 
-void FBlueprintEditor::NavigateToChildGraph_Clicked()
+void FBlueprintEditor::NavigateToChildGraph()
 {
 	if (FocusedGraphEdPtr.IsValid())
 	{
diff --git a/Engine/Source/Editor/Kismet/Public/BlueprintEditor.h b/Engine/Source/Editor/Kismet/Public/BlueprintEditor.h
--- a/Engine/Source/Editor/Kismet/Public/BlueprintEditor.h
+++ b/Engine/Source/Editor/Kismet/Public/BlueprintEditor.h
@@ -667,9 +667,10 @@ protected:
 	bool CanZoomToSelection() const;
 
 	// Navigating into/out of graphs
-	void NavigateToParentGraph_Clicked();
+	void NavigateToParentGraph();
+	void NavigateToParentGraphByDoubleClick();
 	bool CanNavigateToParentGraph() const;
-	void NavigateToChildGraph_Clicked();
+	void NavigateToChildGraph();
 	bool CanNavigateToChildGraph() const;
 
 	/** Determines visibility of the parent class manipulation buttons on the menu bar overlay */
diff --git a/Engine/Source/Editor/UnrealEd/Public/GraphEditor.h b/Engine/Source/Editor/UnrealEd/Public/GraphEditor.h
--- a/Engine/Source/Editor/UnrealEd/Public/GraphEditor.h
+++ b/Engine/Source/Editor/UnrealEd/Public/GraphEditor.h
@@ -96,6 +96,8 @@ public:
 
 	DECLARE_DELEGATE_TwoParams( FOnDisallowedPinConnection, const UEdGraphPin*, const UEdGraphPin* );
 
+	DECLARE_DELEGATE( FOnDoubleClicked );
+
 	/** Info about events occurring in/on the graph */
 	struct FGraphEditorEvents
 	{
@@ -123,6 +125,8 @@ public:
 		FOnNodeSpawnedByKeymap OnNodeSpawnedByKeymap;
 		/** Called when the user generates a warning tooltip because a connection was invalid */
 		FOnDisallowedPinConnection OnDisallowedPinConnection;
+		/** Called when the graph itself is double clicked */
+		FOnDoubleClicked OnDoubleClicked;
 	};
 
 
