Bluelprints: Enable type promotion by default!

Change from using a CVar to a BP editor preferences setting for a better user experience now that it is feature complete with no known bugs.

#rnx

diff --git a/Engine/Source/Editor/BlueprintGraph/Classes/BlueprintTypePromotion.h b/Engine/Source/Editor/BlueprintGraph/Classes/BlueprintTypePromotion.h
--- a/Engine/Source/Editor/BlueprintGraph/Classes/BlueprintTypePromotion.h
+++ b/Engine/Source/Editor/BlueprintGraph/Classes/BlueprintTypePromotion.h
@@ -7,6 +7,7 @@
 #include "BlueprintActionDatabase.h"
 #include "BlueprintFunctionNodeSpawner.h"
 #include "HAL/IConsoleManager.h"
+#include "BlueprintEditorSettings.h"
 
 /** 
 * Contains behavior needed to handle type promotion in blueprints. 
@@ -184,10 +185,12 @@ private:
 
 namespace TypePromoDebug
 {
-	/** Checks if the CVar for type promotion is true or false (BP.TypePromo.IsEnabled) */
+	/** 
+	* Checks if the type promotion editor pref is true or false
+	* (The bEnableTypePromotion BP editor setting)
+	*/
 	static bool IsTypePromoEnabled()
 	{
-		static IConsoleVariable* TypePromoCVar = IConsoleManager::Get().FindConsoleVariable(TEXT("BP.TypePromo.IsEnabled"));
-		return TypePromoCVar ? TypePromoCVar->GetBool() : false;
+		return GetDefault<UBlueprintEditorSettings>()->bEnableTypePromotion;
 	}
 }
\ No newline at end of file
diff --git a/Engine/Source/Editor/BlueprintGraph/Private/BlueprintEditorSettings.cpp b/Engine/Source/Editor/BlueprintGraph/Private/BlueprintEditorSettings.cpp
--- a/Engine/Source/Editor/BlueprintGraph/Private/BlueprintEditorSettings.cpp
+++ b/Engine/Source/Editor/BlueprintGraph/Private/BlueprintEditorSettings.cpp
@@ -8,6 +8,7 @@
 #include "FindInBlueprintManager.h"
 #include "Subsystems/AssetEditorSubsystem.h"
 #include "Editor.h"
+#include "BlueprintTypePromotion.h"
 
 UBlueprintEditorSettings::UBlueprintEditorSettings(const FObjectInitializer& ObjectInitializer)
 	: Super(ObjectInitializer)
@@ -30,6 +31,7 @@ UBlueprintEditorSettings::UBlueprintEditorSettings(const FObjectInitializer& Obj
 	, bHostFindInBlueprintsInGlobalTab(true)
 	, bNavigateToNativeFunctionsFromCallNodes(true)
 	, bDoubleClickNavigatesToParent(false)
+	, bEnableTypePromotion(true)
 	// Experimental
 	, bFavorPureCastNodes(false)
 	// Compiler Settings
@@ -97,6 +99,14 @@ void UBlueprintEditorSettings::PostEditChangeProperty(FPropertyChangedEvent& Pro
 		bShouldRebuildRegistry = true;
 	}
 
+	// Refresh type promotion when the preference gets changed so that we can correctly rebuild the action database
+	if (PropertyName == GET_MEMBER_NAME_CHECKED(UBlueprintEditorSettings, bEnableTypePromotion))
+	{
+		FTypePromotion::ClearNodeSpawners();
+
+		bShouldRebuildRegistry = true;
+	}
+
 	if (bShouldRebuildRegistry)
 	{
 		FBlueprintActionDatabase::Get().RefreshAll();
diff --git a/Engine/Source/Editor/BlueprintGraph/Private/BlueprintTypePromotion.cpp b/Engine/Source/Editor/BlueprintGraph/Private/BlueprintTypePromotion.cpp
--- a/Engine/Source/Editor/BlueprintGraph/Private/BlueprintTypePromotion.cpp
+++ b/Engine/Source/Editor/BlueprintGraph/Private/BlueprintTypePromotion.cpp
@@ -26,26 +26,6 @@ namespace OperatorNames
 	static const FName Equal		= TEXT("EqualEqual");
 }
 
-namespace TypePromoDebug
-{
-	static bool bIsTypePromoEnabled = false;
-	static FAutoConsoleVariableRef CVarIsTypePromoEnabled(
-		TEXT("BP.TypePromo.IsEnabled"), bIsTypePromoEnabled,
-		TEXT("If true then type promotion inside of blueprints will be enabled"),
-		FConsoleVariableDelegate::CreateLambda([](IConsoleVariable* InVariable)
-		{
-			// Clear the node spawner so that we create the new BP actions correctly
-			FTypePromotion::ClearNodeSpawners();
-
-			// Refresh all the actions so that the context menu goes back to the normal options
-			if (FBlueprintActionDatabase* Actions = FBlueprintActionDatabase::TryGet())
-			{
-				Actions->RefreshAll();
-			}
-		}),
-		ECVF_Default);
-}
-
 FTypePromotion& FTypePromotion::Get()
 {
 	if (Instance == nullptr)
diff --git a/Engine/Source/Editor/BlueprintGraph/Public/BlueprintEditorSettings.h b/Engine/Source/Editor/BlueprintGraph/Public/BlueprintEditorSettings.h
--- a/Engine/Source/Editor/BlueprintGraph/Public/BlueprintEditorSettings.h
+++ b/Engine/Source/Editor/BlueprintGraph/Public/BlueprintEditorSettings.h
@@ -95,6 +95,10 @@ public:
 	UPROPERTY(config, EditAnywhere, Category = Workflow)
 	bool bDoubleClickNavigatesToParent;
 
+	/** Allows for pin types to be promoted to others, i.e. float to double */
+	UPROPERTY(config, EditAnywhere, Category = Workflow)
+	bool bEnableTypePromotion;
+
 // Experimental
 public:
 	// The list of namespaces to always expose in any Blueprint (local per-user)
