disable type promotion by default due to being unable to connect actors to the not equal node

#rnx

diff --git a/Engine/Source/Editor/BlueprintGraph/Private/BlueprintEditorSettings.cpp b/Engine/Source/Editor/BlueprintGraph/Private/BlueprintEditorSettings.cpp
--- a/Engine/Source/Editor/BlueprintGraph/Private/BlueprintEditorSettings.cpp
+++ b/Engine/Source/Editor/BlueprintGraph/Private/BlueprintEditorSettings.cpp
@@ -31,7 +31,7 @@ UBlueprintEditorSettings::UBlueprintEditorSettings(const FObjectInitializer& Obj
 	, bHostFindInBlueprintsInGlobalTab(true)
 	, bNavigateToNativeFunctionsFromCallNodes(true)
 	, bDoubleClickNavigatesToParent(false)
-	, bEnableTypePromotion(true)
+	, bEnableTypePromotion(false)
 	// Experimental
 	, bFavorPureCastNodes(false)
 	// Compiler Settings
