Fix inability to find matching functions when searching for an object or class type in type promotion

#rnx

diff --git a/Engine/Source/Editor/BlueprintGraph/Private/BlueprintTypePromotion.cpp b/Engine/Source/Editor/BlueprintGraph/Private/BlueprintTypePromotion.cpp
--- a/Engine/Source/Editor/BlueprintGraph/Private/BlueprintTypePromotion.cpp
+++ b/Engine/Source/Editor/BlueprintGraph/Private/BlueprintTypePromotion.cpp
@@ -302,8 +302,13 @@ UFunction* FTypePromotion::FindBestMatchingFunc_Internal(FName Operation, const
 
 				for (const UEdGraphPin* Pin : PinsToConsider)
 				{
+					// Object types will not match ever the function input params and neither will classes, because their
+					// sub-category will be different. 
+					const bool bObjectTypesMatch = ParamType.PinCategory == UEdGraphSchema_K2::PC_Object && Pin->PinType.PinCategory == UEdGraphSchema_K2::PC_Object;
+					const bool bClassTypesMatch = ParamType.PinCategory == UEdGraphSchema_K2::PC_Class && Pin->PinType.PinCategory == UEdGraphSchema_K2::PC_Class;
+
 					// Give a point for each function parameter that matches up with a pin to consider
-					if ((!CheckedPins.Contains(Pin) || bIsSinglePin) && Schema->ArePinTypesEquivalent(ParamType, Pin->PinType))
+					if ((!CheckedPins.Contains(Pin) || bIsSinglePin) && (bObjectTypesMatch || bClassTypesMatch || Schema->ArePinTypesEquivalent(ParamType, Pin->PinType)))
 					{
 						// Are the directions compatible? 
 						// If we are a comparison or only a single pin then we don't care about the direction
