Blueprints: Improve compiler messaging when an invalid operator is created via type promotion

#rnx

diff --git a/Engine/Source/Editor/BlueprintGraph/Private/K2Node_PromotableOperator.cpp b/Engine/Source/Editor/BlueprintGraph/Private/K2Node_PromotableOperator.cpp
--- a/Engine/Source/Editor/BlueprintGraph/Private/K2Node_PromotableOperator.cpp
+++ b/Engine/Source/Editor/BlueprintGraph/Private/K2Node_PromotableOperator.cpp
@@ -360,6 +360,20 @@ void UK2Node_PromotableOperator::ExpandNode(FKismetCompilerContext& CompilerCont
 	if (OriginalOutputPin && PrevOutputPin)
 	{
 		CompilerContext.MovePinLinksToIntermediate(*OriginalOutputPin, *PrevOutputPin);
+		// If there is no link to the output pin then the connection response called for a conversion node,
+		// but one was not available. This can occur when there is a connection to the output pin that 
+		// is smaller than an input and cannot be casted. We throw a compiler error instead of auto-breaking
+		// pins because that can be confusing to the user.
+		if (PrevOutputPin->LinkedTo.Num() == 0)
+		{
+			CompilerContext.MessageLog.Error(
+				*FText::Format(LOCTEXT("FailedOutputConnection_ErrorFmt",
+				"Output pin type '{1}' is not compatible with input type of '{0}' on '@@'"),
+				Schema->TypeToText(PrevOutputPin->PinType),
+				Schema->TypeToText(OriginalOutputPin->PinType)).ToString(),
+				PrevIntermediateNode
+			);
+		}
 	}
 }
 
