Ensure that comparison operator's input pins are reset when an all inputs are discocnnected, but still allow the bool output to be connected.

#rnx

diff --git a/Engine/Source/Editor/BlueprintGraph/Private/K2Node_PromotableOperator.cpp b/Engine/Source/Editor/BlueprintGraph/Private/K2Node_PromotableOperator.cpp
--- a/Engine/Source/Editor/BlueprintGraph/Private/K2Node_PromotableOperator.cpp
+++ b/Engine/Source/Editor/BlueprintGraph/Private/K2Node_PromotableOperator.cpp
@@ -681,10 +681,25 @@ void UK2Node_PromotableOperator::EvaluatePinsFromChange(UEdGraphPin* ChangedPin)
 
 	// True if the pin that has changed now has zero connections
 	const bool bWasAFullDisconnect = (ChangedPin->LinkedTo.Num() == 0);
+	const bool bIsComparisonOp = FTypePromotion::IsComparisonOpName(OperationName);
+
+	auto HasDeterminingComparisonTypes = [](const TArray<UEdGraphPin*>& PinsToCheck) -> bool
+	{
+		// For comparison ops, we only want to check the input pins as the output
+		// will always be a bool
+		for(const UEdGraphPin* Pin : PinsToCheck)
+		{
+			if (Pin && Pin->Direction == EGPD_Input && (Pin->LinkedTo.Num() > 0 || !Pin->DoesDefaultValueMatchAutogenerated()))
+			{
+				return true;
+			}
+		}
+		return false;
+	};
 
 	// If we have been totally disconnected and don't have any non-default inputs, 
 	// than we just reset the node to be a regular wildcard
-	if (bWasAFullDisconnect && !HasAnyConnectionsOrDefaults())
+	if ((bWasAFullDisconnect && !HasAnyConnectionsOrDefaults()) || (bIsComparisonOp && !HasDeterminingComparisonTypes(Pins)))
 	{
 		ResetNodeToWildcard();
 		return;
@@ -702,7 +717,7 @@ void UK2Node_PromotableOperator::EvaluatePinsFromChange(UEdGraphPin* ChangedPin)
 	}
 	// If we have connected the output of our comparison operator (which is always a bool) then 
 	// there is no need to update the other pins
-	else if (FTypePromotion::IsComparisonOpName(OperationName) && ChangedPin == GetOutputPin())
+	else if (bIsComparisonOp && ChangedPin == GetOutputPin())
 	{
 		return;
 	}
