Make sure that comparison operators with only a bool connection refresh correctly

#rnx

diff --git a/Engine/Source/Editor/BlueprintGraph/Classes/K2Node_PromotableOperator.h b/Engine/Source/Editor/BlueprintGraph/Classes/K2Node_PromotableOperator.h
--- a/Engine/Source/Editor/BlueprintGraph/Classes/K2Node_PromotableOperator.h
+++ b/Engine/Source/Editor/BlueprintGraph/Classes/K2Node_PromotableOperator.h
@@ -78,6 +78,14 @@ private:
 	*/
 	bool HasAnyConnectionsOrDefaults() const;
 
+	/**
+	* @return	True if this is a comparison operator node and there are 
+	*			no connections/default values on input pins. The output pins
+	*			is not considered in this check, because it will always be a bool
+	*			so it does not determine the types.
+	*/
+	bool HasDeterminingComparisonTypes() const;
+
 	/** Recombines all split pins and sets the node to have default values (all wildcard pins) */
 	void ResetNodeToWildcard();
 
diff --git a/Engine/Source/Editor/BlueprintGraph/Private/K2Node_PromotableOperator.cpp b/Engine/Source/Editor/BlueprintGraph/Private/K2Node_PromotableOperator.cpp
--- a/Engine/Source/Editor/BlueprintGraph/Private/K2Node_PromotableOperator.cpp
+++ b/Engine/Source/Editor/BlueprintGraph/Private/K2Node_PromotableOperator.cpp
@@ -665,6 +665,11 @@ bool UK2Node_PromotableOperator::IsAdditionalPin(const UEdGraphPin* Pin) const
 
 bool UK2Node_PromotableOperator::HasAnyConnectionsOrDefaults() const
 {
+	if(FTypePromotion::IsComparisonOpName(OperationName))
+	{
+		return HasDeterminingComparisonTypes();
+	}
+
 	for (UEdGraphPin* Pin : Pins)
 	{
 		if (Pin->LinkedTo.Num() > 0 || !Pin->DoesDefaultValueMatchAutogenerated())
@@ -675,6 +680,25 @@ bool UK2Node_PromotableOperator::HasAnyConnectionsOrDefaults() const
 	return false;
 }
 
+bool UK2Node_PromotableOperator::HasDeterminingComparisonTypes() const
+{
+	if(!FTypePromotion::IsComparisonOpName(OperationName))
+	{
+		return false;
+	}
+
+	// For comparison ops, we only want to check the input pins as the output
+	// will always be a bool
+	for (const UEdGraphPin* Pin : Pins)
+	{
+		if (Pin && Pin->Direction == EGPD_Input && (Pin->LinkedTo.Num() > 0 || !Pin->DoesDefaultValueMatchAutogenerated()))
+		{
+			return true;
+		}
+	}
+	return false;
+}
+
 void UK2Node_PromotableOperator::EvaluatePinsFromChange(UEdGraphPin* ChangedPin)
 {
 	UpdateOpName();
@@ -683,23 +707,9 @@ void UK2Node_PromotableOperator::EvaluatePinsFromChange(UEdGraphPin* ChangedPin)
 	const bool bWasAFullDisconnect = (ChangedPin->LinkedTo.Num() == 0);
 	const bool bIsComparisonOp = FTypePromotion::IsComparisonOpName(OperationName);
 
-	auto HasDeterminingComparisonTypes = [](const TArray<UEdGraphPin*>& PinsToCheck) -> bool
-	{
-		// For comparison ops, we only want to check the input pins as the output
-		// will always be a bool
-		for(const UEdGraphPin* Pin : PinsToCheck)
-		{
-			if (Pin && Pin->Direction == EGPD_Input && (Pin->LinkedTo.Num() > 0 || !Pin->DoesDefaultValueMatchAutogenerated()))
-			{
-				return true;
-			}
-		}
-		return false;
-	};
-
 	// If we have been totally disconnected and don't have any non-default inputs, 
 	// than we just reset the node to be a regular wildcard
-	if ((bWasAFullDisconnect && !HasAnyConnectionsOrDefaults()) || (bIsComparisonOp && !HasDeterminingComparisonTypes(Pins)))
+	if ((bWasAFullDisconnect && !HasAnyConnectionsOrDefaults()) || (bIsComparisonOp && !HasDeterminingComparisonTypes()))
 	{
 		ResetNodeToWildcard();
 		return;
