Type Promotion: Add a conversion option for promotable operators when they are in a wildcard state.

#rnx

diff --git a/Engine/Source/Editor/BlueprintGraph/Classes/K2Node_PromotableOperator.h b/Engine/Source/Editor/BlueprintGraph/Classes/K2Node_PromotableOperator.h
--- a/Engine/Source/Editor/BlueprintGraph/Classes/K2Node_PromotableOperator.h
+++ b/Engine/Source/Editor/BlueprintGraph/Classes/K2Node_PromotableOperator.h
@@ -8,6 +8,7 @@
 #include "K2Node_PromotableOperator.generated.h"
 
 class UEdGraphPin;
+class UToolMenu;
 
 /** The promotable operator node allows for pin types to be promoted to others, i.e. float to double */
 UCLASS(MinimalAPI)
@@ -108,9 +109,10 @@ private:
 	* Re-evaluates the in types on this node based on all the current connections
 	* and the given pin that has changed. 
 	* 
-	* @param ChangedPin		The pin that has been change
+	* @param ChangedPin			The pin that has been change
+	* @param bFromConversion	True if the pin modification is coming from a conversion.
 	*/
-	void EvaluatePinsFromChange(UEdGraphPin* ChangedPin);
+	void EvaluatePinsFromChange(UEdGraphPin* ChangedPin, const bool bFromConversion = false);
 
 	/** Helper to make sure we have the most up to date operation name. Returns true upon success */
 	bool UpdateOpName();
@@ -136,6 +138,26 @@ private:
 	*/
 	void GetPinsToConsider(TArray<UEdGraphPin*>& OutArray) const;
 
+	/**
+	* Convert the given pin to the new pin type. If the new type is wildcard then break 
+	* all connections and reset the node. Invalidates tooltips upon completion. 
+	* This function is to be used by the context menu and includes a transaction. 
+	* 
+	* @param PinToChange		The Pin to convert
+	* @param NewPinType			The pin's new type
+	*/
+	void ConvertPinType(UEdGraphPin* PinToChange, const FEdGraphPinType NewPinType);
+	
+	/**
+	* Build the context menu to convert the given pin to a different type. 
+	* If the pin is a wildcard, offer any possible types. Otherwise, offer 
+	* to reset the node to wildcard.
+	* 
+	* @param Menu			The parent context menu
+	* @param ContextPin		The pin to convert
+	*/
+	void CreateConversionSubMenu(UToolMenu* Menu, UEdGraphPin* ContextPin) const;
+
 	/** The name that this operation uses ("Add", "Multiply", etc) */
 	FName OperationName;
 
diff --git a/Engine/Source/Editor/BlueprintGraph/Private/K2Node_PromotableOperator.cpp b/Engine/Source/Editor/BlueprintGraph/Private/K2Node_PromotableOperator.cpp
--- a/Engine/Source/Editor/BlueprintGraph/Private/K2Node_PromotableOperator.cpp
+++ b/Engine/Source/Editor/BlueprintGraph/Private/K2Node_PromotableOperator.cpp
@@ -136,6 +136,91 @@ void UK2Node_PromotableOperator::GetNodeContextMenuActions(UToolMenu* Menu, UGra
 			)
 		);
 	}
+
+	// Add the pin conversion sub menu
+	if(Context->Pin)
+	{
+		FToolMenuSection& Section = Menu->AddSection("K2NodePromoOpConversionGraphNode");
+		Section.AddSubMenu(
+			"ConvertPin",
+			LOCTEXT("ConvertPin", "Convert Pin..."),
+			LOCTEXT("ConvertPinTooltip", "Convert the selected pin to another type"),
+			FNewToolMenuDelegate::CreateUObject(this, &UK2Node_PromotableOperator::CreateConversionSubMenu, (UEdGraphPin*)Context->Pin)
+		);
+	}
+}
+
+void UK2Node_PromotableOperator::CreateConversionSubMenu(UToolMenu* Menu, UEdGraphPin* ContextPin) const
+{
+	check(ContextPin);
+
+	static const FName ConvNodeName = FName("PromotableOperatorPinConvs");
+	static const FText ConvNodeStr = LOCTEXT("PromotableOperatorPinConvs", "Pin Conversions");
+
+	// If this pin is not a wildcard pin, then offer to reset it to one.
+	if(!FWildcardNodeUtils::IsWildcardPin(ContextPin))
+	{
+		static const FName ResetWildcardName = FName("PromotableOperatorResetWildcardPinConvs");
+		static const FText ResetWildcardStr = LOCTEXT("PromotableOperatorResetWildcardPinConvs", "Reset To Wildcard");
+		FToolMenuSection& Section = Menu->AddSection(ResetWildcardName, ResetWildcardStr);
+
+		const FText ResetName = LOCTEXT("ResetFunction_Tooltip", "Reset this node to wildcard");
+		
+		Section.AddMenuEntry(
+			FName(ResetName.ToString()),
+			ResetName,
+			LOCTEXT("ResetToWildcardTooltip", "Reset this node to a wildcard state."),
+			FSlateIcon(),
+			FUIAction(
+				FExecuteAction::CreateUObject(const_cast<UK2Node_PromotableOperator*>(this), &UK2Node_PromotableOperator::ConvertPinType, const_cast<UEdGraphPin*>(ContextPin), FWildcardNodeUtils::GetDefaultWildcardPinType())
+			)
+		);
+	}
+	// If the pin is a wildcard, then offer any possible types that are a part of this operator
+	else
+	{
+		// Gather what pin types could possibly be used for a conversion with this operator
+		TArray<UFunction*> AvailableFunctions;
+		FTypePromotion::GetAllFuncsForOp(OperationName, AvailableFunctions);
+		TArray<FEdGraphPinType> PossiblePromos;
+		const UEdGraphSchema_K2* Schema = GetDefault<UEdGraphSchema_K2>();
+
+		for (const UFunction* Func : AvailableFunctions)
+		{
+			for (TFieldIterator<FProperty> PropIt(Func); PropIt && (PropIt->PropertyFlags & CPF_Parm); ++PropIt)
+			{
+				const FProperty* Param = *PropIt;
+				FEdGraphPinType ParamType;
+				if (Schema->ConvertPropertyToPinType(Param, /* out */ ParamType))
+				{
+					PossiblePromos.AddUnique(ParamType);
+				}
+			}
+		}
+
+		ensureMsgf(PossiblePromos.Num() > 0, TEXT("There are no possible promotions on '%s' op node!"), *OperationName.ToString());
+
+		// Add the options to the context menu
+		for (const FEdGraphPinType& PinType : PossiblePromos)
+		{
+			FToolMenuSection& Section = Menu->AddSection(ConvNodeName, ConvNodeStr);
+
+			FFormatNamedArguments Args;
+			Args.Add(TEXT("PinType"), Schema->TypeToText(PinType));
+
+			const FText PinConversionName = FText::Format(LOCTEXT("CallFunction_Tooltip", "{PinType}"), Args);
+
+			Section.AddMenuEntry(
+				FName(PinConversionName.ToString()),
+				PinConversionName,
+				LOCTEXT("ConvertPinTypeTooltip", "Convert this pin type to the selected type"),
+				FSlateIcon(),
+				FUIAction(
+					FExecuteAction::CreateUObject(const_cast<UK2Node_PromotableOperator*>(this), &UK2Node_PromotableOperator::ConvertPinType, const_cast<UEdGraphPin*>(ContextPin), PinType)
+				)
+			);
+		}
+	}
 }
 
 FText UK2Node_PromotableOperator::GetTooltipText() const
@@ -516,6 +601,15 @@ void UK2Node_PromotableOperator::ReallocatePinsDuringReconstruction(TArray<UEdGr
 				++AdditionalPinsFixed;
 			}
 		}
+		
+		// Copy the default value and pin type for pins without any links to preseve it
+		UEdGraphPin* CurrentPin = FindPin(OldPin->GetFName(), OldPin->Direction);
+
+		if(CurrentPin && OldPin->LinkedTo.Num() == 0 && !OldPin->DoesDefaultValueMatchAutogenerated())
+		{
+			CurrentPin->PinType = OldPin->PinType;
+			CurrentPin->DefaultValue = OldPin->DefaultValue;
+		}
 	}
 }
 
@@ -699,17 +793,17 @@ bool UK2Node_PromotableOperator::HasDeterminingComparisonTypes() const
 	return false;
 }
 
-void UK2Node_PromotableOperator::EvaluatePinsFromChange(UEdGraphPin* ChangedPin)
+void UK2Node_PromotableOperator::EvaluatePinsFromChange(UEdGraphPin* ChangedPin, const bool bFromConversion /* = false*/)
 {
 	UpdateOpName();
 
 	// True if the pin that has changed now has zero connections
-	const bool bWasAFullDisconnect = (ChangedPin->LinkedTo.Num() == 0);
+	const bool bWasAFullDisconnect = (ChangedPin->LinkedTo.Num() == 0) && !HasAnyConnectionsOrDefaults();
 	const bool bIsComparisonOp = FTypePromotion::IsComparisonOpName(OperationName);
 
 	// If we have been totally disconnected and don't have any non-default inputs, 
 	// than we just reset the node to be a regular wildcard
-	if ((bWasAFullDisconnect && !HasAnyConnectionsOrDefaults()) || (bIsComparisonOp && !HasDeterminingComparisonTypes()))
+	if (!bFromConversion && (bWasAFullDisconnect || (bIsComparisonOp && !HasDeterminingComparisonTypes())))
 	{
 		ResetNodeToWildcard();
 		return;
@@ -736,6 +830,11 @@ void UK2Node_PromotableOperator::EvaluatePinsFromChange(UEdGraphPin* ChangedPin)
 	TArray<UEdGraphPin*> PinsToConsider;
 	GetPinsToConsider(PinsToConsider);
 
+	if (bFromConversion)
+	{
+		PinsToConsider.Add(ChangedPin);
+	}
+
 	const UFunction* BestMatchingFunc = FTypePromotion::FindBestMatchingFunc(OperationName, PinsToConsider);
 
 	// Store these other function options for later so that the user can convert to them later
@@ -1051,8 +1150,10 @@ void UK2Node_PromotableOperator::UpdatePinsFromFunction(const UFunction* Functio
 			FunctionPinType.PinCategory == UEdGraphSchema_K2::PC_Struct && 
 			HighestLinkedType.PinSubCategoryObject != FunctionPinType.PinSubCategoryObject;
 
+		const bool bHasDeterminingType = NodePin->LinkedTo.Num() > 0 || !NodePin->DoesDefaultValueMatchAutogenerated();
+
 		// If the highest type is the same as the function type, just continue on with life
-		if (NodePin->LinkedTo.Num() > 0 && (HighestLinkedType.PinCategory != FunctionPinType.PinCategory || bDifferingStructs))
+		if (bHasDeterminingType && (HighestLinkedType.PinCategory != FunctionPinType.PinCategory || bDifferingStructs))
 		{
 			NodePin->PinType = FunctionPinType;
 			const bool bValidPromo = 
@@ -1148,4 +1249,32 @@ UEdGraphPin* UK2Node_PromotableOperator::GetOutputPin() const
 	return nullptr;
 }
 
+void UK2Node_PromotableOperator::ConvertPinType(UEdGraphPin* PinToChange, const FEdGraphPinType NewPinType)
+{
+	// No work to be done here!
+	if (!PinToChange || PinToChange->PinType == NewPinType)
+	{
+		return;
+	}
+
+	FScopedTransaction Transaction(LOCTEXT("PromotableOperatorPinConvert", "Convert pin type"));
+	Modify();
+
+	if(NewPinType.PinCategory == UEdGraphSchema_K2::PC_Wildcard)
+	{
+		BreakAllNodeLinks();
+		ResetNodeToWildcard();
+		return;
+	}
+
+	// Break any pin links to this node because the type will be different
+	PinToChange->BreakAllPinLinks();
+
+	PinToChange->PinType = NewPinType;
+
+	EvaluatePinsFromChange(PinToChange, true);
+
+	InvalidatePinTooltips();
+}
+
 #undef LOCTEXT_NAMESPACE	// "PromotableOperatorNode"
\ No newline at end of file
