Allow double click navigation to the subsystem class with GetSubsystem nodes

#rnx

diff --git a/Engine/Source/Editor/BlueprintGraph/Classes/K2Node_GetSubsystem.h b/Engine/Source/Editor/BlueprintGraph/Classes/K2Node_GetSubsystem.h
--- a/Engine/Source/Editor/BlueprintGraph/Classes/K2Node_GetSubsystem.h
+++ b/Engine/Source/Editor/BlueprintGraph/Classes/K2Node_GetSubsystem.h
@@ -24,7 +24,8 @@ public:
 	virtual FLinearColor GetNodeTitleColor() const override;
 	virtual FText GetNodeTitle(ENodeTitleType::Type TitleType) const override;
 	virtual void ExpandNode(FKismetCompilerContext& CompilerContext, UEdGraph* SourceGraph) override;
-
+	virtual bool CanJumpToDefinition() const override;
+	virtual void JumpToDefinition() const override;
 	virtual FText GetTooltipText() const override;
 	//~End UEdGraphNode interface.
 
diff --git a/Engine/Source/Editor/BlueprintGraph/Private/K2Node_GetSubsystem.cpp b/Engine/Source/Editor/BlueprintGraph/Private/K2Node_GetSubsystem.cpp
--- a/Engine/Source/Editor/BlueprintGraph/Private/K2Node_GetSubsystem.cpp
+++ b/Engine/Source/Editor/BlueprintGraph/Private/K2Node_GetSubsystem.cpp
@@ -19,7 +19,8 @@
 #include "GameFramework/PlayerController.h"
 #include "Kismet2/BlueprintEditorUtils.h"
 #include "Subsystems/WorldSubsystem.h"
-
+#include "SourceCodeNavigation.h"
+#include "BlueprintEditorSettings.h"
 
 // ************************************************************************************
 //    UK2Node_GetSubsystem
@@ -713,10 +714,36 @@ FText UK2Node_GetEditorSubsystem::GetTooltipText() const
 	{
 		return FText::FormatNamed(NSLOCTEXT("K2Node", "GetEditorSubsystem_TooltipFormat", "Get {ClassName} an Editor Subsystem"), TEXT("ClassName"), CustomClass->GetDisplayNameText());
 	}
-
+	 
 	return NSLOCTEXT("K2Node", "GetEditorSubsystem_InvalidSubsystemTypeTooltip", "Invalid Subsystem Type");
 }
 
+bool UK2Node_GetSubsystem::CanJumpToDefinition() const
+{
+	return CustomClass != nullptr;
+}
+
+void UK2Node_GetSubsystem::JumpToDefinition() const
+{
+	bool bSucceeded = false;
+	
+	// Attempt to navigate to the header file where the class is defined if the 
+	// blueprint preferences allow for it
+	if (GetDefault<UBlueprintEditorSettings>()->bNavigateToNativeFunctionsFromCallNodes)
+	{
+		if (FSourceCodeNavigation::CanNavigateToClass(CustomClass))
+		{
+			bSucceeded = FSourceCodeNavigation::NavigateToClass(CustomClass);
+		}
+	}
+
+	// Otherwise fall back to the base class which will just bring focus to this node
+	if (!bSucceeded)
+	{
+		Super::JumpToDefinition();
+	}
+}
+
 bool UK2Node_GetEditorSubsystem::IsActionFilteredOut(class FBlueprintActionFilter const& Filter)
 {
 	for(UBlueprint* BP : Filter.Context.Blueprints)
