Do not display the Right Click + Convert pin type menu on comparison operator output pins, because the output is always a boolean and cannot be changed.

#rnx

diff --git a/Engine/Source/Editor/BlueprintGraph/Classes/K2Node_PromotableOperator.h b/Engine/Source/Editor/BlueprintGraph/Classes/K2Node_PromotableOperator.h
--- a/Engine/Source/Editor/BlueprintGraph/Classes/K2Node_PromotableOperator.h
+++ b/Engine/Source/Editor/BlueprintGraph/Classes/K2Node_PromotableOperator.h
@@ -158,6 +158,15 @@ private:
 	*/
 	void CreateConversionSubMenu(UToolMenu* Menu, UEdGraphPin* ContextPin) const;
 
+	/**
+	* Returns true if the given pin can be converted via the context menu to another type
+	*
+	* @param Pin		The pin to consider
+	*
+	* @return bool		True if a conversion is possible
+	*/
+	bool CanConvertPinType(const UEdGraphPin* Pin) const;
+
 	/** The name that this operation uses ("Add", "Multiply", etc) */
 	FName OperationName;
 
diff --git a/Engine/Source/Editor/BlueprintGraph/Private/K2Node_PromotableOperator.cpp b/Engine/Source/Editor/BlueprintGraph/Private/K2Node_PromotableOperator.cpp
--- a/Engine/Source/Editor/BlueprintGraph/Private/K2Node_PromotableOperator.cpp
+++ b/Engine/Source/Editor/BlueprintGraph/Private/K2Node_PromotableOperator.cpp
@@ -138,7 +138,7 @@ void UK2Node_PromotableOperator::GetNodeContextMenuActions(UToolMenu* Menu, UGra
 	}
 
 	// Add the pin conversion sub menu
-	if(Context->Pin)
+	if(CanConvertPinType(Context->Pin))
 	{
 		FToolMenuSection& Section = Menu->AddSection("K2NodePromoOpConversionGraphNode");
 		Section.AddSubMenu(
@@ -223,6 +223,12 @@ void UK2Node_PromotableOperator::CreateConversionSubMenu(UToolMenu* Menu, UEdGra
 	}
 }
 
+bool UK2Node_PromotableOperator::CanConvertPinType(const UEdGraphPin* Pin) const
+{
+	// You can convert any pin except for output pins on comparison functions
+	return Pin && !(Pin->Direction == EGPD_Output && FTypePromotion::IsComparisonFunc(GetTargetFunction()));
+}
+
 FText UK2Node_PromotableOperator::GetTooltipText() const
 {
 	// If there are no connections then just display the op name
