Do not break connection to boolean outputs when resetting nodes to wildcard

#rnx

diff --git a/Engine/Source/Editor/BlueprintGraph/Private/K2Node_PromotableOperator.cpp b/Engine/Source/Editor/BlueprintGraph/Private/K2Node_PromotableOperator.cpp
--- a/Engine/Source/Editor/BlueprintGraph/Private/K2Node_PromotableOperator.cpp
+++ b/Engine/Source/Editor/BlueprintGraph/Private/K2Node_PromotableOperator.cpp
@@ -1268,7 +1268,17 @@ void UK2Node_PromotableOperator::ConvertPinType(UEdGraphPin* PinToChange, const
 
 	if(NewPinType.PinCategory == UEdGraphSchema_K2::PC_Wildcard)
 	{
-		BreakAllNodeLinks();
+		const bool bIsComparison = FTypePromotion::IsComparisonFunc(GetTargetFunction());
+		// Break all input connections, but only break an output if it is not a comparison
+		// because the bool pin type will not change.
+		for(UEdGraphPin* Pin : Pins)
+		{
+			if(Pin->Direction == EGPD_Input || !bIsComparison)
+			{
+				Pin->BreakAllPinLinks();
+			}
+		}		
+		
 		ResetNodeToWildcard();
 		return;
 	}
