Remove #pragma once
#rnx

diff --git a/Engine/Source/Editor/BlueprintGraph/Private/BlueprintNodeStatics.cpp b/Engine/Source/Editor/BlueprintGraph/Private/BlueprintNodeStatics.cpp
--- a/Engine/Source/Editor/BlueprintGraph/Private/BlueprintNodeStatics.cpp
+++ b/Engine/Source/Editor/BlueprintGraph/Private/BlueprintNodeStatics.cpp
@@ -1,7 +1,5 @@
 // Copyright Epic Games, Inc. All Rights Reserved.
 
-#pragma once
-
 #include "BlueprintNodeStatics.h"
 #include "EdGraphSchema_K2.h"
 #include "EdGraph/EdGraphPin.h"
