Fix invalid access exception when upgrading a regular math node to a promotable operator

#rnx

diff --git a/Engine/Source/Editor/BlueprintGraph/Private/K2Node_PromotableOperator.cpp b/Engine/Source/Editor/BlueprintGraph/Private/K2Node_PromotableOperator.cpp
--- a/Engine/Source/Editor/BlueprintGraph/Private/K2Node_PromotableOperator.cpp
+++ b/Engine/Source/Editor/BlueprintGraph/Private/K2Node_PromotableOperator.cpp
@@ -1164,7 +1164,7 @@ void UK2Node_PromotableOperator::UpdatePinsFromFunction(const UFunction* Functio
 			NodePin->PinType = FunctionPinType;
 			const bool bValidPromo = 
 				FTypePromotion::IsValidPromotion(HighestLinkedType, FunctionPinType) || 
-				FTypePromotion::HasStructConversion(NodePin, NodePin->LinkedTo[0]);
+				(NodePin->LinkedTo.Num() > 0 && FTypePromotion::HasStructConversion(NodePin, NodePin->LinkedTo[0]));
 
 			// If the links cannot be promoted to the function type, then we need to break them
 			// We don't want to break the pin if it is the one that the user has dragged on to though,
