Add a null check before pin propagation in knot nodes.

PR #7758: Don't crash editor if input node failed to load (Contributed by [REDACTED_PERSON])

diff --git a/Engine/Source/Editor/BlueprintGraph/Private/K2Node_Knot.cpp b/Engine/Source/Editor/BlueprintGraph/Private/K2Node_Knot.cpp
--- a/Engine/Source/Editor/BlueprintGraph/Private/K2Node_Knot.cpp
+++ b/Engine/Source/Editor/BlueprintGraph/Private/K2Node_Knot.cpp
@@ -147,9 +147,12 @@ void UK2Node_Knot::PropagatePinTypeFromDirection(bool bFromInput)
 	// into this function but the recursion guard will stop it
 	for (UEdGraphPin* InPin : MySourcePin->LinkedTo)
 	{
-		if (UK2Node_Knot* KnotNode = Cast<UK2Node_Knot>(InPin->GetOwningNode()))
+		if (InPin)
 		{
-			KnotNode->PropagatePinTypeFromDirection(bFromInput);
+			if (UK2Node_Knot * KnotNode = Cast<UK2Node_Knot>(InPin->GetOwningNode()))
+			{
+				KnotNode->PropagatePinTypeFromDirection(bFromInput);
+			}
 		}
 	}
 
