Allow conversions on pins that are not wildcards

#rnx

diff --git a/Engine/Source/Editor/BlueprintGraph/Classes/K2Node_PromotableOperator.h b/Engine/Source/Editor/BlueprintGraph/Classes/K2Node_PromotableOperator.h
--- a/Engine/Source/Editor/BlueprintGraph/Classes/K2Node_PromotableOperator.h
+++ b/Engine/Source/Editor/BlueprintGraph/Classes/K2Node_PromotableOperator.h
@@ -121,7 +121,7 @@ private:
 	* Update the pins on this node based on the given function. This modifies pins, meant 
 	* for use by PinConnectionListChanged, not during node construction. 
 	*/
-	void UpdatePinsFromFunction(const UFunction* Function, UEdGraphPin* ChangedPin = nullptr);
+	void UpdatePinsFromFunction(const UFunction* Function, UEdGraphPin* ChangedPin = nullptr, bool bIsFromConversion = false);
 
 	/**
 	* Returns all pins that have the EGPD_Input direction
diff --git a/Engine/Source/Editor/BlueprintGraph/Private/K2Node_PromotableOperator.cpp b/Engine/Source/Editor/BlueprintGraph/Private/K2Node_PromotableOperator.cpp
--- a/Engine/Source/Editor/BlueprintGraph/Private/K2Node_PromotableOperator.cpp
+++ b/Engine/Source/Editor/BlueprintGraph/Private/K2Node_PromotableOperator.cpp
@@ -138,8 +138,28 @@ void UK2Node_PromotableOperator::GetNodeContextMenuActions(UToolMenu* Menu, UGra
 	}
 
 	// Add the pin conversion sub menu
-	if(CanConvertPinType(Context->Pin))
+	if (CanConvertPinType(Context->Pin))
 	{
+		// Give the user an option to reset this node to wildcard
+		if (!FWildcardNodeUtils::IsWildcardPin(Context->Pin))
+		{
+			static const FName ResetWildcardName = FName("PromotableOperatorResetWildcardPinConvs");
+			static const FText ResetWildcardStr = LOCTEXT("PromotableOperatorResetWildcardPinConvs", "Reset To Wildcard");
+			FToolMenuSection& Section = Menu->AddSection(ResetWildcardName, ResetWildcardStr);
+
+			const FText ResetName = LOCTEXT("ResetFunction_Tooltip", "Reset this node to wildcard");
+
+			Section.AddMenuEntry(
+				FName(ResetName.ToString()),
+				ResetName,
+				LOCTEXT("ResetToWildcardTooltip", "Reset this node to a wildcard state."),
+				FSlateIcon(),
+				FUIAction(
+					FExecuteAction::CreateUObject(const_cast<UK2Node_PromotableOperator*>(this), &UK2Node_PromotableOperator::ConvertPinType, const_cast<UEdGraphPin*>(Context->Pin), FWildcardNodeUtils::GetDefaultWildcardPinType())
+				)
+			);
+		}
+
 		FToolMenuSection& Section = Menu->AddSection("K2NodePromoOpConversionGraphNode");
 		Section.AddSubMenu(
 			"ConvertPin",
@@ -157,76 +177,74 @@ void UK2Node_PromotableOperator::CreateConversionSubMenu(UToolMenu* Menu, UEdGra
 	static const FName ConvNodeName = FName("PromotableOperatorPinConvs");
 	static const FText ConvNodeStr = LOCTEXT("PromotableOperatorPinConvs", "Pin Conversions");
 
-	// If this pin is not a wildcard pin, then offer to reset it to one.
-	if(!FWildcardNodeUtils::IsWildcardPin(ContextPin))
-	{
-		static const FName ResetWildcardName = FName("PromotableOperatorResetWildcardPinConvs");
-		static const FText ResetWildcardStr = LOCTEXT("PromotableOperatorResetWildcardPinConvs", "Reset To Wildcard");
-		FToolMenuSection& Section = Menu->AddSection(ResetWildcardName, ResetWildcardStr);
+	// Gather what pin types could possibly be used for a conversion with this operator
+	TArray<UFunction*> AvailableFunctions;
+	FTypePromotion::GetAllFuncsForOp(OperationName, AvailableFunctions);
+	TArray<FEdGraphPinType> PossiblePromos;
+	const UEdGraphSchema_K2* Schema = GetDefault<UEdGraphSchema_K2>();
 
-		const FText ResetName = LOCTEXT("ResetFunction_Tooltip", "Reset this node to wildcard");
-		
-		Section.AddMenuEntry(
-			FName(ResetName.ToString()),
-			ResetName,
-			LOCTEXT("ResetToWildcardTooltip", "Reset this node to a wildcard state."),
-			FSlateIcon(),
-			FUIAction(
-				FExecuteAction::CreateUObject(const_cast<UK2Node_PromotableOperator*>(this), &UK2Node_PromotableOperator::ConvertPinType, const_cast<UEdGraphPin*>(ContextPin), FWildcardNodeUtils::GetDefaultWildcardPinType())
-			)
-		);
-	}
-	// If the pin is a wildcard, then offer any possible types that are a part of this operator
-	else
+	// If this is a split pin, then we need to convert the parent pin, not the child.
+	if (ContextPin->ParentPin != nullptr)
 	{
-		// Gather what pin types could possibly be used for a conversion with this operator
-		TArray<UFunction*> AvailableFunctions;
-		FTypePromotion::GetAllFuncsForOp(OperationName, AvailableFunctions);
-		TArray<FEdGraphPinType> PossiblePromos;
-		const UEdGraphSchema_K2* Schema = GetDefault<UEdGraphSchema_K2>();
+		ContextPin = ContextPin->ParentPin;
+	}
 
-		for (const UFunction* Func : AvailableFunctions)
+	// If we have a pin that matches our current type, then we can use it to see if we can still get a valid function
+	FEdGraphPinType OriginalContextType = ContextPin->PinType;
+
+	for (const UFunction* Func : AvailableFunctions)
+	{
+		for (TFieldIterator<FProperty> PropIt(Func); PropIt && (PropIt->PropertyFlags & CPF_Parm); ++PropIt)
 		{
-			for (TFieldIterator<FProperty> PropIt(Func); PropIt && (PropIt->PropertyFlags & CPF_Parm); ++PropIt)
+			const FProperty* Param = *PropIt;
+			FEdGraphPinType ParamType;
+
+			if (Schema->ConvertPropertyToPinType(Param, /* out */ ParamType))
 			{
-				const FProperty* Param = *PropIt;
-				FEdGraphPinType ParamType;
-				if (Schema->ConvertPropertyToPinType(Param, /* out */ ParamType))
+				if (FWildcardNodeUtils::IsWildcardPin(ContextPin) || FTypePromotion::IsValidPromotion(ParamType, ContextPin->PinType) || FTypePromotion::IsValidPromotion(ContextPin->PinType, ParamType))
 				{
 					PossiblePromos.AddUnique(ParamType);
 				}
 			}
 		}
+	}
 
-		ensureMsgf(PossiblePromos.Num() > 0, TEXT("There are no possible promotions on '%s' op node!"), *OperationName.ToString());
+	// Don't display the conversion menu if there are no valid conversions
+	if (AvailableFunctions.Num() == 0)
+	{
+		return;
+	}
 
-		// Add the options to the context menu
-		for (const FEdGraphPinType& PinType : PossiblePromos)
-		{
-			FToolMenuSection& Section = Menu->AddSection(ConvNodeName, ConvNodeStr);
+	// Add the options to the context menu
+	for (const FEdGraphPinType& PinType : PossiblePromos)
+	{
+		FToolMenuSection& Section = Menu->AddSection(ConvNodeName, ConvNodeStr);
 
-			FFormatNamedArguments Args;
-			Args.Add(TEXT("PinType"), Schema->TypeToText(PinType));
+		FFormatNamedArguments Args;
+		Args.Add(TEXT("PinType"), Schema->TypeToText(PinType));
 
-			const FText PinConversionName = FText::Format(LOCTEXT("CallFunction_Tooltip", "{PinType}"), Args);
+		const FText PinConversionName = FText::Format(LOCTEXT("CallFunction_Tooltip", "{PinType}"), Args);
 
-			Section.AddMenuEntry(
-				FName(PinConversionName.ToString()),
-				PinConversionName,
-				LOCTEXT("ConvertPinTypeTooltip", "Convert this pin type to the selected type"),
-				FSlateIcon(),
-				FUIAction(
-					FExecuteAction::CreateUObject(const_cast<UK2Node_PromotableOperator*>(this), &UK2Node_PromotableOperator::ConvertPinType, const_cast<UEdGraphPin*>(ContextPin), PinType)
-				)
-			);
-		}
+		Section.AddMenuEntry(
+			FName(PinConversionName.ToString()),
+			PinConversionName,
+			LOCTEXT("ConvertPinTypeTooltip", "Convert this pin type to the selected type"),
+			FSlateIcon(),
+			FUIAction(
+				FExecuteAction::CreateUObject(const_cast<UK2Node_PromotableOperator*>(this), &UK2Node_PromotableOperator::ConvertPinType, const_cast<UEdGraphPin*>(ContextPin), PinType)
+			)
+		);
 	}
 }
 
 bool UK2Node_PromotableOperator::CanConvertPinType(const UEdGraphPin* Pin) const
 {
-	// You can convert any pin except for output pins on comparison functions
-	return Pin && !(Pin->Direction == EGPD_Output && FTypePromotion::IsComparisonFunc(GetTargetFunction()));
+	// You can convert any pin except for output pins on comparison functions and the tolerance pin
+	return
+		Pin &&
+		!(Pin->Direction == EGPD_Output &&
+		FTypePromotion::IsComparisonFunc(GetTargetFunction())) &&
+		!IsTolerancePin(Pin);
 }
 
 FText UK2Node_PromotableOperator::GetTooltipText() const
@@ -844,7 +862,7 @@ void UK2Node_PromotableOperator::EvaluatePinsFromChange(UEdGraphPin* ChangedPin,
 	const UFunction* BestMatchingFunc = FTypePromotion::FindBestMatchingFunc(OperationName, PinsToConsider);
 
 	// Store these other function options for later so that the user can convert to them later
-	UpdatePinsFromFunction(BestMatchingFunc, ChangedPin);
+	UpdatePinsFromFunction(BestMatchingFunc, ChangedPin, bFromConversion);
 }
 
 bool UK2Node_PromotableOperator::UpdateOpName()
@@ -1100,7 +1118,7 @@ void UK2Node_PromotableOperator::GetPinsToConsider(TArray<UEdGraphPin*>& OutArra
 	}
 }
 
-void UK2Node_PromotableOperator::UpdatePinsFromFunction(const UFunction* Function, UEdGraphPin* ChangedPin /* = nullptr */)
+void UK2Node_PromotableOperator::UpdatePinsFromFunction(const UFunction* Function, UEdGraphPin* ChangedPin /* = nullptr */, bool bIsFromConversion /*= false*/)
 {
 	if (!Function)
 	{
@@ -1139,7 +1157,7 @@ void UK2Node_PromotableOperator::UpdatePinsFromFunction(const UFunction* Functio
 		}
 	}
 
-	auto ConformPinLambda = [&ChangedPin](const FEdGraphPinType& FunctionPinType, UEdGraphPin* NodePin)
+	auto ConformPinLambda = [&ChangedPin, bIsFromConversion](const FEdGraphPinType& FunctionPinType, UEdGraphPin* NodePin)
 	{
 		// If the pin types are already equal, then we don't have to do any work
 		// If this is linked to wildcard pins, then we can just ignore it and handle it on expansion
@@ -1148,6 +1166,12 @@ void UK2Node_PromotableOperator::UpdatePinsFromFunction(const UFunction* Functio
 			return;
 		}
 
+		// Pins that are underdoing a conversion will have already had their types changed
+		if (NodePin == ChangedPin && bIsFromConversion)
+		{
+			return;
+		}
+
 		// By default, conform to the type of the function param
 		FEdGraphPinType ConformingType = FunctionPinType;
 		const FEdGraphPinType HighestLinkedType = NodePin->LinkedTo.Num() > 0 ? FTypePromotion::GetPromotedType(NodePin->LinkedTo) : NodePin->PinType;
@@ -1286,11 +1310,21 @@ void UK2Node_PromotableOperator::ConvertPinType(UEdGraphPin* PinToChange, const
 	// Break any pin links to this node because the type will be different
 	PinToChange->BreakAllPinLinks();
 
+	// Recombine any split pins that this type has before changing its type
+	const UEdGraphSchema_K2* Schema = GetDefault<UEdGraphSchema_K2>();
+	if (PinToChange->SubPins.Num() > 0)
+	{
+		Schema->RecombinePin(PinToChange);
+	}
+
 	PinToChange->PinType = NewPinType;
 
 	EvaluatePinsFromChange(PinToChange, true);
 
 	InvalidatePinTooltips();
+
+	// Reset the default value on pins that have been converted
+	Schema->ResetPinToAutogeneratedDefaultValue(PinToChange, false);
 }
 
 #undef LOCTEXT_NAMESPACE	// "PromotableOperatorNode"
\ No newline at end of file
