Set the tolerance pin to hidden if necessary when converting pin types

#rnx

diff --git a/Engine/Source/Editor/BlueprintGraph/Private/K2Node_PromotableOperator.cpp b/Engine/Source/Editor/BlueprintGraph/Private/K2Node_PromotableOperator.cpp
--- a/Engine/Source/Editor/BlueprintGraph/Private/K2Node_PromotableOperator.cpp
+++ b/Engine/Source/Editor/BlueprintGraph/Private/K2Node_PromotableOperator.cpp
@@ -1208,18 +1208,18 @@ void UK2Node_PromotableOperator::UpdatePinsFromFunction(const UFunction* Functio
 	};
 
 	// Check if we need to add a tolerance pin or not with this new function
-	if (FTypePromotion::IsComparisonFunc(Function) && PromotableOpUtils::FindTolerancePinType(Function, ToleranceType))
+	if (FTypePromotion::IsComparisonFunc(Function))
 	{
-		// Set the tolerance pin to visible if it is currently on the node
-		if (UEdGraphPin* ExistingTolPin = FindTolerancePin())
-		{
-			ExistingTolPin->bHidden = false;
-		}
-		// Otherwise we need to create a tolerance pin
-		else
+		UEdGraphPin* ExistingTolPin = FindTolerancePin();
+		const bool bHasTolerancePin = PromotableOpUtils::FindTolerancePinType(Function, ToleranceType);
+
+		if (!ExistingTolPin)
 		{
-			CreatePin(EGPD_Input, ToleranceType, TolerancePin_Name);
+			ExistingTolPin = CreatePin(EGPD_Input, ToleranceType, TolerancePin_Name);
 		}
+
+		// Set the tolerance pin to visible if it is currently on the node
+		ExistingTolPin->bHidden = !bHasTolerancePin;
 	}
 
 	int32 CurPinIndex = 0;
