Consider pins even if they have a non-modified default value. This allows for a more coherent pin conversion  experience so that you do not have to enter a default value before converting a second pin.

#rnx

diff --git a/Engine/Source/Editor/BlueprintGraph/Private/K2Node_PromotableOperator.cpp b/Engine/Source/Editor/BlueprintGraph/Private/K2Node_PromotableOperator.cpp
--- a/Engine/Source/Editor/BlueprintGraph/Private/K2Node_PromotableOperator.cpp
+++ b/Engine/Source/Editor/BlueprintGraph/Private/K2Node_PromotableOperator.cpp
@@ -856,7 +856,7 @@ void UK2Node_PromotableOperator::EvaluatePinsFromChange(UEdGraphPin* ChangedPin,
 
 	if (bFromConversion)
 	{
-		PinsToConsider.Add(ChangedPin);
+		PinsToConsider.AddUnique(ChangedPin);
 	}
 
 	const UFunction* BestMatchingFunc = FTypePromotion::FindBestMatchingFunc(OperationName, PinsToConsider);
@@ -1103,7 +1103,7 @@ void UK2Node_PromotableOperator::GetPinsToConsider(TArray<UEdGraphPin*>& OutArra
 				OutArray.Emplace(Link);
 			}
 		}
-		else if (!Pin->DoesDefaultValueMatchAutogenerated() && !FWildcardNodeUtils::IsWildcardPin(Pin))
+		else if (!FWildcardNodeUtils::IsWildcardPin(Pin))
 		{
 			// If this is from a split pin, then we care about the type of the parent, not this pin
 			if (Pin->ParentPin != nullptr)
