Give a preference to actions that match the same container type to ensure that they get the right modifier functions like "Add"

#rnx

diff --git a/Engine/Source/Editor/BlueprintGraph/Private/EdGraphSchema_K2.cpp b/Engine/Source/Editor/BlueprintGraph/Private/EdGraphSchema_K2.cpp
--- a/Engine/Source/Editor/BlueprintGraph/Private/EdGraphSchema_K2.cpp
+++ b/Engine/Source/Editor/BlueprintGraph/Private/EdGraphSchema_K2.cpp
@@ -7462,6 +7462,13 @@ namespace ContextMenuConsoleVariables
 		TEXT("The bonus given if node is a favorite"),
 		ECVF_Default);
 
+	/** The bonus given if an action has the same container type as the dragged from pin */
+	static float ContainerBonus = 1000.0f;
+	static FAutoConsoleVariableRef CVarContainerBonus(
+		TEXT("ContextMenu.ContainerBonus"), ContainerBonus,
+		TEXT("The bonus given if the dragged from pin matches the same container type of the action"),
+		ECVF_Default);
+
 	/** Enabling the debug printing of context menu selections */
 	static bool bPrintDebugContextSelection = false;
 	static FAutoConsoleVariableRef CVarPrintDebugContextSelection(
@@ -7567,14 +7574,43 @@ float UEdGraphSchema_K2::GetActionFilteredWeight(const FGraphActionListBuilderBa
 			const TArray<FString>& InActionCategories = InCurrentAction.GetCategoryChain();
 			bool bAddMatchBonus = false;
 
+			/** Get a string reference for an EPinContainerType */
+			auto GetContainerTypeString = [](const EPinContainerType Type) -> const FString&
+			{
+				static const FString ArrayName = TEXT("Array");
+				static const FString MapName = TEXT("Map");
+				static const FString SetName = TEXT("Set");
+				static const FString InvalidName = TEXT("INVALID");
+
+				switch (Type)
+				{
+					case EPinContainerType::Array:
+						return ArrayName;
+					case EPinContainerType::Map:
+						return MapName;
+					case EPinContainerType::Set:
+						return SetName;
+					default:
+						return InvalidName;
+				}
+			};
+
+			bool bAddedContainerPreferenceBonus = false;
+
 			for (const FString& InActionCategory : InActionCategories)
 			{
 				for (UEdGraphPin* const FromPin : DraggedFromPins)
 				{
 					check(FromPin != nullptr);
 
-					// If we can't find anything there, check the subcategory of the object
-					// This covers most of the more complex struct types (LinearColor, date time, etc)
+					// For containers, add a preference for functions that are marked in their category
+					if (!bAddedContainerPreferenceBonus && FromPin->PinType.IsContainer() && InActionCategory == GetContainerTypeString(FromPin->PinType.ContainerType))
+					{
+						TotalWeight += ContextMenuConsoleVariables::ContainerBonus;
+						bAddedContainerPreferenceBonus = true;
+					}
+
+					// Check the subcategory of the object to cover more more complex struct types (LinearColor, date time, etc)
 					if (UObject* const SubCatObj = FromPin->PinType.PinSubCategoryObject.Get())
 					{
 						const FString& SubCatObjName = SubCatObj->GetFullName();
