Set the source path of the blueprint compiler log when pasting nodes correctly. This stops an ensure within FPackageName::DoesPackageExist when given an empty package name.

#rnx

diff --git a/Engine/Source/Editor/BlueprintGraph/Private/EdGraphSchema_K2.cpp b/Engine/Source/Editor/BlueprintGraph/Private/EdGraphSchema_K2.cpp
--- a/Engine/Source/Editor/BlueprintGraph/Private/EdGraphSchema_K2.cpp
+++ b/Engine/Source/Editor/BlueprintGraph/Private/EdGraphSchema_K2.cpp
@@ -6417,6 +6417,7 @@ UEdGraphNode* UEdGraphSchema_K2::CreateSubstituteNode(UEdGraphNode* Node, const
 					Blueprint->PreCompileLog = TSharedPtr<FCompilerResultsLog>(new FCompilerResultsLog(false));
 					Blueprint->PreCompileLog->bSilentMode = false;
 					Blueprint->PreCompileLog->bAnnotateMentionedNodes = false;
+					Blueprint->PreCompileLog->SetSourcePath(Blueprint->GetPathName());
 				}
 
 				// Append a warning to the node and to the logs
