Allow object references on the pin defaults of a blueprint graph to be handled as wrapped object pointers to assist in lazy loading in editor.

diff --git a/Engine/Source/Editor/BlueprintGraph/Classes/EdGraphSchema_K2.h b/Engine/Source/Editor/BlueprintGraph/Classes/EdGraphSchema_K2.h
--- a/Engine/Source/Editor/BlueprintGraph/Classes/EdGraphSchema_K2.h
+++ b/Engine/Source/Editor/BlueprintGraph/Classes/EdGraphSchema_K2.h
@@ -5,6 +5,7 @@
 #include "CoreMinimal.h"
 #include "Misc/EnumClassFlags.h"
 #include "UObject/ObjectMacros.h"
+#include "UObject/ObjectPtr.h"
 #include "UObject/Class.h"
 #include "UObject/SoftObjectPath.h"
 #include "EdGraph/EdGraphNode.h"
@@ -614,7 +615,7 @@ public:
 	UK2Node* CreateSplitPinNode(UEdGraphPin* Pin, const FCreateSplitPinNodeParams& Params) const;
 
 	/** Reads in a FString and gets the values of the pin defaults for that type. This can be passed to DefaultValueSimpleValidation to validate. OwningObject can be null */
-	virtual void GetPinDefaultValuesFromString(const FEdGraphPinType& PinType, UObject* OwningObject, const FString& NewValue, FString& UseDefaultValue, UObject*& UseDefaultObject, FText& UseDefaultText, bool bPreserveTextIdentity = true) const;
+	virtual void GetPinDefaultValuesFromString(const FEdGraphPinType& PinType, UObject* OwningObject, const FString& NewValue, FString& UseDefaultValue, TObjectPtr<UObject>& UseDefaultObject, FText& UseDefaultText, bool bPreserveTextIdentity = true) const;
 
 	/** Do validation, that doesn't require a knowledge about actual pin */
 	virtual bool DefaultValueSimpleValidation(const FEdGraphPinType& PinType, const FName PinName, const FString& NewDefaultValue, UObject* NewDefaultObject, const FText& InText, FString* OutMsg = nullptr) const;
diff --git a/Engine/Source/Editor/BlueprintGraph/Private/EdGraphSchema_K2.cpp b/Engine/Source/Editor/BlueprintGraph/Private/EdGraphSchema_K2.cpp
--- a/Engine/Source/Editor/BlueprintGraph/Private/EdGraphSchema_K2.cpp
+++ b/Engine/Source/Editor/BlueprintGraph/Private/EdGraphSchema_K2.cpp
@@ -4538,7 +4538,7 @@ void UEdGraphSchema_K2::HandleGraphBeingDeleted(UEdGraph& GraphBeingRemoved) con
 	}
 }
 
-void UEdGraphSchema_K2::GetPinDefaultValuesFromString(const FEdGraphPinType& PinType, UObject* OwningObject, const FString& NewDefaultValue, FString& UseDefaultValue, UObject*& UseDefaultObject, FText& UseDefaultText, bool bPreserveTextIdentity) const
+void UEdGraphSchema_K2::GetPinDefaultValuesFromString(const FEdGraphPinType& PinType, UObject* OwningObject, const FString& NewDefaultValue, FString& UseDefaultValue, TObjectPtr<UObject>& UseDefaultObject, FText& UseDefaultText, bool bPreserveTextIdentity) const
 {
 	if ((PinType.PinCategory == PC_Object)
 		|| (PinType.PinCategory == PC_Class)
@@ -4629,7 +4629,7 @@ void UEdGraphSchema_K2::GetPinDefaultValuesFromString(const FEdGraphPinType& Pin
 void UEdGraphSchema_K2::TrySetDefaultValue(UEdGraphPin& Pin, const FString& NewDefaultValue, bool bMarkAsModified) const
 {
 	FString UseDefaultValue;
-	UObject* UseDefaultObject = nullptr;
+	TObjectPtr<UObject> UseDefaultObject = nullptr;
 	FText UseDefaultText;
 
 	GetPinDefaultValuesFromString(Pin.PinType, Pin.GetOwningNodeUnchecked(), NewDefaultValue, UseDefaultValue, UseDefaultObject, UseDefaultText, /*bPreserveTextIdentity*/false);
diff --git a/Engine/Source/Editor/BlueprintGraph/Private/K2Node_FunctionEntry.cpp b/Engine/Source/Editor/BlueprintGraph/Private/K2Node_FunctionEntry.cpp
--- a/Engine/Source/Editor/BlueprintGraph/Private/K2Node_FunctionEntry.cpp
+++ b/Engine/Source/Editor/BlueprintGraph/Private/K2Node_FunctionEntry.cpp
@@ -289,7 +289,7 @@ void UK2Node_FunctionEntry::Serialize(FArchive& Ar)
 			for (FBPVariableDescription& LocalVar : LocalVariables)
 			{
 				FString UseDefaultValue;
-				UObject* UseDefaultObject = nullptr;
+				TObjectPtr<UObject> UseDefaultObject = nullptr;
 				FText UseDefaultText;
 
 				if (!LocalVar.DefaultValue.IsEmpty())
diff --git a/Engine/Source/Runtime/Engine/Classes/EdGraph/EdGraphPin.h b/Engine/Source/Runtime/Engine/Classes/EdGraph/EdGraphPin.h
--- a/Engine/Source/Runtime/Engine/Classes/EdGraph/EdGraphPin.h
+++ b/Engine/Source/Runtime/Engine/Classes/EdGraph/EdGraphPin.h
@@ -5,6 +5,7 @@
 #include "CoreMinimal.h"
 #include "UObject/ObjectMacros.h"
 #include "UObject/Object.h"
+#include "UObject/ObjectPtr.h"
 #include "Misc/Guid.h"
 #include "UObject/Class.h"
 #include "UObject/WeakObjectPtr.h"
@@ -424,7 +425,7 @@ public:
 	FString AutogeneratedDefaultValue;
 
 	/** If the default value for this pin should be an object, we store a pointer to it */
-	class UObject* DefaultObject;
+	TObjectPtr<class UObject> DefaultObject;
 
 	/** If the default value for this pin should be an FText, it is stored here. */
 	FText DefaultTextValue;
diff --git a/Engine/Source/Runtime/Engine/Private/EdGraph/EdGraphPin.cpp b/Engine/Source/Runtime/Engine/Private/EdGraph/EdGraphPin.cpp
--- a/Engine/Source/Runtime/Engine/Private/EdGraph/EdGraphPin.cpp
+++ b/Engine/Source/Runtime/Engine/Private/EdGraph/EdGraphPin.cpp
@@ -683,9 +683,9 @@ bool UEdGraphPin::HasAnyConnections() const
 
 FString UEdGraphPin::GetDefaultAsString() const
 {
-	if(DefaultObject)
+	if(!DefaultObject.IsNullNoResolve())
 	{
-		return DefaultObject->GetPathName();
+		return DefaultObject.GetPath();
 	}
 	else if(!DefaultTextValue.IsEmpty())
 	{
