Properly expand DefaultToSelf pins within Asyc Task nodes

diff --git a/Engine/Source/Editor/BlueprintGraph/Classes/K2Node_BaseAsyncTask.h b/Engine/Source/Editor/BlueprintGraph/Classes/K2Node_BaseAsyncTask.h
--- a/Engine/Source/Editor/BlueprintGraph/Classes/K2Node_BaseAsyncTask.h
+++ b/Engine/Source/Editor/BlueprintGraph/Classes/K2Node_BaseAsyncTask.h
@@ -15,6 +15,7 @@ class UEdGraphPin;
 class UEdGraphSchema_K2;
 class UK2Node_CustomEvent;
 class UK2Node_TemporaryVariable;
+class UK2Node_CallFunction;
 
 /** struct to remap pins for Async Tasks.
  * a single K2 node is shared by many proxy classes.
@@ -46,7 +47,7 @@ class BLUEPRINTGRAPH_API UK2Node_BaseAsyncTask : public UK2Node
 	// End of UEdGraphNode interface
 
 	// UK2Node interface
-	virtual void ExpandNode(class FKismetCompilerContext& CompilerContext, UEdGraph* SourceGraph) override;
+	virtual void ExpandNode(FKismetCompilerContext& CompilerContext, UEdGraph* SourceGraph) override;
 	virtual bool HasExternalDependencies(TArray<class UStruct*>* OptionalOutput) const override;
 	virtual FName GetCornerIcon() const override;
 	virtual void GetMenuActions(FBlueprintActionDatabaseRegistrar& ActionRegistrar) const override;
@@ -62,6 +63,19 @@ protected:
 	/** Determines what the possible redirect pin names are **/
 	virtual void GetRedirectPinNames(const UEdGraphPin& Pin, TArray<FString>& RedirectPinNames) const;
 
+	/**
+	 * If a the DefaultToSelf pin exists then it needs an actual connection to get properly casted
+	 * during compilation.
+	 *
+	 * @param CompilerContext			The current compiler context used during expansion
+	 * @param SourceGraph				The graph to place the expanded self node on
+	 * @param IntermediateProxyNode		The spawned intermediate proxy node that has a DefaultToSelfPin
+	 *
+	 * @return	True if a successful connection was made to an intermediate node
+	 *			or if one was not necessary. False if a connection was failed.
+	 */
+	bool ExpandDefaultToSelfPin(FKismetCompilerContext& CompilerContext, UEdGraph* SourceGraph, UK2Node_CallFunction* IntermediateProxyNode);
+	
 protected:
 	// The name of the function to call to create a proxy object
 	UPROPERTY()
diff --git a/Engine/Source/Editor/BlueprintGraph/Private/K2Node_BaseAsyncTask.cpp b/Engine/Source/Editor/BlueprintGraph/Private/K2Node_BaseAsyncTask.cpp
--- a/Engine/Source/Editor/BlueprintGraph/Private/K2Node_BaseAsyncTask.cpp
+++ b/Engine/Source/Editor/BlueprintGraph/Private/K2Node_BaseAsyncTask.cpp
@@ -318,8 +318,38 @@ bool UK2Node_BaseAsyncTask::FBaseAsyncTaskHelper::HandleDelegateImplementation(
 	return bIsErrorFree;
 }
 
+bool UK2Node_BaseAsyncTask::ExpandDefaultToSelfPin(FKismetCompilerContext& CompilerContext, UEdGraph* SourceGraph, UK2Node_CallFunction* IntermediateProxyNode)
+{
+	if(SourceGraph && IntermediateProxyNode)
+	{
+		// Connect a self reference pin if there is a TScriptInterface default to self
+		if (const UFunction* TargetFunc = IntermediateProxyNode->GetTargetFunction())
+		{
+			const FString& MetaData = TargetFunc->GetMetaData(FBlueprintMetadata::MD_DefaultToSelf);
+			if (!MetaData.IsEmpty())
+			{
+				// Find the default to self value pin
+				if (UEdGraphPin* DefaultToSelfPin = IntermediateProxyNode->FindPinChecked(MetaData, EGPD_Input))
+				{
+					// If it has no links then spawn a new self node here
+					if (DefaultToSelfPin->PinType.PinCategory == UEdGraphSchema_K2::PC_Interface && DefaultToSelfPin->LinkedTo.Num() == 0)
+					{
+						const UEdGraphSchema_K2* Schema = CompilerContext.GetSchema();
 
-void UK2Node_BaseAsyncTask::ExpandNode(class FKismetCompilerContext& CompilerContext, UEdGraph* SourceGraph)
+						UK2Node_Self* SelfNode = CompilerContext.SpawnIntermediateNode<UK2Node_Self>(this, SourceGraph);
+						SelfNode->AllocateDefaultPins();
+						UEdGraphPin* SelfPin = SelfNode->FindPinChecked(UEdGraphSchema_K2::PSC_Self);
+						// Make a connection from this intermediate self pin to here
+						return Schema->TryCreateConnection(DefaultToSelfPin, SelfPin);
+					}
+				}
+			}
+		}	
+	}
+	return true;
+}
+
+void UK2Node_BaseAsyncTask::ExpandNode(FKismetCompilerContext& CompilerContext, UEdGraph* SourceGraph)
 {
     Super::ExpandNode(CompilerContext, SourceGraph);
 
@@ -360,6 +390,8 @@ void UK2Node_BaseAsyncTask::ExpandNode(class FKismetCompilerContext& CompilerCon
 	UEdGraphPin* OutputAsyncTaskProxy = FindPin(FBaseAsyncTaskHelper::GetAsyncTaskProxyName());
 	bIsErrorFree &= !OutputAsyncTaskProxy || CompilerContext.MovePinLinksToIntermediate(*OutputAsyncTaskProxy, *ProxyObjectPin).CanSafeConnect();
 
+	bIsErrorFree &= ExpandDefaultToSelfPin(CompilerContext, SourceGraph, CallCreateProxyObjectNode);		
+
 	// GATHER OUTPUT PARAMETERS AND PAIR THEM WITH LOCAL VARIABLES
 	TArray<FBaseAsyncTaskHelper::FOutputPinAndLocalVariable> VariableOutputs;
 	bool bPassedFactoryOutputs = false;
diff --git a/Engine/Source/Editor/GameplayTasksEditor/Private/K2Node_LatentGameplayTaskCall.cpp b/Engine/Source/Editor/GameplayTasksEditor/Private/K2Node_LatentGameplayTaskCall.cpp
--- a/Engine/Source/Editor/GameplayTasksEditor/Private/K2Node_LatentGameplayTaskCall.cpp
+++ b/Engine/Source/Editor/GameplayTasksEditor/Private/K2Node_LatentGameplayTaskCall.cpp
@@ -581,6 +581,8 @@ void UK2Node_LatentGameplayTaskCall::ExpandNode(class FKismetCompilerContext& Co
 			}
 		}
 	}
+	
+	bIsErrorFree &= ExpandDefaultToSelfPin(CompilerContext, SourceGraph, CallCreateProxyObjectNode);
 
 	// Expose Async Task Proxy object
 	UEdGraphPin* const ProxyObjectPin = CallCreateProxyObjectNode->GetReturnValuePin();
