Do not allow marco nodes to have breakpoints placed on them

diff --git a/Engine/Source/Editor/BlueprintGraph/Classes/K2Node.h b/Engine/Source/Editor/BlueprintGraph/Classes/K2Node.h
--- a/Engine/Source/Editor/BlueprintGraph/Classes/K2Node.h
+++ b/Engine/Source/Editor/BlueprintGraph/Classes/K2Node.h
@@ -222,7 +222,7 @@ public:
 	virtual bool HasExternalDependencies(TArray<class UStruct*>* OptionalOutput = nullptr) const { return false; }
 
 	/** Returns whether this node can have breakpoints placed on it in the debugger */
-	virtual bool CanPlaceBreakpoints() const { return !IsNodePure(); }
+	BLUEPRINTGRAPH_API virtual bool CanPlaceBreakpoints() const;
 
 	/** Return whether to draw this node as an entry */
 	virtual bool DrawNodeAsEntry() const { return false; }
diff --git a/Engine/Source/Editor/BlueprintGraph/Private/K2Node.cpp b/Engine/Source/Editor/BlueprintGraph/Private/K2Node.cpp
--- a/Engine/Source/Editor/BlueprintGraph/Private/K2Node.cpp
+++ b/Engine/Source/Editor/BlueprintGraph/Private/K2Node.cpp
@@ -608,6 +608,30 @@ void UK2Node::JumpToDefinition() const
 	}
 }
 
+bool UK2Node::CanPlaceBreakpoints() const
+{
+	// Pure nodes have no execs and therefore cannot have breakpoints placed on them
+	if(IsNodePure())
+	{
+		return false;
+	}
+
+	// If this node is within a macro or interface blueprint then it cannot have breakpoints
+	// added to it as they will get expanded during compilation and never get hit. 
+	if(const UEdGraph* Graph = GetGraph())
+	{
+		if (const UBlueprint* Blueprint = FBlueprintEditorUtils::FindBlueprintForGraph(Graph))
+		{
+			return
+				!(Blueprint->BlueprintType == BPTYPE_MacroLibrary ||
+				Blueprint->BlueprintType == BPTYPE_Interface ||
+				Blueprint->MacroGraphs.Contains(Graph));
+		}
+	}
+	
+	return true;
+}
+
 void UK2Node::ReallocatePinsDuringReconstruction(TArray<UEdGraphPin*>& OldPins)
 {
 	AllocateDefaultPins();
