Fix "ctrl + drag" shortcut on promotable operators with multiple connections

Instead of always breaking the links witha promotable connection at the schema level, only break links if the most recent connection is a different type then before. This stops incorrect breaking of valid connections when there are multiple at a time while keeping a good UX.

#rnx

diff --git a/Engine/Source/Editor/BlueprintGraph/Private/EdGraphSchema_K2.cpp b/Engine/Source/Editor/BlueprintGraph/Private/EdGraphSchema_K2.cpp
--- a/Engine/Source/Editor/BlueprintGraph/Private/EdGraphSchema_K2.cpp
+++ b/Engine/Source/Editor/BlueprintGraph/Private/EdGraphSchema_K2.cpp
@@ -2736,8 +2736,6 @@ bool UEdGraphSchema_K2::CreatePromotedConnection(UEdGraphPin* PinA, UEdGraphPin*
 	PinA->Modify();
 	PinB->Modify();
 
-	PinB->BreakAllPinLinks(true);
-
 	PinA->MakeLinkTo(PinB);
 
 	UBlueprint* Blueprint = FBlueprintEditorUtils::FindBlueprintForNodeChecked(PinA->GetOwningNode());
diff --git a/Engine/Source/Editor/BlueprintGraph/Private/K2Node_PromotableOperator.cpp b/Engine/Source/Editor/BlueprintGraph/Private/K2Node_PromotableOperator.cpp
--- a/Engine/Source/Editor/BlueprintGraph/Private/K2Node_PromotableOperator.cpp
+++ b/Engine/Source/Editor/BlueprintGraph/Private/K2Node_PromotableOperator.cpp
@@ -850,6 +850,21 @@ void UK2Node_PromotableOperator::EvaluatePinsFromChange(UEdGraphPin* ChangedPin,
 		return;
 	}
 
+	// If the newest connection to this this pin is a different type,
+	// then we need to break all other type connections that are not the same
+	if(ChangedPin->LinkedTo.Num() > 1)
+	{
+		const FEdGraphPinType& MostRecentConnection = ChangedPin->LinkedTo.Last()->PinType;
+		for(int32 i = ChangedPin->LinkedTo.Num() - 1; i >= 0; --i)
+		{
+			UEdGraphPin* Link = ChangedPin->LinkedTo[i];
+			if(Link->PinType != MostRecentConnection)
+			{
+				ChangedPin->BreakLinkTo(Link);				
+			}
+		}
+	}
+	
 	// Gather all pins and their links so we can determine the highest type
 	TArray<UEdGraphPin*> PinsToConsider;
 	GetPinsToConsider(PinsToConsider);
