Undo fix for UE-94768 as it created worse bug

diff --git a/Engine/Source/Editor/BlueprintGraph/Private/K2Node_GetClassDefaults.cpp b/Engine/Source/Editor/BlueprintGraph/Private/K2Node_GetClassDefaults.cpp
--- a/Engine/Source/Editor/BlueprintGraph/Private/K2Node_GetClassDefaults.cpp
+++ b/Engine/Source/Editor/BlueprintGraph/Private/K2Node_GetClassDefaults.cpp
@@ -485,30 +485,10 @@ void UK2Node_GetClassDefaults::OnBlueprintClassModified(UBlueprint* TargetBluepr
 
 void UK2Node_GetClassDefaults::CreateOutputPins(UClass* InClass)
 {
-	TArray<FOptionalPinFromProperty> OldShowPins = MoveTemp(ShowPinForProperties);
-
 	// Create the set of output pins through the optional pin manager
 	FClassDefaultsOptionalPinManager OptionalPinManager(InClass, bExcludeObjectContainers, bExcludeObjectArrays_DEPRECATED);
 	OptionalPinManager.RebuildPropertyList(ShowPinForProperties, InClass);
 
-	// re-enable any old show pins
-	for (const FOptionalPinFromProperty& Pin : OldShowPins)
-	{
-		if (Pin.bShowPin)
-		{
-			FOptionalPinFromProperty* NewPin = ShowPinForProperties.FindByPredicate(
-				[&Pin](const FOptionalPinFromProperty& OtherPin)
-				{
-					return OtherPin.PropertyName == Pin.PropertyName && OtherPin.CategoryName == Pin.CategoryName;
-				});
-			
-			if (NewPin)
-			{
-				NewPin->bShowPin = true;
-			}
-		}
-	}
-
 	OptionalPinManager.CreateVisiblePins(ShowPinForProperties, InClass, EGPD_Output, this);
 
 	// Check for any advanced properties (outputs)
@@ -596,6 +576,9 @@ void UK2Node_GetClassDefaults::OnClassPinChanged()
 		}
 	}
 
+	// Clear the current output pin settings (so they don't carry over to the new set)
+	ShowPinForProperties.Reset();
+
 	// Create output pins for the new class type
 	UClass* InputClass = GetInputClass();
 	CreateOutputPins(InputClass);
