Only compare the pin category, subcategory, and subcategory object when considering connections to break. This stops the member reference being different causing a break

#rnx

diff --git a/Engine/Source/Editor/BlueprintGraph/Private/K2Node_PromotableOperator.cpp b/Engine/Source/Editor/BlueprintGraph/Private/K2Node_PromotableOperator.cpp
--- a/Engine/Source/Editor/BlueprintGraph/Private/K2Node_PromotableOperator.cpp
+++ b/Engine/Source/Editor/BlueprintGraph/Private/K2Node_PromotableOperator.cpp
@@ -858,7 +858,10 @@ void UK2Node_PromotableOperator::EvaluatePinsFromChange(UEdGraphPin* ChangedPin,
 		for(int32 i = ChangedPin->LinkedTo.Num() - 1; i >= 0; --i)
 		{
 			UEdGraphPin* Link = ChangedPin->LinkedTo[i];
-			if(Link->PinType != MostRecentConnection)
+			if( Link->PinType.PinCategory != MostRecentConnection.PinCategory ||
+				Link->PinType.PinSubCategory != MostRecentConnection.PinSubCategory ||
+				Link->PinType.PinSubCategoryObject != MostRecentConnection.PinSubCategoryObject
+			)
 			{
 				ChangedPin->BreakLinkTo(Link);				
 			}
