Consider the split pins of the transform pin as spawn vars

diff --git a/Engine/Source/Editor/BlueprintGraph/Classes/K2Node_SpawnActor.h b/Engine/Source/Editor/BlueprintGraph/Classes/K2Node_SpawnActor.h
--- a/Engine/Source/Editor/BlueprintGraph/Classes/K2Node_SpawnActor.h
+++ b/Engine/Source/Editor/BlueprintGraph/Classes/K2Node_SpawnActor.h
@@ -42,7 +42,7 @@ class UK2Node_SpawnActor : public UK2Node
 	/** Create new pins to show properties on archetype */
 	void CreatePinsForClass(UClass* InClass);
 	/** See if this is a spawn variable pin, or a 'default' pin */
-	BLUEPRINTGRAPH_API bool IsSpawnVarPin(UEdGraphPin* Pin);
+	BLUEPRINTGRAPH_API bool IsSpawnVarPin(UEdGraphPin* Pin) const;
 
 	/** Get the then output pin */
 	BLUEPRINTGRAPH_API UEdGraphPin* GetThenPin() const;
diff --git a/Engine/Source/Editor/BlueprintGraph/Private/K2Node_ConstructObjectFromClass.cpp b/Engine/Source/Editor/BlueprintGraph/Private/K2Node_ConstructObjectFromClass.cpp
--- a/Engine/Source/Editor/BlueprintGraph/Private/K2Node_ConstructObjectFromClass.cpp
+++ b/Engine/Source/Editor/BlueprintGraph/Private/K2Node_ConstructObjectFromClass.cpp
@@ -185,6 +185,11 @@ void UK2Node_ConstructObjectFromClass::AddSearchMetaDataInfo(TArray<struct FSear
 
 bool UK2Node_ConstructObjectFromClass::IsSpawnVarPin(UEdGraphPin* Pin) const
 {
+	if (Pin->ParentPin)
+	{
+		return IsSpawnVarPin(Pin->ParentPin);
+	}
+
 	return(	Pin->PinName != UEdGraphSchema_K2::PN_Execute &&
 			Pin->PinName != UEdGraphSchema_K2::PN_Then &&
 			Pin->PinName != UEdGraphSchema_K2::PN_ReturnValue &&
diff --git a/Engine/Source/Editor/BlueprintGraph/Private/K2Node_SpawnActor.cpp b/Engine/Source/Editor/BlueprintGraph/Private/K2Node_SpawnActor.cpp
--- a/Engine/Source/Editor/BlueprintGraph/Private/K2Node_SpawnActor.cpp
+++ b/Engine/Source/Editor/BlueprintGraph/Private/K2Node_SpawnActor.cpp
@@ -130,8 +130,13 @@ void UK2Node_SpawnActor::ReallocatePinsDuringReconstruction(TArray<UEdGraphPin*>
 	RestoreSplitPins(OldPins);
 }
 
-bool UK2Node_SpawnActor::IsSpawnVarPin(UEdGraphPin* Pin)
+bool UK2Node_SpawnActor::IsSpawnVarPin(UEdGraphPin* Pin) const
 {
+	if (Pin->ParentPin)
+	{
+		return IsSpawnVarPin(Pin->ParentPin);
+	}
+
 	return(	Pin->PinName != UEdGraphSchema_K2::PN_Execute &&
 			Pin->PinName != UEdGraphSchema_K2::PN_Then &&
 			Pin->PinName != UEdGraphSchema_K2::PN_ReturnValue &&
