Fix resolving local variable asset references on load to the function entry node's variable cache for override implementations.

diff --git a/Engine/Source/Editor/BlueprintGraph/Private/K2Node_FunctionEntry.cpp b/Engine/Source/Editor/BlueprintGraph/Private/K2Node_FunctionEntry.cpp
--- a/Engine/Source/Editor/BlueprintGraph/Private/K2Node_FunctionEntry.cpp
+++ b/Engine/Source/Editor/BlueprintGraph/Private/K2Node_FunctionEntry.cpp
@@ -434,7 +434,13 @@ TSharedPtr<FStructOnScope> UK2Node_FunctionEntry::GetFunctionVariableCache(bool
 
 	if (!FunctionVariableCache.IsValid() || !FunctionVariableCache->IsValid())
 	{
-		if (UFunction* const Function = FindSignatureFunction())
+		// Locate the UFunction object in the class hierarchy, starting at the current class. Note that local
+		// variables are generated as fields (properties) within the function context that contains this node,
+		// so for parent class/interface function overrides we want to make sure we're looking at the most-
+		// derived UFunction object. Also note that FunctionFromNode() looks at the skeleton class rather than
+		// the [authoritative] generated class, since the skeleton class is implicitly recompiled after e.g. adding
+		// an input/output argument or local variable, whereas the generated class must be explicitly recompiled.
+		if (UFunction* const Function = FFunctionFromNodeHelper::FunctionFromNode(this))
 		{
 			if (LocalVariables.Num() > 0)
 			{
@@ -801,16 +807,6 @@ void UK2Node_FunctionEntry::ExpandNode(class FKismetCompilerContext& CompilerCon
 		// Find the associated UFunction
 		UFunction* Function = FindUField<UFunction>(CompilerContext.Blueprint->SkeletonGeneratedClass, *OriginalNode->GetOuter()->GetName());
 
-		// When regenerating on load, we may need to import text on certain properties to force load the assets
-		TSharedPtr<FStructOnScope> LocalVarData;
-		if (Function && CompilerContext.Blueprint->bIsRegeneratingOnLoad)
-		{
-			if (Function->GetStructureSize() > 0 || !ensure(Function->PropertyLink == nullptr))
-			{
-				LocalVarData = MakeShareable(new FStructOnScope(Function));
-			}
-		}
-
 		for (TFieldIterator<FProperty> It(Function); It; ++It)
 		{
 			if (const FProperty* Property = *It)
@@ -871,16 +867,6 @@ void UK2Node_FunctionEntry::ExpandNode(class FKismetCompilerContext& CompilerCon
 							}
 							else
 							{
-								if (CompilerContext.Blueprint->bIsRegeneratingOnLoad)
-								{
-									// When regenerating on load, we want to force load assets referenced by local variables.
-									// This functionality is already handled when generating Terms in the Kismet Compiler for Arrays and Structs, so we do not have to worry about them.
-									if (LocalVar.VarType.PinCategory == UEdGraphSchema_K2::PC_Object || LocalVar.VarType.PinCategory == UEdGraphSchema_K2::PC_Class || LocalVar.VarType.PinCategory == UEdGraphSchema_K2::PC_Interface)
-									{
-										FBlueprintEditorUtils::PropertyValueFromString(Property, LocalVar.DefaultValue, LocalVarData->GetStructMemory());
-									}
-								}
-
 								// Set the default value
 								Schema->TrySetDefaultValue(*SetPin, LocalVar.DefaultValue);
 							}
diff --git a/Engine/Source/Editor/Kismet/Private/BlueprintCompilationManager.cpp b/Engine/Source/Editor/Kismet/Private/BlueprintCompilationManager.cpp
--- a/Engine/Source/Editor/Kismet/Private/BlueprintCompilationManager.cpp
+++ b/Engine/Source/Editor/Kismet/Private/BlueprintCompilationManager.cpp
@@ -2640,6 +2640,16 @@ UClass* FBlueprintCompilationManagerImpl::FastGenerateSkeletonClass(UBlueprint*
 
 					CompilerContext.SetCalculatedMetaDataAndFlags(NewFunction, EntryNode, Schema);
 				}
+
+				if (BP->bIsRegeneratingOnLoad)
+				{
+					// Ensure that the function's variable cache is up-to-date after property creation. Note that
+					// this may incur a load if the variable's default value (a string) refers to an external asset;
+					// that won't result in a package import until after the BP is compiled, and sometimes users will
+					// save the Blueprint after having set the variable's default value without also recompiling it.
+					// In that case we want to ensure these assets are loaded as part of regenerating classes on load.
+					EntryNode->RefreshFunctionVariableCache();
+				}
 			}
 		}
 	};
