Added new ReloadEnum method to the NodeDependingOnEnumInterface which will be utilized by the reload system to notify the node that the enumeration has been reloaded and that the given enum is the new address of the enum.

#rnx

diff --git a/Engine/Source/Editor/AnimGraph/Private/AnimGraphNode_BlendListByEnum.cpp b/Engine/Source/Editor/AnimGraph/Private/AnimGraphNode_BlendListByEnum.cpp
--- a/Engine/Source/Editor/AnimGraph/Private/AnimGraphNode_BlendListByEnum.cpp
+++ b/Engine/Source/Editor/AnimGraph/Private/AnimGraphNode_BlendListByEnum.cpp
@@ -331,6 +331,12 @@ void UAnimGraphNode_BlendListByEnum::BakeDataDuringCompilation(class FCompilerRe
 	}
 }
 
+void UAnimGraphNode_BlendListByEnum::ReloadEnum(class UEnum* InEnum)
+{
+	BoundEnum = InEnum;
+	CachedNodeTitle.MarkDirty();
+}
+
 /////////////////////////////////////////////////////
 
 #undef LOCTEXT_NAMESPACE
diff --git a/Engine/Source/Editor/AnimGraph/Public/AnimGraphNode_BlendListByEnum.h b/Engine/Source/Editor/AnimGraph/Public/AnimGraphNode_BlendListByEnum.h
--- a/Engine/Source/Editor/AnimGraph/Public/AnimGraphNode_BlendListByEnum.h
+++ b/Engine/Source/Editor/AnimGraph/Public/AnimGraphNode_BlendListByEnum.h
@@ -52,6 +52,7 @@ public:
 
 	// INodeDependingOnEnumInterface
 	virtual class UEnum* GetEnum() const override { return BoundEnum; }
+	virtual void ReloadEnum(class UEnum* InEnum) override;
 	virtual bool ShouldBeReconstructedAfterEnumChanged() const override {return true;}
 	// End of INodeDependingOnEnumInterface
 protected:
diff --git a/Engine/Source/Editor/BlueprintGraph/Classes/K2Node_BitmaskLiteral.h b/Engine/Source/Editor/BlueprintGraph/Classes/K2Node_BitmaskLiteral.h
--- a/Engine/Source/Editor/BlueprintGraph/Classes/K2Node_BitmaskLiteral.h
+++ b/Engine/Source/Editor/BlueprintGraph/Classes/K2Node_BitmaskLiteral.h
@@ -39,6 +39,7 @@ class UK2Node_BitmaskLiteral : public UK2Node, public INodeDependingOnEnumInterf
 
 	//~ Begin INodeDependingOnEnumInterface
 	virtual class UEnum* GetEnum() const override { return BitflagsEnum; }
+	virtual void ReloadEnum(class UEnum* InEnum) override;
 	virtual bool ShouldBeReconstructedAfterEnumChanged() const override { return true; }
 	//~ End INodeDependingOnEnumInterface
 
diff --git a/Engine/Source/Editor/BlueprintGraph/Classes/K2Node_CastByteToEnum.h b/Engine/Source/Editor/BlueprintGraph/Classes/K2Node_CastByteToEnum.h
--- a/Engine/Source/Editor/BlueprintGraph/Classes/K2Node_CastByteToEnum.h
+++ b/Engine/Source/Editor/BlueprintGraph/Classes/K2Node_CastByteToEnum.h
@@ -7,6 +7,7 @@
 #include "Textures/SlateIcon.h"
 #include "K2Node.h"
 #include "EdGraph/EdGraphNodeUtils.h"
+#include "NodeDependingOnEnumInterface.h"
 #include "K2Node_CastByteToEnum.generated.h"
 
 class FBlueprintActionDatabaseRegistrar;
@@ -15,7 +16,7 @@ class FNodeHandlingFunctor;
 class UEdGraph;
 
 UCLASS(MinimalAPI)
-class UK2Node_CastByteToEnum : public UK2Node
+class UK2Node_CastByteToEnum : public UK2Node, public INodeDependingOnEnumInterface
 {
 	GENERATED_UCLASS_BODY()
 
@@ -47,6 +48,12 @@ class UK2Node_CastByteToEnum : public UK2Node
 	virtual FText GetMenuCategory() const override;
 	//~ End UK2Node Interface
 
+	// INodeDependingOnEnumInterface
+	virtual class UEnum* GetEnum() const override { return Enum; }
+	virtual void ReloadEnum(class UEnum* InEnum) override;
+	virtual bool ShouldBeReconstructedAfterEnumChanged() const override { return false; }
+	// End of INodeDependingOnEnumInterface
+
 	virtual FName GetFunctionName() const;
 
 private:
diff --git a/Engine/Source/Editor/BlueprintGraph/Classes/K2Node_EnumLiteral.h b/Engine/Source/Editor/BlueprintGraph/Classes/K2Node_EnumLiteral.h
--- a/Engine/Source/Editor/BlueprintGraph/Classes/K2Node_EnumLiteral.h
+++ b/Engine/Source/Editor/BlueprintGraph/Classes/K2Node_EnumLiteral.h
@@ -42,6 +42,7 @@ class UK2Node_EnumLiteral : public UK2Node, public INodeDependingOnEnumInterface
 
 	// INodeDependingOnEnumInterface
 	virtual class UEnum* GetEnum() const override { return Enum; }
+	virtual void ReloadEnum(class UEnum* InEnum) override;
 	virtual bool ShouldBeReconstructedAfterEnumChanged() const override { return true; }
 	// End of INodeDependingOnEnumInterface
 
diff --git a/Engine/Source/Editor/BlueprintGraph/Classes/K2Node_ForEachElementInEnum.h b/Engine/Source/Editor/BlueprintGraph/Classes/K2Node_ForEachElementInEnum.h
--- a/Engine/Source/Editor/BlueprintGraph/Classes/K2Node_ForEachElementInEnum.h
+++ b/Engine/Source/Editor/BlueprintGraph/Classes/K2Node_ForEachElementInEnum.h
@@ -7,13 +7,14 @@
 #include "Textures/SlateIcon.h"
 #include "K2Node.h"
 #include "EdGraph/EdGraphNodeUtils.h"
+#include "NodeDependingOnEnumInterface.h"
 #include "K2Node_ForEachElementInEnum.generated.h"
 
 class FBlueprintActionDatabaseRegistrar;
 class UEdGraph;
 
 UCLASS(MinimalAPI)
-class UK2Node_ForEachElementInEnum : public UK2Node
+class UK2Node_ForEachElementInEnum : public UK2Node, public INodeDependingOnEnumInterface
 {
 	GENERATED_UCLASS_BODY()
 
@@ -40,6 +41,12 @@ class UK2Node_ForEachElementInEnum : public UK2Node
 	virtual void PostPlacedNewNode() override;
 	//~ End UK2Node Interface
 
+	// INodeDependingOnEnumInterface
+	virtual class UEnum* GetEnum() const override { return Enum; }
+	virtual void ReloadEnum(class UEnum* InEnum) override;
+	virtual bool ShouldBeReconstructedAfterEnumChanged() const override { return false; }
+	// End of INodeDependingOnEnumInterface
+
 private:
 	/** Constructing FText strings can be costly, so we cache the node's title */
 	FNodeTextCache CachedNodeTitle;
diff --git a/Engine/Source/Editor/BlueprintGraph/Classes/K2Node_GetNumEnumEntries.h b/Engine/Source/Editor/BlueprintGraph/Classes/K2Node_GetNumEnumEntries.h
--- a/Engine/Source/Editor/BlueprintGraph/Classes/K2Node_GetNumEnumEntries.h
+++ b/Engine/Source/Editor/BlueprintGraph/Classes/K2Node_GetNumEnumEntries.h
@@ -37,6 +37,7 @@ class UK2Node_GetNumEnumEntries : public UK2Node, public INodeDependingOnEnumInt
 
 	// INodeDependingOnEnumInterface
 	virtual class UEnum* GetEnum() const override { return Enum; }
+	virtual void ReloadEnum(class UEnum* InEnum) override;
 	virtual bool ShouldBeReconstructedAfterEnumChanged() const override {return false;}
 	// End of INodeDependingOnEnumInterface
 
diff --git a/Engine/Source/Editor/BlueprintGraph/Classes/K2Node_Select.h b/Engine/Source/Editor/BlueprintGraph/Classes/K2Node_Select.h
--- a/Engine/Source/Editor/BlueprintGraph/Classes/K2Node_Select.h
+++ b/Engine/Source/Editor/BlueprintGraph/Classes/K2Node_Select.h
@@ -83,6 +83,7 @@ public:
 
 	// INodeDependingOnEnumInterface
 	virtual class UEnum* GetEnum() const override { return Enum; }
+	virtual void ReloadEnum(class UEnum* InEnum) override;
 	virtual bool ShouldBeReconstructedAfterEnumChanged() const override { return true; }
 	// End of INodeDependingOnEnumInterface
 
diff --git a/Engine/Source/Editor/BlueprintGraph/Classes/K2Node_SwitchEnum.h b/Engine/Source/Editor/BlueprintGraph/Classes/K2Node_SwitchEnum.h
--- a/Engine/Source/Editor/BlueprintGraph/Classes/K2Node_SwitchEnum.h
+++ b/Engine/Source/Editor/BlueprintGraph/Classes/K2Node_SwitchEnum.h
@@ -32,6 +32,7 @@ class UK2Node_SwitchEnum : public UK2Node_Switch, public INodeDependingOnEnumInt
 
 	// INodeDependingOnEnumInterface
 	virtual class UEnum* GetEnum() const override { return Enum; }
+	virtual void ReloadEnum(class UEnum* InEnum) override;
 	virtual bool ShouldBeReconstructedAfterEnumChanged() const override {return true;}
 	// End of INodeDependingOnEnumInterface
 
diff --git a/Engine/Source/Editor/BlueprintGraph/Classes/NodeDependingOnEnumInterface.h b/Engine/Source/Editor/BlueprintGraph/Classes/NodeDependingOnEnumInterface.h
--- a/Engine/Source/Editor/BlueprintGraph/Classes/NodeDependingOnEnumInterface.h
+++ b/Engine/Source/Editor/BlueprintGraph/Classes/NodeDependingOnEnumInterface.h
@@ -19,5 +19,7 @@ class BLUEPRINTGRAPH_API INodeDependingOnEnumInterface
 
 	virtual class UEnum* GetEnum() const PURE_VIRTUAL(INodeDependingOnEnumInterface::GetEnum,return NULL;);
 
+	virtual void ReloadEnum(class UEnum*) PURE_VIRTUAL(INodeDependingOnEnumInterface::ReloadEnum, return;);
+
 	virtual bool ShouldBeReconstructedAfterEnumChanged() const {return false;}
 };
diff --git a/Engine/Source/Editor/BlueprintGraph/Private/K2Node_BitmaskLiteral.cpp b/Engine/Source/Editor/BlueprintGraph/Private/K2Node_BitmaskLiteral.cpp
--- a/Engine/Source/Editor/BlueprintGraph/Private/K2Node_BitmaskLiteral.cpp
+++ b/Engine/Source/Editor/BlueprintGraph/Private/K2Node_BitmaskLiteral.cpp
@@ -155,4 +155,9 @@ FText UK2Node_BitmaskLiteral::GetMenuCategory() const
 	return FEditorCategoryUtils::GetCommonCategory(FCommonEditorCategory::Math);
 }
 
+void UK2Node_BitmaskLiteral::ReloadEnum(class UEnum* InEnum)
+{
+	BitflagsEnum = InEnum;
+}
+
 #undef LOCTEXT_NAMESPACE
diff --git a/Engine/Source/Editor/BlueprintGraph/Private/K2Node_CastByteToEnum.cpp b/Engine/Source/Editor/BlueprintGraph/Private/K2Node_CastByteToEnum.cpp
--- a/Engine/Source/Editor/BlueprintGraph/Private/K2Node_CastByteToEnum.cpp
+++ b/Engine/Source/Editor/BlueprintGraph/Private/K2Node_CastByteToEnum.cpp
@@ -205,3 +205,9 @@ FText UK2Node_CastByteToEnum::GetMenuCategory() const
 {
 	return FEditorCategoryUtils::GetCommonCategory(FCommonEditorCategory::Enum);
 }
+
+void UK2Node_CastByteToEnum::ReloadEnum(class UEnum* InEnum)
+{
+	Enum = InEnum;
+	CachedTooltip.MarkDirty();
+}
diff --git a/Engine/Source/Editor/BlueprintGraph/Private/K2Node_EnumLiteral.cpp b/Engine/Source/Editor/BlueprintGraph/Private/K2Node_EnumLiteral.cpp
--- a/Engine/Source/Editor/BlueprintGraph/Private/K2Node_EnumLiteral.cpp
+++ b/Engine/Source/Editor/BlueprintGraph/Private/K2Node_EnumLiteral.cpp
@@ -174,3 +174,8 @@ FText UK2Node_EnumLiteral::GetMenuCategory() const
 	return FEditorCategoryUtils::GetCommonCategory(FCommonEditorCategory::Enum);
 }
 
+void UK2Node_EnumLiteral::ReloadEnum(class UEnum* InEnum)
+{
+	Enum = InEnum;
+	CachedTooltip.MarkDirty();
+}
diff --git a/Engine/Source/Editor/BlueprintGraph/Private/K2Node_ForEachElementInEnum.cpp b/Engine/Source/Editor/BlueprintGraph/Private/K2Node_ForEachElementInEnum.cpp
--- a/Engine/Source/Editor/BlueprintGraph/Private/K2Node_ForEachElementInEnum.cpp
+++ b/Engine/Source/Editor/BlueprintGraph/Private/K2Node_ForEachElementInEnum.cpp
@@ -347,4 +347,10 @@ void UK2Node_ForEachElementInEnum::PostPlacedNewNode()
 	}
 }
 
+void UK2Node_ForEachElementInEnum::ReloadEnum(class UEnum* InEnum)
+{
+	Enum = InEnum;
+	CachedNodeTitle.MarkDirty();
+}
+
 #undef LOCTEXT_NAMESPACE
diff --git a/Engine/Source/Editor/BlueprintGraph/Private/K2Node_GetNumEnumEntries.cpp b/Engine/Source/Editor/BlueprintGraph/Private/K2Node_GetNumEnumEntries.cpp
--- a/Engine/Source/Editor/BlueprintGraph/Private/K2Node_GetNumEnumEntries.cpp
+++ b/Engine/Source/Editor/BlueprintGraph/Private/K2Node_GetNumEnumEntries.cpp
@@ -123,3 +123,9 @@ FText UK2Node_GetNumEnumEntries::GetMenuCategory() const
 	return FEditorCategoryUtils::GetCommonCategory(FCommonEditorCategory::Enum);
 }
 
+void UK2Node_GetNumEnumEntries::ReloadEnum(class UEnum* InEnum)
+{
+	Enum = InEnum;
+	CachedTooltip.MarkDirty();
+	CachedNodeTitle.MarkDirty();
+}
diff --git a/Engine/Source/Editor/BlueprintGraph/Private/K2Node_Select.cpp b/Engine/Source/Editor/BlueprintGraph/Private/K2Node_Select.cpp
--- a/Engine/Source/Editor/BlueprintGraph/Private/K2Node_Select.cpp
+++ b/Engine/Source/Editor/BlueprintGraph/Private/K2Node_Select.cpp
@@ -953,4 +953,9 @@ void UK2Node_Select::ExpandNode(class FKismetCompilerContext& CompilerContext, U
 	}
 }
 
+void UK2Node_Select::ReloadEnum(class UEnum* InEnum)
+{
+	SetEnum(InEnum, true);
+}
+
 #undef LOCTEXT_NAMESPACE
diff --git a/Engine/Source/Editor/BlueprintGraph/Private/K2Node_SwitchEnum.cpp b/Engine/Source/Editor/BlueprintGraph/Private/K2Node_SwitchEnum.cpp
--- a/Engine/Source/Editor/BlueprintGraph/Private/K2Node_SwitchEnum.cpp
+++ b/Engine/Source/Editor/BlueprintGraph/Private/K2Node_SwitchEnum.cpp
@@ -220,4 +220,9 @@ void UK2Node_SwitchEnum::RemovePinFromSwitchNode(UEdGraphPin* Pin)
 	}
 }
 
+void UK2Node_SwitchEnum::ReloadEnum(class UEnum* InEnum)
+{
+	SetEnum(InEnum);
+}
+
 #undef LOCTEXT_NAMESPACE
