Check for a valid Blueprint before attempting to refresh the local variable reference cache on function entry nodes.

#rnx

diff --git a/Engine/Source/Editor/BlueprintGraph/Private/K2Node_FunctionEntry.cpp b/Engine/Source/Editor/BlueprintGraph/Private/K2Node_FunctionEntry.cpp
--- a/Engine/Source/Editor/BlueprintGraph/Private/K2Node_FunctionEntry.cpp
+++ b/Engine/Source/Editor/BlueprintGraph/Private/K2Node_FunctionEntry.cpp
@@ -432,7 +432,7 @@ TSharedPtr<FStructOnScope> UK2Node_FunctionEntry::GetFunctionVariableCache(bool
 		FunctionVariableCache.Reset();
 	}
 
-	if (!FunctionVariableCache.IsValid() || !FunctionVariableCache->IsValid())
+	if (!FunctionVariableCache.IsValid() && HasValidBlueprint() && LocalVariables.Num() > 0)
 	{
 		// Locate the UFunction object in the class hierarchy, starting at the current class. Note that local
 		// variables are generated as fields (properties) within the function context that contains this node,
@@ -442,13 +442,10 @@ TSharedPtr<FStructOnScope> UK2Node_FunctionEntry::GetFunctionVariableCache(bool
 		// an input/output argument or local variable, whereas the generated class must be explicitly recompiled.
 		if (UFunction* const Function = FFunctionFromNodeHelper::FunctionFromNode(this))
 		{
-			if (LocalVariables.Num() > 0)
-			{
-				FunctionVariableCache = MakeShared<FStructOnScope>(Function);
-				FunctionVariableCache->SetPackage(GetOutermost());
+			FunctionVariableCache = MakeShared<FStructOnScope>(Function);
+			FunctionVariableCache->SetPackage(GetOutermost());
 
-				RefreshFunctionVariableCache();
-			}
+			RefreshFunctionVariableCache();
 		}
 	}
 	
