Split pin building logic up for UK2Node_StructMemberSet

Tweak to allow UK2Node_StructMemberSet to rebuild pins in a more optimal way for large sparsely-connected structs that are built by the anim BP compiler

diff --git a/Engine/Source/Editor/BlueprintGraph/Classes/K2Node_StructMemberSet.h b/Engine/Source/Editor/BlueprintGraph/Classes/K2Node_StructMemberSet.h
--- a/Engine/Source/Editor/BlueprintGraph/Classes/K2Node_StructMemberSet.h
+++ b/Engine/Source/Editor/BlueprintGraph/Classes/K2Node_StructMemberSet.h
@@ -36,6 +36,9 @@ class UK2Node_StructMemberSet : public UK2Node_StructOperation
 	virtual class FNodeHandlingFunctor* CreateNodeHandler(class FKismetCompilerContext& CompilerContext) const override;
 	// End of UK2Node interface
 
+	// Helper for AllocateDefaultPins
+	BLUEPRINTGRAPH_API void AllocateExecPins();
+
 private:
 	/** Constructing FText strings can be costly, so we cache the node's title/tooltip */
 	FNodeTextCache CachedTooltip;
diff --git a/Engine/Source/Editor/BlueprintGraph/Private/K2Node_StructMemberSet.cpp b/Engine/Source/Editor/BlueprintGraph/Private/K2Node_StructMemberSet.cpp
--- a/Engine/Source/Editor/BlueprintGraph/Private/K2Node_StructMemberSet.cpp
+++ b/Engine/Source/Editor/BlueprintGraph/Private/K2Node_StructMemberSet.cpp
@@ -37,11 +37,16 @@ void UK2Node_StructMemberSet::PostEditChangeProperty(FPropertyChangedEvent& Prop
 	Super::PostEditChangeProperty(PropertyChangedEvent);
 }
 
-void UK2Node_StructMemberSet::AllocateDefaultPins()
+void UK2Node_StructMemberSet::AllocateExecPins()
 {
-	// Add the execution sequencing pin
 	CreatePin(EGPD_Input, UEdGraphSchema_K2::PC_Exec, UEdGraphSchema_K2::PN_Execute);
 	CreatePin(EGPD_Output, UEdGraphSchema_K2::PC_Exec, UEdGraphSchema_K2::PN_Then);
+}
+
+void UK2Node_StructMemberSet::AllocateDefaultPins()
+{
+	// Add the execution sequencing pin
+	AllocateExecPins();
 
 	// Display any currently visible optional pins
 	{
