Implemented PostPlacedNewNode for BreakStruct so that the Serialize function does not mutate the object during save
Moved the Serialize function mutation to PreSave as well as a precaution

diff --git a/Engine/Source/Editor/BlueprintGraph/Classes/K2Node_BreakStruct.h b/Engine/Source/Editor/BlueprintGraph/Classes/K2Node_BreakStruct.h
--- a/Engine/Source/Editor/BlueprintGraph/Classes/K2Node_BreakStruct.h
+++ b/Engine/Source/Editor/BlueprintGraph/Classes/K2Node_BreakStruct.h
@@ -36,6 +36,7 @@ class UK2Node_BreakStruct : public UK2Node_StructMemberGet
 	static bool CanBeSplit(const UScriptStruct* Struct) { return CanBeBroken(Struct); }
 
 	// UObject interface
+	virtual void PreSave(FObjectPreSaveContext SaveContext) override;
 	virtual void Serialize(FArchive& Ar) override;
 	// End of UObject interface
 
@@ -47,6 +48,7 @@ class UK2Node_BreakStruct : public UK2Node_StructMemberGet
 	virtual FText GetTooltipText() const override;
 	virtual void ValidateNodeDuringCompilation(class FCompilerResultsLog& MessageLog) const override;
 	virtual FSlateIcon GetIconAndTint(FLinearColor& OutColor) const override;
+	virtual void PostPlacedNewNode() override;
 	//~ End  UEdGraphNode Interface
 
 	//~ Begin K2Node Interface
diff --git a/Engine/Source/Editor/BlueprintGraph/Classes/K2Node_MakeStruct.h b/Engine/Source/Editor/BlueprintGraph/Classes/K2Node_MakeStruct.h
--- a/Engine/Source/Editor/BlueprintGraph/Classes/K2Node_MakeStruct.h
+++ b/Engine/Source/Editor/BlueprintGraph/Classes/K2Node_MakeStruct.h
@@ -37,6 +37,7 @@ class UK2Node_MakeStruct : public UK2Node_StructMemberSet
 	BLUEPRINTGRAPH_API static bool CanBeSplit(const UScriptStruct* Struct, UBlueprint* InBP);
 
 	// UObject interface
+	virtual void PreSave(FObjectPreSaveContext SaveContext) override;
 	virtual void Serialize(FArchive& Ar) override;
 	// End of UObject interface
 
diff --git a/Engine/Source/Editor/BlueprintGraph/Private/K2Node_BreakStruct.cpp b/Engine/Source/Editor/BlueprintGraph/Private/K2Node_BreakStruct.cpp
--- a/Engine/Source/Editor/BlueprintGraph/Private/K2Node_BreakStruct.cpp
+++ b/Engine/Source/Editor/BlueprintGraph/Private/K2Node_BreakStruct.cpp
@@ -15,6 +15,7 @@
 #include "PropertyCustomizationHelpers.h"
 #include "BlueprintEditorSettings.h"
 #include "Kismet/KismetMathLibrary.h"
+#include "UObject/ObjectSaveContext.h"
 
 #define LOCTEXT_NAMESPACE "K2Node_BreakStruct"
 
@@ -377,6 +378,24 @@ FText UK2Node_BreakStruct::GetMenuCategory() const
 	return FEditorCategoryUtils::GetCommonCategory(FCommonEditorCategory::Struct);
 }
 
+void UK2Node_BreakStruct::PreSave(FObjectPreSaveContext SaveContext)
+{
+	Super::PreSave(SaveContext);
+	UBlueprint* Blueprint = FBlueprintEditorUtils::FindBlueprintForNode(this);
+	if (Blueprint && !Blueprint->bBeingCompiled)
+	{
+		bMadeAfterOverridePinRemoval = true;
+	}
+}
+
+void UK2Node_BreakStruct::PostPlacedNewNode()
+{
+	Super::PostPlacedNewNode();
+
+	// New nodes automatically have this set.
+	bMadeAfterOverridePinRemoval = true;
+}
+
 void UK2Node_BreakStruct::Serialize(FArchive& Ar)
 {
 	Super::Serialize(Ar);
@@ -404,15 +423,6 @@ void UK2Node_BreakStruct::Serialize(FArchive& Ar)
 			}
 		}
 	}
-	else if (Ar.IsSaving() && !Ar.IsTransacting())
-	{
-		UBlueprint* Blueprint = FBlueprintEditorUtils::FindBlueprintForNode(this);
-
-		if (Blueprint && !Blueprint->bBeingCompiled)
-		{
-			bMadeAfterOverridePinRemoval = true;
-		}
-	}
 }
 
 void UK2Node_BreakStruct::ConvertDeprecatedNode(UEdGraph* Graph, bool bOnlySafeChanges)
diff --git a/Engine/Source/Editor/BlueprintGraph/Private/K2Node_MakeStruct.cpp b/Engine/Source/Editor/BlueprintGraph/Private/K2Node_MakeStruct.cpp
--- a/Engine/Source/Editor/BlueprintGraph/Private/K2Node_MakeStruct.cpp
+++ b/Engine/Source/Editor/BlueprintGraph/Private/K2Node_MakeStruct.cpp
@@ -16,6 +16,7 @@
 #include "BlueprintActionDatabaseRegistrar.h"
 #include "PropertyCustomizationHelpers.h"
 #include "Kismet/KismetMathLibrary.h"
+#include "UObject/ObjectSaveContext.h"
 
 #define LOCTEXT_NAMESPACE "K2Node_MakeStruct"
 
@@ -336,6 +337,16 @@ FText UK2Node_MakeStruct::GetMenuCategory() const
 	return FEditorCategoryUtils::GetCommonCategory(FCommonEditorCategory::Struct);
 }
 
+void UK2Node_MakeStruct::PreSave(FObjectPreSaveContext SaveContext)
+{
+	Super::PreSave(SaveContext);
+	UBlueprint* Blueprint = FBlueprintEditorUtils::FindBlueprintForNode(this);
+	if (Blueprint && !Blueprint->bBeingCompiled)
+	{
+		bMadeAfterOverridePinRemoval = true;
+	}
+}
+
 void UK2Node_MakeStruct::PostPlacedNewNode()
 {
 	Super::PostPlacedNewNode();
@@ -348,11 +359,10 @@ void UK2Node_MakeStruct::Serialize(FArchive& Ar)
 {
 	Super::Serialize(Ar);
 
-	UBlueprint* Blueprint = FBlueprintEditorUtils::FindBlueprintForNode(this);
-
-	if (Blueprint && !Ar.IsTransacting() && !HasAllFlags(RF_Transient))
+	if (Ar.IsLoading() && !Ar.IsTransacting() && !HasAllFlags(RF_Transient))
 	{
-		if (Ar.IsLoading() && !bMadeAfterOverridePinRemoval)
+		UBlueprint* Blueprint = FBlueprintEditorUtils::FindBlueprintForNode(this);
+		if (Blueprint && !bMadeAfterOverridePinRemoval)
 		{
 			// Check if this node actually requires warning the user that functionality has changed.
 			bMadeAfterOverridePinRemoval = true;
@@ -416,13 +426,6 @@ void UK2Node_MakeStruct::Serialize(FArchive& Ar)
 				}
 			}
 		}
-		else if (Ar.IsSaving())
-		{
-			if (!Blueprint->bBeingCompiled)
-			{
-				bMadeAfterOverridePinRemoval = true;
-			}
-		}
 	}
 }
 
