Fixed issue where blueprints derived from editor-only classes would not always show all possible actions in their right-click context menu on the Event Graph.  One case was a user BP derived from AnimationModifier would be missing "Animation Blueprint Library" actions until it was force-saved (a fluke that happened to work around the true problem).

diff --git a/Engine/Source/Editor/BlueprintGraph/Private/K2Node_CallFunction.cpp b/Engine/Source/Editor/BlueprintGraph/Private/K2Node_CallFunction.cpp
--- a/Engine/Source/Editor/BlueprintGraph/Private/K2Node_CallFunction.cpp
+++ b/Engine/Source/Editor/BlueprintGraph/Private/K2Node_CallFunction.cpp
@@ -1319,7 +1319,12 @@ bool UK2Node_CallFunction::CanEditorOnlyFunctionBeCalled(const UFunction* InFunc
 	if (InFunction && InObject &&
 		(IsEditorOnlyObject(InFunction) || InFunction->HasAnyFunctionFlags(FUNC_EditorOnly)))
 	{
-		return IsEditorOnlyObject(InObject);
+		if (!IsEditorOnlyObject(InObject))
+		{
+			// InObject isn't editor-only, but it's still possible that it's a blueprint derived from an editor-only class, so let's check for that case
+			const UBlueprint* InObjectAsBP = Cast<const UBlueprint>(InObject->GetOuter());
+			return (InObjectAsBP && InObjectAsBP->ParentClass && IsEditorOnlyObject(InObjectAsBP->ParentClass));
+		}
 	}
 
 	return true;
