Update comments about disallowing containers for promotable operators

#rnx

diff --git a/Engine/Source/Editor/BlueprintGraph/Private/K2Node_PromotableOperator.cpp b/Engine/Source/Editor/BlueprintGraph/Private/K2Node_PromotableOperator.cpp
--- a/Engine/Source/Editor/BlueprintGraph/Private/K2Node_PromotableOperator.cpp
+++ b/Engine/Source/Editor/BlueprintGraph/Private/K2Node_PromotableOperator.cpp
@@ -531,7 +531,7 @@ bool UK2Node_PromotableOperator::IsConnectionDisallowed(const UEdGraphPin* MyPin
 
 	const UEdGraphSchema_K2* K2Schema = GetDefault<UEdGraphSchema_K2>();
 
-	// #TODO_BH Just disallow containers for now
+	// Container types cannot be promotable operators as they have their own special case wildcard propagation
 	if (OtherPin->PinType.IsContainer())
 	{
 		OutReason = LOCTEXT("NoExecPinsAllowed", "Promotable Operator nodes cannot have containers or references.").ToString();
