Fix merge error with missing if (Function) test that allowed Function->HasMetaData call to be made on nullptr.
#rnx

diff --git a/Engine/Source/Editor/BlueprintGraph/Private/K2Node_CallFunction.cpp b/Engine/Source/Editor/BlueprintGraph/Private/K2Node_CallFunction.cpp
--- a/Engine/Source/Editor/BlueprintGraph/Private/K2Node_CallFunction.cpp
+++ b/Engine/Source/Editor/BlueprintGraph/Private/K2Node_CallFunction.cpp
@@ -2054,20 +2054,20 @@ void UK2Node_CallFunction::ValidateNodeDuringCompilation(class FCompilerResultsL
 	if (!bAllowUnsafeBlueprintCalls)
 	{
 		const bool bIsUncookedOnlyFunction = Function && Function->GetOutermost()->HasAllPackagesFlags(PKG_UncookedOnly);
-		if (	bIsUncookedOnlyFunction &&
-				// Only allow calls to uncooked only functions from editor only/uncooked only
-				// contexts:
-				!(	GetOutermost()->HasAnyPackageFlags(PKG_UncookedOnly|PKG_EditorOnly) ||
-					bIsEditorOnlyBlueprintBaseClass ))
+		if (bIsUncookedOnlyFunction &&
+			// Only allow calls to uncooked only functions from editor only/uncooked only
+			// contexts:
+			!(GetOutermost()->HasAnyPackageFlags(PKG_UncookedOnly | PKG_EditorOnly) ||
+				bIsEditorOnlyBlueprintBaseClass))
 		{
 			MessageLog.Error(*LOCTEXT("UncookedOnlyError", "Attempting to call uncooked only function @@ in runtime blueprint").ToString(), this);
 		}
-	
+
 		// Ensure that editor module BP exposed UFunctions can only be called in blueprints for which the base class is also part of an editor module
 		// Also check for functions wrapped in WITH_EDITOR 
 		if (Function && Blueprint &&
 			(IsEditorOnlyObject(Function) || Function->HasAnyFunctionFlags(FUNC_EditorOnly)))
-		{	
+		{
 			if (!bIsEditorOnlyBlueprintBaseClass)
 			{
 				FString const FunctName = Function->GetName();
@@ -2075,7 +2075,10 @@ void UK2Node_CallFunction::ValidateNodeDuringCompilation(class FCompilerResultsL
 				MessageLog.Error(*FText::Format(WarningFormat, FText::FromString(FunctName)).ToString(), this);
 			}
 		}
+	}
 
+	if (Function)
+	{
 		// enforce UnsafeDuringActorConstruction keyword
 		if (Function->HasMetaData(FBlueprintMetadata::MD_UnsafeForConstructionScripts))
 		{
