Refresh a call function node on NotifyPinConnectionListChanged to get the correct literal text box values after propegation of the array type.

#rnx

diff --git a/Engine/Source/Editor/BlueprintGraph/Private/K2Node_CallFunction.cpp b/Engine/Source/Editor/BlueprintGraph/Private/K2Node_CallFunction.cpp
--- a/Engine/Source/Editor/BlueprintGraph/Private/K2Node_CallFunction.cpp
+++ b/Engine/Source/Editor/BlueprintGraph/Private/K2Node_CallFunction.cpp
@@ -1213,7 +1213,9 @@ void UK2Node_CallFunction::NotifyPinConnectionListChanged(UEdGraphPin* Pin)
 	FCustomStructureParamHelper::UpdateCustomStructurePins(GetTargetFunction(), this, Pin);
 
 	// Refresh the node to hide internal-only pins once the [invalid] connection has been broken
-	if (Pin->bHidden && Pin->bNotConnectable && Pin->LinkedTo.Num() == 0)
+	// If the pin was a container then it needs to be refreshed to get the correct pin literal text boxes
+	// for its default value
+	if (Pin->PinType.IsContainer() || (Pin->bHidden && Pin->bNotConnectable && Pin->LinkedTo.Num() == 0))
 	{
 		GetGraph()->NotifyGraphChanged();
 	}
