Fix bug where it was failing to separate the category from translated description which caused false matches
Fix it to strip the "Class" from category description, this was prioritizing interface functions too high
Change the way the shorter length weight works by giving a bonus for matching short words instead of penalizing for irrelevant long words
Changed all word match bonuses to use the best one so redundant keywords like "add add" don't get prioritized and lower percentage match weight
Bump the description weight if there are no keywords, and minor weight tweaks in general

diff --git a/Engine/Source/Editor/BlueprintGraph/Private/EdGraphSchema_K2.cpp b/Engine/Source/Editor/BlueprintGraph/Private/EdGraphSchema_K2.cpp
--- a/Engine/Source/Editor/BlueprintGraph/Private/EdGraphSchema_K2.cpp
+++ b/Engine/Source/Editor/BlueprintGraph/Private/EdGraphSchema_K2.cpp
@@ -7402,13 +7402,20 @@ bool UEdGraphSchema_K2::SafeDeleteNodeFromGraph(UEdGraph* Graph, UEdGraphNode* N
 /** CVars for tweaking how the blueprint context menu search picks the best match */
 namespace ContextMenuConsoleVariables
 {
-	/** Increasing this weight will make shorter words preferred */
-	static float ShorterWeight = 15.0f;
+	/** Increasing this weight will give a bonus to shorter matching words */
+	static float ShorterWeight = 10.0f;
 	static FAutoConsoleVariableRef CVarShorterWeight(
 		TEXT("ContextMenu.ShorterWeight"), ShorterWeight,
 		TEXT("Increasing this weight will make shorter words preferred"),
 		ECVF_Default);
 
+	/** When calculating shorter weight, this is the maximum length to make it relative to */
+	static int32 MaxWordLength = 30;
+	static FAutoConsoleVariableRef CVarMaxWordLength(
+		TEXT("ContextMenu.MaxWordLength"), MaxWordLength,
+		TEXT("Maximum length to count while awarding short word weight"),
+		ECVF_Default);
+
 	/** Increasing this will prefer whole percentage matches when comparing the keyword to what the user has typed in */
 	static float PercentageMatchWeightMultiplier = 1.0f;
 	static FAutoConsoleVariableRef CVarPercentageMatchWeightMultiplier(
@@ -7438,7 +7445,7 @@ namespace ContextMenuConsoleVariables
 		ECVF_Default);
 
 	/** How much weight the node's title has */
-	static float NodeTitleWeight = 5.0f;
+	static float NodeTitleWeight = 10.0f;
 	static FAutoConsoleVariableRef CVarNodeTitleWeight(
 		TEXT("ContextMenu.NodeTitleWeight"), NodeTitleWeight,
 		TEXT("The amount of weight placed on the search items title"),
@@ -7451,18 +7458,18 @@ namespace ContextMenuConsoleVariables
 		TEXT("The amount of weight placed on search items keyword"),
 		ECVF_Default);
 
-	/** The multiplier given if the keyword starts with a letter the user typed in */
-	static float StartsWithBonusWeightMultiplier = 5.0f;
+	/** The multiplier given if the keyword starts with a term the user typed in */
+	static float StartsWithBonusWeightMultiplier = 2.0f;
 	static FAutoConsoleVariableRef CVarStartsWithBonusWeightMultiplier(
 		TEXT("ContextMenu.StartsWithBonusWeightMultiplier"), StartsWithBonusWeightMultiplier,
-		TEXT("The multiplier given if the keyword starts with a letter the user typed in"),
+		TEXT("The multiplier given if the keyword starts with a term the user typed in"),
 		ECVF_Default);
 
-	/** The multiplier given if the keyword starts with a letter the user typed in */
+	/** The multiplier given if the keyword starts with a term the user typed in */
 	static float WordContainsLetterWeightMultiplier = 0.5f;
 	static FAutoConsoleVariableRef CVarWordContainsLetterWeightMultiplier(
 		TEXT("ContextMenu.WordContainsLetterWeightMultiplier"), WordContainsLetterWeightMultiplier,
-		TEXT("The multiplier given if the keyword only contains a letter the user typed in"),
+		TEXT("The multiplier given if the keyword only contains a term the user typed in"),
 		ECVF_Default);
 
 	/** The bonus given if node is a favorite */
@@ -7506,16 +7513,16 @@ struct FContextMenuWeightDebugInfo
 	/** The calculated match percentage */
 	float PercMatch = 0.0f;
 	float ShorterWeight = 0.0f;
-	float CategoryBonusWieight = 0.0f;
+	float CategoryBonusWeight = 0.0f;
 
 	/**
 	* Print out the debug info about this weight info to the console
 	*/
-	void Print()
+	void Print(const FString& Desc)
 	{
-		UE_LOG(LogTemp, Warning, TEXT("[Weight Debug info] \
-		TotalWeight: %-8.2f | PercentageMatchWeight: %-8.2f | PercMatch: %-8.2f | ShorterWeight: %-8.2f | CategoryBonusWeight: %-8.2f | KeywordArrayWeight: %-8.2f | DescriptionWeight: %-8.2f | NodeTitleWeight: %-8.2f | CategoryWeight: %-8.2f | Fav. Bonus:%-8.2f\n"),
-			TotalWeight, PercentageMatchWeight, PercMatch, ShorterWeight, CategoryBonusWieight, KeywordArrayWeight, DescriptionWeight, NodeTitleWeight, CategoryWeight, FavoriteBonusWeight);
+		UE_LOG(LogTemp, Warning, TEXT("[Weight for %s] \
+TotalWeight: %-8.2f | PercentageMatchWeight: %-8.2f | PercMatch: %-8.2f | ShorterWeight: %-8.2f | CategoryBonusWeight: %-8.2f | KeywordArrayWeight: %-8.2f | DescriptionWeight: %-8.2f | NodeTitleWeight: %-8.2f | CategoryWeight: %-8.2f | Fav. Bonus:%-8.2f\n"),
+			*Desc, TotalWeight, PercentageMatchWeight, PercMatch, ShorterWeight, CategoryBonusWeight, KeywordArrayWeight, DescriptionWeight, NodeTitleWeight, CategoryWeight, FavoriteBonusWeight);
 	}
 };
 
@@ -7554,13 +7561,17 @@ float UEdGraphSchema_K2::GetActionFilteredWeight(const FGraphActionListBuilderBa
 		// and keywords, so we only need to use the first one for filtering.
 		const FString& SearchText = InCurrentAction.GetSearchTextForFirstAction();
 
+		// If there are no keywords, bump the weight on description to compensate
+		const TArray<FString>& LocKeywords = InCurrentAction.GetLocalizedSearchKeywordsArrayForFirstAction();
+		const float DescriptionWeight = LocKeywords.Num() > 0 ? ContextMenuConsoleVariables::DescriptionWeight : ContextMenuConsoleVariables::DescriptionWeight * 2.0f;
+
 		// First the localized keywords
-		WeightedArrayList.Add(FArrayWithWeight(&InCurrentAction.GetLocalizedSearchKeywordsArrayForFirstAction(), ContextMenuConsoleVariables::KeywordWeight, &OutDebugInfo.KeywordArrayWeight));
+		WeightedArrayList.Add(FArrayWithWeight(&LocKeywords, ContextMenuConsoleVariables::KeywordWeight, &OutDebugInfo.KeywordArrayWeight));
 
 		// The localized description
-		WeightedArrayList.Add(FArrayWithWeight(&InCurrentAction.GetLocalizedMenuDescriptionArrayForFirstAction(), ContextMenuConsoleVariables::DescriptionWeight, &OutDebugInfo.DescriptionWeight));
+		WeightedArrayList.Add(FArrayWithWeight(&InCurrentAction.GetLocalizedMenuDescriptionArrayForFirstAction(), DescriptionWeight, &OutDebugInfo.DescriptionWeight));
 
-		// The node search localized title weight
+		// The node search localized title weight, in most case this will be the same as description
 		WeightedArrayList.Add(FArrayWithWeight(&InCurrentAction.GetLocalizedSearchTitleArrayForFirstAction(), ContextMenuConsoleVariables::NodeTitleWeight, &OutDebugInfo.NodeTitleWeight));
 
 		// The localized category
@@ -7570,7 +7581,7 @@ float UEdGraphSchema_K2::GetActionFilteredWeight(const FGraphActionListBuilderBa
 		WeightedArrayList.Add(FArrayWithWeight(&InCurrentAction.GetSearchKeywordsArrayForFirstAction(), ContextMenuConsoleVariables::KeywordWeight, &OutDebugInfo.KeywordArrayWeight));
 
 		// The description
-		WeightedArrayList.Add(FArrayWithWeight(&InCurrentAction.GetMenuDescriptionArrayForFirstAction(), ContextMenuConsoleVariables::DescriptionWeight, &OutDebugInfo.DescriptionWeight));
+		WeightedArrayList.Add(FArrayWithWeight(&InCurrentAction.GetMenuDescriptionArrayForFirstAction(), DescriptionWeight, &OutDebugInfo.DescriptionWeight));
 
 		// The node search title weight
 		WeightedArrayList.Add(FArrayWithWeight(&InCurrentAction.GetSearchTitleArrayForFirstAction(), ContextMenuConsoleVariables::NodeTitleWeight, &OutDebugInfo.NodeTitleWeight));
@@ -7623,7 +7634,7 @@ float UEdGraphSchema_K2::GetActionFilteredWeight(const FGraphActionListBuilderBa
 					// Check the subcategory of the object to cover more more complex struct types (LinearColor, date time, etc)
 					if (UObject* const SubCatObj = FromPin->PinType.PinSubCategoryObject.Get())
 					{
-						const FString& SubCatObjName = SubCatObj->GetFullName();
+						const FString& SubCatObjName = SubCatObj->GetPathName();
 						// The pin SubObjectCategory names don't have any spaces, so split up the category
 						TArray<FString> DelimitedArray;
 						InActionCategory.ParseIntoArray(DelimitedArray, TEXT(" "), true);
@@ -7646,7 +7657,7 @@ float UEdGraphSchema_K2::GetActionFilteredWeight(const FGraphActionListBuilderBa
 					if (bAddMatchBonus)
 					{
 						TotalWeight += ContextMenuConsoleVariables::MatchingFromPinCategory;
-						OutDebugInfo.CategoryBonusWieight += ContextMenuConsoleVariables::MatchingFromPinCategory;
+						OutDebugInfo.CategoryBonusWeight += ContextMenuConsoleVariables::MatchingFromPinCategory;
 
 						// Break out of the loop so that we don't give any extra bonuses
 						break;
@@ -7687,32 +7698,39 @@ float UEdGraphSchema_K2::GetActionFilteredWeight(const FGraphActionListBuilderBa
 				// Count of how many words in this keyword array contain a filter(letter) that the user has typed in
 				int32 WordMatchCount = 0;
 
-				// The number of characters in this keyword array
-				int32 KeywordArrayCharLength = 0;
+				// The number of characters in the best matching word
+				int32 BestMatchCharLength = 0;
 
 				// Loop through every word that the user could be looking for
 				for (int32 iEachWord = 0; iEachWord < KeywordArray.Num(); ++iEachWord)
 				{
-					// Keep track of how long all the words in the array are
-					KeywordArrayCharLength += KeywordArray[iEachWord].Len();
+					float WeightPerWord = 0.0f;
 
 					// If a word contains the letter that the user has typed in, than increment the whole match count					
 					if (KeywordArray[iEachWord].Contains(*EachTermSanitized, ESearchCase::CaseSensitive) || KeywordArray[iEachWord].Contains(*EachTerm, ESearchCase::CaseSensitive))
 					{
 						++WordMatchCount;
-						WeightPerList += KeywordArrayWeight * ContextMenuConsoleVariables::WordContainsLetterWeightMultiplier;
+						WeightPerWord += KeywordArrayWeight * ContextMenuConsoleVariables::WordContainsLetterWeightMultiplier;
 
 						// If the word starts with the letter, give it a little extra boost of weight
 						if (KeywordArray[iEachWord].StartsWith(*EachTermSanitized, ESearchCase::CaseSensitive) || KeywordArray[iEachWord].StartsWith(*EachTerm, ESearchCase::CaseSensitive))
 						{
-							WeightPerList += KeywordArrayWeight * ContextMenuConsoleVariables::StartsWithBonusWeightMultiplier;
+							WeightPerWord += KeywordArrayWeight * ContextMenuConsoleVariables::StartsWithBonusWeightMultiplier;
+						}
+
+						if (WeightPerWord > WeightPerList)
+						{
+							// Use the best word match weight, we don't want to double-count redundant keywords like add and addmap here
+							WeightPerList = WeightPerWord;
+							BestMatchCharLength = KeywordArray[iEachWord].Len();
 						}
 					}
 				}
 
 				// If the user has dragged off of a pin then do not prefer shorter things, because that will result
 				// in the matching of "Add" for a container instead of "+" for numeric types
-				if (KeywordArrayCharLength > 0)
+				// We only care about length penalty if something actually matched
+				if (BestMatchCharLength > 0 && WeightPerList > 0)
 				{
 					// How many words that we are checking had partial matches compared to what the user typed in?
 					float PercMatch = (float)WordMatchCount / (float)KeywordArray.Num();
@@ -7720,25 +7738,31 @@ float UEdGraphSchema_K2::GetActionFilteredWeight(const FGraphActionListBuilderBa
 					float PercentageBonus = (WeightPerList * PercMatch * ContextMenuConsoleVariables::PercentageMatchWeightMultiplier);
 					WeightPerList += PercentageBonus;
 
-					// The longer the match is, the more points it loses
-					float ShortWeight = KeywordArrayCharLength * ContextMenuConsoleVariables::ShorterWeight * (bIsFromDrag ? 0.25f : 1.0f);
-					WeightPerList -= ShortWeight;
+					// The shorter the matching word, the larger bonus it gets
+					float ShortFactor = ContextMenuConsoleVariables::MaxWordLength - FMath::Min(BestMatchCharLength, ContextMenuConsoleVariables::MaxWordLength);
+					float ShortWeight = ShortFactor * ContextMenuConsoleVariables::ShorterWeight * (bIsFromDrag ? 0.25f : 1.0f);
+					WeightPerList += ShortWeight;
 
 					OutDebugInfo.PercMatch += PercMatch;
-					OutDebugInfo.ShorterWeight -= ShortWeight;
+					OutDebugInfo.ShorterWeight += ShortWeight;
 					OutDebugInfo.PercentageMatchWeight += PercentageBonus;
 				}
 
 				TotalWeight += WeightPerList;
 				if (WeightedArrayList[iFindCount].OutWeight)
 				{
-					*WeightedArrayList[iFindCount].OutWeight = WeightPerList;
+					// Each weight is used twice so add them
+					*WeightedArrayList[iFindCount].OutWeight += WeightPerList;
 				}
 			}
 		}
-	}
+		OutDebugInfo.TotalWeight = TotalWeight;
 
-	OutDebugInfo.TotalWeight = TotalWeight;
+		if (ContextMenuConsoleVariables::bPrintDebugContextSelection)
+		{
+			OutDebugInfo.Print(SearchText);
+		}
+	}
 
 	return TotalWeight;
 }
diff --git a/Engine/Source/Runtime/Engine/Private/EdGraph/EdGraphSchema.cpp b/Engine/Source/Runtime/Engine/Private/EdGraph/EdGraphSchema.cpp
--- a/Engine/Source/Runtime/Engine/Private/EdGraph/EdGraphSchema.cpp
+++ b/Engine/Source/Runtime/Engine/Private/EdGraph/EdGraphSchema.cpp
@@ -61,6 +61,8 @@ void FEdGraphSchemaAction::UpdateSearchText()
 		SearchText += Entry;
 	}
 
+	SearchText.Append(LINE_TERMINATOR);
+
 	for (FString& Entry : FullSearchTitlesArray)
 	{
 		Entry.ToLowerInline();
