Promote Resolve Soft Reference to the default category and change title to Make when going the opposite direction
Move the UDN documentation node out of the top level category because end users never want it

diff --git a/Engine/Source/Editor/BlueprintGraph/Classes/K2Node_ConvertAsset.h b/Engine/Source/Editor/BlueprintGraph/Classes/K2Node_ConvertAsset.h
--- a/Engine/Source/Editor/BlueprintGraph/Classes/K2Node_ConvertAsset.h
+++ b/Engine/Source/Editor/BlueprintGraph/Classes/K2Node_ConvertAsset.h
@@ -12,6 +12,7 @@ class FBlueprintActionDatabaseRegistrar;
 class UEdGraph;
 class UEdGraphPin;
 
+/** This node converts between hard and soft references, for both objects and classes. The name is historical from when they were called asset IDs */
 UCLASS(MinimalAPI)
 class UK2Node_ConvertAsset : public UK2Node
 {
@@ -20,10 +21,10 @@ public:
 
 	/** Returns the class for the object being converted */
 	UClass* GetTargetClass() const;
-	/** True if this is converting an Asset Class, false for Asset Id */
-	bool IsAssetClassType() const;
-	/** True if this is going from object to asset, false for default of asset to object */
-	bool IsConvertToAsset() const;
+	/** True if this is converting a class, false for object or unknown */
+	bool IsClassType() const;
+	/** True if this is going from hard to soft, false for opposite or unknown */
+	bool IsConvertToSoft() const;
 
 	// UEdGraphNode interface
 	virtual void AllocateDefaultPins() override;
diff --git a/Engine/Source/Editor/BlueprintGraph/Classes/K2Node_SetFieldsInStruct.h b/Engine/Source/Editor/BlueprintGraph/Classes/K2Node_SetFieldsInStruct.h
--- a/Engine/Source/Editor/BlueprintGraph/Classes/K2Node_SetFieldsInStruct.h
+++ b/Engine/Source/Editor/BlueprintGraph/Classes/K2Node_SetFieldsInStruct.h
@@ -23,6 +23,7 @@ class UK2Node_SetFieldsInStruct : public UK2Node_MakeStruct
 	virtual void ValidateNodeDuringCompilation(class FCompilerResultsLog& MessageLog) const override;
 	virtual bool IsConnectionDisallowed(const UEdGraphPin* MyPin, const UEdGraphPin* OtherPin, FString& OutReason) const override;
 	virtual bool CanSplitPin(const UEdGraphPin* Pin) const override;
+	virtual void GetMenuActions(FBlueprintActionDatabaseRegistrar& ActionRegistrar) const override;
 	//~ End  UEdGraphNode Interface
 
 	//~ Begin K2Node Interface
diff --git a/Engine/Source/Editor/BlueprintGraph/Classes/K2Node_StructOperation.h b/Engine/Source/Editor/BlueprintGraph/Classes/K2Node_StructOperation.h
--- a/Engine/Source/Editor/BlueprintGraph/Classes/K2Node_StructOperation.h
+++ b/Engine/Source/Editor/BlueprintGraph/Classes/K2Node_StructOperation.h
@@ -48,5 +48,9 @@ protected:
 	};
 
 	bool DoRenamedPinsMatch(const UEdGraphPin* NewPin, const UEdGraphPin* OldPin, bool bStructInVariablesOut) const;
+
+	/** Utility function to set up menu actions to set the struct type and promote category */
+	DECLARE_DELEGATE_RetVal_TwoParams(bool, FMakeStructSpawnerAllowedDelegate, const UScriptStruct*, bool);
+	void SetupMenuActions(FBlueprintActionDatabaseRegistrar& ActionRegistrar, const FMakeStructSpawnerAllowedDelegate& AllowedDelegate, EEdGraphPinDirection PinDirectionToPromote) const;
 };
 
diff --git a/Engine/Source/Editor/BlueprintGraph/Private/BlueprintActionDatabase.cpp b/Engine/Source/Editor/BlueprintGraph/Private/BlueprintActionDatabase.cpp
--- a/Engine/Source/Editor/BlueprintGraph/Private/BlueprintActionDatabase.cpp
+++ b/Engine/Source/Editor/BlueprintGraph/Private/BlueprintActionDatabase.cpp
@@ -927,7 +927,9 @@ static void BlueprintActionDatabaseImpl::GetNodeSpecificActions(TSubclassOf<UEdG
 	else if (NodeClass == UEdGraphNode_Documentation::StaticClass())
 	{
 		// @TODO: BOOOOOOO! (see comment above)
-		Registrar.AddBlueprintAction(MakeDocumentationNodeSpawner<UEdGraphNode_Documentation>());
+		UBlueprintNodeSpawner* DocumentationSpawner = MakeDocumentationNodeSpawner<UEdGraphNode_Documentation>();
+		DocumentationSpawner->DefaultMenuSignature.Category = LOCTEXT("DocumentationNodeCategory", "Documentation");
+		Registrar.AddBlueprintAction(DocumentationSpawner);
 	}
 }
 
diff --git a/Engine/Source/Editor/BlueprintGraph/Private/K2Node_BreakStruct.cpp b/Engine/Source/Editor/BlueprintGraph/Private/K2Node_BreakStruct.cpp
--- a/Engine/Source/Editor/BlueprintGraph/Private/K2Node_BreakStruct.cpp
+++ b/Engine/Source/Editor/BlueprintGraph/Private/K2Node_BreakStruct.cpp
@@ -333,44 +333,7 @@ FNodeHandlingFunctor* UK2Node_BreakStruct::CreateNodeHandler(class FKismetCompil
 
 void UK2Node_BreakStruct::GetMenuActions(FBlueprintActionDatabaseRegistrar& ActionRegistrar) const
 {
-	struct GetMenuActions_Utils
-	{
-		static void SetNodeStruct(UEdGraphNode* NewNode, FFieldVariant /*StructField*/, TWeakObjectPtr<UScriptStruct> NonConstStructPtr)
-		{
-			UK2Node_BreakStruct* BreakNode = CastChecked<UK2Node_BreakStruct>(NewNode);
-			BreakNode->StructType = NonConstStructPtr.Get();
-		}
-
-		static void OverrideCategory(FBlueprintActionContext const& Context, IBlueprintNodeBinder::FBindingSet const& /*Bindings*/, FBlueprintActionUiSpec* UiSpecOut, TWeakObjectPtr<UScriptStruct> StructPtr)
-		{
-			for (UEdGraphPin* Pin : Context.Pins)
-			{
-				UScriptStruct* PinStruct = Cast<UScriptStruct>(Pin->PinType.PinSubCategoryObject.Get());
-				if ((PinStruct != nullptr) && (StructPtr.Get() == PinStruct) && (Pin->Direction == EGPD_Output))
-				{
-					UiSpecOut->Category = LOCTEXT("EmptyCategory", "|");
-					break;
-				}
-			}
-		}
-	};
-
-	UClass* NodeClass = GetClass();
-	ActionRegistrar.RegisterStructActions( FBlueprintActionDatabaseRegistrar::FMakeStructSpawnerDelegate::CreateLambda([NodeClass](const UScriptStruct* Struct)->UBlueprintNodeSpawner*
-	{
-		UBlueprintFieldNodeSpawner* NodeSpawner = nullptr;
-		
-		if (UK2Node_BreakStruct::CanBeBroken(Struct))
-		{
-			NodeSpawner = UBlueprintFieldNodeSpawner::Create(NodeClass, const_cast<UScriptStruct*>(Struct));
-			check(NodeSpawner != nullptr);
-			TWeakObjectPtr<UScriptStruct> NonConstStructPtr = MakeWeakObjectPtr(const_cast<UScriptStruct*>(Struct));
-			NodeSpawner->SetNodeFieldDelegate     = UBlueprintFieldNodeSpawner::FSetNodeFieldDelegate::CreateStatic(GetMenuActions_Utils::SetNodeStruct, NonConstStructPtr);
-			NodeSpawner->DynamicUiSignatureGetter = UBlueprintFieldNodeSpawner::FUiSpecOverrideDelegate::CreateStatic(GetMenuActions_Utils::OverrideCategory, NonConstStructPtr);
-
-		}
-		return NodeSpawner;
-	}) );
+	Super::SetupMenuActions(ActionRegistrar, FMakeStructSpawnerAllowedDelegate::CreateStatic(&UK2Node_BreakStruct::CanBeBroken), EGPD_Output);
 }
 
 FText UK2Node_BreakStruct::GetMenuCategory() const
diff --git a/Engine/Source/Editor/BlueprintGraph/Private/K2Node_ConvertAsset.cpp b/Engine/Source/Editor/BlueprintGraph/Private/K2Node_ConvertAsset.cpp
--- a/Engine/Source/Editor/BlueprintGraph/Private/K2Node_ConvertAsset.cpp
+++ b/Engine/Source/Editor/BlueprintGraph/Private/K2Node_ConvertAsset.cpp
@@ -12,6 +12,7 @@
 #include "K2Node_ClassDynamicCast.h"
 #include "KismetCompiler.h"
 #include "BlueprintNodeSpawner.h"
+#include "BlueprintFieldNodeSpawner.h"
 #include "BlueprintActionDatabaseRegistrar.h"
 
 #define LOCTEXT_NAMESPACE "K2Node_ConvertAsset"
@@ -30,18 +31,16 @@ UClass* UK2Node_ConvertAsset::GetTargetClass() const
 	return SourcePin ? FBlueprintEditorUtils::GetTypeForPin(*SourcePin) : nullptr;
 }
 
-bool UK2Node_ConvertAsset::IsAssetClassType() const
+bool UK2Node_ConvertAsset::IsClassType() const
 {
-	// get first input, return if class asset
 	UEdGraphPin* InputPin = FindPin(UK2Node_ConvertAssetImpl::InputPinName);
 	bool bIsConnected = InputPin && InputPin->LinkedTo.Num() && InputPin->LinkedTo[0];
 	UEdGraphPin* SourcePin = bIsConnected ? InputPin->LinkedTo[0] : nullptr;
 	return SourcePin ? (SourcePin->PinType.PinCategory == UEdGraphSchema_K2::PC_SoftClass || SourcePin->PinType.PinCategory == UEdGraphSchema_K2::PC_Class) : false;
 }
 
-bool UK2Node_ConvertAsset::IsConvertToAsset() const
+bool UK2Node_ConvertAsset::IsConvertToSoft() const
 {
-	// get first input, return if class asset
 	UEdGraphPin* InputPin = FindPin(UK2Node_ConvertAssetImpl::InputPinName);
 	bool bIsConnected = InputPin && InputPin->LinkedTo.Num() && InputPin->LinkedTo[0];
 	UEdGraphPin* SourcePin = bIsConnected ? InputPin->LinkedTo[0] : nullptr;
@@ -74,22 +73,22 @@ void UK2Node_ConvertAsset::RefreshPinTypes()
 	{
 		const bool bIsConnected = InputPin->LinkedTo.Num() > 0;
 		UClass* TargetType = bIsConnected ? GetTargetClass() : nullptr;
-		const bool bIsAssetClass = bIsConnected ? IsAssetClassType() : false;
-		const bool bIsConvertToAsset = bIsConnected ? IsConvertToAsset() : false;
+		const bool bIsClassType = bIsConnected ? IsClassType() : false;
+		const bool bIsConvertToSoft = bIsConnected ? IsConvertToSoft() : false;
 
 		FName InputCategory = UEdGraphSchema_K2::PC_Wildcard;
 		FName OutputCategory = UEdGraphSchema_K2::PC_Wildcard;
 		if (bIsConnected)
 		{
-			if (bIsConvertToAsset)
+			if (bIsConvertToSoft)
 			{
-				InputCategory = (bIsAssetClass ? UEdGraphSchema_K2::PC_Class : UEdGraphSchema_K2::PC_Object);
-				OutputCategory = (bIsAssetClass ? UEdGraphSchema_K2::PC_SoftClass : UEdGraphSchema_K2::PC_SoftObject);
+				InputCategory = (bIsClassType ? UEdGraphSchema_K2::PC_Class : UEdGraphSchema_K2::PC_Object);
+				OutputCategory = (bIsClassType ? UEdGraphSchema_K2::PC_SoftClass : UEdGraphSchema_K2::PC_SoftObject);
 			}
 			else
 			{
-				InputCategory = (bIsAssetClass ? UEdGraphSchema_K2::PC_SoftClass : UEdGraphSchema_K2::PC_SoftObject);
-				OutputCategory = (bIsAssetClass ? UEdGraphSchema_K2::PC_Class : UEdGraphSchema_K2::PC_Object);
+				InputCategory = (bIsClassType ? UEdGraphSchema_K2::PC_SoftClass : UEdGraphSchema_K2::PC_SoftObject);
+				OutputCategory = (bIsClassType ? UEdGraphSchema_K2::PC_Class : UEdGraphSchema_K2::PC_Object);
 			}
 		}
 			
@@ -160,19 +159,51 @@ UK2Node::ERedirectType UK2Node_ConvertAsset::DoPinsMatchForReconstruction(const
 
 void UK2Node_ConvertAsset::GetMenuActions(FBlueprintActionDatabaseRegistrar& ActionRegistrar) const
 {
-	// actions get registered under specific object-keys; the idea is that 
-	// actions might have to be updated (or deleted) if their object-key is  
-	// mutated (or removed)... here we use the node's class (so if the node 
-	// type disappears, then the action should go with it)
+	struct GetMenuActions_Utils
+	{
+		static void OverrideUI(FBlueprintActionContext const& Context, IBlueprintNodeBinder::FBindingSet const& /*Bindings*/, FBlueprintActionUiSpec* UiSpecOut)
+		{
+			bool bShouldPromote = false;
+			bool bIsMakeSoft = false;
+			for (UEdGraphPin* Pin : Context.Pins)
+			{
+				// Auto promote for soft refs
+				if (Pin->PinType.PinCategory == UEdGraphSchema_K2::PC_SoftClass || Pin->PinType.PinCategory == UEdGraphSchema_K2::PC_SoftObject)
+				{
+					if (Pin->Direction == EGPD_Output)
+					{
+						// The wildcards don't get set when connecting right side, so don't promote
+						bShouldPromote = true;
+					}
+				}
+
+				// Change title for hard refs
+				if (Pin->PinType.PinCategory == UEdGraphSchema_K2::PC_Class || Pin->PinType.PinCategory == UEdGraphSchema_K2::PC_Object)
+				{
+					if (Pin->Direction == EGPD_Output)
+					{
+						bIsMakeSoft = true;
+					}
+				}
+			}
+
+			if (bShouldPromote)
+			{
+				UiSpecOut->Category = NSLOCTEXT("BlueprintFunctionNodeSpawner", "EmptyFunctionCategory", "|");
+			}
+			if (bIsMakeSoft)
+			{
+				UiSpecOut->MenuName = LOCTEXT("MakeSoftTitle", "Make Soft Reference");
+			}
+		}
+	};
+
 	UClass* ActionKey = GetClass();
-	// to keep from needlessly instantiating a UBlueprintNodeSpawner, first   
-	// check to make sure that the registrar is looking for actions of this type
-	// (could be regenerating actions for a specific asset, and therefore the 
-	// registrar would only accept actions corresponding to that asset)
 	if (ActionRegistrar.IsOpenForRegistration(ActionKey))
 	{
 		UBlueprintNodeSpawner* NodeSpawner = UBlueprintNodeSpawner::Create(GetClass());
 		check(NodeSpawner != nullptr);
+		NodeSpawner->DynamicUiSignatureGetter = UBlueprintFieldNodeSpawner::FUiSpecOverrideDelegate::CreateStatic(GetMenuActions_Utils::OverrideUI);
 
 		ActionRegistrar.AddBlueprintAction(ActionKey, NodeSpawner);
 	}
@@ -186,15 +217,15 @@ void UK2Node_ConvertAsset::ExpandNode(class FKismetCompilerContext& CompilerCont
 	UClass* TargetType = GetTargetClass();
 	if (TargetType && Schema && (2 == Pins.Num()))
 	{
-		const bool bIsAssetClass = IsAssetClassType();
+		const bool bIsClassType = IsClassType();
 		UClass *TargetClass = GetTargetClass();
 		bool bIsErrorFree = true;
 
-		if (IsConvertToAsset())
+		if (IsConvertToSoft())
 		{
 			//Create Convert Function
 			UK2Node_CallFunction* ConvertToObjectFunc = CompilerContext.SpawnIntermediateNode<UK2Node_CallFunction>(this, SourceGraph);
-			FName ConvertFunctionName = bIsAssetClass
+			FName ConvertFunctionName = bIsClassType
 				? GET_FUNCTION_NAME_CHECKED(UKismetSystemLibrary, Conv_ClassToSoftClassReference)
 				: GET_FUNCTION_NAME_CHECKED(UKismetSystemLibrary, Conv_ObjectToSoftObjectReference);
 
@@ -203,7 +234,7 @@ void UK2Node_ConvertAsset::ExpandNode(class FKismetCompilerContext& CompilerCont
 
 			//Connect input to convert
 			UEdGraphPin* InputPin = FindPin(UK2Node_ConvertAssetImpl::InputPinName);
-			const FName ConvertInputName = bIsAssetClass ? FName(TEXT("Class")) : FName(TEXT("Object"));
+			const FName ConvertInputName = bIsClassType ? FName(TEXT("Class")) : FName(TEXT("Object"));
 			UEdGraphPin* ConvertInput = ConvertToObjectFunc->FindPin(ConvertInputName);
 			bIsErrorFree = InputPin && ConvertInput && CompilerContext.MovePinLinksToIntermediate(*InputPin, *ConvertInput).CanSafeConnect();
 
@@ -224,7 +255,7 @@ void UK2Node_ConvertAsset::ExpandNode(class FKismetCompilerContext& CompilerCont
 		{
 			//Create Convert Function
 			UK2Node_CallFunction* ConvertToObjectFunc = CompilerContext.SpawnIntermediateNode<UK2Node_CallFunction>(this, SourceGraph);
-			FName ConvertFunctionName = bIsAssetClass
+			FName ConvertFunctionName = bIsClassType
 				? GET_FUNCTION_NAME_CHECKED(UKismetSystemLibrary, Conv_SoftClassReferenceToClass)
 				: GET_FUNCTION_NAME_CHECKED(UKismetSystemLibrary, Conv_SoftObjectReferenceToObject);
 
@@ -233,7 +264,7 @@ void UK2Node_ConvertAsset::ExpandNode(class FKismetCompilerContext& CompilerCont
 
 			//Connect input to convert
 			UEdGraphPin* InputPin = FindPin(UK2Node_ConvertAssetImpl::InputPinName);
-			const FName ConvertInputName = bIsAssetClass ? FName(TEXT("SoftClass")) : FName(TEXT("SoftObject"));
+			const FName ConvertInputName = bIsClassType ? FName(TEXT("SoftClass")) : FName(TEXT("SoftObject"));
 			UEdGraphPin* ConvertInput = ConvertToObjectFunc->FindPin(ConvertInputName);
 			bIsErrorFree = InputPin && ConvertInput && CompilerContext.MovePinLinksToIntermediate(*InputPin, *ConvertInput).CanSafeConnect();
 
@@ -242,7 +273,7 @@ void UK2Node_ConvertAsset::ExpandNode(class FKismetCompilerContext& CompilerCont
 			if (UObject::StaticClass() != TargetType)
 			{
 				//Create Cast Node
-				UK2Node_DynamicCast* CastNode = bIsAssetClass
+				UK2Node_DynamicCast* CastNode = bIsClassType
 					? CompilerContext.SpawnIntermediateNode<UK2Node_ClassDynamicCast>(this, SourceGraph)
 					: CompilerContext.SpawnIntermediateNode<UK2Node_DynamicCast>(this, SourceGraph);
 				CastNode->SetPurity(true);
@@ -281,23 +312,41 @@ FText UK2Node_ConvertAsset::GetCompactNodeTitle() const
 
 FText UK2Node_ConvertAsset::GetMenuCategory() const
 {
-	return FText(LOCTEXT("UK2Node_LoadAssetGetMenuCategory", "Utilities"));
+	return LOCTEXT("UK2Node_LoadAssetGetMenuCategory", "Utilities");
 }
 
 FText UK2Node_ConvertAsset::GetNodeTitle(ENodeTitleType::Type TitleType) const
 {
-	return FText(LOCTEXT("UK2Node_ConvertAssetGetNodeTitle", "Resolve Soft Reference"));
+	const bool bIsConvertToSoft = IsConvertToSoft();
+
+	if (bIsConvertToSoft)
+	{
+		return LOCTEXT("MakeSoftTitle", "Make Soft Reference");
+	}
+
+	return LOCTEXT("UK2Node_ConvertAssetGetNodeTitle", "Resolve Soft Reference");
 }
 
 FText UK2Node_ConvertAsset::GetKeywords() const
 {
 	// Return old name here
-	return FText(LOCTEXT("UK2Node_ConvertAssetGetKeywords", "Resolve Asset ID"));
+	return LOCTEXT("UK2Node_ConvertAssetGetKeywords", "Resolve Asset ID Convert Soft");
 }
 
 FText UK2Node_ConvertAsset::GetTooltipText() const
 {
-	return FText(LOCTEXT("UK2Node_ConvertAssetGetTooltipText", "Resolves a Soft Reference or Soft Class Reference into an object/class or vice versa. If the object isn't already loaded it returns none."));
+	UEdGraphPin* InputPin = FindPin(UK2Node_ConvertAssetImpl::InputPinName);
+	const bool bIsConvertToSoft = IsConvertToSoft();
+	if (bIsConvertToSoft)
+	{
+		return LOCTEXT("MakeSoftTooltip", "Takes a hard Class or Object reference and makes a Soft Reference.");
+	}
+	else if (InputPin && InputPin->LinkedTo.Num() > 0)
+	{
+		return LOCTEXT("UK2Node_ConvertAssetGetTooltipText", "Resolves a Soft Reference and gets the Class or Object it is pointing to. If the object isn't already loaded in memory this will return none.");
+	}
+
+	return LOCTEXT("UnknownTypeTooltip", "Resolves or makes a Soft Reference, connect a soft or hard reference to the input pin.");
 }
 
 #undef LOCTEXT_NAMESPACE
diff --git a/Engine/Source/Editor/BlueprintGraph/Private/K2Node_MakeStruct.cpp b/Engine/Source/Editor/BlueprintGraph/Private/K2Node_MakeStruct.cpp
--- a/Engine/Source/Editor/BlueprintGraph/Private/K2Node_MakeStruct.cpp
+++ b/Engine/Source/Editor/BlueprintGraph/Private/K2Node_MakeStruct.cpp
@@ -293,43 +293,7 @@ UK2Node::ERedirectType UK2Node_MakeStruct::DoPinsMatchForReconstruction(const UE
 
 void UK2Node_MakeStruct::GetMenuActions(FBlueprintActionDatabaseRegistrar& ActionRegistrar) const
 {
-	struct GetMenuActions_Utils
-	{
-		static void SetNodeStruct(UEdGraphNode* NewNode, FFieldVariant /*StructField*/, TWeakObjectPtr<UScriptStruct> NonConstStructPtr)
-		{
-			UK2Node_MakeStruct* MakeNode = CastChecked<UK2Node_MakeStruct>(NewNode);
-			MakeNode->StructType = NonConstStructPtr.Get();
-		}
-
-		static void OverrideCategory(FBlueprintActionContext const& Context, IBlueprintNodeBinder::FBindingSet const& /*Bindings*/, FBlueprintActionUiSpec* UiSpecOut, TWeakObjectPtr<UScriptStruct> StructPtr)
-		{
-			for (UEdGraphPin* Pin : Context.Pins)
-			{
-				UScriptStruct* PinStruct = Cast<UScriptStruct>(Pin->PinType.PinSubCategoryObject.Get());
-				if ((PinStruct != nullptr) && (StructPtr.Get() == PinStruct) && (Pin->Direction == EGPD_Input))
-				{
-					UiSpecOut->Category = LOCTEXT("EmptyCategory", "|");
-					break;
-				}
-			}
-		}
-	};
-
-	UClass* NodeClass = GetClass();
-	ActionRegistrar.RegisterStructActions( FBlueprintActionDatabaseRegistrar::FMakeStructSpawnerDelegate::CreateLambda([NodeClass](const UScriptStruct* Struct)->UBlueprintNodeSpawner*
-	{
-		UBlueprintFieldNodeSpawner* NodeSpawner = nullptr;
-		
-		if (UK2Node_MakeStruct::CanBeMade(Struct))
-		{
-			NodeSpawner = UBlueprintFieldNodeSpawner::Create(NodeClass, const_cast<UScriptStruct*>(Struct));
-			check(NodeSpawner != nullptr);
-			TWeakObjectPtr<UScriptStruct> NonConstStructPtr = MakeWeakObjectPtr(const_cast<UScriptStruct*>(Struct));
-			NodeSpawner->SetNodeFieldDelegate     = UBlueprintFieldNodeSpawner::FSetNodeFieldDelegate::CreateStatic(GetMenuActions_Utils::SetNodeStruct, NonConstStructPtr);
-			NodeSpawner->DynamicUiSignatureGetter = UBlueprintFieldNodeSpawner::FUiSpecOverrideDelegate::CreateStatic(GetMenuActions_Utils::OverrideCategory, NonConstStructPtr);
-		}
-		return NodeSpawner;
-	}) );
+	Super::SetupMenuActions(ActionRegistrar, FMakeStructSpawnerAllowedDelegate::CreateStatic(&UK2Node_MakeStruct::CanBeMade), EGPD_Input);
 }
 
 FText UK2Node_MakeStruct::GetMenuCategory() const
diff --git a/Engine/Source/Editor/BlueprintGraph/Private/K2Node_SetFieldsInStruct.cpp b/Engine/Source/Editor/BlueprintGraph/Private/K2Node_SetFieldsInStruct.cpp
--- a/Engine/Source/Editor/BlueprintGraph/Private/K2Node_SetFieldsInStruct.cpp
+++ b/Engine/Source/Editor/BlueprintGraph/Private/K2Node_SetFieldsInStruct.cpp
@@ -371,4 +371,10 @@ bool UK2Node_SetFieldsInStruct::CanSplitPin(const UEdGraphPin* Pin) const
 }
 
 
+void UK2Node_SetFieldsInStruct::GetMenuActions(FBlueprintActionDatabaseRegistrar& ActionRegistrar) const
+{
+	// Use Make's filter function but auto promote for output pins
+	Super::SetupMenuActions(ActionRegistrar, FMakeStructSpawnerAllowedDelegate::CreateStatic(&UK2Node_MakeStruct::CanBeMade), EGPD_Output);
+}
+
 #undef LOCTEXT_NAMESPACE
diff --git a/Engine/Source/Editor/BlueprintGraph/Private/K2Node_StructOperation.cpp b/Engine/Source/Editor/BlueprintGraph/Private/K2Node_StructOperation.cpp
--- a/Engine/Source/Editor/BlueprintGraph/Private/K2Node_StructOperation.cpp
+++ b/Engine/Source/Editor/BlueprintGraph/Private/K2Node_StructOperation.cpp
@@ -6,6 +6,9 @@
 #include "UserDefinedStructure/UserDefinedStructEditorData.h"
 #include "Kismet2/StructureEditorUtils.h"
 #include "BlueprintActionFilter.h"
+#include "BlueprintFieldNodeSpawner.h"
+#include "EditorCategoryUtils.h"
+#include "BlueprintActionDatabaseRegistrar.h"
 
 //////////////////////////////////////////////////////////////////////////
 // UK2Node_StructOperation
@@ -83,6 +86,48 @@ bool UK2Node_StructOperation::DoRenamedPinsMatch(const UEdGraphPin* NewPin, cons
 	return false;
 }
 
+void UK2Node_StructOperation::SetupMenuActions(FBlueprintActionDatabaseRegistrar& ActionRegistrar, const FMakeStructSpawnerAllowedDelegate& AllowedDelegate, EEdGraphPinDirection PinDirectionToPromote) const
+{
+	struct GetMenuActions_Utils
+	{
+		static void SetNodeStruct(UEdGraphNode* NewNode, FFieldVariant /*StructField*/, TWeakObjectPtr<UScriptStruct> NonConstStructPtr)
+		{
+			UK2Node_StructOperation* StructNode = CastChecked<UK2Node_StructOperation>(NewNode);
+			StructNode->StructType = NonConstStructPtr.Get();
+		}
+
+		static void OverrideCategory(FBlueprintActionContext const& Context, IBlueprintNodeBinder::FBindingSet const& /*Bindings*/, FBlueprintActionUiSpec* UiSpecOut, TWeakObjectPtr<UScriptStruct> StructPtr, EEdGraphPinDirection PinDirectionToPromote)
+		{
+			for (UEdGraphPin* Pin : Context.Pins)
+			{
+				UScriptStruct* PinStruct = Cast<UScriptStruct>(Pin->PinType.PinSubCategoryObject.Get());
+				if ((PinStruct != nullptr) && (StructPtr.Get() == PinStruct) && (Pin->Direction == PinDirectionToPromote))
+				{
+					UiSpecOut->Category = NSLOCTEXT("BlueprintFunctionNodeSpawner", "EmptyFunctionCategory", "|");
+					break;
+				}
+			}
+		}
+	};
+
+	UClass* NodeClass = GetClass();
+	ActionRegistrar.RegisterStructActions(FBlueprintActionDatabaseRegistrar::FMakeStructSpawnerDelegate::CreateLambda([NodeClass, AllowedDelegate, PinDirectionToPromote](const UScriptStruct* Struct)->UBlueprintNodeSpawner*
+	{
+		UBlueprintFieldNodeSpawner* NodeSpawner = nullptr;
+
+		if (AllowedDelegate.Execute(Struct, false))
+		{
+			NodeSpawner = UBlueprintFieldNodeSpawner::Create(NodeClass, const_cast<UScriptStruct*>(Struct));
+			check(NodeSpawner != nullptr);
+			TWeakObjectPtr<UScriptStruct> NonConstStructPtr = MakeWeakObjectPtr(const_cast<UScriptStruct*>(Struct));
+			NodeSpawner->SetNodeFieldDelegate = UBlueprintFieldNodeSpawner::FSetNodeFieldDelegate::CreateStatic(GetMenuActions_Utils::SetNodeStruct, NonConstStructPtr);
+			NodeSpawner->DynamicUiSignatureGetter = UBlueprintFieldNodeSpawner::FUiSpecOverrideDelegate::CreateStatic(GetMenuActions_Utils::OverrideCategory, NonConstStructPtr, PinDirectionToPromote);
+
+		}
+		return NodeSpawner;
+	}));
+}
+
 FString UK2Node_StructOperation::GetPinMetaData(FName InPinName, FName InKey)
 {
 	for (TFieldIterator<FProperty> It(StructType); It; ++It)
