This fixes confusing UI behavior where it would dirty the blueprint as you saved it, and keeps function defaults that were set before being connected to a delegate

diff --git a/Engine/Source/Editor/BlueprintGraph/Private/K2Node_CustomEvent.cpp b/Engine/Source/Editor/BlueprintGraph/Private/K2Node_CustomEvent.cpp
--- a/Engine/Source/Editor/BlueprintGraph/Private/K2Node_CustomEvent.cpp
+++ b/Engine/Source/Editor/BlueprintGraph/Private/K2Node_CustomEvent.cpp
@@ -393,29 +393,42 @@ void UK2Node_CustomEvent::SetDelegateSignature(const UFunction* DelegateSignatur
 
 	const UEdGraphSchema_K2* K2Schema = GetDefault<UEdGraphSchema_K2>();
 	
-		UserDefinedPins.Empty();
-		for (TFieldIterator<FProperty> PropIt(DelegateSignature); PropIt && (PropIt->PropertyFlags & CPF_Parm); ++PropIt)
+	TArray < TSharedPtr<FUserPinInfo> > OldPins = UserDefinedPins;
+	UserDefinedPins.Empty();
+	for (TFieldIterator<FProperty> PropIt(DelegateSignature); PropIt && (PropIt->PropertyFlags & CPF_Parm); ++PropIt)
+	{
+		const FProperty* Param = *PropIt;
+		if (!Param->HasAnyPropertyFlags(CPF_OutParm) || Param->HasAnyPropertyFlags(CPF_ReferenceParm))
 		{
-			const FProperty* Param = *PropIt;
-			if (!Param->HasAnyPropertyFlags(CPF_OutParm) || Param->HasAnyPropertyFlags(CPF_ReferenceParm))
+			FEdGraphPinType PinType;
+			K2Schema->ConvertPropertyToPinType(Param, /*out*/ PinType);
+
+			FName NewPinName = Param->GetFName();
+			int32 Index = 1;
+			while ((DelegateOutputName == NewPinName) || (UEdGraphSchema_K2::PN_Then == NewPinName))
 			{
-				FEdGraphPinType PinType;
-				K2Schema->ConvertPropertyToPinType(Param, /*out*/ PinType);
+				++Index;
+				NewPinName = *FString::Printf(TEXT("%s%d"), *NewPinName.ToString(), Index);
+			}
+			TSharedPtr<FUserPinInfo> NewPinInfo = MakeShareable(new FUserPinInfo());
+			NewPinInfo->PinName = NewPinName;
+			NewPinInfo->PinType = PinType;
+			NewPinInfo->DesiredPinDirection = EGPD_Output;
+			int32 NewIndex = UserDefinedPins.Num();
 
-				FName NewPinName = Param->GetFName();
-				int32 Index = 1;
-				while ((DelegateOutputName == NewPinName) || (UEdGraphSchema_K2::PN_Then == NewPinName))
+			// Copy over old default value if type matches
+			if (OldPins.IsValidIndex(NewIndex) && OldPins[NewIndex].IsValid())
+			{
+				TSharedPtr<FUserPinInfo> OldPinInfo = OldPins[NewIndex];
+				if (NewPinInfo->PinName == OldPinInfo->PinName && NewPinInfo->PinType == OldPinInfo->PinType && NewPinInfo->DesiredPinDirection == OldPinInfo->DesiredPinDirection)
 				{
-					++Index;
-					NewPinName = *FString::Printf(TEXT("%s%d"), *NewPinName.ToString(), Index);
+					NewPinInfo->PinDefaultValue = OldPinInfo->PinDefaultValue;
 				}
-			TSharedPtr<FUserPinInfo> NewPinInfo = MakeShareable(new FUserPinInfo());
-				NewPinInfo->PinName = NewPinName;
-				NewPinInfo->PinType = PinType;
-				NewPinInfo->DesiredPinDirection = EGPD_Output;
-				UserDefinedPins.Add(NewPinInfo);
 			}
+
+			UserDefinedPins.Add(NewPinInfo);
 		}
+	}
 }
 
 UK2Node_CustomEvent* UK2Node_CustomEvent::CreateFromFunction(FVector2D GraphPosition, UEdGraph* ParentGraph, const FString& Name, const UFunction* Function, bool bSelectNewNode/* = true*/)
