During compilation of an Anim BP, if this node gets pruned then the outer will be the /Engine/Transient/ package, which will result in this node not having a valid SelfContext, so there is no need to set any further information based on this function as it will not be used.

This prevents an Ensure being hit in UK2Node_CallFunction::SetFromFunction which will cause cook failures.

#rnx

diff --git a/Engine/Source/Editor/BlueprintGraph/Private/K2Node_PromotableOperator.cpp b/Engine/Source/Editor/BlueprintGraph/Private/K2Node_PromotableOperator.cpp
--- a/Engine/Source/Editor/BlueprintGraph/Private/K2Node_PromotableOperator.cpp
+++ b/Engine/Source/Editor/BlueprintGraph/Private/K2Node_PromotableOperator.cpp
@@ -742,12 +742,18 @@ UEdGraphPin* UK2Node_PromotableOperator::FindTolerancePin() const
 // UK2Node_CallFunction interface
 void UK2Node_PromotableOperator::SetFromFunction(const UFunction* Function)
 {
-	Super::SetFromFunction(Function);
-
 	if(Function)
 	{
 		OperationName = FTypePromotion::GetOpNameFromFunction(Function);
 	}
+
+	// During compilation of an Anim BP, if this node gets pruned then the outer will be the /Engine/Transient/
+	// package which will result in this node not having a valid SelfContext, so there is no need
+	// to set any further information based on this function as it will not be used. 
+	if(GetBlueprintClassFromNode())
+	{
+		Super::SetFromFunction(Function);
+	}
 }
 
 ///////////////////////////////////////////////////////////
