Fix: Outer Pin of Construct Object from Class should not have an asset picker
- changed pin visual to a "self" style
- made pin default to "self" if there's nothing connected to it

diff --git a/Engine/Source/Editor/BlueprintGraph/Classes/K2Node_GenericCreateObject.h b/Engine/Source/Editor/BlueprintGraph/Classes/K2Node_GenericCreateObject.h
--- a/Engine/Source/Editor/BlueprintGraph/Classes/K2Node_GenericCreateObject.h
+++ b/Engine/Source/Editor/BlueprintGraph/Classes/K2Node_GenericCreateObject.h
@@ -8,6 +8,8 @@
 #include "K2Node_GenericCreateObject.generated.h"
 
 class UEdGraph;
+class FKismetCompilerContext;
+class UK2Node_CallFunction;
 
 UCLASS()
 class BLUEPRINTGRAPH_API UK2Node_GenericCreateObject : public UK2Node_ConstructObjectFromClass
@@ -15,11 +17,22 @@ class BLUEPRINTGRAPH_API UK2Node_GenericCreateObject : public UK2Node_ConstructO
 	GENERATED_BODY()
 
 	//~ Begin UEdGraphNode Interface.
-	virtual void ExpandNode(class FKismetCompilerContext& CompilerContext, UEdGraph* SourceGraph) override;
-	virtual void EarlyValidation(class FCompilerResultsLog& MessageLog) const override;
+	virtual void ExpandNode(FKismetCompilerContext& CompilerContext, UEdGraph* SourceGraph) override;
+	virtual void EarlyValidation(FCompilerResultsLog& MessageLog) const override;
 	virtual bool IsCompatibleWithGraph(const UEdGraph* TargetGraph) const override;
 	//~ End UEdGraphNode Interface.
 
 	virtual bool UseWorldContext() const override { return false; }
 	virtual bool UseOuter() const override { return true; }
+
+	/**
+	 * attaches a self node to the self pin of 'this' if the CallCreateNode function has DefaultToSelf in it's metadata
+	 *
+	 * @param	CompilerContext		the context to expand in - likely passed from ExpandNode
+	 * @param	SourceGraph			the graph to expand in - likely passed from ExpandNode
+	 * @param	CallCreateNode		the CallFunction node that 'this' is imitating
+	 *
+	 * @return	true on success.
+	 */
+	bool ExpandDefaultToSelfPin(FKismetCompilerContext& CompilerContext, UEdGraph* SourceGraph, UK2Node_CallFunction* CallCreateNode);
 };
diff --git a/Engine/Source/Editor/BlueprintGraph/Private/K2Node_ConstructObjectFromClass.cpp b/Engine/Source/Editor/BlueprintGraph/Private/K2Node_ConstructObjectFromClass.cpp
--- a/Engine/Source/Editor/BlueprintGraph/Private/K2Node_ConstructObjectFromClass.cpp
+++ b/Engine/Source/Editor/BlueprintGraph/Private/K2Node_ConstructObjectFromClass.cpp
@@ -13,12 +13,12 @@ struct FK2Node_ConstructObjectFromClassHelper
 {
 	static FName WorldContextPinName;
 	static FName ClassPinName;
-	static FName OuterPinName;
+	static FName OuterPinFriendlyName;
 };
 
 FName FK2Node_ConstructObjectFromClassHelper::WorldContextPinName(TEXT("WorldContextObject"));
 FName FK2Node_ConstructObjectFromClassHelper::ClassPinName(TEXT("Class"));
-FName FK2Node_ConstructObjectFromClassHelper::OuterPinName(TEXT("Outer"));
+FName FK2Node_ConstructObjectFromClassHelper::OuterPinFriendlyName(TEXT("Outer"));
 
 #define LOCTEXT_NAMESPACE "K2Node_ConstructObjectFromClass"
 
@@ -58,9 +58,11 @@ void UK2Node_ConstructObjectFromClass::AllocateDefaultPins()
 	// Result pin
 	UEdGraphPin* ResultPin = CreatePin(EGPD_Output, UEdGraphSchema_K2::PC_Object, GetClassPinBaseClass(), UEdGraphSchema_K2::PN_ReturnValue);
 	
+	// Outer pin
 	if (UseOuter())
 	{
-		UEdGraphPin* OuterPin = CreatePin(EGPD_Input, UEdGraphSchema_K2::PC_Object, UObject::StaticClass(), FK2Node_ConstructObjectFromClassHelper::OuterPinName);
+		UEdGraphPin* OuterPin = CreatePin(EGPD_Input, UEdGraphSchema_K2::PC_Object, UEdGraphSchema_K2::PSC_Self, nullptr, UEdGraphSchema_K2::PN_Self);
+		OuterPin->PinFriendlyName = FText::FromName(FK2Node_ConstructObjectFromClassHelper::OuterPinFriendlyName);
 	}
 
 	Super::AllocateDefaultPins();
@@ -68,7 +70,8 @@ void UK2Node_ConstructObjectFromClass::AllocateDefaultPins()
 
 UEdGraphPin* UK2Node_ConstructObjectFromClass::GetOuterPin() const
 {
-	UEdGraphPin* Pin = FindPin(FK2Node_ConstructObjectFromClassHelper::OuterPinName);
+	// the pin is named self now because it's a "default to self" type pin
+	UEdGraphPin* Pin = FindPin(UEdGraphSchema_K2::PN_Self);
 	ensure(nullptr == Pin || Pin->Direction == EGPD_Input);
 	return Pin;
 }
@@ -195,7 +198,7 @@ bool UK2Node_ConstructObjectFromClass::IsSpawnVarPin(UEdGraphPin* Pin) const
 			Pin->PinName != UEdGraphSchema_K2::PN_ReturnValue &&
 			Pin->PinName != FK2Node_ConstructObjectFromClassHelper::ClassPinName &&
 			Pin->PinName != FK2Node_ConstructObjectFromClassHelper::WorldContextPinName &&
-			Pin->PinName != FK2Node_ConstructObjectFromClassHelper::OuterPinName);
+			Pin->PinName != FK2Node_ConstructObjectFromClassHelper::OuterPinFriendlyName);
 }
 
 void UK2Node_ConstructObjectFromClass::OnClassPinChanged()
diff --git a/Engine/Source/Editor/BlueprintGraph/Private/K2Node_GenericCreateObject.cpp b/Engine/Source/Editor/BlueprintGraph/Private/K2Node_GenericCreateObject.cpp
--- a/Engine/Source/Editor/BlueprintGraph/Private/K2Node_GenericCreateObject.cpp
+++ b/Engine/Source/Editor/BlueprintGraph/Private/K2Node_GenericCreateObject.cpp
@@ -3,6 +3,7 @@
 #include "K2Node_GenericCreateObject.h"
 #include "Kismet/GameplayStatics.h"
 #include "K2Node_CallFunction.h"
+#include "K2Node_Self.h"
 #include "KismetCompilerMisc.h"
 #include "KismetCompiler.h"
 
@@ -72,10 +73,13 @@ void UK2Node_GenericCreateObject::ExpandNode(class FKismetCompilerContext& Compi
 	CallCreateNode->FunctionReference.SetExternalMember(GET_FUNCTION_NAME_CHECKED(UGameplayStatics, SpawnObject), UGameplayStatics::StaticClass());
 	CallCreateNode->AllocateDefaultPins();
 
+	bool bSucceeded = true;
+
+	// connect GenericCreateObject's self pin to the self node
+	bSucceeded &= ExpandDefaultToSelfPin(CompilerContext, SourceGraph, CallCreateNode);
+
 	// store off the class to spawn before we mutate pin connections:
 	UClass* ClassToSpawn = GetClassToSpawn();
-	
-	bool bSucceeded = true;
 	//connect exe
 	{
 		UEdGraphPin* SpawnExecPin = GetExecPin();
@@ -96,7 +100,6 @@ void UK2Node_GenericCreateObject::ExpandNode(class FKismetCompilerContext& Compi
 		UEdGraphPin* CallOuterPin = CallCreateNode->FindPin(TEXT("Outer"));
 		bSucceeded &= SpawnOuterPin && CallOuterPin && CompilerContext.MovePinLinksToIntermediate(*SpawnOuterPin, *CallOuterPin).CanSafeConnect();
 	}
-
 	UEdGraphPin* CallResultPin = nullptr;
 	//connect result
 	{
@@ -136,12 +139,6 @@ void UK2Node_GenericCreateObject::EarlyValidation(class FCompilerResultsLog& Mes
 	{
 		MessageLog.Error(*FText::Format(LOCTEXT("GenericCreateObject_WrongClassFmt", "Cannot construct objects of type '{0}' in @@"), FText::FromString(GetPathNameSafe(ClassToSpawn))).ToString(), this);
 	}
-
-	UEdGraphPin* OuterPin = GetOuterPin();
-	if (!OuterPin || (!OuterPin->DefaultObject && !OuterPin->LinkedTo.Num()))
-	{
-		MessageLog.Error(*LOCTEXT("GenericCreateObject_NoOuter", "Outer object is required in @@").ToString(), this);
-	}
 }
 
 bool UK2Node_GenericCreateObject::IsCompatibleWithGraph(const UEdGraph* TargetGraph) const
@@ -149,4 +146,37 @@ bool UK2Node_GenericCreateObject::IsCompatibleWithGraph(const UEdGraph* TargetGr
 	return UK2Node::IsCompatibleWithGraph(TargetGraph);
 }
 
+bool UK2Node_GenericCreateObject::ExpandDefaultToSelfPin(FKismetCompilerContext& CompilerContext, UEdGraph* SourceGraph, UK2Node_CallFunction* CallCreateNode)
+{
+	// ensure valid input
+	if (!SourceGraph || !CallCreateNode)
+	{
+		return false;
+	}
+
+	// Connect a self reference pin if there is a TScriptInterface default to self
+	if (const UFunction* TargetFunc = CallCreateNode->GetTargetFunction())
+	{
+		const FString& MetaData = TargetFunc->GetMetaData(FBlueprintMetadata::MD_DefaultToSelf);
+		if (!MetaData.IsEmpty())
+		{
+			// Get the Self Pin from this so we can default it correctly
+			const UEdGraphSchema_K2* Schema = CompilerContext.GetSchema();
+			if (UEdGraphPin* DefaultToSelfPin = Schema->FindSelfPin(*this, EGPD_Input))
+			{
+				// If it has no links then spawn a new self node here
+				if (DefaultToSelfPin->LinkedTo.IsEmpty())
+				{
+					UK2Node_Self* SelfNode = CompilerContext.SpawnIntermediateNode<UK2Node_Self>(this, SourceGraph);
+					SelfNode->AllocateDefaultPins();
+					UEdGraphPin* SelfPin = SelfNode->FindPinChecked(UEdGraphSchema_K2::PSC_Self);
+					// Make a connection from this intermediate self pin to here
+					return Schema->TryCreateConnection(DefaultToSelfPin, SelfPin);
+				}
+			}
+		}
+	}
+	return true;
+}
+
 #undef LOCTEXT_NAMESPACE
diff --git a/Engine/Source/Runtime/Engine/Classes/Kismet/GameplayStatics.h b/Engine/Source/Runtime/Engine/Classes/Kismet/GameplayStatics.h
--- a/Engine/Source/Runtime/Engine/Classes/Kismet/GameplayStatics.h
+++ b/Engine/Source/Runtime/Engine/Classes/Kismet/GameplayStatics.h
@@ -50,7 +50,7 @@ class ENGINE_API UGameplayStatics : public UBlueprintFunctionLibrary
 	GENERATED_UCLASS_BODY()
 
 	// --- Create Object
-	UFUNCTION(BlueprintCallable, Category = "Spawning", meta = (BlueprintInternalUseOnly = "true"))
+	UFUNCTION(BlueprintCallable, Category = "Spawning", meta = (BlueprintInternalUseOnly = "true", DefaultToSelf = "Outer"))
 	static UObject* SpawnObject(TSubclassOf<UObject> ObjectClass, UObject* Outer);
 
 	// --- Spawning functions ------------------------------
