Fix: Outer Pin of Construct Object from Class should not have an asset picker
- changed pin visual to a "self" style
- made pin default to "self" if there's nothing connected to it

diff --git a/Engine/Source/Editor/BlueprintGraph/Classes/K2Node_ConstructObjectFromClass.h b/Engine/Source/Editor/BlueprintGraph/Classes/K2Node_ConstructObjectFromClass.h
--- a/Engine/Source/Editor/BlueprintGraph/Classes/K2Node_ConstructObjectFromClass.h
+++ b/Engine/Source/Editor/BlueprintGraph/Classes/K2Node_ConstructObjectFromClass.h
@@ -33,6 +33,7 @@ class BLUEPRINTGRAPH_API UK2Node_ConstructObjectFromClass : public UK2Node
 	//~ Begin UK2Node Interface
 	virtual bool IsNodeSafeToIgnore() const override { return true; }
 	virtual void ReallocatePinsDuringReconstruction(TArray<UEdGraphPin*>& OldPins) override;
+	virtual ERedirectType DoPinsMatchForReconstruction(const UEdGraphPin* NewPin, int32 NewPinIndex, const UEdGraphPin* OldPin, int32 OldPinIndex) const override;
 	virtual void GetNodeAttributes( TArray<TKeyValuePair<FString, FString>>& OutNodeAttributes ) const override;
 	virtual void GetMenuActions(FBlueprintActionDatabaseRegistrar& ActionRegistrar) const override;
 	virtual FText GetMenuCategory() const override;
diff --git a/Engine/Source/Editor/BlueprintGraph/Classes/K2Node_GenericCreateObject.h b/Engine/Source/Editor/BlueprintGraph/Classes/K2Node_GenericCreateObject.h
--- a/Engine/Source/Editor/BlueprintGraph/Classes/K2Node_GenericCreateObject.h
+++ b/Engine/Source/Editor/BlueprintGraph/Classes/K2Node_GenericCreateObject.h
@@ -8,6 +8,8 @@
 #include "K2Node_GenericCreateObject.generated.h"
 
 class UEdGraph;
+class FKismetCompilerContext;
+class UK2Node_CallFunction;
 
 UCLASS()
 class BLUEPRINTGRAPH_API UK2Node_GenericCreateObject : public UK2Node_ConstructObjectFromClass
@@ -22,4 +24,15 @@ class BLUEPRINTGRAPH_API UK2Node_GenericCreateObject : public UK2Node_ConstructO
 
 	virtual bool UseWorldContext() const override { return false; }
 	virtual bool UseOuter() const override { return true; }
+
+	/**
+	 * attaches a self node to the self pin of 'this' if the CallCreateNode function has DefaultToSelf in it's metadata
+	 *
+	 * @param	CompilerContext		the context to expand in - likely passed from ExpandNode
+	 * @param	SourceGraph			the graph to expand in - likely passed from ExpandNode
+	 * @param	CallCreateNode		the CallFunction node that 'this' is imitating
+	 *
+	 * @return	true on success.
+	 */
+	bool ExpandDefaultToSelfPin(FKismetCompilerContext& CompilerContext, UEdGraph* SourceGraph, UK2Node_CallFunction* CallCreateNode);
 };
diff --git a/Engine/Source/Editor/BlueprintGraph/Private/K2Node_ConstructObjectFromClass.cpp b/Engine/Source/Editor/BlueprintGraph/Private/K2Node_ConstructObjectFromClass.cpp
--- a/Engine/Source/Editor/BlueprintGraph/Private/K2Node_ConstructObjectFromClass.cpp
+++ b/Engine/Source/Editor/BlueprintGraph/Private/K2Node_ConstructObjectFromClass.cpp
@@ -13,12 +13,14 @@ struct FK2Node_ConstructObjectFromClassHelper
 {
 	static FName WorldContextPinName;
 	static FName ClassPinName;
+	static FName OuterPinFriendlyName;
 	static FName OuterPinName;
 };
 
 FName FK2Node_ConstructObjectFromClassHelper::WorldContextPinName(TEXT("WorldContextObject"));
 FName FK2Node_ConstructObjectFromClassHelper::ClassPinName(TEXT("Class"));
-FName FK2Node_ConstructObjectFromClassHelper::OuterPinName(TEXT("Outer"));
+FName FK2Node_ConstructObjectFromClassHelper::OuterPinFriendlyName(TEXT("Outer"));
+FName FK2Node_ConstructObjectFromClassHelper::OuterPinName(UEdGraphSchema_K2::PN_Self);
 
 #define LOCTEXT_NAMESPACE "K2Node_ConstructObjectFromClass"
 
@@ -60,7 +62,8 @@ void UK2Node_ConstructObjectFromClass::AllocateDefaultPins()
 	
 	if (UseOuter())
 	{
-		UEdGraphPin* OuterPin = CreatePin(EGPD_Input, UEdGraphSchema_K2::PC_Object, UObject::StaticClass(), FK2Node_ConstructObjectFromClassHelper::OuterPinName);
+		UEdGraphPin* OuterPin = CreatePin(EGPD_Input, UEdGraphSchema_K2::PC_Object, UEdGraphSchema_K2::PSC_Self, nullptr, FK2Node_ConstructObjectFromClassHelper::OuterPinName);
+		OuterPin->PinFriendlyName = FText::FromName(FK2Node_ConstructObjectFromClassHelper::OuterPinFriendlyName);
 	}
 
 	Super::AllocateDefaultPins();
@@ -167,6 +170,18 @@ void UK2Node_ConstructObjectFromClass::ReallocatePinsDuringReconstruction(TArray
 	RestoreSplitPins(OldPins);
 }
 
+UK2Node::ERedirectType UK2Node_ConstructObjectFromClass::DoPinsMatchForReconstruction(const UEdGraphPin* NewPin, int32 NewPinIndex, const UEdGraphPin* OldPin, int32 OldPinIndex) const
+{
+	// the name of the outer pin was changed and its friendly name was updated to match
+	// the legacy naming. Use this to identify the change
+	if (NewPin->PinName == FK2Node_ConstructObjectFromClassHelper::OuterPinName &&
+		OldPin->PinName == FK2Node_ConstructObjectFromClassHelper::OuterPinFriendlyName)
+	{
+		return ERedirectType_Name;
+	}
+	return Super::DoPinsMatchForReconstruction(NewPin, NewPinIndex, OldPin, OldPinIndex);
+}
+
 void UK2Node_ConstructObjectFromClass::PostPlacedNewNode()
 {
 	Super::PostPlacedNewNode();
diff --git a/Engine/Source/Editor/BlueprintGraph/Private/K2Node_GenericCreateObject.cpp b/Engine/Source/Editor/BlueprintGraph/Private/K2Node_GenericCreateObject.cpp
--- a/Engine/Source/Editor/BlueprintGraph/Private/K2Node_GenericCreateObject.cpp
+++ b/Engine/Source/Editor/BlueprintGraph/Private/K2Node_GenericCreateObject.cpp
@@ -3,6 +3,7 @@
 #include "K2Node_GenericCreateObject.h"
 #include "Kismet/GameplayStatics.h"
 #include "K2Node_CallFunction.h"
+#include "K2Node_Self.h"
 #include "KismetCompilerMisc.h"
 #include "KismetCompiler.h"
 
@@ -71,11 +72,13 @@ void UK2Node_GenericCreateObject::ExpandNode(class FKismetCompilerContext& Compi
 	UK2Node_CallFunction* CallCreateNode = CompilerContext.SpawnIntermediateNode<UK2Node_CallFunction>(this, SourceGraph);
 	CallCreateNode->FunctionReference.SetExternalMember(GET_FUNCTION_NAME_CHECKED(UGameplayStatics, SpawnObject), UGameplayStatics::StaticClass());
 	CallCreateNode->AllocateDefaultPins();
+	
+	// connect GenericCreateObject's self pin to the self node
+	bool bSucceeded = ExpandDefaultToSelfPin(CompilerContext, SourceGraph, CallCreateNode);
 
 	// store off the class to spawn before we mutate pin connections:
 	UClass* ClassToSpawn = GetClassToSpawn();
-	
-	bool bSucceeded = true;
+
 	//connect exe
 	{
 		UEdGraphPin* SpawnExecPin = GetExecPin();
@@ -136,12 +139,6 @@ void UK2Node_GenericCreateObject::EarlyValidation(class FCompilerResultsLog& Mes
 	{
 		MessageLog.Error(*FText::Format(LOCTEXT("GenericCreateObject_WrongClassFmt", "Cannot construct objects of type '{0}' in @@"), FText::FromString(GetPathNameSafe(ClassToSpawn))).ToString(), this);
 	}
-
-	UEdGraphPin* OuterPin = GetOuterPin();
-	if (!OuterPin || (!OuterPin->DefaultObject && !OuterPin->LinkedTo.Num()))
-	{
-		MessageLog.Error(*LOCTEXT("GenericCreateObject_NoOuter", "Outer object is required in @@").ToString(), this);
-	}
 }
 
 bool UK2Node_GenericCreateObject::IsCompatibleWithGraph(const UEdGraph* TargetGraph) const
@@ -149,4 +146,37 @@ bool UK2Node_GenericCreateObject::IsCompatibleWithGraph(const UEdGraph* TargetGr
 	return UK2Node::IsCompatibleWithGraph(TargetGraph);
 }
 
+bool UK2Node_GenericCreateObject::ExpandDefaultToSelfPin(FKismetCompilerContext& CompilerContext, UEdGraph* SourceGraph, UK2Node_CallFunction* CallCreateNode)
+{
+	// ensure valid input
+	if (!SourceGraph || !CallCreateNode)
+	{
+		return false;
+	}
+
+	// Connect a self reference pin if there is a TScriptInterface default to self
+	if (const UFunction* TargetFunc = CallCreateNode->GetTargetFunction())
+	{
+		const FString& MetaData = TargetFunc->GetMetaData(FBlueprintMetadata::MD_DefaultToSelf);
+		if (!MetaData.IsEmpty())
+		{
+			// Get the Self Pin from this so we can default it correctly
+			const UEdGraphSchema_K2* Schema = CompilerContext.GetSchema();
+			if (UEdGraphPin* DefaultToSelfPin = Schema->FindSelfPin(*this, EGPD_Input))
+			{
+				// If it has no links then spawn a new self node here
+				if (DefaultToSelfPin->LinkedTo.IsEmpty())
+				{
+					UK2Node_Self* SelfNode = CompilerContext.SpawnIntermediateNode<UK2Node_Self>(this, SourceGraph);
+					SelfNode->AllocateDefaultPins();
+					UEdGraphPin* SelfPin = SelfNode->FindPinChecked(UEdGraphSchema_K2::PSC_Self);
+					// Make a connection from this intermediate self pin to here
+					return Schema->TryCreateConnection(DefaultToSelfPin, SelfPin);
+				}
+			}
+		}
+	}
+	return true;
+}
+
 #undef LOCTEXT_NAMESPACE
diff --git a/Engine/Source/Runtime/Engine/Classes/Kismet/GameplayStatics.h b/Engine/Source/Runtime/Engine/Classes/Kismet/GameplayStatics.h
--- a/Engine/Source/Runtime/Engine/Classes/Kismet/GameplayStatics.h
+++ b/Engine/Source/Runtime/Engine/Classes/Kismet/GameplayStatics.h
@@ -50,7 +50,7 @@ class ENGINE_API UGameplayStatics : public UBlueprintFunctionLibrary
 	GENERATED_UCLASS_BODY()
 
 	// --- Create Object
-	UFUNCTION(BlueprintCallable, Category = "Spawning", meta = (BlueprintInternalUseOnly = "true"))
+	UFUNCTION(BlueprintCallable, Category = "Spawning", meta = (BlueprintInternalUseOnly = "true", DefaultToSelf = "Outer"))
 	static UObject* SpawnObject(TSubclassOf<UObject> ObjectClass, UObject* Outer);
 
 	// --- Spawning functions ------------------------------
