Fix issue where a property get/set node fails to refresh after the reference class is re-instanced via Live Coding

#rnx

diff --git a/Engine/Source/Editor/BlueprintGraph/Private/K2Node_Variable.cpp b/Engine/Source/Editor/BlueprintGraph/Private/K2Node_Variable.cpp
--- a/Engine/Source/Editor/BlueprintGraph/Private/K2Node_Variable.cpp
+++ b/Engine/Source/Editor/BlueprintGraph/Private/K2Node_Variable.cpp
@@ -829,12 +829,13 @@ FName UK2Node_Variable::GetCornerIcon() const
 
 bool UK2Node_Variable::HasExternalDependencies(TArray<class UStruct*>* OptionalOutput) const
 {
-	UClass* SourceClass = GetVariableSourceClass();
 	UBlueprint* SourceBlueprint = GetBlueprint();
-	bool bResult = (SourceClass && (SourceClass->ClassGeneratedBy != SourceBlueprint));
+	FProperty* VariableProperty = GetPropertyForVariable();
+	UClass* PropertySourceClass = VariableProperty ? VariableProperty->GetOwnerClass() : nullptr;
+	bool bResult = (PropertySourceClass && (PropertySourceClass->ClassGeneratedBy != SourceBlueprint));
 	if (bResult && OptionalOutput)
 	{
-		OptionalOutput->AddUnique(SourceClass);
+		OptionalOutput->AddUnique(PropertySourceClass);
 	}
 
 	// Also include underlying non-native variable types as external dependencies. Otherwise, contextual
