Fix math expression nodes to properly report errors by making sure a connected output pin will be saved as broken and copy error text accross reconstruct
Add ability to write out intermediate compiler errors at append time, and fix issue where the same source node would get appended to multiple times

diff --git a/Engine/Source/Editor/BlueprintGraph/Private/K2Node_MathExpression.cpp b/Engine/Source/Editor/BlueprintGraph/Private/K2Node_MathExpression.cpp
--- a/Engine/Source/Editor/BlueprintGraph/Private/K2Node_MathExpression.cpp
+++ b/Engine/Source/Editor/BlueprintGraph/Private/K2Node_MathExpression.cpp
@@ -18,6 +18,7 @@
 #include "BlueprintActionDatabaseRegistrar.h"
 #include "DiffResults.h"
 #include "MathExpressionHandler.h"
+#include "Misc/DefaultValueHelper.h"
 #include "BlueprintNodeSpawner.h"
 
 #define LOCTEXT_NAMESPACE "K2Node"
@@ -817,6 +818,10 @@ public:
 						{
 							FunctionName = TestFunction->GetMetaData(FBlueprintMetadata::MD_DisplayName);
 						}
+
+						// Remove spaces from display name as the parser cannot handle it
+						FunctionName = FDefaultValueHelper::RemoveWhitespaces(FunctionName);
+
 						Add(FunctionName, TestFunction);
 					}
 				}
@@ -1327,8 +1332,10 @@ public:
 		}
 		else
 		{
-			MessageLog.Error(*LOCTEXT("NoGraphGenerated", "No root node generated from the expression: '@@'").ToString(),
-				CompilingNode);
+			if (MessageLog.NumErrors == 0)
+			{
+				MessageLog.Error(*LOCTEXT("NoGraphGenerated", "No root node generated from the expression: '@@'").ToString(), CompilingNode);
+			}
 		}
 
 		// position the entry and exit nodes somewhere sane
@@ -2442,7 +2449,7 @@ UK2Node_MathExpression::UK2Node_MathExpression(const FObjectInitializer& ObjectI
 	bCanRenameNode = true;
 
 	bMadeAfterRotChange = false;
-	OrphanedPinSaveMode = ESaveOrphanPinMode::SaveNone;
+	OrphanedPinSaveMode = ESaveOrphanPinMode::SaveAll;
 }
 
 void UK2Node_MathExpression::Serialize(FArchive& Ar)
@@ -2620,7 +2627,10 @@ void UK2Node_MathExpression::RebuildExpression(FString InExpression)
 				// a series of errors being attached to the node).
 				if (!GraphGenerator.GenerateCode(ExpressionRoot.ToSharedRef(), *CachedMessageLog))
 				{
-					CachedMessageLog->Error(*LOCTEXT("MathExprGFailedGen", "Failed to generate full expression graph for: '@@'").ToString(), this);
+					if (CachedMessageLog->NumErrors == 0)
+					{
+						CachedMessageLog->Error(*LOCTEXT("MathExprGFailedGen", "Failed to generate full expression graph for: '@@'").ToString(), this);
+					}
 				}
 				else
 				{
@@ -2654,8 +2664,8 @@ void UK2Node_MathExpression::RebuildExpression(FString InExpression)
 			}
 		}
 
-		// refresh the node since the connections may have changed
-		Super::ReconstructNode();
+		// refresh the node since the connections may have changed, this won't be reentrant due to bool above
+		ReconstructNode();
 
 		// finally, recompile
 		UBlueprint* Blueprint = FBlueprintEditorUtils::FindBlueprintForNodeChecked(this);
@@ -2694,6 +2704,7 @@ void UK2Node_MathExpression::ClearExpression()
 	// becoming the "CurrentEventTarget", we pass false - we append logs 
 	// collected by this one to the full compiler log later on anyways (so they won't be missed)
 	CachedMessageLog = MakeShareable(new FCompilerResultsLog(/*bIsCompatibleWithEvents =*/false));
+	CachedMessageLog->bSilentMode = true;
 	
 	Expression.Empty();
 }
@@ -2705,7 +2716,7 @@ void UK2Node_MathExpression::ValidateNodeDuringCompilation(FCompilerResultsLog&
 
 	if (CachedMessageLog.IsValid())
 	{
-		MessageLog.Append(*CachedMessageLog);
+		MessageLog.Append(*CachedMessageLog, true);
 	}
 	// else, this may be some intermediate node in the compile, let's look at the errors from the original...
 	else 
@@ -2724,7 +2735,7 @@ void UK2Node_MathExpression::ValidateNodeDuringCompilation(FCompilerResultsLog&
 			// re-parse/re-gen to fish out the same errors)
 			if (MathExpression->CachedMessageLog.IsValid())
 			{
-				MessageLog.Append(*MathExpression->CachedMessageLog);
+				MessageLog.Append(*MathExpression->CachedMessageLog, true);
 			}
 		}
 	}
@@ -2789,6 +2800,20 @@ void UK2Node_MathExpression::ReconstructNode()
 	const FString OldErrorMessage = ErrorMsg;
 	Super::ReconstructNode();
 	ErrorMsg = OldErrorMessage;
+
+	// Mark our input pins as not saved, but we want to save orphaned output pins to show compile errors
+	for (UEdGraphPin* Pin : Pins)
+	{
+		// Recombine the sub pins back into the OptionPin
+		if (Pin->Direction == EEdGraphPinDirection::EGPD_Input)
+		{
+			Pin->SetSavePinIfOrphaned(false);
+		}
+		else
+		{
+			Pin->SetSavePinIfOrphaned(true);
+		}
+	}
 }
 
 //------------------------------------------------------------------------------
diff --git a/Engine/Source/Editor/UnrealEd/Private/Kismet2/CompilerResultsLog.cpp b/Engine/Source/Editor/UnrealEd/Private/Kismet2/CompilerResultsLog.cpp
--- a/Engine/Source/Editor/UnrealEd/Private/Kismet2/CompilerResultsLog.cpp
+++ b/Engine/Source/Editor/UnrealEd/Private/Kismet2/CompilerResultsLog.cpp
@@ -632,7 +632,7 @@ void FCompilerResultsLog::GetNodesFromTokens(const TArray<TSharedRef<IMessageTok
 			UObject* ObjectArgument = ObjectPtr.Get();
 			if(UEdGraphNode* Node = Cast<UEdGraphNode>(ObjectArgument))
 			{
-				OutOwnerNodes.Add(Node);
+				OutOwnerNodes.AddUnique(Node);
 			}
 		}
 		else if (Token->GetType() == EMessageToken::EdGraph)
@@ -644,18 +644,18 @@ void FCompilerResultsLog::GetNodesFromTokens(const TArray<TSharedRef<IMessageTok
 			{
 				if (PinBeingReferenced)
 				{
-					OutOwnerNodes.Add(Cast<UEdGraphNode>(PinBeingReferenced->GetOwningNodeUnchecked()));
+					OutOwnerNodes.AddUnique(Cast<UEdGraphNode>(PinBeingReferenced->GetOwningNodeUnchecked()));
 				}
 			}
 			else
 			{
-				OutOwnerNodes.Add(OwnerNode);
+				OutOwnerNodes.AddUnique(OwnerNode);
 			}
 		}
 	}
 }
 
-void FCompilerResultsLog::Append(FCompilerResultsLog const& Other)
+void FCompilerResultsLog::Append(FCompilerResultsLog const& Other, bool bWriteToSystemLog)
 {
 	for (TSharedRef<FTokenizedMessage> const& Message : Other.Messages)
 	{
@@ -678,11 +678,19 @@ void FCompilerResultsLog::Append(FCompilerResultsLog const& Other)
 				break;
 			}
 		}
-		Messages.Add(Message);
-		
+
 		TArray<UEdGraphNode*> OwnerNodes;
 		GetNodesFromTokens(Message->GetMessageTokens(), OwnerNodes);
-		AnnotateNode(OwnerNodes, Message);
+
+		if (bWriteToSystemLog)
+		{
+			InternalLogMessage(Message->GetIdentifier(), Message, OwnerNodes);
+		}
+		else
+		{
+			Messages.Add(Message);
+			AnnotateNode(OwnerNodes, Message);
+		}
 	}
 }
 
diff --git a/Engine/Source/Editor/UnrealEd/Public/Kismet2/CompilerResultsLog.h b/Engine/Source/Editor/UnrealEd/Public/Kismet2/CompilerResultsLog.h
--- a/Engine/Source/Editor/UnrealEd/Public/Kismet2/CompilerResultsLog.h
+++ b/Engine/Source/Editor/UnrealEd/Public/Kismet2/CompilerResultsLog.h
@@ -299,7 +299,8 @@ public:
 	UEdGraphPin* FindSourcePin(UEdGraphPin* PossiblyDuplicatedPin);
 	const UEdGraphPin* FindSourcePin(const UEdGraphPin* PossiblyDuplicatedPin) const;
 	
-	void Append(FCompilerResultsLog const& Other);
+	/** Copy errors from an existing log into this one, and optionally write out to log if it was suppressed the first time */
+	void Append(FCompilerResultsLog const& Other, bool bWriteToSystemLog = false);
 
 	/** Begin a new compiler event */
 	void BeginEvent(const TCHAR* InName);
