Disable warning for uncooked blueprint functions that are placed in runtime blueprints. This warning was accidently turned on by a previous change and is invalid.

#rnx

diff --git a/Engine/Source/Editor/BlueprintGraph/Private/K2Node_CallFunction.cpp b/Engine/Source/Editor/BlueprintGraph/Private/K2Node_CallFunction.cpp
--- a/Engine/Source/Editor/BlueprintGraph/Private/K2Node_CallFunction.cpp
+++ b/Engine/Source/Editor/BlueprintGraph/Private/K2Node_CallFunction.cpp
@@ -2053,10 +2053,11 @@ void UK2Node_CallFunction::ValidateNodeDuringCompilation(class FCompilerResultsL
 	const bool bIsEditorOnlyBlueprintBaseClass = !BlueprintClass || IsEditorOnlyObject(BlueprintClass);
 	static bool bAllowUnsafeBlueprintCalls = FParse::Param(FCommandLine::Get(), TEXT("AllowUnsafeBlueprintCalls"));
 
-	// This error is disabled while we figure out how we can identify uncooked only
-	// blueprints that want to make use of uncooked only APIs:
 	if (!bAllowUnsafeBlueprintCalls)
 	{
+		// This error is disabled while we figure out how we can identify uncooked only
+		// blueprints that want to make use of uncooked only APIs:
+		#if 0
 		const bool bIsUncookedOnlyFunction = Function && Function->GetOutermost()->HasAllPackagesFlags(PKG_UncookedOnly);
 		if (bIsUncookedOnlyFunction &&
 			// Only allow calls to uncooked only functions from editor only/uncooked only
@@ -2066,6 +2067,7 @@ void UK2Node_CallFunction::ValidateNodeDuringCompilation(class FCompilerResultsL
 		{
 			MessageLog.Error(*LOCTEXT("UncookedOnlyError", "Attempting to call uncooked only function @@ in runtime blueprint").ToString(), this);
 		}
+		#endif	// 0
 
 		// Ensure that editor module BP exposed UFunctions can only be called in blueprints for which the base class is also part of an editor module
 		// Also check for functions wrapped in WITH_EDITOR 
