Add a context target menu option to allow non-imported type filtering in the context menu to be toggled on/off.

The option will be visible only if namespace filtering features are turned on (currently via CVar).

Also revised filter behavior to route through the shared import set and apply to actions that pass the non-target member test (i.e. global fields).

diff --git a/Engine/Source/Editor/BlueprintGraph/Private/BlueprintActionFilter.cpp b/Engine/Source/Editor/BlueprintGraph/Private/BlueprintActionFilter.cpp
--- a/Engine/Source/Editor/BlueprintGraph/Private/BlueprintActionFilter.cpp
+++ b/Engine/Source/Editor/BlueprintGraph/Private/BlueprintActionFilter.cpp
@@ -11,7 +11,6 @@
 #include "EdGraph/EdGraphSchema.h"
 #include "EdGraph/EdGraph.h"
 #include "EdGraphSchema_K2.h"
-#include "BlueprintNamespaceHelper.h"
 #include "K2Node.h"
 #include "K2Node_Event.h"
 #include "K2Node_CallFunction.h"
@@ -37,6 +36,8 @@
 #include "K2Node_Message.h"
 #include "K2Node_PromotableOperator.h"
 #include "EditorCategoryUtils.h"
+#include "BlueprintEditorModule.h"
+#include "Kismet2/KismetEditorUtilities.h"
 
 /*******************************************************************************
  * Static BlueprintActionFilter Helpers
@@ -230,9 +231,20 @@ namespace BlueprintActionFilterImpl
 	 * 
 	 * @param  Filter			Holds the TagetClass context for this test.
 	 * @param  BlueprintAction	The action you wish to query.
+	 * @param  Flags			Filter flags that may modify this test.
 	 * @return 
 	 */
-	static bool IsRejectedGlobalField(FBlueprintActionFilter const& Filter, FBlueprintActionInfo& BlueprintAction);
+	static bool IsRejectedGlobalField(FBlueprintActionFilter const& Filter, FBlueprintActionInfo& BlueprintAction, uint32 Flags);
+
+	/**
+	 * Rejection test that checks to see if the node-spawner represents a field
+	 * that does NOT exist within the scope of the editor's current import set.
+	 *
+	 * @param  Filter			The filter context for this test.
+	 * @param  BlueprintAction	The action you wish to query.
+	 * @return True if the action is associated with a member field that is not included in any imported or global namespace.
+	 */
+	static bool IsNonImportedField(FBlueprintActionFilter const& Filter, FBlueprintActionInfo& BlueprintAction);
 
 	/**
 	 * Rejection test that checks to see if the node-spawner is associated with 
@@ -297,13 +309,6 @@ namespace BlueprintActionFilterImpl
 	 */
 	static bool IsOutOfScopeLocalVariable(FBlueprintActionFilter const& Filter, FBlueprintActionInfo& BlueprintAction);
 
-	/**
-	 * Rejection test that checks to see if the action references a member variable or function that is
-	 * outside the scope of the current set of namespaces imported by the current Blueprint context(s).
-	 *
-	 */
-	static bool IsOutOfScopeMemberField(FBlueprintActionFilter const& Filter, FBlueprintActionInfo& BlueprintAction);
-
 	/**
 	 * 
 	 * 
@@ -876,13 +881,50 @@ static bool BlueprintActionFilterImpl::IsDeprecated(FBlueprintActionFilter const
 }
 
 //------------------------------------------------------------------------------
-static bool BlueprintActionFilterImpl::IsRejectedGlobalField(FBlueprintActionFilter const& Filter, FBlueprintActionInfo& BlueprintAction)
+static bool BlueprintActionFilterImpl::IsNonImportedField(FBlueprintActionFilter const& Filter, FBlueprintActionInfo& BlueprintAction)
 {
+	// All fields are considered to be imported if there is no context.
 	bool bIsFilteredOut = false;
+
 	FFieldVariant Field = BlueprintAction.GetAssociatedMemberField();
 	if (Field.IsValid())
 	{
-		bIsFilteredOut = IsGloballyAccessible(Field);
+		// @todo_namespaces - Maybe cache this as part of the filter context, or decouple queries from the editor?
+		TArray<TSharedPtr<IBlueprintEditor>> BlueprintEditors;
+		for (const UBlueprint* Blueprint : Filter.Context.Blueprints)
+		{
+			const bool bOpenEditor = false;
+			TSharedPtr<IBlueprintEditor> BlueprintEditor = FKismetEditorUtilities::GetIBlueprintEditorForObject(Blueprint, bOpenEditor);
+			if (BlueprintEditor.IsValid())
+			{
+				BlueprintEditors.AddUnique(BlueprintEditor);
+			}
+		}
+
+		for (const TSharedPtr<IBlueprintEditor>& BlueprintEditor : BlueprintEditors)
+		{
+			if (BlueprintEditor->IsNonImportedField(Field))
+			{
+				bIsFilteredOut = true;
+				break;
+			}
+		}
+	}
+
+	return bIsFilteredOut;
+}
+
+//------------------------------------------------------------------------------
+static bool BlueprintActionFilterImpl::IsRejectedGlobalField(FBlueprintActionFilter const& Filter, FBlueprintActionInfo& BlueprintAction, uint32 Flags)
+{
+	const bool bRejectAllGlobalFields = ((Flags & FBlueprintActionFilter::EFlags::BPFILTER_RejectGlobalFields) != 0);
+	const bool bRejectNonImportedFields = ((Flags & FBlueprintActionFilter::EFlags::BPFILTER_RejectNonImportedFields) != 0);
+
+	bool bIsFilteredOut = false;
+	FFieldVariant Field = BlueprintAction.GetAssociatedMemberField();
+	if (Field.IsValid() && IsGloballyAccessible(Field))
+	{
+		bIsFilteredOut = bRejectAllGlobalFields || (bRejectNonImportedFields && IsNonImportedField(Filter, BlueprintAction));
 		
 		UClass* FieldClass = Field.GetOwnerClass();
 		if (bIsFilteredOut && (FieldClass != nullptr))
@@ -1165,36 +1207,6 @@ static bool BlueprintActionFilterImpl::IsOutOfScopeLocalVariable(FBlueprintActio
 	return bIsFilteredOut;
 }
 
-//------------------------------------------------------------------------------
-static bool BlueprintActionFilterImpl::IsOutOfScopeMemberField(FBlueprintActionFilter const& Filter, FBlueprintActionInfo& BlueprintAction)
-{
-	bool bIsFilteredOut = false;
-	FFieldVariant MemberField = BlueprintAction.GetAssociatedMemberField();
-	if (MemberField.IsValid())
-	{
-		if (const UClass* OwnerClass = MemberField.GetOwnerClass())
-		{
-			// If the owning class was generated from a Blueprint, check to see if its namespace falls within the scope of the filter context.
-			if (const UBlueprint* OwnerBlueprint = UBlueprint::GetBlueprintFromClass(OwnerClass))
-			{
-				bIsFilteredOut = true;
-				for (const UBlueprint* Blueprint : Filter.Context.Blueprints)
-				{
-					// @todo - Cache this as a transient member on the BP itself and/or expose an API to query w/o incurring the cost of recreating this locally
-					FBlueprintNamespaceHelper NamespaceHelper(Blueprint);
-					if (NamespaceHelper.IsIncludedInNamespaceList(OwnerBlueprint->BlueprintNamespace))
-					{
-						// This action's member belongs to a class that falls within the scope of our current import set, so don't filter it out.
-						bIsFilteredOut = false;
-						break;
-					}
-				}
-			}
-		}
-	}
-	return bIsFilteredOut;
-}
-
 //------------------------------------------------------------------------------
 static bool BlueprintActionFilterImpl::IsLevelScriptActionValid(FBlueprintActionFilter const& Filter, FBlueprintActionInfo& BlueprintAction)
 {
@@ -2036,9 +2048,9 @@ FBlueprintActionFilter::FBlueprintActionFilter(uint32 Flags/*= 0x00*/)
 	
 	AddRejectionTest(FRejectionTestDelegate::CreateStatic(IsActionHiddenByConfig));
 	AddRejectionTest(FRejectionTestDelegate::CreateStatic(IsFieldCategoryHidden));
-	if (Flags & BPFILTER_RejectGlobalFields)
+	if (Flags & (BPFILTER_RejectGlobalFields | BPFILTER_RejectNonImportedFields))
 	{
-		AddRejectionTest(FRejectionTestDelegate::CreateStatic(IsRejectedGlobalField));
+		AddRejectionTest(FRejectionTestDelegate::CreateStatic(IsRejectedGlobalField, Flags));
 	}
 	
 	AddRejectionTest(FRejectionTestDelegate::CreateStatic(IsFieldInaccessible));
@@ -2060,10 +2072,6 @@ FBlueprintActionFilter::FBlueprintActionFilter(uint32 Flags/*= 0x00*/)
 	AddRejectionTest(FRejectionTestDelegate::CreateStatic(IsNonTargetMember, !(Flags & BPFILTER_RejectGlobalFields)));
 	AddRejectionTest(FRejectionTestDelegate::CreateStatic(IsUnBoundBindingSpawner));
 	AddRejectionTest(FRejectionTestDelegate::CreateStatic(IsOutOfScopeLocalVariable));
-	if (Flags & BPFILTER_RejectOutOfScopeMembers)
-	{
-		AddRejectionTest(FRejectionTestDelegate::CreateStatic(IsOutOfScopeMemberField));
-	}
 	AddRejectionTest(FRejectionTestDelegate::CreateStatic(IsLevelScriptActionValid));
 
 	AddRejectionTest(FRejectionTestDelegate::CreateStatic(IsHiddenInNonEditorBlueprint));	
diff --git a/Engine/Source/Editor/BlueprintGraph/Public/BlueprintActionFilter.h b/Engine/Source/Editor/BlueprintGraph/Public/BlueprintActionFilter.h
--- a/Engine/Source/Editor/BlueprintGraph/Public/BlueprintActionFilter.h
+++ b/Engine/Source/Editor/BlueprintGraph/Public/BlueprintActionFilter.h
@@ -193,7 +193,7 @@ public:
 		 * to a struct that exists outside the scope of the current set of
 		 * imported namespaces will be filtered out.
 		 */
-		BPFILTER_RejectOutOfScopeMembers	= (1<<4),
+		BPFILTER_RejectNonImportedFields	= (1<<4),
 
 		/**
 		 * Rejects actions that would spawn nodes with thread safety characteristics
diff --git a/Engine/Source/Editor/Kismet/Private/BlueprintActionMenuUtils.cpp b/Engine/Source/Editor/Kismet/Private/BlueprintActionMenuUtils.cpp
--- a/Engine/Source/Editor/Kismet/Private/BlueprintActionMenuUtils.cpp
+++ b/Engine/Source/Editor/Kismet/Private/BlueprintActionMenuUtils.cpp
@@ -456,9 +456,9 @@ void FBlueprintActionMenuUtils::MakeContextMenu(FBlueprintActionContext const& C
 		FilterFlags |= FBlueprintActionFilter::BPFILTER_RejectGlobalFields;
 	}
 
-	if (bIsContextSensitive && GetDefault<UBlueprintEditorSettings>()->bEnableNamespaceFilteringFeatures)
+	if ( bIsContextSensitive && ((ClassTargetMask & EContextTargetFlags::TARGET_NonImportedTypes) == 0) )
 	{
-		FilterFlags |= FBlueprintActionFilter::BPFILTER_RejectOutOfScopeMembers;
+		FilterFlags |= FBlueprintActionFilter::BPFILTER_RejectNonImportedFields;
 	}
 
 	if(bIsContextSensitive)
diff --git a/Engine/Source/Editor/Kismet/Private/BlueprintEditor.cpp b/Engine/Source/Editor/Kismet/Private/BlueprintEditor.cpp
--- a/Engine/Source/Editor/Kismet/Private/BlueprintEditor.cpp
+++ b/Engine/Source/Editor/Kismet/Private/BlueprintEditor.cpp
@@ -10169,6 +10169,47 @@ void FBlueprintEditor::ClearAllGraphEditorQuickJumps()
 	LocalSettings->SaveConfig();
 }
 
+bool FBlueprintEditor::IsNonImportedField(FFieldVariant Field) const
+{
+	bool bNotImported = false;
+
+	const UStruct* FieldType = nullptr;
+	if (Field.IsUObject())
+	{
+		if (const UField* FieldAsUObject = Cast<UField>(Field.ToUObject()))
+		{
+			FieldType = Cast<UStruct>(FieldAsUObject);
+			if (!FieldType)
+			{
+				FieldType = FieldAsUObject->GetOwnerStruct();
+			}
+		}
+	}
+	else
+	{
+		FieldType = Field.ToField()->GetOwnerStruct();
+	}
+
+	if (FieldType)
+	{
+		for (const UObject* EditingObj : GetEditingObjects())
+		{
+			if (const UBlueprint* BP = Cast<UBlueprint>(EditingObj))
+			{
+				// Casting away the 'const' here currently because this is a non-const method that can modify the internally-cached set.
+				TSharedRef<FBlueprintNamespaceHelper> NamespaceHelper = const_cast<FBlueprintEditor*>(this)->GetOrCreateNamespaceHelperForBlueprint(BP);
+				if (!NamespaceHelper->IsImportedType(FieldType))
+				{
+					bNotImported = true;
+					break;
+				}
+			}
+		}
+	}
+
+	return bNotImported;
+}
+
 /////////////////////////////////////////////////////
 /////////////////////////////////////////////////////
 /////////////////////////////////////////////////////
diff --git a/Engine/Source/Editor/Kismet/Private/SBlueprintContextTargetMenu.cpp b/Engine/Source/Editor/Kismet/Private/SBlueprintContextTargetMenu.cpp
--- a/Engine/Source/Editor/Kismet/Private/SBlueprintContextTargetMenu.cpp
+++ b/Engine/Source/Editor/Kismet/Private/SBlueprintContextTargetMenu.cpp
@@ -1,4 +1,4 @@
-﻿// Copyright Epic Games, Inc. All Rights Reserved.
+// Copyright Epic Games, Inc. All Rights Reserved.
 
 #include "SBlueprintContextTargetMenu.h"
 #include "Engine/Blueprint.h"
@@ -372,6 +372,12 @@ FContextMenuTargetProfile::FContextMenuTargetProfile(const FBlueprintActionConte
 		SavedTargetFlags = 0;
 	}
 
+	if (BpSettings->bEnableNamespaceFilteringFeatures)
+	{
+		// Filter out non-imported types by default, but do so only if namespace filtering features are also enabled.
+		SavedTargetFlags &= ~EContextTargetFlags::TARGET_NonImportedTypes;
+	}
+
 	if (!LoadProfile())
 	{
 		// maybe they were originally using the shared context profile? so let's default to that
@@ -550,6 +556,12 @@ void SBlueprintContextTargetMenu::Construct(const FArguments& InArgs, const FBlu
 		{
 			continue;
 		}
+
+		// Keep this target hidden for now unless namespace filtering features are turned on.
+		if (ContextTarget == EContextTargetFlags::TARGET_NonImportedTypes && !GetDefault<UBlueprintEditorSettings>()->bEnableNamespaceFilteringFeatures)
+		{
+			continue;
+		}
 		
 		FText const MenuName = BlueprintContextTargetMenuImpl::GetContextTargetDisplayName(TargetEnum, BitMaskOffset);
 		const FContextMenuTargetProfile& ProfileRef = TargetProfile;
diff --git a/Engine/Source/Editor/Kismet/Public/BlueprintActionMenuUtils.h b/Engine/Source/Editor/Kismet/Public/BlueprintActionMenuUtils.h
--- a/Engine/Source/Editor/Kismet/Public/BlueprintActionMenuUtils.h
+++ b/Engine/Source/Editor/Kismet/Public/BlueprintActionMenuUtils.h
@@ -21,6 +21,7 @@ namespace EContextTargetFlags
 		TARGET_PinObject			= 0x00000008 UMETA(DisplayName="Pin Type Class", ToolTip="Include functions and variables that belong to this pin type."),
 		TARGET_SiblingPinObjects	= 0x00000010 UMETA(DisplayName="Other Object Outputs", ToolTip="Include functions and variables that belong to any of this node's output types."),
 		TARGET_BlueprintLibraries	= 0x00000020 UMETA(DisplayName="Libraries", ToolTip="Include static functions that are globally accessible (belonging to function/macro libraries, etc.)."),
+		TARGET_NonImportedTypes		= 0x00000040 UMETA(DisplayName="Non-Imported Types", ToolTip="Include functions from namespaces that are not imported by this Blueprint."),
 
 		// +1 to the last flag (so we can easily iterate these flags)
 		ContextTargetFlagsEnd UMETA(Hidden),
diff --git a/Engine/Source/Editor/Kismet/Public/BlueprintEditor.h b/Engine/Source/Editor/Kismet/Public/BlueprintEditor.h
--- a/Engine/Source/Editor/Kismet/Public/BlueprintEditor.h
+++ b/Engine/Source/Editor/Kismet/Public/BlueprintEditor.h
@@ -243,6 +243,7 @@ public:
 	void AnalyticsTrackCompileEvent(UBlueprint* Blueprint, int32 NumErrors, int32 NumWarnings) const;
 	virtual TSharedPtr<class IClassViewerFilter> GetImportedClassViewerFilter() const override { return ImportedClassViewerFilter; }
 	virtual TSharedPtr<class IPinTypeSelectorFilter> GetImportedPinTypeSelectorFilter() const override { return ImportedPinTypeSelectorFilter; }
+	virtual bool IsNonImportedField(FFieldVariant Field) const;
 	//~ End IBlueprintEditor Interface
 
 	//~ Begin FTickableEditorObject Interface
diff --git a/Engine/Source/Editor/Kismet/Public/BlueprintEditorModule.h b/Engine/Source/Editor/Kismet/Public/BlueprintEditorModule.h
--- a/Engine/Source/Editor/Kismet/Public/BlueprintEditorModule.h
+++ b/Engine/Source/Editor/Kismet/Public/BlueprintEditorModule.h
@@ -104,6 +104,9 @@ public:
 	/** Return the pin type selector filter associated with the current set of imported namespaces within this editor context. Default is NULL (no filter). */
 	virtual TSharedPtr<class IPinTypeSelectorFilter> GetImportedPinTypeSelectorFilter() const { return nullptr; }
 
+	/** Return whether the given field falls outside the scope the current set of imported namespaces within this editor context. Default is FALSE (imported). */
+	virtual bool IsNonImportedField(FFieldVariant Field) const { return false; }
+
 	UE_DEPRECATED(5.0, "GetSelectedSCSEditorTreeNodes has been deprecated. Use GetSelectedSubobjectEditorTreeNodes instead.")
 	virtual TArray<TSharedPtr<class FSCSEditorTreeNode> >  GetSelectedSCSEditorTreeNodes() const = 0;
 	UE_DEPRECATED(5.0, "FindAndSelectSCSEditorTreeNode has been deprecated. Use FindAndSelectSubobjectEditorTreeNode instead.")
