Add a blacklist of pin types to blueprint type promotion to allow users to disallow types that they want explicit functions. Add the string type by default as it has multiple compare functions that should be a choice for the user.

#rnx

diff --git a/Engine/Source/Editor/BlueprintGraph/Private/BlueprintEditorSettings.cpp b/Engine/Source/Editor/BlueprintGraph/Private/BlueprintEditorSettings.cpp
--- a/Engine/Source/Editor/BlueprintGraph/Private/BlueprintEditorSettings.cpp
+++ b/Engine/Source/Editor/BlueprintGraph/Private/BlueprintEditorSettings.cpp
@@ -33,6 +33,7 @@ UBlueprintEditorSettings::UBlueprintEditorSettings(const FObjectInitializer& Obj
 	, bNavigateToNativeFunctionsFromCallNodes(true)
 	, bDoubleClickNavigatesToParent(true)
 	, bEnableTypePromotion(true)
+	, TypePromotionPinBlacklist { UEdGraphSchema_K2::PC_String, UEdGraphSchema_K2::PC_Text }
 	// Experimental
 	, bEnableNamespaceFilteringFeatures(false)
 	, bEnableNamespaceImportingFeatures(false)
@@ -123,9 +124,10 @@ void UBlueprintEditorSettings::PostEditChangeProperty(FPropertyChangedEvent& Pro
 	{
 		bShouldRebuildRegistry = true;
 	}
-
+	
 	// Refresh type promotion when the preference gets changed so that we can correctly rebuild the action database
-	if (PropertyName == GET_MEMBER_NAME_CHECKED(UBlueprintEditorSettings, bEnableTypePromotion))
+	if (PropertyName == GET_MEMBER_NAME_CHECKED(UBlueprintEditorSettings, bEnableTypePromotion) || 
+		PropertyName == GET_MEMBER_NAME_CHECKED(UBlueprintEditorSettings, TypePromotionPinBlacklist))
 	{
 		FTypePromotion::ClearNodeSpawners();
 
diff --git a/Engine/Source/Editor/BlueprintGraph/Private/BlueprintTypePromotion.cpp b/Engine/Source/Editor/BlueprintGraph/Private/BlueprintTypePromotion.cpp
--- a/Engine/Source/Editor/BlueprintGraph/Private/BlueprintTypePromotion.cpp
+++ b/Engine/Source/Editor/BlueprintGraph/Private/BlueprintTypePromotion.cpp
@@ -522,12 +522,42 @@ void FTypePromotion::AddOpFunction(FName OpName, UFunction* Function)
 	OperatorTable.FindOrAdd(OpName).Add(Function);
 }
 
+static bool HasBlackListedPinType(const UFunction* Function)
+{
+	const TSet<FName>& Blacklist = GetDefault<UBlueprintEditorSettings>()->TypePromotionPinBlacklist;
+
+	check(Function);
+	const UEdGraphSchema_K2* Schema = GetDefault<UEdGraphSchema_K2>();
+
+	// For each property in the func, see if the pin type is on the blacklist 
+	for (TFieldIterator<FProperty> PropIt(Function); PropIt && (PropIt->PropertyFlags & CPF_Parm); ++PropIt)
+	{
+		const FProperty* Param = *PropIt;
+		FEdGraphPinType ParamType;
+		if (Schema->ConvertPropertyToPinType(Param, /* out */ ParamType))
+		{
+			if (Blacklist.Contains(ParamType.PinCategory))
+			{
+				return true;
+			}
+		}
+	}
+
+	return false;
+}
+
 bool FTypePromotion::IsPromotableFunction(const UFunction* Function)
 {
+	TRACE_CPUPROFILER_EVENT_SCOPE(FTypePromotion::IsPromotableFunction);
+
 	// Ensure that we don't have an invalid OpName as well for extra safety when this function 
 	// is called outside of this class, not during the OpTable creation process
 	FName OpName = GetOpNameFromFunction(Function);
-	return Function && Function->HasAnyFunctionFlags(FUNC_BlueprintPure) && Function->GetReturnProperty() && OpName != OperatorNames::NoOp;
+	return Function &&
+		Function->HasAnyFunctionFlags(FUNC_BlueprintPure) &&
+		Function->GetReturnProperty() &&
+		OpName != OperatorNames::NoOp && 
+		!HasBlackListedPinType(Function);
 }
 
 bool FTypePromotion::IsOperatorSpawnerRegistered(UFunction const* const Func)
diff --git a/Engine/Source/Editor/BlueprintGraph/Public/BlueprintEditorSettings.h b/Engine/Source/Editor/BlueprintGraph/Public/BlueprintEditorSettings.h
--- a/Engine/Source/Editor/BlueprintGraph/Public/BlueprintEditorSettings.h
+++ b/Engine/Source/Editor/BlueprintGraph/Public/BlueprintEditorSettings.h
@@ -130,6 +130,10 @@ public:
 	UPROPERTY(config, EditAnywhere, Category = Workflow)
 	bool bEnableTypePromotion;
 
+	/** If a pin type is within this blacklist, then it will never be marked as a possible promotable function. */
+	UPROPERTY(config, EditAnywhere, Category = Workflow)
+	TSet<FName> TypePromotionPinBlacklist;
+
 // Experimental
 public:
 	/** Whether to enable namespace filtering features in the Blueprint editor */
