Blueprints: Show the class tooltip from the underlying subsystem for Get Subsystem nodes

diff --git a/Engine/Source/Editor/BlueprintGraph/Private/K2Node_GetSubsystem.cpp b/Engine/Source/Editor/BlueprintGraph/Private/K2Node_GetSubsystem.cpp
--- a/Engine/Source/Editor/BlueprintGraph/Private/K2Node_GetSubsystem.cpp
+++ b/Engine/Source/Editor/BlueprintGraph/Private/K2Node_GetSubsystem.cpp
@@ -292,7 +292,10 @@ FText UK2Node_GetSubsystem::GetTooltipText() const
 		{
 			SubsystemTypeText = NSLOCTEXT("K2Node", "GetSubsystem_LocalPlayerSubsystemTooltip", "LocalPlayer Subsystem");
 		}
-		return FText::FormatNamed(NSLOCTEXT("K2Node", "GetSubsystem_TooltipFormat", "Get {ClassName} a {SubsystemType}"), TEXT("ClassName"), CustomClass->GetDisplayNameText(), TEXT("SubsystemType"), SubsystemTypeText);
+		return FText::FormatNamed(NSLOCTEXT("K2Node", "GetSubsystem_TooltipFormat", "Get {ClassName} ({SubsystemType})\n\n{ClassTooltip}"),
+			TEXT("ClassName"), CustomClass->GetDisplayNameText(),
+			TEXT("SubsystemType"), SubsystemTypeText,
+			TEXT("ClassTooltip"), CustomClass->GetToolTipText(/*bShortTooltip=*/ true));
 	}
 
 	return NSLOCTEXT("K2Node", "GetSubsystem_InvalidSubsystemTypeTooltip", "Invalid Subsystem Type");
