Blueprints: Attempt to show tooltips for delegate return values from async task nodes (by harvesting [REDACTED_HANDLE] style markup from the comment on the delegate type declaration, e.g., DECLARE_DYNAMIC_MULTICAST_DELEGATE)

diff --git a/Engine/Source/Editor/BlueprintGraph/Private/K2Node_BaseAsyncTask.cpp b/Engine/Source/Editor/BlueprintGraph/Private/K2Node_BaseAsyncTask.cpp
--- a/Engine/Source/Editor/BlueprintGraph/Private/K2Node_BaseAsyncTask.cpp
+++ b/Engine/Source/Editor/BlueprintGraph/Private/K2Node_BaseAsyncTask.cpp
@@ -1,6 +1,5 @@
 // Copyright Epic Games, Inc. All Rights Reserved.
 
-
 #include "K2Node_BaseAsyncTask.h"
 #include "EdGraph/EdGraphPin.h"
 #include "UObject/UnrealType.h"
@@ -143,7 +142,7 @@ void UK2Node_BaseAsyncTask::AllocateDefaultPins()
 				UEdGraphPin* Pin = CreatePin(EGPD_Output, NAME_None, Param->GetFName());
 				K2Schema->ConvertPropertyToPinType(Param, /*out*/ Pin->PinType);
 
-				Pin->PinToolTip = Param->GetToolTipText().ToString();
+				UK2Node_CallFunction::GeneratePinTooltipFromFunction(*Pin, DelegateSignatureFunction);
 			}
 		}
 	}
