Adjust weights a bit to prioritize starts of words more

diff --git a/Engine/Source/Editor/BlueprintGraph/Private/EdGraphSchema_K2.cpp b/Engine/Source/Editor/BlueprintGraph/Private/EdGraphSchema_K2.cpp
--- a/Engine/Source/Editor/BlueprintGraph/Private/EdGraphSchema_K2.cpp
+++ b/Engine/Source/Editor/BlueprintGraph/Private/EdGraphSchema_K2.cpp
@@ -7468,7 +7468,7 @@ namespace ContextMenuConsoleVariables
 		ECVF_Default);
 
 	/** Weight that a match to a category search has */
-	static float CategoryWeight = 5.0f;
+	static float CategoryWeight = 4.0f;
 	static FAutoConsoleVariableRef CVarCategoryWeight(
 		TEXT("ContextMenu.CategoryWeight"), CategoryWeight,
 		TEXT("The amount of weight placed on categories that match what the user has typed in"),
@@ -7489,13 +7489,13 @@ namespace ContextMenuConsoleVariables
 		ECVF_Default);
 
 	/** The multiplier given if the keyword starts with a term the user typed in */
-	static float StartsWithBonusWeightMultiplier = 2.0f;
+	static float StartsWithBonusWeightMultiplier = 4.0f;
 	static FAutoConsoleVariableRef CVarStartsWithBonusWeightMultiplier(
 		TEXT("ContextMenu.StartsWithBonusWeightMultiplier"), StartsWithBonusWeightMultiplier,
 		TEXT("The multiplier given if the keyword starts with a term the user typed in"),
 		ECVF_Default);
 
-	/** The multiplier given if the keyword starts with a term the user typed in */
+	/** The multiplier given if the keyword contains a term the user typed in */
 	static float WordContainsLetterWeightMultiplier = 0.5f;
 	static FAutoConsoleVariableRef CVarWordContainsLetterWeightMultiplier(
 		TEXT("ContextMenu.WordContainsLetterWeightMultiplier"), WordContainsLetterWeightMultiplier,
diff --git a/Engine/Source/Editor/BlueprintGraph/Private/K2Node_CallFunction.cpp b/Engine/Source/Editor/BlueprintGraph/Private/K2Node_CallFunction.cpp
--- a/Engine/Source/Editor/BlueprintGraph/Private/K2Node_CallFunction.cpp
+++ b/Engine/Source/Editor/BlueprintGraph/Private/K2Node_CallFunction.cpp
@@ -1740,12 +1740,8 @@ FText UK2Node_CallFunction::GetDefaultCategoryForFunction(const UFunction* Funct
 
 FText UK2Node_CallFunction::GetKeywordsForFunction(const UFunction* Function)
 {
-	// If the friendly name and real function name do not match add the real function name friendly name as a keyword.
-	FString Keywords;
-	if( Function->GetName() != GetUserFacingFunctionName(Function).ToString() )
-	{
-		Keywords = Function->GetName();
-	}
+	// Always add the real function name as the first keyword, even if it matches the display name we don't want to penalize one word function names in later searches
+	FString Keywords = Function->GetName();
 
 	if (ShouldDrawCompact(Function))
 	{
