Refresh the type promotion table after modifying the blacklist

#rnx

diff --git a/Engine/Source/Editor/BlueprintGraph/Classes/BlueprintTypePromotion.h b/Engine/Source/Editor/BlueprintGraph/Classes/BlueprintTypePromotion.h
--- a/Engine/Source/Editor/BlueprintGraph/Classes/BlueprintTypePromotion.h
+++ b/Engine/Source/Editor/BlueprintGraph/Classes/BlueprintTypePromotion.h
@@ -107,6 +107,11 @@ public:
 	*/
 	static void ClearNodeSpawners();
 
+	/**
+	 * Clear out the promotion table and rebuild it, reassessing each available UFunction. 
+	 */
+	static void RefreshPromotionTables();
+
 	/**
 	* Get the "Primitive Promotion Table" which represents what base Pin types 
 	* can be promoted to others. 
diff --git a/Engine/Source/Editor/BlueprintGraph/Private/BlueprintEditorSettings.cpp b/Engine/Source/Editor/BlueprintGraph/Private/BlueprintEditorSettings.cpp
--- a/Engine/Source/Editor/BlueprintGraph/Private/BlueprintEditorSettings.cpp
+++ b/Engine/Source/Editor/BlueprintGraph/Private/BlueprintEditorSettings.cpp
@@ -130,7 +130,8 @@ void UBlueprintEditorSettings::PostEditChangeProperty(FPropertyChangedEvent& Pro
 		PropertyName == GET_MEMBER_NAME_CHECKED(UBlueprintEditorSettings, TypePromotionPinBlacklist))
 	{
 		FTypePromotion::ClearNodeSpawners();
-
+		FTypePromotion::RefreshPromotionTables();
+		
 		bShouldRebuildRegistry = true;
 	}
 
diff --git a/Engine/Source/Editor/BlueprintGraph/Private/BlueprintTypePromotion.cpp b/Engine/Source/Editor/BlueprintGraph/Private/BlueprintTypePromotion.cpp
--- a/Engine/Source/Editor/BlueprintGraph/Private/BlueprintTypePromotion.cpp
+++ b/Engine/Source/Editor/BlueprintGraph/Private/BlueprintTypePromotion.cpp
@@ -591,4 +591,13 @@ void FTypePromotion::ClearNodeSpawners()
 	}
 }
 
-#undef LOCTEXT_NAMESPACE
\ No newline at end of file
+void FTypePromotion::RefreshPromotionTables()
+{
+	if(Instance)
+	{
+		Instance->OperatorTable.Empty();
+		Instance->CreateOpTable();
+	}
+}
+
+#undef LOCTEXT_NAMESPACE
