Add an edit condition to the blueprint editor settings to only allow changing type promotion black list if it is enabled.

#rnx

diff --git a/Engine/Source/Editor/BlueprintGraph/Public/BlueprintEditorSettings.h b/Engine/Source/Editor/BlueprintGraph/Public/BlueprintEditorSettings.h
--- a/Engine/Source/Editor/BlueprintGraph/Public/BlueprintEditorSettings.h
+++ b/Engine/Source/Editor/BlueprintGraph/Public/BlueprintEditorSettings.h
@@ -131,7 +131,7 @@ public:
 	bool bEnableTypePromotion;
 
 	/** If a pin type is within this blacklist, then it will never be marked as a possible promotable function. */
-	UPROPERTY(config, EditAnywhere, Category = Workflow)
+	UPROPERTY(config, EditAnywhere, Category = Workflow, meta=(EditCondition="bEnableTypePromotion"))
 	TSet<FName> TypePromotionPinBlacklist;
 
 // Experimental
