Fix for a regression that incorrectly restricts the "Outer" input pin on ConstructObjectFromClass nodes to self class subtypes.

diff --git a/Engine/Source/Editor/BlueprintGraph/Private/K2Node_ConstructObjectFromClass.cpp b/Engine/Source/Editor/BlueprintGraph/Private/K2Node_ConstructObjectFromClass.cpp
--- a/Engine/Source/Editor/BlueprintGraph/Private/K2Node_ConstructObjectFromClass.cpp
+++ b/Engine/Source/Editor/BlueprintGraph/Private/K2Node_ConstructObjectFromClass.cpp
@@ -62,7 +62,7 @@ void UK2Node_ConstructObjectFromClass::AllocateDefaultPins()
 	
 	if (UseOuter())
 	{
-		UEdGraphPin* OuterPin = CreatePin(EGPD_Input, UEdGraphSchema_K2::PC_Object, UEdGraphSchema_K2::PSC_Self, nullptr, FK2Node_ConstructObjectFromClassHelper::OuterPinName);
+		UEdGraphPin* OuterPin = CreatePin(EGPD_Input, UEdGraphSchema_K2::PC_Object, UObject::StaticClass(), UEdGraphSchema_K2::PN_Self);
 		OuterPin->PinFriendlyName = FText::FromName(FK2Node_ConstructObjectFromClassHelper::OuterPinFriendlyName);
 	}
 
