Allow pasted variable get nodes in function scope to look for parameters

diff --git a/Engine/Source/Editor/BlueprintGraph/Classes/K2Node_Variable.h b/Engine/Source/Editor/BlueprintGraph/Classes/K2Node_Variable.h
--- a/Engine/Source/Editor/BlueprintGraph/Classes/K2Node_Variable.h
+++ b/Engine/Source/Editor/BlueprintGraph/Classes/K2Node_Variable.h
@@ -187,5 +187,8 @@ protected:
 	 * @return 
 	 */
 	FBPVariableDescription const* GetBlueprintVarDescription() const;
+
+	/** Returns whether a Function Graph contains a parameter with the given name */
+	static bool FunctionParameterExists(const UEdGraph* InFunctionGraph, const FName InParameterName);
 };
 
diff --git a/Engine/Source/Editor/BlueprintGraph/Private/K2Node_Variable.cpp b/Engine/Source/Editor/BlueprintGraph/Private/K2Node_Variable.cpp
--- a/Engine/Source/Editor/BlueprintGraph/Private/K2Node_Variable.cpp
+++ b/Engine/Source/Editor/BlueprintGraph/Private/K2Node_Variable.cpp
@@ -2,6 +2,7 @@
 
 
 #include "K2Node_Variable.h"
+#include "K2Node_FunctionEntry.h"
 #include "BlueprintCompilationManager.h"
 #include "UObject/UObjectHash.h"
 #include "Components/PrimitiveComponent.h"
@@ -1011,7 +1012,8 @@ void UK2Node_Variable::PostPasteNode()
 		{
 			UEdGraph* FunctionGraph = FBlueprintEditorUtils::GetTopLevelGraph(GetGraph());
 			FBPVariableDescription* VariableDescription = FBlueprintEditorUtils::FindLocalVariable(Blueprint, FunctionGraph, VariableReference.GetMemberName());
-			if(VariableDescription)
+			bool bFoundParam = FunctionParameterExists(FunctionGraph, VariableReference.GetMemberName());
+			if(VariableDescription || bFoundParam)
 			{
 				VariableReference.SetLocalMember(VariableReference.GetMemberName(), FunctionGraph->GetName(), VariableReference.GetMemberGuid());
 			}
@@ -1099,4 +1101,19 @@ void UK2Node_Variable::JumpToDefinition() const
 	Super::JumpToDefinition();
 }
 
+bool UK2Node_Variable::FunctionParameterExists(const UEdGraph* InFunctionGraph, const FName InParameterName)
+{
+	TArray<UK2Node_FunctionEntry*> Entry;
+	InFunctionGraph->GetNodesOfClass<UK2Node_FunctionEntry>(Entry);
+
+	if (ensureMsgf(Entry.Num() == 1, TEXT("Couldn't find a Function Entry node in graph %s"), *InFunctionGraph->GetName()))
+	{
+		check(Entry[0]);
+
+		return Entry[0]->UserDefinedPinExists(InParameterName);
+	}
+
+	return false;
+}
+
 #undef LOCTEXT_NAMESPACE
