Update connection message to clarify that a type is blacklisted from type promotion

#rnx

diff --git a/Engine/Source/Editor/BlueprintGraph/Private/K2Node_PromotableOperator.cpp b/Engine/Source/Editor/BlueprintGraph/Private/K2Node_PromotableOperator.cpp
--- a/Engine/Source/Editor/BlueprintGraph/Private/K2Node_PromotableOperator.cpp
+++ b/Engine/Source/Editor/BlueprintGraph/Private/K2Node_PromotableOperator.cpp
@@ -568,7 +568,16 @@ bool UK2Node_PromotableOperator::IsConnectionDisallowed(const UEdGraphPin* MyPin
 			Args.Add(TEXT("OtherPinType"), K2Schema->TypeToText(OtherPin->PinType));
 			Args.Add(TEXT("OpName"), FText::FromName(OperationName));
 
-			OutReason = FText::Format(LOCTEXT("NoCompatibleStructConv", "No matching '{OpName}' function for '{OtherPinType}'"), Args).ToString();
+			const TSet<FName>& Blacklist = GetDefault<UBlueprintEditorSettings>()->TypePromotionPinBlacklist;
+			if (Blacklist.Contains(OtherPin->PinType.PinCategory))
+			{
+				OutReason = FText::Format(LOCTEXT("NoCompatibleStructConv_Blacklisted", "No matching '{OpName}' function for '{OtherPinType}'. This type is blacklisted from operators in the blueprint editor settings."), Args).ToString();
+			}
+			else
+			{
+				OutReason = FText::Format(LOCTEXT("NoCompatibleStructConv", "No matching '{OpName}' function for '{OtherPinType}'"), Args).ToString();
+			}
+
 			return true;
 		}
 	}
