PR #8438: UHT: Fix // & /* not parsing in a quoted string (Contributed by [REDACTED_PERSON])

diff --git a/Engine/Source/Editor/BlueprintGraph/Private/BasicTokenParser.cpp b/Engine/Source/Editor/BlueprintGraph/Private/BasicTokenParser.cpp
--- a/Engine/Source/Editor/BlueprintGraph/Private/BasicTokenParser.cpp
+++ b/Engine/Source/Editor/BlueprintGraph/Private/BasicTokenParser.cpp
@@ -628,7 +628,7 @@ bool FBasicTokenParser::GetRawTokenRespectingQuotes(FBasicToken& Token, TCHAR St
 
 	while( !BasicTokenParserImpl::IsEOL(c) && ((c != StopChar) || bInQuote) )
 	{
-		if( (c=='/' && PeekChar()=='/') || (c=='/' && PeekChar()=='*') )
+		if( !bInQuote && ( (c=='/' && PeekChar()=='/') || (c=='/' && PeekChar()=='*') ) )
 		{
 			break;
 		}
diff --git a/Engine/Source/Programs/UnrealHeaderTool/Private/BaseParser.cpp b/Engine/Source/Programs/UnrealHeaderTool/Private/BaseParser.cpp
--- a/Engine/Source/Programs/UnrealHeaderTool/Private/BaseParser.cpp
+++ b/Engine/Source/Programs/UnrealHeaderTool/Private/BaseParser.cpp
@@ -631,7 +631,7 @@ bool FBaseParser::GetRawStringRespectingQuotes(FTokenString& String, TCHAR StopC
 
 	while( !IsEOL(c) && ((c != StopChar) || bInQuote) )
 	{
-		if( (c=='/' && PeekChar()=='/') || (c=='/' && PeekChar()=='*') )
+		if( !bInQuote && ( (c=='/' && PeekChar()=='/') || (c=='/' && PeekChar()=='*') ) )
 		{
 			break;
 		}
