Include non-native struct/class variable types in external node dependency outputs for variable get/set nodes.

May indirectly resolve other one-off script serialization crash reports we've seen during post-edit bytecode fixup (e.g. UE-73413).

diff --git a/Engine/Source/Editor/BlueprintGraph/Private/K2Node_Variable.cpp b/Engine/Source/Editor/BlueprintGraph/Private/K2Node_Variable.cpp
--- a/Engine/Source/Editor/BlueprintGraph/Private/K2Node_Variable.cpp
+++ b/Engine/Source/Editor/BlueprintGraph/Private/K2Node_Variable.cpp
@@ -831,11 +831,30 @@ bool UK2Node_Variable::HasExternalDependencies(TArray<class UStruct*>* OptionalO
 {
 	UClass* SourceClass = GetVariableSourceClass();
 	UBlueprint* SourceBlueprint = GetBlueprint();
-	const bool bResult = (SourceClass && (SourceClass->ClassGeneratedBy != SourceBlueprint));
+	bool bResult = (SourceClass && (SourceClass->ClassGeneratedBy != SourceBlueprint));
 	if (bResult && OptionalOutput)
 	{
 		OptionalOutput->AddUnique(SourceClass);
 	}
+
+	// Also include underlying non-native variable types as external dependencies. Otherwise, contextual
+	// type references serialized to bytecode can potentially be invalidated when the type is regenerated.
+	if (const UEdGraphPin* VarPin = FindPin(GetVarName()))
+	{
+		if (UStruct* PinTypeStruct = Cast<UStruct>(VarPin->PinType.PinSubCategoryObject.Get()))
+		{
+			UClass* PinTypeClass = Cast<UClass>(PinTypeStruct);
+			if (!PinTypeStruct->IsNative() && (!PinTypeClass || PinTypeClass->ClassGeneratedBy != SourceBlueprint))
+			{
+				bResult = true;
+				if (OptionalOutput)
+				{
+					OptionalOutput->AddUnique(PinTypeStruct);
+				}
+			}
+		}
+	}
+	
 	const bool bSuperResult = Super::HasExternalDependencies(OptionalOutput);
 	return bSuperResult || bResult;
 }
