Fixes for dtor typos that attempted to unregister handlers from the wrong delegate interface on editor shutdown.

diff --git a/Engine/Source/Editor/BlueprintGraph/Private/BlueprintActionDatabase.cpp b/Engine/Source/Editor/BlueprintGraph/Private/BlueprintActionDatabase.cpp
--- a/Engine/Source/Editor/BlueprintGraph/Private/BlueprintActionDatabase.cpp
+++ b/Engine/Source/Editor/BlueprintGraph/Private/BlueprintActionDatabase.cpp
@@ -1179,8 +1179,8 @@ FBlueprintActionDatabase::~FBlueprintActionDatabase()
 	{
 		IAssetRegistry& AssetRegistry = FModuleManager::LoadModuleChecked<FAssetRegistryModule>(TEXT("AssetRegistry")).Get();
 		AssetRegistry.OnAssetAdded().Remove(OnAssetAddedDelegateHandle);
-		AssetRegistry.OnAssetAdded().Remove(OnAssetRemovedDelegateHandle);
-		AssetRegistry.OnAssetAdded().Remove(OnAssetRenamedDelegateHandle);
+		AssetRegistry.OnAssetRemoved().Remove(OnAssetRemovedDelegateHandle);
+		AssetRegistry.OnAssetRenamed().Remove(OnAssetRenamedDelegateHandle);
 	}
 
 	FEditorDelegates::OnAssetsPreDelete.Remove(OnAssetsPreDeleteDelegateHandle);
@@ -1189,7 +1189,7 @@ FBlueprintActionDatabase::~FBlueprintActionDatabase()
 	if (GEngine)
 	{
 		GEngine->OnWorldAdded().Remove(OnWorldAddedDelegateHandle);
-		GEngine->OnWorldAdded().Remove(OnWorldDestroyedDelegateHandle);
+		GEngine->OnWorldDestroyed().Remove(OnWorldDestroyedDelegateHandle);
 	}
 
 	FWorldDelegates::RefreshLevelScriptActions.Remove(RefreshLevelScriptActionsDelegateHandle);
diff --git a/Engine/Source/Editor/Kismet/Private/BlueprintNamespaceRegistry.cpp b/Engine/Source/Editor/Kismet/Private/BlueprintNamespaceRegistry.cpp
--- a/Engine/Source/Editor/Kismet/Private/BlueprintNamespaceRegistry.cpp
+++ b/Engine/Source/Editor/Kismet/Private/BlueprintNamespaceRegistry.cpp
@@ -78,8 +78,8 @@ void FBlueprintNamespaceRegistry::Shutdown()
 	{
 		IAssetRegistry& AssetRegistry = FModuleManager::LoadModuleChecked<FAssetRegistryModule>(TEXT("AssetRegistry")).Get();
 		AssetRegistry.OnAssetAdded().Remove(OnAssetAddedDelegateHandle);
-		AssetRegistry.OnAssetAdded().Remove(OnAssetRemovedDelegateHandle);
-		AssetRegistry.OnAssetAdded().Remove(OnAssetRenamedDelegateHandle);
+		AssetRegistry.OnAssetRemoved().Remove(OnAssetRemovedDelegateHandle);
+		AssetRegistry.OnAssetRenamed().Remove(OnAssetRenamedDelegateHandle);
 	}
 
 	bIsInitialized = false;
