Fix GetClassDefaults Node discarding ShowPins settings during compilation if connected to a reroute node

diff --git a/Engine/Source/Editor/BlueprintGraph/Private/K2Node_GetClassDefaults.cpp b/Engine/Source/Editor/BlueprintGraph/Private/K2Node_GetClassDefaults.cpp
--- a/Engine/Source/Editor/BlueprintGraph/Private/K2Node_GetClassDefaults.cpp
+++ b/Engine/Source/Editor/BlueprintGraph/Private/K2Node_GetClassDefaults.cpp
@@ -576,8 +576,14 @@ void UK2Node_GetClassDefaults::OnClassPinChanged()
 		}
 	}
 
-	// Clear the current output pin settings (so they don't carry over to the new set)
-	ShowPinForProperties.Reset();
+	// If the blueprint is currently compiling, we shouldn't be changing the showpins settings
+	const bool bIsCompiling = GetBlueprint()->bBeingCompiled;
+
+	if (!bIsCompiling)
+	{
+		// Clear the current output pin settings (so they don't carry over to the new set)
+		ShowPinForProperties.Reset();
+	}
 
 	// Create output pins for the new class type
 	UClass* InputClass = GetInputClass();
