Add metadata to hide the asset picker on Object pins

#rnx

diff --git a/Engine/Source/Editor/BlueprintGraph/Classes/EdGraphSchema_K2.h b/Engine/Source/Editor/BlueprintGraph/Classes/EdGraphSchema_K2.h
--- a/Engine/Source/Editor/BlueprintGraph/Classes/EdGraphSchema_K2.h
+++ b/Engine/Source/Editor/BlueprintGraph/Classes/EdGraphSchema_K2.h
@@ -167,6 +167,9 @@ public:
 	/** If true, an unconnected pin will generate a UPROPERTY under the hood to connect as the input, which will be set to the literal value for the pin.  Only valid for reference parameters. */
 	static const FName MD_AutoCreateRefTerm;
 
+	/** The specified parameter should hide the asset picker on the pin, even if it is a valid UObject Asset. */
+	static const FName MD_HideAssetPicker;
+
 	/** If true, the hidden world context pin will be visible when the function is placed in a child blueprint of the class. */
 	static const FName MD_ShowWorldContextPin;
 
diff --git a/Engine/Source/Editor/BlueprintGraph/Private/EdGraphSchema_K2.cpp b/Engine/Source/Editor/BlueprintGraph/Private/EdGraphSchema_K2.cpp
--- a/Engine/Source/Editor/BlueprintGraph/Private/EdGraphSchema_K2.cpp
+++ b/Engine/Source/Editor/BlueprintGraph/Private/EdGraphSchema_K2.cpp
@@ -141,6 +141,7 @@ const FName FBlueprintMetadata::MD_WorldContext(TEXT("WorldContext"));
 const FName FBlueprintMetadata::MD_CallableWithoutWorldContext(TEXT("CallableWithoutWorldContext"));
 const FName FBlueprintMetadata::MD_DevelopmentOnly(TEXT("DevelopmentOnly"));
 const FName FBlueprintMetadata::MD_AutoCreateRefTerm(TEXT("AutoCreateRefTerm"));
+const FName FBlueprintMetadata::MD_HideAssetPicker(TEXT("HideAssetPicker"));
 
 const FName FBlueprintMetadata::MD_ShowWorldContextPin(TEXT("ShowWorldContextPin"));
 const FName FBlueprintMetadata::MD_Private(TEXT("BlueprintPrivate"));
@@ -4892,10 +4893,28 @@ bool UEdGraphSchema_K2::ShouldShowAssetPickerForPin(UEdGraphPin* Pin) const
 			{
 				if (UK2Node_CallFunction* CallFunctionNode = Cast<UK2Node_CallFunction>(Pin->GetOwningNode()))
 				{
-					if ( UFunction* FunctionRef = CallFunctionNode->GetTargetFunction() )
+					if (UFunction* FunctionRef = CallFunctionNode->GetTargetFunction())
 					{
 						const UEdGraphPin* WorldContextPin = CallFunctionNode->FindPin(FunctionRef->GetMetaData(FBlueprintMetadata::MD_WorldContext));
 						bShow = ( WorldContextPin != Pin );
+
+						// Check if we have explictly marked this pin as hiding the asset picker
+						const FString& HideAssetPickerMetaData = FunctionRef->GetMetaData(FBlueprintMetadata::MD_HideAssetPicker);
+						if(!HideAssetPickerMetaData.IsEmpty())
+						{
+							TArray<FString> PinNames;
+							HideAssetPickerMetaData.ParseIntoArray(PinNames, TEXT(","), true);
+							const FString PinName = Pin->GetName();
+							for(FString& ParamNameToHide : PinNames)
+							{
+								ParamNameToHide.TrimStartAndEndInline();
+								if(ParamNameToHide == PinName)
+								{
+									bShow = false;
+									break;
+								}
+							}
+						}
 					}
 				}
 				else if (Cast<UK2Node_CreateDelegate>( Pin->GetOwningNode())) 
diff --git a/Engine/Source/Runtime/CoreUObject/Public/UObject/ObjectMacros.h b/Engine/Source/Runtime/CoreUObject/Public/UObject/ObjectMacros.h
--- a/Engine/Source/Runtime/CoreUObject/Public/UObject/ObjectMacros.h
+++ b/Engine/Source/Runtime/CoreUObject/Public/UObject/ObjectMacros.h
@@ -1428,6 +1428,9 @@ namespace UM
 		/// This also allows for inline editing of the default value on some types (take FRotator for instance). Only valid for inputs.
 		AutoCreateRefTerm,
 
+		/// [FunctionMetadata] Forcibly hide the asset picker for pins matching any parameter names in this piece of metadata.
+		HideAssetPicker,
+
 		/// [FunctionMetadata] This function is an internal implementation detail, used to implement another function or node.  It is never directly exposed in a graph.
 		BlueprintInternalUseOnly,
 
