Disable Pin Value Inspection Tooltips for Pose Pins in AnimBPs

diff --git a/Engine/Source/Editor/AnimGraph/Private/AnimationGraphSchema.cpp b/Engine/Source/Editor/AnimGraph/Private/AnimationGraphSchema.cpp
--- a/Engine/Source/Editor/AnimGraph/Private/AnimationGraphSchema.cpp
+++ b/Engine/Source/Editor/AnimGraph/Private/AnimationGraphSchema.cpp
@@ -439,6 +439,11 @@ void UAnimationGraphSchema::CreateFunctionGraphTerminators(UEdGraph& Graph, UCla
 }
 
 
+bool UAnimationGraphSchema::CanShowDataTooltipForPin(const UEdGraphPin& Pin) const
+{
+	return !IsPosePin(Pin.PinType) && UEdGraphSchema_K2::CanShowDataTooltipForPin(Pin);
+}
+
 bool UAnimationGraphSchema::CanGraphBeDropped(TSharedPtr<FEdGraphSchemaAction> InAction) const
 {
 	if (!InAction.IsValid())
diff --git a/Engine/Source/Editor/AnimGraph/Public/AnimationGraphSchema.h b/Engine/Source/Editor/AnimGraph/Public/AnimationGraphSchema.h
--- a/Engine/Source/Editor/AnimGraph/Public/AnimationGraphSchema.h
+++ b/Engine/Source/Editor/AnimGraph/Public/AnimationGraphSchema.h
@@ -75,7 +75,7 @@ class UAnimationGraphSchema : public UEdGraphSchema_K2
 	virtual bool ArePinsCompatible(const UEdGraphPin* PinA, const UEdGraphPin* PinB, const UClass* CallingContext = NULL, bool bIgnoreArray = false) const override;
 	virtual bool DoesSupportAnimNotifyActions() const override;
 	virtual void CreateFunctionGraphTerminators(UEdGraph& Graph, UClass* Class) const override;
-
+	virtual bool CanShowDataTooltipForPin(const UEdGraphPin& Pin) const override;
 	//~ End UEdGraphSchema_K2 Interface
 
 	/** Spawn the correct node in the Animation Graph using the given AnimationAsset at the supplied location */
diff --git a/Engine/Source/Editor/BlueprintGraph/Classes/EdGraphSchema_K2.h b/Engine/Source/Editor/BlueprintGraph/Classes/EdGraphSchema_K2.h
--- a/Engine/Source/Editor/BlueprintGraph/Classes/EdGraphSchema_K2.h
+++ b/Engine/Source/Editor/BlueprintGraph/Classes/EdGraphSchema_K2.h
@@ -552,6 +552,7 @@ public:
 	virtual void ForceVisualizationCacheClear() const override;
 	virtual bool SafeDeleteNodeFromGraph(UEdGraph* Graph, UEdGraphNode* NodeToDelete) const override;
 	virtual bool CanVariableBeDropped(UEdGraph* InGraph, FProperty* InVariableToDrop) const override { return true; }
+	virtual bool CanShowDataTooltipForPin(const UEdGraphPin& Pin) const override;
 
 #if WITH_EDITORONLY_DATA
 	virtual float GetActionFilteredWeight(const FGraphActionListBuilderBase::ActionGroup& InCurrentAction, const TArray<FString>& InFilterTerms, const TArray<FString>& InSanitizedFilterTerms, const TArray<UEdGraphPin*>& DraggedFromPins) const override;
diff --git a/Engine/Source/Editor/BlueprintGraph/Private/EdGraphSchema_K2.cpp b/Engine/Source/Editor/BlueprintGraph/Private/EdGraphSchema_K2.cpp
--- a/Engine/Source/Editor/BlueprintGraph/Private/EdGraphSchema_K2.cpp
+++ b/Engine/Source/Editor/BlueprintGraph/Private/EdGraphSchema_K2.cpp
@@ -5488,6 +5488,11 @@ bool UEdGraphSchema_K2::IsSelfPin(const UEdGraphPin& Pin) const
 	return (Pin.PinName == PN_Self);
 }
 
+bool UEdGraphSchema_K2::CanShowDataTooltipForPin(const UEdGraphPin& Pin) const
+{
+	return !IsExecPin(Pin) && !IsDelegateCategory(Pin.PinType.PinCategory);
+}
+
 bool UEdGraphSchema_K2::IsDelegateCategory(const FName Category) const
 {
 	return (Category == PC_Delegate);
diff --git a/Engine/Source/Editor/UnrealEd/Private/Kismet2/KismetDebugUtilities.cpp b/Engine/Source/Editor/UnrealEd/Private/Kismet2/KismetDebugUtilities.cpp
--- a/Engine/Source/Editor/UnrealEd/Private/Kismet2/KismetDebugUtilities.cpp
+++ b/Engine/Source/Editor/UnrealEd/Private/Kismet2/KismetDebugUtilities.cpp
@@ -1634,7 +1634,7 @@ bool FKismetDebugUtilities::CanInspectPinValue(const UEdGraphPin* Pin)
 	// Can't inspect exec pins or delegate pins; their values are not defined.
 	// Disallow non-K2 Schemas (like ControlRig)
 	const UEdGraphSchema_K2* K2Schema = Cast<UEdGraphSchema_K2>(OwningNode->GetSchema());
-	if (!K2Schema || K2Schema->IsExecPin(*Pin) || K2Schema->IsDelegateCategory(Pin->PinType.PinCategory))
+	if (!K2Schema || !K2Schema->CanShowDataTooltipForPin(*Pin))
 	{
 		return false;
 	}
diff --git a/Engine/Source/Runtime/Engine/Classes/EdGraph/EdGraphSchema.h b/Engine/Source/Runtime/Engine/Classes/EdGraph/EdGraphSchema.h
--- a/Engine/Source/Runtime/Engine/Classes/EdGraph/EdGraphSchema.h
+++ b/Engine/Source/Runtime/Engine/Classes/EdGraph/EdGraphSchema.h
@@ -776,6 +776,14 @@ class ENGINE_API UEdGraphSchema : public UObject
 	 */
 	virtual void ClearPinWatch(UEdGraphPin const* Pin) const {}
 
+	/**
+	 * Checks to see if a pin supports Pin Value Inspection Tooltips
+	 *
+	 * @param	Pin The pin to check
+	 * @return	true if it supports data tooltips
+	 */
+	virtual bool CanShowDataTooltipForPin(const UEdGraphPin& Pin) const { return false; }
+
 	/**
 	 * Sets the string to the specified pin; even if it is invalid it is still set.
 	 *
