Add support for Input Axis and Action nodes in any Uobject based blueprint, not just components and actors. There is no reason why UObject based blueprints cannot be supported anymore after 18396294

#rnx

diff --git a/Engine/Source/Editor/BlueprintGraph/Private/K2Node_InputAxisEvent.cpp b/Engine/Source/Editor/BlueprintGraph/Private/K2Node_InputAxisEvent.cpp
--- a/Engine/Source/Editor/BlueprintGraph/Private/K2Node_InputAxisEvent.cpp
+++ b/Engine/Source/Editor/BlueprintGraph/Private/K2Node_InputAxisEvent.cpp
@@ -123,7 +123,10 @@ bool UK2Node_InputAxisEvent::IsCompatibleWithGraph(const UEdGraph* TargetGraph)
 	UBlueprint* Blueprint = FBlueprintEditorUtils::FindBlueprintForGraph(TargetGraph);
 	if (Blueprint && Blueprint->SkeletonGeneratedClass)
 	{
-		bIsCompatible = Blueprint->ParentClass->IsChildOf(AActor::StaticClass());
+		UEdGraphSchema_K2 const* K2Schema = Cast<UEdGraphSchema_K2>(TargetGraph->GetSchema());
+		bool const bIsConstructionScript = (K2Schema != nullptr) ? UEdGraphSchema_K2::IsConstructionScript(TargetGraph) : false;
+		
+		bIsCompatible = Blueprint && Blueprint->SupportsInputEvents() && !bIsConstructionScript;
 	}
 
 	return bIsCompatible && Super::IsCompatibleWithGraph(TargetGraph);
diff --git a/Engine/Source/Editor/BlueprintGraph/Private/K2Node_InputAxisKeyEvent.cpp b/Engine/Source/Editor/BlueprintGraph/Private/K2Node_InputAxisKeyEvent.cpp
--- a/Engine/Source/Editor/BlueprintGraph/Private/K2Node_InputAxisKeyEvent.cpp
+++ b/Engine/Source/Editor/BlueprintGraph/Private/K2Node_InputAxisKeyEvent.cpp
@@ -110,7 +110,7 @@ bool UK2Node_InputAxisKeyEvent::IsCompatibleWithGraph(const UEdGraph* TargetGrap
 	UBlueprint* Blueprint = FBlueprintEditorUtils::FindBlueprintForGraph(TargetGraph);
 	if (Blueprint != nullptr)
 	{
-		bIsCompatible = FBlueprintEditorUtils::IsActorBased(Blueprint) && Blueprint->SupportsInputEvents();
+		bIsCompatible = Blueprint && Blueprint->SupportsInputEvents();
 	}
 
 	UEdGraphSchema_K2 const* K2Schema = Cast<UEdGraphSchema_K2>(TargetGraph->GetSchema());
diff --git a/Engine/Source/Runtime/Engine/Private/Blueprint.cpp b/Engine/Source/Runtime/Engine/Private/Blueprint.cpp
--- a/Engine/Source/Runtime/Engine/Private/Blueprint.cpp
+++ b/Engine/Source/Runtime/Engine/Private/Blueprint.cpp
@@ -1021,7 +1021,7 @@ bool UBlueprint::AllowsDynamicBinding() const
 
 bool UBlueprint::SupportsInputEvents() const
 {
-	return FBlueprintEditorUtils::IsActorBased(this) || FBlueprintEditorUtils::IsComponentBased(this);
+	return ParentClass && ParentClass->IsChildOf(UObject::StaticClass());
 }
 
 struct FBlueprintInnerHelper
