Fix crash when pasting local variable from parent blueprint into child event graph

diff --git a/Engine/Source/Editor/BlueprintGraph/Private/K2Node_Variable.cpp b/Engine/Source/Editor/BlueprintGraph/Private/K2Node_Variable.cpp
--- a/Engine/Source/Editor/BlueprintGraph/Private/K2Node_Variable.cpp
+++ b/Engine/Source/Editor/BlueprintGraph/Private/K2Node_Variable.cpp
@@ -1014,8 +1014,10 @@ void UK2Node_Variable::PostPasteNode()
 	{
 		// Local scoped variables should always validate whether they are being placed in the same graph as their scope
 		// ResolveMember will not return nullptr when the graph changes but the Blueprint remains the same.
-		UEdGraph* ScopeGraph = FBlueprintEditorUtils::FindScopeGraph(Blueprint, VariableReference.GetMemberScope(GetBlueprintClassFromNode()));
-		if(ScopeGraph != GetGraph())
+		const UStruct* MemberScope = VariableReference.GetMemberScope(GetBlueprintClassFromNode());
+		UEdGraph* ScopeGraph = FBlueprintEditorUtils::FindScopeGraph(Blueprint, MemberScope);
+		const bool bMemberScopeInvalid = (ScopeGraph && MemberScope && ScopeGraph->GetFName() != MemberScope->GetFName());
+		if(ScopeGraph != GetGraph() || bMemberScopeInvalid)
 		{
 			bInvalidateVariable = true;
 		}
