Added the ability to override the list of default Blueprint categories

diff --git a/Engine/Source/Editor/BlueprintGraph/Public/BlueprintEditorSettings.h b/Engine/Source/Editor/BlueprintGraph/Public/BlueprintEditorSettings.h
--- a/Engine/Source/Editor/BlueprintGraph/Public/BlueprintEditorSettings.h
+++ b/Engine/Source/Editor/BlueprintGraph/Public/BlueprintEditorSettings.h
@@ -49,6 +49,18 @@ template<> struct TStructOpsTypeTraits<FPerBlueprintSettings> : public TStructOp
 	};
 };
 
+USTRUCT()
+struct BLUEPRINTGRAPH_API FAdditionalBlueprintCategory
+{
+	GENERATED_BODY()
+
+	UPROPERTY(EditAnywhere, Category = Category)
+	FText Name;
+
+	UPROPERTY(EditAnywhere, Category = Category)
+	FSoftClassPath ClassFilter;
+};
+
 UCLASS(config=EditorPerProjectUserSettings)
 class BLUEPRINTGRAPH_API UBlueprintEditorSettings : public UDeveloperSettings
 {
@@ -135,6 +147,10 @@ public:
 	UPROPERTY(config, EditAnywhere, Category = Workflow, meta=(EditCondition="bEnableTypePromotion"))
 	TSet<FName> TypePromotionPinDenyList;
 
+	/** List of additional category names to show in Blueprints, optionally filtered by parent class type. */
+	UPROPERTY(config, EditAnywhere, Category = Workflow)
+	TArray<FAdditionalBlueprintCategory> AdditionalBlueprintCategories;
+
 	/** How to handle previously-set breakpoints on reload. */
 	UPROPERTY(config, EditAnywhere, Category = Workflow)
 	EBlueprintBreakpointReloadMethod BreakpointReloadMethod;
diff --git a/Engine/Source/Editor/Kismet/Private/BlueprintDetailsCustomization.cpp b/Engine/Source/Editor/Kismet/Private/BlueprintDetailsCustomization.cpp
--- a/Engine/Source/Editor/Kismet/Private/BlueprintDetailsCustomization.cpp
+++ b/Engine/Source/Editor/Kismet/Private/BlueprintDetailsCustomization.cpp
@@ -1541,6 +1541,13 @@ void FBlueprintVarActionDetails::PopulateCategories(SMyBlueprint* MyBlueprint, T
 
 	CategorySource.Reset();
 	CategorySource.Add(MakeShared<FText>(UEdGraphSchema_K2::VR_DefaultCategory));
+	for (const FAdditionalBlueprintCategory& AdditionalBlueprintCategory : GetDefault<UBlueprintEditorSettings>()->AdditionalBlueprintCategories)
+	{
+		if (!AdditionalBlueprintCategory.Name.IsEmpty() && (AdditionalBlueprintCategory.ClassFilter.IsNull() || Blueprint->ParentClass->IsChildOf(AdditionalBlueprintCategory.ClassFilter.TryLoadClass<UObject>())))
+		{
+			CategorySource.Add(MakeShared<FText>(AdditionalBlueprintCategory.Name));
+		}
+	}
 	for (const FName& VariableName : VisibleVariables)
 	{
 		FText Category = FBlueprintEditorUtils::GetBlueprintVariableCategory(Blueprint, VariableName, nullptr);
