Merging //UE5/Dev-LargeWorldCoordinates [at] 18802167 to //UE5/Release-5.0

Blueprint real number support.

This change deprecates the use the of "float" and "double" types in Blueprints in favor of a new "real". By default, "real" is back by a double precision floating point number. However, it can be single precision if the number is a native float property or function parameter. This distinction won't be visible to the Blueprint user: in both instances, they'll be represented by "real" pin types. During deserialization, we'll automatically convert Blueprint pin types to use real/doubles, unless they're used to represent native code (including delegate signatures).

One consequence of this change is that we need to perform implicit casts between single and double precision real numbers. During Blueprint compilation, the compiler will detect points in the graph for when either a widening or narrowing conversion needs to occur. Subsequently, the script bytecode will contain a new cast instruction that performs the conversion. This also works on container types, but each entry in the container will have to be converted. This can introduce unwanted overhead for large containers that are frequently passed between Blueprint and native code.

The scope of this change affects Blueprints used by Gameplay, Animation, Control Rig, and UMG.

diff --git a/Engine/Plugins/Animation/ControlRig/Source/ControlRig/Private/Rigs/RigHierarchy.cpp b/Engine/Plugins/Animation/ControlRig/Source/ControlRig/Private/Rigs/RigHierarchy.cpp
--- a/Engine/Plugins/Animation/ControlRig/Source/ControlRig/Private/Rigs/RigHierarchy.cpp
+++ b/Engine/Plugins/Animation/ControlRig/Source/ControlRig/Private/Rigs/RigHierarchy.cpp
@@ -848,6 +848,10 @@ FEdGraphPinType URigHierarchy::GetControlPinType(FRigControlElement* InControlEl
 	static const FName PC_Int(TEXT("int"));
 	static const FName PC_Struct(TEXT("struct"));
 
+#if ENABLE_BLUEPRINT_REAL_NUMBERS
+	static const FName PC_Real(TEXT("real"));
+#endif
+
 	FEdGraphPinType PinType;
 
 	switch(InControlElement->Settings.ControlType)
@@ -857,11 +861,20 @@ FEdGraphPinType URigHierarchy::GetControlPinType(FRigControlElement* InControlEl
 			PinType.PinCategory = PC_Boolean;
 			break;
 		}
+#if ENABLE_BLUEPRINT_REAL_NUMBERS
+		case ERigControlType::Float:
+		{
+			PinType.PinCategory = PC_Real;
+			PinType.PinSubCategory = PC_Float;
+			break;
+		}
+#else
 		case ERigControlType::Float:
 		{
 			PinType.PinCategory = PC_Float;
 			break;
 		}
+#endif
 		case ERigControlType::Integer:
 		{
 			PinType.PinCategory = PC_Int;
diff --git a/Engine/Plugins/Animation/ControlRig/Source/ControlRigDeveloper/Private/ControlRigBlueprint.cpp b/Engine/Plugins/Animation/ControlRig/Source/ControlRigDeveloper/Private/ControlRigBlueprint.cpp
--- a/Engine/Plugins/Animation/ControlRig/Source/ControlRigDeveloper/Private/ControlRigBlueprint.cpp
+++ b/Engine/Plugins/Animation/ControlRig/Source/ControlRigDeveloper/Private/ControlRigBlueprint.cpp
@@ -455,6 +455,7 @@ void UControlRigBlueprint::PostLoad()
 
 		PatchFunctionReferencesOnLoad();
 		PatchVariableNodesOnLoad();
+		PatchVariableNodesWithIncorrectType();
 		PatchRigElementKeyCacheOnLoad();
 		PatchBoundVariables();
 
@@ -2584,6 +2585,23 @@ void UControlRigBlueprint::PopulateModelFromGraphForBackwardsCompatibility(UCont
 					}
 
 					FName DataType = PinType.PinCategory;
+#if ENABLE_BLUEPRINT_REAL_NUMBERS
+					if (PinType.PinCategory == UEdGraphSchema_K2::PC_Real)
+					{
+						if (PinType.PinSubCategory == UEdGraphSchema_K2::PC_Float)
+						{
+							DataType = TEXT("float");
+						}
+						else if (PinType.PinSubCategory == UEdGraphSchema_K2::PC_Double)
+						{
+							DataType = TEXT("double");
+						}
+						else
+						{
+							ensure(false);
+						}
+					}
+#endif
 					UObject* DataTypeObject = nullptr;
 					if (DataType == NAME_None)
 					{
@@ -3598,6 +3616,73 @@ void UControlRigBlueprint::PatchBoundVariables()
 	}
 }
 
+void UControlRigBlueprint::PatchVariableNodesWithIncorrectType()
+{
+	TGuardValue<bool> GuardNotifsSelf(bSuspendModelNotificationsForSelf, true);
+
+	struct Local
+	{
+		static bool RefreshIfNeeded(URigVMController* Controller, URigVMVariableNode* VariableNode, const FString& CPPType, UObject* CPPTypeObject)
+		{
+			if (URigVMPin* ValuePin = VariableNode->GetValuePin())
+			{
+				if (ValuePin->GetCPPType() != CPPType || ValuePin->GetCPPTypeObject() != CPPTypeObject)
+				{
+					Controller->RefreshVariableNode(VariableNode->GetFName(), VariableNode->GetVariableName(), CPPType, CPPTypeObject, false);
+					return true;
+				}
+			}
+			return false;
+		}
+	};
+
+	for (URigVMGraph* Graph : GetAllModels())
+	{
+		URigVMController* Controller = GetOrCreateController(Graph);
+		TArray<URigVMNode*> Nodes = Graph->GetNodes();
+		for (URigVMNode* Node : Nodes)
+		{
+			if (URigVMVariableNode* VariableNode = Cast<URigVMVariableNode>(Node))
+			{
+				FRigVMGraphVariableDescription Description = VariableNode->GetVariableDescription();
+
+				// Check for inputs and local variables
+				TArray<FRigVMGraphVariableDescription> LocalVariables = Graph->GetLocalVariables(true);
+				bool bLocalVariableFound = false;
+				for (FRigVMGraphVariableDescription Variable : LocalVariables)
+				{
+					if (Variable.Name == Description.Name)
+					{
+						if (Local::RefreshIfNeeded(Controller, VariableNode, Variable.CPPType, Variable.CPPTypeObject))
+						{
+							bDirtyDuringLoad = true;
+						}
+						bLocalVariableFound = true;
+						break;
+					}
+				}
+
+				if (!bLocalVariableFound)
+				{
+					for (struct FBPVariableDescription& Variable : NewVariables)
+					{
+						if (Variable.VarName == Description.Name)
+						{
+							FString CPPType;
+							UObject* CPPTypeObject = nullptr;
+							RigVMTypeUtils::CPPTypeFromPinType(Variable.VarType, CPPType, &CPPTypeObject);
+							if (Local::RefreshIfNeeded(Controller, VariableNode, CPPType, CPPTypeObject))
+							{
+								bDirtyDuringLoad = true;
+							}
+						}
+					}
+				}
+			}
+		}
+	}
+}
+
 void UControlRigBlueprint::PropagatePoseFromInstanceToBP(UControlRig* InControlRig)
 {
 	check(InControlRig);
diff --git a/Engine/Plugins/Animation/ControlRig/Source/ControlRigDeveloper/Private/Graph/ControlRigGraphNode.cpp b/Engine/Plugins/Animation/ControlRig/Source/ControlRigDeveloper/Private/Graph/ControlRigGraphNode.cpp
--- a/Engine/Plugins/Animation/ControlRig/Source/ControlRigDeveloper/Private/Graph/ControlRigGraphNode.cpp
+++ b/Engine/Plugins/Animation/ControlRig/Source/ControlRigDeveloper/Private/Graph/ControlRigGraphNode.cpp
@@ -1164,6 +1164,18 @@ FEdGraphPinType UControlRigGraphNode::GetPinTypeForModelPin(URigVMPin* InModelPi
 	{
 		PinType.PinCategory = UEdGraphSchema_K2::PC_Int;
 	}
+#if ENABLE_BLUEPRINT_REAL_NUMBERS
+	else if (ModelPinCPPType == TEXT("float"))
+	{
+		PinType.PinCategory = UEdGraphSchema_K2::PC_Real;
+		PinType.PinSubCategory = UEdGraphSchema_K2::PC_Float;
+	}
+	else if (ModelPinCPPType == TEXT("double"))
+	{
+		PinType.PinCategory = UEdGraphSchema_K2::PC_Real;
+		PinType.PinSubCategory = UEdGraphSchema_K2::PC_Double;
+	}
+#else
 	else if (ModelPinCPPType == TEXT("float"))
 	{
 		PinType.PinCategory = UEdGraphSchema_K2::PC_Float;
@@ -1172,6 +1184,7 @@ FEdGraphPinType UControlRigGraphNode::GetPinTypeForModelPin(URigVMPin* InModelPi
 	{
 		PinType.PinCategory = UEdGraphSchema_K2::PC_Double;
 	}
+#endif
 	else if (ModelPinCPPType == TEXT("FName"))
 	{
 		PinType.PinCategory = UEdGraphSchema_K2::PC_Name;
@@ -1191,15 +1204,15 @@ FEdGraphPinType UControlRigGraphNode::GetPinTypeForModelPin(URigVMPin* InModelPi
 		PinType.PinSubCategoryObject = InModelPin->GetCPPTypeObject();
 	}
 	else if (InModelPin->GetEnum() != nullptr)
-		{
+	{
 		PinType.PinCategory = UEdGraphSchema_K2::PC_Byte;
 		PinType.PinSubCategoryObject = InModelPin->GetEnum();
-		}
+	}
 
 	if (InModelPin->IsArray())
-		{
+	{
 		PinType.ContainerType = EPinContainerType::Array;
-		}
+	}
 	else
 	{
 		PinType.ContainerType = EPinContainerType::None;
diff --git a/Engine/Plugins/Animation/ControlRig/Source/ControlRigDeveloper/Private/Graph/ControlRigGraphSchema.cpp b/Engine/Plugins/Animation/ControlRig/Source/ControlRigDeveloper/Private/Graph/ControlRigGraphSchema.cpp
--- a/Engine/Plugins/Animation/ControlRig/Source/ControlRigDeveloper/Private/Graph/ControlRigGraphSchema.cpp
+++ b/Engine/Plugins/Animation/ControlRig/Source/ControlRigDeveloper/Private/Graph/ControlRigGraphSchema.cpp
@@ -757,12 +757,14 @@ FLinearColor UControlRigGraphSchema::GetPinTypeColor(const FEdGraphPinType& PinT
 		}
 	}
 
+#if !ENABLE_BLUEPRINT_REAL_NUMBERS
 	if(TypeName == UEdGraphSchema_K2::PC_Double)
 	{
 		FEdGraphPinType LocalPinType = PinType;
 		LocalPinType.PinCategory = UEdGraphSchema_K2::PC_Float;
 		return GetPinTypeColor(LocalPinType);
 	}
+#endif
 	
 	return GetDefault<UEdGraphSchema_K2>()->GetPinTypeColor(PinType);
 }
@@ -884,8 +886,12 @@ bool UControlRigGraphSchema::SupportsPinType(TWeakPtr<const FEdGraphSchemaAction
 
 	if (TypeName == UEdGraphSchema_K2::PC_Boolean ||
 		TypeName == UEdGraphSchema_K2::PC_Int ||
+#if ENABLE_BLUEPRINT_REAL_NUMBERS
+		TypeName == UEdGraphSchema_K2::PC_Real ||
+#else
 		TypeName == UEdGraphSchema_K2::PC_Float ||
 		TypeName == UEdGraphSchema_K2::PC_Double ||
+#endif
 		TypeName == UEdGraphSchema_K2::PC_Name ||
 		TypeName == UEdGraphSchema_K2::PC_String ||
 		TypeName == UEdGraphSchema_K2::PC_Enum)
@@ -1652,8 +1658,12 @@ void UControlRigGraphSchema::ResetPinDefaultsRecursive(UEdGraphPin* InPin) const
 void UControlRigGraphSchema::GetVariablePinTypes(TArray<FEdGraphPinType>& PinTypes) const
 {
 	PinTypes.Add(FEdGraphPinType(UEdGraphSchema_K2::PC_Boolean, FName(NAME_None), nullptr, EPinContainerType::None, false, FEdGraphTerminalType()));
+#if ENABLE_BLUEPRINT_REAL_NUMBERS
+	PinTypes.Add(FEdGraphPinType(UEdGraphSchema_K2::PC_Real, FName(NAME_None), nullptr, EPinContainerType::None, false, FEdGraphTerminalType()));
+#else
 	PinTypes.Add(FEdGraphPinType(UEdGraphSchema_K2::PC_Float, FName(NAME_None), nullptr, EPinContainerType::None, false, FEdGraphTerminalType()));
 	PinTypes.Add(FEdGraphPinType(UEdGraphSchema_K2::PC_Double, FName(NAME_None), nullptr, EPinContainerType::None, false, FEdGraphTerminalType()));
+#endif
 	PinTypes.Add(FEdGraphPinType(UEdGraphSchema_K2::PC_Int, FName(NAME_None), nullptr, EPinContainerType::None, false, FEdGraphTerminalType()));
 	PinTypes.Add(FEdGraphPinType(UEdGraphSchema_K2::PC_Struct, FName(NAME_None), TBaseStructure<FVector>::Get(), EPinContainerType::None, false, FEdGraphTerminalType()));
 	PinTypes.Add(FEdGraphPinType(UEdGraphSchema_K2::PC_Struct, FName(NAME_None), TBaseStructure<FVector2D>::Get(), EPinContainerType::None, false, FEdGraphTerminalType()));
diff --git a/Engine/Plugins/Animation/ControlRig/Source/ControlRigDeveloper/Public/ControlRigBlueprint.h b/Engine/Plugins/Animation/ControlRig/Source/ControlRigDeveloper/Public/ControlRigBlueprint.h
--- a/Engine/Plugins/Animation/ControlRig/Source/ControlRigDeveloper/Public/ControlRigBlueprint.h
+++ b/Engine/Plugins/Animation/ControlRig/Source/ControlRigDeveloper/Public/ControlRigBlueprint.h
@@ -545,6 +545,7 @@ private:
 	void PatchVariableNodesOnLoad();
 	void PatchRigElementKeyCacheOnLoad();
 	void PatchBoundVariables();
+	void PatchVariableNodesWithIncorrectType();
 
 	TMap<FName, int32> AddedMemberVariableMap;
 	TArray<FBPVariableDescription> LastNewVariables;
diff --git a/Engine/Plugins/Animation/LiveLink/Source/LiveLinkGraphNode/Private/K2Node_EvaluateLiveLinkCustom.cpp b/Engine/Plugins/Animation/LiveLink/Source/LiveLinkGraphNode/Private/K2Node_EvaluateLiveLinkCustom.cpp
--- a/Engine/Plugins/Animation/LiveLink/Source/LiveLinkGraphNode/Private/K2Node_EvaluateLiveLinkCustom.cpp
+++ b/Engine/Plugins/Animation/LiveLink/Source/LiveLinkGraphNode/Private/K2Node_EvaluateLiveLinkCustom.cpp
@@ -36,8 +36,11 @@ void UK2Node_EvaluateLiveLinkFrameAtWorldTime::AllocateDefaultPins()
 
 	const UEdGraphSchema_K2* K2Schema = GetDefault<UEdGraphSchema_K2>();
 
-	// float pin
+#if ENABLE_BLUEPRINT_REAL_NUMBERS
+	UEdGraphPin* LiveLinkWorldTimePin = CreatePin(EGPD_Input, UEdGraphSchema_K2::PC_Real, UEdGraphSchema_K2::PC_Double, UK2Node_EvaluateLiveLinkFrameHelper::LiveLinkWorldTimePinName);
+#else
 	UEdGraphPin* LiveLinkWorldTimePin = CreatePin(EGPD_Input, UEdGraphSchema_K2::PC_Float, UK2Node_EvaluateLiveLinkFrameHelper::LiveLinkWorldTimePinName);
+#endif
 	LiveLinkWorldTimePin->PinFriendlyName = LOCTEXT("LiveLinkWorldTimePin", "World Time");
 	SetPinToolTip(*LiveLinkWorldTimePin, LOCTEXT("LiveLinkWorldTimePinDescription", "The World Time the subject will be evaluated to"));
 }
diff --git a/Engine/Plugins/Experimental/EnhancedInput/Source/InputBlueprintNodes/Private/K2Node_EnhancedInputAction.cpp b/Engine/Plugins/Experimental/EnhancedInput/Source/InputBlueprintNodes/Private/K2Node_EnhancedInputAction.cpp
--- a/Engine/Plugins/Experimental/EnhancedInput/Source/InputBlueprintNodes/Private/K2Node_EnhancedInputAction.cpp
+++ b/Engine/Plugins/Experimental/EnhancedInput/Source/InputBlueprintNodes/Private/K2Node_EnhancedInputAction.cpp
@@ -77,8 +77,13 @@ void UK2Node_EnhancedInputAction::AllocateDefaultPins()
 	
 	Schema->SetPinAutogeneratedDefaultValueBasedOnType(ValuePin);
 
+#if ENABLE_BLUEPRINT_REAL_NUMBERS
+	CreatePin(EGPD_Output, UEdGraphSchema_K2::PC_Real, UEdGraphSchema_K2::PC_Double, ElapsedSecondsPinName)->bAdvancedView = true;
+	CreatePin(EGPD_Output, UEdGraphSchema_K2::PC_Real, UEdGraphSchema_K2::PC_Double, TriggeredSecondsPinName)->bAdvancedView = true;
+#else
 	CreatePin(EGPD_Output, UEdGraphSchema_K2::PC_Float, ElapsedSecondsPinName)->bAdvancedView = true;
 	CreatePin(EGPD_Output, UEdGraphSchema_K2::PC_Float, TriggeredSecondsPinName)->bAdvancedView = true;
+#endif
 
 	if(InputAction)
 	{
@@ -315,14 +320,27 @@ void UK2Node_EnhancedInputAction::ExpandNode(FKismetCompilerContext& CompilerCon
 	// Create temporary variables to copy ActionValue and ElapsedSeconds into
 	UK2Node_TemporaryVariable* ActionValueVar = CompilerContext.SpawnIntermediateNode<UK2Node_TemporaryVariable>(this, SourceGraph);
 	ActionValueVar->VariableType.PinCategory = UK2Node_GetInputActionValue::GetValueCategory(InputAction);
+	ActionValueVar->VariableType.PinSubCategory = UK2Node_GetInputActionValue::GetValueSubCategory(InputAction);
 	ActionValueVar->VariableType.PinSubCategoryObject = UK2Node_GetInputActionValue::GetValueSubCategoryObject(InputAction);
 	ActionValueVar->AllocateDefaultPins();
+
+#if ENABLE_BLUEPRINT_REAL_NUMBERS
+	UK2Node_TemporaryVariable* ElapsedSecondsVar = CompilerContext.SpawnIntermediateNode<UK2Node_TemporaryVariable>(this, SourceGraph);
+	ElapsedSecondsVar->VariableType.PinCategory = UEdGraphSchema_K2::PC_Real;
+	ElapsedSecondsVar->VariableType.PinSubCategory = UEdGraphSchema_K2::PC_Double;
+	ElapsedSecondsVar->AllocateDefaultPins();
+	UK2Node_TemporaryVariable* TriggeredSecondsVar = CompilerContext.SpawnIntermediateNode<UK2Node_TemporaryVariable>(this, SourceGraph);
+	TriggeredSecondsVar->VariableType.PinCategory = UEdGraphSchema_K2::PC_Real;
+	TriggeredSecondsVar->VariableType.PinSubCategory = UEdGraphSchema_K2::PC_Double;
+	TriggeredSecondsVar->AllocateDefaultPins();
+#else
 	UK2Node_TemporaryVariable* ElapsedSecondsVar = CompilerContext.SpawnIntermediateNode<UK2Node_TemporaryVariable>(this, SourceGraph);
 	ElapsedSecondsVar->VariableType.PinCategory = UEdGraphSchema_K2::PC_Float;
 	ElapsedSecondsVar->AllocateDefaultPins();
 	UK2Node_TemporaryVariable* TriggeredSecondsVar = CompilerContext.SpawnIntermediateNode<UK2Node_TemporaryVariable>(this, SourceGraph);
 	TriggeredSecondsVar->VariableType.PinCategory = UEdGraphSchema_K2::PC_Float;
 	TriggeredSecondsVar->AllocateDefaultPins();
+#endif
 	UK2Node_TemporaryVariable* InputActionVar = CompilerContext.SpawnIntermediateNode<UK2Node_TemporaryVariable>(this, SourceGraph);
 	InputActionVar->VariableType.PinCategory = UEdGraphSchema_K2::PC_Object;
 	InputActionVar->VariableType.PinSubCategoryObject = InputAction->GetClass();
diff --git a/Engine/Plugins/Experimental/EnhancedInput/Source/InputBlueprintNodes/Private/K2Node_GetInputActionValue.cpp b/Engine/Plugins/Experimental/EnhancedInput/Source/InputBlueprintNodes/Private/K2Node_GetInputActionValue.cpp
--- a/Engine/Plugins/Experimental/EnhancedInput/Source/InputBlueprintNodes/Private/K2Node_GetInputActionValue.cpp
+++ b/Engine/Plugins/Experimental/EnhancedInput/Source/InputBlueprintNodes/Private/K2Node_GetInputActionValue.cpp
@@ -33,19 +33,28 @@ UK2Node_GetInputActionValue::UK2Node_GetInputActionValue(const FObjectInitialize
 
 struct FValueTypeData
 {
-	FValueTypeData(FName InCategory, UScriptStruct* InSubCategory = nullptr)
-		: Category(InCategory), SubCategory(InSubCategory) {}
+	FValueTypeData(FName InCategory, FName InSubCategory = NAME_None, UScriptStruct* InSubCategoryObject = nullptr)
+		: Category(InCategory)
+		, SubCategory(InSubCategory)
+		, SubCategoryObject(InSubCategoryObject)
+	{
+	}
 
 	FName Category;
-	UScriptStruct* SubCategory;
+	FName SubCategory;
+	UScriptStruct* SubCategoryObject;
 };
 
 static TMap<EInputActionValueType, FValueTypeData> ValueLookups =
 {
 	{ EInputActionValueType::Boolean, FValueTypeData(UEdGraphSchema_K2::PC_Boolean) },
+#if ENABLE_BLUEPRINT_REAL_NUMBERS
+	{ EInputActionValueType::Axis1D, FValueTypeData(UEdGraphSchema_K2::PC_Real, UEdGraphSchema_K2::PC_Float) },
+#else
 	{ EInputActionValueType::Axis1D, FValueTypeData(UEdGraphSchema_K2::PC_Float) },
-	{ EInputActionValueType::Axis2D, FValueTypeData(UEdGraphSchema_K2::PC_Struct, TBaseStructure<FVector2D>::Get()) },
-	{ EInputActionValueType::Axis3D, FValueTypeData(UEdGraphSchema_K2::PC_Struct, TBaseStructure<FVector>::Get()) },
+#endif
+	{ EInputActionValueType::Axis2D, FValueTypeData(UEdGraphSchema_K2::PC_Struct, NAME_None, TBaseStructure<FVector2D>::Get()) },
+	{ EInputActionValueType::Axis3D, FValueTypeData(UEdGraphSchema_K2::PC_Struct, NAME_None, TBaseStructure<FVector>::Get()) },
 };
 
 FName UK2Node_GetInputActionValue::GetValueCategory(const UInputAction* InputAction)
@@ -54,12 +63,17 @@ FName UK2Node_GetInputActionValue::GetValueCategory(const UInputAction* InputAct
 	return ValueLookups[Type].Category;
 }
 
-UScriptStruct* UK2Node_GetInputActionValue::GetValueSubCategoryObject(const UInputAction* InputAction)
+FName UK2Node_GetInputActionValue::GetValueSubCategory(const UInputAction* InputAction)
 {
 	EInputActionValueType Type = InputAction ? InputAction->ValueType : EInputActionValueType::Boolean;
 	return ValueLookups[Type].SubCategory;
 }
 
+UScriptStruct* UK2Node_GetInputActionValue::GetValueSubCategoryObject(const UInputAction* InputAction)
+{
+	EInputActionValueType Type = InputAction ? InputAction->ValueType : EInputActionValueType::Boolean;
+	return ValueLookups[Type].SubCategoryObject;
+}
 
 void UK2Node_GetInputActionValue::AllocateDefaultPins()
 {
@@ -68,7 +82,16 @@ void UK2Node_GetInputActionValue::AllocateDefaultPins()
 	Super::AllocateDefaultPins();
 
 	// Dynamically typed output
-	CreatePin(EGPD_Output, UK2Node_GetInputActionValue::GetValueCategory(InputAction), UK2Node_GetInputActionValue::GetValueSubCategoryObject(InputAction), UEdGraphSchema_K2::PN_ReturnValue);
+	FName SubCategory = UK2Node_GetInputActionValue::GetValueSubCategory(InputAction);
+
+	if (SubCategory != NAME_None)
+	{
+		CreatePin(EGPD_Output, UK2Node_GetInputActionValue::GetValueCategory(InputAction), SubCategory, UEdGraphSchema_K2::PN_ReturnValue);
+	}
+	else
+	{
+		CreatePin(EGPD_Output, UK2Node_GetInputActionValue::GetValueCategory(InputAction), UK2Node_GetInputActionValue::GetValueSubCategoryObject(InputAction), UEdGraphSchema_K2::PN_ReturnValue);
+	}
 }
 
 void UK2Node_GetInputActionValue::Initialize(const UInputAction* Action)
diff --git a/Engine/Plugins/Experimental/EnhancedInput/Source/InputBlueprintNodes/Public/K2Node_GetInputActionValue.h b/Engine/Plugins/Experimental/EnhancedInput/Source/InputBlueprintNodes/Public/K2Node_GetInputActionValue.h
--- a/Engine/Plugins/Experimental/EnhancedInput/Source/InputBlueprintNodes/Public/K2Node_GetInputActionValue.h
+++ b/Engine/Plugins/Experimental/EnhancedInput/Source/InputBlueprintNodes/Public/K2Node_GetInputActionValue.h
@@ -43,6 +43,7 @@ class UK2Node_GetInputActionValue : public UK2Node
 
 	// Auto pin generation helpers
 	static FName GetValueCategory(const UInputAction* InputAction);
+	static FName GetValueSubCategory(const UInputAction* InputAction);
 	static UScriptStruct* GetValueSubCategoryObject(const UInputAction* InputAction);
 
 private:
diff --git a/Engine/Plugins/Runtime/PropertyAccess/Source/PropertyAccessEditor/Private/PropertyAccessEditor.cpp b/Engine/Plugins/Runtime/PropertyAccess/Source/PropertyAccessEditor/Private/PropertyAccessEditor.cpp
--- a/Engine/Plugins/Runtime/PropertyAccess/Source/PropertyAccessEditor/Private/PropertyAccessEditor.cpp
+++ b/Engine/Plugins/Runtime/PropertyAccess/Source/PropertyAccessEditor/Private/PropertyAccessEditor.cpp
@@ -456,6 +456,28 @@ struct FPropertyAccessEditorSystem
 					return EPropertyAccessCopyType::DemoteDoubleToFloat;
 				}
 			}
+#if ENABLE_BLUEPRINT_REAL_NUMBERS
+			else if (SrcProperty->IsA<FArrayProperty>() && DestProperty->IsA<FArrayProperty>())
+			{
+				const FArrayProperty* SrcArrayProperty = CastField<const FArrayProperty>(SrcProperty);
+				const FArrayProperty* DestArrayProperty = CastField<const FArrayProperty>(DestProperty);
+
+				if (SrcArrayProperty->Inner->IsA<FFloatProperty>())
+				{
+					if (DestArrayProperty->Inner->IsA<FDoubleProperty>())
+					{
+						return EPropertyAccessCopyType::PromoteArrayFloatToDouble;
+					}
+				}
+				else if (SrcArrayProperty->Inner->IsA<FDoubleProperty>())
+				{
+					if (DestArrayProperty->Inner->IsA<FFloatProperty>())
+					{
+						return EPropertyAccessCopyType::DemoteArrayDoubleToFloat;
+					}
+				}
+			}
+#endif
 		}
 
 		OutErrorMessage = FText::Format(LOCTEXT("CopyTypeInvalidFormat", "@@ Cannot copy property ({0} -> {1})"), FText::FromString(SrcProperty->GetCPPType()), FText::FromString(DestProperty->GetCPPType()));
@@ -612,6 +634,28 @@ namespace PropertyAccess
 					return EPropertyAccessCompatibility::Promotable;	// LWC_TODO: Incorrect! Do not ship this!
 				}
 			}
+#if ENABLE_BLUEPRINT_REAL_NUMBERS
+			else if (InPropertyA->IsA<FArrayProperty>() && InPropertyB->IsA<FArrayProperty>())
+			{
+				const FArrayProperty* ArrayPropertyA = CastField<const FArrayProperty>(InPropertyA);
+				const FArrayProperty* ArrayPropertyB = CastField<const FArrayProperty>(InPropertyB);
+
+				if (ArrayPropertyA->Inner->IsA<FFloatProperty>())
+				{
+					if (ArrayPropertyB->Inner->IsA<FDoubleProperty>())
+					{
+						return EPropertyAccessCompatibility::Promotable;
+					}
+				}
+				else if (ArrayPropertyA->Inner->IsA<FDoubleProperty>())
+				{
+					if (ArrayPropertyB->Inner->IsA<FFloatProperty>())
+					{
+						return EPropertyAccessCompatibility::Promotable;
+					}
+				}
+			}
+#endif
 		}
 
 		return EPropertyAccessCompatibility::Incompatible;
@@ -620,34 +664,47 @@ namespace PropertyAccess
 	EPropertyAccessCompatibility GetPinTypeCompatibility(const FEdGraphPinType& InPinTypeA, const FEdGraphPinType& InPinTypeB)
 	{
 		const UEdGraphSchema_K2* Schema = GetDefault<UEdGraphSchema_K2>();
-		if(Schema->ArePinTypesCompatible(InPinTypeA, InPinTypeB))
+		if (Schema->ArePinTypesCompatible(InPinTypeA, InPinTypeB))
 		{
 			return EPropertyAccessCompatibility::Compatible;
 		}
 		else
 		{
 			// Not directly compatible, check for promotions
-			if(InPinTypeA.PinCategory == UEdGraphSchema_K2::PC_Boolean)
+			if (InPinTypeA.PinCategory == UEdGraphSchema_K2::PC_Boolean)
 			{
-				if(InPinTypeB.PinCategory == UEdGraphSchema_K2::PC_Byte || InPinTypeB.PinCategory == UEdGraphSchema_K2::PC_Int || InPinTypeB.PinCategory == UEdGraphSchema_K2::PC_Int64 || InPinTypeB.PinCategory == UEdGraphSchema_K2::PC_Float || InPinTypeB.PinCategory == UEdGraphSchema_K2::PC_Double)
+#if ENABLE_BLUEPRINT_REAL_NUMBERS
+				if (InPinTypeB.PinCategory == UEdGraphSchema_K2::PC_Byte || InPinTypeB.PinCategory == UEdGraphSchema_K2::PC_Int || InPinTypeB.PinCategory == UEdGraphSchema_K2::PC_Int64 || InPinTypeB.PinCategory == UEdGraphSchema_K2::PC_Real)
+#else
+				if (InPinTypeB.PinCategory == UEdGraphSchema_K2::PC_Byte || InPinTypeB.PinCategory == UEdGraphSchema_K2::PC_Int || InPinTypeB.PinCategory == UEdGraphSchema_K2::PC_Int64 || InPinTypeB.PinCategory == UEdGraphSchema_K2::PC_Float || InPinTypeB.PinCategory == UEdGraphSchema_K2::PC_Double)
+#endif
 				{
 					return EPropertyAccessCompatibility::Promotable;
 				}
 			}
-			else if(InPinTypeA.PinCategory == UEdGraphSchema_K2::PC_Byte)
+			else if (InPinTypeA.PinCategory == UEdGraphSchema_K2::PC_Byte)
 			{
-				if(InPinTypeB.PinCategory == UEdGraphSchema_K2::PC_Int || InPinTypeB.PinCategory == UEdGraphSchema_K2::PC_Int64 || InPinTypeB.PinCategory == UEdGraphSchema_K2::PC_Float || InPinTypeB.PinCategory == UEdGraphSchema_K2::PC_Double)
+#if ENABLE_BLUEPRINT_REAL_NUMBERS
+				if (InPinTypeB.PinCategory == UEdGraphSchema_K2::PC_Int || InPinTypeB.PinCategory == UEdGraphSchema_K2::PC_Int64 || InPinTypeB.PinCategory == UEdGraphSchema_K2::PC_Real)
+#else
+				if (InPinTypeB.PinCategory == UEdGraphSchema_K2::PC_Int || InPinTypeB.PinCategory == UEdGraphSchema_K2::PC_Int64 || InPinTypeB.PinCategory == UEdGraphSchema_K2::PC_Float || InPinTypeB.PinCategory == UEdGraphSchema_K2::PC_Double)
+#endif
 				{
 					return EPropertyAccessCompatibility::Promotable;
 				}
 			}
-			else if(InPinTypeA.PinCategory == UEdGraphSchema_K2::PC_Int)
+			else if (InPinTypeA.PinCategory == UEdGraphSchema_K2::PC_Int)
 			{
-				if(InPinTypeB.PinCategory == UEdGraphSchema_K2::PC_Int64 || InPinTypeB.PinCategory == UEdGraphSchema_K2::PC_Float || InPinTypeB.PinCategory == UEdGraphSchema_K2::PC_Double)
+#if ENABLE_BLUEPRINT_REAL_NUMBERS
+				if (InPinTypeB.PinCategory == UEdGraphSchema_K2::PC_Int64 || InPinTypeB.PinCategory == UEdGraphSchema_K2::PC_Real)
+#else
+				if (InPinTypeB.PinCategory == UEdGraphSchema_K2::PC_Int64 || InPinTypeB.PinCategory == UEdGraphSchema_K2::PC_Float || InPinTypeB.PinCategory == UEdGraphSchema_K2::PC_Double)
+#endif
 				{
 					return EPropertyAccessCompatibility::Promotable;
 				}
 			}
+#if !ENABLE_BLUEPRINT_REAL_NUMBERS
 			else if (InPinTypeA.PinCategory == UEdGraphSchema_K2::PC_Float)
 			{
 				if (InPinTypeB.PinCategory == UEdGraphSchema_K2::PC_Double)
@@ -662,6 +719,7 @@ namespace PropertyAccess
 					return EPropertyAccessCompatibility::Promotable;	// LWC_TODO: Incorrect! Do not ship this!
 				}
 			}
+#endif
 		}
 
 		return EPropertyAccessCompatibility::Incompatible;
diff --git a/Engine/Plugins/ScriptPlugin/Source/ScriptEditorPlugin/Private/ScriptBlueprintCompiler.cpp b/Engine/Plugins/ScriptPlugin/Source/ScriptEditorPlugin/Private/ScriptBlueprintCompiler.cpp
--- a/Engine/Plugins/ScriptPlugin/Source/ScriptEditorPlugin/Private/ScriptBlueprintCompiler.cpp
+++ b/Engine/Plugins/ScriptPlugin/Source/ScriptEditorPlugin/Private/ScriptBlueprintCompiler.cpp
@@ -49,13 +49,19 @@ void FScriptBlueprintCompiler::CreateClassVariablesFromBlueprint()
 		if (Field.PropertyClass)
 		{
 			FName PinCategory;
+			FName PinSubCategory = NAME_None;
 			if (Field.PropertyClass->IsChildOf(FStrProperty::StaticClass()))
 			{
 				PinCategory = UEdGraphSchema_K2::PC_String;
 			}
 			else if (Field.PropertyClass->IsChildOf(FFloatProperty::StaticClass()))
 			{
+#if ENABLE_BLUEPRINT_REAL_NUMBERS
+				PinCategory = UEdGraphSchema_K2::PC_Real;
+				PinSubCategory = UEdGraphSchema_K2::PC_Float;
+#else
 				PinCategory = UEdGraphSchema_K2::PC_Float;
+#endif
 			}
 			else if (Field.PropertyClass->IsChildOf(FIntProperty::StaticClass()))
 			{
@@ -77,7 +83,7 @@ void FScriptBlueprintCompiler::CreateClassVariablesFromBlueprint()
 			}
 			if (!PinCategory.IsNone())
 			{
-				FEdGraphPinType ScriptPinType(PinCategory, NAME_None, InnerType, EPinContainerType::None, false, FEdGraphTerminalType());
+				FEdGraphPinType ScriptPinType(PinCategory, PinSubCategory, InnerType, EPinContainerType::None, false, FEdGraphTerminalType());
 				FProperty* ScriptProperty = CreateVariable(Field.Name, ScriptPinType);
 				if (ScriptProperty)
 				{
diff --git a/Engine/Source/Developer/RigVMDeveloper/Private/RigVMTypeUtils.cpp b/Engine/Source/Developer/RigVMDeveloper/Private/RigVMTypeUtils.cpp
--- a/Engine/Source/Developer/RigVMDeveloper/Private/RigVMTypeUtils.cpp
+++ b/Engine/Source/Developer/RigVMDeveloper/Private/RigVMTypeUtils.cpp
@@ -85,6 +85,25 @@ FRigVMExternalVariable RigVMTypeUtils::ExternalVariableFromPinType(const FName&
 		}
 		ExternalVariable.Size = sizeof(uint8);
 	}
+#if ENABLE_BLUEPRINT_REAL_NUMBERS
+	else if (InPinType.PinCategory == UEdGraphSchema_K2::PC_Real)
+	{
+		if (InPinType.PinSubCategory == UEdGraphSchema_K2::PC_Float)
+		{
+			ExternalVariable.TypeName = TEXT("float");
+			ExternalVariable.Size = sizeof(float);
+		}
+		else if (InPinType.PinSubCategory == UEdGraphSchema_K2::PC_Double)
+		{
+			ExternalVariable.TypeName = TEXT("double");
+			ExternalVariable.Size = sizeof(double);
+		}
+		else
+		{
+			checkf(false, TEXT("Unexpected subcategory for PC_Real pin type."));
+		}
+	}		
+#else
 	else if (InPinType.PinCategory == UEdGraphSchema_K2::PC_Float)
 	{
 		ExternalVariable.TypeName = TEXT("float");
@@ -95,6 +114,7 @@ FRigVMExternalVariable RigVMTypeUtils::ExternalVariableFromPinType(const FName&
 		ExternalVariable.TypeName = TEXT("double");
 		ExternalVariable.Size = sizeof(double);
 	}
+#endif
 	else if (InPinType.PinCategory == UEdGraphSchema_K2::PC_Name)
 	{
 		ExternalVariable.TypeName = TEXT("FName");
@@ -324,6 +344,18 @@ FEdGraphPinType RigVMTypeUtils::PinTypeFromExternalVariable(const FRigVMExternal
 	{
 		PinType.PinCategory = UEdGraphSchema_K2::PC_Int;
 	}
+#if ENABLE_BLUEPRINT_REAL_NUMBERS
+	else if (InExternalVariable.TypeName == TEXT("float"))
+	{
+		PinType.PinCategory = UEdGraphSchema_K2::PC_Real;
+		PinType.PinSubCategory = UEdGraphSchema_K2::PC_Float;
+	}
+	else if (InExternalVariable.TypeName == TEXT("double"))
+	{
+		PinType.PinCategory = UEdGraphSchema_K2::PC_Real;
+		PinType.PinSubCategory = UEdGraphSchema_K2::PC_Double;
+	}
+#else
 	else if (InExternalVariable.TypeName == TEXT("float"))
 	{
 		PinType.PinCategory = UEdGraphSchema_K2::PC_Float;
@@ -332,6 +364,7 @@ FEdGraphPinType RigVMTypeUtils::PinTypeFromExternalVariable(const FRigVMExternal
 	{
 		PinType.PinCategory = UEdGraphSchema_K2::PC_Double;
 	}
+#endif
 	else if (InExternalVariable.TypeName == TEXT("FName"))
 	{
 		PinType.PinCategory = UEdGraphSchema_K2::PC_Name;
@@ -389,6 +422,18 @@ FEdGraphPinType RigVMTypeUtils::PinTypeFromRigVMVariableDescription(
 	{
 		PinType.PinCategory = UEdGraphSchema_K2::PC_Int;
 	}
+#if ENABLE_BLUEPRINT_REAL_NUMBERS
+	else if (CurrentCPPType == TEXT("float"))
+	{
+		PinType.PinCategory = UEdGraphSchema_K2::PC_Real;
+		PinType.PinSubCategory = UEdGraphSchema_K2::PC_Float;
+	}
+	else if (CurrentCPPType == TEXT("double"))
+	{
+		PinType.PinCategory = UEdGraphSchema_K2::PC_Real;
+		PinType.PinSubCategory = UEdGraphSchema_K2::PC_Double;
+	}
+#else
 	else if (CurrentCPPType == TEXT("float"))
 	{
 		PinType.PinCategory = UEdGraphSchema_K2::PC_Float;
@@ -397,6 +442,7 @@ FEdGraphPinType RigVMTypeUtils::PinTypeFromRigVMVariableDescription(
 	{
 		PinType.PinCategory = UEdGraphSchema_K2::PC_Double;
 	}
+#endif
 	else if (CurrentCPPType == TEXT("FName"))
 	{
 		PinType.PinCategory = UEdGraphSchema_K2::PC_Name;
@@ -493,6 +539,32 @@ bool RigVMTypeUtils::CPPTypeFromPinType(const FEdGraphPinType& InPinType, FStrin
 	{
 		OutCPPType = Prefix + TEXT("int32") + Suffix;
 	}
+#if ENABLE_BLUEPRINT_REAL_NUMBERS
+	else if (InPinType.PinCategory == UEdGraphSchema_K2::PC_Real)
+	{
+		if (InPinType.PinSubCategory == UEdGraphSchema_K2::PC_Float)
+		{
+			OutCPPType = Prefix + TEXT("float") + Suffix;
+		}
+		else if (InPinType.PinSubCategory == UEdGraphSchema_K2::PC_Double)
+		{
+			OutCPPType = Prefix + TEXT("double") + Suffix;
+		}
+		else
+		{
+			checkf(false, TEXT("Unexpected subcategory for PC_Real pin type."));
+		}
+	}		
+#else
+	else if (InPinType.PinCategory == UEdGraphSchema_K2::PC_Float)
+	{
+		OutCPPType = Prefix + TEXT("float") + Suffix;
+	}
+	else if (InPinType.PinCategory == UEdGraphSchema_K2::PC_Double)
+	{
+		OutCPPType = Prefix + TEXT("double") + Suffix;
+	}
+#endif
 	else if (InPinType.PinCategory == UEdGraphSchema_K2::PC_Float)
 	{
 		OutCPPType = Prefix + TEXT("float") + Suffix;
diff --git a/Engine/Source/Developer/RigVMDeveloper/Public/RigVMModel/Nodes/RigVMVariableNode.h b/Engine/Source/Developer/RigVMDeveloper/Public/RigVMModel/Nodes/RigVMVariableNode.h
--- a/Engine/Source/Developer/RigVMDeveloper/Public/RigVMModel/Nodes/RigVMVariableNode.h
+++ b/Engine/Source/Developer/RigVMDeveloper/Public/RigVMModel/Nodes/RigVMVariableNode.h
@@ -72,6 +72,7 @@ private:
 	
 
 	friend class URigVMController;
+	friend class UControlRigBlueprint;
 	friend struct FRigVMRemoveNodeAction;
 	friend class URigVMPin;
 	friend class URigVMCompiler;
diff --git a/Engine/Source/Developer/ScriptDisassembler/Private/ScriptDisassembler.cpp b/Engine/Source/Developer/ScriptDisassembler/Private/ScriptDisassembler.cpp
--- a/Engine/Source/Developer/ScriptDisassembler/Private/ScriptDisassembler.cpp
+++ b/Engine/Source/Developer/ScriptDisassembler/Private/ScriptDisassembler.cpp
@@ -215,29 +215,104 @@ FString FKismetBytecodeDisassembler::ReadString16(int32& ScriptIndex)
 	return Result;
 }
 
-void FKismetBytecodeDisassembler::ProcessCastByte(int32 CastType, int32& ScriptIndex)
-{
-	// Expression of cast
-	SerializeExpr(ScriptIndex);
-}
-
 void FKismetBytecodeDisassembler::ProcessCommon(int32& ScriptIndex, EExprToken Opcode)
 {
+	static const TCHAR* CastNameTable[CST_Max] = {
+		TEXT("ObjectToInterface"),
+		TEXT("ObjectToBool"),
+		TEXT("InterfaceToBool"),
+		TEXT("DoubleToFloat"),
+		TEXT("DoubleToFloatArray"),
+		TEXT("DoubleToFloatSet"),
+		TEXT("FloatToDouble"),
+		TEXT("FloatToDoubleArray"),
+		TEXT("FloatToDoubleSet"),
+		TEXT("VectorToVector3f"),
+		TEXT("Vector3fToVector"),
+		TEXT("FloatToDoubleKeys_Map"),
+		TEXT("DoubleToFloatKeys_Map"),
+		TEXT("FloatToDoubleValues_Map"),
+		TEXT("DoubleToFloatValues_Map"),
+		TEXT("FloatToDoubleKeys_FloatToDoubleValues_Map"),
+		TEXT("DoubleToFloatKeys_FloatToDoubleValues_Map"),
+		TEXT("DoubleToFloatKeys_DoubleToFloatValues_Map"),
+		TEXT("FloatToDoubleKeys_DoubleToFloatValues_Map"),
+	};
+
+	auto PrintVariable = [&ScriptIndex, Opcode, this](FStringView VariableDescription)
+	{
+		FProperty* PropertyPtr = ReadPointer<FProperty>(ScriptIndex);
+		FString PropertyName = TEXT("(null)");
+		FString PropertyType = TEXT("(null)");
+		FString ParameterType;
+
+		if (PropertyPtr)
+		{
+			PropertyName = PropertyPtr->GetName();
+
+			FString ExtendedPropertyType;
+			PropertyType = PropertyPtr->GetCPPType(&ExtendedPropertyType);
+			PropertyType += ExtendedPropertyType;
+
+			if (PropertyPtr->HasAnyPropertyFlags(CPF_ParmFlags))
+			{
+				ParameterType = TEXT("(");
+				if (PropertyPtr->HasAnyPropertyFlags(CPF_Parm))
+				{
+					ParameterType += TEXT("Parameter,");
+				}
+
+				if (PropertyPtr->HasAnyPropertyFlags(CPF_OutParm))
+				{
+					ParameterType += TEXT("Out,");
+				}
+
+				if (PropertyPtr->HasAnyPropertyFlags(CPF_ReturnParm))
+				{
+					ParameterType += TEXT("Return,");
+				}
+
+				if (PropertyPtr->HasAnyPropertyFlags(CPF_ReferenceParm))
+				{
+					ParameterType += TEXT("Reference,");
+				}
+
+				if (PropertyPtr->HasAnyPropertyFlags(CPF_ConstParm))
+				{
+					ParameterType += TEXT("Const,");
+				}
+
+				int32 LastCommaLocation = ParameterType.Len() - 1;
+				check(LastCommaLocation > 0);
+				ParameterType[LastCommaLocation] = TCHAR(')');
+			}
+		}
+
+		FString Output = 
+			FString::Printf(TEXT("%s $%X: %s of type %s named %s."), *Indents, (int32)Opcode, VariableDescription.GetData(), *PropertyType, *PropertyName);
+
+		if (ParameterType.Len() > 0)
+		{
+			Output += FString::Printf(TEXT(" Parameter flags: %s."), *ParameterType);
+		}
+
+		Ar.Logf(TEXT("%s"), *Output);
+	};
+
 	switch (Opcode)
 	{
-	case EX_PrimitiveCast:
+	case EX_Cast:
 		{
 			// A type conversion.
 			uint8 ConversionType = ReadBYTE(ScriptIndex);
-			Ar.Logf(TEXT("%s $%X: PrimitiveCast of type %d"), *Indents, (int32)Opcode, ConversionType);
+			check(CastNameTable[ConversionType] != nullptr);
+			Ar.Logf(TEXT("%s $%X: Cast of type %d (%s)"), *Indents, (int32)Opcode, ConversionType, CastNameTable[ConversionType]);
 			AddIndent();
 
 			Ar.Logf(TEXT("%s Argument:"), *Indents);
-			ProcessCastByte(ConversionType, ScriptIndex);
+			SerializeExpr(ScriptIndex);
 
-			//@TODO:
-			//Ar.Logf(TEXT("%s Expression:"), *Indents);
-			//SerializeExpr( ScriptIndex );
+			DropIndent();
 			break;
 		}
 	case EX_SetSet:
@@ -245,7 +320,7 @@ void FKismetBytecodeDisassembler::ProcessCommon(int32& ScriptIndex, EExprToken O
  			Ar.Logf(TEXT("%s $%X: set set"), *Indents, (int32)Opcode);
 			SerializeExpr(ScriptIndex);
 			ReadINT(ScriptIndex);
- 			while( SerializeExpr(ScriptIndex) != EX_EndSet)
+ 			while (SerializeExpr(ScriptIndex) != EX_EndSet)
  			{
  				// Set contents
  			}
@@ -277,7 +352,7 @@ void FKismetBytecodeDisassembler::ProcessCommon(int32& ScriptIndex, EExprToken O
  			Ar.Logf(TEXT("%s $%X: set map"), *Indents, (int32)Opcode);
 			SerializeExpr(ScriptIndex);
  			ReadINT(ScriptIndex);
- 			while( SerializeExpr(ScriptIndex) != EX_EndMap)
+ 			while (SerializeExpr(ScriptIndex) != EX_EndMap)
  			{
  				// Map contents
  			}
@@ -463,7 +538,7 @@ void FKismetBytecodeDisassembler::ProcessCommon(int32& ScriptIndex, EExprToken O
 			UStruct* StackNode = ReadPointer<UStruct>(ScriptIndex);
 			Ar.Logf(TEXT("%s $%X: Local Final Script Function (stack node %s::%s)"), *Indents, (int32)Opcode, StackNode ? *StackNode->GetOuter()->GetName() : TEXT("(null)"), StackNode ? *StackNode->GetName() : TEXT("(null)"));
 
-			while (SerializeExpr( ScriptIndex ) != EX_EndFunctionParms)
+			while (SerializeExpr(ScriptIndex) != EX_EndFunctionParms)
 			{
 				// Params
 			}
@@ -485,7 +560,6 @@ void FKismetBytecodeDisassembler::ProcessCommon(int32& ScriptIndex, EExprToken O
 			DropIndent();
 			break;
 		}
-
 	case EX_ComputedJump:
 		{
 			Ar.Logf(TEXT("%s $%X: Computed Jump, offset specified by expression:"), *Indents, (int32)Opcode);
@@ -496,7 +570,6 @@ void FKismetBytecodeDisassembler::ProcessCommon(int32& ScriptIndex, EExprToken O
 
 			break;
 		}
-
 	case EX_Jump:
 		{
 			CodeSkipSizeType SkipCount = ReadSkipCount(ScriptIndex);
@@ -505,32 +578,27 @@ void FKismetBytecodeDisassembler::ProcessCommon(int32& ScriptIndex, EExprToken O
 		}
 	case EX_LocalVariable:
 		{
-			FProperty* PropertyPtr = ReadPointer<FProperty>(ScriptIndex);
-			Ar.Logf(TEXT("%s $%X: Local variable named %s"), *Indents, (int32)Opcode, PropertyPtr ? *PropertyPtr->GetName() : TEXT("(null)"));
+			PrintVariable(TEXT("Local variable"));
 			break;
 		}
 	case EX_DefaultVariable:
 		{
-			FProperty* PropertyPtr = ReadPointer<FProperty>(ScriptIndex);
-			Ar.Logf(TEXT("%s $%X: Default variable named %s"), *Indents, (int32)Opcode, PropertyPtr ? *PropertyPtr->GetName() : TEXT("(null)"));
+			PrintVariable(TEXT("Default variable"));
 			break;
 		}
 	case EX_InstanceVariable:
 		{
-			FProperty* PropertyPtr = ReadPointer<FProperty>(ScriptIndex);
-			Ar.Logf(TEXT("%s $%X: Instance variable named %s"), *Indents, (int32)Opcode, PropertyPtr ? *PropertyPtr->GetName() : TEXT("(null)"));
+			PrintVariable(TEXT("Instance variable"));
 			break;
 		}
 	case EX_LocalOutVariable:
 		{
-			FProperty* PropertyPtr = ReadPointer<FProperty>(ScriptIndex);
-			Ar.Logf(TEXT("%s $%X: Local out variable named %s"), *Indents, (int32)Opcode, PropertyPtr ? *PropertyPtr->GetName() : TEXT("(null)"));
+			PrintVariable(TEXT("Local out variable"));
 			break;
 		}
 	case EX_ClassSparseDataVariable:
 		{
-			FProperty* PropertyPtr = ReadPointer<FProperty>(ScriptIndex);
-			Ar.Logf(TEXT("%s $%X: Class sparse data variable named %s"), *Indents, (int32)Opcode, PropertyPtr ? *PropertyPtr->GetName() : TEXT("(null)"));
+			PrintVariable(TEXT("Class sparse data variable"));
 			break;
 		}
 	case EX_InterfaceContext:
@@ -637,7 +705,7 @@ void FKismetBytecodeDisassembler::ProcessCommon(int32& ScriptIndex, EExprToken O
 			UStruct* StackNode = ReadPointer<UStruct>(ScriptIndex);
 			Ar.Logf(TEXT("%s $%X: Final Function (stack node %s::%s)"), *Indents, (int32)Opcode, StackNode ? *StackNode->GetOuter()->GetName() : TEXT("(null)"), StackNode ? *StackNode->GetName() : TEXT("(null)"));
 
-			while (SerializeExpr( ScriptIndex ) != EX_EndFunctionParms)
+			while (SerializeExpr(ScriptIndex) != EX_EndFunctionParms)
 			{
 				// Params
 			}
@@ -649,7 +717,7 @@ void FKismetBytecodeDisassembler::ProcessCommon(int32& ScriptIndex, EExprToken O
 			Ar.Logf(TEXT("%s $%X: CallMulticastDelegate (signature %s::%s) delegate:"), *Indents, (int32)Opcode, StackNode ? *StackNode->GetOuter()->GetName() : TEXT("(null)"), StackNode ? *StackNode->GetName() : TEXT("(null)"));
 			SerializeExpr( ScriptIndex );
 			Ar.Logf(TEXT("Params:"));
-			while (SerializeExpr( ScriptIndex ) != EX_EndFunctionParms)
+			while (SerializeExpr(ScriptIndex) != EX_EndFunctionParms)
 			{
 				// Params
 			}
@@ -840,6 +908,12 @@ void FKismetBytecodeDisassembler::ProcessCommon(int32& ScriptIndex, EExprToken O
 			Ar.Logf(TEXT("%s $%X: literal vector (%f,%f,%f)"), *Indents, (int32)Opcode, Vec.X, Vec.Y, Vec.Z);
 			break;
 		}
+	case EX_Vector3fConst:
+		{
+			FVector3f Vec = ReadFVECTOR(ScriptIndex);
+			Ar.Logf(TEXT("%s $%X: literal float vector (%f,%f,%f)"), *Indents, (int32)Opcode, Vec.X, Vec.Y, Vec.Z);
+			break;
+		}
 	case EX_TransformConst:
 		{
 			float RotX = ReadFLOAT(ScriptIndex);
@@ -858,7 +932,7 @@ void FKismetBytecodeDisassembler::ProcessCommon(int32& ScriptIndex, EExprToken O
 			UScriptStruct* Struct = ReadPointer<UScriptStruct>(ScriptIndex);
 			int32 SerializedSize = ReadINT(ScriptIndex);
 			Ar.Logf(TEXT("%s $%X: literal struct %s (serialized size: %d)"), *Indents, (int32)Opcode, *Struct->GetName(), SerializedSize);
-			while( SerializeExpr(ScriptIndex) != EX_EndStructConst )
+			while (SerializeExpr(ScriptIndex) != EX_EndStructConst)
 			{
 				// struct contents
 			}
@@ -868,7 +942,7 @@ void FKismetBytecodeDisassembler::ProcessCommon(int32& ScriptIndex, EExprToken O
 		{
  			Ar.Logf(TEXT("%s $%X: set array"), *Indents, (int32)Opcode);
 			SerializeExpr(ScriptIndex);
- 			while( SerializeExpr(ScriptIndex) != EX_EndArray)
+ 			while (SerializeExpr(ScriptIndex) != EX_EndArray)
  			{
  				// Array contents
  			}
@@ -1063,8 +1137,8 @@ void FKismetBytecodeDisassembler::ProcessCommon(int32& ScriptIndex, EExprToken O
 		}
 	case EX_SwitchValue:
 		{
-			const auto NumCases = ReadWORD(ScriptIndex);
-			const auto AfterSkip = ReadSkipCount(ScriptIndex);
+			const uint16 NumCases = ReadWORD(ScriptIndex);
+			const CodeSkipSizeType AfterSkip = ReadSkipCount(ScriptIndex);
 
 			Ar.Logf(TEXT("%s $%X: Switch Value %d cases, end in 0x%X"), *Indents, (int32)Opcode, NumCases, AfterSkip);
 			AddIndent();
@@ -1075,7 +1149,7 @@ void FKismetBytecodeDisassembler::ProcessCommon(int32& ScriptIndex, EExprToken O
 			{
 				Ar.Logf(TEXT("%s [%d] Case Index (label: 0x%X):"), *Indents, CaseIndex, ScriptIndex);
 				SerializeExpr(ScriptIndex);	// case index value term
-				const auto OffsetToNextCase = ReadSkipCount(ScriptIndex);
+				const CodeSkipSizeType OffsetToNextCase = ReadSkipCount(ScriptIndex);
 				Ar.Logf(TEXT("%s [%d] Offset to the next case: 0x%X"), *Indents, CaseIndex, OffsetToNextCase);
 				Ar.Logf(TEXT("%s [%d] Case Result:"), *Indents, CaseIndex);
 				SerializeExpr(ScriptIndex);	// case term
diff --git a/Engine/Source/Developer/ScriptDisassembler/Public/ScriptDisassembler.h b/Engine/Source/Developer/ScriptDisassembler/Public/ScriptDisassembler.h
--- a/Engine/Source/Developer/ScriptDisassembler/Public/ScriptDisassembler.h
+++ b/Engine/Source/Developer/ScriptDisassembler/Public/ScriptDisassembler.h
@@ -58,7 +58,6 @@ private:
 	FString ReadString16(int32& ScriptIndex);
 
 	EExprToken SerializeExpr(int32& ScriptIndex);
-	void ProcessCastByte(int32 CastType, int32& ScriptIndex);
 	void ProcessCommon(int32& ScriptIndex, EExprToken Opcode);
 
 	void InitTables();
diff --git a/Engine/Source/Editor/AnimGraph/Private/AnimBlueprintCompiler.cpp b/Engine/Source/Editor/AnimGraph/Private/AnimBlueprintCompiler.cpp
--- a/Engine/Source/Editor/AnimGraph/Private/AnimBlueprintCompiler.cpp
+++ b/Engine/Source/Editor/AnimGraph/Private/AnimBlueprintCompiler.cpp
@@ -32,6 +32,7 @@
 #include "AnimBlueprintVariableCreationContext.h"
 #include "Animation/AnimSubsystemInstance.h"
 #include "AnimBlueprintExtension.h"
+#include "UObject/FrameworkObjectVersion.h"
 #include "UObject/UE5MainStreamObjectVersion.h"
 
 #define LOCTEXT_NAMESPACE "AnimBlueprintCompiler"
@@ -2060,6 +2061,17 @@ const IAnimBlueprintCompilationContext::FFoldedPropertyRecord* FAnimBlueprintCom
 	return nullptr;
 }
 
+void FAnimBlueprintCompilerContext::PreCompileUpdateBlueprintOnLoad(UBlueprint* BP)
+{
+	if (UAnimBlueprint* AnimBP = Cast<UAnimBlueprint>(BP))
+	{
+		if (AnimBP->GetLinkerCustomVersion(FFrameworkObjectVersion::GUID) < FFrameworkObjectVersion::AnimBlueprintSubgraphFix)
+		{
+			AnimationEditorUtils::RegenerateSubGraphArrays(AnimBP);
+		}
+	}
+}
+
 //////////////////////////////////////////////////////////////////////////
 
 #undef LOCTEXT_NAMESPACE
diff --git a/Engine/Source/Editor/AnimGraph/Private/AnimBlueprintCompiler.h b/Engine/Source/Editor/AnimGraph/Private/AnimBlueprintCompiler.h
--- a/Engine/Source/Editor/AnimGraph/Private/AnimBlueprintCompiler.h
+++ b/Engine/Source/Editor/AnimGraph/Private/AnimBlueprintCompiler.h
@@ -74,6 +74,7 @@ protected:
 	virtual void SetCalculatedMetaDataAndFlags(UFunction* Function, UK2Node_FunctionEntry* EntryNode, const UEdGraphSchema_K2* Schema ) override;
 	virtual bool ShouldForceKeepNode(const UEdGraphNode* Node) const override;
 	virtual void PostExpansionStep(const UEdGraph* Graph) override;
+	virtual void PreCompileUpdateBlueprintOnLoad(UBlueprint* BP) override;
 	// End of FKismetCompilerContext interface
 
 protected:
diff --git a/Engine/Source/Editor/AnimGraph/Private/K2Node_TransitionRuleGetter.cpp b/Engine/Source/Editor/AnimGraph/Private/K2Node_TransitionRuleGetter.cpp
--- a/Engine/Source/Editor/AnimGraph/Private/K2Node_TransitionRuleGetter.cpp
+++ b/Engine/Source/Editor/AnimGraph/Private/K2Node_TransitionRuleGetter.cpp
@@ -26,7 +26,11 @@ UK2Node_TransitionRuleGetter::UK2Node_TransitionRuleGetter(const FObjectInitiali
 
 void UK2Node_TransitionRuleGetter::AllocateDefaultPins()
 {
+#if ENABLE_BLUEPRINT_REAL_NUMBERS
+	UEdGraphPin* OutputPin = CreatePin(EGPD_Output, UEdGraphSchema_K2::PC_Real, UEdGraphSchema_K2::PC_Float, TEXT("Output"));
+#else
 	UEdGraphPin* OutputPin = CreatePin(EGPD_Output, UEdGraphSchema_K2::PC_Float, TEXT("Output"));
+#endif
 	OutputPin->PinFriendlyName = GetFriendlyName(GetterType);
 
 	PreloadObject(AssociatedAnimAssetPlayerNode);
diff --git a/Engine/Source/Editor/BlueprintGraph/Classes/BlueprintTypePromotion.h b/Engine/Source/Editor/BlueprintGraph/Classes/BlueprintTypePromotion.h
--- a/Engine/Source/Editor/BlueprintGraph/Classes/BlueprintTypePromotion.h
+++ b/Engine/Source/Editor/BlueprintGraph/Classes/BlueprintTypePromotion.h
@@ -121,7 +121,7 @@ public:
 	*
 	* @return Const pointer to the primitive promotion table
 	*/
-	static const TMap<FName, TArray<FName>>* const GetPrimativePromotionTable();
+	static const TMap<FName, TArray<FName>>* const GetPrimitivePromotionTable();
 	
 	/**
 	* Get a pointer to an array of available promotion types to a given pin type
diff --git a/Engine/Source/Editor/BlueprintGraph/Classes/EdGraphSchema_K2.h b/Engine/Source/Editor/BlueprintGraph/Classes/EdGraphSchema_K2.h
--- a/Engine/Source/Editor/BlueprintGraph/Classes/EdGraphSchema_K2.h
+++ b/Engine/Source/Editor/BlueprintGraph/Classes/EdGraphSchema_K2.h
@@ -331,6 +331,9 @@ class BLUEPRINTGRAPH_API UEdGraphSchema_K2 : public UEdGraphSchema
 	static const FName PC_Int64;
 	static const FName PC_Float;
 	static const FName PC_Double;
+#if ENABLE_BLUEPRINT_REAL_NUMBERS
+	static const FName PC_Real;
+#endif
 	static const FName PC_Name;
 	static const FName PC_Delegate;    // SubCategoryObject is the UFunction of the delegate signature
 	static const FName PC_MCDelegate;  // SubCategoryObject is the UFunction of the delegate signature
@@ -1186,6 +1189,7 @@ private:
 	bool DoesFunctionHaveOutParameters( const UFunction* Function ) const;
 
 	static const UScriptStruct* VectorStruct;
+	static const UScriptStruct* Vector3fStruct;
 	static const UScriptStruct* RotatorStruct;
 	static const UScriptStruct* TransformStruct;
 	static const UScriptStruct* LinearColorStruct;
diff --git a/Engine/Source/Editor/BlueprintGraph/Classes/K2Node_AssignmentStatement.h b/Engine/Source/Editor/BlueprintGraph/Classes/K2Node_AssignmentStatement.h
--- a/Engine/Source/Editor/BlueprintGraph/Classes/K2Node_AssignmentStatement.h
+++ b/Engine/Source/Editor/BlueprintGraph/Classes/K2Node_AssignmentStatement.h
@@ -26,7 +26,7 @@ class UK2Node_AssignmentStatement : public UK2Node
 	virtual void AllocateDefaultPins() override;
 	virtual FText GetTooltipText() const override;
 	virtual FText GetNodeTitle(ENodeTitleType::Type TitleType) const override;
-	virtual bool IsCompatibleWithGraph(UEdGraph const* TargetGraph) const override;
+	virtual bool IsCompatibleWithGraph(const UEdGraph* TargetGraph) const override;
 	virtual bool CanPasteHere(const UEdGraph* TargetGraph) const override;
 	//~ End UEdGraphNode Interface
 
diff --git a/Engine/Source/Editor/BlueprintGraph/Private/BlueprintTypePromotion.cpp b/Engine/Source/Editor/BlueprintGraph/Private/BlueprintTypePromotion.cpp
--- a/Engine/Source/Editor/BlueprintGraph/Private/BlueprintTypePromotion.cpp
+++ b/Engine/Source/Editor/BlueprintGraph/Private/BlueprintTypePromotion.cpp
@@ -71,16 +71,25 @@ void FTypePromotion::CreatePromotionTable()
 {
 	PromotionTable =
 	{
+
+#if ENABLE_BLUEPRINT_REAL_NUMBERS
+		// Type_X...						Can be promoted to...
+		{ UEdGraphSchema_K2::PC_Int,		{ UEdGraphSchema_K2::PC_Real, UEdGraphSchema_K2::PC_Int64 } },
+		{ UEdGraphSchema_K2::PC_Byte,		{ UEdGraphSchema_K2::PC_Real, UEdGraphSchema_K2::PC_Int, UEdGraphSchema_K2::PC_Int64 } },
+		{ UEdGraphSchema_K2::PC_Real,		{ UEdGraphSchema_K2::PC_Int64 } },
+		{ UEdGraphSchema_K2::PC_Wildcard,	{ UEdGraphSchema_K2::PC_Int, UEdGraphSchema_K2::PC_Int64, UEdGraphSchema_K2::PC_Real, UEdGraphSchema_K2::PC_Byte, UEdGraphSchema_K2::PC_Boolean } },
+#else
 		// Type_X...						Can be promoted to...
 		{ UEdGraphSchema_K2::PC_Int,		{ UEdGraphSchema_K2::PC_Float, UEdGraphSchema_K2::PC_Double, UEdGraphSchema_K2::PC_Int64 } },
 		{ UEdGraphSchema_K2::PC_Byte,		{ UEdGraphSchema_K2::PC_Float, UEdGraphSchema_K2::PC_Int, UEdGraphSchema_K2::PC_Int64, UEdGraphSchema_K2::PC_Double } },
 		{ UEdGraphSchema_K2::PC_Float,		{ UEdGraphSchema_K2::PC_Double, UEdGraphSchema_K2::PC_Int64 } },
 		{ UEdGraphSchema_K2::PC_Double,		{ UEdGraphSchema_K2::PC_Int64 } },
 		{ UEdGraphSchema_K2::PC_Wildcard,	{ UEdGraphSchema_K2::PC_Int, UEdGraphSchema_K2::PC_Int64, UEdGraphSchema_K2::PC_Float, UEdGraphSchema_K2::PC_Double, UEdGraphSchema_K2::PC_Byte, UEdGraphSchema_K2::PC_Boolean } },
+#endif
 	};
 }
 
-const TMap<FName, TArray<FName>>* const FTypePromotion::GetPrimativePromotionTable()
+const TMap<FName, TArray<FName>>* const FTypePromotion::GetPrimitivePromotionTable()
 {
 	if (Instance)
 	{
@@ -92,7 +101,7 @@ const TMap<FName, TArray<FName>>* const FTypePromotion::GetPrimativePromotionTab
 
 const TArray<FName>* FTypePromotion::GetAvailablePrimitivePromotions(const FEdGraphPinType& Type)
 {
-	const TMap<FName, TArray<FName>>* PromoTable = GetPrimativePromotionTable();
+	const TMap<FName, TArray<FName>>* PromoTable = GetPrimitivePromotionTable();
 
 	return PromoTable->Find(Type.PinCategory);
 }
diff --git a/Engine/Source/Editor/BlueprintGraph/Private/CallFunctionHandler.cpp b/Engine/Source/Editor/BlueprintGraph/Private/CallFunctionHandler.cpp
--- a/Engine/Source/Editor/BlueprintGraph/Private/CallFunctionHandler.cpp
+++ b/Engine/Source/Editor/BlueprintGraph/Private/CallFunctionHandler.cpp
@@ -15,6 +15,7 @@
 
 #include "EdGraphUtilities.h"
 #include "Engine/BlueprintGeneratedClass.h"
+#include "KismetCastingUtils.h"
 #include "KismetCompiler.h"
 #include "Net/Core/PushModel/PushModelMacros.h"
 #include "PushModelHelpers.h"
@@ -147,9 +148,25 @@ void FKCHandler_CallFunction::CreateFunctionCallStatement(FKismetFunctionContext
 
 		// Grab the special case structs that use their own literal path
 		UScriptStruct* VectorStruct = TBaseStructure<FVector>::Get();
+		UScriptStruct* Vector3fStruct = TVariantStructure<FVector3f>::Get();
 		UScriptStruct* RotatorStruct = TBaseStructure<FRotator>::Get();
 		UScriptStruct* TransformStruct = TBaseStructure<FTransform>::Get();
 
+#if ENABLE_BLUEPRINT_REAL_NUMBERS
+		// If a function parameter needs an implicit double<->float cast *and* it's a non-const reference,
+		// then we need to copy the value of the casted temporary *back* to its source.
+		// 
+		// Just to illustrate the scenario, take the following example in pseudocode:
+		// 
+		//     double Input = 2.0
+		//     float CastedInput = (float)Input					; Narrowing conversion needed for function input
+		//     NativeFunctionWithReferenceParam(CastedInput)	; CastedInput has possibly changed since this function takes a float&
+		//     Input = (double)CastedInput						; Now we need to propagate that change back to Input
+		//
+		using CastEntry = TTuple<FBPTerminal*, FBPTerminal*, EKismetCompiledStatementType>;
+		TArray<CastEntry> ModifiedCastInputs;
+#endif
+
 		// Check each property
 		bool bMatchedAllParams = true;
 		for (TFieldIterator<FProperty> It(Function); It && (It->PropertyFlags & CPF_Parm); ++It)
@@ -178,6 +195,7 @@ void FKCHandler_CallFunction::CreateFunctionCallStatement(FKismetFunctionContext
 								{
 									UScriptStruct* Struct = StructProperty->Struct;
 									if( Struct != VectorStruct
+										&& Struct != Vector3fStruct
 										&& Struct != RotatorStruct
 										&& Struct != TransformStruct )
 									{
@@ -237,6 +255,36 @@ void FKCHandler_CallFunction::CreateFunctionCallStatement(FKismetFunctionContext
 									RHSTerm = *InterfaceTerm;
 								}
 
+#if ENABLE_BLUEPRINT_REAL_NUMBERS
+								{
+									using namespace UE::KismetCompiler;
+
+									TOptional<TPair<FBPTerminal*, EKismetCompiledStatementType>> ImplicitCastEntry =
+										CastingUtils::InsertImplicitCastStatement(Context, PinMatch, RHSTerm);
+
+									if (ImplicitCastEntry)
+									{
+										FBPTerminal* CastTerm = ImplicitCastEntry->Get<0>();
+
+										bool bIsNonConstReference =
+											Property->HasAllPropertyFlags(CPF_OutParm | CPF_ReferenceParm) &&
+											!Property->HasAllPropertyFlags(CPF_ConstParm);
+
+										if (bIsNonConstReference)
+										{
+											TOptional<EKismetCompiledStatementType> CastType = 
+												CastingUtils::GetInverseCastStatement(ImplicitCastEntry->Get<1>());
+
+											check(CastType.IsSet());
+
+											ModifiedCastInputs.Add(CastEntry{RHSTerm, CastTerm, *CastType});
+										}
+
+										RHSTerm = CastTerm;
+									}
+								}
+#endif
+
 								int32 ParameterIndex = RHSTerms.Add(RHSTerm);
 
 								if (PinMatch == LatentInfoPin)
@@ -410,10 +458,6 @@ void FKCHandler_CallFunction::CreateFunctionCallStatement(FKismetFunctionContext
 				pSrcEventNode = CompilerContext.CallsIntoUbergraph.Find(Node);
 			}
 
-			bool bInlineEventCall = false;
-			bool bEmitInstrumentPushState = false;
-			FName EventName = NAME_None;
-
 			// Iterate over all the contexts this functions needs to be called on, and emit a call function statement for each
 			FBlueprintCompiledStatement* LatentStatement = nullptr;
 			for (FBPTerminal* Target : ContextTerms)
@@ -464,6 +508,27 @@ void FKCHandler_CallFunction::CreateFunctionCallStatement(FKismetFunctionContext
 				}
 			}
 
+#if ENABLE_BLUEPRINT_REAL_NUMBERS
+			{
+				using namespace UE::KismetCompiler;
+
+				for (const auto& It : ModifiedCastInputs)
+				{
+					FBPTerminal* LocalLHSTerm = It.Get<0>();
+					FBPTerminal* LocalRHSTerm = It.Get<1>();
+					EKismetCompiledStatementType LocalCastType = It.Get<2>();
+
+					check(LocalLHSTerm);
+					check(LocalRHSTerm);
+
+					FBlueprintCompiledStatement& CastStatement = Context.AppendStatementForNode(Node);
+					CastStatement.LHS = LocalLHSTerm;
+					CastStatement.Type = LocalCastType;
+					CastStatement.RHS.Add(LocalRHSTerm);
+				}
+			}
+#endif
+
 			// Create the exit from this node if there is one
 			if (bIsLatent)
 			{
@@ -546,6 +611,8 @@ UClass* FKCHandler_CallFunction::GetTrueCallingClass(FKismetFunctionContext& Con
 
 void FKCHandler_CallFunction::RegisterNets(FKismetFunctionContext& Context, UEdGraphNode* Node)
 {
+	check(Node);
+
 	if (UFunction* Function = FindFunction(Context, Node))
 	{
 		TArray<FName> DefaultToSelfParamNames;
@@ -600,6 +667,8 @@ void FKCHandler_CallFunction::RegisterNets(FKismetFunctionContext& Context, UEdG
 
 	for (UEdGraphPin* Pin : Node->Pins)
 	{
+		check(Pin);
+
 		if ((Pin->Direction != EGPD_Input) || (Pin->LinkedTo.Num() == 0))
 		{
 			continue;
diff --git a/Engine/Source/Editor/BlueprintGraph/Private/EdGraphSchema_K2.cpp b/Engine/Source/Editor/BlueprintGraph/Private/EdGraphSchema_K2.cpp
--- a/Engine/Source/Editor/BlueprintGraph/Private/EdGraphSchema_K2.cpp
+++ b/Engine/Source/Editor/BlueprintGraph/Private/EdGraphSchema_K2.cpp
@@ -554,8 +554,12 @@ void UEdGraphSchema_K2::FPinTypeTreeInfo::Init(const FText& InFriendlyName, cons
 
 	FriendlyName = InFriendlyName;
 	Tooltip = InTooltip;
-	PinType.PinCategory = (CategoryName == PC_Enum? PC_Byte : CategoryName);
+	PinType.PinCategory = (CategoryName == PC_Enum ? PC_Byte : CategoryName);
+#if ENABLE_BLUEPRINT_REAL_NUMBERS
+	PinType.PinSubCategory = (CategoryName == PC_Real ? PC_Double : NAME_None);
+#else
 	PinType.PinSubCategory = NAME_None;
+#endif
 	PinType.PinSubCategoryObject = nullptr;
 
 	bReadOnly = bInReadOnly;
@@ -642,6 +646,9 @@ const FName UEdGraphSchema_K2::PC_Int(TEXT("int"));
 const FName UEdGraphSchema_K2::PC_Int64(TEXT("int64"));
 const FName UEdGraphSchema_K2::PC_Float(TEXT("float"));
 const FName UEdGraphSchema_K2::PC_Double(TEXT("double"));
+#if ENABLE_BLUEPRINT_REAL_NUMBERS
+const FName UEdGraphSchema_K2::PC_Real(TEXT("real"));
+#endif
 const FName UEdGraphSchema_K2::PC_Name(TEXT("name"));
 const FName UEdGraphSchema_K2::PC_Delegate(TEXT("delegate"));
 const FName UEdGraphSchema_K2::PC_MCDelegate(TEXT("mcdelegate"));
@@ -688,6 +695,7 @@ const FText UEdGraphSchema_K2::VR_DefaultCategory(LOCTEXT("Default", "Default"))
 const int32 UEdGraphSchema_K2::AG_LevelReference = 100;
 
 const UScriptStruct* UEdGraphSchema_K2::VectorStruct = nullptr;
+const UScriptStruct* UEdGraphSchema_K2::Vector3fStruct = nullptr;
 const UScriptStruct* UEdGraphSchema_K2::RotatorStruct = nullptr;
 const UScriptStruct* UEdGraphSchema_K2::TransformStruct = nullptr;
 const UScriptStruct* UEdGraphSchema_K2::LinearColorStruct = nullptr;
@@ -705,6 +713,7 @@ UEdGraphSchema_K2::UEdGraphSchema_K2(const FObjectInitializer& ObjectInitializer
 	if (VectorStruct == nullptr)
 	{
 		VectorStruct = TBaseStructure<FVector>::Get();
+		Vector3fStruct = TVariantStructure<FVector3f>::Get();
 		RotatorStruct = TBaseStructure<FRotator>::Get();
 		TransformStruct = TBaseStructure<FTransform>::Get();
 		LinearColorStruct = TBaseStructure<FLinearColor>::Get();
@@ -1415,6 +1424,7 @@ bool UEdGraphSchema_K2::PinDefaultValueIsEditable(const UEdGraphPin& InGraphPin)
 		// See FNodeFactory::CreatePinWidget for justification of the above statement!
 		UObject const& SubCategoryObject = *InGraphPin.PinType.PinSubCategoryObject;
 		return &SubCategoryObject == VectorStruct 
+			|| &SubCategoryObject == Vector3fStruct
 			|| &SubCategoryObject == RotatorStruct
 			|| &SubCategoryObject == TransformStruct
 			|| &SubCategoryObject == LinearColorStruct
@@ -1432,6 +1442,7 @@ bool UEdGraphSchema_K2::PinHasCustomDefaultFormat(const UEdGraphPin& InGraphPin)
 		// Some struct types have custom formats for default value for historical reasons
 		UObject const& SubCategoryObject = *InGraphPin.PinType.PinSubCategoryObject;
 		return &SubCategoryObject == VectorStruct
+			|| &SubCategoryObject == Vector3fStruct
 			|| &SubCategoryObject == RotatorStruct
 			|| &SubCategoryObject == TransformStruct
 			|| &SubCategoryObject == LinearColorStruct;
@@ -2881,6 +2892,12 @@ FLinearColor UEdGraphSchema_K2::GetPinTypeColor(const FEdGraphPinType& PinType)
 	{
 		return Settings->InterfacePinTypeColor;
 	}
+#if ENABLE_BLUEPRINT_REAL_NUMBERS
+	else if (TypeName == PC_Real)
+	{
+		return Settings->RealPinTypeColor;
+	}
+#else
 	else if (TypeName == PC_Float)
 	{
 		return Settings->FloatPinTypeColor;
@@ -2889,6 +2906,7 @@ FLinearColor UEdGraphSchema_K2::GetPinTypeColor(const FEdGraphPinType& PinType)
 	{
 		return Settings->DoublePinTypeColor;
 	}
+#endif
 	else if (TypeName == PC_Boolean)
 	{
 		return Settings->BooleanPinTypeColor;
@@ -2907,7 +2925,7 @@ FLinearColor UEdGraphSchema_K2::GetPinTypeColor(const FEdGraphPinType& PinType)
 	}
 	else if (TypeName == PC_Struct)
 	{
-		if (PinType.PinSubCategoryObject == VectorStruct)
+		if ((PinType.PinSubCategoryObject == VectorStruct) || (PinType.PinSubCategoryObject == Vector3fStruct))
 		{
 			// vector
 			return Settings->VectorPinTypeColor;
@@ -3332,6 +3350,18 @@ bool UEdGraphSchema_K2::GetPropertyCategoryInfo(const FProperty* TestProperty, F
 			}
 		}
 	}
+#if ENABLE_BLUEPRINT_REAL_NUMBERS
+	else if (TestProperty->IsA<FFloatProperty>())
+	{
+		OutCategory = PC_Real;
+		OutSubCategory = PC_Float;
+	}
+	else if (TestProperty->IsA<FDoubleProperty>())
+	{
+		OutCategory = PC_Real;
+		OutSubCategory = PC_Double;
+	}
+#else
 	else if (TestProperty->IsA<FFloatProperty>())
 	{
 		OutCategory = PC_Float;
@@ -3340,6 +3370,7 @@ bool UEdGraphSchema_K2::GetPropertyCategoryInfo(const FProperty* TestProperty, F
 	{
 		OutCategory = PC_Double;
 	}
+#endif
 	else if (TestProperty->IsA<FInt64Property>())
 	{
 		OutCategory = PC_Int64;
@@ -3635,8 +3666,12 @@ FText UEdGraphSchema_K2::GetCategoryText(const FName Category, const bool bForMe
 		CategoryDescriptions.Add(PC_Class, LOCTEXT("ClassCategory","Class Reference"));
 		CategoryDescriptions.Add(PC_Int, LOCTEXT("IntCategory", "Integer"));
 		CategoryDescriptions.Add(PC_Int64, LOCTEXT("Int64Category", "Integer64"));
-		CategoryDescriptions.Add(PC_Float, LOCTEXT("FloatCategory","Float"));
+#if ENABLE_BLUEPRINT_REAL_NUMBERS
+		CategoryDescriptions.Add(PC_Real, LOCTEXT("RealCategory", "Real"));
+#else
+		CategoryDescriptions.Add(PC_Float, LOCTEXT("FloatCategory", "Float"));
 		CategoryDescriptions.Add(PC_Double, LOCTEXT("DoubleCategory", "Double"));
+#endif
 		CategoryDescriptions.Add(PC_Name, LOCTEXT("NameCategory","Name"));
 		CategoryDescriptions.Add(PC_Delegate, LOCTEXT("DelegateCategory","Delegate"));
 		CategoryDescriptions.Add(PC_MCDelegate, LOCTEXT("MulticastDelegateCategory","Multicast Delegate"));
@@ -3697,7 +3732,7 @@ FText UEdGraphSchema_K2::TerminalTypeToText(const FName Category, const FName Su
 				Args.Add(TEXT("ObjectName"), FText::FromString(FName::NameToDisplayString(SubCategoryObjName, /*bIsBool =*/false)));
 
 				// Don't display the category for "well-known" struct types
-				if (Category == UEdGraphSchema_K2::PC_Struct && (SubCategoryObject == UEdGraphSchema_K2::VectorStruct || SubCategoryObject == UEdGraphSchema_K2::RotatorStruct || SubCategoryObject == UEdGraphSchema_K2::TransformStruct))
+				if (Category == UEdGraphSchema_K2::PC_Struct && (SubCategoryObject == UEdGraphSchema_K2::VectorStruct || SubCategoryObject == UEdGraphSchema_K2::Vector3fStruct  || SubCategoryObject == UEdGraphSchema_K2::RotatorStruct || SubCategoryObject == UEdGraphSchema_K2::TransformStruct))
 				{
 					PropertyText = FText::Format(LOCTEXT("ObjectAsTextWithoutCategory", "{ObjectName}"), Args);
 				}
@@ -3823,8 +3858,12 @@ void UEdGraphSchema_K2::GetVariableTypeTree(TArray< TSharedPtr<FPinTypeTreeInfo>
 
 	if (!bIndexTypesOnly)
 	{
+#if ENABLE_BLUEPRINT_REAL_NUMBERS
+		TypeTree.Add(MakeShareable(new FPinTypeTreeInfo(GetCategoryText(PC_Real, true), PC_Real, this, LOCTEXT("RealType", "Floating point number"))));
+#else
 		TypeTree.Add(MakeShareable(new FPinTypeTreeInfo(GetCategoryText(PC_Float, true), PC_Float, this, LOCTEXT("FloatType", "Floating point number"))));
 		TypeTree.Add(MakeShareable(new FPinTypeTreeInfo(GetCategoryText(PC_Double, true), PC_Double, this, LOCTEXT("DoubleType", "64 bit floating point number"))));
+#endif
 		TypeTree.Add(MakeShareable(new FPinTypeTreeInfo(GetCategoryText(PC_Name, true), PC_Name, this, LOCTEXT("NameType", "A text name"))));
 		TypeTree.Add(MakeShareable(new FPinTypeTreeInfo(GetCategoryText(PC_String, true), PC_String, this, LOCTEXT("StringType", "A text string"))));
 		TypeTree.Add(MakeShareable(new FPinTypeTreeInfo(GetCategoryText(PC_Text, true), PC_Text, this, LOCTEXT("TextType", "A localizable text string"))));
@@ -4097,6 +4136,18 @@ bool UEdGraphSchema_K2::DefaultValueSimpleValidation(const FEdGraphPinType& PinT
 			}
 		}
 	}
+#if ENABLE_BLUEPRINT_REAL_NUMBERS
+	else if (PinCategory == PC_Real)
+	{
+		if (!NewDefaultValue.IsEmpty())
+		{
+			if (!FDefaultValueHelper::IsStringValidFloat(NewDefaultValue))
+			{
+				DVSV_RETURN_MSG(TEXT("Expected a valid number for a real property"));
+			}
+		}
+	}
+#else
 	else if (PinCategory == PC_Float)
 	{
 		if (!NewDefaultValue.IsEmpty())
@@ -4117,6 +4168,7 @@ bool UEdGraphSchema_K2::DefaultValueSimpleValidation(const FEdGraphPinType& PinT
 			}
 		}
 	}
+#endif
 	else if (PinCategory == PC_Int)
 	{
 		if (!NewDefaultValue.IsEmpty())
@@ -4219,7 +4271,7 @@ bool UEdGraphSchema_K2::DefaultValueSimpleValidation(const FEdGraphPinType& PinT
 		}
 		else if (!NewDefaultValue.IsEmpty())
 		{
-			if (StructType == VectorStruct)
+			if ((StructType == VectorStruct) || (StructType == Vector3fStruct))
 			{
 				if (!FDefaultValueHelper::IsStringValidVector(NewDefaultValue))
 				{
@@ -4276,8 +4328,8 @@ bool UEdGraphSchema_K2::DefaultValueSimpleValidation(const FEdGraphPinType& PinT
 
 bool UEdGraphSchema_K2::ArePinTypesCompatible(const FEdGraphPinType& Output, const FEdGraphPinType& Input, const UClass* CallingContext, bool bIgnoreArray /*= false*/) const
 {
-	if( !bIgnoreArray && 
-		( Output.ContainerType != Input.ContainerType ) && 
+	if (!bIgnoreArray && 
+		(Output.ContainerType != Input.ContainerType) && 
 		(Input.PinCategory != PC_Wildcard || Input.IsContainer()) && 
 		(Output.PinCategory != PC_Wildcard || Output.IsContainer()))
 	{
@@ -4285,6 +4337,10 @@ bool UEdGraphSchema_K2::ArePinTypesCompatible(const FEdGraphPinType& Output, con
 	}
 	else if (Output.PinCategory == Input.PinCategory)
 	{
+		bool bAreConvertibleVectorTypes =
+			((Output.PinSubCategoryObject == VectorStruct) && (Input.PinSubCategoryObject == Vector3fStruct)) ||
+			((Output.PinSubCategoryObject == Vector3fStruct) && (Input.PinSubCategoryObject == VectorStruct));
+
 		if ((Output.PinSubCategory == Input.PinSubCategory) 
 			&& (Output.PinSubCategoryObject == Input.PinSubCategoryObject)
 			&& (Output.PinSubCategoryMemberReference == Input.PinSubCategoryMemberReference))
@@ -4294,10 +4350,25 @@ bool UEdGraphSchema_K2::ArePinTypesCompatible(const FEdGraphPinType& Output, con
 				return 
 					Input.PinValueType.TerminalCategory == PC_Wildcard ||
 					Output.PinValueType.TerminalCategory == PC_Wildcard ||
+#if ENABLE_BLUEPRINT_REAL_NUMBERS
+					((Input.PinValueType.TerminalCategory == PC_Real) && (Output.PinValueType.TerminalCategory == PC_Real)) ||
+#endif
 					Input.PinValueType == Output.PinValueType;
 			}
 			return true;
 		}
+#if ENABLE_BLUEPRINT_REAL_NUMBERS
+		// Reals, whether they're actually a float or double, are always compatible.
+		// We'll insert an implicit conversion in the bytecode where necessary.
+		else if (Output.PinCategory == PC_Real)
+		{
+			return true;
+		}
+		else if (bAreConvertibleVectorTypes)
+		{
+			return true;
+		}
+#endif
 		else if (Output.PinCategory == PC_Interface)
 		{
 			UClass const* OutputClass = Cast<UClass const>(Output.PinSubCategoryObject.Get());
@@ -4772,6 +4843,14 @@ bool UEdGraphSchema_K2::DoesDefaultValueMatchAutogenerated(const UEdGraphPin& In
 	}
 	else if (!InPin.bUseBackwardsCompatForEmptyAutogeneratedValue)
 	{
+#if ENABLE_BLUEPRINT_REAL_NUMBERS
+		if (InPin.PinType.PinCategory == PC_Real)
+		{
+			const double AutogeneratedDouble = FCString::Atod(*InPin.AutogeneratedDefaultValue);
+			const double DefaultDouble = FCString::Atod(*InPin.DefaultValue);
+			return (AutogeneratedDouble == DefaultDouble);
+		}
+#else
 		if (InPin.PinType.PinCategory == PC_Float)
 		{
 			const float AutogeneratedFloat = FCString::Atof(*InPin.AutogeneratedDefaultValue);
@@ -4784,9 +4863,10 @@ bool UEdGraphSchema_K2::DoesDefaultValueMatchAutogenerated(const UEdGraphPin& In
 			const double DefaultDouble = FCString::Atod(*InPin.DefaultValue);
 			return (AutogeneratedDouble == DefaultDouble);
 		}
+#endif
 		else if (InPin.PinType.PinCategory == PC_Struct)
 		{
-			if (InPin.PinType.PinSubCategoryObject == VectorStruct)
+			if ((InPin.PinType.PinSubCategoryObject == VectorStruct) || (InPin.PinType.PinSubCategoryObject == Vector3fStruct))
 			{
 				FVector AutogeneratedVector = FVector::ZeroVector;
 				FVector DefaultVector = FVector::ZeroVector;
@@ -4999,6 +5079,12 @@ void UEdGraphSchema_K2::SetPinAutogeneratedDefaultValueBasedOnType(UEdGraphPin*
 			NewValue = TEXT("0");
 		}
 	}
+#if ENABLE_BLUEPRINT_REAL_NUMBERS
+	else if (Pin->PinType.PinCategory == PC_Real)
+	{
+		NewValue = TEXT("0.0");
+	}
+#else
 	else if (Pin->PinType.PinCategory == PC_Float)
 	{
 		// This is a slightly different format than is produced by PropertyValueToString, but changing it has backward compatibility issues
@@ -5008,6 +5094,7 @@ void UEdGraphSchema_K2::SetPinAutogeneratedDefaultValueBasedOnType(UEdGraphPin*
 	{
 		NewValue = TEXT("0.0");
 	}
+#endif
 	else if (Pin->PinType.PinCategory == PC_Boolean)
 	{
 		NewValue = TEXT("false");
@@ -5016,7 +5103,7 @@ void UEdGraphSchema_K2::SetPinAutogeneratedDefaultValueBasedOnType(UEdGraphPin*
 	{
 		NewValue = TEXT("None");
 	}
-	else if ((Pin->PinType.PinCategory == PC_Struct) && ((Pin->PinType.PinSubCategoryObject == VectorStruct) || (Pin->PinType.PinSubCategoryObject == RotatorStruct)))
+	else if ((Pin->PinType.PinCategory == PC_Struct) && ((Pin->PinType.PinSubCategoryObject == VectorStruct) || (Pin->PinType.PinSubCategoryObject == Vector3fStruct) || (Pin->PinType.PinSubCategoryObject == RotatorStruct)))
 	{
 		// This is a slightly different format than is produced by PropertyValueToString, but changing it has backward compatibility issues
 		NewValue = TEXT("0, 0, 0");
@@ -5155,6 +5242,7 @@ namespace FSetVariableByNameFunctionNames
 	static const FName SetSoftClassName(GET_FUNCTION_NAME_CHECKED(UKismetSystemLibrary, SetSoftClassPropertyByName));
 	static const FName SetNameName(GET_FUNCTION_NAME_CHECKED(UKismetSystemLibrary, SetNamePropertyByName));
 	static const FName SetVectorName(GET_FUNCTION_NAME_CHECKED(UKismetSystemLibrary, SetVectorPropertyByName));
+	static const FName SetVector3fName(GET_FUNCTION_NAME_CHECKED(UKismetSystemLibrary, SetVector3fPropertyByName));
 	static const FName SetRotatorName(GET_FUNCTION_NAME_CHECKED(UKismetSystemLibrary, SetRotatorPropertyByName));
 	static const FName SetLinearColorName(GET_FUNCTION_NAME_CHECKED(UKismetSystemLibrary, SetLinearColorPropertyByName));
 	static const FName SetColorName(GET_FUNCTION_NAME_CHECKED(UKismetSystemLibrary, SetColorPropertyByName));
@@ -5209,6 +5297,12 @@ UFunction* UEdGraphSchema_K2::FindSetVariableByNameFunction(const FEdGraphPinTyp
 	{
 		SetFunctionName = FSetVariableByNameFunctionNames::SetByteName;
 	}
+#if ENABLE_BLUEPRINT_REAL_NUMBERS
+	else if(PinType.PinCategory == UEdGraphSchema_K2::PC_Real)
+	{
+		SetFunctionName = FSetVariableByNameFunctionNames::SetDoubleName;
+	}
+#else
 	else if(PinType.PinCategory == UEdGraphSchema_K2::PC_Float)
 	{
 		SetFunctionName = FSetVariableByNameFunctionNames::SetFloatName;
@@ -5217,6 +5311,7 @@ UFunction* UEdGraphSchema_K2::FindSetVariableByNameFunction(const FEdGraphPinTyp
 	{
 		SetFunctionName = FSetVariableByNameFunctionNames::SetDoubleName;
 	}
+#endif
 	else if(PinType.PinCategory == UEdGraphSchema_K2::PC_Boolean)
 	{
 		SetFunctionName = FSetVariableByNameFunctionNames::SetBoolName;
@@ -5257,6 +5352,10 @@ UFunction* UEdGraphSchema_K2::FindSetVariableByNameFunction(const FEdGraphPinTyp
 	{
 		SetFunctionName = FSetVariableByNameFunctionNames::SetVectorName;
 	}
+	else if (PinType.PinCategory == UEdGraphSchema_K2::PC_Struct && PinType.PinSubCategoryObject == Vector3fStruct)
+	{
+		SetFunctionName = FSetVariableByNameFunctionNames::SetVector3fName;
+	}
 	else if(PinType.PinCategory == UEdGraphSchema_K2::PC_Struct && PinType.PinSubCategoryObject == RotatorStruct)
 	{
 		SetFunctionName = FSetVariableByNameFunctionNames::SetRotatorName;
diff --git a/Engine/Source/Editor/BlueprintGraph/Private/K2Node_AssignmentStatement.cpp b/Engine/Source/Editor/BlueprintGraph/Private/K2Node_AssignmentStatement.cpp
--- a/Engine/Source/Editor/BlueprintGraph/Private/K2Node_AssignmentStatement.cpp
+++ b/Engine/Source/Editor/BlueprintGraph/Private/K2Node_AssignmentStatement.cpp
@@ -3,6 +3,8 @@
 
 #include "K2Node_AssignmentStatement.h"
 #include "EdGraphSchema_K2.h"
+#include "EdGraphUtilities.h"
+#include "KismetCastingUtils.h"
 #include "KismetCompiler.h"
 #include "VariableSetHandler.h"
 #include "BlueprintNodeSpawner.h"
@@ -35,10 +37,50 @@ public:
 
 		if (VariablePin->LinkedTo.Num() == 0)
 		{
-			CompilerContext.MessageLog.Error(*LOCTEXT("NoVarriableConnected_Error", "A variable needs to be connected to @@").ToString(), VariablePin);
+			CompilerContext.MessageLog.Error(*LOCTEXT("NoVariableConnected_Error", "A variable needs to be connected to @@").ToString(), VariablePin);
+			return;
 		}
 
 		ValidateAndRegisterNetIfLiteral(Context, ValuePin);
+
+#if ENABLE_BLUEPRINT_REAL_NUMBERS
+		{
+			using namespace UE::KismetCompiler;
+
+			UEdGraphPin* VariablePinNet = FEdGraphUtilities::GetNetFromPin(VariablePin);
+			UEdGraphPin* ValuePinNet = FEdGraphUtilities::GetNetFromPin(ValuePin);
+
+			if (VariablePinNet && ValuePinNet)
+			{
+				TOptional<CastingUtils::StatementNamePair> ConversionType =
+					CastingUtils::GetFloatingPointConversionType(*ValuePinNet, *VariablePinNet);
+
+				if (ConversionType)
+				{
+
+					FString DescriptiveName = Node->GetName();
+
+					FString TerminalName = FString::Printf(TEXT("%s_%s_%s"),
+						*DescriptiveName,
+						*VariablePinNet->PinName.ToString(),
+						ConversionType->Get<1>());
+
+					FBPTerminal* NewTerm = Context.CreateLocalTerminal();
+					NewTerm->Name = TerminalName;
+					NewTerm->Type = VariablePinNet->PinType;
+					NewTerm->Source = Node;
+
+					EKismetCompiledStatementType CastType = ConversionType->Get<0>();
+
+					Context.ImplicitCastMap.Add(VariablePin, FImplicitCastParams{CastType, NewTerm, Node});
+				}
+			}
+			else
+			{
+				CompilerContext.MessageLog.Error(*LOCTEXT("NoVariableOrValueNets_Error", "Expected Variable and Value pins to have valid connections in @@").ToString(), Node);
+			}
+		}
+#endif
 	}
 
 
@@ -47,11 +89,18 @@ public:
 		UEdGraphPin* VariablePin = Node->FindPin(TEXT("Variable"));
 		UEdGraphPin* ValuePin = Node->FindPin(TEXT("Value"));
 
+		check(VariablePin);
+		check(VariablePin->LinkedTo.Num() == 1);
+		check(ValuePin);
+
 		InnerAssignment(Context, Node, VariablePin, ValuePin);
 
 		// Generate the output impulse from this node
 		GenerateSimpleThenGoto(Context, *Node);
 	}
+
+protected:
+	virtual bool UsesVariablePinAsKey() const override { return true; }
 };
 
 
@@ -85,13 +134,13 @@ FText UK2Node_AssignmentStatement::GetNodeTitle(ENodeTitleType::Type TitleType)
 	return LOCTEXT("Assign", "Assign");
 }
 
-bool UK2Node_AssignmentStatement::IsCompatibleWithGraph(UEdGraph const* TargetGraph) const
+bool UK2Node_AssignmentStatement::IsCompatibleWithGraph(const UEdGraph* TargetGraph) const
 {
 	bool bIsCompatible = Super::IsCompatibleWithGraph(TargetGraph);
 	if (bIsCompatible)
 	{
-		EGraphType const GraphType = TargetGraph->GetSchema()->GetGraphType(TargetGraph);
-		bIsCompatible = GraphType == GT_Macro;
+		const EGraphType GraphType = TargetGraph->GetSchema()->GetGraphType(TargetGraph);
+		bIsCompatible = (GraphType == GT_Macro);
 	}
 
 	return bIsCompatible;
diff --git a/Engine/Source/Editor/BlueprintGraph/Private/K2Node_ConstructObjectFromClass.cpp b/Engine/Source/Editor/BlueprintGraph/Private/K2Node_ConstructObjectFromClass.cpp
--- a/Engine/Source/Editor/BlueprintGraph/Private/K2Node_ConstructObjectFromClass.cpp
+++ b/Engine/Source/Editor/BlueprintGraph/Private/K2Node_ConstructObjectFromClass.cpp
@@ -11,16 +11,26 @@
 
 struct FK2Node_ConstructObjectFromClassHelper
 {
-	static FName WorldContextPinName;
-	static FName ClassPinName;
-	static FName OuterPinFriendlyName;
-	static FName OuterPinName;
-};
+	static FName GetWorldContextPinName()
+	{
+		return TEXT("WorldContextObject");
+	}
+
+	static FName GetClassPinName()
+	{
+		return TEXT("Class");
+	}
 
-FName FK2Node_ConstructObjectFromClassHelper::WorldContextPinName(TEXT("WorldContextObject"));
-FName FK2Node_ConstructObjectFromClassHelper::ClassPinName(TEXT("Class"));
-FName FK2Node_ConstructObjectFromClassHelper::OuterPinFriendlyName(TEXT("Outer"));
-FName FK2Node_ConstructObjectFromClassHelper::OuterPinName(UEdGraphSchema_K2::PN_Self);
+	static FName GetOuterPinFriendlyName()
+	{
+		return TEXT("Outer");
+	}
+
+	static FName GetOuterPinName()
+	{
+		return UEdGraphSchema_K2::PN_Self;
+	}
+};
 
 #define LOCTEXT_NAMESPACE "K2Node_ConstructObjectFromClass"
 
@@ -51,11 +61,11 @@ void UK2Node_ConstructObjectFromClass::AllocateDefaultPins()
 	// If required add the world context pin
 	if (UseWorldContext())
 	{
-		CreatePin(EGPD_Input, UEdGraphSchema_K2::PC_Object, UObject::StaticClass(), FK2Node_ConstructObjectFromClassHelper::WorldContextPinName);
+		CreatePin(EGPD_Input, UEdGraphSchema_K2::PC_Object, UObject::StaticClass(), FK2Node_ConstructObjectFromClassHelper::GetWorldContextPinName());
 	}
 
 	// Add blueprint pin
-	UEdGraphPin* ClassPin = CreatePin(EGPD_Input, UEdGraphSchema_K2::PC_Class, GetClassPinBaseClass(), FK2Node_ConstructObjectFromClassHelper::ClassPinName);
+	UEdGraphPin* ClassPin = CreatePin(EGPD_Input, UEdGraphSchema_K2::PC_Class, GetClassPinBaseClass(), FK2Node_ConstructObjectFromClassHelper::GetClassPinName());
 	
 	// Result pin
 	UEdGraphPin* ResultPin = CreatePin(EGPD_Output, UEdGraphSchema_K2::PC_Object, GetClassPinBaseClass(), UEdGraphSchema_K2::PN_ReturnValue);
@@ -63,7 +73,7 @@ void UK2Node_ConstructObjectFromClass::AllocateDefaultPins()
 	if (UseOuter())
 	{
 		UEdGraphPin* OuterPin = CreatePin(EGPD_Input, UEdGraphSchema_K2::PC_Object, UObject::StaticClass(), UEdGraphSchema_K2::PN_Self);
-		OuterPin->PinFriendlyName = FText::FromName(FK2Node_ConstructObjectFromClassHelper::OuterPinFriendlyName);
+		OuterPin->PinFriendlyName = FText::FromName(FK2Node_ConstructObjectFromClassHelper::GetOuterPinFriendlyName());
 	}
 
 	Super::AllocateDefaultPins();
@@ -71,7 +81,7 @@ void UK2Node_ConstructObjectFromClass::AllocateDefaultPins()
 
 UEdGraphPin* UK2Node_ConstructObjectFromClass::GetOuterPin() const
 {
-	UEdGraphPin* Pin = FindPin(FK2Node_ConstructObjectFromClassHelper::OuterPinName);
+	UEdGraphPin* Pin = FindPin(FK2Node_ConstructObjectFromClassHelper::GetOuterPinName());
 	ensure(nullptr == Pin || Pin->Direction == EGPD_Input);
 	return Pin;
 }
@@ -174,8 +184,8 @@ UK2Node::ERedirectType UK2Node_ConstructObjectFromClass::DoPinsMatchForReconstru
 {
 	// the name of the outer pin was changed and its friendly name was updated to match
 	// the legacy naming. Use this to identify the change
-	if (NewPin->PinName == FK2Node_ConstructObjectFromClassHelper::OuterPinName &&
-		OldPin->PinName == FK2Node_ConstructObjectFromClassHelper::OuterPinFriendlyName)
+	if (NewPin->PinName == FK2Node_ConstructObjectFromClassHelper::GetOuterPinName() &&
+		OldPin->PinName == FK2Node_ConstructObjectFromClassHelper::GetOuterPinFriendlyName())
 	{
 		return ERedirectType_Name;
 	}
@@ -208,9 +218,9 @@ bool UK2Node_ConstructObjectFromClass::IsSpawnVarPin(UEdGraphPin* Pin) const
 	return(	Pin->PinName != UEdGraphSchema_K2::PN_Execute &&
 			Pin->PinName != UEdGraphSchema_K2::PN_Then &&
 			Pin->PinName != UEdGraphSchema_K2::PN_ReturnValue &&
-			Pin->PinName != FK2Node_ConstructObjectFromClassHelper::ClassPinName &&
-			Pin->PinName != FK2Node_ConstructObjectFromClassHelper::WorldContextPinName &&
-			Pin->PinName != FK2Node_ConstructObjectFromClassHelper::OuterPinName);
+			Pin->PinName != FK2Node_ConstructObjectFromClassHelper::GetClassPinName() &&
+			Pin->PinName != FK2Node_ConstructObjectFromClassHelper::GetWorldContextPinName() &&
+			Pin->PinName != FK2Node_ConstructObjectFromClassHelper::GetOuterPinName());
 }
 
 void UK2Node_ConstructObjectFromClass::OnClassPinChanged()
@@ -266,7 +276,7 @@ void UK2Node_ConstructObjectFromClass::PinConnectionListChanged(UEdGraphPin* Pin
 {
 	Super::PinConnectionListChanged(Pin);
 
-	if (Pin && (Pin->PinName == FK2Node_ConstructObjectFromClassHelper::ClassPinName))
+	if (Pin && (Pin->PinName == FK2Node_ConstructObjectFromClassHelper::GetClassPinName()))
 	{
 		OnClassPinChanged();
 	}
@@ -292,7 +302,7 @@ void UK2Node_ConstructObjectFromClass::GetPinHoverText(const UEdGraphPin& Pin, F
 
 void UK2Node_ConstructObjectFromClass::PinDefaultValueChanged(UEdGraphPin* ChangedPin) 
 {
-	if (ChangedPin && (ChangedPin->PinName == FK2Node_ConstructObjectFromClassHelper::ClassPinName))
+	if (ChangedPin && (ChangedPin->PinName == FK2Node_ConstructObjectFromClassHelper::GetClassPinName()))
 	{
 		OnClassPinChanged();
 	}
@@ -317,7 +327,7 @@ UEdGraphPin* UK2Node_ConstructObjectFromClass::GetClassPin(const TArray<UEdGraph
 	UEdGraphPin* Pin = nullptr;
 	for (UEdGraphPin* TestPin : *PinsToSearch)
 	{
-		if (TestPin && TestPin->PinName == FK2Node_ConstructObjectFromClassHelper::ClassPinName)
+		if (TestPin && TestPin->PinName == FK2Node_ConstructObjectFromClassHelper::GetClassPinName())
 		{
 			Pin = TestPin;
 			break;
@@ -329,7 +339,7 @@ UEdGraphPin* UK2Node_ConstructObjectFromClass::GetClassPin(const TArray<UEdGraph
 
 UEdGraphPin* UK2Node_ConstructObjectFromClass::GetWorldContextPin() const
 {
-	UEdGraphPin* Pin = FindPin(FK2Node_ConstructObjectFromClassHelper::WorldContextPinName);
+	UEdGraphPin* Pin = FindPin(FK2Node_ConstructObjectFromClassHelper::GetWorldContextPinName());
 	check(Pin == nullptr || Pin->Direction == EGPD_Input);
 	return Pin;
 }
diff --git a/Engine/Source/Editor/BlueprintGraph/Private/K2Node_EaseFunction.cpp b/Engine/Source/Editor/BlueprintGraph/Private/K2Node_EaseFunction.cpp
--- a/Engine/Source/Editor/BlueprintGraph/Private/K2Node_EaseFunction.cpp
+++ b/Engine/Source/Editor/BlueprintGraph/Private/K2Node_EaseFunction.cpp
@@ -101,7 +101,11 @@ void UK2Node_EaseFunction::AllocateDefaultPins()
 	// Make sure that the default value is set correctly if none has been set
 	K2Schema->SetPinAutogeneratedDefaultValueBasedOnType(CachedEaseFuncPin);
 
+#if ENABLE_BLUEPRINT_REAL_NUMBERS
+	UEdGraphPin* AlphaPin = CreatePin(EGPD_Input, UEdGraphSchema_K2::PC_Real, UEdGraphSchema_K2::PC_Double, FEaseFunctionNodeHelper::GetAlphaPinName());
+#else
 	UEdGraphPin* AlphaPin = CreatePin(EGPD_Input, UEdGraphSchema_K2::PC_Float, FEaseFunctionNodeHelper::GetAlphaPinName());
+#endif
 	SetPinToolTip(*AlphaPin, LOCTEXT("AlphaPinTooltip", "Alpha value used to specify the easing in time."));
 
 	// Add wildcard pins for A, B and the return Pin
@@ -118,7 +122,11 @@ void UK2Node_EaseFunction::AllocateDefaultPins()
 	SetPinToolTip(*ShortestPathPin, LOCTEXT("ShortestPathPinTooltip", "When interpolating the shortest path should be taken."));
 	K2Schema->SetPinAutogeneratedDefaultValue(ShortestPathPin, TEXT("true"));
 
+#if ENABLE_BLUEPRINT_REAL_NUMBERS
+	UEdGraphPin* BlendExpPin = CreatePin(EGPD_Input, UEdGraphSchema_K2::PC_Real, UEdGraphSchema_K2::PC_Double, FEaseFunctionNodeHelper::GetBlendExpPinName());
+#else
 	UEdGraphPin* BlendExpPin = CreatePin(EGPD_Input, UEdGraphSchema_K2::PC_Float, FEaseFunctionNodeHelper::GetBlendExpPinName());
+#endif
 	SetPinToolTip(*BlendExpPin, LOCTEXT("BlendExpPinDescription", "Blend Exponent for basic ease functions"));
 	K2Schema->SetPinAutogeneratedDefaultValue(BlendExpPin, TEXT("2.0"));
 
@@ -224,6 +232,20 @@ bool UK2Node_EaseFunction::IsConnectionDisallowed(const UEdGraphPin* MyPin, cons
 		MyPin->PinName == FEaseFunctionNodeHelper::GetBPinName() || 
 		MyPin->PinName == FEaseFunctionNodeHelper::GetResultPinName())
 	{
+#if ENABLE_BLUEPRINT_REAL_NUMBERS
+		const bool bConnectionOk = (
+			OtherPin->PinType.PinCategory == UEdGraphSchema_K2::PC_Real ||
+			(
+				OtherPin->PinType.PinCategory == UEdGraphSchema_K2::PC_Struct &&
+				OtherPin->PinType.PinSubCategoryObject.IsValid() &&
+				(
+					OtherPin->PinType.PinSubCategoryObject.Get()->GetName() == TEXT("Vector") ||
+					OtherPin->PinType.PinSubCategoryObject.Get()->GetName() == TEXT("Rotator") ||
+					OtherPin->PinType.PinSubCategoryObject.Get()->GetName() == TEXT("Transform")
+					)
+				)
+			) && !OtherPin->PinType.IsContainer();
+#else
 		const bool bConnectionOk = (
 			OtherPin->PinType.PinCategory == UEdGraphSchema_K2::PC_Float ||
 			(
@@ -236,6 +258,8 @@ bool UK2Node_EaseFunction::IsConnectionDisallowed(const UEdGraphPin* MyPin, cons
 				)
 			)
 		) && !OtherPin->PinType.IsContainer();
+#endif
+
 		if (!bConnectionOk)
 		{
 			OutReason = LOCTEXT("PinConnectionDisallowed", "Pin type is not supported by function.").ToString();
@@ -310,10 +334,17 @@ void UK2Node_EaseFunction::PinTypeChanged(UEdGraphPin* Pin)
 				EaseFunctionName = NAME_None;
 
 				// Generate the right function name
+#if ENABLE_BLUEPRINT_REAL_NUMBERS
+				if (InstigatorPin->PinType.PinCategory == UEdGraphSchema_K2::PC_Real)
+				{
+					EaseFunctionName = TEXT("Ease");
+				}
+#else
 				if (InstigatorPin->PinType.PinCategory == UEdGraphSchema_K2::PC_Float)
 				{
 					EaseFunctionName = TEXT("Ease");
 				}
+#endif
 				else if (InstigatorPin->PinType.PinCategory == UEdGraphSchema_K2::PC_Struct)
 				{
 					if (InstigatorPin->PinType.PinSubCategoryObject.Get()->GetName() == TEXT("Vector"))
diff --git a/Engine/Source/Editor/BlueprintGraph/Private/K2Node_EditablePinBase.cpp b/Engine/Source/Editor/BlueprintGraph/Private/K2Node_EditablePinBase.cpp
--- a/Engine/Source/Editor/BlueprintGraph/Private/K2Node_EditablePinBase.cpp
+++ b/Engine/Source/Editor/BlueprintGraph/Private/K2Node_EditablePinBase.cpp
@@ -22,7 +22,7 @@ FArchive& operator<<(FArchive& Ar, FUserPinInfo& Info)
 
 	if (Ar.CustomVer(FFrameworkObjectVersion::GUID) >= FFrameworkObjectVersion::PinsStoreFName)
 	{
-	Ar << Info.PinName;
+		Ar << Info.PinName;
 	}
 	else
 	{
@@ -46,8 +46,8 @@ FArchive& operator<<(FArchive& Ar, FUserPinInfo& Info)
 		bool bIsReference = Info.PinType.bIsReference;
 		Ar << bIsReference;
 
-			Info.PinType.ContainerType = (bIsArray ? EPinContainerType::Array : EPinContainerType::None);
-			Info.PinType.bIsReference = bIsReference;
+		Info.PinType.ContainerType = (bIsArray ? EPinContainerType::Array : EPinContainerType::None);
+		Info.PinType.bIsReference = bIsReference;
 
 		FString PinCategoryStr;
 		FString PinSubCategoryStr;
@@ -58,6 +58,17 @@ FArchive& operator<<(FArchive& Ar, FUserPinInfo& Info)
 		Info.PinType.PinCategory = *PinCategoryStr;
 		Info.PinType.PinSubCategory = *PinSubCategoryStr;
 
+#if ENABLE_BLUEPRINT_REAL_NUMBERS
+		bool bFixupPinCategories =
+			((Info.PinType.PinCategory == TEXT("double")) || (Info.PinType.PinCategory == TEXT("float")));
+
+		if (bFixupPinCategories)
+		{
+			Info.PinType.PinCategory = TEXT("real");
+			Info.PinType.PinSubCategory = TEXT("double");
+		}
+#endif
+
 		Ar << Info.PinType.PinSubCategoryObject;
 	}
 
diff --git a/Engine/Source/Editor/BlueprintGraph/Private/K2Node_FormatText.cpp b/Engine/Source/Editor/BlueprintGraph/Private/K2Node_FormatText.cpp
--- a/Engine/Source/Editor/BlueprintGraph/Private/K2Node_FormatText.cpp
+++ b/Engine/Source/Editor/BlueprintGraph/Private/K2Node_FormatText.cpp
@@ -381,21 +381,41 @@ void UK2Node_FormatText::ExpandNode(class FKismetCompilerContext& CompilerContex
 				// Connect the int output pin to the argument value
 				CallFloatToDoubleFunction->FindPinChecked(UEdGraphSchema_K2::PN_ReturnValue)->MakeLinkTo(MakeFormatArgumentDataStruct->FindPinChecked(GET_MEMBER_NAME_STRING_CHECKED(FFormatArgumentData, ArgumentValueInt)));
 			}
+#if ENABLE_BLUEPRINT_REAL_NUMBERS
+			else if (ArgumentPinCategory == UEdGraphSchema_K2::PC_Real)
+			{
+				if (ArgumentPin->PinType.PinSubCategory == UEdGraphSchema_K2::PC_Float)
+				{
+					MakeFormatArgumentDataStruct->GetSchema()->TrySetDefaultValue(*ArgumentTypePin, TEXT("Float"));
+					CompilerContext.MovePinLinksToIntermediate(*ArgumentPin, *MakeFormatArgumentDataStruct->FindPinChecked(GET_MEMBER_NAME_STRING_CHECKED(FFormatArgumentData, ArgumentValueFloat)));
+				}
+				else if (ArgumentPin->PinType.PinSubCategory == UEdGraphSchema_K2::PC_Double)
+				{
+					MakeFormatArgumentDataStruct->GetSchema()->TrySetDefaultValue(*ArgumentTypePin, TEXT("Double"));
+					CompilerContext.MovePinLinksToIntermediate(*ArgumentPin, *MakeFormatArgumentDataStruct->FindPinChecked(GET_MEMBER_NAME_STRING_CHECKED(FFormatArgumentData, ArgumentValueDouble)));
+				}
+				else
+				{
+					check(false);
+				}
+			}
+#else
 			else if (ArgumentPinCategory == UEdGraphSchema_K2::PC_Float)
 			{
 				MakeFormatArgumentDataStruct->GetSchema()->TrySetDefaultValue(*ArgumentTypePin, TEXT("Float"));
 				CompilerContext.MovePinLinksToIntermediate(*ArgumentPin, *MakeFormatArgumentDataStruct->FindPinChecked(GET_MEMBER_NAME_STRING_CHECKED(FFormatArgumentData, ArgumentValueFloat)));
 			}
-			else if (ArgumentPinCategory == UEdGraphSchema_K2::PC_Int64)
-			{
-				MakeFormatArgumentDataStruct->GetSchema()->TrySetDefaultValue(*ArgumentTypePin, TEXT("Int64"));
-				CompilerContext.MovePinLinksToIntermediate(*ArgumentPin, *MakeFormatArgumentDataStruct->FindPinChecked(GET_MEMBER_NAME_STRING_CHECKED(FFormatArgumentData, ArgumentValueInt)));
-			}
 			else if (ArgumentPinCategory == UEdGraphSchema_K2::PC_Double)
 			{
 				MakeFormatArgumentDataStruct->GetSchema()->TrySetDefaultValue(*ArgumentTypePin, TEXT("Double"));
 				CompilerContext.MovePinLinksToIntermediate(*ArgumentPin, *MakeFormatArgumentDataStruct->FindPinChecked(GET_MEMBER_NAME_STRING_CHECKED(FFormatArgumentData, ArgumentValueDouble)));
 			}
+#endif
+			else if (ArgumentPinCategory == UEdGraphSchema_K2::PC_Int64)
+			{
+				MakeFormatArgumentDataStruct->GetSchema()->TrySetDefaultValue(*ArgumentTypePin, TEXT("Int64"));
+				CompilerContext.MovePinLinksToIntermediate(*ArgumentPin, *MakeFormatArgumentDataStruct->FindPinChecked(GET_MEMBER_NAME_STRING_CHECKED(FFormatArgumentData, ArgumentValueInt)));
+			}
 			else if (ArgumentPinCategory == UEdGraphSchema_K2::PC_Text)
 			{
 				MakeFormatArgumentDataStruct->GetSchema()->TrySetDefaultValue(*ArgumentTypePin, TEXT("Text"));
@@ -546,10 +566,17 @@ bool UK2Node_FormatText::IsConnectionDisallowed(const UEdGraphPin* MyPin, const
 		const FName& OtherPinCategory = OtherPin->PinType.PinCategory;
 
 		bool bIsValidType = false;
+#if ENABLE_BLUEPRINT_REAL_NUMBERS
+		if (OtherPinCategory == UEdGraphSchema_K2::PC_Int || OtherPinCategory == UEdGraphSchema_K2::PC_Real || OtherPinCategory == UEdGraphSchema_K2::PC_Text ||
+			(OtherPinCategory == UEdGraphSchema_K2::PC_Byte && !OtherPin->PinType.PinSubCategoryObject.IsValid()) ||
+			OtherPinCategory == UEdGraphSchema_K2::PC_Boolean || OtherPinCategory == UEdGraphSchema_K2::PC_String || OtherPinCategory == UEdGraphSchema_K2::PC_Name || OtherPinCategory == UEdGraphSchema_K2::PC_Object ||
+			OtherPinCategory == UEdGraphSchema_K2::PC_Wildcard || OtherPinCategory == UEdGraphSchema_K2::PC_Int64)
+#else
 		if (OtherPinCategory == UEdGraphSchema_K2::PC_Int || OtherPinCategory == UEdGraphSchema_K2::PC_Float || OtherPinCategory == UEdGraphSchema_K2::PC_Text ||
 			(OtherPinCategory == UEdGraphSchema_K2::PC_Byte && !OtherPin->PinType.PinSubCategoryObject.IsValid()) ||
             OtherPinCategory == UEdGraphSchema_K2::PC_Boolean || OtherPinCategory == UEdGraphSchema_K2::PC_String || OtherPinCategory == UEdGraphSchema_K2::PC_Name || OtherPinCategory == UEdGraphSchema_K2::PC_Object ||
 			OtherPinCategory == UEdGraphSchema_K2::PC_Wildcard || OtherPinCategory == UEdGraphSchema_K2::PC_Int64 || OtherPinCategory == UEdGraphSchema_K2::PC_Double)
+#endif
         {
 			bIsValidType = true;
 		}
diff --git a/Engine/Source/Editor/BlueprintGraph/Private/K2Node_FunctionResult.cpp b/Engine/Source/Editor/BlueprintGraph/Private/K2Node_FunctionResult.cpp
--- a/Engine/Source/Editor/BlueprintGraph/Private/K2Node_FunctionResult.cpp
+++ b/Engine/Source/Editor/BlueprintGraph/Private/K2Node_FunctionResult.cpp
@@ -117,6 +117,9 @@ public:
 	{
 		return true;
 	}
+
+protected:
+	virtual bool UsesVariablePinAsKey() const override { return true; }
 };
 
 #if WITH_EDITORONLY_DATA
@@ -235,7 +238,7 @@ void UK2Node_FunctionResult::GetMenuActions(FBlueprintActionDatabaseRegistrar& A
 
 bool UK2Node_FunctionResult::IsCompatibleWithGraph(UEdGraph const* Graph) const
 {
-	auto K2Schema = Cast<const UEdGraphSchema_K2>(Graph ? Graph->GetSchema() : nullptr);
+	const UEdGraphSchema_K2* K2Schema = Cast<UEdGraphSchema_K2>(Graph ? Graph->GetSchema() : nullptr);
 	const bool bIsConstructionScript = (K2Schema != nullptr) ? K2Schema->IsConstructionScript(Graph) : false;
 	const bool bIsCompatible = (K2Schema != nullptr) ? (EGraphType::GT_Function == K2Schema->GetGraphType(Graph)) : false;
 	return bIsCompatible && !bIsConstructionScript && Super::IsCompatibleWithGraph(Graph);
@@ -244,7 +247,7 @@ bool UK2Node_FunctionResult::IsCompatibleWithGraph(UEdGraph const* Graph) const
 TArray<UK2Node_FunctionResult*> UK2Node_FunctionResult::GetAllResultNodes() const
 {
 	TArray<UK2Node_FunctionResult*> AllResultNodes;
-	if (auto Graph = GetGraph())
+	if (const UEdGraph* Graph = GetGraph())
 	{
 		Graph->GetNodesOfClass(AllResultNodes);
 	}
@@ -334,9 +337,9 @@ void UK2Node_FunctionResult::SyncWithEntryNode()
 
 void UK2Node_FunctionResult::SyncWithPrimaryResultNode()
 {
-	UK2Node_FunctionResult* PrimaryNode = nullptr;
+	const UK2Node_FunctionResult* PrimaryNode = nullptr;
 	TArray<UK2Node_FunctionResult*> AllResultNodes = GetAllResultNodes();
-	for (auto ResultNode : AllResultNodes)
+	for (const UK2Node_FunctionResult* ResultNode : AllResultNodes)
 	{
 		if (ResultNode && (this != ResultNode))
 		{
@@ -351,10 +354,10 @@ void UK2Node_FunctionResult::SyncWithPrimaryResultNode()
 		bIsEditable = PrimaryNode->bIsEditable;
 
 		// Temporary array that will contain our list of Old Pins that are no longer part of the return signature
-		TArray< TSharedPtr<FUserPinInfo> > OldPins = UserDefinedPins;
+		TArray<TSharedPtr<FUserPinInfo>> OldPins = UserDefinedPins;
 
 		// Temporary array that will contain our list of Signature Pins that need to be added
-		TArray< TSharedPtr<FUserPinInfo> > SignaturePins = PrimaryNode->UserDefinedPins;
+		TArray<TSharedPtr<FUserPinInfo>> SignaturePins = PrimaryNode->UserDefinedPins;
 
 		for (int OldIndex = OldPins.Num() - 1; OldIndex >= 0; --OldIndex)
 		{
@@ -428,13 +431,14 @@ void UK2Node_FunctionResult::ValidateNodeDuringCompilation(class FCompilerResult
 {
 	Super::ValidateNodeDuringCompilation(MessageLog);
 
-	auto AllResultNodes = GetAllResultNodes();
-	UK2Node_FunctionResult* OtherResult = AllResultNodes.Num() ? AllResultNodes[0] : nullptr;
+	TArray<UK2Node_FunctionResult*> AllResultNodes = GetAllResultNodes();
+	const UK2Node_FunctionResult* OtherResult = AllResultNodes.Num() ? AllResultNodes[0] : nullptr;
 	if (OtherResult && (OtherResult != this))
 	{
-		for (auto Pin : Pins)
+		for (UEdGraphPin* Pin : Pins)
 		{
-			auto OtherPin = OtherResult->FindPin(Pin->PinName);
+			check(Pin);
+			UEdGraphPin* OtherPin = OtherResult->FindPin(Pin->PinName);
 			if (!OtherPin || (OtherPin->PinType != Pin->PinType))
 			{
 				MessageLog.Error(*NSLOCTEXT("K2Node", "FunctionResult_DifferentReturnError", "Return nodes don't match each other: @@, @@").ToString(), this, OtherResult);
diff --git a/Engine/Source/Editor/BlueprintGraph/Private/K2Node_MakeContainer.cpp b/Engine/Source/Editor/BlueprintGraph/Private/K2Node_MakeContainer.cpp
--- a/Engine/Source/Editor/BlueprintGraph/Private/K2Node_MakeContainer.cpp
+++ b/Engine/Source/Editor/BlueprintGraph/Private/K2Node_MakeContainer.cpp
@@ -12,6 +12,7 @@
 
 #include "ScopedTransaction.h"
 #include "EdGraphUtilities.h"
+#include "KismetCastingUtils.h"
 #include "KismetCompiledFunctionContext.h"
 #include "KismetCompilerMisc.h"
 #include "BlueprintNodeSpawner.h"
@@ -36,27 +37,45 @@ void FKCHandler_MakeContainer::RegisterNets(FKismetFunctionContext& Context, UEd
 
 void FKCHandler_MakeContainer::Compile(FKismetFunctionContext& Context, UEdGraphNode* Node)
 {
-	UK2Node_MakeContainer* ContainerNode = CastChecked<UK2Node_MakeContainer>(Node);
-	UEdGraphPin* OutputPin = ContainerNode->GetOutputPin();
-
-	FBPTerminal** ContainerTerm = Context.NetMap.Find(OutputPin);
-	check(ContainerTerm);
-
-	FBlueprintCompiledStatement& CreateContainerStatement = Context.AppendStatementForNode(Node);
-	CreateContainerStatement.Type = CompiledStatementType;
-	CreateContainerStatement.LHS = *ContainerTerm;
+	TArray<FBPTerminal*> RHSTerms;
 
 	for (UEdGraphPin* Pin : Node->Pins)
 	{
-		if(Pin && Pin->Direction == EGPD_Input)
+		if (Pin && Pin->Direction == EGPD_Input)
 		{
 			FBPTerminal** InputTerm = Context.NetMap.Find(FEdGraphUtilities::GetNetFromPin(Pin));
-			if( InputTerm )
+			if (InputTerm)
 			{
-				CreateContainerStatement.RHS.Add(*InputTerm);
+				FBPTerminal* RHSTerm = *InputTerm;
+
+#if ENABLE_BLUEPRINT_REAL_NUMBERS
+				{
+					using namespace UE::KismetCompiler;
+
+					TOptional<TPair<FBPTerminal*, EKismetCompiledStatementType>> ImplicitCastEntry =
+						CastingUtils::InsertImplicitCastStatement(Context, Pin, RHSTerm);
+
+					if (ImplicitCastEntry)
+					{
+						RHSTerm = ImplicitCastEntry->Get<0>();
+					}
+				}
+#endif
+				RHSTerms.Add(RHSTerm);
 			}
 		}
 	}
+
+	UK2Node_MakeContainer* ContainerNode = CastChecked<UK2Node_MakeContainer>(Node);
+	UEdGraphPin* OutputPin = ContainerNode->GetOutputPin();
+
+	FBPTerminal** ContainerTerm = Context.NetMap.Find(OutputPin);
+	check(ContainerTerm);
+
+	FBlueprintCompiledStatement& CreateContainerStatement = Context.AppendStatementForNode(Node);
+	CreateContainerStatement.Type = CompiledStatementType;
+	CreateContainerStatement.LHS = *ContainerTerm;
+	CreateContainerStatement.RHS = MoveTemp(RHSTerms);
 }
 
 /////////////////////////////////////////////////////
diff --git a/Engine/Source/Editor/BlueprintGraph/Private/K2Node_MathExpression.cpp b/Engine/Source/Editor/BlueprintGraph/Private/K2Node_MathExpression.cpp
--- a/Engine/Source/Editor/BlueprintGraph/Private/K2Node_MathExpression.cpp
+++ b/Engine/Source/Editor/BlueprintGraph/Private/K2Node_MathExpression.cpp
@@ -79,7 +79,12 @@ static bool PromoteIntToFloat(FEdGraphPinType& InOutType)
 {
 	if (InOutType.PinCategory == UEdGraphSchema_K2::PC_Int)
 	{
+#if ENABLE_BLUEPRINT_REAL_NUMBERS
+		InOutType.PinCategory = UEdGraphSchema_K2::PC_Real;
+		InOutType.PinSubCategory = UEdGraphSchema_K2::PC_Double;
+#else
 		InOutType.PinCategory = UEdGraphSchema_K2::PC_Float;
+#endif
 		InOutType.PinSubCategoryObject = nullptr;
 		return true;
 	}
@@ -548,7 +553,7 @@ public:
 class FFunctionExpression : public IFExpressionNode
 {
 public:
-	FFunctionExpression(FString const& InFuncName, TSharedRef<FExpressionList> InParamList)
+	FFunctionExpression(const FString& InFuncName, TSharedRef<FExpressionList> InParamList)
 		: FuncName(InFuncName)
 		, ParamList(InParamList)
 	{
@@ -575,13 +580,13 @@ public:
 	/** For debug purposes, constructs a textual representation of this expression */
 	virtual FString ToString() const override
 	{
-		FString const ParamsString = ParamList->ToString();
+		const FString ParamsString = ParamList->ToString();
 		return FString::Printf(TEXT("(%s%s)"), *FuncName, *ParamsString);
 	}
 
 	virtual FString ToDisplayString(UBlueprint* InBlueprint) const
 	{
-		FString const ParamsString = ParamList->ToDisplayString(InBlueprint);
+		const FString ParamsString = ParamList->ToDisplayString(InBlueprint);
 		return FString::Printf(TEXT("(%s%s)"), *FuncName, *ParamsString);
 	}
 public:
@@ -707,7 +712,7 @@ public:
 	 * Checks to see if there are any functions associated with the specified 
 	 * operator. 
 	 */
-	bool Contains(FString const& Operator) const
+	bool Contains(const FString& Operator) const
 	{
 		return LookupTable.Contains(Operator);
 	}	
@@ -722,7 +727,7 @@ public:
 	 * @param  InputTypeList	A list of parameter types you want to feed the function.
 	 * @return A pointer to the matching function (if one was found), otherwise nullptr.
 	 */
-	UFunction* FindMatchingFunction(FString const& Operator, TArray<FEdGraphPinType> const& InputTypeList) const
+	UFunction* FindMatchingFunction(const FString& Operator, const TArray<FEdGraphPinType>& InputTypeList) const
 	{
 		// make a local copy of the desired input types so that we can promote 
 		// those types as needed
@@ -736,7 +741,7 @@ public:
 		// float), and see if we can lookup a function with those types
 		for (int32 promoterIndex = 0; (promoterIndex < OrderedTypePromoters.Num()) && (MatchingFunc == NULL); ++promoterIndex)
 		{
-			FTypePromoter const& PromotionOperator = OrderedTypePromoters[promoterIndex];
+			const FTypePromoter& PromotionOperator = OrderedTypePromoters[promoterIndex];
 
 			// Apply the promotion operator to any values that match
 			bool bMadeChanges = false;
@@ -766,7 +771,7 @@ public:
 	 * Flags the specified function as one associated with the supplied 
 	 * operator.
 	 */
-	void Add(FString const& Operator, UFunction* OperatorFunc)
+	void Add(const FString& Operator, UFunction* OperatorFunc)
 	{
 		LookupTable.FindOrAdd(Operator).Add(OperatorFunc);
 	}
@@ -798,12 +803,12 @@ public:
 					}
 					
 					FString FunctionName = TestFunction->GetName();
-					TArray<FString> const& OperatorAliases = GetOperatorAliases(FunctionName);
+					const TArray<FString>& OperatorAliases = GetOperatorAliases(FunctionName);
 					
 					// if there are aliases, use those instead of the function's standard name
 					if (OperatorAliases.Num() > 0)
 					{
-						for (FString const& Alias : OperatorAliases)
+						for (const FString& Alias : OperatorAliases)
 						{
 							Add(Alias, TestFunction);
 						}
@@ -847,14 +852,14 @@ private:
 	 * @param  InputTypeList	A list of parameter types you want to feed the function.
 	 * @return A pointer to the matching function (if one was found), otherwise nullptr.
 	 */
-	UFunction* FindFunctionInternal(FString const& Operator, TArray<FEdGraphPinType> const& InputTypeList) const
+	UFunction* FindFunctionInternal(const FString& Operator, const TArray<FEdGraphPinType>& InputTypeList) const
 	{
 		UFunction* MatchedFunction = nullptr;
 
-		FFunctionsList const* OperatorFunctions = LookupTable.Find(Operator);
+		const FFunctionsList* OperatorFunctions = LookupTable.Find(Operator);
 		if (OperatorFunctions != nullptr)
 		{
-			UEdGraphSchema_K2 const* K2Schema = GetDefault<UEdGraphSchema_K2>();
+			const UEdGraphSchema_K2* K2Schema = GetDefault<UEdGraphSchema_K2>();
 			for (UFunction* TestFunction : *OperatorFunctions)
 			{
 				int32 ArgumentIndex = 0;
@@ -868,7 +873,7 @@ private:
 							FEdGraphPinType ParamType;
 							if (K2Schema->ConvertPropertyToPinType(Param, /*out*/ParamType))
 							{
-								FEdGraphPinType const& TypeToMatch = InputTypeList[ArgumentIndex];
+								const FEdGraphPinType& TypeToMatch = InputTypeList[ArgumentIndex];
 								if (!K2Schema->ArePinTypesCompatible(TypeToMatch, ParamType))
 								{
 									break; // type mismatch
@@ -909,7 +914,7 @@ private:
 	 * @param  FunctionName		The raw name of the function you're looking to replace (not the friendly name)
 	 * @return A reference to the array of aliases for the specified function (an empty array if none were found).
 	 */
-	static TArray<FString> const& GetOperatorAliases(FString const& FunctionName)
+	static const TArray<FString>& GetOperatorAliases(const FString& FunctionName)
 	{
 #define FUNC_ALIASES_BEGIN(FuncName) \
 		if (FunctionName == FString(TEXT(FuncName))) \
@@ -1096,7 +1101,7 @@ public:
 	 * 
 	 * @return The pin type of this fragment's output.
 	 */
-	FEdGraphPinType const& GetOutputType() const
+	const FEdGraphPinType& GetOutputType() const
 	{
 		return FragmentType;
 	}
@@ -1138,7 +1143,7 @@ private:
 class FCodeGenFragment_VariableGet : public FCodeGenFragment
 {
 public:
-	FCodeGenFragment_VariableGet(UK2Node_VariableGet* InNode, FEdGraphPinType const& InType)
+	FCodeGenFragment_VariableGet(UK2Node_VariableGet* InNode, const FEdGraphPinType& InType)
 		: FCodeGenFragment(InType)
 		, GeneratedNode(InNode)
 	{
@@ -1177,7 +1182,7 @@ private:
 class FCodeGenFragment_FuntionCall : public FCodeGenFragment
 {
 public:
-	FCodeGenFragment_FuntionCall(UK2Node_CallFunction* InNode, FEdGraphPinType const& InType)
+	FCodeGenFragment_FuntionCall(UK2Node_CallFunction* InNode, const FEdGraphPinType& InType)
 		: FCodeGenFragment(InType)
 		, GeneratedNode(InNode)
 	{
@@ -1215,7 +1220,7 @@ private:
 class FCodeGenFragment_Literal : public FCodeGenFragment
 {
 public:
-	FCodeGenFragment_Literal(FString const& LiteralVal, FEdGraphPinType const& ResultType) 
+	FCodeGenFragment_Literal(const FString& LiteralVal, const FEdGraphPinType& ResultType)
 		: FCodeGenFragment(ResultType) 
 		, DefaultValue(LiteralVal)
 	{}
@@ -1225,7 +1230,7 @@ public:
 	/// Begin FCodeGenFragment Interface
 	virtual bool ConnectToInput(UEdGraphPin* InputPin, FCompilerResultsLog& MessageLog) override
 	{
-		UEdGraphSchema_K2 const* K2Schema = Cast<UEdGraphSchema_K2>(InputPin->GetSchema());
+		const UEdGraphSchema_K2* K2Schema = Cast<UEdGraphSchema_K2>(InputPin->GetSchema());
 		bool bSuccess = true;//K2Schema->ArePinTypesCompatible(GetOutputType(), InputPin->PinType);
 		if (bSuccess)
 		{
@@ -1380,7 +1385,7 @@ public:
 
 		if (ExpressionNode.Token.TokenType == FBasicToken::TOKEN_Identifier || ExpressionNode.Token.TokenType == FBasicToken::TOKEN_Guid)
 		{
-			FString const VariableIdentifier = ExpressionNode.Token.Identifier;
+			const FString VariableIdentifier = ExpressionNode.Token.Identifier;
 			// first we try to match up variables with existing variable properties on the blueprint
 
 			FMemberReference VariableReference;
@@ -1634,7 +1639,7 @@ private:
 	TSharedPtr<FCodeGenFragment_InputPin> GenerateInputPinFragment(const FName VariableIdentifier)
 	{
 		TSharedPtr<FCodeGenFragment_InputPin> InputPinFragment;
-		UEdGraphSchema_K2 const* K2Schema = GetDefault<UEdGraphSchema_K2>();
+		const UEdGraphSchema_K2* K2Schema = GetDefault<UEdGraphSchema_K2>();
 		
 		UK2Node_Tunnel* EntryNode = CompilingNode->GetEntryNode();
 		// if a pin under this name already exists, use that
@@ -1648,7 +1653,12 @@ private:
 			// Create an input pin (using the default guessed type)
 			FEdGraphPinType DefaultType;
 			// currently, generated expressions ALWAYS take a float (it is the most versatile type)
+#if ENABLE_BLUEPRINT_REAL_NUMBERS
+			DefaultType.PinCategory = UEdGraphSchema_K2::PC_Real;
+			DefaultType.PinSubCategory = UEdGraphSchema_K2::PC_Double;
+#else
 			DefaultType.PinCategory = UEdGraphSchema_K2::PC_Float;
+#endif
 			
 			UEdGraphPin* NewInputPin = EntryNode->CreateUserDefinedPin(VariableIdentifier, DefaultType, EGPD_Output);
 			InputPinFragment = MakeShareable(new FCodeGenFragment_InputPin(NewInputPin));
@@ -1678,7 +1688,7 @@ private:
 	{
 		check(ExpressionContext.Token.TokenType == FBasicToken::TOKEN_Identifier || ExpressionContext.Token.TokenType == FBasicToken::TOKEN_Guid);
 		check(VariableProperty != nullptr);
-		UEdGraphSchema_K2 const* K2Schema = GetDefault<UEdGraphSchema_K2>();
+		const UEdGraphSchema_K2* K2Schema = GetDefault<UEdGraphSchema_K2>();
 		
 		TSharedPtr<FCodeGenFragment_VariableGet> VariableGetFragment;
 		
@@ -1720,7 +1730,7 @@ private:
 	 * @param  ExpressionNode	The expression node that we're generating this fragment for.
      * @return A new literal fragment.
      */
-	TSharedPtr<FCodeGenFragment_Literal> GenerateLiteralFragment(FBasicToken const& Token, FCompilerResultsLog& MessageLog)
+	TSharedPtr<FCodeGenFragment_Literal> GenerateLiteralFragment(const FBasicToken& Token, FCompilerResultsLog& MessageLog)
 	{
 		check(Token.TokenType == FBasicToken::TOKEN_Const);
 		
@@ -1731,7 +1741,12 @@ private:
 				LiteralType.PinCategory = UEdGraphSchema_K2::PC_Boolean;
 				break;
 			case CPT_Float:
+#if ENABLE_BLUEPRINT_REAL_NUMBERS
+				LiteralType.PinCategory = UEdGraphSchema_K2::PC_Real;
+				LiteralType.PinSubCategory = UEdGraphSchema_K2::PC_Double;
+#else
 				LiteralType.PinCategory = UEdGraphSchema_K2::PC_Float;
+#endif
 				break;
 			case CPT_Int:
 				LiteralType.PinCategory = UEdGraphSchema_K2::PC_Int;
@@ -1804,7 +1819,7 @@ private:
 			}
 			else
 			{
-				UEdGraphSchema_K2 const* K2Schema = GetDefault<UEdGraphSchema_K2>();
+				const UEdGraphSchema_K2* K2Schema = GetDefault<UEdGraphSchema_K2>();
 				
 				FEdGraphPinType ReturnType;
 				if (K2Schema->ConvertPropertyToPinType(ReturnProperty, /*out*/ReturnType))
@@ -2721,9 +2736,9 @@ void UK2Node_MathExpression::ValidateNodeDuringCompilation(FCompilerResultsLog&
 	// else, this may be some intermediate node in the compile, let's look at the errors from the original...
 	else 
 	{
-		if(UObject const* SourceObject = MessageLog.FindSourceObject(this))
+		if(const UObject* SourceObject = MessageLog.FindSourceObject(this))
 		{
-			UK2Node_MathExpression const* MathExpression = MessageLog.FindSourceObjectTypeChecked<UK2Node_MathExpression>(this);
+			const UK2Node_MathExpression* MathExpression = MessageLog.FindSourceObjectTypeChecked<UK2Node_MathExpression>(this);
 
 			// Should always be able to find the source math expression
 			check(MathExpression);
diff --git a/Engine/Source/Editor/BlueprintGraph/Private/K2Node_Select.cpp b/Engine/Source/Editor/BlueprintGraph/Private/K2Node_Select.cpp
--- a/Engine/Source/Editor/BlueprintGraph/Private/K2Node_Select.cpp
+++ b/Engine/Source/Editor/BlueprintGraph/Private/K2Node_Select.cpp
@@ -11,6 +11,7 @@
 #include "EdGraphUtilities.h"
 #include "BPTerminal.h"
 #include "BlueprintCompiledStatement.h"
+#include "KismetCastingUtils.h"
 #include "KismetCompilerMisc.h"
 #include "KismetCompiler.h"
 #include "BlueprintNodeSpawner.h"
@@ -119,7 +120,8 @@ public:
 
 		TArray<UEdGraphPin*> OptionPins;
 		SelectNode->GetOptionPins(OptionPins);
-		for (int32 OptionIdx = 0; OptionIdx < OptionPins.Num(); OptionIdx++)
+
+		for (int32 OptionIdx = 0; OptionIdx < OptionPins.Num(); ++OptionIdx)
 		{
 			{
 				FBPTerminal* LiteralTerm = Context.CreateLocalTerminal(ETerminalSpecification::TS_Literal);
@@ -135,6 +137,7 @@ public:
 				}
 				SelectStatement->RHS.Add(LiteralTerm);
 			}
+
 			{
 				UEdGraphPin* NetPin = OptionPins[OptionIdx] ? FEdGraphUtilities::GetNetFromPin(OptionPins[OptionIdx]) : nullptr;
 				FBPTerminal** ValueTermPtr = NetPin ? Context.NetMap.Find(NetPin) : nullptr;
@@ -144,6 +147,21 @@ public:
 					Context.MessageLog.Error(*LOCTEXT("Error_NoTermFound", "No term registered for pin @@").ToString(), NetPin);
 					return;
 				}
+
+#if ENABLE_BLUEPRINT_REAL_NUMBERS
+				{
+					using namespace UE::KismetCompiler;
+
+					TOptional<TPair<FBPTerminal*, EKismetCompiledStatementType>> ImplicitCastEntry =
+						CastingUtils::InsertImplicitCastStatement(Context, OptionPins[OptionIdx], ValueTerm);
+
+					if (ImplicitCastEntry)
+					{
+						ValueTerm = ImplicitCastEntry->Get<0>();
+					}
+				}
+#endif
+
 				SelectStatement->RHS.Add(ValueTerm);
 			}
 		}
diff --git a/Engine/Source/Editor/BlueprintGraph/Private/K2Node_TemporaryVariable.cpp b/Engine/Source/Editor/BlueprintGraph/Private/K2Node_TemporaryVariable.cpp
--- a/Engine/Source/Editor/BlueprintGraph/Private/K2Node_TemporaryVariable.cpp
+++ b/Engine/Source/Editor/BlueprintGraph/Private/K2Node_TemporaryVariable.cpp
@@ -204,10 +204,15 @@ void UK2Node_TemporaryVariable::GetMenuActions(FBlueprintActionDatabaseRegistrar
 	ActionRegistrar.AddBlueprintAction(ActionKey, MakeTempVarNodeSpawner(FEdGraphPinType(UEdGraphSchema_K2::PC_Int, NAME_None, nullptr, EPinContainerType::Array, /*bIsReference =*/ false, /*InValueTerminalType=*/FEdGraphTerminalType()), /*bIsPersistent =*/false));
 	ActionRegistrar.AddBlueprintAction(ActionKey, MakeTempVarNodeSpawner(FEdGraphPinType(UEdGraphSchema_K2::PC_Int64, NAME_None, nullptr, EPinContainerType::None, /*bIsReference =*/ false, /*InValueTerminalType=*/FEdGraphTerminalType()), /*bIsPersistent =*/false));
 	ActionRegistrar.AddBlueprintAction(ActionKey, MakeTempVarNodeSpawner(FEdGraphPinType(UEdGraphSchema_K2::PC_Int64, NAME_None, nullptr, EPinContainerType::Array, /*bIsReference =*/ false, /*InValueTerminalType=*/FEdGraphTerminalType()), /*bIsPersistent =*/false));
+#if ENABLE_BLUEPRINT_REAL_NUMBERS
+	ActionRegistrar.AddBlueprintAction(ActionKey, MakeTempVarNodeSpawner(FEdGraphPinType(UEdGraphSchema_K2::PC_Real, UEdGraphSchema_K2::PC_Double, nullptr, EPinContainerType::None, /*bIsReference =*/ false, /*InValueTerminalType=*/FEdGraphTerminalType()), /*bIsPersistent =*/false));
+	ActionRegistrar.AddBlueprintAction(ActionKey, MakeTempVarNodeSpawner(FEdGraphPinType(UEdGraphSchema_K2::PC_Real, UEdGraphSchema_K2::PC_Double, nullptr, EPinContainerType::Array, /*bIsReference =*/ false, /*InValueTerminalType=*/FEdGraphTerminalType()), /*bIsPersistent =*/false));
+#else
 	ActionRegistrar.AddBlueprintAction(ActionKey, MakeTempVarNodeSpawner(FEdGraphPinType(UEdGraphSchema_K2::PC_Float, NAME_None, nullptr, EPinContainerType::None, /*bIsReference =*/ false, /*InValueTerminalType=*/FEdGraphTerminalType()), /*bIsPersistent =*/false));
 	ActionRegistrar.AddBlueprintAction(ActionKey, MakeTempVarNodeSpawner(FEdGraphPinType(UEdGraphSchema_K2::PC_Float, NAME_None, nullptr, EPinContainerType::Array, /*bIsReference =*/ false, /*InValueTerminalType=*/FEdGraphTerminalType()), /*bIsPersistent =*/false));	
 	ActionRegistrar.AddBlueprintAction(ActionKey, MakeTempVarNodeSpawner(FEdGraphPinType(UEdGraphSchema_K2::PC_Double, NAME_None, nullptr, EPinContainerType::None, /*bIsReference =*/ false, /*InValueTerminalType=*/FEdGraphTerminalType()), /*bIsPersistent =*/false));
 	ActionRegistrar.AddBlueprintAction(ActionKey, MakeTempVarNodeSpawner(FEdGraphPinType(UEdGraphSchema_K2::PC_Double, NAME_None, nullptr, EPinContainerType::Array, /*bIsReference =*/ false, /*InValueTerminalType=*/FEdGraphTerminalType()), /*bIsPersistent =*/false));
+#endif
 	ActionRegistrar.AddBlueprintAction(ActionKey, MakeTempVarNodeSpawner(FEdGraphPinType(UEdGraphSchema_K2::PC_Boolean, NAME_None, nullptr, EPinContainerType::None, /*bIsReference =*/ false, /*InValueTerminalType=*/FEdGraphTerminalType()), /*bIsPersistent =*/false));
 	ActionRegistrar.AddBlueprintAction(ActionKey, MakeTempVarNodeSpawner(FEdGraphPinType(UEdGraphSchema_K2::PC_Boolean, NAME_None, nullptr, EPinContainerType::Array, /*bIsReference =*/ false, /*InValueTerminalType=*/FEdGraphTerminalType()), /*bIsPersistent =*/false));
 	ActionRegistrar.AddBlueprintAction(ActionKey, MakeTempVarNodeSpawner(FEdGraphPinType(UEdGraphSchema_K2::PC_String, NAME_None, nullptr, EPinContainerType::None, /*bIsReference =*/ false, /*InValueTerminalType=*/FEdGraphTerminalType()), /*bIsPersistent =*/false));
diff --git a/Engine/Source/Editor/BlueprintGraph/Private/K2Node_Timeline.cpp b/Engine/Source/Editor/BlueprintGraph/Private/K2Node_Timeline.cpp
--- a/Engine/Source/Editor/BlueprintGraph/Private/K2Node_Timeline.cpp
+++ b/Engine/Source/Editor/BlueprintGraph/Private/K2Node_Timeline.cpp
@@ -20,11 +20,50 @@
 #include "BlueprintNodeSpawner.h"
 #include "DiffResults.h"
 #include "Kismet2/Kismet2NameValidators.h"
+#include "KismetCastingUtils.h"
 #include "KismetCompilerMisc.h"
 #include "KismetCompiler.h"
 
 #define LOCTEXT_NAMESPACE "K2Node_Timeline"
 
+/////////////////////////////////////////////////////
+// FKCHandler_Timeline
+class FKCHandler_Timeline : public FNodeHandlingFunctor
+{
+public:
+	FKCHandler_Timeline(FKismetCompilerContext& InCompilerContext)
+		: FNodeHandlingFunctor(InCompilerContext)
+	{
+	}
+
+	virtual void Compile(FKismetFunctionContext& Context, UEdGraphNode* Node) override
+	{
+		UK2Node_Timeline* TimelineNode = CastChecked<UK2Node_Timeline>(Node);
+		UEdGraphPin* NewTimePin = TimelineNode->GetNewTimePin();
+		UEdGraphPin* SetNewTimePin = TimelineNode->GetSetNewTimePin();
+
+		if (NewTimePin && SetNewTimePin)
+		{
+#if ENABLE_BLUEPRINT_REAL_NUMBERS
+			// If SetNewTimePin is linked, then FKismetCompilerContext::ExpandTimelineNodes will create a function
+			// that handles the casting of the "NewTime" parameter for UTimelineComponent::SetNewTime.
+			// If SetNewTimePin is *not* linked, then we don't use "NewTime" at all, and can throw away the cast entry.
+
+			const bool bSetNewTimePinConnected = (SetNewTimePin->LinkedTo.Num() > 0);
+			if (!bSetNewTimePinConnected)
+			{
+				using namespace UE::KismetCompiler;
+				CastingUtils::RemoveRegisteredImplicitCast(Context, NewTimePin);
+			}
+#endif
+		}
+		else
+		{
+			Context.MessageLog.Error(*LOCTEXT("Error_InvalidNewTimePin", "ICE: missing one or more 'New Time'-related pins on node @@").ToString(), Node);
+		}
+	}
+};
+
 /////////////////////////////////////////////////////
 // UK2Node_Timeline
 
@@ -50,7 +89,7 @@ static FName DirectionPinName(TEXT("Direction"));
 
 namespace 
 {
-	UEdGraphPin* GetPin (const UK2Node_Timeline* Timeline, const FName PinName, EEdGraphPinDirection DesiredDirection) 
+	UEdGraphPin* GetPin(const UK2Node_Timeline* Timeline, const FName PinName, EEdGraphPinDirection DesiredDirection) 
 	{
 		UEdGraphPin* Pin = Timeline->FindPin(PinName);
 		check(Pin);
@@ -75,7 +114,11 @@ void UK2Node_Timeline::AllocateDefaultPins()
 
 	CreatePin(EGPD_Input, UEdGraphSchema_K2::PC_Exec, SetNewTimePinName);
 
+#if ENABLE_BLUEPRINT_REAL_NUMBERS
+	UEdGraphPin* NewPositionPin = CreatePin(EGPD_Input, UEdGraphSchema_K2::PC_Real, UEdGraphSchema_K2::PC_Float, NewTimePinName);
+#else
 	UEdGraphPin* NewPositionPin = CreatePin(EGPD_Input, UEdGraphSchema_K2::PC_Float, NewTimePinName);
+#endif
 	K2Schema->SetPinAutogeneratedDefaultValue(NewPositionPin, TEXT("0.0"));
 
 	CreatePin(EGPD_Output, UEdGraphSchema_K2::PC_Byte, FTimeline::GetTimelineDirectionEnum(), DirectionPinName);
@@ -100,7 +143,11 @@ void UK2Node_Timeline::AllocateDefaultPins()
 			else if (TrackId.TrackType == FTTTrackBase::TT_FloatInterp)
 			{
 				const FTTFloatTrack& FloatTrack = Timeline->FloatTracks[TrackId.TrackIndex];
+#if ENABLE_BLUEPRINT_REAL_NUMBERS
+				CreatePin(EGPD_Output, UEdGraphSchema_K2::PC_Real, UEdGraphSchema_K2::PC_Float, FloatTrack.GetTrackName());
+#else
 				CreatePin(EGPD_Output, UEdGraphSchema_K2::PC_Float, FloatTrack.GetTrackName());
+#endif
 			}
 			else if (TrackId.TrackType == FTTTrackBase::TT_VectorInterp)
 			{
@@ -382,7 +429,7 @@ UEdGraphPin* UK2Node_Timeline::GetSetNewTimePin() const
 	return GetPin(this, SetNewTimePinName, EGPD_Input);
 }
 
-bool UK2Node_Timeline::RenameTimeline (const FString& NewName)
+bool UK2Node_Timeline::RenameTimeline(const FString& NewName)
 {
 	UBlueprint* Blueprint = GetBlueprint();
 	check(Blueprint);
@@ -547,7 +594,7 @@ TSharedPtr<class INameValidatorInterface> UK2Node_Timeline::MakeNameValidator()
 
 FNodeHandlingFunctor* UK2Node_Timeline::CreateNodeHandler(FKismetCompilerContext& CompilerContext) const
 {
-	return new FNodeHandlingFunctor(CompilerContext);
+	return new FKCHandler_Timeline(CompilerContext);
 }
 
 void UK2Node_Timeline::ExpandForPin(UEdGraphPin* TimelinePin, const FName PropertyName, FKismetCompilerContext& CompilerContext, UEdGraph* SourceGraph)
diff --git a/Engine/Source/Editor/BlueprintGraph/Private/K2Node_VariableSetRef.cpp b/Engine/Source/Editor/BlueprintGraph/Private/K2Node_VariableSetRef.cpp
--- a/Engine/Source/Editor/BlueprintGraph/Private/K2Node_VariableSetRef.cpp
+++ b/Engine/Source/Editor/BlueprintGraph/Private/K2Node_VariableSetRef.cpp
@@ -4,6 +4,7 @@
 #include "K2Node_VariableSetRef.h"
 #include "EdGraphSchema_K2.h"
 #include "EdGraphUtilities.h"
+#include "KismetCastingUtils.h"
 #include "KismetCompiler.h"
 #include "VariableSetHandler.h"
 #include "BlueprintActionFilter.h"
@@ -29,29 +30,107 @@ public:
 		UK2Node_VariableSetRef* VarRefNode = CastChecked<UK2Node_VariableSetRef>(Node);
 		UEdGraphPin* ValuePin = VarRefNode->GetValuePin();
 		ValidateAndRegisterNetIfLiteral(Context, ValuePin);
+
+#if ENABLE_BLUEPRINT_REAL_NUMBERS
+		{
+			// The pins on UK2Node_VariableSetRef don't actually refer to storage, so there's nothing to register in the NetMap.
+			// However, their pins' nets do, so we need to check if a cast is required for assignment to the target net.
+			// The cast detection performed by RegisterImplicitCasts doesn't understand the reference semantics used in this node.
+			// As a result, we need to manually handle any implicit casts.
+
+			using namespace UE::KismetCompiler;
+
+			UEdGraphPin* VariablePin = VarRefNode->GetTargetPin();
+			UEdGraphPin* VariablePinNet = FEdGraphUtilities::GetNetFromPin(VariablePin);
+			UEdGraphPin* ValuePinNet = FEdGraphUtilities::GetNetFromPin(ValuePin);
+
+			if ((VariablePinNet != nullptr) && (ValuePinNet != nullptr))
+			{
+				TOptional<CastingUtils::StatementNamePair> ConversionType =
+					CastingUtils::GetFloatingPointConversionType(*ValuePinNet, *VariablePinNet);
+
+				if (ConversionType)
+				{
+					check(!ImplicitCastMap.Contains(VarRefNode));
+
+					FBPTerminal* NewTerminal = Context.CreateLocalTerminal();
+					UEdGraphNode* OwningNode = VariablePinNet->GetOwningNode();
+					NewTerminal->CopyFromPin(VariablePinNet, Context.NetNameMap->MakeValidName(VariablePin, ConversionType->Get<1>()));
+					NewTerminal->Source = OwningNode;
+
+					ImplicitCastMap.Add(VarRefNode, TPair<FBPTerminal*, EKismetCompiledStatementType>{ NewTerminal, ConversionType->Get<0>() });
+				}
+			}
+		}
+#endif
+	}
+
+	virtual void Compile(FKismetFunctionContext& Context, UEdGraphNode* Node) override
+	{
+		UK2Node_VariableSetRef* VarRefNode = CastChecked<UK2Node_VariableSetRef>(Node);
+		UEdGraphPin* VarTargetPin = VarRefNode->GetTargetPin();
+		UEdGraphPin* ValuePin = VarRefNode->GetValuePin();
+
+		InnerAssignment(Context, Node, VarTargetPin, ValuePin);
+
+		// Generate the output impulse from this node
+		GenerateSimpleThenGoto(Context, *Node);
 	}
 
+private:
+
 	void InnerAssignment(FKismetFunctionContext& Context, UEdGraphNode* Node, UEdGraphPin* VariablePin, UEdGraphPin* ValuePin)
 	{
+		UEdGraphPin* VariablePinNet = FEdGraphUtilities::GetNetFromPin(VariablePin);
+		UEdGraphPin* ValuePinNet = FEdGraphUtilities::GetNetFromPin(ValuePin);
+
 		FBPTerminal** VariableTerm = Context.NetMap.Find(VariablePin);
 		if (VariableTerm == nullptr)
 		{
-			VariableTerm = Context.NetMap.Find(FEdGraphUtilities::GetNetFromPin(VariablePin));
+			VariableTerm = Context.NetMap.Find(VariablePinNet);
 		}
 
 		FBPTerminal** ValueTerm = Context.LiteralHackMap.Find(ValuePin);
 		if (ValueTerm == nullptr)
 		{
-			ValueTerm = Context.NetMap.Find(FEdGraphUtilities::GetNetFromPin(ValuePin));
+			ValueTerm = Context.NetMap.Find(ValuePinNet);
 		}
 
 		if ((VariableTerm != nullptr) && (ValueTerm != nullptr))
 		{
-			FBlueprintCompiledStatement& Statement = Context.AppendStatementForNode(Node);
+			FBPTerminal* LHSTerm = *VariableTerm;
+			FBPTerminal* RHSTerm = *ValueTerm;
+
+#if ENABLE_BLUEPRINT_REAL_NUMBERS
+			{
+				using namespace UE::KismetCompiler;
+
+				UK2Node_VariableSetRef* VarRefNode = CastChecked<UK2Node_VariableSetRef>(Node);
+				if (TPair<FBPTerminal*, EKismetCompiledStatementType>* CastEntry = ImplicitCastMap.Find(VarRefNode))
+				{
+					FBPTerminal* CastTerminal = CastEntry->Get<0>();
+					EKismetCompiledStatementType StatementType = CastEntry->Get<1>();
 
+					FBlueprintCompiledStatement& CastStatement = Context.AppendStatementForNode(Node);
+					CastStatement.LHS = CastTerminal;
+					CastStatement.Type = StatementType;
+					CastStatement.RHS.Add(RHSTerm);
+
+					RHSTerm = CastTerminal;
+
+					ImplicitCastMap.Remove(VarRefNode);
+
+					// We've manually registered our cast statement, so it can be removed from the context.
+					CastingUtils::RemoveRegisteredImplicitCast(Context, VariablePin);
+					CastingUtils::RemoveRegisteredImplicitCast(Context, ValuePin);
+				}
+			}
+#endif
+
+			FBlueprintCompiledStatement& Statement = Context.AppendStatementForNode(Node);
 			Statement.Type = KCST_Assignment;
-			Statement.LHS = *VariableTerm;
-			Statement.RHS.Add(*ValueTerm);
+			Statement.LHS = LHSTerm;
+			Statement.RHS.Add(RHSTerm);
 
 			if (!(*VariableTerm)->IsTermWritable())
 			{
@@ -71,17 +150,7 @@ public:
 		}
 	}
 
-	virtual void Compile(FKismetFunctionContext& Context, UEdGraphNode* Node) override
-	{
-		UK2Node_VariableSetRef* VarRefNode = CastChecked<UK2Node_VariableSetRef>(Node);
-		UEdGraphPin* VarTargetPin = VarRefNode->GetTargetPin();
-		UEdGraphPin* ValuePin = VarRefNode->GetValuePin();
-
-		InnerAssignment(Context, Node, VarTargetPin, ValuePin);
-
-		// Generate the output impulse from this node
-		GenerateSimpleThenGoto(Context, *Node);
-	}
+	TMap<UK2Node_VariableSetRef*, TPair<FBPTerminal*, EKismetCompiledStatementType>> ImplicitCastMap;
 };
 
 UK2Node_VariableSetRef::UK2Node_VariableSetRef(const FObjectInitializer& ObjectInitializer)
diff --git a/Engine/Source/Editor/BlueprintGraph/Private/MathExpressionHandler.cpp b/Engine/Source/Editor/BlueprintGraph/Private/MathExpressionHandler.cpp
--- a/Engine/Source/Editor/BlueprintGraph/Private/MathExpressionHandler.cpp
+++ b/Engine/Source/Editor/BlueprintGraph/Private/MathExpressionHandler.cpp
@@ -8,7 +8,9 @@
 #include "K2Node_CallFunction.h"
 #include "K2Node_VariableGet.h"
 #include "EdGraphUtilities.h"
+#include "KismetCastingUtils.h"
 #include "KismetCompiler.h"
+#include "Kismet/KismetMathLibrary.h"
 
 #define LOCTEXT_NAMESPACE "KCHandler_MathExpression"
 
@@ -83,9 +85,9 @@ FBlueprintCompiledStatement* FKCHandler_MathExpression::GenerateFunctionRPN(UEdG
 		FProperty* Property = *It;
 		if (Property && !Property->HasAnyPropertyFlags(CPF_ReturnParm | CPF_OutParm))
 		{
+			UEdGraphPin* PinMatch = CallFunctionNode->FindPin(Property->GetFName());
 			UEdGraphPin* PinToTry = nullptr;
 			{
-				UEdGraphPin* PinMatch = CallFunctionNode->FindPin(Property->GetFName());
 				const bool bGoodPin = PinMatch && FKismetCompilerUtilities::IsTypeCompatibleWithProperty(PinMatch, Property, CompilerContext.MessageLog, CompilerContext.GetSchema(), Context.NewClass);
 				PinToTry = bGoodPin ? FEdGraphUtilities::GetNetFromPin(PinMatch) : nullptr;
 			}
@@ -129,6 +131,48 @@ FBlueprintCompiledStatement* FKCHandler_MathExpression::GenerateFunctionRPN(UEdG
 
 			if (RHSTerm)
 			{
+#if ENABLE_BLUEPRINT_REAL_NUMBERS
+				{
+					const FImplicitCastParams* CastParams =
+						Context.ImplicitCastMap.Find(PinMatch);
+					if (CastParams)
+					{
+						check(CastParams->TargetTerminal);
+
+						UFunction* CastFunction = nullptr;
+
+						switch (CastParams->CastType)
+						{
+						case KCST_DoubleToFloatCast:
+							CastFunction = 
+								UKismetMathLibrary::StaticClass()->FindFunctionByName(GET_FUNCTION_NAME_CHECKED(UKismetMathLibrary, Conv_DoubleToFloat));
+							break;
+
+						case KCST_FloatToDoubleCast:
+							CastFunction = 
+								UKismetMathLibrary::StaticClass()->FindFunctionByName(GET_FUNCTION_NAME_CHECKED(UKismetMathLibrary, Conv_FloatToDouble));
+							break;
+
+						default:
+							checkf(false, TEXT("Unsupported cast type used in math expression node: %d"), CastParams->CastType);
+						}
+
+						check(CastFunction);
+
+						FBlueprintCompiledStatement* CastStatement = new FBlueprintCompiledStatement();
+						CastStatement->FunctionToCall = CastFunction;
+						CastStatement->Type = KCST_CallFunction;
+						CastStatement->RHS.Add(RHSTerm);
+
+						RHSTerm = CastParams->TargetTerminal;
+						CastParams->TargetTerminal->InlineGeneratedParameter = CastStatement;
+						Context.AllGeneratedStatements.Add(CastStatement);
+
+						Context.ImplicitCastMap.Remove(PinMatch);
+					}
+				}
+#endif
+
 				RHSTerms.Add(RHSTerm);
 			}
 			else
@@ -242,6 +286,25 @@ void FKCHandler_MathExpression::Compile(FKismetFunctionContext& Context, UEdGrap
 		Context.AllGeneratedStatements.Add(DetachedStatement);
 		TArray<FBlueprintCompiledStatement*>& StatementList = Context.StatementsPerNode.FindOrAdd(Node);
 		StatementList.Add(DetachedStatement);
+
+#if ENABLE_BLUEPRINT_REAL_NUMBERS
+		{
+			using namespace UE::KismetCompiler;
+
+			// The casts that we care about occur *within* the nodes generated by the math expression.
+			// The expression node merely serves as a proxy, so we can silently remove the casts
+			// on the input pins that were registered by the compiler.
+
+			for (UEdGraphPin* Pin : Node_MathExpression->Pins)
+			{
+				check(Pin);
+				if (!Context.Schema->IsMetaPin(*Pin) && (Pin->Direction == EGPD_Input))
+				{
+					CastingUtils::RemoveRegisteredImplicitCast(Context, Pin);
+				}
+			}
+		}
+#endif
 	}
 	else
 	{
diff --git a/Engine/Source/Editor/BlueprintGraph/Private/Tests/TypePromotionTests.cpp b/Engine/Source/Editor/BlueprintGraph/Private/Tests/TypePromotionTests.cpp
--- a/Engine/Source/Editor/BlueprintGraph/Private/Tests/TypePromotionTests.cpp
+++ b/Engine/Source/Editor/BlueprintGraph/Private/Tests/TypePromotionTests.cpp
@@ -121,14 +121,34 @@ namespace TypePromoTestUtils
 		UK2Node* NodeName = NewObject<UK2Node_CallFunction>(/*outer*/ OwningGraph );	\
 		OwningGraph->AddNode(NodeName);
 
-#define MakeTestPin(OwningNode, PinArray, InPinName, InPinType, PinDirection)		UEdGraphPin* InPinName = UEdGraphPin::CreatePin(OwningNode);		\
+#define MakeTestPin(OwningNode, PinArray, InPinName, InPinType, PinDirection)	UEdGraphPin* InPinName = UEdGraphPin::CreatePin(OwningNode);		\
 																				InPinName->PinType.PinCategory = InPinType;							\
 																				PinArray.Add(InPinName);											\
 																				InPinName->Direction = PinDirection;
 
-#define MakeTestPins(OwningNode, OutArray)																\
+#if ENABLE_BLUEPRINT_REAL_NUMBERS
+#define MakeRealNumberTestPin(OwningNode, PinArray, InPinName, PinDirection)	UEdGraphPin* InPinName = UEdGraphPin::CreatePin(OwningNode);		\
+																				InPinName->PinType.PinCategory = UEdGraphSchema_K2::PC_Real;		\
+																				InPinName->PinType.PinSubCategory = UEdGraphSchema_K2::PC_Double;	\
+																				PinArray.Add(InPinName);											\
+																				InPinName->Direction = PinDirection;
+
+#define MakeRealNumberTestPins(OwningNode, OutArray)													\
+	MakeRealNumberTestPin(OwningNode, OutArray, RealPinA, EGPD_Input);									\
+	MakeRealNumberTestPin(OwningNode, OutArray, RealPinB, EGPD_Input);									\
+	MakeRealNumberTestPin(OwningNode, OutArray, RealOutputPin, EGPD_Output);		
+#else
+#define MakeRealNumberTestPins(OwningNode, OutArray)													\
+	MakeTestPin(OwningNode, OutArray, FloatPinA, UEdGraphSchema_K2::PC_Float, EGPD_Input);				\
+	MakeTestPin(OwningNode, OutArray, FloatPinB, UEdGraphSchema_K2::PC_Float, EGPD_Input);				\
+	MakeTestPin(OwningNode, OutArray, FloatOutputPin, UEdGraphSchema_K2::PC_Float, EGPD_Output);		\
 	MakeTestPin(OwningNode, OutArray, DoublePinA, UEdGraphSchema_K2::PC_Double, EGPD_Output);			\
 	MakeTestPin(OwningNode, OutArray, DoublePinB, UEdGraphSchema_K2::PC_Double, EGPD_Output);			\
+	MakeTestPin(OwningNode, OutArray, DoubleOutputPin, UEdGraphSchema_K2::PC_Double, EGPD_Output);
+#endif
+
+#define MakeTestPins(OwningNode, OutArray)																\
+	MakeRealNumberTestPins(OwningNode, OutArray)														\
 	MakeTestPin(OwningNode, OutArray, Int64PinA, UEdGraphSchema_K2::PC_Int64, EGPD_Output);				\
 	MakeTestPin(OwningNode, OutArray, Int64PinB, UEdGraphSchema_K2::PC_Int64, EGPD_Input);				\
 	MakeTestPin(OwningNode, OutArray, BytePinA, UEdGraphSchema_K2::PC_Byte, EGPD_Output);				\
@@ -137,10 +157,6 @@ namespace TypePromoTestUtils
 	MakeTestPin(OwningNode, OutArray, BytePinB, UEdGraphSchema_K2::PC_Byte, EGPD_Input);				\
 	MakeTestPin(OwningNode, OutArray, BoolPinA, UEdGraphSchema_K2::PC_Boolean, EGPD_Output);			\
 	MakeTestPin(OwningNode, OutArray, BoolPinB, UEdGraphSchema_K2::PC_Boolean, EGPD_Input);				\
-	MakeTestPin(OwningNode, OutArray, DoubleOutputPin, UEdGraphSchema_K2::PC_Double, EGPD_Output);		\
-	MakeTestPin(OwningNode, OutArray, FloatPinA, UEdGraphSchema_K2::PC_Float, EGPD_Input);				\
-	MakeTestPin(OwningNode, OutArray, FloatPinB, UEdGraphSchema_K2::PC_Float, EGPD_Input);				\
-	MakeTestPin(OwningNode, OutArray, FloatOutputPin, UEdGraphSchema_K2::PC_Float, EGPD_Output);		\
 	MakeTestPin(OwningNode, OutArray, BoolOutputPin, UEdGraphSchema_K2::PC_Boolean, EGPD_Output);		\
 	MakeTestPin(OwningNode, OutArray, IntPinA, UEdGraphSchema_K2::PC_Int, EGPD_Output);					\
 	MakeTestPin(OwningNode, OutArray, VecInputPinA, UEdGraphSchema_K2::PC_Struct, EGPD_Input);			\
@@ -161,6 +177,34 @@ namespace TypePromoTestUtils
 IMPLEMENT_SIMPLE_AUTOMATION_TEST(FTypePromotionTest, "Blueprints.Compiler.TypePromotion", EAutomationTestFlags::EditorContext | EAutomationTestFlags::SmokeFilter)
 bool FTypePromotionTest::RunTest(const FString& Parameters)
 {
+#if ENABLE_BLUEPRINT_REAL_NUMBERS
+	FEdGraphPinType RealPin = {};		RealPin.PinCategory = UEdGraphSchema_K2::PC_Real; RealPin.PinSubCategory = UEdGraphSchema_K2::PC_Double;
+	FEdGraphPinType IntPin = {};		IntPin.PinCategory = UEdGraphSchema_K2::PC_Int;
+	FEdGraphPinType Int64Pin = {};		Int64Pin.PinCategory = UEdGraphSchema_K2::PC_Int64;
+	FEdGraphPinType BytePin = {};		BytePin.PinCategory = UEdGraphSchema_K2::PC_Byte;
+	FEdGraphPinType VecPin = {};		VecPin.PinCategory = UEdGraphSchema_K2::PC_Struct; VecPin.PinSubCategoryObject = TBaseStructure<FVector>::Get();
+
+	// Test promotions that should happen
+	TestEqual(TEXT("Testing real to vector"), FTypePromotion::GetHigherType(RealPin, VecPin), FTypePromotion::ETypeComparisonResult::TypeBHigher);
+
+	TestEqual(TEXT("Testing int to real"), FTypePromotion::GetHigherType(IntPin, RealPin), FTypePromotion::ETypeComparisonResult::TypeBHigher);
+	TestEqual(TEXT("Testing int to int64"), FTypePromotion::GetHigherType(IntPin, Int64Pin), FTypePromotion::ETypeComparisonResult::TypeBHigher);
+
+	TestEqual(TEXT("Testing byte to int"), FTypePromotion::GetHigherType(BytePin, IntPin), FTypePromotion::ETypeComparisonResult::TypeBHigher);
+	TestEqual(TEXT("Testing byte to int64"), FTypePromotion::GetHigherType(BytePin, Int64Pin), FTypePromotion::ETypeComparisonResult::TypeBHigher);
+
+	TestEqual(TEXT("Testing real to int64"), FTypePromotion::GetHigherType(RealPin, Int64Pin), FTypePromotion::ETypeComparisonResult::TypeBHigher);
+
+	// Test Equality of pins
+	TestEqual(TEXT("Testing byte == byte"), FTypePromotion::GetHigherType(BytePin, BytePin), FTypePromotion::ETypeComparisonResult::TypesEqual);
+	TestEqual(TEXT("Testing real == real"), FTypePromotion::GetHigherType(RealPin, RealPin), FTypePromotion::ETypeComparisonResult::TypesEqual);
+	TestEqual(TEXT("Testing int == int"), FTypePromotion::GetHigherType(IntPin, IntPin), FTypePromotion::ETypeComparisonResult::TypesEqual);
+	TestEqual(TEXT("Testing int64 == int64"), FTypePromotion::GetHigherType(Int64Pin, Int64Pin), FTypePromotion::ETypeComparisonResult::TypesEqual);
+
+	// Test promotions that should not happen
+	TestEqual(TEXT("Testing int64 cannot go to byte"), FTypePromotion::GetHigherType(Int64Pin, BytePin), FTypePromotion::ETypeComparisonResult::TypeAHigher);
+	TestEqual(TEXT("Testing int64 cannot go to int"), FTypePromotion::GetHigherType(Int64Pin, IntPin), FTypePromotion::ETypeComparisonResult::TypeAHigher);
+#else
 	FEdGraphPinType DoublePin = {};		DoublePin.PinCategory = UEdGraphSchema_K2::PC_Double;
 	FEdGraphPinType FloatPin = {};		FloatPin.PinCategory = UEdGraphSchema_K2::PC_Float;
 	FEdGraphPinType IntPin = {};		IntPin.PinCategory = UEdGraphSchema_K2::PC_Int;
@@ -193,6 +237,7 @@ bool FTypePromotionTest::RunTest(const FString& Parameters)
 	TestEqual(TEXT("Testing int64 cannot go to byte"), FTypePromotion::GetHigherType(Int64Pin, BytePin), FTypePromotion::ETypeComparisonResult::TypeAHigher);
 	TestEqual(TEXT("Testing int64 cannot go to int"), FTypePromotion::GetHigherType(Int64Pin, IntPin), FTypePromotion::ETypeComparisonResult::TypeAHigher);
 	TestEqual(TEXT("Testing int64 cannot go to float"), FTypePromotion::GetHigherType(Int64Pin, FloatPin), FTypePromotion::ETypeComparisonResult::TypeAHigher);
+#endif
 
 	return true;
 }
@@ -214,20 +259,20 @@ bool FFindBestMatchingFunc::RunTest(const FString& Parameters)
 	TArray<UEdGraphPin*> PinTypes = {};
 	MakeTestPins(TestNode, PinTypes);
 
-	#define TestMatchingFunc(OpName, TestPins, ExpectedFuncName) \
-	{\
-		const UFunction* FoundFunc = FTypePromotion::FindBestMatchingFunc(OpName, TestPins);\
-		const FName ExpectedName = FName(ExpectedFuncName);\
-		const FString TestNullMessage = FString::Printf(TEXT(" Find Function '%s' null check"), *ExpectedName.ToString()); \
-		if (TestNotNull(TestNullMessage, FoundFunc)) \
-		{\
-		 	const FString PinTypesString = TypePromoTestUtils::GetPinListDisplayName(TestPins);\
-			const FString TestMessage = FString::Printf(TEXT("Given pins %s Expecting Function '%s' and got '%s'"),\
-				*PinTypesString,\
-				*ExpectedName.ToString(),\
-				*FoundFunc->GetFName().ToString());\
-				TestEqual(TestMessage, FoundFunc->GetFName(), ExpectedName);\
-		}\
+	#define TestMatchingFunc(OpName, TestPins, ExpectedFuncName)															\
+	{																														\
+		const UFunction* FoundFunc = FTypePromotion::FindBestMatchingFunc(OpName, TestPins);								\
+		const FName ExpectedName = FName(ExpectedFuncName);																	\
+		const FString TestNullMessage = FString::Printf(TEXT(" Find Function '%s' null check"), *ExpectedName.ToString());	\
+		if (TestNotNull(TestNullMessage, FoundFunc))																		\
+		{																													\
+		 	const FString PinTypesString = TypePromoTestUtils::GetPinListDisplayName(TestPins);								\
+			const FString TestMessage = FString::Printf(TEXT("Given pins %s Expecting Function '%s' and got '%s'"),			\
+				*PinTypesString,																							\
+				*ExpectedName.ToString(),																					\
+				*FoundFunc->GetFName().ToString());																			\
+				TestEqual(TestMessage, FoundFunc->GetFName(), ExpectedName);												\
+		}																													\
 	}
 
 	{
@@ -238,6 +283,90 @@ bool FFindBestMatchingFunc::RunTest(const FString& Parameters)
 		TestMatchingFunc(TEXT("Add"), TestPins, TEXT("Add_Vector2DVector2D"));
 	}
 
+#if ENABLE_BLUEPRINT_REAL_NUMBERS
+	// Multiply_VectorVector given A real input, vector input, and a vector output
+	{
+		TArray<UEdGraphPin*> TestPins =
+		{
+			RealPinA, VecInputPinB, VecOutputPinA,
+		};
+		TestMatchingFunc(TEXT("Multiply"), TestPins, TEXT("Multiply_VectorVector"));
+	}
+
+	// Multiply_VectorVector given a real, vector, real
+	// Order shouldn't matter when passing these pins in, which is what we are testing here
+	{
+		TArray<UEdGraphPin*> TestPins =
+		{
+			RealPinA, VecOutputPinA, VecInputPinA,
+		};
+		TestMatchingFunc(TEXT("Multiply"), TestPins, TEXT("Multiply_VectorVector"));
+	}
+
+	// Multiply_VectorVector given two vector inputs and a vector output
+	{
+		TArray<UEdGraphPin*> TestPins =
+		{
+			VecInputPinA, VecInputPinB, VecOutputPinA,
+		};
+		TestMatchingFunc(TEXT("Multiply"), TestPins, TEXT("Multiply_VectorVector"));
+	}
+
+	// Add_DoubleDouble
+	{
+		TArray<UEdGraphPin*> TestPins =
+		{
+			RealPinA, RealPinB, RealOutputPin
+		};
+		TestMatchingFunc(TEXT("Add"), TestPins, TEXT("Add_DoubleDouble"));
+	}
+
+	// Subtract_DoubleDouble
+	{
+		TArray<UEdGraphPin*> TestPins =
+		{
+			RealPinA, RealPinB, RealOutputPin
+		};
+		TestMatchingFunc(TEXT("Subtract"), TestPins, TEXT("Subtract_DoubleDouble"));
+	}
+
+	// Add_DoubleDouble given only one real pin. This simulates the first connection being made to a 
+	// promotable operator, in which case we should default to a regular old real + real
+	{
+		TArray<UEdGraphPin*> TestPins =
+		{
+			RealPinA,
+		};
+		TestMatchingFunc(TEXT("Add"), TestPins, TEXT("Add_DoubleDouble"));
+	}
+
+	// Less_DoubleDouble given a real
+	{
+		TArray<UEdGraphPin*> TestPins =
+		{
+			RealPinA, BoolOutputPin
+		};
+		TestMatchingFunc(TEXT("Less"), TestPins, TEXT("Less_DoubleDouble"));
+	}
+
+	// Less_DoubleDouble given just a single real
+	{
+		TArray<UEdGraphPin*> TestPins =
+		{
+			RealPinA,
+		};
+		TestMatchingFunc(TEXT("Less"), TestPins, TEXT("Less_DoubleDouble"));
+	}
+
+	// Greater_DoubleDouble given reals
+	{
+		TArray<UEdGraphPin*> TestPins =
+		{
+			RealPinA, RealPinA
+		};
+		TestMatchingFunc(TEXT("Greater"), TestPins, TEXT("Greater_DoubleDouble"));
+	}
+#else
 	// Multiply_VectorVector given A float input, vector input, and a vector output
 	{
 		TArray<UEdGraphPin*> TestPins =
@@ -328,6 +457,7 @@ bool FFindBestMatchingFunc::RunTest(const FString& Parameters)
 		};
 		TestMatchingFunc(TEXT("Greater"), TestPins, TEXT("Greater_DoubleDouble"));
 	}
+#endif
 	
 	TypePromoTestUtils::CleanupTestPins(PinTypes);
 	TestNode->MarkAsGarbage();
@@ -358,8 +488,8 @@ bool FPromotableTypeToOperator::RunTest(const FString& Parameters)
 
 	const UEdGraphSchema_K2* K2Schema = GetDefault<UEdGraphSchema_K2>();
 
-	const TMap<FName, TArray<FName>>* const PromoTable = FTypePromotion::GetPrimativePromotionTable();
-	TestNotNull(TEXT("Primative Promotion table exists"), PromoTable);
+	const TMap<FName, TArray<FName>>* const PromoTable = FTypePromotion::GetPrimitivePromotionTable();
+	TestNotNull(TEXT("Primitive Promotion table exists"), PromoTable);
 
 	for (const FName OpName : AllOpNames)
 	{
@@ -635,6 +765,16 @@ bool FPromotableOpNodeAddPinInterface::RunTest(const FString& Parameters)
 		UEdGraphPin* AdditonalPin = AddNode->GetAdditionalPin(2);
 		TestNotNull(TEXT("Additional Pin is not null"), AdditonalPin);
 		
+#if ENABLE_BLUEPRINT_REAL_NUMBERS
+		// Connect a real pin to the additional input pin
+		const bool bConnected = TypePromoTestUtils::TestPromotedConnection(AdditonalPin, RealOutputPin);
+		TestTrue(TEXT("Connection to additional pin success"), bConnected);
+
+		// The other pins have propagated correctly with this new connection 
+		TestTrue(TEXT("Top Pin type propegates to new connection"), TopInputPin->PinType.PinCategory == RealOutputPin->PinType.PinCategory);
+		TestTrue(TEXT("Bottom Pin type propegates to new connection"), BottomInputPin->PinType.PinCategory == RealOutputPin->PinType.PinCategory);
+		TestTrue(TEXT("Out Pin type propegates to new connection"), OutputPin->PinType.PinCategory == RealOutputPin->PinType.PinCategory);
+#else
 		// Connect a float pin to the additional input pin
 		const bool bConnected = TypePromoTestUtils::TestPromotedConnection(AdditonalPin, FloatOutputPin);
 		TestTrue(TEXT("Connection to additional pin success"), bConnected);
@@ -643,6 +783,7 @@ bool FPromotableOpNodeAddPinInterface::RunTest(const FString& Parameters)
 		TestTrue(TEXT("Top Pin type propegates to new connection"), TopInputPin->PinType.PinCategory == FloatOutputPin->PinType.PinCategory);
 		TestTrue(TEXT("Bottom Pin type propegates to new connection"), BottomInputPin->PinType.PinCategory == FloatOutputPin->PinType.PinCategory);
 		TestTrue(TEXT("Out Pin type propegates to new connection"), OutputPin->PinType.PinCategory == FloatOutputPin->PinType.PinCategory);
+#endif
 
 		// Removing the only pin with a connection with reset the node to wildcard
 		AddNode->RemoveInputPin(AdditonalPin);
@@ -702,10 +843,17 @@ bool FPromotableOperatorConnectionChanged::RunTest(const FString& Parameters)
 
 		check(TopInputPin && BottomInputPin && OutputPin);
 
+#if ENABLE_BLUEPRINT_REAL_NUMBERS
+		const bool bConnected = K2Schema->TryCreateConnection(TopInputPin, RealOutputPin);
+		AddNode->NotifyPinConnectionListChanged(TopInputPin);
+
+		TestTrue(TEXT("Bottom Pin type propagates to real"), bConnected && BottomInputPin->PinType.PinCategory == RealPinB->PinType.PinCategory);
+#else
 		const bool bConnected = K2Schema->TryCreateConnection(TopInputPin, FloatOutputPin);
 		AddNode->NotifyPinConnectionListChanged(TopInputPin);
 
-		TestTrue(TEXT("Bottom Pin type propegates to float"), bConnected && BottomInputPin->PinType.PinCategory == FloatPinB->PinType.PinCategory);
+		TestTrue(TEXT("Bottom Pin type propagates to float"), bConnected && BottomInputPin->PinType.PinCategory == FloatPinB->PinType.PinCategory);
+#endif
 	}
 
 	// Connecting a vector output should make the other input be a vector as well
@@ -772,7 +920,7 @@ bool FPromotableOperatorPrimitivePromotions::RunTest(const FString& Parameters)
 	FTypePromotion::ClearNodeSpawners();
 	FBlueprintActionDatabase::Get().RefreshAll();
 
-	MakeTestableBP(BP_Primative_Connections, TestGraph);
+	MakeTestableBP(BP_Primitive_Connections, TestGraph);
 	MakeTestableNode(TestNode, TestGraph);
 
 	// Create test pins!
@@ -781,8 +929,8 @@ bool FPromotableOperatorPrimitivePromotions::RunTest(const FString& Parameters)
 
 	const UEdGraphSchema_K2* K2Schema = GetDefault<UEdGraphSchema_K2>();
 
-	const TMap<FName, TArray<FName>>* const PromoTable = FTypePromotion::GetPrimativePromotionTable();
-	TestNotNull(TEXT("Primative Promotion table exists"), PromoTable);
+	const TMap<FName, TArray<FName>>* const PromoTable = FTypePromotion::GetPrimitivePromotionTable();
+	TestNotNull(TEXT("Primitive Promotion table exists"), PromoTable);
 	
 	if (!PromoTable)
 	{
@@ -838,8 +986,8 @@ bool FPromotableOperatorPrimitivePromotions::RunTest(const FString& Parameters)
 	{
 		TypePromoTestUtils::CleanupTestPins(PinTypes);
 
-		BP_Primative_Connections->MarkAsGarbage();
-		BP_Primative_Connections->Rename(nullptr, nullptr, REN_DontCreateRedirectors);
+		BP_Primitive_Connections->MarkAsGarbage();
+		BP_Primitive_Connections->Rename(nullptr, nullptr, REN_DontCreateRedirectors);
 		TestGraph->MarkAsGarbage();
 		TestNode->MarkAsGarbage();
 	}
diff --git a/Engine/Source/Editor/BlueprintGraph/Private/VariableSetHandler.cpp b/Engine/Source/Editor/BlueprintGraph/Private/VariableSetHandler.cpp
--- a/Engine/Source/Editor/BlueprintGraph/Private/VariableSetHandler.cpp
+++ b/Engine/Source/Editor/BlueprintGraph/Private/VariableSetHandler.cpp
@@ -75,7 +75,11 @@ void FKCHandler_VariableSet::InnerAssignment(FKismetFunctionContext& Context, UE
 
 	if (VariableTerm && ValueTerm)
 	{
-		FKismetCompilerUtilities::CreateObjectAssignmentStatement(Context, Node, *ValueTerm, *VariableTerm);
+		FKismetCompilerUtilities::CreateObjectAssignmentStatement(Context,
+																  Node,
+																  *ValueTerm,
+																  *VariableTerm,
+																  (UsesVariablePinAsKey() ? VariablePin : nullptr));
 
 		if (!(*VariableTerm)->IsTermWritable())
 		{
diff --git a/Engine/Source/Editor/BlueprintGraph/Public/CallFunctionHandler.h b/Engine/Source/Editor/BlueprintGraph/Public/CallFunctionHandler.h
--- a/Engine/Source/Editor/BlueprintGraph/Public/CallFunctionHandler.h
+++ b/Engine/Source/Editor/BlueprintGraph/Public/CallFunctionHandler.h
@@ -69,6 +69,6 @@ public:
 	virtual void CheckIfFunctionIsCallable(UFunction* Function, FKismetFunctionContext& Context, UEdGraphNode* Node);
 	virtual void AdditionalCompiledStatementHandling(FKismetFunctionContext& Context, UEdGraphNode* Node, FBlueprintCompiledStatement& Statement) {}
 
-protected:
+private:
 	TMap<UEdGraphPin*, FBPTerminal*> InterfaceTermMap;
 };
diff --git a/Engine/Source/Editor/BlueprintGraph/Public/VariableSetHandler.h b/Engine/Source/Editor/BlueprintGraph/Public/VariableSetHandler.h
--- a/Engine/Source/Editor/BlueprintGraph/Public/VariableSetHandler.h
+++ b/Engine/Source/Editor/BlueprintGraph/Public/VariableSetHandler.h
@@ -25,4 +25,10 @@ public:
 	void GenerateAssigments(FKismetFunctionContext& Context, UEdGraphNode* Node);
 	virtual void Compile(FKismetFunctionContext& Context, UEdGraphNode* Node) override;
 	virtual void Transform(FKismetFunctionContext& Context, UEdGraphNode* Node) override;
+
+protected:
+
+	// Used for implicit casting.
+	// Some nodes need to use the variable pin when performing a lookup in the implicit cast table.
+	virtual bool UsesVariablePinAsKey() const { return false; }
 };
diff --git a/Engine/Source/Editor/GraphEditor/Private/GraphEditorSettings.cpp b/Engine/Source/Editor/GraphEditor/Private/GraphEditorSettings.cpp
--- a/Engine/Source/Editor/GraphEditor/Private/GraphEditorSettings.cpp
+++ b/Engine/Source/Editor/GraphEditor/Private/GraphEditorSettings.cpp
@@ -42,6 +42,7 @@ UGraphEditorSettings::UGraphEditorSettings( const FObjectInitializer& ObjectInit
 	Int64PinTypeColor = FLinearColor(0.413575f, 0.770000f, 0.429609f, 1.0f);	
 	FloatPinTypeColor = FLinearColor(0.357667f, 1.0f, 0.060000f, 1.0f);			// bright green
 	DoublePinTypeColor = FLinearColor(0.039216f, 0.666667f, 0.0f, 1.0f);		// darker green
+	RealPinTypeColor = FLinearColor(0.039216f, 0.666667f, 0.0f, 1.0f);			// darker green
 	NamePinTypeColor = FLinearColor(0.607717f, 0.224984f, 1.0f, 1.0f);			// lilac
 	DelegatePinTypeColor = FLinearColor(1.0f, 0.04f, 0.04f, 1.0f);				// bright red
 	ObjectPinTypeColor = FLinearColor(0.0f, 0.4f, 0.910000f, 1.0f);				// sharp blue
diff --git a/Engine/Source/Editor/GraphEditor/Private/NodeFactory.cpp b/Engine/Source/Editor/GraphEditor/Private/NodeFactory.cpp
--- a/Engine/Source/Editor/GraphEditor/Private/NodeFactory.cpp
+++ b/Engine/Source/Editor/GraphEditor/Private/NodeFactory.cpp
@@ -325,6 +325,12 @@ TSharedPtr<SGraphPin> FNodeFactory::CreateK2PinWidget(UEdGraphPin* InPin)
 	{
 		return SNew(SGraphPinNum<int64>, InPin);
 	}
+#if ENABLE_BLUEPRINT_REAL_NUMBERS
+	else if (InPin->PinType.PinCategory == UEdGraphSchema_K2::PC_Real)
+	{
+		return SNew(SGraphPinNum<double>, InPin);
+	}
+#else
 	else if (InPin->PinType.PinCategory == UEdGraphSchema_K2::PC_Float)
 	{
 		return SNew(SGraphPinNum<float>, InPin);
@@ -333,6 +339,7 @@ TSharedPtr<SGraphPin> FNodeFactory::CreateK2PinWidget(UEdGraphPin* InPin)
 	{
 		return SNew(SGraphPinNum<double>, InPin);
 	}
+#endif
 	else if (InPin->PinType.PinCategory == UEdGraphSchema_K2::PC_String || InPin->PinType.PinCategory == UEdGraphSchema_K2::PC_Name)
 	{
 		return SNew(SGraphPinString, InPin);
@@ -342,6 +349,7 @@ TSharedPtr<SGraphPin> FNodeFactory::CreateK2PinWidget(UEdGraphPin* InPin)
 		// If you update this logic you'll probably need to update UEdGraphSchema_K2::ShouldHidePinDefaultValue!
 		UScriptStruct* ColorStruct = TBaseStructure<FLinearColor>::Get();
 		UScriptStruct* VectorStruct = TBaseStructure<FVector>::Get();
+		UScriptStruct* Vector3fStruct = TVariantStructure<FVector3f>::Get();
 		UScriptStruct* Vector2DStruct = TBaseStructure<FVector2D>::Get();
 		UScriptStruct* RotatorStruct = TBaseStructure<FRotator>::Get();
 
@@ -349,7 +357,7 @@ TSharedPtr<SGraphPin> FNodeFactory::CreateK2PinWidget(UEdGraphPin* InPin)
 		{
 			return SNew(SGraphPinColor, InPin);
 		}
-		else if ((InPin->PinType.PinSubCategoryObject == VectorStruct) || (InPin->PinType.PinSubCategoryObject == RotatorStruct))
+		else if ((InPin->PinType.PinSubCategoryObject == VectorStruct) || (InPin->PinType.PinSubCategoryObject == Vector3fStruct) || (InPin->PinType.PinSubCategoryObject == RotatorStruct))
 		{
 			return SNew(SGraphPinVector, InPin);
 		}
diff --git a/Engine/Source/Editor/GraphEditor/Public/GraphEditorSettings.h b/Engine/Source/Editor/GraphEditor/Public/GraphEditorSettings.h
--- a/Engine/Source/Editor/GraphEditor/Public/GraphEditorSettings.h
+++ b/Engine/Source/Editor/GraphEditor/Public/GraphEditorSettings.h
@@ -154,7 +154,7 @@ public:
 	FLinearColor IntPinTypeColor;
 
 	/** Integer64 pin type color */
-	UPROPERTY(EditAnywhere, config, Category = PinColors)
+	UPROPERTY(EditAnywhere, config, Category=PinColors)
 	FLinearColor Int64PinTypeColor;
 
 	/** Floating-point pin type color */
@@ -162,19 +162,23 @@ public:
 	FLinearColor FloatPinTypeColor;
 
 	/** Double pin type color */
-	UPROPERTY(EditAnywhere, config, Category = PinColors)
+	UPROPERTY(EditAnywhere, config, Category=PinColors)
 	FLinearColor DoublePinTypeColor;
 
+	/** Real pin type color */
+	UPROPERTY(EditAnywhere, config, Category=PinColors)
+	FLinearColor RealPinTypeColor;
+
 	/** Name pin type color */
 	UPROPERTY(EditAnywhere, config, Category=PinColors)
 	FLinearColor NamePinTypeColor;
 
 	/** Asset pin type color */
-	UPROPERTY(EditAnywhere, config, Category = PinColors)
+	UPROPERTY(EditAnywhere, config, Category=PinColors)
 	FLinearColor SoftObjectPinTypeColor;
 
 	/** Asset Class pin type color */
-	UPROPERTY(EditAnywhere, config, Category = PinColors)
+	UPROPERTY(EditAnywhere, config, Category=PinColors)
 	FLinearColor SoftClassPinTypeColor;
 
 	/** Delegate pin type color */
diff --git a/Engine/Source/Editor/KismetCompiler/Private/KismetCastingUtils.cpp b/Engine/Source/Editor/KismetCompiler/Private/KismetCastingUtils.cpp
new file mode 100644
--- /dev/null
+++ b/Engine/Source/Editor/KismetCompiler/Private/KismetCastingUtils.cpp
@@ -0,0 +1,424 @@
+// Copyright Epic Games, Inc. All Rights Reserved.
+
+#include "KismetCastingUtils.h"
+
+#include "KismetCompiledFunctionContext.h"
+#include "KismetCompilerMisc.h"
+
+#if ENABLE_BLUEPRINT_REAL_NUMBERS
+
+namespace UE::KismetCompiler::CastingUtils::Private
+{
+
+enum class EPinType : uint8
+{
+	Float,
+	FloatArray,
+	FloatSet,
+
+	Double,
+	DoubleArray,
+	DoubleSet,
+
+	Vector3f,
+	Vector3fArray,
+	Vector3fSet,
+
+	Vector,
+	VectorArray,
+	VectorSet,
+
+	// Maps are particularly sinister, and are categorized separately here.
+	// Keys and values are casted independently of one another.
+
+	FloatKeyOtherValueMap,
+	DoubleKeyOtherValueMap,
+	OtherKeyFloatValueMap,
+	OtherKeyDoubleValueMap,
+	FloatKeyFloatValueMap,
+	FloatKeyDoubleValueMap,
+	DoubleKeyFloatValueMap,
+	DoubleKeyDoubleValueMap,
+
+	Other,
+};
+
+EPinType GetPinType(const UEdGraphPin& Pin)
+{
+	static UScriptStruct* VectorStruct = TBaseStructure<FVector>::Get();
+	static UScriptStruct* Vector3fStruct = TVariantStructure<FVector3f>::Get();
+
+	if (Pin.PinType.IsMap())
+	{
+		if ((Pin.PinType.PinCategory == UEdGraphSchema_K2::PC_Real) && (Pin.PinType.PinSubCategory == UEdGraphSchema_K2::PC_Float))
+		{
+			if ((Pin.PinType.PinValueType.TerminalCategory == UEdGraphSchema_K2::PC_Real) && (Pin.PinType.PinValueType.TerminalSubCategory == UEdGraphSchema_K2::PC_Float))
+			{
+				return EPinType::FloatKeyFloatValueMap;
+			}
+			else if ((Pin.PinType.PinValueType.TerminalCategory == UEdGraphSchema_K2::PC_Real) && (Pin.PinType.PinValueType.TerminalSubCategory == UEdGraphSchema_K2::PC_Double))
+			{
+				return EPinType::FloatKeyDoubleValueMap;
+			}
+			else
+			{
+				return EPinType::FloatKeyOtherValueMap;
+			}
+		}
+		else if ((Pin.PinType.PinCategory == UEdGraphSchema_K2::PC_Real) && (Pin.PinType.PinSubCategory == UEdGraphSchema_K2::PC_Double))
+		{
+			if ((Pin.PinType.PinValueType.TerminalCategory == UEdGraphSchema_K2::PC_Real) && (Pin.PinType.PinValueType.TerminalSubCategory == UEdGraphSchema_K2::PC_Float))
+			{
+				return EPinType::DoubleKeyFloatValueMap;
+			}
+			else if ((Pin.PinType.PinValueType.TerminalCategory == UEdGraphSchema_K2::PC_Real) && (Pin.PinType.PinValueType.TerminalSubCategory == UEdGraphSchema_K2::PC_Double))
+			{
+				return EPinType::DoubleKeyDoubleValueMap;
+			}
+			else
+			{
+				return EPinType::DoubleKeyOtherValueMap;
+			}
+		}
+		else
+		{
+			if ((Pin.PinType.PinValueType.TerminalCategory == UEdGraphSchema_K2::PC_Real) && (Pin.PinType.PinValueType.TerminalSubCategory == UEdGraphSchema_K2::PC_Float))
+			{
+				return EPinType::OtherKeyFloatValueMap;
+			}
+			else if ((Pin.PinType.PinValueType.TerminalCategory == UEdGraphSchema_K2::PC_Real) && (Pin.PinType.PinValueType.TerminalSubCategory == UEdGraphSchema_K2::PC_Double))
+			{
+				return EPinType::OtherKeyDoubleValueMap;
+			}
+		}
+	}
+	else
+	{
+		if ((Pin.PinType.PinCategory == UEdGraphSchema_K2::PC_Real) && (Pin.PinType.PinSubCategory == UEdGraphSchema_K2::PC_Float))
+		{
+			if (Pin.PinType.IsArray())
+			{
+				return EPinType::FloatArray;
+			}
+			else if (Pin.PinType.IsSet())
+			{
+				return EPinType::FloatSet;
+			}
+			else
+			{
+				return EPinType::Float;
+			}
+		}
+		else if ((Pin.PinType.PinCategory == UEdGraphSchema_K2::PC_Real) && (Pin.PinType.PinSubCategory == UEdGraphSchema_K2::PC_Double))
+		{
+			if (Pin.PinType.IsArray())
+			{
+				return EPinType::DoubleArray;
+			}
+			else if (Pin.PinType.IsSet())
+			{
+				return EPinType::DoubleSet;
+			}
+			else
+			{
+				return EPinType::Double;
+			}
+		}
+		else if (Pin.PinType.PinCategory == UEdGraphSchema_K2::PC_Struct)
+		{
+			if (Pin.PinType.PinSubCategoryObject == Vector3fStruct)
+			{
+				if (Pin.PinType.IsArray())
+				{
+					return EPinType::Vector3fArray;
+				}
+				else if (Pin.PinType.IsSet())
+				{
+					return EPinType::Vector3fSet;
+				}
+				else
+				{
+					return EPinType::Vector3f;
+				}
+			}
+			else if (Pin.PinType.PinSubCategoryObject == VectorStruct)
+			{
+				if (Pin.PinType.IsArray())
+				{
+					return EPinType::VectorArray;
+				}
+				else if (Pin.PinType.IsSet())
+				{
+					return EPinType::VectorSet;
+				}
+				else
+				{
+					return EPinType::Vector;
+				}
+			}
+		}
+	}
+
+	return EPinType::Other;
+}
+
+} // namespace UE::KismetCompiler::CastingUtils::Private
+
+namespace UE::KismetCompiler::CastingUtils
+{
+
+TOptional<EKismetCompiledStatementType> GetInverseCastStatement(EKismetCompiledStatementType Statement)
+{
+	TOptional<EKismetCompiledStatementType> Result;
+
+	switch (Statement)
+	{
+	case KCST_FloatToDoubleCast:
+		Result = KCST_DoubleToFloatCast;
+		break;
+	case KCST_FloatToDoubleArrayCast:
+		Result = KCST_DoubleToFloatArrayCast;
+		break;
+	case KCST_FloatToDoubleSetCast:
+		Result = KCST_DoubleToFloatSetCast;
+		break;
+
+	case KCST_DoubleToFloatCast:
+		Result = KCST_FloatToDoubleCast;
+		break;
+	case KCST_DoubleToFloatArrayCast:
+		Result = KCST_FloatToDoubleArrayCast;
+		break;
+	case KCST_DoubleToFloatSetCast:
+		Result = KCST_FloatToDoubleSetCast;
+		break;
+
+	case KCST_Vector3fToVectorCast:
+		Result = KCST_VectorToVector3fCast;
+		break;
+	case KCST_Vector3fToVectorArrayCast:
+		Result = KCST_VectorToVector3fArrayCast;
+		break;
+	case KCST_Vector3fToVectorSetCast:
+		Result = KCST_VectorToVector3fSetCast;
+		break;
+
+	case KCST_VectorToVector3fCast:
+		Result = KCST_Vector3fToVectorCast;
+		break;
+	case KCST_VectorToVector3fArrayCast:
+		Result = KCST_Vector3fToVectorArrayCast;
+		break;
+	case KCST_VectorToVector3fSetCast:
+		Result = KCST_Vector3fToVectorSetCast;
+		break;
+
+	case KCST_FloatToDoubleKeys_MapCast:
+		Result = KCST_DoubleToFloatKeys_MapCast;
+		break;
+	case KCST_DoubleToFloatKeys_MapCast:
+		Result = KCST_FloatToDoubleKeys_MapCast;
+		break;
+	case KCST_FloatToDoubleValues_MapCast:
+		Result = KCST_DoubleToFloatValues_MapCast;
+		break;
+	case KCST_DoubleToFloatValues_MapCast:
+		Result = KCST_FloatToDoubleValues_MapCast;
+		break;
+
+	case KCST_FloatToDoubleKeys_FloatToDoubleValues_MapCast:
+		Result = KCST_DoubleToFloatKeys_DoubleToFloatValues_MapCast;
+		break;
+	case KCST_DoubleToFloatKeys_FloatToDoubleValues_MapCast:
+		Result = KCST_FloatToDoubleKeys_DoubleToFloatValues_MapCast;
+		break;
+	case KCST_DoubleToFloatKeys_DoubleToFloatValues_MapCast:
+		Result = KCST_FloatToDoubleKeys_FloatToDoubleValues_MapCast;
+		break;
+	case KCST_FloatToDoubleKeys_DoubleToFloatValues_MapCast:
+		Result = KCST_DoubleToFloatKeys_FloatToDoubleValues_MapCast;
+		break;
+	}
+
+	return Result;
+}
+
+void RegisterImplicitCasts(FKismetFunctionContext& Context)
+{
+	using namespace UE::KismetCompiler::CastingUtils::Private;
+
+	auto AddCastMapping = [&Context](UEdGraphPin* DestinationPin, EKismetCompiledStatementType CastType, const TCHAR* TermName)
+	{
+		check(DestinationPin);
+		check(TermName);
+
+		FBPTerminal* NewTerm = Context.CreateLocalTerminal();
+		UEdGraphNode* OwningNode = DestinationPin->GetOwningNode();
+		NewTerm->CopyFromPin(DestinationPin, Context.NetNameMap->MakeValidName(DestinationPin, TermName));
+		NewTerm->Source = OwningNode;
+
+		Context.ImplicitCastMap.Add(DestinationPin, FImplicitCastParams{ CastType, NewTerm, OwningNode });
+	};
+
+	// The current context's NetMap can be a mix of input and output pin types.
+	// We need to check both pin types in order to get adequate coverage for potential cast points.
+	for (const auto& It : Context.NetMap)
+	{
+		UEdGraphPin* CurrentPin = It.Key;
+		check(CurrentPin);
+
+		bool bIsConnectedOutput =
+			(CurrentPin->Direction == EGPD_Output) && (CurrentPin->LinkedTo.Num() > 0);
+
+		bool bIsConnectedInput =
+			(CurrentPin->Direction == EGPD_Input) && (CurrentPin->LinkedTo.Num() > 0);
+
+		if (bIsConnectedOutput)
+		{
+			for (UEdGraphPin* DestinationPin : CurrentPin->LinkedTo)
+			{
+				check(DestinationPin);
+
+				if (Context.ImplicitCastMap.Contains(DestinationPin))
+				{
+					continue;
+				}
+
+				TOptional<StatementNamePair> ConversionType =
+					GetFloatingPointConversionType(*CurrentPin, *DestinationPin);
+
+				if (ConversionType)
+				{
+					AddCastMapping(DestinationPin, ConversionType->Get<0>(), ConversionType->Get<1>());
+				}
+			}
+		}
+		else if (bIsConnectedInput)
+		{
+			if (Context.ImplicitCastMap.Contains(CurrentPin))
+			{
+				continue;
+			}
+
+			if (CurrentPin->LinkedTo.Num() > 0)
+			{
+				const UEdGraphPin* SourcePin = CurrentPin->LinkedTo[0];
+				check(SourcePin);
+
+				TOptional<StatementNamePair> ConversionType =
+					GetFloatingPointConversionType(*SourcePin, *CurrentPin);
+
+				if (ConversionType)
+				{
+					AddCastMapping(CurrentPin, ConversionType->Get<0>(), ConversionType->Get<1>());
+				}
+			}
+		}
+	}
+}
+
+TOptional<TPair<FBPTerminal*, EKismetCompiledStatementType>>
+InsertImplicitCastStatement(FKismetFunctionContext& Context, UEdGraphPin* DestinationPin, FBPTerminal* RHSTerm)
+{
+	using namespace UE::KismetCompiler::CastingUtils::Private;
+
+	check(DestinationPin);
+	check(RHSTerm);
+
+	TOptional<TPair<FBPTerminal*, EKismetCompiledStatementType>> Result;
+
+	const FImplicitCastParams* CastParams =
+		Context.ImplicitCastMap.Find(DestinationPin);
+
+	if (CastParams != nullptr)
+	{
+		check(CastParams->TargetTerminal);
+		check(CastParams->TargetNode);
+
+		FBlueprintCompiledStatement& CastStatement = Context.AppendStatementForNode(CastParams->TargetNode);
+		CastStatement.LHS = CastParams->TargetTerminal;
+		CastStatement.Type = CastParams->CastType;
+		CastStatement.RHS.Add(RHSTerm);
+
+		// Removal of the pin entry indicates to the compiler that the implicit cast has been processed.
+		Context.ImplicitCastMap.Remove(DestinationPin);
+
+		Result = TPair<FBPTerminal*, EKismetCompiledStatementType>{CastParams->TargetTerminal, CastParams->CastType};
+	}
+
+	return Result;
+}
+
+bool RemoveRegisteredImplicitCast(FKismetFunctionContext& Context, const UEdGraphPin* DestinationPin)
+{
+	check(DestinationPin);
+
+	int32 RemovedCount = Context.ImplicitCastMap.Remove(DestinationPin);
+
+	return (RemovedCount > 0);
+}
+
+TOptional<StatementNamePair> GetFloatingPointConversionType(const UEdGraphPin& SourcePin, const UEdGraphPin& DestinationPin)
+{
+	using namespace UE::KismetCompiler::CastingUtils::Private;
+
+	using CastPair = TPair<EPinType, EPinType>;
+
+	static TMap<CastPair, StatementNamePair> ImplicitCastTable =
+	{
+		{ CastPair{EPinType::Float,						EPinType::Double},					StatementNamePair{KCST_FloatToDoubleCast,								TEXT("WideningCast")}		},
+		{ CastPair{EPinType::FloatArray,				EPinType::DoubleArray},				StatementNamePair{KCST_FloatToDoubleArrayCast,							TEXT("WideningArrayCast")}	},
+		{ CastPair{EPinType::FloatSet,					EPinType::DoubleSet},				StatementNamePair{KCST_FloatToDoubleSetCast,							TEXT("WideningSetCast")}	},
+
+		{ CastPair{EPinType::Double,					EPinType::Float},					StatementNamePair{KCST_DoubleToFloatCast,								TEXT("NarrowingCast")}		},
+		{ CastPair{EPinType::DoubleArray,				EPinType::FloatArray},				StatementNamePair{KCST_DoubleToFloatArrayCast,							TEXT("NarrowingArrayCast")}	},
+		{ CastPair{EPinType::DoubleSet,					EPinType::FloatSet},				StatementNamePair{KCST_DoubleToFloatSetCast,							TEXT("NarrowingSetCast")}	},
+
+		{ CastPair{EPinType::Vector3f,					EPinType::Vector},					StatementNamePair{KCST_Vector3fToVectorCast,							TEXT("WideningCast")}		},
+		{ CastPair{EPinType::Vector3fArray,				EPinType::VectorArray},				StatementNamePair{KCST_Vector3fToVectorArrayCast,						TEXT("WideningArrayCast")}	},
+		{ CastPair{EPinType::Vector3fSet,				EPinType::VectorSet},				StatementNamePair{KCST_Vector3fToVectorSetCast,							TEXT("WideningSetCast")}	},
+
+		{ CastPair{EPinType::Vector,					EPinType::Vector3f},				StatementNamePair{KCST_VectorToVector3fCast,							TEXT("NarrowingCast")}		},
+		{ CastPair{EPinType::VectorArray,				EPinType::Vector3fArray},			StatementNamePair{KCST_VectorToVector3fArrayCast,						TEXT("NarrowingArrayCast")}	},
+		{ CastPair{EPinType::VectorSet,					EPinType::Vector3fSet},				StatementNamePair{KCST_VectorToVector3fSetCast,							TEXT("NarrowingSetCast")}	},
+
+		{ CastPair{EPinType::FloatKeyOtherValueMap,		EPinType::DoubleKeyOtherValueMap},	StatementNamePair{KCST_FloatToDoubleKeys_MapCast,						TEXT("MapCast")}	},
+		{ CastPair{EPinType::DoubleKeyOtherValueMap,	EPinType::FloatKeyOtherValueMap},	StatementNamePair{KCST_DoubleToFloatKeys_MapCast,						TEXT("MapCast")}	},
+		{ CastPair{EPinType::OtherKeyFloatValueMap,		EPinType::OtherKeyDoubleValueMap},	StatementNamePair{KCST_FloatToDoubleValues_MapCast,						TEXT("MapCast")}	},
+		{ CastPair{EPinType::OtherKeyDoubleValueMap,	EPinType::OtherKeyFloatValueMap},	StatementNamePair{KCST_DoubleToFloatValues_MapCast,						TEXT("MapCast")}	},
+
+		{ CastPair{EPinType::FloatKeyFloatValueMap,		EPinType::DoubleKeyDoubleValueMap},	StatementNamePair{KCST_FloatToDoubleKeys_FloatToDoubleValues_MapCast,	TEXT("MapCast")}	},
+		{ CastPair{EPinType::FloatKeyFloatValueMap,		EPinType::FloatKeyDoubleValueMap},	StatementNamePair{KCST_FloatToDoubleValues_MapCast,						TEXT("MapCast")}	},
+		{ CastPair{EPinType::FloatKeyFloatValueMap,		EPinType::DoubleKeyFloatValueMap},	StatementNamePair{KCST_FloatToDoubleKeys_MapCast,						TEXT("MapCast")}	},
+
+		{ CastPair{EPinType::DoubleKeyFloatValueMap,	EPinType::DoubleKeyDoubleValueMap},	StatementNamePair{KCST_FloatToDoubleValues_MapCast,						TEXT("MapCast")}	},
+		{ CastPair{EPinType::DoubleKeyFloatValueMap,	EPinType::FloatKeyDoubleValueMap},	StatementNamePair{KCST_DoubleToFloatKeys_FloatToDoubleValues_MapCast,	TEXT("MapCast")}	},
+		{ CastPair{EPinType::DoubleKeyFloatValueMap,	EPinType::FloatKeyFloatValueMap},	StatementNamePair{KCST_DoubleToFloatKeys_MapCast,						TEXT("MapCast")}	},
+
+		{ CastPair{EPinType::DoubleKeyDoubleValueMap,	EPinType::DoubleKeyFloatValueMap},	StatementNamePair{KCST_DoubleToFloatValues_MapCast,						TEXT("MapCast")}	},
+		{ CastPair{EPinType::DoubleKeyDoubleValueMap,	EPinType::FloatKeyDoubleValueMap},	StatementNamePair{KCST_DoubleToFloatKeys_MapCast,						TEXT("MapCast")}	},
+		{ CastPair{EPinType::DoubleKeyDoubleValueMap,	EPinType::FloatKeyFloatValueMap},	StatementNamePair{KCST_DoubleToFloatKeys_DoubleToFloatValues_MapCast,	TEXT("MapCast")}	},
+
+		{ CastPair{EPinType::FloatKeyDoubleValueMap,	EPinType::DoubleKeyFloatValueMap},	StatementNamePair{KCST_FloatToDoubleKeys_DoubleToFloatValues_MapCast,	TEXT("MapCast")}	},
+		{ CastPair{EPinType::FloatKeyDoubleValueMap,	EPinType::DoubleKeyDoubleValueMap},	StatementNamePair{KCST_FloatToDoubleKeys_MapCast,						TEXT("MapCast")}	},
+		{ CastPair{EPinType::FloatKeyDoubleValueMap,	EPinType::FloatKeyFloatValueMap},	StatementNamePair{KCST_DoubleToFloatValues_MapCast,						TEXT("MapCast")}	},
+	};
+
+	CastPair LookupPair{ GetPinType(SourcePin), GetPinType(DestinationPin) };
+
+	const StatementNamePair* TableLookupResult = ImplicitCastTable.Find(LookupPair);
+
+	TOptional<StatementNamePair> FinalResult;
+
+	if (TableLookupResult)
+	{
+		FinalResult = *TableLookupResult;
+	}
+
+	return FinalResult;
+}
+
+} // namespace UE::KismetCompiler::CastingUtils
+
+#endif
+
diff --git a/Engine/Source/Editor/KismetCompiler/Private/KismetCompiler.cpp b/Engine/Source/Editor/KismetCompiler/Private/KismetCompiler.cpp
--- a/Engine/Source/Editor/KismetCompiler/Private/KismetCompiler.cpp
+++ b/Engine/Source/Editor/KismetCompiler/Private/KismetCompiler.cpp
@@ -23,6 +23,7 @@
 #include "Engine/UserDefinedStruct.h"
 #include "Blueprint/BlueprintExtension.h"
 #include "EdGraphUtilities.h"
+#include "K2Node_AddDelegate.h"
 #include "K2Node_CallFunction.h"
 #include "K2Node_Composite.h"
 #include "K2Node_CreateDelegate.h"
@@ -39,6 +40,7 @@
 #include "K2Node_VariableGet.h"
 #include "K2Node_VariableSet.h"
 #include "K2Node_EditablePinBase.h" // for FUserPinInfo
+#include "KismetCastingUtils.h"
 #include "KismetCompilerBackend.h"
 #include "Kismet2/KismetReinstanceUtilities.h"
 #include "Engine/SCS_Node.h"
@@ -88,7 +90,7 @@ DECLARE_CYCLE_STAT(TEXT("Calculate checksum of signature"), EKismetCompilerStats
 DECLARE_CYCLE_STAT(TEXT("Pruning"), EKismetCompilerStats_PruneIsolatedNodes, STATGROUP_KismetCompiler);
 DECLARE_CYCLE_STAT(TEXT("Merge Ubergraph Pages In"), EKismetCompilerStats_MergeUbergraphPagesIn, STATGROUP_KismetCompiler);
 
-namespace
+namespace UE::KismetCompiler::Private
 {
 	// The function collects all nodes, that can represents entry points of the execution. Any node connected to "root" node (by execution link) won't be consider isolated.
 	static void GatherRootSet(const UEdGraph* Graph, TArray<UEdGraphNode*>& RootSet, bool bIncludeNodesThatCouldBeExpandedToRootSet)
@@ -122,6 +124,319 @@ namespace
 			}
 		}
 	}
+
+#if ENABLE_BLUEPRINT_REAL_NUMBERS
+	// When we change pins back to real/float types, we also expect that the series of pin types matches that of the UFunction.
+	// If there's ever a discrepancy, then there's not much we can do other than log a warning.
+	// It typically means that the BP is in a bad state, which prevents us from deducing the corresponding pin.
+	static void RestoreFloatPinsToNode(UK2Node_EditablePinBase* Node, const UFunction* FunctionSignature, bool bOutputParamsOnly = false)
+	{
+		check(Node);
+		check(FunctionSignature);
+
+		int UserDefinedPinCursor = 0;
+		int PinCursor = 0;
+
+		if (Node->UserDefinedPins.Num() > 0)
+		{
+			bool bMatchingPinFound = false;
+			
+			check(Node->UserDefinedPins[0]);
+
+			FName UserPinName = Node->UserDefinedPins[0]->PinName;
+
+			for (int i = 0; i < Node->Pins.Num(); ++i)
+			{
+				check(Node->Pins[i]);
+
+				if (Node->Pins[i]->PinName == UserPinName)
+				{
+					PinCursor = i;
+					bMatchingPinFound = true;
+					break;
+				}
+			}
+
+			if (!bMatchingPinFound)
+			{
+				UE_LOG(LogK2Compiler, Warning, TEXT("User pin '%s' ('%s') was not found in the pins list!"),
+					*UserPinName.ToString(),
+					*Node->GetFullName());
+			}
+		}
+
+		auto IsFloatProperty = [Node](const FProperty* Property)
+		{
+			check(Property);
+
+			if (const FFloatProperty* FloatProperty = CastField<FFloatProperty>(Property))
+			{
+				return true;
+			}
+			else if (const FArrayProperty* ArrayProperty = CastField<FArrayProperty>(Property))
+			{
+				check(ArrayProperty->Inner);
+				return ArrayProperty->Inner->IsA<FFloatProperty>();
+			}
+			else if (const FSetProperty* SetProperty = CastField<FSetProperty>(Property))
+			{
+				check(SetProperty->ElementProp);
+				return SetProperty->ElementProp->IsA<FFloatProperty>();
+			}
+			else if (const FMapProperty* MapProperty = CastField<FMapProperty>(Property))
+			{
+				check(MapProperty->KeyProp);
+				check(MapProperty->ValueProp);
+
+				if (MapProperty->KeyProp->IsA<FFloatProperty>() || MapProperty->ValueProp->IsA<FFloatProperty>())
+				{
+					UE_LOG(LogK2Compiler, Warning, TEXT("A map with float entries was found in '%s'. Delegate signature may be inaccurate."),
+						*Node->GetFullName());
+				}
+
+				return false;
+			}
+
+			return false;
+		};
+
+		auto IsValidParameter = [bOutputParamsOnly](FProperty* Property)
+		{
+			check(Property);
+
+			if (bOutputParamsOnly)
+			{
+				return Property->HasAnyPropertyFlags(CPF_OutParm);
+			}
+			else
+			{
+				// HACK: We still have to reckon with arrays that are implicitly copy-by-ref.
+				// Ideally, we should be excluding all output params.
+				bool bIsValidInput =
+					(Property->IsA<FArrayProperty>() && Property->HasAnyPropertyFlags(CPF_Parm)) ||
+					(Property->HasAnyPropertyFlags(CPF_Parm) && !Property->HasAnyPropertyFlags(CPF_OutParm));
+
+				return bIsValidInput;
+			}
+		};
+
+		for (TFieldIterator<FProperty> PropIt(FunctionSignature); PropIt; ++PropIt)
+		{
+			FProperty* CurrentProperty = *PropIt;
+			check(CurrentProperty);
+
+			if (IsValidParameter(CurrentProperty))
+			{
+				if (UserDefinedPinCursor < Node->UserDefinedPins.Num())
+				{
+					if (IsFloatProperty(CurrentProperty))
+					{
+						TSharedPtr<FUserPinInfo>& UserPin = Node->UserDefinedPins[UserDefinedPinCursor];
+						check(UserPin);
+
+						if (UserPin->PinType.PinCategory == UEdGraphSchema_K2::PC_Real)
+						{
+							UserPin->PinType.PinSubCategory = UEdGraphSchema_K2::PC_Float;
+						}
+						else
+						{
+							UE_LOG(LogK2Compiler, Warning, TEXT("Expected a 'real' type for user pin '%s' ('%s'), but found type '%s' instead!"),
+								*UserPin->PinName.ToString(),
+								*Node->GetFullName(),
+								*UserPin->PinType.PinCategory.ToString());
+						}
+
+						if (PinCursor < Node->Pins.Num())
+						{
+							UEdGraphPin* Pin = Node->Pins[PinCursor];
+							check(Pin);
+
+							if (Pin->PinType.PinCategory == UEdGraphSchema_K2::PC_Real)
+							{
+								Pin->PinType.PinSubCategory = UEdGraphSchema_K2::PC_Float;
+							}
+							else
+							{
+								UE_LOG(LogK2Compiler, Warning, TEXT("Expected a 'real' type for pin '%s' ('%s'), but found type '%s' instead!"),
+									*Pin->PinName.ToString(),
+									*Node->GetFullName(),
+									*Pin->PinType.PinCategory.ToString());
+							}
+						}
+						else
+						{
+							UE_LOG(LogK2Compiler, Warning, TEXT("PinCursor ('%s') has an invalid index: %d!"),
+								*Node->GetFullName(),
+								PinCursor);
+						}
+					}
+
+					++UserDefinedPinCursor;
+					++PinCursor;
+				}
+				else
+				{
+					// It's possible that we'll use up pins before reaching the end of the function signature's properties.
+					// This can happen when a function has reference params. The corresponding pins are split between two nodes (ie: entry and result nodes).
+					break;
+				}
+			}
+		}
+	}
+
+	static UFunction* GetLinkedDelegateSignature(UK2Node* NodeWithDelegate)
+	{
+		check(NodeWithDelegate);
+
+		const UEdGraphPin* DelegateOutPin = NodeWithDelegate->FindPin(UK2Node_Event::DelegateOutputName);
+		if (DelegateOutPin && (DelegateOutPin->LinkedTo.Num() > 0))
+		{
+			// Since the BP hasn't been compiled yet, we'll run into knot nodes that need to be skipped.
+			check(DelegateOutPin->LinkedTo[0]);
+			UEdGraphNode* LinkedNode = DelegateOutPin->LinkedTo[0]->GetOwningNode();
+			while (UK2Node_Knot* KnotNode = Cast<UK2Node_Knot>(LinkedNode))
+			{
+				check(KnotNode->GetOutputPin());
+				UEdGraphPin* OutputPin = KnotNode->GetOutputPin();
+				check(OutputPin);
+				check(OutputPin->LinkedTo.Num() > 0);
+				
+				LinkedNode = OutputPin->LinkedTo[0]->GetOwningNode();
+			}
+
+			// If the delegate pin has been linked, then we should run into either a
+			// AddDelegate or RemoveDelegate node.
+			// In either case, it should have the signature that we need.
+			UK2Node_BaseMCDelegate* DelegateNode = Cast<UK2Node_BaseMCDelegate>(LinkedNode);
+			if (DelegateNode)
+			{
+				return DelegateNode->GetDelegateSignature();
+			}
+		}
+
+		return nullptr;
+	}
+
+	// By default, float pins are changed to real/double pins during serialization.
+	// However, Blueprint functions that are bound to native delegates need to use real/float pins if
+	// that's what the underlying delegate signature requires.
+	// Failure to do so will lead to incorrect UFunction signatures in the skeleton class,
+	// which can cause compilation errors.
+	static void MatchNodesToDelegateSignatures(UBlueprint* BP)
+	{
+		check(BP);
+
+		TArray<UK2Node_CreateDelegate*> CreateDelegateNodes;
+		FBlueprintEditorUtils::GetAllNodesOfClass(BP, CreateDelegateNodes);
+
+		TMap<FName, UK2Node_CreateDelegate*> CreateDelegateNodeMap;
+		for (UK2Node_CreateDelegate* CreateDelegateNode : CreateDelegateNodes)
+		{
+			check(CreateDelegateNode);
+			if (CreateDelegateNode->SelectedFunctionName != NAME_None)
+			{
+				CreateDelegateNodeMap.Add(CreateDelegateNode->SelectedFunctionName, CreateDelegateNode);
+			}
+		}
+
+		// Handle function terminator nodes
+
+		TArray<UK2Node_FunctionTerminator*> TerminatorNodes;
+		FBlueprintEditorUtils::GetAllNodesOfClass(BP, TerminatorNodes);
+		for (UK2Node_FunctionTerminator* FunctionTerminatorNode : TerminatorNodes)
+		{
+			check(FunctionTerminatorNode);
+
+			// If the function is an override, then we can't change the signature.
+			UFunction* Function = FunctionTerminatorNode->FunctionReference.ResolveMember<UFunction>(BP->ParentClass);
+			if (!Function)
+			{
+				UK2Node_CreateDelegate** DelegateNodePtr = CreateDelegateNodeMap.Find(FunctionTerminatorNode->FunctionReference.GetMemberName());
+				if (DelegateNodePtr)
+				{
+					UK2Node_CreateDelegate* DelegateNode = *DelegateNodePtr;
+					check(DelegateNode);
+
+					UFunction* DelegateSignature = DelegateNode->GetDelegateSignature();
+
+					if (DelegateSignature == nullptr)
+					{
+						// In some edge cases, the delegate pin for UK2Node_CreateDelegate might be missing PinSubCategoryMemberReference information.
+						// Our fallback case is to find the signature from the linked K2Node_AddDelegate. 
+						DelegateSignature = GetLinkedDelegateSignature(DelegateNode);
+					}
+
+					if (DelegateSignature)
+					{
+						if (DelegateSignature->GetNativeFunc())
+						{
+							bool bUseOutputParams =
+								FunctionTerminatorNode->IsA<UK2Node_FunctionResult>();
+
+							RestoreFloatPinsToNode(FunctionTerminatorNode, DelegateSignature, bUseOutputParams);
+						}
+					}
+					else
+					{
+						UE_LOG(LogK2Compiler, Warning, TEXT("Unable to determine delegate signature for node '%s'!"),
+							*DelegateNode->GetFullName());
+					}
+				}
+			}
+		}
+
+		// Handle custom event nodes
+
+		TArray<UK2Node_CustomEvent*> CustomEventNodes;
+		FBlueprintEditorUtils::GetAllNodesOfClass(BP, CustomEventNodes);
+		for (UK2Node_CustomEvent* CustomEventNode : CustomEventNodes)
+		{
+			check(CustomEventNode);
+
+			// We'll give precedence to the signature of the linked delegate.
+			// Failing that, we'll then see if the custom event was possibly created via a CreateDelegate node.
+			UFunction* DelegateSignature = GetLinkedDelegateSignature(CustomEventNode);
+
+			if (DelegateSignature && DelegateSignature->GetNativeFunc())
+			{
+				RestoreFloatPinsToNode(CustomEventNode, DelegateSignature);
+			}
+			else
+			{
+				FName NodeName = CustomEventNode->CustomFunctionName;
+
+				UK2Node_CreateDelegate** DelegateNodePtr = CreateDelegateNodeMap.Find(NodeName);
+				if (DelegateNodePtr)
+				{
+					UK2Node_CreateDelegate* DelegateNode = *DelegateNodePtr;
+					check(DelegateNode);
+
+					DelegateSignature = DelegateNode->GetDelegateSignature();
+
+					if (DelegateSignature == nullptr)
+					{
+						// In some edge cases, the delegate pin for UK2Node_CreateDelegate might be missing PinSubCategoryMemberReference information.
+						// Our fallback case is to find the signature from the linked K2Node_AddDelegate. 
+						DelegateSignature = GetLinkedDelegateSignature(DelegateNode);
+					}
+
+					if (DelegateSignature)
+					{
+						if (DelegateSignature->GetNativeFunc())
+						{
+							RestoreFloatPinsToNode(CustomEventNode, DelegateSignature);
+						}
+					}
+					else
+					{
+						UE_LOG(LogK2Compiler, Warning, TEXT("Unable to determine delegate signature for node '%s'!"),
+							*DelegateNode->GetFullName());
+					}
+				}
+			}
+		}
+	}
+#endif
 }
 
 //////////////////////////////////////////////////////////////////////////
@@ -857,7 +1172,11 @@ void FKismetCompilerContext::CreateClassVariablesFromBlueprint()
 		FEdGraphPinType DirectionPinType(UEdGraphSchema_K2::PC_Byte, NAME_None, FTimeline::GetTimelineDirectionEnum(), EPinContainerType::None, false, FEdGraphTerminalType());
 		CreateVariable(Timeline->GetDirectionPropertyName(), DirectionPinType);
 
+#if ENABLE_BLUEPRINT_REAL_NUMBERS
+		FEdGraphPinType FloatPinType(UEdGraphSchema_K2::PC_Real, UEdGraphSchema_K2::PC_Float, nullptr, EPinContainerType::None, false, FEdGraphTerminalType());
+#else
 		FEdGraphPinType FloatPinType(UEdGraphSchema_K2::PC_Float, NAME_None, nullptr, EPinContainerType::None, false, FEdGraphTerminalType());
+#endif
 		for (const FTTFloatTrack& FloatTrack : Timeline->FloatTracks)
 		{
 			CreateVariable(FloatTrack.GetPropertyName(), FloatPinType);
@@ -1498,7 +1817,7 @@ void FKismetCompilerContext::PruneIsolatedNodes(UEdGraph* InGraph, bool bInInclu
 {
 	TArray<UEdGraphNode*> RootSet;
 	// Find any all entry points caused by special nodes
-	GatherRootSet(InGraph, RootSet, bInIncludeNodesThatCouldBeExpandedToRootSet);
+	UE::KismetCompiler::Private::GatherRootSet(InGraph, RootSet, bInIncludeNodesThatCouldBeExpandedToRootSet);
 
 	// Find the connected subgraph starting at the root node and prune out unused nodes
 	PruneIsolatedNodes(RootSet, InGraph->Nodes);
@@ -2057,6 +2376,17 @@ void FKismetCompilerContext::PrecompileFunction(FKismetFunctionContext& Context,
 	}
 }
 
+void FKismetCompilerContext::PreCompileUpdateBlueprintOnLoad(UBlueprint* BP)
+{
+#if ENABLE_BLUEPRINT_REAL_NUMBERS
+	check(BP);
+	if (BP->GetLinkerCustomVersion(FUE5ReleaseStreamObjectVersion::GUID) < FUE5ReleaseStreamObjectVersion::BlueprintPinsUseRealNumbers)
+	{
+		UE::KismetCompiler::Private::MatchNodesToDelegateSignatures(BP);
+	}
+#endif
+}
+
 /** Inserts a new item into an array in a sorted position; using an externally stored sort index map */
 template<typename DataType, typename SortKeyType>
 void OrderedInsertIntoArray(TArray<DataType>& Array, const TMap<DataType, SortKeyType>& SortKeyMap, const DataType& NewItem)
@@ -2149,6 +2479,18 @@ void FKismetCompilerContext::CompileFunction(FKismetFunctionContext& Context)
 		}
 	}
 	
+	if (Context.ImplicitCastMap.Num() > 0)
+	{
+		UE_LOG(LogK2Compiler, Warning, TEXT("Unhandled implicit casts found during compilation of function '%s'!"),
+			*Context.Function->GetFullName());
+		for (const auto& It : Context.ImplicitCastMap)
+		{
+			UE_LOG(LogK2Compiler, Warning, TEXT("\tPin '%s' was not handled by node '%s'!"),
+				*It.Key->PinName.ToString(),
+				*It.Key->GetOwningNode()->GetName());
+		}
+	}
+
 	// The LinearExecutionList should be immutable at this point
 	check(Context.LinearExecutionList.Num() == NumNodesAtStart);
 
@@ -3282,6 +3624,12 @@ void FKismetCompilerContext::CreateLocalsAndRegisterNets(FKismetFunctionContext&
 		}
 	}
 
+#if ENABLE_BLUEPRINT_REAL_NUMBERS
+	using namespace UE::KismetCompiler;
+
+	CastingUtils::RegisterImplicitCasts(Context);
+#endif
+
 	// Create net variable declarations
 	CreateLocalVariablesForFunction(Context, FunctionPropertyStorageLocation);
 }
diff --git a/Engine/Source/Editor/KismetCompiler/Private/KismetCompilerMisc.cpp b/Engine/Source/Editor/KismetCompiler/Private/KismetCompilerMisc.cpp
--- a/Engine/Source/Editor/KismetCompiler/Private/KismetCompilerMisc.cpp
+++ b/Engine/Source/Editor/KismetCompiler/Private/KismetCompilerMisc.cpp
@@ -32,6 +32,7 @@
 #include "K2Node_FunctionResult.h"
 #include "K2Node_Timeline.h"
 #include "K2Node_Variable.h"
+#include "KismetCastingUtils.h"
 #include "KismetCompiledFunctionContext.h"
 #include "KismetCompiler.h"
 
@@ -120,6 +121,26 @@ static bool DoesTypeNotMatchProperty(UEdGraphPin* SourcePin, const FEdGraphPinTy
 			}
 		}
 	}
+#if ENABLE_BLUEPRINT_REAL_NUMBERS
+	else if (PinCategory == UEdGraphSchema_K2::PC_Real)
+	{
+		if (PinSubCategory == UEdGraphSchema_K2::PC_Float)
+		{
+			FFloatProperty* SpecificProperty = CastField<FFloatProperty>(TestProperty);
+			bTypeMismatch = (SpecificProperty == nullptr);
+		}
+		else if (PinSubCategory == UEdGraphSchema_K2::PC_Double)
+		{
+			FDoubleProperty* SpecificProperty = CastField<FDoubleProperty>(TestProperty);
+			bTypeMismatch = (SpecificProperty == nullptr);
+		}
+		else
+		{
+			checkf(false, TEXT("Erroneous pin subcategory for PC_Real: %s"), *PinSubCategory.ToString());
+			bTypeMismatch = true;
+		}
+	}
+#else
 	else if (PinCategory == UEdGraphSchema_K2::PC_Float)
 	{
 		FFloatProperty* SpecificProperty = CastField<FFloatProperty>(TestProperty);
@@ -130,6 +151,7 @@ static bool DoesTypeNotMatchProperty(UEdGraphPin* SourcePin, const FEdGraphPinTy
 		FDoubleProperty* SpecificProperty = CastField<FDoubleProperty>(TestProperty);
 		bTypeMismatch = (SpecificProperty == nullptr);
 	}
+#endif
 	else if (PinCategory == UEdGraphSchema_K2::PC_Int)
 	{
 		FIntProperty* SpecificProperty = CastField<FIntProperty>(TestProperty);
@@ -851,7 +873,7 @@ UEdGraphPin* FKismetCompilerUtilities::GenerateAssignmentNodes(class FKismetComp
 	return LastThen;
 }
 
-void FKismetCompilerUtilities::CreateObjectAssignmentStatement(FKismetFunctionContext& Context, UEdGraphNode* Node, FBPTerminal* SrcTerm, FBPTerminal* DstTerm)
+void FKismetCompilerUtilities::CreateObjectAssignmentStatement(FKismetFunctionContext& Context, UEdGraphNode* Node, FBPTerminal* SrcTerm, FBPTerminal* DstTerm, UEdGraphPin* DstPin)
 {
 	UClass* InputObjClass = Cast<UClass>(SrcTerm->Type.PinSubCategoryObject.Get());
 	UClass* OutputObjClass = Cast<UClass>(DstTerm->Type.PinSubCategoryObject.Get());
@@ -878,10 +900,36 @@ void FKismetCompilerUtilities::CreateObjectAssignmentStatement(FKismetFunctionCo
 	}
 	else
 	{
+		FBPTerminal* RHSTerm = SrcTerm;
+
+#if ENABLE_BLUEPRINT_REAL_NUMBERS
+		using namespace UE::KismetCompiler;
+
+		FBPTerminal* ImplicitCastTerm = nullptr;
+
+		// Some pins can share a single terminal (eg: those in UK2Node_FunctionResult)
+		// In those cases, it's preferable to use a specific pin instead of relying on what DstTerm points to.
+		UEdGraphPin* DstPinSearchKey = DstPin ? DstPin : DstTerm->SourcePin;
+
+		// Some terms don't necessarily have a valid SourcePin (eg: FKCHandler_FunctionEntry)
+		if (DstPinSearchKey)
+		{
+			TOptional<TPair<FBPTerminal*, EKismetCompiledStatementType>> ImplicitCastEntry =
+				CastingUtils::InsertImplicitCastStatement(Context, DstPinSearchKey, RHSTerm);
+
+			ImplicitCastTerm = ImplicitCastEntry ? ImplicitCastEntry->Get<0>() : nullptr;
+		}
+
+		if (ImplicitCastTerm != nullptr)
+		{
+			RHSTerm = ImplicitCastTerm;
+		}
+#endif
+
 		FBlueprintCompiledStatement& Statement = Context.AppendStatementForNode(Node);
 		Statement.Type = KCST_Assignment;
 		Statement.LHS = DstTerm;
-		Statement.RHS.Add(SrcTerm);
+		Statement.RHS.Add(RHSTerm);
 	}
 }
 
@@ -1038,6 +1086,25 @@ FProperty* FKismetCompilerUtilities::CreatePrimitiveProperty(FFieldVariant Prope
 		NewProperty = new FInt64Property(PropertyScope, ValidatedPropertyName, ObjectFlags);
 		NewProperty->SetPropertyFlags(CPF_HasGetValueTypeHash);
 	}
+#if ENABLE_BLUEPRINT_REAL_NUMBERS
+	else if (PinCategory == UEdGraphSchema_K2::PC_Real)
+	{
+		if (PinSubCategory == UEdGraphSchema_K2::PC_Float)
+		{
+			NewProperty = new FFloatProperty(PropertyScope, ValidatedPropertyName, ObjectFlags);
+			NewProperty->SetPropertyFlags(CPF_HasGetValueTypeHash);
+		}
+		else if (PinSubCategory == UEdGraphSchema_K2::PC_Double)
+		{
+			NewProperty = new FDoubleProperty(PropertyScope, ValidatedPropertyName, ObjectFlags);
+			NewProperty->SetPropertyFlags(CPF_HasGetValueTypeHash);
+		}
+		else
+		{
+			checkf(false, TEXT("Erroneous pin subcategory for PC_Real: %s"), *PinSubCategory.ToString());
+		}
+	}
+#else
 	else if (PinCategory == UEdGraphSchema_K2::PC_Float)
 	{
 		NewProperty = new FFloatProperty(PropertyScope, ValidatedPropertyName, ObjectFlags);
@@ -1048,6 +1115,7 @@ FProperty* FKismetCompilerUtilities::CreatePrimitiveProperty(FFieldVariant Prope
 		NewProperty = new FDoubleProperty(PropertyScope, ValidatedPropertyName, ObjectFlags);
 		NewProperty->SetPropertyFlags(CPF_HasGetValueTypeHash);
 	}
+#endif
 	else if (PinCategory == UEdGraphSchema_K2::PC_Boolean)
 	{
 		FBoolProperty* BoolProperty = new FBoolProperty(PropertyScope, ValidatedPropertyName, ObjectFlags);
@@ -1954,6 +2022,26 @@ bool FKismetCompilerUtilities::CheckFunctionCompiledStatementsThreadSafety(const
 			case KCST_DebugSite:
 			case KCST_CastObjToInterface:
 			case KCST_DynamicCast:
+			case KCST_DoubleToFloatCast:
+			case KCST_DoubleToFloatArrayCast:
+			case KCST_DoubleToFloatSetCast:
+			case KCST_FloatToDoubleCast:
+			case KCST_FloatToDoubleArrayCast:
+			case KCST_FloatToDoubleSetCast:
+			case KCST_VectorToVector3fCast:
+			case KCST_VectorToVector3fArrayCast:
+			case KCST_VectorToVector3fSetCast:
+			case KCST_Vector3fToVectorCast:
+			case KCST_Vector3fToVectorArrayCast:
+			case KCST_Vector3fToVectorSetCast:
+			case KCST_FloatToDoubleKeys_MapCast:
+			case KCST_DoubleToFloatKeys_MapCast:
+			case KCST_FloatToDoubleValues_MapCast:
+			case KCST_DoubleToFloatValues_MapCast:
+			case KCST_FloatToDoubleKeys_FloatToDoubleValues_MapCast:
+			case KCST_DoubleToFloatKeys_FloatToDoubleValues_MapCast:
+			case KCST_DoubleToFloatKeys_DoubleToFloatValues_MapCast:
+			case KCST_FloatToDoubleKeys_DoubleToFloatValues_MapCast:
 			case KCST_ObjectToBool:
 				break;
 			case KCST_AddMulticastDelegate:
@@ -2031,7 +2119,6 @@ void FNodeHandlingFunctor::ResolveAndRegisterScopedTerm(FKismetFunctionContext&
 	FProperty* BoundProperty = FKismetCompilerUtilities::FindPropertyInScope(SearchScope, Net, CompilerContext.MessageLog, CompilerContext.GetSchema(), Context.NewClass, bIsSparseProperty);
 	if (BoundProperty != NULL)
 	{
-		UBlueprintEditorSettings* Settings = GetMutableDefault<UBlueprintEditorSettings>();
 		// Create the term in the list
 		FBPTerminal* Term = new FBPTerminal();
 		NetArray.Add(Term);
diff --git a/Engine/Source/Editor/KismetCompiler/Private/KismetCompilerVMBackend.cpp b/Engine/Source/Editor/KismetCompiler/Private/KismetCompilerVMBackend.cpp
--- a/Engine/Source/Editor/KismetCompiler/Private/KismetCompilerVMBackend.cpp
+++ b/Engine/Source/Editor/KismetCompiler/Private/KismetCompilerVMBackend.cpp
@@ -269,6 +269,7 @@ private:
 
 	// Pointers to commonly used structures (found in constructor)
 	UScriptStruct* VectorStruct;
+	UScriptStruct* Vector3fStruct;
 	UScriptStruct* RotatorStruct;
 	UScriptStruct* TransformStruct;
 	UScriptStruct* LatentInfoStruct;
@@ -404,6 +405,7 @@ public:
 		, PureNodeEntryStart(0)
 	{
 		VectorStruct = TBaseStructure<FVector>::Get();
+		Vector3fStruct = TVariantStructure<FVector3f>::Get();
 		RotatorStruct = TBaseStructure<FRotator>::Get();
 		TransformStruct = TBaseStructure<FTransform>::Get();
 		LatentInfoStruct = FLatentActionInfo::StaticStruct();
@@ -476,6 +478,25 @@ public:
 			return Type && (Type->PinCategory == UEdGraphSchema_K2::PC_Text);
 		}
 
+#if ENABLE_BLUEPRINT_REAL_NUMBERS
+		static bool IsFloat(const FEdGraphPinType* Type, const FProperty* Property)
+		{
+			if (Property)
+			{
+				return Property->IsA<FFloatProperty>();
+			}
+			return Type && (Type->PinCategory == UEdGraphSchema_K2::PC_Real) && (Type->PinSubCategory == UEdGraphSchema_K2::PC_Float);
+		}
+
+		static bool IsDouble(const FEdGraphPinType* Type, const FProperty* Property)
+		{
+			if (Property)
+			{
+				return Property->IsA<FDoubleProperty>();
+			}
+			return Type && (Type->PinCategory == UEdGraphSchema_K2::PC_Real) && (Type->PinSubCategory == UEdGraphSchema_K2::PC_Double);
+		}
+#else
 		static bool IsFloat(const FEdGraphPinType* Type, const FProperty* Property)
 		{
 			if (Property)
@@ -493,6 +514,7 @@ public:
 			}
 			return Type && (Type->PinCategory == UEdGraphSchema_K2::PC_Double);
 		}
+#endif
 
 		static bool IsInt(const FEdGraphPinType* Type, const FProperty* Property)
 		{
@@ -604,7 +626,7 @@ public:
 		}
 	};
 
-	virtual void EmitTermExpr(FBPTerminal* Term, FProperty* CoerceProperty = NULL, bool bAllowStaticArray = false)
+	virtual void EmitTermExpr(FBPTerminal* Term, const FProperty* CoerceProperty = NULL, bool bAllowStaticArray = false)
 	{
 		if (Term->bIsLiteral)
 		{
@@ -769,11 +791,11 @@ public:
 
 				UEnum* EnumPtr = nullptr;
 
-				if (FByteProperty* ByteProp = CastField< FByteProperty >(CoerceProperty))
+				if (const FByteProperty* ByteProp = CastField< FByteProperty >(CoerceProperty))
 				{
 					EnumPtr = ByteProp->Enum;
 				}
-				else if (FEnumProperty* EnumProp = CastField< FEnumProperty >(CoerceProperty))
+				else if (const FEnumProperty* EnumProp = CastField< FEnumProperty >(CoerceProperty))
 				{
 					EnumPtr = EnumProp->GetEnum();
 				}
@@ -811,7 +833,7 @@ public:
 			}
 			else if (FLiteralTypeHelper::IsStruct(&Term->Type, CoerceProperty))
 			{
-				FStructProperty* StructProperty = CastField<FStructProperty>(CoerceProperty);
+				const FStructProperty* StructProperty = CastField<FStructProperty>(CoerceProperty);
 				UScriptStruct* Struct = StructProperty ? StructProperty->Struct : Cast<UScriptStruct>(Term->Type.PinSubCategoryObject.Get());
 				check(Struct);
 
@@ -829,6 +851,20 @@ public:
 					Writer << EX_VectorConst;
 					Writer << V;
 				}
+				else if (Struct == Vector3fStruct)
+				{
+					FVector3f V = FVector3f::ZeroVector;
+					if (!Term->Name.IsEmpty())
+					{
+						const bool bParsedUsingCustomFormat = FDefaultValueHelper::ParseVector(Term->Name, /*out*/ V);
+						if (!bParsedUsingCustomFormat)
+						{
+							Struct->ImportText(*Term->Name, &V, nullptr, PPF_None, GWarn, GetPathNameSafe(StructProperty));
+						}
+					}
+					Writer << EX_Vector3fConst;
+					Writer << V;
+				}
 				else if (Struct == RotatorStruct)
 				{
 					FRotator R = FRotator::ZeroRotator;
@@ -913,7 +949,7 @@ public:
 					Writer << EX_EndStructConst;
 				}
 			}
-			else if (FArrayProperty* ArrayPropr = CastField<FArrayProperty>(CoerceProperty))
+			else if (const FArrayProperty* ArrayPropr = CastField<FArrayProperty>(CoerceProperty))
 			{
 				FProperty* InnerProp = ArrayPropr->Inner;
 				ensure(InnerProp);
@@ -933,7 +969,7 @@ public:
 				}
 				Writer << EX_EndArrayConst;
 			}
-			else if (FSetProperty* SetPropr = CastField<FSetProperty>(CoerceProperty))
+			else if (const FSetProperty* SetPropr = CastField<FSetProperty>(CoerceProperty))
 			{
 				FProperty* InnerProp = SetPropr->ElementProp;
 				ensure(InnerProp);
@@ -960,7 +996,7 @@ public:
 				}
 				Writer << EX_EndSetConst;
 			}
-			else if (FMapProperty* MapPropr = CastField<FMapProperty>(CoerceProperty))
+			else if (const FMapProperty* MapPropr = CastField<FMapProperty>(CoerceProperty))
 			{
 				FProperty* KeyProp = MapPropr->KeyProp;
 				FProperty* ValProp = MapPropr->ValueProp;
@@ -1373,7 +1409,7 @@ public:
 		Writer << EX_EndFunctionParms;
 	}
 
-	void EmitTerm(FBPTerminal* Term, FProperty* CoerceProperty = NULL, FBPTerminal* RValueTerm = NULL)
+	void EmitTerm(FBPTerminal* Term, const FProperty* CoerceProperty = NULL, FBPTerminal* RValueTerm = NULL)
 	{
 		if (Term->InlineGeneratedParameter)
 		{
@@ -1592,7 +1628,7 @@ public:
 		Writer << PropertyToHandleComplexStruct;
 		EmitTerm(DestinationExpression);
 
-		Writer << EX_PrimitiveCast;
+		Writer << EX_Cast;
 		uint8 CastType = !bIsInterfaceCast ? CST_ObjectToBool : CST_InterfaceToBool;
 		Writer << CastType;
 		
@@ -2014,6 +2050,80 @@ public:
 		EmitTerm(Statement.RHS[1], (FProperty*)(GetDefault<FIntProperty>()));
 	}
 
+	void EmitCastStatement(FBlueprintCompiledStatement& Statement)
+	{
+		FBPTerminal* DestinationExpression = Statement.LHS;
+		FBPTerminal* TargetExpression = Statement.RHS[0];
+
+		Writer << EX_Let;
+		FProperty* PropertyToHandleComplexStruct = nullptr;
+		Writer << PropertyToHandleComplexStruct;
+		EmitTerm(DestinationExpression);
+
+		Writer << EX_Cast;
+
+		ECastToken CastType = CST_Max;
+
+		switch (Statement.Type)
+		{
+			case KCST_DoubleToFloatCast:
+				CastType = CST_DoubleToFloat;
+				break;
+			case KCST_DoubleToFloatArrayCast:
+				CastType = CST_DoubleToFloatArray;
+				break;
+			case KCST_DoubleToFloatSetCast:
+				CastType = CST_DoubleToFloatSet;
+				break;
+			case KCST_FloatToDoubleCast:
+				CastType = CST_FloatToDouble;
+				break;
+			case KCST_FloatToDoubleArrayCast:
+				CastType = CST_FloatToDoubleArray;
+				break;
+			case KCST_FloatToDoubleSetCast:
+				CastType = CST_FloatToDoubleSet;
+				break;
+			case KCST_VectorToVector3fCast:
+				CastType = CST_VectorToVector3f;
+				break;
+			case KCST_Vector3fToVectorCast:
+				CastType = CST_Vector3fToVector;
+				break;
+			case KCST_FloatToDoubleKeys_MapCast:
+				CastType = CST_FloatToDoubleKeys_Map;
+				break;
+			case KCST_DoubleToFloatKeys_MapCast:
+				CastType = CST_DoubleToFloatKeys_Map;
+				break;
+			case KCST_FloatToDoubleValues_MapCast:
+				CastType = CST_FloatToDoubleValues_Map;
+				break;
+			case KCST_DoubleToFloatValues_MapCast:
+				CastType = CST_DoubleToFloatValues_Map;
+				break;
+			case KCST_FloatToDoubleKeys_FloatToDoubleValues_MapCast:
+				CastType = CST_FloatToDoubleKeys_FloatToDoubleValues_Map;
+				break;
+			case KCST_DoubleToFloatKeys_FloatToDoubleValues_MapCast:
+				CastType = CST_DoubleToFloatKeys_FloatToDoubleValues_Map;
+				break;
+			case KCST_DoubleToFloatKeys_DoubleToFloatValues_MapCast:
+				CastType = CST_DoubleToFloatKeys_DoubleToFloatValues_Map;
+				break;
+			case KCST_FloatToDoubleKeys_DoubleToFloatValues_MapCast:
+				CastType = CST_FloatToDoubleKeys_DoubleToFloatValues_Map;
+				break;
+			default:
+				check(false);
+				break;
+		}
+
+		Writer << CastType;
+
+		EmitTerm(TargetExpression);
+	}
+
 	void PushReturnAddress(FBlueprintCompiledStatement& ReturnTarget)
 	{
 		Writer << EX_PushExecutionFlow;
@@ -2140,6 +2250,28 @@ public:
 		case KCST_CreateMap:
 			EmitCreateMapStatement(Statement);
 			break;
+		case KCST_DoubleToFloatCast:
+		case KCST_DoubleToFloatArrayCast:
+		case KCST_DoubleToFloatSetCast:
+		case KCST_FloatToDoubleCast:
+		case KCST_FloatToDoubleArrayCast:
+		case KCST_FloatToDoubleSetCast:
+		case KCST_VectorToVector3fCast:
+		case KCST_VectorToVector3fArrayCast:
+		case KCST_VectorToVector3fSetCast:
+		case KCST_Vector3fToVectorCast:
+		case KCST_Vector3fToVectorArrayCast:
+		case KCST_Vector3fToVectorSetCast:
+		case KCST_FloatToDoubleKeys_MapCast:
+		case KCST_DoubleToFloatKeys_MapCast:
+		case KCST_FloatToDoubleValues_MapCast:
+		case KCST_DoubleToFloatValues_MapCast:
+		case KCST_FloatToDoubleKeys_FloatToDoubleValues_MapCast:
+		case KCST_DoubleToFloatKeys_FloatToDoubleValues_MapCast:
+		case KCST_DoubleToFloatKeys_DoubleToFloatValues_MapCast:
+		case KCST_FloatToDoubleKeys_DoubleToFloatValues_MapCast:
+			EmitCastStatement(Statement);
+			break;
 		default:
 			UE_LOG(LogK2Compiler, Warning, TEXT("VM backend encountered unsupported statement type %d"), (int32)Statement.Type);
 		}
diff --git a/Engine/Source/Editor/KismetCompiler/Public/BlueprintCompiledStatement.h b/Engine/Source/Editor/KismetCompiler/Public/BlueprintCompiledStatement.h
--- a/Engine/Source/Editor/KismetCompiler/Public/BlueprintCompiledStatement.h
+++ b/Engine/Source/Editor/KismetCompiler/Public/BlueprintCompiledStatement.h
@@ -67,6 +67,32 @@ enum EKismetCompiledStatementType
 	KCST_GotoReturnIfNot = 28,
 	KCST_SwitchValue = 29,
 	
+	KCST_DoubleToFloatCast = 30,
+	KCST_DoubleToFloatArrayCast = 31,
+	KCST_DoubleToFloatSetCast = 32,
+	
+	KCST_FloatToDoubleCast = 33,
+	KCST_FloatToDoubleArrayCast = 34,
+	KCST_FloatToDoubleSetCast = 35,
+	
+	KCST_VectorToVector3fCast = 36,
+	KCST_VectorToVector3fArrayCast = 37,
+	KCST_VectorToVector3fSetCast = 38,
+	
+	KCST_Vector3fToVectorCast = 39,
+	KCST_Vector3fToVectorArrayCast = 40,
+	KCST_Vector3fToVectorSetCast = 41,
+
+	KCST_FloatToDoubleKeys_MapCast = 42,
+	KCST_DoubleToFloatKeys_MapCast = 43,
+	KCST_FloatToDoubleValues_MapCast = 44,
+	KCST_DoubleToFloatValues_MapCast = 45,
+
+	KCST_FloatToDoubleKeys_FloatToDoubleValues_MapCast = 46,
+	KCST_DoubleToFloatKeys_FloatToDoubleValues_MapCast = 47,
+	KCST_DoubleToFloatKeys_DoubleToFloatValues_MapCast = 48,
+	KCST_FloatToDoubleKeys_DoubleToFloatValues_MapCast = 49,
+
 	//~ Kismet instrumentation extensions:
 
 	// Instrumented event
diff --git a/Engine/Source/Editor/KismetCompiler/Public/KismetCastingUtils.h b/Engine/Source/Editor/KismetCompiler/Public/KismetCastingUtils.h
new file mode 100644
--- /dev/null
+++ b/Engine/Source/Editor/KismetCompiler/Public/KismetCastingUtils.h
@@ -0,0 +1,74 @@
+// Copyright Epic Games, Inc. All Rights Reserved.
+
+#pragma once
+
+#include "CoreMinimal.h"
+#include "BlueprintCompiledStatement.h"
+
+struct FBPTerminal;
+struct FKismetFunctionContext;
+class UEdGraphPin;
+
+namespace UE::KismetCompiler::CastingUtils
+{
+	using StatementNamePair = TPair<EKismetCompiledStatementType, const TCHAR*>;
+
+	/**
+	 * Given a specific EKismetCompiledStatementType enum, this function returns its inverse.
+	 * eg: the inverse of KCST_DoubleToFloatCast is KCST_FloatToDoubleCast.
+	 * Invalid enums (ie: not cast-related) will return an unset TOptional.
+	 * 
+	 * @param Statement - The enum to query for an inverse.
+	 * @return A set TOptional with the enum's inverse. Unset if no inverse is found.
+	 */
+	KISMETCOMPILER_API TOptional<EKismetCompiledStatementType> GetInverseCastStatement(EKismetCompiledStatementType Statement);
+
+	/**
+	 * Analyzes the NetMap of the current function context for potential implicit casts.
+	 * If any are found, they're added to ImplicitCastMap in the context.
+	 * After function compilation, the Kismet compiler will validate that the map is empty.
+	 * It's up to the nodes to check the map and insert cast statements where necessary.
+	 * 
+	 * @param Context - Current function context to analyze. Assumes the NetMap has been populated.
+	 */
+	KISMETCOMPILER_API void RegisterImplicitCasts(FKismetFunctionContext& Context);
+
+	/**
+	 * Utility function used by nodes for inserting implicit cast statements.
+	 * During compilation, a node that potentially may need to handle a cast should call this function.
+	 * If the current pin needs a cast, a statement is inserted, and a new terminal for the temporary is returned.
+	 * 
+	 * @param Context - Current function context to analyze. Assumes the ImplicitCastMap has been populated.
+	 * @param DestinationPin - Used as a key in the ImplicitCastMap. These pins are always inputs.
+	 * @param RHSTerm - The current terminal that should have its data read from.
+	 * @return A new FBPTerminal, EKismetCompiledStatementType pair for the temporary variable that has the casted result (if one exists)
+	 */
+	KISMETCOMPILER_API TOptional<TPair<FBPTerminal*, EKismetCompiledStatementType>> InsertImplicitCastStatement(FKismetFunctionContext& Context,
+																												UEdGraphPin* DestinationPin,
+																												FBPTerminal* RHSTerm);
+
+	/**
+	 * Removes the specific UEdGraphPin from the context's implicit cast map.
+	 * In most cases, InsertImplicitCastStatement should be used to remove the cast map entry.
+	 * However, some nodes need to implement custom behavior for casting.
+	 * 
+	 * @param Context - Current function context to analyze. Assumes the ImplicitCastMap has been populated.
+	 * @param DestinationPin - Used as a key in the ImplicitCastMap. These pins are always inputs.
+	 * @return True if DestinationPin was found in the ImplicitCastMap.
+	 */
+	KISMETCOMPILER_API bool RemoveRegisteredImplicitCast(FKismetFunctionContext& Context, const UEdGraphPin* DestinationPin);
+
+	/**
+	 * Retrieves the conversion type needed between two arbitrary pins (if necessary). Specifically, this indicates if either
+	 * a narrowing or widening cast is needed between a float or a double type (including containers). In addition to the
+	 * corresponding EKismetCompiledStatementType that represents the cast type, a string literal describing the cast is also
+	 * returned.
+	 * 
+	 * @param SourcePin - The source pin to compare.
+	 * @param DestinationPin - The destination pin to compare.
+	 * @return A new StatementNamePair containing the cast information. The result is unset if no conversion is needed.
+	 */
+	KISMETCOMPILER_API TOptional<StatementNamePair> GetFloatingPointConversionType(const UEdGraphPin& SourcePin, const UEdGraphPin& DestinationPin);
+
+} // UE::KismetCompiler::CastingUtils
+
diff --git a/Engine/Source/Editor/KismetCompiler/Public/KismetCompiledFunctionContext.h b/Engine/Source/Editor/KismetCompiler/Public/KismetCompiledFunctionContext.h
--- a/Engine/Source/Editor/KismetCompiler/Public/KismetCompiledFunctionContext.h
+++ b/Engine/Source/Editor/KismetCompiler/Public/KismetCompiledFunctionContext.h
@@ -32,6 +32,13 @@ enum ETerminalSpecification
 	TS_ForcedShared,
 };
 
+struct FImplicitCastParams
+{
+	EKismetCompiledStatementType CastType = KCST_Nop;
+	FBPTerminal* TargetTerminal = nullptr;
+	UEdGraphNode* TargetNode = nullptr;
+};
+
 //////////////////////////////////////////////////////////////////////////
 // FKismetFunctionContext
 
@@ -60,17 +67,16 @@ public:
 	TArray< FBlueprintCompiledStatement* > AllGeneratedStatements;
 
 	// Individual execution lists for every node that generated code to be consumed by the backend
-	TMap< UEdGraphNode*, TArray<FBlueprintCompiledStatement*> > StatementsPerNode;
+	TMap<UEdGraphNode*, TArray<FBlueprintCompiledStatement*>> StatementsPerNode;
 
 	// Goto fixup requests (each statement (key) wants to goto the first statement attached to the exec out-pin (value))
-	TMap< FBlueprintCompiledStatement*, UEdGraphPin* > GotoFixupRequestMap;
+	TMap<FBlueprintCompiledStatement*, UEdGraphPin*> GotoFixupRequestMap;
 
 	// @todo: BP2CPP_remove
 	// Used to split uber graph into subfunctions by C++ backend
 	UE_DEPRECATED(5.0, "This member is no longer in use and will be removed.")
 	TArray<TSet<UEdGraphNode*>> UnsortedSeparateExecutionGroups;
 
-	// Map from a net to an term (either a literal or a storage location)
 	TIndirectArray<FBPTerminal> Parameters;
 	TIndirectArray<FBPTerminal> Results;
 	TIndirectArray<FBPTerminal> VariableReferences;
@@ -80,9 +86,14 @@ public:
 	TIndirectArray<FBPTerminal> EventGraphLocals;
 	TIndirectArray<FBPTerminal>	LevelActorReferences;
 	TIndirectArray<FBPTerminal>	InlineGeneratedValues; // A function generating the parameter will be called inline. The value won't be stored in a local variable.
+
+	// Map from a net to an term (either a literal or a storage location)
 	TMap<UEdGraphPin*, FBPTerminal*> NetMap;
 	TMap<UEdGraphPin*, FBPTerminal*> LiteralHackMap;
 
+	// Contains a map of destination pins that will need an implicit cast to either a float or double
+	TMap<UEdGraphPin*, FImplicitCastParams> ImplicitCastMap;
+
 	bool bIsUbergraph;
 	bool bCannotBeCalledFromOtherKismet;
 	bool bIsInterfaceStub;
diff --git a/Engine/Source/Editor/KismetCompiler/Public/KismetCompiler.h b/Engine/Source/Editor/KismetCompiler/Public/KismetCompiler.h
--- a/Engine/Source/Editor/KismetCompiler/Public/KismetCompiler.h
+++ b/Engine/Source/Editor/KismetCompiler/Public/KismetCompiler.h
@@ -498,7 +498,8 @@ protected:
 	/**
 	 * Used for performing custom patching during stage IX of the compilation during load.
 	 */
-	virtual void PreCompileUpdateBlueprintOnLoad(UBlueprint* BP) {}
+	virtual void PreCompileUpdateBlueprintOnLoad(UBlueprint* BP);
+
 	/**
 	 * Second phase of compiling a function graph
 	 *   - Generates an executable statement list
diff --git a/Engine/Source/Editor/KismetCompiler/Public/KismetCompilerMisc.h b/Engine/Source/Editor/KismetCompiler/Public/KismetCompilerMisc.h
--- a/Engine/Source/Editor/KismetCompiler/Public/KismetCompilerMisc.h
+++ b/Engine/Source/Editor/KismetCompiler/Public/KismetCompilerMisc.h
@@ -86,7 +86,7 @@ public:
 	static UEdGraphPin* GenerateAssignmentNodes( class FKismetCompilerContext& CompilerContext, UEdGraph* SourceGraph, UK2Node_CallFunction* CallBeginSpawnNode, UEdGraphNode* SpawnNode, UEdGraphPin* CallBeginResult, const UClass* ForClass );
 
 	/** Create Kismet assignment statement with proper object <-> interface cast */
-	static void CreateObjectAssignmentStatement(FKismetFunctionContext& Context, UEdGraphNode* Node, FBPTerminal* SrcTerm, FBPTerminal* DstTerm);
+	static void CreateObjectAssignmentStatement(FKismetFunctionContext& Context, UEdGraphNode* Node, FBPTerminal* SrcTerm, FBPTerminal* DstTerm, UEdGraphPin* DstPin = nullptr);
 
 	/** Checks if each execution path ends with a Return node */
 	static void ValidateProperEndExecutionPath(FKismetFunctionContext& Context);
diff --git a/Engine/Source/Editor/UMGEditor/Private/WidgetBlueprint.cpp b/Engine/Source/Editor/UMGEditor/Private/WidgetBlueprint.cpp
--- a/Engine/Source/Editor/UMGEditor/Private/WidgetBlueprint.cpp
+++ b/Engine/Source/Editor/UMGEditor/Private/WidgetBlueprint.cpp
@@ -20,6 +20,7 @@
 #include "WidgetBlueprintCompiler.h"
 #include "UObject/EditorObjectVersion.h"
 #include "UObject/FortniteMainBranchObjectVersion.h"
+#include "UObject/UE5ReleaseStreamObjectVersion.h"
 #include "UObject/ObjectSaveContext.h"
 #include "WidgetGraphSchema.h"
 #include "UMGEditorProjectSettings.h"
@@ -33,6 +34,7 @@
 #include "K2Node_CallFunction.h"
 #include "K2Node_MacroInstance.h"
 #include "K2Node_Composite.h"
+#include "K2Node_FunctionResult.h"
 #include "Blueprint/WidgetNavigation.h"
 
 #define LOCTEXT_NAMESPACE "UMG"
@@ -582,6 +584,92 @@ void UWidgetBlueprint::ReplaceDeprecatedNodes()
 		}
 	}
 
+#if WITH_EDITORONLY_DATA && ENABLE_BLUEPRINT_REAL_NUMBERS
+	if (GetLinkerCustomVersion(FUE5ReleaseStreamObjectVersion::GUID) < FUE5ReleaseStreamObjectVersion::BlueprintPinsUseRealNumbers)
+	{
+		// Revert any overzealous PC_Float to PC_Real/PC_Double conversions.
+
+		// The Blueprint real number changes will automatically convert pin types to doubles if used in a non-native context.
+		// However, UMG property bindings are a special case: the BP functions that bind to the native delegate must agree on their underlying types.
+		// Specifically, bindings used with float properties *must* use the PC_Float type as the return value in a BP function.
+		// In order to correct this behavior, we need to:
+		// * Iterate throught the property bindings.
+		// * Find the corresponding delegate signature.
+		// * Find the function graph that matches the binding.
+		// * Find the result node.
+		// * Change the pin type back to float if that's what the delegate signature expects.
+
+		TArray<UEdGraph*> Graphs;
+		GetAllGraphs(Graphs);
+
+		for (const FDelegateEditorBinding& Binding : Bindings)
+		{
+			if (Binding.IsAttributePropertyBinding(this))
+			{
+				check(WidgetTree);
+				if (UWidget* TargetWidget = WidgetTree->FindWidget(FName(*Binding.ObjectName)))
+				{
+					const FDelegateProperty* BindableProperty =
+						FindFProperty<FDelegateProperty>(TargetWidget->GetClass(), FName(*(Binding.PropertyName.ToString() + TEXT("Delegate"))));
+
+					if (BindableProperty)
+					{
+						auto GraphMatchesBindingPredicate = [&Binding](const UEdGraph* Graph) {
+							check(Graph);
+							return (Binding.FunctionName == Graph->GetFName());
+						};
+
+						if (UEdGraph** GraphEntry = Graphs.FindByPredicate(GraphMatchesBindingPredicate))
+						{
+							UEdGraph* CurrentGraph = *GraphEntry;
+							check(CurrentGraph);
+
+							for (UEdGraphNode* Node : CurrentGraph->Nodes)
+							{
+								check(Node);
+								if (Node->IsA<UK2Node_FunctionResult>())
+								{
+									for (UEdGraphPin* Pin : Node->Pins)
+									{
+										check(Pin);
+										if (Pin->PinType.PinCategory == UEdGraphSchema_K2::PC_Real)
+										{
+											FName PinName = Pin->GetFName();
+
+											const UFunction* DelegateFunction = BindableProperty->SignatureFunction;
+											check(DelegateFunction);
+
+											auto OutputParameterMatchesPin = [PinName](FFloatProperty* FloatParam) {
+												check(FloatParam);
+												bool bHasMatch =
+													(FloatParam->PropertyFlags & CPF_OutParm) &&
+													(FloatParam->GetFName() == PinName);
+
+												return bHasMatch;
+											};
+
+											for (TFieldIterator<FFloatProperty> It(DelegateFunction); It; ++It)
+											{
+												if (OutputParameterMatchesPin(*It))
+												{
+													Pin->PinType.PinSubCategory = UEdGraphSchema_K2::PC_Float;
+													break;
+												}
+											}
+
+											break;
+										}
+									}
+								}
+							}
+						}
+					}
+				}
+			}
+		}
+	}
+#endif
+
 	Super::ReplaceDeprecatedNodes();
 }
 
@@ -847,6 +935,7 @@ void UWidgetBlueprint::Serialize(FArchive& Ar)
 
 	Ar.UsingCustomVersion(FEditorObjectVersion::GUID);
 	Ar.UsingCustomVersion(FFortniteMainBranchObjectVersion::GUID);
+	Ar.UsingCustomVersion(FUE5ReleaseStreamObjectVersion::GUID);
 }
 
 void UWidgetBlueprint::PostLoad()
diff --git a/Engine/Source/Editor/UnrealEd/Classes/UserDefinedStructure/UserDefinedStructEditorData.h b/Engine/Source/Editor/UnrealEd/Classes/UserDefinedStructure/UserDefinedStructEditorData.h
--- a/Engine/Source/Editor/UnrealEd/Classes/UserDefinedStructure/UserDefinedStructEditorData.h
+++ b/Engine/Source/Editor/UnrealEd/Classes/UserDefinedStructure/UserDefinedStructEditorData.h
@@ -131,6 +131,7 @@ public:
 public:
 	// UObject interface.
 	virtual TSharedPtr<ITransactionObjectAnnotation> FactoryTransactionAnnotation(const ETransactionAnnotationCreationMode InCreationMode) const override;
+	virtual void Serialize(FArchive& Ar) override;
 	virtual void PostEditUndo() override;
 	virtual void PostEditUndo(TSharedPtr<ITransactionObjectAnnotation> TransactionAnnotation) override;
 	virtual void PostLoadSubobjects(struct FObjectInstancingGraph* OuterInstanceGraph) override;
diff --git a/Engine/Source/Editor/UnrealEd/Private/Kismet2/BlueprintEditorUtils.cpp b/Engine/Source/Editor/UnrealEd/Private/Kismet2/BlueprintEditorUtils.cpp
--- a/Engine/Source/Editor/UnrealEd/Private/Kismet2/BlueprintEditorUtils.cpp
+++ b/Engine/Source/Editor/UnrealEd/Private/Kismet2/BlueprintEditorUtils.cpp
@@ -580,7 +580,6 @@ void FBlueprintEditorUtils::ReconstructAllNodes(UBlueprint* Blueprint)
 	TArray<UK2Node*> AllNodes;
 	FBlueprintEditorUtils::GetAllNodesOfClass(Blueprint, AllNodes);
 
-	const bool bIsMacro = (Blueprint->BlueprintType == BPTYPE_MacroLibrary);
 	if (AllNodes.Num() > 1)
 	{
 		AllNodes.Sort(FCompareNodePriority());
diff --git a/Engine/Source/Editor/UnrealEd/Private/UserDefinedStructEditorData.cpp b/Engine/Source/Editor/UnrealEd/Private/UserDefinedStructEditorData.cpp
--- a/Engine/Source/Editor/UnrealEd/Private/UserDefinedStructEditorData.cpp
+++ b/Engine/Source/Editor/UnrealEd/Private/UserDefinedStructEditorData.cpp
@@ -2,6 +2,7 @@
 
 #include "UserDefinedStructure/UserDefinedStructEditorData.h"
 #include "Misc/ITransaction.h"
+#include "UObject/UE5ReleaseStreamObjectVersion.h"
 #include "UObject/UnrealType.h"
 #include "UObject/ObjectSaveContext.h"
 #include "Engine/UserDefinedStruct.h"
@@ -31,6 +32,19 @@ void FStructVariableDescription::PostSerialize(const FArchive& Ar)
 			Category = TEXT("softclass");
 		}
 	}
+
+#if ENABLE_BLUEPRINT_REAL_NUMBERS
+	bool FixupPinCategories =
+		Ar.IsLoading() &&
+		(Ar.CustomVer(FUE5ReleaseStreamObjectVersion::GUID) < FUE5ReleaseStreamObjectVersion::BlueprintPinsUseRealNumbers) &&
+		((Category == TEXT("double")) || (Category == TEXT("float")));
+
+	if (FixupPinCategories)
+	{
+		Category = TEXT("real");
+		SubCategory = TEXT("double");
+	}
+#endif
 }
 
 bool FStructVariableDescription::SetPinType(const FEdGraphPinType& VarType)
@@ -66,6 +80,15 @@ UUserDefinedStruct* UUserDefinedStructEditorData::GetOwnerStruct() const
 	return Cast<UUserDefinedStruct>(GetOuter());
 }
 
+void UUserDefinedStructEditorData::Serialize(FArchive& Ar)
+{
+	Super::Serialize(Ar);
+
+#if ENABLE_BLUEPRINT_REAL_NUMBERS
+	Ar.UsingCustomVersion(FUE5ReleaseStreamObjectVersion::GUID);
+#endif
+}
+
 void UUserDefinedStructEditorData::PostUndo(bool bSuccess)
 {
 	GEditor->UnregisterForUndo(this);
diff --git a/Engine/Source/Runtime/Core/Public/Misc/Build.h b/Engine/Source/Runtime/Core/Public/Misc/Build.h
--- a/Engine/Source/Runtime/Core/Public/Misc/Build.h
+++ b/Engine/Source/Runtime/Core/Public/Misc/Build.h
@@ -364,6 +364,11 @@
 /** Enable fast calls for event thunks into an event graph that have no parameters  */
 #define UE_BLUEPRINT_EVENTGRAPH_FASTCALLS 1
 
+/** Toggles the use of real numbers in blueprints. Explicit float and double usage will not be available. */
+#ifndef ENABLE_BLUEPRINT_REAL_NUMBERS
+#define ENABLE_BLUEPRINT_REAL_NUMBERS 1
+#endif
+
 /** Enable perf counters on dedicated servers */
 #define USE_SERVER_PERF_COUNTERS ((UE_SERVER || UE_EDITOR) && WITH_PERFCOUNTERS)
 
diff --git a/Engine/Source/Runtime/Core/Public/UObject/UE5ReleaseStreamObjectVersion.h b/Engine/Source/Runtime/Core/Public/UObject/UE5ReleaseStreamObjectVersion.h
--- a/Engine/Source/Runtime/Core/Public/UObject/UE5ReleaseStreamObjectVersion.h
+++ b/Engine/Source/Runtime/Core/Public/UObject/UE5ReleaseStreamObjectVersion.h
@@ -96,6 +96,11 @@ struct CORE_API FUE5ReleaseStreamObjectVersion
 		// Large Worlds - serialize double types as doubles
 		LargeWorldCoordinates,
 		
+#if ENABLE_BLUEPRINT_REAL_NUMBERS
+		// Deserialize old BP float&double types as real numbers for pins
+		BlueprintPinsUseRealNumbers,
+#endif
+		
 		// -----<new versions can be added above this line>-------------------------------------------------
 		VersionPlusOne,
 		LatestVersion = VersionPlusOne - 1
diff --git a/Engine/Source/Runtime/CoreUObject/Private/UObject/ScriptCore.cpp b/Engine/Source/Runtime/CoreUObject/Private/UObject/ScriptCore.cpp
--- a/Engine/Source/Runtime/CoreUObject/Private/UObject/ScriptCore.cpp
+++ b/Engine/Source/Runtime/CoreUObject/Private/UObject/ScriptCore.cpp
@@ -24,6 +24,7 @@
 #include "UObject/SoftObjectPtr.h"
 #include "UObject/PropertyPortFlags.h"
 #include "UObject/UnrealType.h"
+#include "UObject/ScriptCastingUtils.h"
 #include "UObject/Stack.h"
 #include "UObject/Reload.h"
 #include "Blueprint/BlueprintSupport.h"
@@ -87,9 +88,6 @@ static FAutoConsoleVariableRef CVarMaxFunctionStatDepth(
 COREUOBJECT_API FNativeFuncPtr GNatives[EX_Max];
 COREUOBJECT_API int32 GNativeDuplicate=0;
 
-COREUOBJECT_API FNativeFuncPtr GCasts[CST_Max];
-COREUOBJECT_API int32 GCastDuplicate=0;
-
 COREUOBJECT_API int32 GMaximumScriptLoopIterations = 1000000;
 
 #if DO_BLUEPRINT_GUARD
@@ -701,7 +699,7 @@ int32 FScriptInstrumentationSignal::GetScriptCodeOffset() const
 // Register a native function.
 // Warning: Called at startup time, before engine initialization.
 //
-COREUOBJECT_API uint8 GRegisterNative( int32 NativeBytecodeIndex, const FNativeFuncPtr& Func )
+COREUOBJECT_API uint8 GRegisterNative(int32 NativeBytecodeIndex, const FNativeFuncPtr& Func)
 {
 	static bool bInitialized = false;
 	if (!bInitialized)
@@ -715,7 +713,7 @@ COREUOBJECT_API uint8 GRegisterNative( int32 NativeBytecodeIndex, const FNativeF
 
 	if( NativeBytecodeIndex != INDEX_NONE )
 	{
-		if( NativeBytecodeIndex<0 || (uint32)NativeBytecodeIndex>UE_ARRAY_COUNT(GNatives) || GNatives[NativeBytecodeIndex]!=&UObject::execUndefined) 
+		if( NativeBytecodeIndex<0 || (uint32)NativeBytecodeIndex > UE_ARRAY_COUNT(GNatives) || GNatives[NativeBytecodeIndex] != &UObject::execUndefined) 
 		{
 			CA_SUPPRESS(6385)
 			if (!ReloadNotifyFunctionRemap(Func, GNatives[NativeBytecodeIndex]))
@@ -730,7 +728,9 @@ COREUOBJECT_API uint8 GRegisterNative( int32 NativeBytecodeIndex, const FNativeF
 	return 0;
 }
 
-COREUOBJECT_API uint8 GRegisterCast( int32 CastCode, const FNativeFuncPtr& Func )
+static FNativeFuncPtr GCasts[CST_Max];
+
+static uint8 GRegisterCast(ECastToken CastCode, const FNativeFuncPtr& Func)
 {
 	static int32 bInitialized = false;
 	if (!bInitialized)
@@ -742,15 +742,8 @@ COREUOBJECT_API uint8 GRegisterCast( int32 CastCode, const FNativeFuncPtr& Func
 		}
 	}
 
-	//@TODO: UCREMOVAL: Remove rest of cast machinery
-	check((CastCode == CST_ObjectToBool) || (CastCode == CST_ObjectToInterface) || (CastCode == CST_InterfaceToBool));
-
-	if (CastCode != INDEX_NONE)
+	if (CastCode != CST_Max)
 	{
-		if(!IsReloadActive() && (CastCode<0 || (uint32)CastCode>UE_ARRAY_COUNT(GCasts) || GCasts[CastCode]!=&UObject::execUndefined) ) 
-		{
-			GCastDuplicate = CastCode;
-		}
 		GCasts[CastCode] = Func;
 	}
 	return 0;
@@ -3336,6 +3329,14 @@ DEFINE_FUNCTION(UObject::execVectorConst)
 }
 IMPLEMENT_VM_FUNCTION( EX_VectorConst, execVectorConst );
 
+DEFINE_FUNCTION(UObject::execVector3fConst)
+{
+	((FVector3f*)RESULT_PARAM)->X = Stack.ReadFloat();
+	((FVector3f*)RESULT_PARAM)->Y = Stack.ReadFloat();
+	((FVector3f*)RESULT_PARAM)->Z = Stack.ReadFloat();
+}
+IMPLEMENT_VM_FUNCTION(EX_Vector3fConst, execVector3fConst);
+
 DEFINE_FUNCTION(UObject::execTransformConst)
 {
 	// Rotation
@@ -3649,12 +3650,12 @@ DEFINE_FUNCTION(UObject::execMetaCast)
 }
 IMPLEMENT_VM_FUNCTION( EX_MetaCast, execMetaCast );
 
-DEFINE_FUNCTION(UObject::execPrimitiveCast)
+DEFINE_FUNCTION(UObject::execCast)
 {
 	int32 B = *(Stack.Code)++;
 	(*GCasts[B])( Stack.Object, Stack, RESULT_PARAM );
 }
-IMPLEMENT_VM_FUNCTION( EX_PrimitiveCast, execPrimitiveCast );
+IMPLEMENT_VM_FUNCTION( EX_Cast, execCast );
 
 DEFINE_FUNCTION(UObject::execInterfaceCast)
 {
@@ -3662,6 +3663,110 @@ DEFINE_FUNCTION(UObject::execInterfaceCast)
 }
 IMPLEMENT_VM_FUNCTION( EX_ObjToInterfaceCast, execInterfaceCast );
 
+DEFINE_FUNCTION(UObject::execDoubleToFloatCast)
+{
+	Stack.Step(Stack.Object, nullptr);
+
+	DoubleToFloatCast(nullptr, Stack.MostRecentPropertyAddress, RESULT_PARAM);
+}
+IMPLEMENT_CAST_FUNCTION( CST_DoubleToFloat, execDoubleToFloatCast )
+
+DEFINE_FUNCTION(UObject::execDoubleToFloatArrayCast)
+{
+	CopyAndCastArrayFromStack<double, float>(Stack, RESULT_PARAM);
+}
+IMPLEMENT_CAST_FUNCTION( CST_DoubleToFloatArray, execDoubleToFloatArrayCast )
+
+DEFINE_FUNCTION(UObject::execDoubleToFloatSetCast)
+{
+	CopyAndCastSetFromStack<double, float>(Stack, RESULT_PARAM);
+}
+IMPLEMENT_CAST_FUNCTION( CST_DoubleToFloatSet, execDoubleToFloatSetCast )
+
+DEFINE_FUNCTION(UObject::execFloatToDoubleCast)
+{
+	Stack.Step(Stack.Object, nullptr);
+
+	FloatToDoubleCast(nullptr, Stack.MostRecentPropertyAddress, RESULT_PARAM);
+}
+IMPLEMENT_CAST_FUNCTION( CST_FloatToDouble, execFloatToDoubleCast )
+
+DEFINE_FUNCTION(UObject::execFloatToDoubleArrayCast)
+{
+	CopyAndCastArrayFromStack<float, double>(Stack, RESULT_PARAM);
+}
+IMPLEMENT_CAST_FUNCTION( CST_FloatToDoubleArray, execFloatToDoubleArrayCast )
+
+DEFINE_FUNCTION(UObject::execFloatToDoubleSetCast)
+{
+	CopyAndCastSetFromStack<float, double>(Stack, RESULT_PARAM);
+}
+IMPLEMENT_CAST_FUNCTION( CST_FloatToDoubleSet, execFloatToDoubleSetCast )
+
+DEFINE_FUNCTION(UObject::execVectorToVector3fCast)
+{
+	Stack.Step(Stack.Object, nullptr);
+
+	FloatingPointCast<FVector, FVector3f>(nullptr, Stack.MostRecentPropertyAddress, RESULT_PARAM);
+}
+IMPLEMENT_CAST_FUNCTION( CST_VectorToVector3f, execVectorToVector3fCast )
+
+DEFINE_FUNCTION(UObject::execVector3fToVectorCast)
+{
+	Stack.Step(Stack.Object, nullptr);
+
+	FloatingPointCast<FVector3f, FVector>(nullptr, Stack.MostRecentPropertyAddress, RESULT_PARAM);
+}
+IMPLEMENT_CAST_FUNCTION( CST_Vector3fToVector, execVector3fToVectorCast )
+
+DEFINE_FUNCTION(UObject::execFloatToDoubleKeysMapCast)
+{
+	CopyAndCastMapFromStack<FloatToDoubleCast, CopyElement>(Stack, RESULT_PARAM);
+}
+IMPLEMENT_CAST_FUNCTION( CST_FloatToDoubleKeys_Map, execFloatToDoubleKeysMapCast )
+
+DEFINE_FUNCTION(UObject::execDoubleToFloatKeysMapCast)
+{
+	CopyAndCastMapFromStack<DoubleToFloatCast, CopyElement>(Stack, RESULT_PARAM);
+}
+IMPLEMENT_CAST_FUNCTION( CST_DoubleToFloatKeys_Map, execDoubleToFloatKeysMapCast )
+
+DEFINE_FUNCTION(UObject::execFloatToDoubleValuesMapCast)
+{
+	CopyAndCastMapFromStack<CopyElement, FloatToDoubleCast>(Stack, RESULT_PARAM);
+}
+IMPLEMENT_CAST_FUNCTION( CST_FloatToDoubleValues_Map, execFloatToDoubleValuesMapCast )
+
+DEFINE_FUNCTION(UObject::execDoubleToFloatValuesMapCast)
+{
+	CopyAndCastMapFromStack<CopyElement, DoubleToFloatCast>(Stack, RESULT_PARAM);
+}
+IMPLEMENT_CAST_FUNCTION( CST_DoubleToFloatValues_Map, execDoubleToFloatValuesMapCast )
+
+DEFINE_FUNCTION(UObject::execFloatToDoubleKeysFloatToDoubleValuesMapCast)
+{
+	CopyAndCastMapFromStack<FloatToDoubleCast, FloatToDoubleCast>(Stack, RESULT_PARAM);
+}
+IMPLEMENT_CAST_FUNCTION( CST_FloatToDoubleKeys_FloatToDoubleValues_Map, execFloatToDoubleKeysFloatToDoubleValuesMapCast )
+
+DEFINE_FUNCTION(UObject::execDoubleToFloatKeysFloatToDoubleValuesMapCast)
+{
+	CopyAndCastMapFromStack<DoubleToFloatCast, FloatToDoubleCast>(Stack, RESULT_PARAM);
+}
+IMPLEMENT_CAST_FUNCTION( CST_DoubleToFloatKeys_FloatToDoubleValues_Map, execDoubleToFloatKeysFloatToDoubleValuesMapCast )
+
+DEFINE_FUNCTION(UObject::execDoubleToFloatKeysDoubleToFloatValuesMapCast)
+{
+	CopyAndCastMapFromStack<DoubleToFloatCast, DoubleToFloatCast>(Stack, RESULT_PARAM);
+}
+IMPLEMENT_CAST_FUNCTION( CST_DoubleToFloatKeys_DoubleToFloatValues_Map, execDoubleToFloatKeysDoubleToFloatValuesMapCast )
+
+DEFINE_FUNCTION(UObject::execFloatToDoubleKeysDoubleToFloatValuesMapCast)
+{
+	CopyAndCastMapFromStack<FloatToDoubleCast, DoubleToFloatCast>(Stack, RESULT_PARAM);
+}
+IMPLEMENT_CAST_FUNCTION( CST_FloatToDoubleKeys_DoubleToFloatValues_Map, execFloatToDoubleKeysDoubleToFloatValuesMapCast )
+
 DEFINE_FUNCTION(UObject::execObjectToBool)
 {
 	UObject* Obj=NULL;
diff --git a/Engine/Source/Runtime/CoreUObject/Private/UObject/UObjectGlobals.cpp b/Engine/Source/Runtime/CoreUObject/Private/UObject/UObjectGlobals.cpp
--- a/Engine/Source/Runtime/CoreUObject/Private/UObject/UObjectGlobals.cpp
+++ b/Engine/Source/Runtime/CoreUObject/Private/UObject/UObjectGlobals.cpp
@@ -498,12 +498,6 @@ void StaticTick( float DeltaTime, bool bUseFullTimeLimit, float AsyncLoadingTime
 	{
 		UE_LOG(LogUObjectGlobals, Fatal, TEXT("Duplicate native registered: %i"), GNativeDuplicate );
 	}
-	// Check for duplicates.
-	extern int32 GCastDuplicate;
-	if( GCastDuplicate )
-	{
-		UE_LOG(LogUObjectGlobals, Fatal, TEXT("Duplicate cast registered: %i"), GCastDuplicate );
-	}
 
 #if STATS
 	// Set name table stats.
diff --git a/Engine/Source/Runtime/CoreUObject/Public/UObject/CoreNative.h b/Engine/Source/Runtime/CoreUObject/Public/UObject/CoreNative.h
--- a/Engine/Source/Runtime/CoreUObject/Public/UObject/CoreNative.h
+++ b/Engine/Source/Runtime/CoreUObject/Public/UObject/CoreNative.h
@@ -22,10 +22,6 @@ struct FNameNativePtrPair
 	FNativeFuncPtr Pointer;
 };
 
-extern COREUOBJECT_API FNativeFuncPtr GCasts[];
-uint8 COREUOBJECT_API GRegisterCast( int32 CastCode, const FNativeFuncPtr& Func );
-
-
 /** A struct that maps a string name to a native function */
 struct FNativeFunctionRegistrar
 {
diff --git a/Engine/Source/Runtime/CoreUObject/Public/UObject/Field.h b/Engine/Source/Runtime/CoreUObject/Public/UObject/Field.h
--- a/Engine/Source/Runtime/CoreUObject/Public/UObject/Field.h
+++ b/Engine/Source/Runtime/CoreUObject/Public/UObject/Field.h
@@ -878,6 +878,12 @@ FORCEINLINE FieldType* ExactCastField(FField* Src)
 	return (Src && (Src->GetClass() == FieldType::StaticClass())) ? static_cast<FieldType*>(Src) : nullptr;
 }
 
+template<typename FieldType>
+FORCEINLINE FieldType* ExactCastField(const FField* Src)
+{
+	return (Src && (Src->GetClass() == FieldType::StaticClass())) ? static_cast<const FieldType*>(Src) : nullptr;
+}
+
 template<typename FieldType>
 FUNCTION_NON_NULL_RETURN_START
 FORCEINLINE FieldType* CastFieldChecked(FField* Src)
diff --git a/Engine/Source/Runtime/CoreUObject/Public/UObject/NoExportTypes.h b/Engine/Source/Runtime/CoreUObject/Public/UObject/NoExportTypes.h
--- a/Engine/Source/Runtime/CoreUObject/Public/UObject/NoExportTypes.h
+++ b/Engine/Source/Runtime/CoreUObject/Public/UObject/NoExportTypes.h
@@ -466,7 +466,7 @@ struct FGuid
  * A point or direction FVector in 3d space.
  * @note The full C++ class is located here: Engine\Source\Runtime\Core\Public\Math\Vector.h
  */
-USTRUCT(immutable, noexport)
+USTRUCT(immutable, noexport, BlueprintType, meta = (HasNativeBreak = "Engine.KismetMathLibrary.BreakVector3f"))
 struct FVector3f
 {
 	UPROPERTY(EditAnywhere, Category = Vector, SaveGame)
diff --git a/Engine/Source/Runtime/CoreUObject/Public/UObject/Object.h b/Engine/Source/Runtime/CoreUObject/Public/UObject/Object.h
--- a/Engine/Source/Runtime/CoreUObject/Public/UObject/Object.h
+++ b/Engine/Source/Runtime/CoreUObject/Public/UObject/Object.h
@@ -1487,6 +1487,7 @@ public:
 	DECLARE_FUNCTION(execIntConstByte);
 	DECLARE_FUNCTION(execRotationConst);
 	DECLARE_FUNCTION(execVectorConst);
+	DECLARE_FUNCTION(execVector3fConst);
 	DECLARE_FUNCTION(execTransformConst);
 	DECLARE_FUNCTION(execStructConst);
 	DECLARE_FUNCTION(execSetArray);
@@ -1502,18 +1503,35 @@ public:
 	DECLARE_FUNCTION(execNativeParm);
 
 	// Conversions 
+	DECLARE_FUNCTION(execCast);
 	DECLARE_FUNCTION(execDynamicCast);
 	DECLARE_FUNCTION(execMetaCast);
-	DECLARE_FUNCTION(execPrimitiveCast);
 	DECLARE_FUNCTION(execInterfaceCast);
-
-	// Cast functions
+	DECLARE_FUNCTION(execDoubleToFloatCast);
+	DECLARE_FUNCTION(execDoubleToFloatArrayCast);
+	DECLARE_FUNCTION(execDoubleToFloatSetCast);
+	DECLARE_FUNCTION(execFloatToDoubleCast);
+	DECLARE_FUNCTION(execFloatToDoubleArrayCast);
+	DECLARE_FUNCTION(execFloatToDoubleSetCast);
+	DECLARE_FUNCTION(execVectorToVector3fCast);
+	DECLARE_FUNCTION(execVector3fToVectorCast);
 	DECLARE_FUNCTION(execObjectToBool);
 	DECLARE_FUNCTION(execInterfaceToBool);
 	DECLARE_FUNCTION(execObjectToInterface);
 	DECLARE_FUNCTION(execInterfaceToInterface);
 	DECLARE_FUNCTION(execInterfaceToObject);
 
+	// Map-specific conversions
+	DECLARE_FUNCTION(execFloatToDoubleKeysMapCast);
+	DECLARE_FUNCTION(execDoubleToFloatKeysMapCast);
+	DECLARE_FUNCTION(execFloatToDoubleValuesMapCast);
+	DECLARE_FUNCTION(execDoubleToFloatValuesMapCast);
+	DECLARE_FUNCTION(execFloatToDoubleKeysFloatToDoubleValuesMapCast);
+	DECLARE_FUNCTION(execDoubleToFloatKeysFloatToDoubleValuesMapCast);
+	DECLARE_FUNCTION(execDoubleToFloatKeysDoubleToFloatValuesMapCast);
+	DECLARE_FUNCTION(execFloatToDoubleKeysDoubleToFloatValuesMapCast);
+
+
 	// Dynamic array functions
 	// Array support
 	DECLARE_FUNCTION(execGetDynArrayElement);
diff --git a/Engine/Source/Runtime/CoreUObject/Public/UObject/Script.h b/Engine/Source/Runtime/CoreUObject/Public/UObject/Script.h
--- a/Engine/Source/Runtime/CoreUObject/Public/UObject/Script.h
+++ b/Engine/Source/Runtime/CoreUObject/Public/UObject/Script.h
@@ -219,7 +219,7 @@ enum EExprToken
 	EX_Int64Const			= 0x35,	// 64-bit integer constant.
 	EX_UInt64Const			= 0x36,	// 64-bit unsigned integer constant.
 	EX_DoubleConst			= 0x37, // Double constant.
-	EX_PrimitiveCast		= 0x38,	// A casting operator for primitives which reads the type as the subsequent byte
+	EX_Cast					= 0x38,	// A casting operator which reads the type as the subsequent byte
 	EX_SetSet				= 0x39,
 	EX_EndSet				= 0x3A,
 	EX_SetMap				= 0x3B,
@@ -228,7 +228,7 @@ enum EExprToken
 	EX_EndSetConst			= 0x3E,
 	EX_MapConst				= 0x3F,
 	EX_EndMapConst			= 0x40,
-	//						= 0x41,
+	EX_Vector3fConst		= 0x41,	// A float vector constant.
 	EX_StructMemberContext	= 0x42, // Context expression to address a property within a struct
 	EX_LetMulticastDelegate	= 0x43, // Assignment to a multi-cast delegate
 	EX_LetDelegate			= 0x44, // Assignment to a delegate
@@ -276,13 +276,29 @@ enum EExprToken
 	EX_Max					= 0x100,
 };
 
-
 enum ECastToken
 {
-	CST_ObjectToInterface	= 0x46,
-	CST_ObjectToBool		= 0x47,
-	CST_InterfaceToBool		= 0x49,
-	CST_Max					= 0xFF,
+	CST_ObjectToInterface		= 0x00,
+	CST_ObjectToBool			= 0x01,
+	CST_InterfaceToBool			= 0x02,
+	CST_DoubleToFloat			= 0x03,
+	CST_DoubleToFloatArray		= 0x04,
+	CST_DoubleToFloatSet		= 0x05,
+	CST_FloatToDouble			= 0x06,
+	CST_FloatToDoubleArray		= 0x07,
+	CST_FloatToDoubleSet		= 0x08,
+	CST_VectorToVector3f		= 0x09,
+	CST_Vector3fToVector		= 0x0A,
+	CST_FloatToDoubleKeys_Map	= 0x0B,
+	CST_DoubleToFloatKeys_Map	= 0x0C,
+	CST_FloatToDoubleValues_Map	= 0x0D,
+	CST_DoubleToFloatValues_Map	= 0x0E,
+	CST_FloatToDoubleKeys_FloatToDoubleValues_Map	= 0x0F,
+	CST_DoubleToFloatKeys_FloatToDoubleValues_Map	= 0x10,
+	CST_DoubleToFloatKeys_DoubleToFloatValues_Map	= 0x11,
+	CST_FloatToDoubleKeys_DoubleToFloatValues_Map	= 0x12,
+
+	CST_Max						= 0xFF,
 };
 
 // Kinds of text literals
diff --git a/Engine/Source/Runtime/CoreUObject/Public/UObject/ScriptCastingUtils.h b/Engine/Source/Runtime/CoreUObject/Public/UObject/ScriptCastingUtils.h
new file mode 100644
--- /dev/null
+++ b/Engine/Source/Runtime/CoreUObject/Public/UObject/ScriptCastingUtils.h
@@ -0,0 +1,187 @@
+// Copyright Epic Games, Inc. All Rights Reserved.
+
+#pragma once
+
+#include "CoreMinimal.h"
+#include "UObject/Stack.h"
+
+template <typename, typename> struct TAreValidFloatingPointPairs : public std::false_type {};
+template <> struct TAreValidFloatingPointPairs<float, double> : public std::true_type {};
+template <> struct TAreValidFloatingPointPairs<double, float> : public std::true_type {};
+
+template <typename, typename> struct TAreValidVectorPairs : public std::false_type {};
+template <> struct TAreValidVectorPairs<FVector, FVector3f> : public std::true_type {};
+template <> struct TAreValidVectorPairs<FVector3f, FVector> : public std::true_type {};
+
+template <typename T>
+constexpr bool TAlwaysFalse()
+{
+	return false;
+}
+
+using CastFunction = void (*)(const FProperty*, const void*, void*);
+
+FORCEINLINE void ImplicitCast(float Source, double& Destination)
+{
+	Destination = static_cast<double>(Source);
+}
+
+FORCEINLINE void ImplicitCast(double Source, float& Destination)
+{
+	Destination = static_cast<float>(Source);
+}
+
+FORCEINLINE void ImplicitCast(const FVector3f& Source, FVector& Destination)
+{
+	Destination.X = static_cast<double>(Source.X);
+	Destination.Y = static_cast<double>(Source.Y);
+	Destination.Z = static_cast<double>(Source.Z);
+}
+
+FORCEINLINE void ImplicitCast(const FVector& Source, FVector3f& Destination)
+{
+	Destination.X = static_cast<float>(Source.X);
+	Destination.Y = static_cast<float>(Source.Y);
+	Destination.Z = static_cast<float>(Source.Z);
+}
+
+template <typename SourceType, typename DestinationType>
+void FloatingPointCast(const FProperty* /* unused */, const void* SourceRawData, void* DestinationRawData)
+{
+	checkSlow(SourceRawData);
+	checkSlow(DestinationRawData);
+
+	const SourceType* Source = reinterpret_cast<const SourceType*>(SourceRawData);
+	DestinationType* Destination = reinterpret_cast<DestinationType*>(DestinationRawData);
+
+	if constexpr (TAreValidFloatingPointPairs<SourceType, DestinationType>::value || TAreValidVectorPairs<SourceType, DestinationType>::value)
+	{
+		ImplicitCast(*Source, *Destination);
+	}
+	else
+	{
+		static_assert(TAlwaysFalse<SourceType>(), "Unsupported type pairs used for casting!");
+	}
+}
+
+FORCEINLINE void FloatToDoubleCast(const FProperty* /* unused */, const void* SourceRawData, void* DestinationRawData)
+{
+	FloatingPointCast<float, double>(nullptr, SourceRawData, DestinationRawData);
+}
+
+FORCEINLINE void DoubleToFloatCast(const FProperty* /* unused */, const void* SourceRawData, void* DestinationRawData)
+{
+	FloatingPointCast<double, float>(nullptr, SourceRawData, DestinationRawData);
+}
+
+FORCEINLINE void CopyElement(const FProperty* SourceProperty, const void* SourceRawData, void* DestinationRawData)
+{
+	checkSlow(SourceProperty);
+
+	SourceProperty->CopySingleValue(DestinationRawData, SourceRawData);
+}
+
+template <typename SourceType, typename DestinationType>
+void CopyAndCastArray(const FArrayProperty* SourceArrayProperty,
+					  const void* SourceAddress,
+					  const FArrayProperty* DestinationArrayProperty,
+					  void* DestinationAddress)
+{
+	checkSlow(SourceArrayProperty);
+	checkSlow(SourceAddress);
+	checkSlow(DestinationArrayProperty);
+	checkSlow(DestinationAddress);
+
+	FScriptArrayHelper SourceArrayHelper(SourceArrayProperty, SourceAddress);
+	FScriptArrayHelper DestinationArrayHelper(DestinationArrayProperty, DestinationAddress);
+
+	DestinationArrayHelper.Resize(SourceArrayHelper.Num());
+	for (int32 i = 0; i < SourceArrayHelper.Num(); ++i)
+	{
+		const void* SourceRawData = SourceArrayHelper.GetRawPtr(i);
+		void* DestinationRawData = DestinationArrayHelper.GetRawPtr(i);
+
+		FloatingPointCast<SourceType, DestinationType>(nullptr, SourceRawData, DestinationRawData);
+	}
+}
+
+template <typename SourceType, typename DestinationType>
+void CopyAndCastArrayFromStack(FFrame& Stack, RESULT_DECL)
+{
+	FArrayProperty* DestinationArrayProperty = ExactCastField<FArrayProperty>(Stack.MostRecentProperty);
+	void* DestinationAddress = RESULT_PARAM;
+
+	Stack.Step(Stack.Object, nullptr);
+	
+	FArrayProperty* SourceArrayProperty = ExactCastField<FArrayProperty>(Stack.MostRecentProperty);
+	void* SourceAddress = Stack.MostRecentPropertyAddress;
+
+	CopyAndCastArray<SourceType, DestinationType>(SourceArrayProperty, SourceAddress, DestinationArrayProperty, DestinationAddress);
+}
+
+template <typename SourceType, typename DestinationType>
+void CopyAndCastSetFromStack(FFrame& Stack, RESULT_DECL)
+{
+	FSetProperty* DestinationSetProperty = ExactCastField<FSetProperty>(Stack.MostRecentProperty);
+	checkSlow(DestinationSetProperty);
+	checkSlow(RESULT_PARAM);
+	FScriptSetHelper DestinationSetHelper(DestinationSetProperty, RESULT_PARAM);
+
+	Stack.Step(Stack.Object, nullptr);
+
+	checkSlow(Stack.MostRecentProperty);
+	checkSlow(Stack.MostRecentPropertyAddress);
+
+	FSetProperty* SourceSetProperty = ExactCastField<FSetProperty>(Stack.MostRecentProperty);
+	FScriptSetHelper SourceSetHelper(SourceSetProperty, Stack.MostRecentPropertyAddress);
+
+	DestinationSetHelper.EmptyElements(SourceSetHelper.Num());
+	for (int32 i = 0; i < SourceSetHelper.Num(); ++i)
+	{
+		int32 NewIndex = DestinationSetHelper.AddDefaultValue_Invalid_NeedsRehash();
+		void* DestinationRawData = DestinationSetHelper.GetElementPtr(NewIndex);
+		const void* SourceRawData = SourceSetHelper.GetElementPtr(i);
+
+		FloatingPointCast<SourceType, DestinationType>(nullptr, SourceRawData, DestinationRawData);
+	}
+	DestinationSetHelper.Rehash();
+}
+
+template <CastFunction KeyCastFunction, CastFunction ValueCastFunction>
+void CopyAndCastMapFromStack(FFrame& Stack, RESULT_DECL)
+{
+	FMapProperty* DestinationMapProperty = ExactCastField<FMapProperty>(Stack.MostRecentProperty);
+	checkSlow(DestinationMapProperty);
+
+	FScriptMapHelper DestinationMapHelper(DestinationMapProperty, RESULT_PARAM);
+
+	Stack.Step(Stack.Object, nullptr);
+
+	FMapProperty* SourceMapProperty = ExactCastField<FMapProperty>(Stack.MostRecentProperty);
+	checkSlow(SourceMapProperty);
+
+	FScriptMapHelper SourceMapHelper(SourceMapProperty, Stack.MostRecentPropertyAddress);
+
+	const FProperty* SourceKeyProperty = SourceMapProperty->KeyProp;
+	checkSlow(SourceKeyProperty);
+
+	const FProperty* SourceValueProperty = SourceMapProperty->ValueProp;
+	checkSlow(SourceValueProperty);
+
+	DestinationMapHelper.EmptyValues(SourceMapHelper.Num());
+	for (int32 i = 0; i < SourceMapHelper.Num(); ++i)
+	{
+		int32 NewIndex = DestinationMapHelper.AddDefaultValue_Invalid_NeedsRehash();
+
+		const void* SourceKeyRawData = SourceMapHelper.GetKeyPtr(i);
+		void* DestinationKeyRawData = DestinationMapHelper.GetKeyPtr(NewIndex);
+		
+		KeyCastFunction(SourceKeyProperty, SourceKeyRawData, DestinationKeyRawData);
+
+		const void* SourceValueRawData = SourceMapHelper.GetValuePtr(i);
+		void* DestinationValueRawData = DestinationMapHelper.GetValuePtr(NewIndex);
+
+		ValueCastFunction(SourceValueProperty, SourceValueRawData, DestinationValueRawData);
+	}
+	DestinationMapHelper.Rehash();
+}
diff --git a/Engine/Source/Runtime/CoreUObject/Public/UObject/ScriptSerialization.h b/Engine/Source/Runtime/CoreUObject/Public/UObject/ScriptSerialization.h
--- a/Engine/Source/Runtime/CoreUObject/Public/UObject/ScriptSerialization.h
+++ b/Engine/Source/Runtime/CoreUObject/Public/UObject/ScriptSerialization.h
@@ -157,7 +157,7 @@
 
 	switch( Expr )
 	{
-		case EX_PrimitiveCast:
+		case EX_Cast:
 		{
 			// A type conversion.
 			XFER(uint8); //which kind of conversion
@@ -420,6 +420,11 @@
 			}
 			break;
 		}
+		case EX_Vector3fConst:
+		{
+			XFER(float); XFER(float); XFER(float);
+			break;
+		}
 		case EX_TransformConst:
 		{
 			if(Ar.UEVer() >= EUnrealEngineObjectUE5Version::LARGE_WORLD_COORDINATES)
@@ -574,11 +579,11 @@
 			break;
 		}
 		case EX_ArrayGetByRef:
-			{
-				SerializeExpr( iCode, Ar );
-				SerializeExpr( iCode, Ar );
-				break;
-			}
+		{
+			SerializeExpr( iCode, Ar );
+			SerializeExpr( iCode, Ar );
+			break;
+		}
 		default:
 		{
 			// This should never occur.
diff --git a/Engine/Source/Runtime/Engine/Classes/Engine/Blueprint.h b/Engine/Source/Runtime/Engine/Classes/Engine/Blueprint.h
--- a/Engine/Source/Runtime/Engine/Classes/Engine/Blueprint.h
+++ b/Engine/Source/Runtime/Engine/Classes/Engine/Blueprint.h
@@ -856,7 +856,7 @@ public:
 	/**
 	* Allows derived blueprints to require compilation on load, otherwise they may get treated as data only and not compiled on load.
 	*/
-	virtual bool AlwaysCompileOnLoad() const { return false; }
+	virtual bool AlwaysCompileOnLoad() const;
 
 	/**
 	 * Some Blueprints (and classes) can recompile while we are debugging a live session (play in editor).
diff --git a/Engine/Source/Runtime/Engine/Classes/Kismet/KismetMathLibrary.h b/Engine/Source/Runtime/Engine/Classes/Kismet/KismetMathLibrary.h
--- a/Engine/Source/Runtime/Engine/Classes/Kismet/KismetMathLibrary.h
+++ b/Engine/Source/Runtime/Engine/Classes/Kismet/KismetMathLibrary.h
@@ -1331,6 +1331,9 @@ class ENGINE_API UKismetMathLibrary : public UBlueprintFunctionLibrary
 	UFUNCTION(BlueprintPure, Category="Math|Vector", meta=(NativeBreakFunc))
 	static void BreakVector(FVector InVec, float& X, float& Y, float& Z);
 
+	UFUNCTION(BlueprintPure, Category = "Math|Vector", meta = (NativeBreakFunc))
+	static void BreakVector3f(FVector3f InVec, float& X, float& Y, float& Z);
+
 	/** Converts a vector to LinearColor */
 	UFUNCTION(BlueprintPure, meta=(DisplayName = "To LinearColor (Vector)", CompactNodeTitle = "->", ScriptMethod = "LinearColor", Keywords="cast convert", BlueprintAutocast), Category="Math|Conversions")
 	static FLinearColor Conv_VectorToLinearColor(FVector InVec);
@@ -3028,10 +3031,14 @@ class ENGINE_API UKismetMathLibrary : public UBlueprintFunctionLibrary
 	UFUNCTION(BlueprintCallable, meta = (ScriptMethod = "SetRandomHue"), Category = "Math|Color")
 	static void LinearColor_SetRandomHue(UPARAM(ref) FLinearColor& InOutColor);
 
-	/** Convert a float into a LinearColor, where each element is that float */
+	/** Convert a float into a LinearColor, where each element is a float */
 	UFUNCTION(BlueprintPure, meta = (DisplayName = "To LinearColor (Float)", CompactNodeTitle = "->", Keywords = "cast convert", BlueprintAutocast), Category = "Math|Conversions")
 	static FLinearColor Conv_FloatToLinearColor(float InFloat);
 
+	/** Convert a float into a LinearColor, where each element is a double */
+	UFUNCTION(BlueprintPure, meta = (DisplayName = "To LinearColor (Double)", CompactNodeTitle = "->", Keywords = "cast convert", BlueprintAutocast), Category = "Math|Conversions")
+	static FLinearColor Conv_DoubleToLinearColor(double InDouble);
+
 	/** Make a color from individual color components (HSV space; Hue is [0..360) while Saturation and Value are 0..1) */
 	UFUNCTION(BlueprintPure, Category = "Math|Color", meta = (DisplayName = "HSV to RGB"))
 	static FLinearColor HSVToRGB(float H, float S, float V, float A = 1.0f);
@@ -3596,6 +3603,9 @@ class ENGINE_API UKismetMathLibrary : public UBlueprintFunctionLibrary
 	UFUNCTION(BlueprintPure, meta=(DisplayName = "To Vector (Float)", CompactNodeTitle = "->", Keywords="cast convert", BlueprintAutocast), Category="Math|Conversions")
 	static FVector Conv_FloatToVector(float InFloat);
 
+	/** Convert a double into a vector, where each element is that double */
+	UFUNCTION(BlueprintPure, meta = (DisplayName = "To Vector (Double)", CompactNodeTitle = "->", Keywords = "cast convert", BlueprintAutocast), Category = "Math|Conversions")
+	static FVector Conv_DoubleToVector(double InDouble);
 
 	//
 	// Box functions
diff --git a/Engine/Source/Runtime/Engine/Classes/Kismet/KismetMathLibrary.inl b/Engine/Source/Runtime/Engine/Classes/Kismet/KismetMathLibrary.inl
--- a/Engine/Source/Runtime/Engine/Classes/Kismet/KismetMathLibrary.inl
+++ b/Engine/Source/Runtime/Engine/Classes/Kismet/KismetMathLibrary.inl
@@ -1463,6 +1463,14 @@ void UKismetMathLibrary::BreakVector(FVector InVec, float& X, float& Y, float& Z
 	Z = InVec.Z;
 }
 
+KISMET_MATH_FORCEINLINE
+void UKismetMathLibrary::BreakVector3f(FVector3f InVec, float& X, float& Y, float& Z)
+{
+	X = InVec.X;
+	Y = InVec.Y;
+	Z = InVec.Z;
+}
+
 KISMET_MATH_FORCEINLINE
 FRotator UKismetMathLibrary::Conv_VectorToRotator(FVector InVec)
 {
@@ -2795,12 +2803,24 @@ FVector UKismetMathLibrary::Conv_FloatToVector(float InFloat)
 	return FVector(InFloat);
 }
 
+KISMET_MATH_FORCEINLINE
+FVector UKismetMathLibrary::Conv_DoubleToVector(double InDouble)
+{
+	return FVector(InDouble);
+}
+
 KISMET_MATH_FORCEINLINE
 FLinearColor UKismetMathLibrary::Conv_FloatToLinearColor(float InFloat)
 {
 	return FLinearColor(InFloat, InFloat, InFloat);
 }
 
+KISMET_MATH_FORCEINLINE
+FLinearColor UKismetMathLibrary::Conv_DoubleToLinearColor(double InDouble)
+{
+	return FLinearColor(InDouble, InDouble, InDouble);
+}
+
 KISMET_MATH_FORCEINLINE
 FBox UKismetMathLibrary::MakeBox(FVector Min, FVector Max)
 {
diff --git a/Engine/Source/Runtime/Engine/Classes/Kismet/KismetStringLibrary.h b/Engine/Source/Runtime/Engine/Classes/Kismet/KismetStringLibrary.h
--- a/Engine/Source/Runtime/Engine/Classes/Kismet/KismetStringLibrary.h
+++ b/Engine/Source/Runtime/Engine/Classes/Kismet/KismetStringLibrary.h
@@ -40,6 +40,10 @@ class ENGINE_API UKismetStringLibrary : public UBlueprintFunctionLibrary
 	UFUNCTION(BlueprintPure, meta=(DisplayName = "To String (Vector)", CompactNodeTitle = "->", BlueprintAutocast), Category="Utilities|String")
 	static FString Conv_VectorToString(FVector InVec);
 
+	/** Converts a float vector value to a string, in the form 'X= Y= Z=' */
+	UFUNCTION(BlueprintPure, meta = (DisplayName = "To String (Float Vector)", CompactNodeTitle = "->", BlueprintAutocast), Category = "Utilities|String")
+	static FString Conv_Vector3fToString(FVector3f InVec);
+
 	/** Converts an IntVector value to a string, in the form 'X= Y= Z=' */
 	UFUNCTION(BlueprintPure, meta=(DisplayName = "To String (IntVector)", CompactNodeTitle = "->", BlueprintAutocast), Category="Utilities|String")
 	static FString Conv_IntVectorToString(FIntVector InIntVec);
@@ -92,6 +96,10 @@ class ENGINE_API UKismetStringLibrary : public UBlueprintFunctionLibrary
 	UFUNCTION(BlueprintPure, meta=(DisplayName = "String To Vector", CompactNodeTitle = "->"), Category="Utilities|String")
 	static void Conv_StringToVector(const FString& InString, FVector& OutConvertedVector, bool& OutIsValid);
 
+	/** Convert String Back To Float Vector. IsValid indicates whether or not the string could be successfully converted. */
+	UFUNCTION(BlueprintPure, meta = (DisplayName = "String To Float Vector", CompactNodeTitle = "->"), Category = "Utilities|String")
+	static void Conv_StringToVector3f(const FString& InString, FVector3f& OutConvertedVector, bool& OutIsValid);
+
 	/** Convert String Back To Vector2D. IsValid indicates whether or not the string could be successfully converted. */
 	UFUNCTION(BlueprintPure, meta=(DisplayName = "String To Vector2D", CompactNodeTitle = "->"), Category="Utilities|String")
 	static void Conv_StringToVector2D(const FString& InString, FVector2D& OutConvertedVector2D, bool& OutIsValid);
diff --git a/Engine/Source/Runtime/Engine/Classes/Kismet/KismetSystemLibrary.h b/Engine/Source/Runtime/Engine/Classes/Kismet/KismetSystemLibrary.h
--- a/Engine/Source/Runtime/Engine/Classes/Kismet/KismetSystemLibrary.h
+++ b/Engine/Source/Runtime/Engine/Classes/Kismet/KismetSystemLibrary.h
@@ -876,6 +876,10 @@ class ENGINE_API UKismetSystemLibrary : public UBlueprintFunctionLibrary
 	UFUNCTION(BlueprintCallable, meta=(BlueprintInternalUseOnly = "true", AutoCreateRefTerm = "Value" ))
 	static void SetVectorPropertyByName(UObject* Object, FName PropertyName, const FVector& Value);
 
+	/** Set a VECTOR3F property by name */
+	UFUNCTION(BlueprintCallable, meta = (BlueprintInternalUseOnly = "true", AutoCreateRefTerm = "Value"))
+	static void SetVector3fPropertyByName(UObject* Object, FName PropertyName, const FVector3f& Value);
+
 	/** Set a ROTATOR property by name */
 	UFUNCTION(BlueprintCallable, meta=(BlueprintInternalUseOnly = "true", AutoCreateRefTerm = "Value" ))
 	static void SetRotatorPropertyByName(UObject* Object, FName PropertyName, const FRotator& Value);
diff --git a/Engine/Source/Runtime/Engine/Private/AnimBlueprint.cpp b/Engine/Source/Runtime/Engine/Private/AnimBlueprint.cpp
--- a/Engine/Source/Runtime/Engine/Private/AnimBlueprint.cpp
+++ b/Engine/Source/Runtime/Engine/Private/AnimBlueprint.cpp
@@ -190,13 +190,6 @@ void UAnimBlueprint::PostLoad()
 		});
 	}
 #endif
-
-#if WITH_EDITORONLY_DATA
-	if(GetLinkerCustomVersion(FFrameworkObjectVersion::GUID) < FFrameworkObjectVersion::AnimBlueprintSubgraphFix)
-	{
-		AnimationEditorUtils::RegenerateSubGraphArrays(this);
-	}
-#endif
 }
 
 bool UAnimBlueprint::CanAlwaysRecompileWhilePlayingInEditor() const
diff --git a/Engine/Source/Runtime/Engine/Private/Blueprint.cpp b/Engine/Source/Runtime/Engine/Private/Blueprint.cpp
--- a/Engine/Source/Runtime/Engine/Private/Blueprint.cpp
+++ b/Engine/Source/Runtime/Engine/Private/Blueprint.cpp
@@ -430,6 +430,10 @@ void UBlueprint::Serialize(FArchive& Ar)
 	Super::Serialize(Ar);
 
 #if WITH_EDITORONLY_DATA
+#if ENABLE_BLUEPRINT_REAL_NUMBERS
+	Ar.UsingCustomVersion(FUE5ReleaseStreamObjectVersion::GUID);
+#endif
+
 	if(Ar.IsLoading() && Ar.UEVer() < VER_UE4_BLUEPRINT_VARS_NOT_READ_ONLY)
 	{
 		// Allow all blueprint defined vars to be read/write.  undoes previous convention of making exposed variables read-only
@@ -874,6 +878,18 @@ void UBlueprint::GetReparentingRules(TSet< const UClass* >& AllowedChildrenOfCla
 
 }
 
+bool UBlueprint::AlwaysCompileOnLoad() const
+{
+	bool bShouldCompile = false;
+
+#if ENABLE_BLUEPRINT_REAL_NUMBERS
+	bShouldCompile = 
+		GetLinkerCustomVersion(FUE5ReleaseStreamObjectVersion::GUID) <  FUE5ReleaseStreamObjectVersion::BlueprintPinsUseRealNumbers;
+#endif
+
+	return bShouldCompile;
+}
+
 bool UBlueprint::CanAlwaysRecompileWhilePlayingInEditor() const
 {
 	return false;
diff --git a/Engine/Source/Runtime/Engine/Private/EdGraph/EdGraphNode.cpp b/Engine/Source/Runtime/Engine/Private/EdGraph/EdGraphNode.cpp
--- a/Engine/Source/Runtime/Engine/Private/EdGraph/EdGraphNode.cpp
+++ b/Engine/Source/Runtime/Engine/Private/EdGraph/EdGraphNode.cpp
@@ -7,6 +7,7 @@
 #include "UObject/FrameworkObjectVersion.h"
 #include "UObject/ObjectSaveContext.h"
 #include "UObject/ReleaseObjectVersion.h"
+#include "UObject/UE5ReleaseStreamObjectVersion.h"
 #include "EdGraph/EdGraphPin.h"
 #include "Textures/SlateIcon.h"
 #include "EdGraph/EdGraph.h"
@@ -39,6 +40,10 @@ FArchive& operator<<(FArchive& Ar, FEdGraphTerminalType& T)
 	Ar.UsingCustomVersion(FFrameworkObjectVersion::GUID);
 	Ar.UsingCustomVersion(FReleaseObjectVersion::GUID);
 
+#if ENABLE_BLUEPRINT_REAL_NUMBERS
+	Ar.UsingCustomVersion(FUE5ReleaseStreamObjectVersion::GUID);
+#endif
+
 	if (Ar.CustomVer(FFrameworkObjectVersion::GUID) >= FFrameworkObjectVersion::PinsStoreFName)
 	{
 		Ar << T.TerminalCategory;
@@ -89,6 +94,21 @@ FArchive& operator<<(FArchive& Ar, FEdGraphTerminalType& T)
 		Ar << T.bTerminalIsUObjectWrapper;
 	}
 
+#if ENABLE_BLUEPRINT_REAL_NUMBERS
+	if (Ar.IsLoading())
+	{
+		bool bFixupPinCategories =
+			(Ar.CustomVer(FUE5ReleaseStreamObjectVersion::GUID) < FUE5ReleaseStreamObjectVersion::BlueprintPinsUseRealNumbers) &&
+			((T.TerminalCategory == TEXT("double")) || (T.TerminalCategory == TEXT("float")));
+
+		if (bFixupPinCategories)
+		{
+			T.TerminalCategory = TEXT("real");
+			T.TerminalSubCategory = TEXT("double");
+		}
+	}
+#endif
+
 	return Ar;
 }
 
diff --git a/Engine/Source/Runtime/Engine/Private/EdGraph/EdGraphPin.cpp b/Engine/Source/Runtime/Engine/Private/EdGraph/EdGraphPin.cpp
--- a/Engine/Source/Runtime/Engine/Private/EdGraph/EdGraphPin.cpp
+++ b/Engine/Source/Runtime/Engine/Private/EdGraph/EdGraphPin.cpp
@@ -6,6 +6,7 @@
 #include "UObject/FrameworkObjectVersion.h"
 #include "UObject/ReleaseObjectVersion.h"
 #include "UObject/UE5MainStreamObjectVersion.h"
+#include "UObject/UE5ReleaseStreamObjectVersion.h"
 #include "UObject/UnrealType.h"
 #include "UObject/TextProperty.h"
 #include "EdGraph/EdGraph.h"
@@ -161,6 +162,10 @@ bool FEdGraphPinType::Serialize(FArchive& Ar)
 	Ar.UsingCustomVersion(FFrameworkObjectVersion::GUID);
 	Ar.UsingCustomVersion(FReleaseObjectVersion::GUID);
 
+#if ENABLE_BLUEPRINT_REAL_NUMBERS
+	Ar.UsingCustomVersion(FUE5ReleaseStreamObjectVersion::GUID);
+#endif
+
 	if (Ar.CustomVer(FFrameworkObjectVersion::GUID) >= FFrameworkObjectVersion::PinsStoreFName)
 	{
 		Ar << PinCategory;
@@ -297,6 +302,19 @@ bool FEdGraphPinType::Serialize(FArchive& Ar)
 				bIsUObjectWrapperBool = false;
 			}
 		}
+
+#if ENABLE_BLUEPRINT_REAL_NUMBERS
+		bool bFixupPinCategories =
+			(Ar.CustomVer(FUE5ReleaseStreamObjectVersion::GUID) < FUE5ReleaseStreamObjectVersion::BlueprintPinsUseRealNumbers) &&
+			((PinCategory == TEXT("double")) || (PinCategory == TEXT("float")));
+
+		if (bFixupPinCategories)
+		{
+			PinCategory = TEXT("real");
+			PinSubCategory = TEXT("double");
+		}
+#endif
+
 #endif
 
 		bIsUObjectWrapper = bIsUObjectWrapperBool;
diff --git a/Engine/Source/Runtime/Engine/Private/KismetStringLibrary.cpp b/Engine/Source/Runtime/Engine/Private/KismetStringLibrary.cpp
--- a/Engine/Source/Runtime/Engine/Private/KismetStringLibrary.cpp
+++ b/Engine/Source/Runtime/Engine/Private/KismetStringLibrary.cpp
@@ -86,6 +86,11 @@ FString UKismetStringLibrary::Conv_VectorToString(FVector InVec)
 	return InVec.ToString();	
 }
 
+FString UKismetStringLibrary::Conv_Vector3fToString(FVector3f InVec)
+{
+	return InVec.ToString();
+}
+
 FString UKismetStringLibrary::Conv_IntVectorToString(FIntVector InIntVec)
 {
 	return InIntVec.ToString();
@@ -151,6 +156,11 @@ void UKismetStringLibrary::Conv_StringToVector(const FString& InString, FVector&
 	OutIsValid = OutConvertedVector.InitFromString(InString);
 }
 
+void UKismetStringLibrary::Conv_StringToVector3f(const FString& InString, FVector3f& OutConvertedVector, bool& OutIsValid)
+{
+	OutIsValid = OutConvertedVector.InitFromString(InString);
+}
+
 void UKismetStringLibrary::Conv_StringToVector2D(const FString& InString, FVector2D& OutConvertedVector2D, bool& OutIsValid)
 {
 	OutIsValid = OutConvertedVector2D.InitFromString(InString);
diff --git a/Engine/Source/Runtime/Engine/Private/KismetSystemLibrary.cpp b/Engine/Source/Runtime/Engine/Private/KismetSystemLibrary.cpp
--- a/Engine/Source/Runtime/Engine/Private/KismetSystemLibrary.cpp
+++ b/Engine/Source/Runtime/Engine/Private/KismetSystemLibrary.cpp
@@ -1222,10 +1222,10 @@ TSoftClassPtr<UObject> UKismetSystemLibrary::Conv_ClassToSoftClassReference(cons
 
 void UKismetSystemLibrary::SetTextPropertyByName(UObject* Object, FName PropertyName, const FText& Value)
 {
-	if(Object != NULL)
+	if(Object != nullptr)
 	{
 		FTextProperty* TextProp = FindFProperty<FTextProperty>(Object->GetClass(), PropertyName);
-		if(TextProp != NULL)
+		if(TextProp != nullptr)
 		{
 			TextProp->SetPropertyValue_InContainer(Object, Value);
 		}		
@@ -1234,24 +1234,37 @@ void UKismetSystemLibrary::SetTextPropertyByName(UObject* Object, FName Property
 
 void UKismetSystemLibrary::SetVectorPropertyByName(UObject* Object, FName PropertyName, const FVector& Value)
 {
-	if(Object != NULL)
+	if(Object != nullptr)
 	{
 		UScriptStruct* VectorStruct = TBaseStructure<FVector>::Get();
 		FStructProperty* VectorProp = FindFProperty<FStructProperty>(Object->GetClass(), PropertyName);
-		if(VectorProp != NULL && VectorProp->Struct == VectorStruct)
+		if(VectorProp != nullptr && VectorProp->Struct == VectorStruct)
 		{
 			*VectorProp->ContainerPtrToValuePtr<FVector>(Object) = Value;
 		}		
 	}
 }
 
+void UKismetSystemLibrary::SetVector3fPropertyByName(UObject* Object, FName PropertyName, const FVector3f& Value)
+{
+	if (Object != nullptr)
+	{
+		UScriptStruct* Vector3fStruct = TVariantStructure<FVector3f>::Get();
+		FStructProperty* Vector3fProp = FindFProperty<FStructProperty>(Object->GetClass(), PropertyName);
+		if (Vector3fProp != nullptr && Vector3fProp->Struct == Vector3fStruct)
+		{
+			*Vector3fProp->ContainerPtrToValuePtr<FVector3f>(Object) = Value;
+		}
+	}
+}
+
 void UKismetSystemLibrary::SetRotatorPropertyByName(UObject* Object, FName PropertyName, const FRotator& Value)
 {
-	if(Object != NULL)
+	if(Object != nullptr)
 	{
 		UScriptStruct* RotatorStruct = TBaseStructure<FRotator>::Get();
 		FStructProperty* RotatorProp = FindFProperty<FStructProperty>(Object->GetClass(), PropertyName);
-		if(RotatorProp != NULL && RotatorProp->Struct == RotatorStruct)
+		if(RotatorProp != nullptr && RotatorProp->Struct == RotatorStruct)
 		{
 			*RotatorProp->ContainerPtrToValuePtr<FRotator>(Object) = Value;
 		}		
@@ -1286,11 +1299,11 @@ void UKismetSystemLibrary::SetColorPropertyByName(UObject* Object, FName Propert
 
 void UKismetSystemLibrary::SetTransformPropertyByName(UObject* Object, FName PropertyName, const FTransform& Value)
 {
-	if(Object != NULL)
+	if(Object != nullptr)
 	{
 		UScriptStruct* TransformStruct = TBaseStructure<FTransform>::Get();
 		FStructProperty* TransformProp = FindFProperty<FStructProperty>(Object->GetClass(), PropertyName);
-		if(TransformProp != NULL && TransformProp->Struct == TransformStruct)
+		if(TransformProp != nullptr && TransformProp->Struct == TransformStruct)
 		{
 			*TransformProp->ContainerPtrToValuePtr<FTransform>(Object) = Value;
 		}		
@@ -1305,10 +1318,10 @@ void UKismetSystemLibrary::SetCollisionProfileNameProperty(UObject* Object, FNam
 
 void UKismetSystemLibrary::Generic_SetStructurePropertyByName(UObject* OwnerObject, FName StructPropertyName, const void* SrcStructAddr)
 {
-	if (OwnerObject != NULL)
+	if (OwnerObject != nullptr)
 	{
 		FStructProperty* StructProp = FindFProperty<FStructProperty>(OwnerObject->GetClass(), StructPropertyName);
-		if (StructProp != NULL)
+		if (StructProp != nullptr)
 		{
 			void* Dest = StructProp->ContainerPtrToValuePtr<void>(OwnerObject);
 			StructProp->CopyValuesInternal(Dest, SrcStructAddr, 1);
diff --git a/Engine/Source/Runtime/Engine/Private/PropertyAccess.cpp b/Engine/Source/Runtime/Engine/Private/PropertyAccess.cpp
--- a/Engine/Source/Runtime/Engine/Private/PropertyAccess.cpp
+++ b/Engine/Source/Runtime/Engine/Private/PropertyAccess.cpp
@@ -2,6 +2,7 @@
 
 #include "PropertyAccess.h"
 #include "Misc/MemStack.h"
+#include "UObject/ScriptCastingUtils.h"
 
 #define LOCTEXT_NAMESPACE "PropertyAccess"
 
@@ -370,6 +371,24 @@ struct FPropertyAccessSystem
 			checkSlow(InDestProperty->IsA<FFloatProperty>());
 			static_cast<const FFloatProperty*>(InDestProperty)->SetPropertyValue(InDestAddr, (float)static_cast<const FDoubleProperty*>(InSrcProperty)->GetPropertyValue(InSrcAddr));
 			break;
+		case EPropertyAccessCopyType::PromoteArrayFloatToDouble:
+		{
+			checkSlow(InSrcProperty->IsA<FArrayProperty>());
+			checkSlow(InDestProperty->IsA<FArrayProperty>());
+			const FArrayProperty* SrcArrayProperty = ExactCastField<const FArrayProperty>(InSrcProperty);
+			const FArrayProperty* DestArrayProperty = ExactCastField<const FArrayProperty>(InDestProperty);
+			CopyAndCastArray<float, double>(SrcArrayProperty, InSrcAddr, DestArrayProperty, InDestAddr);
+			break;
+		}
+		case EPropertyAccessCopyType::DemoteArrayDoubleToFloat:
+		{
+			checkSlow(InSrcProperty->IsA<FArrayProperty>());
+			checkSlow(InDestProperty->IsA<FArrayProperty>());
+			const FArrayProperty* SrcArrayProperty = ExactCastField<const FArrayProperty>(InSrcProperty);
+			const FArrayProperty* DestArrayProperty = ExactCastField<const FArrayProperty>(InDestProperty);
+			CopyAndCastArray<double, float>(SrcArrayProperty, InSrcAddr, DestArrayProperty, InDestAddr);
+			break;
+		}
 		default:
 			check(false);
 			break;
diff --git a/Engine/Source/Runtime/Engine/Public/PropertyAccess.h b/Engine/Source/Runtime/Engine/Public/PropertyAccess.h
--- a/Engine/Source/Runtime/Engine/Public/PropertyAccess.h
+++ b/Engine/Source/Runtime/Engine/Public/PropertyAccess.h
@@ -342,6 +342,9 @@ enum class EPropertyAccessCopyType : uint8
 	// Float promotions		// LWC_TODO: Float/double should become synonyms?
 	PromoteFloatToDouble,
 	DemoteDoubleToFloat,	// LWC_TODO: This should not ship!
+
+	PromoteArrayFloatToDouble,
+	DemoteArrayDoubleToFloat,
 };
 
 // A property copy, represents a one-to-many copy operation
diff --git a/Engine/Source/Runtime/RigVM/Private/RigVMCore/RigVM.cpp b/Engine/Source/Runtime/RigVM/Private/RigVMCore/RigVM.cpp
--- a/Engine/Source/Runtime/RigVM/Private/RigVMCore/RigVM.cpp
+++ b/Engine/Source/Runtime/RigVM/Private/RigVMCore/RigVM.cpp
@@ -3467,27 +3467,36 @@ FString URigVM::GetOperandLabel(const FRigVMOperand& InOperand, TFunction<FStrin
 
 #else
 
-	if(InOperand.GetMemoryType() == ERigVMMemoryType::External)
+	FString RegisterName;
+	FString RegisterOffsetName;
+	if (InOperand.GetMemoryType() == ERigVMMemoryType::External)
 	{
-		if(ExternalVariables.IsValidIndex(InOperand.GetRegisterIndex()))
+		const FRigVMExternalVariable& ExternalVariable = ExternalVariables[InOperand.GetRegisterIndex()];
+		RegisterName = FString::Printf(TEXT("Variable::%s"), *ExternalVariable.Name.ToString());
+		if (InOperand.GetRegisterOffset() != INDEX_NONE)
 		{
-			return ExternalVariables[InOperand.GetRegisterIndex()].Name.ToString();
+			if(ensure(ExternalPropertyPaths.IsValidIndex(InOperand.GetRegisterOffset())))
+			{
+				RegisterOffsetName = ExternalPropertyPaths[InOperand.GetRegisterOffset()].ToString();
+			}
 		}
 	}
-	
-	URigVMMemoryStorage* Memory = GetMemoryByType(InOperand.GetMemoryType());
-	if(Memory == nullptr)
+	else
 	{
-		return FString();
-	}
+		URigVMMemoryStorage* Memory = GetMemoryByType(InOperand.GetMemoryType());
+		if(Memory == nullptr)
+		{
+			return FString();
+		}
 
-	check(Memory->IsValidIndex(InOperand.GetRegisterIndex()));
-	
-	const FString RegisterName = Memory->GetProperties()[InOperand.GetRegisterIndex()]->GetName();
-	const FString RegisterOffsetName =
-		InOperand.GetRegisterOffset() != INDEX_NONE ?
-		Memory->GetPropertyPaths()[InOperand.GetRegisterOffset()].ToString() :
-		FString();
+		check(Memory->IsValidIndex(InOperand.GetRegisterIndex()));
+		
+		RegisterName = Memory->GetProperties()[InOperand.GetRegisterIndex()]->GetName();
+		RegisterOffsetName =
+			InOperand.GetRegisterOffset() != INDEX_NONE ?
+			Memory->GetPropertyPaths()[InOperand.GetRegisterOffset()].ToString() :
+			FString();
+	}
 	
 	FString OperandLabel = RegisterName;
 	
diff --git a/Engine/Source/Runtime/RigVM/Public/RigVMCore/RigVMExternalVariable.h b/Engine/Source/Runtime/RigVM/Public/RigVMCore/RigVMExternalVariable.h
--- a/Engine/Source/Runtime/RigVM/Public/RigVMCore/RigVMExternalVariable.h
+++ b/Engine/Source/Runtime/RigVM/Public/RigVMCore/RigVMExternalVariable.h
@@ -85,6 +85,10 @@ struct RIGVM_API FRigVMExternalVariable
 				*ObjectProperty->PropertyClass->GetName());
 			OutTypeObject = ObjectProperty->PropertyClass;
 		}
+		else
+		{
+			checkNoEntry();
+		}
 	}
 
 	static uint32 GetPropertyTypeHash(const FProperty *InProperty, const uint8 *InMemory)
