UE-140327 - Remove ENABLE_BLUEPRINT_REAL_NUMBERS usage

The original ENABLE_BLUEPRINT_REAL_NUMBERS macro was strictly meant for AB testing in the Dev-LWC. Currently, disabling it won't work, and would likely lead to broken behavior. Its presence might incorrectly indicate to licensees that this is an option that can be toggled safely.

diff --git a/Engine/Plugins/Animation/ControlRig/Source/ControlRig/Private/Rigs/RigHierarchy.cpp b/Engine/Plugins/Animation/ControlRig/Source/ControlRig/Private/Rigs/RigHierarchy.cpp
--- a/Engine/Plugins/Animation/ControlRig/Source/ControlRig/Private/Rigs/RigHierarchy.cpp
+++ b/Engine/Plugins/Animation/ControlRig/Source/ControlRig/Private/Rigs/RigHierarchy.cpp
@@ -847,10 +847,7 @@ FEdGraphPinType URigHierarchy::GetControlPinType(FRigControlElement* InControlEl
 	static const FName PC_Float(TEXT("float"));
 	static const FName PC_Int(TEXT("int"));
 	static const FName PC_Struct(TEXT("struct"));
-
-#if ENABLE_BLUEPRINT_REAL_NUMBERS
 	static const FName PC_Real(TEXT("real"));
-#endif
 
 	FEdGraphPinType PinType;
 
@@ -861,20 +858,12 @@ FEdGraphPinType URigHierarchy::GetControlPinType(FRigControlElement* InControlEl
 			PinType.PinCategory = PC_Boolean;
 			break;
 		}
-#if ENABLE_BLUEPRINT_REAL_NUMBERS
 		case ERigControlType::Float:
 		{
 			PinType.PinCategory = PC_Real;
 			PinType.PinSubCategory = PC_Float;
 			break;
 		}
-#else
-		case ERigControlType::Float:
-		{
-			PinType.PinCategory = PC_Float;
-			break;
-		}
-#endif
 		case ERigControlType::Integer:
 		{
 			PinType.PinCategory = PC_Int;
diff --git a/Engine/Plugins/Animation/ControlRig/Source/ControlRigDeveloper/Private/ControlRigBlueprint.cpp b/Engine/Plugins/Animation/ControlRig/Source/ControlRigDeveloper/Private/ControlRigBlueprint.cpp
--- a/Engine/Plugins/Animation/ControlRig/Source/ControlRigDeveloper/Private/ControlRigBlueprint.cpp
+++ b/Engine/Plugins/Animation/ControlRig/Source/ControlRigDeveloper/Private/ControlRigBlueprint.cpp
@@ -2585,7 +2585,7 @@ void UControlRigBlueprint::PopulateModelFromGraphForBackwardsCompatibility(UCont
 					}
 
 					FName DataType = PinType.PinCategory;
-#if ENABLE_BLUEPRINT_REAL_NUMBERS
+
 					if (PinType.PinCategory == UEdGraphSchema_K2::PC_Real)
 					{
 						if (PinType.PinSubCategory == UEdGraphSchema_K2::PC_Float)
@@ -2601,7 +2601,7 @@ void UControlRigBlueprint::PopulateModelFromGraphForBackwardsCompatibility(UCont
 							ensure(false);
 						}
 					}
-#endif
+
 					UObject* DataTypeObject = nullptr;
 					if (DataType == NAME_None)
 					{
diff --git a/Engine/Plugins/Animation/ControlRig/Source/ControlRigDeveloper/Private/Graph/ControlRigGraphNode.cpp b/Engine/Plugins/Animation/ControlRig/Source/ControlRigDeveloper/Private/Graph/ControlRigGraphNode.cpp
--- a/Engine/Plugins/Animation/ControlRig/Source/ControlRigDeveloper/Private/Graph/ControlRigGraphNode.cpp
+++ b/Engine/Plugins/Animation/ControlRig/Source/ControlRigDeveloper/Private/Graph/ControlRigGraphNode.cpp
@@ -1164,7 +1164,6 @@ FEdGraphPinType UControlRigGraphNode::GetPinTypeForModelPin(URigVMPin* InModelPi
 	{
 		PinType.PinCategory = UEdGraphSchema_K2::PC_Int;
 	}
-#if ENABLE_BLUEPRINT_REAL_NUMBERS
 	else if (ModelPinCPPType == TEXT("float"))
 	{
 		PinType.PinCategory = UEdGraphSchema_K2::PC_Real;
@@ -1175,16 +1174,6 @@ FEdGraphPinType UControlRigGraphNode::GetPinTypeForModelPin(URigVMPin* InModelPi
 		PinType.PinCategory = UEdGraphSchema_K2::PC_Real;
 		PinType.PinSubCategory = UEdGraphSchema_K2::PC_Double;
 	}
-#else
-	else if (ModelPinCPPType == TEXT("float"))
-	{
-		PinType.PinCategory = UEdGraphSchema_K2::PC_Float;
-	}
-	else if (ModelPinCPPType == TEXT("double"))
-	{
-		PinType.PinCategory = UEdGraphSchema_K2::PC_Double;
-	}
-#endif
 	else if (ModelPinCPPType == TEXT("FName"))
 	{
 		PinType.PinCategory = UEdGraphSchema_K2::PC_Name;
diff --git a/Engine/Plugins/Animation/ControlRig/Source/ControlRigDeveloper/Private/Graph/ControlRigGraphSchema.cpp b/Engine/Plugins/Animation/ControlRig/Source/ControlRigDeveloper/Private/Graph/ControlRigGraphSchema.cpp
--- a/Engine/Plugins/Animation/ControlRig/Source/ControlRigDeveloper/Private/Graph/ControlRigGraphSchema.cpp
+++ b/Engine/Plugins/Animation/ControlRig/Source/ControlRigDeveloper/Private/Graph/ControlRigGraphSchema.cpp
@@ -756,15 +756,6 @@ FLinearColor UControlRigGraphSchema::GetPinTypeColor(const FEdGraphPinType& PinT
 			}
 		}
 	}
-
-#if !ENABLE_BLUEPRINT_REAL_NUMBERS
-	if(TypeName == UEdGraphSchema_K2::PC_Double)
-	{
-		FEdGraphPinType LocalPinType = PinType;
-		LocalPinType.PinCategory = UEdGraphSchema_K2::PC_Float;
-		return GetPinTypeColor(LocalPinType);
-	}
-#endif
 	
 	return GetDefault<UEdGraphSchema_K2>()->GetPinTypeColor(PinType);
 }
@@ -886,12 +877,7 @@ bool UControlRigGraphSchema::SupportsPinType(TWeakPtr<const FEdGraphSchemaAction
 
 	if (TypeName == UEdGraphSchema_K2::PC_Boolean ||
 		TypeName == UEdGraphSchema_K2::PC_Int ||
-#if ENABLE_BLUEPRINT_REAL_NUMBERS
 		TypeName == UEdGraphSchema_K2::PC_Real ||
-#else
-		TypeName == UEdGraphSchema_K2::PC_Float ||
-		TypeName == UEdGraphSchema_K2::PC_Double ||
-#endif
 		TypeName == UEdGraphSchema_K2::PC_Name ||
 		TypeName == UEdGraphSchema_K2::PC_String ||
 		TypeName == UEdGraphSchema_K2::PC_Enum)
@@ -1667,12 +1653,7 @@ void UControlRigGraphSchema::ResetPinDefaultsRecursive(UEdGraphPin* InPin) const
 void UControlRigGraphSchema::GetVariablePinTypes(TArray<FEdGraphPinType>& PinTypes) const
 {
 	PinTypes.Add(FEdGraphPinType(UEdGraphSchema_K2::PC_Boolean, FName(NAME_None), nullptr, EPinContainerType::None, false, FEdGraphTerminalType()));
-#if ENABLE_BLUEPRINT_REAL_NUMBERS
 	PinTypes.Add(FEdGraphPinType(UEdGraphSchema_K2::PC_Real, FName(NAME_None), nullptr, EPinContainerType::None, false, FEdGraphTerminalType()));
-#else
-	PinTypes.Add(FEdGraphPinType(UEdGraphSchema_K2::PC_Float, FName(NAME_None), nullptr, EPinContainerType::None, false, FEdGraphTerminalType()));
-	PinTypes.Add(FEdGraphPinType(UEdGraphSchema_K2::PC_Double, FName(NAME_None), nullptr, EPinContainerType::None, false, FEdGraphTerminalType()));
-#endif
 	PinTypes.Add(FEdGraphPinType(UEdGraphSchema_K2::PC_Int, FName(NAME_None), nullptr, EPinContainerType::None, false, FEdGraphTerminalType()));
 	PinTypes.Add(FEdGraphPinType(UEdGraphSchema_K2::PC_Struct, FName(NAME_None), TBaseStructure<FVector>::Get(), EPinContainerType::None, false, FEdGraphTerminalType()));
 	PinTypes.Add(FEdGraphPinType(UEdGraphSchema_K2::PC_Struct, FName(NAME_None), TBaseStructure<FVector2D>::Get(), EPinContainerType::None, false, FEdGraphTerminalType()));
diff --git a/Engine/Plugins/Animation/IKRig/Source/IKRigDeveloper/Private/AnimGraphNode_IKRig.cpp b/Engine/Plugins/Animation/IKRig/Source/IKRigDeveloper/Private/AnimGraphNode_IKRig.cpp
--- a/Engine/Plugins/Animation/IKRig/Source/IKRigDeveloper/Private/AnimGraphNode_IKRig.cpp
+++ b/Engine/Plugins/Animation/IKRig/Source/IKRigDeveloper/Private/AnimGraphNode_IKRig.cpp
@@ -570,12 +570,8 @@ void UAnimGraphNode_IKRig::CreateCustomPinsFromValidAsset()
 	static const FName PC_Float(TEXT("float"));
 	static const FName PC_Double(TEXT("double"));
 	static FEdGraphPinType AlphaPinType;
-#if ENABLE_BLUEPRINT_REAL_NUMBERS
 	AlphaPinType.PinCategory = PC_Real;
 	AlphaPinType.PinSubCategory = PC_Double;
-#else
-	AlphaPinType.PinCategory = PC_Float;
-#endif
 
 	// create pins
 	const TArray<UIKRigEffectorGoal*>& AssetGoals = Node.RigDefinitionAsset->GetGoalArray();
diff --git a/Engine/Plugins/Animation/LiveLink/Source/LiveLinkGraphNode/Private/K2Node_EvaluateLiveLinkCustom.cpp b/Engine/Plugins/Animation/LiveLink/Source/LiveLinkGraphNode/Private/K2Node_EvaluateLiveLinkCustom.cpp
--- a/Engine/Plugins/Animation/LiveLink/Source/LiveLinkGraphNode/Private/K2Node_EvaluateLiveLinkCustom.cpp
+++ b/Engine/Plugins/Animation/LiveLink/Source/LiveLinkGraphNode/Private/K2Node_EvaluateLiveLinkCustom.cpp
@@ -36,11 +36,7 @@ void UK2Node_EvaluateLiveLinkFrameAtWorldTime::AllocateDefaultPins()
 
 	const UEdGraphSchema_K2* K2Schema = GetDefault<UEdGraphSchema_K2>();
 
-#if ENABLE_BLUEPRINT_REAL_NUMBERS
 	UEdGraphPin* LiveLinkWorldTimePin = CreatePin(EGPD_Input, UEdGraphSchema_K2::PC_Real, UEdGraphSchema_K2::PC_Double, UK2Node_EvaluateLiveLinkFrameHelper::LiveLinkWorldTimePinName);
-#else
-	UEdGraphPin* LiveLinkWorldTimePin = CreatePin(EGPD_Input, UEdGraphSchema_K2::PC_Float, UK2Node_EvaluateLiveLinkFrameHelper::LiveLinkWorldTimePinName);
-#endif
 	LiveLinkWorldTimePin->PinFriendlyName = LOCTEXT("LiveLinkWorldTimePin", "World Time");
 	SetPinToolTip(*LiveLinkWorldTimePin, LOCTEXT("LiveLinkWorldTimePinDescription", "The World Time the subject will be evaluated to"));
 }
diff --git a/Engine/Plugins/Experimental/EnhancedInput/Source/InputBlueprintNodes/Private/K2Node_EnhancedInputAction.cpp b/Engine/Plugins/Experimental/EnhancedInput/Source/InputBlueprintNodes/Private/K2Node_EnhancedInputAction.cpp
--- a/Engine/Plugins/Experimental/EnhancedInput/Source/InputBlueprintNodes/Private/K2Node_EnhancedInputAction.cpp
+++ b/Engine/Plugins/Experimental/EnhancedInput/Source/InputBlueprintNodes/Private/K2Node_EnhancedInputAction.cpp
@@ -77,13 +77,8 @@ void UK2Node_EnhancedInputAction::AllocateDefaultPins()
 	
 	Schema->SetPinAutogeneratedDefaultValueBasedOnType(ValuePin);
 
-#if ENABLE_BLUEPRINT_REAL_NUMBERS
 	CreatePin(EGPD_Output, UEdGraphSchema_K2::PC_Real, UEdGraphSchema_K2::PC_Double, ElapsedSecondsPinName)->bAdvancedView = true;
 	CreatePin(EGPD_Output, UEdGraphSchema_K2::PC_Real, UEdGraphSchema_K2::PC_Double, TriggeredSecondsPinName)->bAdvancedView = true;
-#else
-	CreatePin(EGPD_Output, UEdGraphSchema_K2::PC_Float, ElapsedSecondsPinName)->bAdvancedView = true;
-	CreatePin(EGPD_Output, UEdGraphSchema_K2::PC_Float, TriggeredSecondsPinName)->bAdvancedView = true;
-#endif
 
 	if(InputAction)
 	{
@@ -324,7 +319,6 @@ void UK2Node_EnhancedInputAction::ExpandNode(FKismetCompilerContext& CompilerCon
 	ActionValueVar->VariableType.PinSubCategoryObject = UK2Node_GetInputActionValue::GetValueSubCategoryObject(InputAction);
 	ActionValueVar->AllocateDefaultPins();
 
-#if ENABLE_BLUEPRINT_REAL_NUMBERS
 	UK2Node_TemporaryVariable* ElapsedSecondsVar = CompilerContext.SpawnIntermediateNode<UK2Node_TemporaryVariable>(this, SourceGraph);
 	ElapsedSecondsVar->VariableType.PinCategory = UEdGraphSchema_K2::PC_Real;
 	ElapsedSecondsVar->VariableType.PinSubCategory = UEdGraphSchema_K2::PC_Double;
@@ -333,14 +327,7 @@ void UK2Node_EnhancedInputAction::ExpandNode(FKismetCompilerContext& CompilerCon
 	TriggeredSecondsVar->VariableType.PinCategory = UEdGraphSchema_K2::PC_Real;
 	TriggeredSecondsVar->VariableType.PinSubCategory = UEdGraphSchema_K2::PC_Double;
 	TriggeredSecondsVar->AllocateDefaultPins();
-#else
-	UK2Node_TemporaryVariable* ElapsedSecondsVar = CompilerContext.SpawnIntermediateNode<UK2Node_TemporaryVariable>(this, SourceGraph);
-	ElapsedSecondsVar->VariableType.PinCategory = UEdGraphSchema_K2::PC_Float;
-	ElapsedSecondsVar->AllocateDefaultPins();
-	UK2Node_TemporaryVariable* TriggeredSecondsVar = CompilerContext.SpawnIntermediateNode<UK2Node_TemporaryVariable>(this, SourceGraph);
-	TriggeredSecondsVar->VariableType.PinCategory = UEdGraphSchema_K2::PC_Float;
-	TriggeredSecondsVar->AllocateDefaultPins();
-#endif
+
 	UK2Node_TemporaryVariable* InputActionVar = CompilerContext.SpawnIntermediateNode<UK2Node_TemporaryVariable>(this, SourceGraph);
 	InputActionVar->VariableType.PinCategory = UEdGraphSchema_K2::PC_Object;
 	InputActionVar->VariableType.PinSubCategoryObject = InputAction->GetClass();
diff --git a/Engine/Plugins/Experimental/EnhancedInput/Source/InputBlueprintNodes/Private/K2Node_GetInputActionValue.cpp b/Engine/Plugins/Experimental/EnhancedInput/Source/InputBlueprintNodes/Private/K2Node_GetInputActionValue.cpp
--- a/Engine/Plugins/Experimental/EnhancedInput/Source/InputBlueprintNodes/Private/K2Node_GetInputActionValue.cpp
+++ b/Engine/Plugins/Experimental/EnhancedInput/Source/InputBlueprintNodes/Private/K2Node_GetInputActionValue.cpp
@@ -48,11 +48,7 @@ struct FValueTypeData
 static TMap<EInputActionValueType, FValueTypeData> ValueLookups =
 {
 	{ EInputActionValueType::Boolean, FValueTypeData(UEdGraphSchema_K2::PC_Boolean) },
-#if ENABLE_BLUEPRINT_REAL_NUMBERS
 	{ EInputActionValueType::Axis1D, FValueTypeData(UEdGraphSchema_K2::PC_Real, UEdGraphSchema_K2::PC_Float) },
-#else
-	{ EInputActionValueType::Axis1D, FValueTypeData(UEdGraphSchema_K2::PC_Float) },
-#endif
 	{ EInputActionValueType::Axis2D, FValueTypeData(UEdGraphSchema_K2::PC_Struct, NAME_None, TBaseStructure<FVector2D>::Get()) },
 	{ EInputActionValueType::Axis3D, FValueTypeData(UEdGraphSchema_K2::PC_Struct, NAME_None, TBaseStructure<FVector>::Get()) },
 };
diff --git a/Engine/Plugins/Runtime/PropertyAccess/Source/PropertyAccessEditor/Private/PropertyAccessEditor.cpp b/Engine/Plugins/Runtime/PropertyAccess/Source/PropertyAccessEditor/Private/PropertyAccessEditor.cpp
--- a/Engine/Plugins/Runtime/PropertyAccess/Source/PropertyAccessEditor/Private/PropertyAccessEditor.cpp
+++ b/Engine/Plugins/Runtime/PropertyAccess/Source/PropertyAccessEditor/Private/PropertyAccessEditor.cpp
@@ -456,7 +456,6 @@ struct FPropertyAccessEditorSystem
 					return EPropertyAccessCopyType::DemoteDoubleToFloat;
 				}
 			}
-#if ENABLE_BLUEPRINT_REAL_NUMBERS
 			else if (SrcProperty->IsA<FArrayProperty>() && DestProperty->IsA<FArrayProperty>())
 			{
 				const FArrayProperty* SrcArrayProperty = CastField<const FArrayProperty>(SrcProperty);
@@ -477,7 +476,6 @@ struct FPropertyAccessEditorSystem
 					}
 				}
 			}
-#endif
 		}
 
 		OutErrorMessage = FText::Format(LOCTEXT("CopyTypeInvalidFormat", "@@ Cannot copy property ({0} -> {1})"), FText::FromString(SrcProperty->GetCPPType()), FText::FromString(DestProperty->GetCPPType()));
@@ -634,7 +632,6 @@ namespace PropertyAccess
 					return EPropertyAccessCompatibility::Promotable;	// LWC_TODO: Incorrect! Do not ship this!
 				}
 			}
-#if ENABLE_BLUEPRINT_REAL_NUMBERS
 			else if (InPropertyA->IsA<FArrayProperty>() && InPropertyB->IsA<FArrayProperty>())
 			{
 				const FArrayProperty* ArrayPropertyA = CastField<const FArrayProperty>(InPropertyA);
@@ -655,7 +652,6 @@ namespace PropertyAccess
 					}
 				}
 			}
-#endif
 		}
 
 		return EPropertyAccessCompatibility::Incompatible;
@@ -673,53 +669,25 @@ namespace PropertyAccess
 			// Not directly compatible, check for promotions
 			if (InPinTypeA.PinCategory == UEdGraphSchema_K2::PC_Boolean)
 			{
-#if ENABLE_BLUEPRINT_REAL_NUMBERS
 				if (InPinTypeB.PinCategory == UEdGraphSchema_K2::PC_Byte || InPinTypeB.PinCategory == UEdGraphSchema_K2::PC_Int || InPinTypeB.PinCategory == UEdGraphSchema_K2::PC_Int64 || InPinTypeB.PinCategory == UEdGraphSchema_K2::PC_Real)
-#else
-				if (InPinTypeB.PinCategory == UEdGraphSchema_K2::PC_Byte || InPinTypeB.PinCategory == UEdGraphSchema_K2::PC_Int || InPinTypeB.PinCategory == UEdGraphSchema_K2::PC_Int64 || InPinTypeB.PinCategory == UEdGraphSchema_K2::PC_Float || InPinTypeB.PinCategory == UEdGraphSchema_K2::PC_Double)
-#endif
 				{
 					return EPropertyAccessCompatibility::Promotable;
 				}
 			}
 			else if (InPinTypeA.PinCategory == UEdGraphSchema_K2::PC_Byte)
 			{
-#if ENABLE_BLUEPRINT_REAL_NUMBERS
 				if (InPinTypeB.PinCategory == UEdGraphSchema_K2::PC_Int || InPinTypeB.PinCategory == UEdGraphSchema_K2::PC_Int64 || InPinTypeB.PinCategory == UEdGraphSchema_K2::PC_Real)
-#else
-				if (InPinTypeB.PinCategory == UEdGraphSchema_K2::PC_Int || InPinTypeB.PinCategory == UEdGraphSchema_K2::PC_Int64 || InPinTypeB.PinCategory == UEdGraphSchema_K2::PC_Float || InPinTypeB.PinCategory == UEdGraphSchema_K2::PC_Double)
-#endif
 				{
 					return EPropertyAccessCompatibility::Promotable;
 				}
 			}
 			else if (InPinTypeA.PinCategory == UEdGraphSchema_K2::PC_Int)
 			{
-#if ENABLE_BLUEPRINT_REAL_NUMBERS
 				if (InPinTypeB.PinCategory == UEdGraphSchema_K2::PC_Int64 || InPinTypeB.PinCategory == UEdGraphSchema_K2::PC_Real)
-#else
-				if (InPinTypeB.PinCategory == UEdGraphSchema_K2::PC_Int64 || InPinTypeB.PinCategory == UEdGraphSchema_K2::PC_Float || InPinTypeB.PinCategory == UEdGraphSchema_K2::PC_Double)
-#endif
 				{
 					return EPropertyAccessCompatibility::Promotable;
 				}
 			}
-#if !ENABLE_BLUEPRINT_REAL_NUMBERS
-			else if (InPinTypeA.PinCategory == UEdGraphSchema_K2::PC_Float)
-			{
-				if (InPinTypeB.PinCategory == UEdGraphSchema_K2::PC_Double)
-				{
-					return EPropertyAccessCompatibility::Promotable;
-				}
-			}
-			else if (InPinTypeA.PinCategory == UEdGraphSchema_K2::PC_Double)
-			{
-				if (InPinTypeB.PinCategory == UEdGraphSchema_K2::PC_Float)
-				{
-					return EPropertyAccessCompatibility::Promotable;	// LWC_TODO: Incorrect! Do not ship this!
-				}
-			}
-#endif
 		}
 
 		return EPropertyAccessCompatibility::Incompatible;
diff --git a/Engine/Plugins/ScriptPlugin/Source/ScriptEditorPlugin/Private/ScriptBlueprintCompiler.cpp b/Engine/Plugins/ScriptPlugin/Source/ScriptEditorPlugin/Private/ScriptBlueprintCompiler.cpp
--- a/Engine/Plugins/ScriptPlugin/Source/ScriptEditorPlugin/Private/ScriptBlueprintCompiler.cpp
+++ b/Engine/Plugins/ScriptPlugin/Source/ScriptEditorPlugin/Private/ScriptBlueprintCompiler.cpp
@@ -56,12 +56,8 @@ void FScriptBlueprintCompiler::CreateClassVariablesFromBlueprint()
 			}
 			else if (Field.PropertyClass->IsChildOf(FFloatProperty::StaticClass()))
 			{
-#if ENABLE_BLUEPRINT_REAL_NUMBERS
 				PinCategory = UEdGraphSchema_K2::PC_Real;
 				PinSubCategory = UEdGraphSchema_K2::PC_Float;
-#else
-				PinCategory = UEdGraphSchema_K2::PC_Float;
-#endif
 			}
 			else if (Field.PropertyClass->IsChildOf(FIntProperty::StaticClass()))
 			{
diff --git a/Engine/Source/Developer/RigVMDeveloper/Private/RigVMTypeUtils.cpp b/Engine/Source/Developer/RigVMDeveloper/Private/RigVMTypeUtils.cpp
--- a/Engine/Source/Developer/RigVMDeveloper/Private/RigVMTypeUtils.cpp
+++ b/Engine/Source/Developer/RigVMDeveloper/Private/RigVMTypeUtils.cpp
@@ -85,7 +85,6 @@ FRigVMExternalVariable RigVMTypeUtils::ExternalVariableFromPinType(const FName&
 		}
 		ExternalVariable.Size = sizeof(uint8);
 	}
-#if ENABLE_BLUEPRINT_REAL_NUMBERS
 	else if (InPinType.PinCategory == UEdGraphSchema_K2::PC_Real)
 	{
 		if (InPinType.PinSubCategory == UEdGraphSchema_K2::PC_Float)
@@ -102,19 +101,7 @@ FRigVMExternalVariable RigVMTypeUtils::ExternalVariableFromPinType(const FName&
 		{
 			checkf(false, TEXT("Unexpected subcategory for PC_Real pin type."));
 		}
-	}		
-#else
-	else if (InPinType.PinCategory == UEdGraphSchema_K2::PC_Float)
-	{
-		ExternalVariable.TypeName = TEXT("float");
-		ExternalVariable.Size = sizeof(float);
-	}
-	else if (InPinType.PinCategory == UEdGraphSchema_K2::PC_Double)
-	{
-		ExternalVariable.TypeName = TEXT("double");
-		ExternalVariable.Size = sizeof(double);
 	}
-#endif
 	else if (InPinType.PinCategory == UEdGraphSchema_K2::PC_Name)
 	{
 		ExternalVariable.TypeName = TEXT("FName");
@@ -300,7 +287,6 @@ FEdGraphPinType RigVMTypeUtils::PinTypeFromCPPType(const FString& InCPPType, UOb
 	{
 		PinType.PinCategory = UEdGraphSchema_K2::PC_Int;
 	}
-#if ENABLE_BLUEPRINT_REAL_NUMBERS
 	else if (BaseCPPType == TEXT("float"))
 	{
 		PinType.PinCategory = UEdGraphSchema_K2::PC_Real;
@@ -311,16 +297,6 @@ FEdGraphPinType RigVMTypeUtils::PinTypeFromCPPType(const FString& InCPPType, UOb
 		PinType.PinCategory = UEdGraphSchema_K2::PC_Real;
 		PinType.PinSubCategory = UEdGraphSchema_K2::PC_Double;
 	}
-#else
-	else if (BaseCPPType == TEXT("float"))
-	{
-		PinType.PinCategory = UEdGraphSchema_K2::PC_Float;
-	}
-	else if (BaseCPPType == TEXT("double"))
-	{
-		PinType.PinCategory = UEdGraphSchema_K2::PC_Double;
-	}
-#endif
 	else if (BaseCPPType == TEXT("FName"))
 	{
 		PinType.PinCategory = UEdGraphSchema_K2::PC_Name;
@@ -357,7 +333,6 @@ FEdGraphPinType RigVMTypeUtils::PinTypeFromExternalVariable(const FRigVMExternal
 	{
 		PinType.PinCategory = UEdGraphSchema_K2::PC_Int;
 	}
-#if ENABLE_BLUEPRINT_REAL_NUMBERS
 	else if (InExternalVariable.TypeName == TEXT("float"))
 	{
 		PinType.PinCategory = UEdGraphSchema_K2::PC_Real;
@@ -368,16 +343,6 @@ FEdGraphPinType RigVMTypeUtils::PinTypeFromExternalVariable(const FRigVMExternal
 		PinType.PinCategory = UEdGraphSchema_K2::PC_Real;
 		PinType.PinSubCategory = UEdGraphSchema_K2::PC_Double;
 	}
-#else
-	else if (InExternalVariable.TypeName == TEXT("float"))
-	{
-		PinType.PinCategory = UEdGraphSchema_K2::PC_Float;
-	}
-	else if (InExternalVariable.TypeName == TEXT("double"))
-	{
-		PinType.PinCategory = UEdGraphSchema_K2::PC_Double;
-	}
-#endif
 	else if (InExternalVariable.TypeName == TEXT("FName"))
 	{
 		PinType.PinCategory = UEdGraphSchema_K2::PC_Name;
@@ -435,7 +400,6 @@ FEdGraphPinType RigVMTypeUtils::PinTypeFromRigVMVariableDescription(
 	{
 		PinType.PinCategory = UEdGraphSchema_K2::PC_Int;
 	}
-#if ENABLE_BLUEPRINT_REAL_NUMBERS
 	else if (CurrentCPPType == TEXT("float"))
 	{
 		PinType.PinCategory = UEdGraphSchema_K2::PC_Real;
@@ -446,16 +410,6 @@ FEdGraphPinType RigVMTypeUtils::PinTypeFromRigVMVariableDescription(
 		PinType.PinCategory = UEdGraphSchema_K2::PC_Real;
 		PinType.PinSubCategory = UEdGraphSchema_K2::PC_Double;
 	}
-#else
-	else if (CurrentCPPType == TEXT("float"))
-	{
-		PinType.PinCategory = UEdGraphSchema_K2::PC_Float;
-	}
-	else if (CurrentCPPType == TEXT("double"))
-	{
-		PinType.PinCategory = UEdGraphSchema_K2::PC_Double;
-	}
-#endif
 	else if (CurrentCPPType == TEXT("FName"))
 	{
 		PinType.PinCategory = UEdGraphSchema_K2::PC_Name;
@@ -552,7 +506,6 @@ bool RigVMTypeUtils::CPPTypeFromPinType(const FEdGraphPinType& InPinType, FStrin
 	{
 		OutCPPType = Prefix + TEXT("int32") + Suffix;
 	}
-#if ENABLE_BLUEPRINT_REAL_NUMBERS
 	else if (InPinType.PinCategory == UEdGraphSchema_K2::PC_Real)
 	{
 		if (InPinType.PinSubCategory == UEdGraphSchema_K2::PC_Float)
@@ -567,17 +520,7 @@ bool RigVMTypeUtils::CPPTypeFromPinType(const FEdGraphPinType& InPinType, FStrin
 		{
 			checkf(false, TEXT("Unexpected subcategory for PC_Real pin type."));
 		}
-	}		
-#else
-	else if (InPinType.PinCategory == UEdGraphSchema_K2::PC_Float)
-	{
-		OutCPPType = Prefix + TEXT("float") + Suffix;
-	}
-	else if (InPinType.PinCategory == UEdGraphSchema_K2::PC_Double)
-	{
-		OutCPPType = Prefix + TEXT("double") + Suffix;
 	}
-#endif
 	else if (InPinType.PinCategory == UEdGraphSchema_K2::PC_Float)
 	{
 		OutCPPType = Prefix + TEXT("float") + Suffix;
diff --git a/Engine/Source/Editor/AnimGraph/Private/K2Node_TransitionRuleGetter.cpp b/Engine/Source/Editor/AnimGraph/Private/K2Node_TransitionRuleGetter.cpp
--- a/Engine/Source/Editor/AnimGraph/Private/K2Node_TransitionRuleGetter.cpp
+++ b/Engine/Source/Editor/AnimGraph/Private/K2Node_TransitionRuleGetter.cpp
@@ -26,11 +26,8 @@ UK2Node_TransitionRuleGetter::UK2Node_TransitionRuleGetter(const FObjectInitiali
 
 void UK2Node_TransitionRuleGetter::AllocateDefaultPins()
 {
-#if ENABLE_BLUEPRINT_REAL_NUMBERS
 	UEdGraphPin* OutputPin = CreatePin(EGPD_Output, UEdGraphSchema_K2::PC_Real, UEdGraphSchema_K2::PC_Float, TEXT("Output"));
-#else
-	UEdGraphPin* OutputPin = CreatePin(EGPD_Output, UEdGraphSchema_K2::PC_Float, TEXT("Output"));
-#endif
+
 	OutputPin->PinFriendlyName = GetFriendlyName(GetterType);
 
 	PreloadObject(AssociatedAnimAssetPlayerNode);
diff --git a/Engine/Source/Editor/BlueprintGraph/Classes/EdGraphSchema_K2.h b/Engine/Source/Editor/BlueprintGraph/Classes/EdGraphSchema_K2.h
--- a/Engine/Source/Editor/BlueprintGraph/Classes/EdGraphSchema_K2.h
+++ b/Engine/Source/Editor/BlueprintGraph/Classes/EdGraphSchema_K2.h
@@ -331,9 +331,7 @@ class BLUEPRINTGRAPH_API UEdGraphSchema_K2 : public UEdGraphSchema
 	static const FName PC_Int64;
 	static const FName PC_Float;
 	static const FName PC_Double;
-#if ENABLE_BLUEPRINT_REAL_NUMBERS
 	static const FName PC_Real;
-#endif
 	static const FName PC_Name;
 	static const FName PC_Delegate;    // SubCategoryObject is the UFunction of the delegate signature
 	static const FName PC_MCDelegate;  // SubCategoryObject is the UFunction of the delegate signature
diff --git a/Engine/Source/Editor/BlueprintGraph/Private/BlueprintTypePromotion.cpp b/Engine/Source/Editor/BlueprintGraph/Private/BlueprintTypePromotion.cpp
--- a/Engine/Source/Editor/BlueprintGraph/Private/BlueprintTypePromotion.cpp
+++ b/Engine/Source/Editor/BlueprintGraph/Private/BlueprintTypePromotion.cpp
@@ -72,20 +72,11 @@ void FTypePromotion::CreatePromotionTable()
 	PromotionTable =
 	{
 
-#if ENABLE_BLUEPRINT_REAL_NUMBERS
 		// Type_X...						Can be promoted to...
 		{ UEdGraphSchema_K2::PC_Int,		{ UEdGraphSchema_K2::PC_Real, UEdGraphSchema_K2::PC_Int64 } },
 		{ UEdGraphSchema_K2::PC_Byte,		{ UEdGraphSchema_K2::PC_Real, UEdGraphSchema_K2::PC_Int, UEdGraphSchema_K2::PC_Int64 } },
 		{ UEdGraphSchema_K2::PC_Real,		{ UEdGraphSchema_K2::PC_Int64 } },
 		{ UEdGraphSchema_K2::PC_Wildcard,	{ UEdGraphSchema_K2::PC_Int, UEdGraphSchema_K2::PC_Int64, UEdGraphSchema_K2::PC_Real, UEdGraphSchema_K2::PC_Byte, UEdGraphSchema_K2::PC_Boolean } },
-#else
-		// Type_X...						Can be promoted to...
-		{ UEdGraphSchema_K2::PC_Int,		{ UEdGraphSchema_K2::PC_Float, UEdGraphSchema_K2::PC_Double, UEdGraphSchema_K2::PC_Int64 } },
-		{ UEdGraphSchema_K2::PC_Byte,		{ UEdGraphSchema_K2::PC_Float, UEdGraphSchema_K2::PC_Int, UEdGraphSchema_K2::PC_Int64, UEdGraphSchema_K2::PC_Double } },
-		{ UEdGraphSchema_K2::PC_Float,		{ UEdGraphSchema_K2::PC_Double, UEdGraphSchema_K2::PC_Int64 } },
-		{ UEdGraphSchema_K2::PC_Double,		{ UEdGraphSchema_K2::PC_Int64 } },
-		{ UEdGraphSchema_K2::PC_Wildcard,	{ UEdGraphSchema_K2::PC_Int, UEdGraphSchema_K2::PC_Int64, UEdGraphSchema_K2::PC_Float, UEdGraphSchema_K2::PC_Double, UEdGraphSchema_K2::PC_Byte, UEdGraphSchema_K2::PC_Boolean } },
-#endif
 	};
 }
 
diff --git a/Engine/Source/Editor/BlueprintGraph/Private/CallFunctionHandler.cpp b/Engine/Source/Editor/BlueprintGraph/Private/CallFunctionHandler.cpp
--- a/Engine/Source/Editor/BlueprintGraph/Private/CallFunctionHandler.cpp
+++ b/Engine/Source/Editor/BlueprintGraph/Private/CallFunctionHandler.cpp
@@ -152,7 +152,6 @@ void FKCHandler_CallFunction::CreateFunctionCallStatement(FKismetFunctionContext
 		UScriptStruct* RotatorStruct = TBaseStructure<FRotator>::Get();
 		UScriptStruct* TransformStruct = TBaseStructure<FTransform>::Get();
 
-#if ENABLE_BLUEPRINT_REAL_NUMBERS
 		// If a function parameter needs an implicit double<->float cast *and* it's a non-const reference,
 		// then we need to copy the value of the casted temporary *back* to its source.
 		// 
@@ -165,7 +164,6 @@ void FKCHandler_CallFunction::CreateFunctionCallStatement(FKismetFunctionContext
 		//
 		using CastEntry = TTuple<FBPTerminal*, FBPTerminal*, EKismetCompiledStatementType>;
 		TArray<CastEntry> ModifiedCastInputs;
-#endif
 
 		// Check each property
 		bool bMatchedAllParams = true;
@@ -255,7 +253,6 @@ void FKCHandler_CallFunction::CreateFunctionCallStatement(FKismetFunctionContext
 									RHSTerm = *InterfaceTerm;
 								}
 
-#if ENABLE_BLUEPRINT_REAL_NUMBERS
 								{
 									using namespace UE::KismetCompiler;
 
@@ -283,7 +280,6 @@ void FKCHandler_CallFunction::CreateFunctionCallStatement(FKismetFunctionContext
 										RHSTerm = CastTerm;
 									}
 								}
-#endif
 
 								int32 ParameterIndex = RHSTerms.Add(RHSTerm);
 
@@ -508,7 +504,6 @@ void FKCHandler_CallFunction::CreateFunctionCallStatement(FKismetFunctionContext
 				}
 			}
 
-#if ENABLE_BLUEPRINT_REAL_NUMBERS
 			{
 				using namespace UE::KismetCompiler;
 
@@ -527,7 +522,6 @@ void FKCHandler_CallFunction::CreateFunctionCallStatement(FKismetFunctionContext
 					CastStatement.RHS.Add(LocalRHSTerm);
 				}
 			}
-#endif
 
 			// Create the exit from this node if there is one
 			if (bIsLatent)
diff --git a/Engine/Source/Editor/BlueprintGraph/Private/EdGraphSchema_K2.cpp b/Engine/Source/Editor/BlueprintGraph/Private/EdGraphSchema_K2.cpp
--- a/Engine/Source/Editor/BlueprintGraph/Private/EdGraphSchema_K2.cpp
+++ b/Engine/Source/Editor/BlueprintGraph/Private/EdGraphSchema_K2.cpp
@@ -555,11 +555,7 @@ void UEdGraphSchema_K2::FPinTypeTreeInfo::Init(const FText& InFriendlyName, cons
 	FriendlyName = InFriendlyName;
 	Tooltip = InTooltip;
 	PinType.PinCategory = (CategoryName == PC_Enum ? PC_Byte : CategoryName);
-#if ENABLE_BLUEPRINT_REAL_NUMBERS
 	PinType.PinSubCategory = (CategoryName == PC_Real ? PC_Double : NAME_None);
-#else
-	PinType.PinSubCategory = NAME_None;
-#endif
 	PinType.PinSubCategoryObject = nullptr;
 
 	bReadOnly = bInReadOnly;
@@ -646,9 +642,7 @@ const FName UEdGraphSchema_K2::PC_Int(TEXT("int"));
 const FName UEdGraphSchema_K2::PC_Int64(TEXT("int64"));
 const FName UEdGraphSchema_K2::PC_Float(TEXT("float"));
 const FName UEdGraphSchema_K2::PC_Double(TEXT("double"));
-#if ENABLE_BLUEPRINT_REAL_NUMBERS
 const FName UEdGraphSchema_K2::PC_Real(TEXT("real"));
-#endif
 const FName UEdGraphSchema_K2::PC_Name(TEXT("name"));
 const FName UEdGraphSchema_K2::PC_Delegate(TEXT("delegate"));
 const FName UEdGraphSchema_K2::PC_MCDelegate(TEXT("mcdelegate"));
@@ -2892,21 +2886,10 @@ FLinearColor UEdGraphSchema_K2::GetPinTypeColor(const FEdGraphPinType& PinType)
 	{
 		return Settings->InterfacePinTypeColor;
 	}
-#if ENABLE_BLUEPRINT_REAL_NUMBERS
 	else if (TypeName == PC_Real)
 	{
 		return Settings->RealPinTypeColor;
 	}
-#else
-	else if (TypeName == PC_Float)
-	{
-		return Settings->FloatPinTypeColor;
-	}
-	else if (TypeName == PC_Double)
-	{
-		return Settings->DoublePinTypeColor;
-	}
-#endif
 	else if (TypeName == PC_Boolean)
 	{
 		return Settings->BooleanPinTypeColor;
@@ -3350,7 +3333,6 @@ bool UEdGraphSchema_K2::GetPropertyCategoryInfo(const FProperty* TestProperty, F
 			}
 		}
 	}
-#if ENABLE_BLUEPRINT_REAL_NUMBERS
 	else if (TestProperty->IsA<FFloatProperty>())
 	{
 		OutCategory = PC_Real;
@@ -3361,16 +3343,6 @@ bool UEdGraphSchema_K2::GetPropertyCategoryInfo(const FProperty* TestProperty, F
 		OutCategory = PC_Real;
 		OutSubCategory = PC_Double;
 	}
-#else
-	else if (TestProperty->IsA<FFloatProperty>())
-	{
-		OutCategory = PC_Float;
-	}
-	else if(TestProperty->IsA<FDoubleProperty>())
-	{
-		OutCategory = PC_Double;
-	}
-#endif
 	else if (TestProperty->IsA<FInt64Property>())
 	{
 		OutCategory = PC_Int64;
@@ -3666,12 +3638,7 @@ FText UEdGraphSchema_K2::GetCategoryText(const FName Category, const bool bForMe
 		CategoryDescriptions.Add(PC_Class, LOCTEXT("ClassCategory","Class Reference"));
 		CategoryDescriptions.Add(PC_Int, LOCTEXT("IntCategory", "Integer"));
 		CategoryDescriptions.Add(PC_Int64, LOCTEXT("Int64Category", "Integer64"));
-#if ENABLE_BLUEPRINT_REAL_NUMBERS
 		CategoryDescriptions.Add(PC_Real, LOCTEXT("RealCategory", "Real"));
-#else
-		CategoryDescriptions.Add(PC_Float, LOCTEXT("FloatCategory", "Float"));
-		CategoryDescriptions.Add(PC_Double, LOCTEXT("DoubleCategory", "Double"));
-#endif
 		CategoryDescriptions.Add(PC_Name, LOCTEXT("NameCategory","Name"));
 		CategoryDescriptions.Add(PC_Delegate, LOCTEXT("DelegateCategory","Delegate"));
 		CategoryDescriptions.Add(PC_MCDelegate, LOCTEXT("MulticastDelegateCategory","Multicast Delegate"));
@@ -3858,12 +3825,7 @@ void UEdGraphSchema_K2::GetVariableTypeTree(TArray< TSharedPtr<FPinTypeTreeInfo>
 
 	if (!bIndexTypesOnly)
 	{
-#if ENABLE_BLUEPRINT_REAL_NUMBERS
 		TypeTree.Add(MakeShareable(new FPinTypeTreeInfo(GetCategoryText(PC_Real, true), PC_Real, this, LOCTEXT("RealType", "Floating point number"))));
-#else
-		TypeTree.Add(MakeShareable(new FPinTypeTreeInfo(GetCategoryText(PC_Float, true), PC_Float, this, LOCTEXT("FloatType", "Floating point number"))));
-		TypeTree.Add(MakeShareable(new FPinTypeTreeInfo(GetCategoryText(PC_Double, true), PC_Double, this, LOCTEXT("DoubleType", "64 bit floating point number"))));
-#endif
 		TypeTree.Add(MakeShareable(new FPinTypeTreeInfo(GetCategoryText(PC_Name, true), PC_Name, this, LOCTEXT("NameType", "A text name"))));
 		TypeTree.Add(MakeShareable(new FPinTypeTreeInfo(GetCategoryText(PC_String, true), PC_String, this, LOCTEXT("StringType", "A text string"))));
 		TypeTree.Add(MakeShareable(new FPinTypeTreeInfo(GetCategoryText(PC_Text, true), PC_Text, this, LOCTEXT("TextType", "A localizable text string"))));
@@ -4136,7 +4098,6 @@ bool UEdGraphSchema_K2::DefaultValueSimpleValidation(const FEdGraphPinType& PinT
 			}
 		}
 	}
-#if ENABLE_BLUEPRINT_REAL_NUMBERS
 	else if (PinCategory == PC_Real)
 	{
 		if (!NewDefaultValue.IsEmpty())
@@ -4147,28 +4108,6 @@ bool UEdGraphSchema_K2::DefaultValueSimpleValidation(const FEdGraphPinType& PinT
 			}
 		}
 	}
-#else
-	else if (PinCategory == PC_Float)
-	{
-		if (!NewDefaultValue.IsEmpty())
-		{
-			if (!FDefaultValueHelper::IsStringValidFloat(NewDefaultValue))
-			{
-				DVSV_RETURN_MSG(TEXT("Expected a valid number for an float property"));
-			}
-		}
-	}
-	else if (PinCategory == PC_Double)
-	{
-		if (!NewDefaultValue.IsEmpty())
-		{
-			if (!FDefaultValueHelper::IsStringValidFloat(NewDefaultValue))
-			{
-				DVSV_RETURN_MSG(TEXT("Expected a valid number for an double property"));
-			}
-		}
-	}
-#endif
 	else if (PinCategory == PC_Int)
 	{
 		if (!NewDefaultValue.IsEmpty())
@@ -4350,14 +4289,11 @@ bool UEdGraphSchema_K2::ArePinTypesCompatible(const FEdGraphPinType& Output, con
 				return 
 					Input.PinValueType.TerminalCategory == PC_Wildcard ||
 					Output.PinValueType.TerminalCategory == PC_Wildcard ||
-#if ENABLE_BLUEPRINT_REAL_NUMBERS
 					((Input.PinValueType.TerminalCategory == PC_Real) && (Output.PinValueType.TerminalCategory == PC_Real)) ||
-#endif
 					Input.PinValueType == Output.PinValueType;
 			}
 			return true;
 		}
-#if ENABLE_BLUEPRINT_REAL_NUMBERS
 		// Reals, whether they're actually a float or double, are always compatible.
 		// We'll insert an implicit conversion in the bytecode where necessary.
 		else if (Output.PinCategory == PC_Real)
@@ -4368,7 +4304,6 @@ bool UEdGraphSchema_K2::ArePinTypesCompatible(const FEdGraphPinType& Output, con
 		{
 			return true;
 		}
-#endif
 		else if (Output.PinCategory == PC_Interface)
 		{
 			UClass const* OutputClass = Cast<UClass const>(Output.PinSubCategoryObject.Get());
@@ -4843,27 +4778,12 @@ bool UEdGraphSchema_K2::DoesDefaultValueMatchAutogenerated(const UEdGraphPin& In
 	}
 	else if (!InPin.bUseBackwardsCompatForEmptyAutogeneratedValue)
 	{
-#if ENABLE_BLUEPRINT_REAL_NUMBERS
 		if (InPin.PinType.PinCategory == PC_Real)
 		{
 			const double AutogeneratedDouble = FCString::Atod(*InPin.AutogeneratedDefaultValue);
 			const double DefaultDouble = FCString::Atod(*InPin.DefaultValue);
 			return (AutogeneratedDouble == DefaultDouble);
 		}
-#else
-		if (InPin.PinType.PinCategory == PC_Float)
-		{
-			const float AutogeneratedFloat = FCString::Atof(*InPin.AutogeneratedDefaultValue);
-			const float DefaultFloat = FCString::Atof(*InPin.DefaultValue);
-			return (AutogeneratedFloat == DefaultFloat);
-		}
-		else if (InPin.PinType.PinCategory == PC_Double)
-		{
-			const double AutogeneratedDouble = FCString::Atod(*InPin.AutogeneratedDefaultValue);
-			const double DefaultDouble = FCString::Atod(*InPin.DefaultValue);
-			return (AutogeneratedDouble == DefaultDouble);
-		}
-#endif
 		else if (InPin.PinType.PinCategory == PC_Struct)
 		{
 			if ((InPin.PinType.PinSubCategoryObject == VectorStruct) || (InPin.PinType.PinSubCategoryObject == Vector3fStruct))
@@ -5079,22 +4999,10 @@ void UEdGraphSchema_K2::SetPinAutogeneratedDefaultValueBasedOnType(UEdGraphPin*
 			NewValue = TEXT("0");
 		}
 	}
-#if ENABLE_BLUEPRINT_REAL_NUMBERS
 	else if (Pin->PinType.PinCategory == PC_Real)
 	{
 		NewValue = TEXT("0.0");
 	}
-#else
-	else if (Pin->PinType.PinCategory == PC_Float)
-	{
-		// This is a slightly different format than is produced by PropertyValueToString, but changing it has backward compatibility issues
-		NewValue = TEXT("0.0");
-	}
-	else if (Pin->PinType.PinCategory == PC_Double)
-	{
-		NewValue = TEXT("0.0");
-	}
-#endif
 	else if (Pin->PinType.PinCategory == PC_Boolean)
 	{
 		NewValue = TEXT("false");
@@ -5297,21 +5205,10 @@ UFunction* UEdGraphSchema_K2::FindSetVariableByNameFunction(const FEdGraphPinTyp
 	{
 		SetFunctionName = FSetVariableByNameFunctionNames::SetByteName;
 	}
-#if ENABLE_BLUEPRINT_REAL_NUMBERS
 	else if(PinType.PinCategory == UEdGraphSchema_K2::PC_Real)
 	{
 		SetFunctionName = FSetVariableByNameFunctionNames::SetDoubleName;
 	}
-#else
-	else if(PinType.PinCategory == UEdGraphSchema_K2::PC_Float)
-	{
-		SetFunctionName = FSetVariableByNameFunctionNames::SetFloatName;
-	}
-	else if(PinType.PinCategory == UEdGraphSchema_K2::PC_Double)
-	{
-		SetFunctionName = FSetVariableByNameFunctionNames::SetDoubleName;
-	}
-#endif
 	else if(PinType.PinCategory == UEdGraphSchema_K2::PC_Boolean)
 	{
 		SetFunctionName = FSetVariableByNameFunctionNames::SetBoolName;
diff --git a/Engine/Source/Editor/BlueprintGraph/Private/K2Node_AssignmentStatement.cpp b/Engine/Source/Editor/BlueprintGraph/Private/K2Node_AssignmentStatement.cpp
--- a/Engine/Source/Editor/BlueprintGraph/Private/K2Node_AssignmentStatement.cpp
+++ b/Engine/Source/Editor/BlueprintGraph/Private/K2Node_AssignmentStatement.cpp
@@ -43,7 +43,6 @@ public:
 
 		ValidateAndRegisterNetIfLiteral(Context, ValuePin);
 
-#if ENABLE_BLUEPRINT_REAL_NUMBERS
 		{
 			using namespace UE::KismetCompiler;
 
@@ -80,7 +79,6 @@ public:
 				CompilerContext.MessageLog.Error(*LOCTEXT("NoVariableOrValueNets_Error", "Expected Variable and Value pins to have valid connections in @@").ToString(), Node);
 			}
 		}
-#endif
 	}
 
 
diff --git a/Engine/Source/Editor/BlueprintGraph/Private/K2Node_EaseFunction.cpp b/Engine/Source/Editor/BlueprintGraph/Private/K2Node_EaseFunction.cpp
--- a/Engine/Source/Editor/BlueprintGraph/Private/K2Node_EaseFunction.cpp
+++ b/Engine/Source/Editor/BlueprintGraph/Private/K2Node_EaseFunction.cpp
@@ -101,11 +101,8 @@ void UK2Node_EaseFunction::AllocateDefaultPins()
 	// Make sure that the default value is set correctly if none has been set
 	K2Schema->SetPinAutogeneratedDefaultValueBasedOnType(CachedEaseFuncPin);
 
-#if ENABLE_BLUEPRINT_REAL_NUMBERS
 	UEdGraphPin* AlphaPin = CreatePin(EGPD_Input, UEdGraphSchema_K2::PC_Real, UEdGraphSchema_K2::PC_Double, FEaseFunctionNodeHelper::GetAlphaPinName());
-#else
-	UEdGraphPin* AlphaPin = CreatePin(EGPD_Input, UEdGraphSchema_K2::PC_Float, FEaseFunctionNodeHelper::GetAlphaPinName());
-#endif
+
 	SetPinToolTip(*AlphaPin, LOCTEXT("AlphaPinTooltip", "Alpha value used to specify the easing in time."));
 
 	// Add wildcard pins for A, B and the return Pin
@@ -122,11 +119,8 @@ void UK2Node_EaseFunction::AllocateDefaultPins()
 	SetPinToolTip(*ShortestPathPin, LOCTEXT("ShortestPathPinTooltip", "When interpolating the shortest path should be taken."));
 	K2Schema->SetPinAutogeneratedDefaultValue(ShortestPathPin, TEXT("true"));
 
-#if ENABLE_BLUEPRINT_REAL_NUMBERS
 	UEdGraphPin* BlendExpPin = CreatePin(EGPD_Input, UEdGraphSchema_K2::PC_Real, UEdGraphSchema_K2::PC_Double, FEaseFunctionNodeHelper::GetBlendExpPinName());
-#else
-	UEdGraphPin* BlendExpPin = CreatePin(EGPD_Input, UEdGraphSchema_K2::PC_Float, FEaseFunctionNodeHelper::GetBlendExpPinName());
-#endif
+
 	SetPinToolTip(*BlendExpPin, LOCTEXT("BlendExpPinDescription", "Blend Exponent for basic ease functions"));
 	K2Schema->SetPinAutogeneratedDefaultValue(BlendExpPin, TEXT("2.0"));
 
@@ -232,7 +226,6 @@ bool UK2Node_EaseFunction::IsConnectionDisallowed(const UEdGraphPin* MyPin, cons
 		MyPin->PinName == FEaseFunctionNodeHelper::GetBPinName() || 
 		MyPin->PinName == FEaseFunctionNodeHelper::GetResultPinName())
 	{
-#if ENABLE_BLUEPRINT_REAL_NUMBERS
 		const bool bConnectionOk = (
 			OtherPin->PinType.PinCategory == UEdGraphSchema_K2::PC_Real ||
 			(
@@ -245,20 +238,6 @@ bool UK2Node_EaseFunction::IsConnectionDisallowed(const UEdGraphPin* MyPin, cons
 					)
 				)
 			) && !OtherPin->PinType.IsContainer();
-#else
-		const bool bConnectionOk = (
-			OtherPin->PinType.PinCategory == UEdGraphSchema_K2::PC_Float ||
-			(
-				OtherPin->PinType.PinCategory == UEdGraphSchema_K2::PC_Struct &&
-				OtherPin->PinType.PinSubCategoryObject.IsValid() &&
-				(
-					OtherPin->PinType.PinSubCategoryObject.Get()->GetName() == TEXT("Vector") ||
-					OtherPin->PinType.PinSubCategoryObject.Get()->GetName() == TEXT("Rotator") ||
-					OtherPin->PinType.PinSubCategoryObject.Get()->GetName() == TEXT("Transform")
-				)
-			)
-		) && !OtherPin->PinType.IsContainer();
-#endif
 
 		if (!bConnectionOk)
 		{
@@ -334,17 +313,10 @@ void UK2Node_EaseFunction::PinTypeChanged(UEdGraphPin* Pin)
 				EaseFunctionName = NAME_None;
 
 				// Generate the right function name
-#if ENABLE_BLUEPRINT_REAL_NUMBERS
 				if (InstigatorPin->PinType.PinCategory == UEdGraphSchema_K2::PC_Real)
 				{
 					EaseFunctionName = TEXT("Ease");
 				}
-#else
-				if (InstigatorPin->PinType.PinCategory == UEdGraphSchema_K2::PC_Float)
-				{
-					EaseFunctionName = TEXT("Ease");
-				}
-#endif
 				else if (InstigatorPin->PinType.PinCategory == UEdGraphSchema_K2::PC_Struct)
 				{
 					if (InstigatorPin->PinType.PinSubCategoryObject.Get()->GetName() == TEXT("Vector"))
diff --git a/Engine/Source/Editor/BlueprintGraph/Private/K2Node_EditablePinBase.cpp b/Engine/Source/Editor/BlueprintGraph/Private/K2Node_EditablePinBase.cpp
--- a/Engine/Source/Editor/BlueprintGraph/Private/K2Node_EditablePinBase.cpp
+++ b/Engine/Source/Editor/BlueprintGraph/Private/K2Node_EditablePinBase.cpp
@@ -58,7 +58,6 @@ FArchive& operator<<(FArchive& Ar, FUserPinInfo& Info)
 		Info.PinType.PinCategory = *PinCategoryStr;
 		Info.PinType.PinSubCategory = *PinSubCategoryStr;
 
-#if ENABLE_BLUEPRINT_REAL_NUMBERS
 		bool bFixupPinCategories =
 			((Info.PinType.PinCategory == TEXT("double")) || (Info.PinType.PinCategory == TEXT("float")));
 
@@ -67,7 +66,6 @@ FArchive& operator<<(FArchive& Ar, FUserPinInfo& Info)
 			Info.PinType.PinCategory = TEXT("real");
 			Info.PinType.PinSubCategory = TEXT("double");
 		}
-#endif
 
 		Ar << Info.PinType.PinSubCategoryObject;
 	}
diff --git a/Engine/Source/Editor/BlueprintGraph/Private/K2Node_FormatText.cpp b/Engine/Source/Editor/BlueprintGraph/Private/K2Node_FormatText.cpp
--- a/Engine/Source/Editor/BlueprintGraph/Private/K2Node_FormatText.cpp
+++ b/Engine/Source/Editor/BlueprintGraph/Private/K2Node_FormatText.cpp
@@ -381,7 +381,6 @@ void UK2Node_FormatText::ExpandNode(class FKismetCompilerContext& CompilerContex
 				// Connect the int output pin to the argument value
 				CallFloatToDoubleFunction->FindPinChecked(UEdGraphSchema_K2::PN_ReturnValue)->MakeLinkTo(MakeFormatArgumentDataStruct->FindPinChecked(GET_MEMBER_NAME_STRING_CHECKED(FFormatArgumentData, ArgumentValueInt)));
 			}
-#if ENABLE_BLUEPRINT_REAL_NUMBERS
 			else if (ArgumentPinCategory == UEdGraphSchema_K2::PC_Real)
 			{
 				if (ArgumentPin->PinType.PinSubCategory == UEdGraphSchema_K2::PC_Float)
@@ -399,18 +398,6 @@ void UK2Node_FormatText::ExpandNode(class FKismetCompilerContext& CompilerContex
 					check(false);
 				}
 			}
-#else
-			else if (ArgumentPinCategory == UEdGraphSchema_K2::PC_Float)
-			{
-				MakeFormatArgumentDataStruct->GetSchema()->TrySetDefaultValue(*ArgumentTypePin, TEXT("Float"));
-				CompilerContext.MovePinLinksToIntermediate(*ArgumentPin, *MakeFormatArgumentDataStruct->FindPinChecked(GET_MEMBER_NAME_STRING_CHECKED(FFormatArgumentData, ArgumentValueFloat)));
-			}
-			else if (ArgumentPinCategory == UEdGraphSchema_K2::PC_Double)
-			{
-				MakeFormatArgumentDataStruct->GetSchema()->TrySetDefaultValue(*ArgumentTypePin, TEXT("Double"));
-				CompilerContext.MovePinLinksToIntermediate(*ArgumentPin, *MakeFormatArgumentDataStruct->FindPinChecked(GET_MEMBER_NAME_STRING_CHECKED(FFormatArgumentData, ArgumentValueDouble)));
-			}
-#endif
 			else if (ArgumentPinCategory == UEdGraphSchema_K2::PC_Int64)
 			{
 				MakeFormatArgumentDataStruct->GetSchema()->TrySetDefaultValue(*ArgumentTypePin, TEXT("Int64"));
@@ -566,17 +553,10 @@ bool UK2Node_FormatText::IsConnectionDisallowed(const UEdGraphPin* MyPin, const
 		const FName& OtherPinCategory = OtherPin->PinType.PinCategory;
 
 		bool bIsValidType = false;
-#if ENABLE_BLUEPRINT_REAL_NUMBERS
 		if (OtherPinCategory == UEdGraphSchema_K2::PC_Int || OtherPinCategory == UEdGraphSchema_K2::PC_Real || OtherPinCategory == UEdGraphSchema_K2::PC_Text ||
 			(OtherPinCategory == UEdGraphSchema_K2::PC_Byte && !OtherPin->PinType.PinSubCategoryObject.IsValid()) ||
 			OtherPinCategory == UEdGraphSchema_K2::PC_Boolean || OtherPinCategory == UEdGraphSchema_K2::PC_String || OtherPinCategory == UEdGraphSchema_K2::PC_Name || OtherPinCategory == UEdGraphSchema_K2::PC_Object ||
 			OtherPinCategory == UEdGraphSchema_K2::PC_Wildcard || OtherPinCategory == UEdGraphSchema_K2::PC_Int64)
-#else
-		if (OtherPinCategory == UEdGraphSchema_K2::PC_Int || OtherPinCategory == UEdGraphSchema_K2::PC_Float || OtherPinCategory == UEdGraphSchema_K2::PC_Text ||
-			(OtherPinCategory == UEdGraphSchema_K2::PC_Byte && !OtherPin->PinType.PinSubCategoryObject.IsValid()) ||
-            OtherPinCategory == UEdGraphSchema_K2::PC_Boolean || OtherPinCategory == UEdGraphSchema_K2::PC_String || OtherPinCategory == UEdGraphSchema_K2::PC_Name || OtherPinCategory == UEdGraphSchema_K2::PC_Object ||
-			OtherPinCategory == UEdGraphSchema_K2::PC_Wildcard || OtherPinCategory == UEdGraphSchema_K2::PC_Int64 || OtherPinCategory == UEdGraphSchema_K2::PC_Double)
-#endif
         {
 			bIsValidType = true;
 		}
diff --git a/Engine/Source/Editor/BlueprintGraph/Private/K2Node_MakeContainer.cpp b/Engine/Source/Editor/BlueprintGraph/Private/K2Node_MakeContainer.cpp
--- a/Engine/Source/Editor/BlueprintGraph/Private/K2Node_MakeContainer.cpp
+++ b/Engine/Source/Editor/BlueprintGraph/Private/K2Node_MakeContainer.cpp
@@ -48,7 +48,6 @@ void FKCHandler_MakeContainer::Compile(FKismetFunctionContext& Context, UEdGraph
 			{
 				FBPTerminal* RHSTerm = *InputTerm;
 
-#if ENABLE_BLUEPRINT_REAL_NUMBERS
 				{
 					using namespace UE::KismetCompiler;
 
@@ -60,7 +59,6 @@ void FKCHandler_MakeContainer::Compile(FKismetFunctionContext& Context, UEdGraph
 						RHSTerm = ImplicitCastEntry->Get<0>();
 					}
 				}
-#endif
 				RHSTerms.Add(RHSTerm);
 			}
 		}
diff --git a/Engine/Source/Editor/BlueprintGraph/Private/K2Node_MathExpression.cpp b/Engine/Source/Editor/BlueprintGraph/Private/K2Node_MathExpression.cpp
--- a/Engine/Source/Editor/BlueprintGraph/Private/K2Node_MathExpression.cpp
+++ b/Engine/Source/Editor/BlueprintGraph/Private/K2Node_MathExpression.cpp
@@ -79,12 +79,8 @@ static bool PromoteIntToFloat(FEdGraphPinType& InOutType)
 {
 	if (InOutType.PinCategory == UEdGraphSchema_K2::PC_Int)
 	{
-#if ENABLE_BLUEPRINT_REAL_NUMBERS
 		InOutType.PinCategory = UEdGraphSchema_K2::PC_Real;
 		InOutType.PinSubCategory = UEdGraphSchema_K2::PC_Double;
-#else
-		InOutType.PinCategory = UEdGraphSchema_K2::PC_Float;
-#endif
 		InOutType.PinSubCategoryObject = nullptr;
 		return true;
 	}
@@ -1653,12 +1649,8 @@ private:
 			// Create an input pin (using the default guessed type)
 			FEdGraphPinType DefaultType;
 			// currently, generated expressions ALWAYS take a float (it is the most versatile type)
-#if ENABLE_BLUEPRINT_REAL_NUMBERS
 			DefaultType.PinCategory = UEdGraphSchema_K2::PC_Real;
 			DefaultType.PinSubCategory = UEdGraphSchema_K2::PC_Double;
-#else
-			DefaultType.PinCategory = UEdGraphSchema_K2::PC_Float;
-#endif
 			
 			UEdGraphPin* NewInputPin = EntryNode->CreateUserDefinedPin(VariableIdentifier, DefaultType, EGPD_Output);
 			InputPinFragment = MakeShareable(new FCodeGenFragment_InputPin(NewInputPin));
@@ -1741,12 +1733,8 @@ private:
 				LiteralType.PinCategory = UEdGraphSchema_K2::PC_Boolean;
 				break;
 			case CPT_Float:
-#if ENABLE_BLUEPRINT_REAL_NUMBERS
 				LiteralType.PinCategory = UEdGraphSchema_K2::PC_Real;
 				LiteralType.PinSubCategory = UEdGraphSchema_K2::PC_Double;
-#else
-				LiteralType.PinCategory = UEdGraphSchema_K2::PC_Float;
-#endif
 				break;
 			case CPT_Int:
 				LiteralType.PinCategory = UEdGraphSchema_K2::PC_Int;
diff --git a/Engine/Source/Editor/BlueprintGraph/Private/K2Node_Select.cpp b/Engine/Source/Editor/BlueprintGraph/Private/K2Node_Select.cpp
--- a/Engine/Source/Editor/BlueprintGraph/Private/K2Node_Select.cpp
+++ b/Engine/Source/Editor/BlueprintGraph/Private/K2Node_Select.cpp
@@ -148,7 +148,6 @@ public:
 					return;
 				}
 
-#if ENABLE_BLUEPRINT_REAL_NUMBERS
 				{
 					using namespace UE::KismetCompiler;
 
@@ -160,7 +159,6 @@ public:
 						ValueTerm = ImplicitCastEntry->Get<0>();
 					}
 				}
-#endif
 
 				SelectStatement->RHS.Add(ValueTerm);
 			}
diff --git a/Engine/Source/Editor/BlueprintGraph/Private/K2Node_TemporaryVariable.cpp b/Engine/Source/Editor/BlueprintGraph/Private/K2Node_TemporaryVariable.cpp
--- a/Engine/Source/Editor/BlueprintGraph/Private/K2Node_TemporaryVariable.cpp
+++ b/Engine/Source/Editor/BlueprintGraph/Private/K2Node_TemporaryVariable.cpp
@@ -204,15 +204,8 @@ void UK2Node_TemporaryVariable::GetMenuActions(FBlueprintActionDatabaseRegistrar
 	ActionRegistrar.AddBlueprintAction(ActionKey, MakeTempVarNodeSpawner(FEdGraphPinType(UEdGraphSchema_K2::PC_Int, NAME_None, nullptr, EPinContainerType::Array, /*bIsReference =*/ false, /*InValueTerminalType=*/FEdGraphTerminalType()), /*bIsPersistent =*/false));
 	ActionRegistrar.AddBlueprintAction(ActionKey, MakeTempVarNodeSpawner(FEdGraphPinType(UEdGraphSchema_K2::PC_Int64, NAME_None, nullptr, EPinContainerType::None, /*bIsReference =*/ false, /*InValueTerminalType=*/FEdGraphTerminalType()), /*bIsPersistent =*/false));
 	ActionRegistrar.AddBlueprintAction(ActionKey, MakeTempVarNodeSpawner(FEdGraphPinType(UEdGraphSchema_K2::PC_Int64, NAME_None, nullptr, EPinContainerType::Array, /*bIsReference =*/ false, /*InValueTerminalType=*/FEdGraphTerminalType()), /*bIsPersistent =*/false));
-#if ENABLE_BLUEPRINT_REAL_NUMBERS
 	ActionRegistrar.AddBlueprintAction(ActionKey, MakeTempVarNodeSpawner(FEdGraphPinType(UEdGraphSchema_K2::PC_Real, UEdGraphSchema_K2::PC_Double, nullptr, EPinContainerType::None, /*bIsReference =*/ false, /*InValueTerminalType=*/FEdGraphTerminalType()), /*bIsPersistent =*/false));
 	ActionRegistrar.AddBlueprintAction(ActionKey, MakeTempVarNodeSpawner(FEdGraphPinType(UEdGraphSchema_K2::PC_Real, UEdGraphSchema_K2::PC_Double, nullptr, EPinContainerType::Array, /*bIsReference =*/ false, /*InValueTerminalType=*/FEdGraphTerminalType()), /*bIsPersistent =*/false));
-#else
-	ActionRegistrar.AddBlueprintAction(ActionKey, MakeTempVarNodeSpawner(FEdGraphPinType(UEdGraphSchema_K2::PC_Float, NAME_None, nullptr, EPinContainerType::None, /*bIsReference =*/ false, /*InValueTerminalType=*/FEdGraphTerminalType()), /*bIsPersistent =*/false));
-	ActionRegistrar.AddBlueprintAction(ActionKey, MakeTempVarNodeSpawner(FEdGraphPinType(UEdGraphSchema_K2::PC_Float, NAME_None, nullptr, EPinContainerType::Array, /*bIsReference =*/ false, /*InValueTerminalType=*/FEdGraphTerminalType()), /*bIsPersistent =*/false));	
-	ActionRegistrar.AddBlueprintAction(ActionKey, MakeTempVarNodeSpawner(FEdGraphPinType(UEdGraphSchema_K2::PC_Double, NAME_None, nullptr, EPinContainerType::None, /*bIsReference =*/ false, /*InValueTerminalType=*/FEdGraphTerminalType()), /*bIsPersistent =*/false));
-	ActionRegistrar.AddBlueprintAction(ActionKey, MakeTempVarNodeSpawner(FEdGraphPinType(UEdGraphSchema_K2::PC_Double, NAME_None, nullptr, EPinContainerType::Array, /*bIsReference =*/ false, /*InValueTerminalType=*/FEdGraphTerminalType()), /*bIsPersistent =*/false));
-#endif
 	ActionRegistrar.AddBlueprintAction(ActionKey, MakeTempVarNodeSpawner(FEdGraphPinType(UEdGraphSchema_K2::PC_Boolean, NAME_None, nullptr, EPinContainerType::None, /*bIsReference =*/ false, /*InValueTerminalType=*/FEdGraphTerminalType()), /*bIsPersistent =*/false));
 	ActionRegistrar.AddBlueprintAction(ActionKey, MakeTempVarNodeSpawner(FEdGraphPinType(UEdGraphSchema_K2::PC_Boolean, NAME_None, nullptr, EPinContainerType::Array, /*bIsReference =*/ false, /*InValueTerminalType=*/FEdGraphTerminalType()), /*bIsPersistent =*/false));
 	ActionRegistrar.AddBlueprintAction(ActionKey, MakeTempVarNodeSpawner(FEdGraphPinType(UEdGraphSchema_K2::PC_String, NAME_None, nullptr, EPinContainerType::None, /*bIsReference =*/ false, /*InValueTerminalType=*/FEdGraphTerminalType()), /*bIsPersistent =*/false));
diff --git a/Engine/Source/Editor/BlueprintGraph/Private/K2Node_Timeline.cpp b/Engine/Source/Editor/BlueprintGraph/Private/K2Node_Timeline.cpp
--- a/Engine/Source/Editor/BlueprintGraph/Private/K2Node_Timeline.cpp
+++ b/Engine/Source/Editor/BlueprintGraph/Private/K2Node_Timeline.cpp
@@ -44,7 +44,6 @@ public:
 
 		if (NewTimePin && SetNewTimePin)
 		{
-#if ENABLE_BLUEPRINT_REAL_NUMBERS
 			// If SetNewTimePin is linked, then FKismetCompilerContext::ExpandTimelineNodes will create a function
 			// that handles the casting of the "NewTime" parameter for UTimelineComponent::SetNewTime.
 			// If SetNewTimePin is *not* linked, then we don't use "NewTime" at all, and can throw away the cast entry.
@@ -55,7 +54,6 @@ public:
 				using namespace UE::KismetCompiler;
 				CastingUtils::RemoveRegisteredImplicitCast(Context, NewTimePin);
 			}
-#endif
 		}
 		else
 		{
@@ -114,11 +112,7 @@ void UK2Node_Timeline::AllocateDefaultPins()
 
 	CreatePin(EGPD_Input, UEdGraphSchema_K2::PC_Exec, SetNewTimePinName);
 
-#if ENABLE_BLUEPRINT_REAL_NUMBERS
 	UEdGraphPin* NewPositionPin = CreatePin(EGPD_Input, UEdGraphSchema_K2::PC_Real, UEdGraphSchema_K2::PC_Float, NewTimePinName);
-#else
-	UEdGraphPin* NewPositionPin = CreatePin(EGPD_Input, UEdGraphSchema_K2::PC_Float, NewTimePinName);
-#endif
 	K2Schema->SetPinAutogeneratedDefaultValue(NewPositionPin, TEXT("0.0"));
 
 	CreatePin(EGPD_Output, UEdGraphSchema_K2::PC_Byte, FTimeline::GetTimelineDirectionEnum(), DirectionPinName);
@@ -143,11 +137,7 @@ void UK2Node_Timeline::AllocateDefaultPins()
 			else if (TrackId.TrackType == FTTTrackBase::TT_FloatInterp)
 			{
 				const FTTFloatTrack& FloatTrack = Timeline->FloatTracks[TrackId.TrackIndex];
-#if ENABLE_BLUEPRINT_REAL_NUMBERS
 				CreatePin(EGPD_Output, UEdGraphSchema_K2::PC_Real, UEdGraphSchema_K2::PC_Float, FloatTrack.GetTrackName());
-#else
-				CreatePin(EGPD_Output, UEdGraphSchema_K2::PC_Float, FloatTrack.GetTrackName());
-#endif
 			}
 			else if (TrackId.TrackType == FTTTrackBase::TT_VectorInterp)
 			{
diff --git a/Engine/Source/Editor/BlueprintGraph/Private/K2Node_VariableSetRef.cpp b/Engine/Source/Editor/BlueprintGraph/Private/K2Node_VariableSetRef.cpp
--- a/Engine/Source/Editor/BlueprintGraph/Private/K2Node_VariableSetRef.cpp
+++ b/Engine/Source/Editor/BlueprintGraph/Private/K2Node_VariableSetRef.cpp
@@ -31,7 +31,6 @@ public:
 		UEdGraphPin* ValuePin = VarRefNode->GetValuePin();
 		ValidateAndRegisterNetIfLiteral(Context, ValuePin);
 
-#if ENABLE_BLUEPRINT_REAL_NUMBERS
 		{
 			// The pins on UK2Node_VariableSetRef don't actually refer to storage, so there's nothing to register in the NetMap.
 			// However, their pins' nets do, so we need to check if a cast is required for assignment to the target net.
@@ -62,7 +61,6 @@ public:
 				}
 			}
 		}
-#endif
 	}
 
 	virtual void Compile(FKismetFunctionContext& Context, UEdGraphNode* Node) override
@@ -101,7 +99,6 @@ private:
 			FBPTerminal* LHSTerm = *VariableTerm;
 			FBPTerminal* RHSTerm = *ValueTerm;
 
-#if ENABLE_BLUEPRINT_REAL_NUMBERS
 			{
 				using namespace UE::KismetCompiler;
 
@@ -125,7 +122,6 @@ private:
 					CastingUtils::RemoveRegisteredImplicitCast(Context, ValuePin);
 				}
 			}
-#endif
 
 			FBlueprintCompiledStatement& Statement = Context.AppendStatementForNode(Node);
 			Statement.Type = KCST_Assignment;
diff --git a/Engine/Source/Editor/BlueprintGraph/Private/MathExpressionHandler.cpp b/Engine/Source/Editor/BlueprintGraph/Private/MathExpressionHandler.cpp
--- a/Engine/Source/Editor/BlueprintGraph/Private/MathExpressionHandler.cpp
+++ b/Engine/Source/Editor/BlueprintGraph/Private/MathExpressionHandler.cpp
@@ -131,7 +131,6 @@ FBlueprintCompiledStatement* FKCHandler_MathExpression::GenerateFunctionRPN(UEdG
 
 			if (RHSTerm)
 			{
-#if ENABLE_BLUEPRINT_REAL_NUMBERS
 				{
 					const FImplicitCastParams* CastParams =
 						Context.ImplicitCastMap.Find(PinMatch);
@@ -171,7 +170,6 @@ FBlueprintCompiledStatement* FKCHandler_MathExpression::GenerateFunctionRPN(UEdG
 						Context.ImplicitCastMap.Remove(PinMatch);
 					}
 				}
-#endif
 
 				RHSTerms.Add(RHSTerm);
 			}
@@ -287,7 +285,6 @@ void FKCHandler_MathExpression::Compile(FKismetFunctionContext& Context, UEdGrap
 		TArray<FBlueprintCompiledStatement*>& StatementList = Context.StatementsPerNode.FindOrAdd(Node);
 		StatementList.Add(DetachedStatement);
 
-#if ENABLE_BLUEPRINT_REAL_NUMBERS
 		{
 			using namespace UE::KismetCompiler;
 
@@ -304,7 +301,6 @@ void FKCHandler_MathExpression::Compile(FKismetFunctionContext& Context, UEdGrap
 				}
 			}
 		}
-#endif
 	}
 	else
 	{
diff --git a/Engine/Source/Editor/BlueprintGraph/Private/Tests/TypePromotionTests.cpp b/Engine/Source/Editor/BlueprintGraph/Private/Tests/TypePromotionTests.cpp
--- a/Engine/Source/Editor/BlueprintGraph/Private/Tests/TypePromotionTests.cpp
+++ b/Engine/Source/Editor/BlueprintGraph/Private/Tests/TypePromotionTests.cpp
@@ -126,7 +126,6 @@ namespace TypePromoTestUtils
 																				PinArray.Add(InPinName);											\
 																				InPinName->Direction = PinDirection;
 
-#if ENABLE_BLUEPRINT_REAL_NUMBERS
 #define MakeRealNumberTestPin(OwningNode, PinArray, InPinName, PinDirection)	UEdGraphPin* InPinName = UEdGraphPin::CreatePin(OwningNode);		\
 																				InPinName->PinType.PinCategory = UEdGraphSchema_K2::PC_Real;		\
 																				InPinName->PinType.PinSubCategory = UEdGraphSchema_K2::PC_Double;	\
@@ -137,15 +136,7 @@ namespace TypePromoTestUtils
 	MakeRealNumberTestPin(OwningNode, OutArray, RealPinA, EGPD_Input);									\
 	MakeRealNumberTestPin(OwningNode, OutArray, RealPinB, EGPD_Input);									\
 	MakeRealNumberTestPin(OwningNode, OutArray, RealOutputPin, EGPD_Output);		
-#else
-#define MakeRealNumberTestPins(OwningNode, OutArray)													\
-	MakeTestPin(OwningNode, OutArray, FloatPinA, UEdGraphSchema_K2::PC_Float, EGPD_Input);				\
-	MakeTestPin(OwningNode, OutArray, FloatPinB, UEdGraphSchema_K2::PC_Float, EGPD_Input);				\
-	MakeTestPin(OwningNode, OutArray, FloatOutputPin, UEdGraphSchema_K2::PC_Float, EGPD_Output);		\
-	MakeTestPin(OwningNode, OutArray, DoublePinA, UEdGraphSchema_K2::PC_Double, EGPD_Output);			\
-	MakeTestPin(OwningNode, OutArray, DoublePinB, UEdGraphSchema_K2::PC_Double, EGPD_Output);			\
-	MakeTestPin(OwningNode, OutArray, DoubleOutputPin, UEdGraphSchema_K2::PC_Double, EGPD_Output);
-#endif
+
 
 #define MakeTestPins(OwningNode, OutArray)																\
 	MakeRealNumberTestPins(OwningNode, OutArray)														\
@@ -177,7 +168,6 @@ namespace TypePromoTestUtils
 IMPLEMENT_SIMPLE_AUTOMATION_TEST(FTypePromotionTest, "Blueprints.Compiler.TypePromotion", EAutomationTestFlags::EditorContext | EAutomationTestFlags::SmokeFilter)
 bool FTypePromotionTest::RunTest(const FString& Parameters)
 {
-#if ENABLE_BLUEPRINT_REAL_NUMBERS
 	FEdGraphPinType RealPin = {};		RealPin.PinCategory = UEdGraphSchema_K2::PC_Real; RealPin.PinSubCategory = UEdGraphSchema_K2::PC_Double;
 	FEdGraphPinType IntPin = {};		IntPin.PinCategory = UEdGraphSchema_K2::PC_Int;
 	FEdGraphPinType Int64Pin = {};		Int64Pin.PinCategory = UEdGraphSchema_K2::PC_Int64;
@@ -204,40 +194,6 @@ bool FTypePromotionTest::RunTest(const FString& Parameters)
 	// Test promotions that should not happen
 	TestEqual(TEXT("Testing int64 cannot go to byte"), FTypePromotion::GetHigherType(Int64Pin, BytePin), FTypePromotion::ETypeComparisonResult::TypeAHigher);
 	TestEqual(TEXT("Testing int64 cannot go to int"), FTypePromotion::GetHigherType(Int64Pin, IntPin), FTypePromotion::ETypeComparisonResult::TypeAHigher);
-#else
-	FEdGraphPinType DoublePin = {};		DoublePin.PinCategory = UEdGraphSchema_K2::PC_Double;
-	FEdGraphPinType FloatPin = {};		FloatPin.PinCategory = UEdGraphSchema_K2::PC_Float;
-	FEdGraphPinType IntPin = {};		IntPin.PinCategory = UEdGraphSchema_K2::PC_Int;
-	FEdGraphPinType Int64Pin = {};		Int64Pin.PinCategory = UEdGraphSchema_K2::PC_Int64;
-	FEdGraphPinType BytePin = {};		BytePin.PinCategory = UEdGraphSchema_K2::PC_Byte;
-	FEdGraphPinType VecPin = {};		VecPin.PinCategory = UEdGraphSchema_K2::PC_Struct; VecPin.PinSubCategoryObject = TBaseStructure<FVector>::Get();
-
-	// Test promotions that should happen
-	TestEqual(TEXT("Testing float to double"), FTypePromotion::GetHigherType(FloatPin, DoublePin), FTypePromotion::ETypeComparisonResult::TypeBHigher);
-	TestEqual(TEXT("Testing float to vector"), FTypePromotion::GetHigherType(FloatPin, VecPin), FTypePromotion::ETypeComparisonResult::TypeBHigher);
-
-	TestEqual(TEXT("Testing int to float"), FTypePromotion::GetHigherType(IntPin, FloatPin), FTypePromotion::ETypeComparisonResult::TypeBHigher);
-	TestEqual(TEXT("Testing int to double"), FTypePromotion::GetHigherType(IntPin, DoublePin), FTypePromotion::ETypeComparisonResult::TypeBHigher);
-	TestEqual(TEXT("Testing int to int64"), FTypePromotion::GetHigherType(IntPin, Int64Pin), FTypePromotion::ETypeComparisonResult::TypeBHigher);
-
-	TestEqual(TEXT("Testing Byte to int"), FTypePromotion::GetHigherType(BytePin, IntPin), FTypePromotion::ETypeComparisonResult::TypeBHigher);
-	TestEqual(TEXT("Testing Byte to int64"), FTypePromotion::GetHigherType(BytePin, Int64Pin), FTypePromotion::ETypeComparisonResult::TypeBHigher);
-
-	TestEqual(TEXT("Testing Double to int64"), FTypePromotion::GetHigherType(DoublePin, Int64Pin), FTypePromotion::ETypeComparisonResult::TypeBHigher);
-
-	// Test Equality of pins
-	TestEqual(TEXT("Testing Byte == Byte"), FTypePromotion::GetHigherType(BytePin, BytePin), FTypePromotion::ETypeComparisonResult::TypesEqual);
-	TestEqual(TEXT("Testing float == float"), FTypePromotion::GetHigherType(FloatPin, FloatPin), FTypePromotion::ETypeComparisonResult::TypesEqual);
-	TestEqual(TEXT("Testing double == double"), FTypePromotion::GetHigherType(DoublePin, DoublePin), FTypePromotion::ETypeComparisonResult::TypesEqual);
-	TestEqual(TEXT("Testing int == int"), FTypePromotion::GetHigherType(IntPin, IntPin), FTypePromotion::ETypeComparisonResult::TypesEqual);
-	TestEqual(TEXT("Testing int64 == int64"), FTypePromotion::GetHigherType(Int64Pin, Int64Pin), FTypePromotion::ETypeComparisonResult::TypesEqual);
-
-
-	// Test promotions that should not happen
-	TestEqual(TEXT("Testing int64 cannot go to byte"), FTypePromotion::GetHigherType(Int64Pin, BytePin), FTypePromotion::ETypeComparisonResult::TypeAHigher);
-	TestEqual(TEXT("Testing int64 cannot go to int"), FTypePromotion::GetHigherType(Int64Pin, IntPin), FTypePromotion::ETypeComparisonResult::TypeAHigher);
-	TestEqual(TEXT("Testing int64 cannot go to float"), FTypePromotion::GetHigherType(Int64Pin, FloatPin), FTypePromotion::ETypeComparisonResult::TypeAHigher);
-#endif
 
 	return true;
 }
@@ -283,7 +239,6 @@ bool FFindBestMatchingFunc::RunTest(const FString& Parameters)
 		TestMatchingFunc(TEXT("Add"), TestPins, TEXT("Add_Vector2DVector2D"));
 	}
 
-#if ENABLE_BLUEPRINT_REAL_NUMBERS
 	// Multiply_VectorVector given A real input, vector input, and a vector output
 	{
 		TArray<UEdGraphPin*> TestPins =
@@ -366,98 +321,6 @@ bool FFindBestMatchingFunc::RunTest(const FString& Parameters)
 		};
 		TestMatchingFunc(TEXT("Greater"), TestPins, TEXT("Greater_DoubleDouble"));
 	}
-#else
-	// Multiply_VectorVector given A float input, vector input, and a vector output
-	{
-		TArray<UEdGraphPin*> TestPins =
-		{
-			FloatPinA, VecInputPinB, VecOutputPinA,
-		};
-		TestMatchingFunc(TEXT("Multiply"), TestPins, TEXT("Multiply_VectorVector"));
-	}
-
-	// Multiply_VectorVector given a float, vector, float
-	// Order shouldn't matter when passing these pins in, which is what we are testing here
-	{
-		TArray<UEdGraphPin*> TestPins =
-		{
-			FloatPinA, VecOutputPinA, VecInputPinA,
-		};
-		TestMatchingFunc(TEXT("Multiply"), TestPins, TEXT("Multiply_VectorVector"));
-	}
-
-	// Multiply_VectorVector given two vector inputs and a vector output
-	{
-		TArray<UEdGraphPin*> TestPins =
-		{
-			VecInputPinA, VecInputPinB, VecOutputPinA,
-		};
-		TestMatchingFunc(TEXT("Multiply"), TestPins, TEXT("Multiply_VectorVector"));
-	}
-
-	// Add_DoubleDouble
-	{
-		TArray<UEdGraphPin*> TestPins =
-		{
-			DoublePinA, DoublePinB, DoubleOutputPin
-		};
-		TestMatchingFunc(TEXT("Add"), TestPins, TEXT("Add_DoubleDouble"));
-	}
-
-	{
-		TArray<UEdGraphPin*> TestPins =
-		{
-			DoublePinA, FloatPinA, DoubleOutputPin
-		};
-		TestMatchingFunc(TEXT("Add"), TestPins, TEXT("Add_DoubleDouble"));
-	}
-
-	// Subtract_FloatFloat
-	{
-		TArray<UEdGraphPin*> TestPins =
-		{
-			FloatPinA, FloatPinB, FloatOutputPin
-		};
-		TestMatchingFunc(TEXT("Subtract"), TestPins, TEXT("Subtract_FloatFloat"));
-	}
-
-	// Add_FloatFloat given only one float pin. This simulates the first connection being made to a 
-	// promotable operator, in which case we should default to a regular old Float + Float
-	{
-		TArray<UEdGraphPin*> TestPins =
-		{
-			FloatPinA,
-		};
-		TestMatchingFunc(TEXT("Add"), TestPins, TEXT("Add_FloatFloat"));
-	}
-
-	// Less_FloatFloat Given a Float and Double
-	{
-		TArray<UEdGraphPin*> TestPins =
-		{
-			FloatPinA, BoolOutputPin
-		};
-		TestMatchingFunc(TEXT("Less"), TestPins, TEXT("Less_FloatFloat"));
-	}
-
-	// Less_FloatFloat Given just a single float
-	{
-		TArray<UEdGraphPin*> TestPins =
-		{
-			FloatPinA,
-		};
-		TestMatchingFunc(TEXT("Less"), TestPins, TEXT("Less_FloatFloat"));
-	}
-
-	// Greater_DoubleDouble Given a Float and Double
-	{
-		TArray<UEdGraphPin*> TestPins =
-		{
-			FloatPinA, DoublePinA
-		};
-		TestMatchingFunc(TEXT("Greater"), TestPins, TEXT("Greater_DoubleDouble"));
-	}
-#endif
 	
 	TypePromoTestUtils::CleanupTestPins(PinTypes);
 	TestNode->MarkAsGarbage();
@@ -765,7 +628,6 @@ bool FPromotableOpNodeAddPinInterface::RunTest(const FString& Parameters)
 		UEdGraphPin* AdditonalPin = AddNode->GetAdditionalPin(2);
 		TestNotNull(TEXT("Additional Pin is not null"), AdditonalPin);
 		
-#if ENABLE_BLUEPRINT_REAL_NUMBERS
 		// Connect a real pin to the additional input pin
 		const bool bConnected = TypePromoTestUtils::TestPromotedConnection(AdditonalPin, RealOutputPin);
 		TestTrue(TEXT("Connection to additional pin success"), bConnected);
@@ -774,16 +636,6 @@ bool FPromotableOpNodeAddPinInterface::RunTest(const FString& Parameters)
 		TestTrue(TEXT("Top Pin type propegates to new connection"), TopInputPin->PinType.PinCategory == RealOutputPin->PinType.PinCategory);
 		TestTrue(TEXT("Bottom Pin type propegates to new connection"), BottomInputPin->PinType.PinCategory == RealOutputPin->PinType.PinCategory);
 		TestTrue(TEXT("Out Pin type propegates to new connection"), OutputPin->PinType.PinCategory == RealOutputPin->PinType.PinCategory);
-#else
-		// Connect a float pin to the additional input pin
-		const bool bConnected = TypePromoTestUtils::TestPromotedConnection(AdditonalPin, FloatOutputPin);
-		TestTrue(TEXT("Connection to additional pin success"), bConnected);
-
-		// The other pins have propagated correctly with this new connection 
-		TestTrue(TEXT("Top Pin type propegates to new connection"), TopInputPin->PinType.PinCategory == FloatOutputPin->PinType.PinCategory);
-		TestTrue(TEXT("Bottom Pin type propegates to new connection"), BottomInputPin->PinType.PinCategory == FloatOutputPin->PinType.PinCategory);
-		TestTrue(TEXT("Out Pin type propegates to new connection"), OutputPin->PinType.PinCategory == FloatOutputPin->PinType.PinCategory);
-#endif
 
 		// Removing the only pin with a connection with reset the node to wildcard
 		AddNode->RemoveInputPin(AdditonalPin);
@@ -843,17 +695,10 @@ bool FPromotableOperatorConnectionChanged::RunTest(const FString& Parameters)
 
 		check(TopInputPin && BottomInputPin && OutputPin);
 
-#if ENABLE_BLUEPRINT_REAL_NUMBERS
 		const bool bConnected = K2Schema->TryCreateConnection(TopInputPin, RealOutputPin);
 		AddNode->NotifyPinConnectionListChanged(TopInputPin);
 
 		TestTrue(TEXT("Bottom Pin type propagates to real"), bConnected && BottomInputPin->PinType.PinCategory == RealPinB->PinType.PinCategory);
-#else
-		const bool bConnected = K2Schema->TryCreateConnection(TopInputPin, FloatOutputPin);
-		AddNode->NotifyPinConnectionListChanged(TopInputPin);
-
-		TestTrue(TEXT("Bottom Pin type propagates to float"), bConnected && BottomInputPin->PinType.PinCategory == FloatPinB->PinType.PinCategory);
-#endif
 	}
 
 	// Connecting a vector output should make the other input be a vector as well
diff --git a/Engine/Source/Editor/GraphEditor/Private/NodeFactory.cpp b/Engine/Source/Editor/GraphEditor/Private/NodeFactory.cpp
--- a/Engine/Source/Editor/GraphEditor/Private/NodeFactory.cpp
+++ b/Engine/Source/Editor/GraphEditor/Private/NodeFactory.cpp
@@ -325,21 +325,10 @@ TSharedPtr<SGraphPin> FNodeFactory::CreateK2PinWidget(UEdGraphPin* InPin)
 	{
 		return SNew(SGraphPinNum<int64>, InPin);
 	}
-#if ENABLE_BLUEPRINT_REAL_NUMBERS
 	else if (InPin->PinType.PinCategory == UEdGraphSchema_K2::PC_Real)
 	{
 		return SNew(SGraphPinNum<double>, InPin);
 	}
-#else
-	else if (InPin->PinType.PinCategory == UEdGraphSchema_K2::PC_Float)
-	{
-		return SNew(SGraphPinNum<float>, InPin);
-	}
-	else if (InPin->PinType.PinCategory == UEdGraphSchema_K2::PC_Double)
-	{
-		return SNew(SGraphPinNum<double>, InPin);
-	}
-#endif
 	else if (InPin->PinType.PinCategory == UEdGraphSchema_K2::PC_String || InPin->PinType.PinCategory == UEdGraphSchema_K2::PC_Name)
 	{
 		return SNew(SGraphPinString, InPin);
diff --git a/Engine/Source/Editor/KismetCompiler/Private/KismetCastingUtils.cpp b/Engine/Source/Editor/KismetCompiler/Private/KismetCastingUtils.cpp
--- a/Engine/Source/Editor/KismetCompiler/Private/KismetCastingUtils.cpp
+++ b/Engine/Source/Editor/KismetCompiler/Private/KismetCastingUtils.cpp
@@ -5,8 +5,6 @@
 #include "KismetCompiledFunctionContext.h"
 #include "KismetCompilerMisc.h"
 
-#if ENABLE_BLUEPRINT_REAL_NUMBERS
-
 namespace UE::KismetCompiler::CastingUtils::Private
 {
 
@@ -420,5 +418,3 @@ TOptional<StatementNamePair> GetFloatingPointConversionType(const UEdGraphPin& S
 
 } // namespace UE::KismetCompiler::CastingUtils
 
-#endif
-
diff --git a/Engine/Source/Editor/KismetCompiler/Private/KismetCompiler.cpp b/Engine/Source/Editor/KismetCompiler/Private/KismetCompiler.cpp
--- a/Engine/Source/Editor/KismetCompiler/Private/KismetCompiler.cpp
+++ b/Engine/Source/Editor/KismetCompiler/Private/KismetCompiler.cpp
@@ -125,7 +125,6 @@ namespace UE::KismetCompiler::Private
 		}
 	}
 
-#if ENABLE_BLUEPRINT_REAL_NUMBERS
 	// When we change pins back to real/float types, we also expect that the series of pin types matches that of the UFunction.
 	// If there's ever a discrepancy, then there's not much we can do other than log a warning.
 	// It typically means that the BP is in a bad state, which prevents us from deducing the corresponding pin.
@@ -436,7 +435,6 @@ namespace UE::KismetCompiler::Private
 			}
 		}
 	}
-#endif
 }
 
 //////////////////////////////////////////////////////////////////////////
@@ -1172,11 +1170,8 @@ void FKismetCompilerContext::CreateClassVariablesFromBlueprint()
 		FEdGraphPinType DirectionPinType(UEdGraphSchema_K2::PC_Byte, NAME_None, FTimeline::GetTimelineDirectionEnum(), EPinContainerType::None, false, FEdGraphTerminalType());
 		CreateVariable(Timeline->GetDirectionPropertyName(), DirectionPinType);
 
-#if ENABLE_BLUEPRINT_REAL_NUMBERS
 		FEdGraphPinType FloatPinType(UEdGraphSchema_K2::PC_Real, UEdGraphSchema_K2::PC_Float, nullptr, EPinContainerType::None, false, FEdGraphTerminalType());
-#else
-		FEdGraphPinType FloatPinType(UEdGraphSchema_K2::PC_Float, NAME_None, nullptr, EPinContainerType::None, false, FEdGraphTerminalType());
-#endif
+
 		for (const FTTFloatTrack& FloatTrack : Timeline->FloatTracks)
 		{
 			CreateVariable(FloatTrack.GetPropertyName(), FloatPinType);
@@ -2378,13 +2373,11 @@ void FKismetCompilerContext::PrecompileFunction(FKismetFunctionContext& Context,
 
 void FKismetCompilerContext::PreCompileUpdateBlueprintOnLoad(UBlueprint* BP)
 {
-#if ENABLE_BLUEPRINT_REAL_NUMBERS
 	check(BP);
 	if (BP->GetLinkerCustomVersion(FUE5ReleaseStreamObjectVersion::GUID) < FUE5ReleaseStreamObjectVersion::BlueprintPinsUseRealNumbers)
 	{
 		UE::KismetCompiler::Private::MatchNodesToDelegateSignatures(BP);
 	}
-#endif
 }
 
 /** Inserts a new item into an array in a sorted position; using an externally stored sort index map */
@@ -3624,11 +3617,9 @@ void FKismetCompilerContext::CreateLocalsAndRegisterNets(FKismetFunctionContext&
 		}
 	}
 
-#if ENABLE_BLUEPRINT_REAL_NUMBERS
 	using namespace UE::KismetCompiler;
 
 	CastingUtils::RegisterImplicitCasts(Context);
-#endif
 
 	// Create net variable declarations
 	CreateLocalVariablesForFunction(Context, FunctionPropertyStorageLocation);
diff --git a/Engine/Source/Editor/KismetCompiler/Private/KismetCompilerMisc.cpp b/Engine/Source/Editor/KismetCompiler/Private/KismetCompilerMisc.cpp
--- a/Engine/Source/Editor/KismetCompiler/Private/KismetCompilerMisc.cpp
+++ b/Engine/Source/Editor/KismetCompiler/Private/KismetCompilerMisc.cpp
@@ -121,7 +121,6 @@ static bool DoesTypeNotMatchProperty(UEdGraphPin* SourcePin, const FEdGraphPinTy
 			}
 		}
 	}
-#if ENABLE_BLUEPRINT_REAL_NUMBERS
 	else if (PinCategory == UEdGraphSchema_K2::PC_Real)
 	{
 		if (PinSubCategory == UEdGraphSchema_K2::PC_Float)
@@ -140,18 +139,6 @@ static bool DoesTypeNotMatchProperty(UEdGraphPin* SourcePin, const FEdGraphPinTy
 			bTypeMismatch = true;
 		}
 	}
-#else
-	else if (PinCategory == UEdGraphSchema_K2::PC_Float)
-	{
-		FFloatProperty* SpecificProperty = CastField<FFloatProperty>(TestProperty);
-		bTypeMismatch = (SpecificProperty == nullptr);
-	}
-	else if (PinCategory == UEdGraphSchema_K2::PC_Double)
-	{
-		FDoubleProperty* SpecificProperty = CastField<FDoubleProperty>(TestProperty);
-		bTypeMismatch = (SpecificProperty == nullptr);
-	}
-#endif
 	else if (PinCategory == UEdGraphSchema_K2::PC_Int)
 	{
 		FIntProperty* SpecificProperty = CastField<FIntProperty>(TestProperty);
@@ -902,7 +889,6 @@ void FKismetCompilerUtilities::CreateObjectAssignmentStatement(FKismetFunctionCo
 	{
 		FBPTerminal* RHSTerm = SrcTerm;
 
-#if ENABLE_BLUEPRINT_REAL_NUMBERS
 		using namespace UE::KismetCompiler;
 
 		FBPTerminal* ImplicitCastTerm = nullptr;
@@ -924,7 +910,6 @@ void FKismetCompilerUtilities::CreateObjectAssignmentStatement(FKismetFunctionCo
 		{
 			RHSTerm = ImplicitCastTerm;
 		}
-#endif
 
 		FBlueprintCompiledStatement& Statement = Context.AppendStatementForNode(Node);
 		Statement.Type = KCST_Assignment;
@@ -1086,7 +1071,6 @@ FProperty* FKismetCompilerUtilities::CreatePrimitiveProperty(FFieldVariant Prope
 		NewProperty = new FInt64Property(PropertyScope, ValidatedPropertyName, ObjectFlags);
 		NewProperty->SetPropertyFlags(CPF_HasGetValueTypeHash);
 	}
-#if ENABLE_BLUEPRINT_REAL_NUMBERS
 	else if (PinCategory == UEdGraphSchema_K2::PC_Real)
 	{
 		if (PinSubCategory == UEdGraphSchema_K2::PC_Float)
@@ -1104,18 +1088,6 @@ FProperty* FKismetCompilerUtilities::CreatePrimitiveProperty(FFieldVariant Prope
 			checkf(false, TEXT("Erroneous pin subcategory for PC_Real: %s"), *PinSubCategory.ToString());
 		}
 	}
-#else
-	else if (PinCategory == UEdGraphSchema_K2::PC_Float)
-	{
-		NewProperty = new FFloatProperty(PropertyScope, ValidatedPropertyName, ObjectFlags);
-		NewProperty->SetPropertyFlags(CPF_HasGetValueTypeHash);
-	}
-	else if (PinCategory == UEdGraphSchema_K2::PC_Double)
-	{
-		NewProperty = new FDoubleProperty(PropertyScope, ValidatedPropertyName, ObjectFlags);
-		NewProperty->SetPropertyFlags(CPF_HasGetValueTypeHash);
-	}
-#endif
 	else if (PinCategory == UEdGraphSchema_K2::PC_Boolean)
 	{
 		FBoolProperty* BoolProperty = new FBoolProperty(PropertyScope, ValidatedPropertyName, ObjectFlags);
diff --git a/Engine/Source/Editor/KismetCompiler/Private/KismetCompilerVMBackend.cpp b/Engine/Source/Editor/KismetCompiler/Private/KismetCompilerVMBackend.cpp
--- a/Engine/Source/Editor/KismetCompiler/Private/KismetCompilerVMBackend.cpp
+++ b/Engine/Source/Editor/KismetCompiler/Private/KismetCompilerVMBackend.cpp
@@ -478,7 +478,6 @@ public:
 			return Type && (Type->PinCategory == UEdGraphSchema_K2::PC_Text);
 		}
 
-#if ENABLE_BLUEPRINT_REAL_NUMBERS
 		static bool IsFloat(const FEdGraphPinType* Type, const FProperty* Property)
 		{
 			if (Property)
@@ -496,25 +495,6 @@ public:
 			}
 			return Type && (Type->PinCategory == UEdGraphSchema_K2::PC_Real) && (Type->PinSubCategory == UEdGraphSchema_K2::PC_Double);
 		}
-#else
-		static bool IsFloat(const FEdGraphPinType* Type, const FProperty* Property)
-		{
-			if (Property)
-			{
-				return Property->IsA<FFloatProperty>();
-			}
-			return Type && (Type->PinCategory == UEdGraphSchema_K2::PC_Float);
-		}
-
-		static bool IsDouble(const FEdGraphPinType* Type, const FProperty* Property)
-		{
-			if (Property)
-			{
-				return Property->IsA<FDoubleProperty>();
-			}
-			return Type && (Type->PinCategory == UEdGraphSchema_K2::PC_Double);
-		}
-#endif
 
 		static bool IsInt(const FEdGraphPinType* Type, const FProperty* Property)
 		{
diff --git a/Engine/Source/Editor/UMGEditor/Private/WidgetBlueprint.cpp b/Engine/Source/Editor/UMGEditor/Private/WidgetBlueprint.cpp
--- a/Engine/Source/Editor/UMGEditor/Private/WidgetBlueprint.cpp
+++ b/Engine/Source/Editor/UMGEditor/Private/WidgetBlueprint.cpp
@@ -584,7 +584,7 @@ void UWidgetBlueprint::ReplaceDeprecatedNodes()
 		}
 	}
 
-#if WITH_EDITORONLY_DATA && ENABLE_BLUEPRINT_REAL_NUMBERS
+#if WITH_EDITORONLY_DATA
 	if (GetLinkerCustomVersion(FUE5ReleaseStreamObjectVersion::GUID) < FUE5ReleaseStreamObjectVersion::BlueprintPinsUseRealNumbers)
 	{
 		// Revert any overzealous PC_Float to PC_Real/PC_Double conversions.
diff --git a/Engine/Source/Editor/UnrealEd/Private/UserDefinedStructEditorData.cpp b/Engine/Source/Editor/UnrealEd/Private/UserDefinedStructEditorData.cpp
--- a/Engine/Source/Editor/UnrealEd/Private/UserDefinedStructEditorData.cpp
+++ b/Engine/Source/Editor/UnrealEd/Private/UserDefinedStructEditorData.cpp
@@ -33,7 +33,6 @@ void FStructVariableDescription::PostSerialize(const FArchive& Ar)
 		}
 	}
 
-#if ENABLE_BLUEPRINT_REAL_NUMBERS
 	bool FixupPinCategories =
 		Ar.IsLoading() &&
 		(Ar.CustomVer(FUE5ReleaseStreamObjectVersion::GUID) < FUE5ReleaseStreamObjectVersion::BlueprintPinsUseRealNumbers) &&
@@ -44,7 +43,6 @@ void FStructVariableDescription::PostSerialize(const FArchive& Ar)
 		Category = TEXT("real");
 		SubCategory = TEXT("double");
 	}
-#endif
 }
 
 bool FStructVariableDescription::SetPinType(const FEdGraphPinType& VarType)
@@ -84,9 +82,7 @@ void UUserDefinedStructEditorData::Serialize(FArchive& Ar)
 {
 	Super::Serialize(Ar);
 
-#if ENABLE_BLUEPRINT_REAL_NUMBERS
 	Ar.UsingCustomVersion(FUE5ReleaseStreamObjectVersion::GUID);
-#endif
 }
 
 void UUserDefinedStructEditorData::PostUndo(bool bSuccess)
diff --git a/Engine/Source/Runtime/Core/Public/Misc/Build.h b/Engine/Source/Runtime/Core/Public/Misc/Build.h
--- a/Engine/Source/Runtime/Core/Public/Misc/Build.h
+++ b/Engine/Source/Runtime/Core/Public/Misc/Build.h
@@ -364,11 +364,6 @@
 /** Enable fast calls for event thunks into an event graph that have no parameters  */
 #define UE_BLUEPRINT_EVENTGRAPH_FASTCALLS 1
 
-/** Toggles the use of real numbers in blueprints. Explicit float and double usage will not be available. */
-#ifndef ENABLE_BLUEPRINT_REAL_NUMBERS
-#define ENABLE_BLUEPRINT_REAL_NUMBERS 1
-#endif
-
 /** Enable perf counters on dedicated servers */
 #define USE_SERVER_PERF_COUNTERS ((UE_SERVER || UE_EDITOR) && WITH_PERFCOUNTERS)
 
diff --git a/Engine/Source/Runtime/Core/Public/UObject/UE5ReleaseStreamObjectVersion.h b/Engine/Source/Runtime/Core/Public/UObject/UE5ReleaseStreamObjectVersion.h
--- a/Engine/Source/Runtime/Core/Public/UObject/UE5ReleaseStreamObjectVersion.h
+++ b/Engine/Source/Runtime/Core/Public/UObject/UE5ReleaseStreamObjectVersion.h
@@ -96,10 +96,8 @@ struct CORE_API FUE5ReleaseStreamObjectVersion
 		// Large Worlds - serialize double types as doubles
 		LargeWorldCoordinates,
 		
-#if ENABLE_BLUEPRINT_REAL_NUMBERS
 		// Deserialize old BP float&double types as real numbers for pins
 		BlueprintPinsUseRealNumbers,
-#endif
 		
 		// Changed shadow defaults for directional light components, version needed to not affect old things
 		UpdatedDirectionalLightShadowDefaults,
diff --git a/Engine/Source/Runtime/Engine/Private/Blueprint.cpp b/Engine/Source/Runtime/Engine/Private/Blueprint.cpp
--- a/Engine/Source/Runtime/Engine/Private/Blueprint.cpp
+++ b/Engine/Source/Runtime/Engine/Private/Blueprint.cpp
@@ -430,9 +430,7 @@ void UBlueprint::Serialize(FArchive& Ar)
 	Super::Serialize(Ar);
 
 #if WITH_EDITORONLY_DATA
-#if ENABLE_BLUEPRINT_REAL_NUMBERS
 	Ar.UsingCustomVersion(FUE5ReleaseStreamObjectVersion::GUID);
-#endif
 
 	if(Ar.IsLoading() && Ar.UEVer() < VER_UE4_BLUEPRINT_VARS_NOT_READ_ONLY)
 	{
@@ -882,10 +880,8 @@ bool UBlueprint::AlwaysCompileOnLoad() const
 {
 	bool bShouldCompile = false;
 
-#if ENABLE_BLUEPRINT_REAL_NUMBERS
 	bShouldCompile = 
 		GetLinkerCustomVersion(FUE5ReleaseStreamObjectVersion::GUID) <  FUE5ReleaseStreamObjectVersion::BlueprintPinsUseRealNumbers;
-#endif
 
 	return bShouldCompile;
 }
diff --git a/Engine/Source/Runtime/Engine/Private/EdGraph/EdGraphNode.cpp b/Engine/Source/Runtime/Engine/Private/EdGraph/EdGraphNode.cpp
--- a/Engine/Source/Runtime/Engine/Private/EdGraph/EdGraphNode.cpp
+++ b/Engine/Source/Runtime/Engine/Private/EdGraph/EdGraphNode.cpp
@@ -39,10 +39,7 @@ FArchive& operator<<(FArchive& Ar, FEdGraphTerminalType& T)
 {
 	Ar.UsingCustomVersion(FFrameworkObjectVersion::GUID);
 	Ar.UsingCustomVersion(FReleaseObjectVersion::GUID);
-
-#if ENABLE_BLUEPRINT_REAL_NUMBERS
 	Ar.UsingCustomVersion(FUE5ReleaseStreamObjectVersion::GUID);
-#endif
 
 	if (Ar.CustomVer(FFrameworkObjectVersion::GUID) >= FFrameworkObjectVersion::PinsStoreFName)
 	{
@@ -94,7 +91,6 @@ FArchive& operator<<(FArchive& Ar, FEdGraphTerminalType& T)
 		Ar << T.bTerminalIsUObjectWrapper;
 	}
 
-#if ENABLE_BLUEPRINT_REAL_NUMBERS
 	if (Ar.IsLoading())
 	{
 		bool bFixupPinCategories =
@@ -107,7 +103,6 @@ FArchive& operator<<(FArchive& Ar, FEdGraphTerminalType& T)
 			T.TerminalSubCategory = TEXT("double");
 		}
 	}
-#endif
 
 	return Ar;
 }
diff --git a/Engine/Source/Runtime/Engine/Private/EdGraph/EdGraphPin.cpp b/Engine/Source/Runtime/Engine/Private/EdGraph/EdGraphPin.cpp
--- a/Engine/Source/Runtime/Engine/Private/EdGraph/EdGraphPin.cpp
+++ b/Engine/Source/Runtime/Engine/Private/EdGraph/EdGraphPin.cpp
@@ -161,10 +161,7 @@ bool FEdGraphPinType::Serialize(FArchive& Ar)
 
 	Ar.UsingCustomVersion(FFrameworkObjectVersion::GUID);
 	Ar.UsingCustomVersion(FReleaseObjectVersion::GUID);
-
-#if ENABLE_BLUEPRINT_REAL_NUMBERS
 	Ar.UsingCustomVersion(FUE5ReleaseStreamObjectVersion::GUID);
-#endif
 
 	if (Ar.CustomVer(FFrameworkObjectVersion::GUID) >= FFrameworkObjectVersion::PinsStoreFName)
 	{
@@ -303,7 +300,6 @@ bool FEdGraphPinType::Serialize(FArchive& Ar)
 			}
 		}
 
-#if ENABLE_BLUEPRINT_REAL_NUMBERS
 		bool bFixupPinCategories =
 			(Ar.CustomVer(FUE5ReleaseStreamObjectVersion::GUID) < FUE5ReleaseStreamObjectVersion::BlueprintPinsUseRealNumbers) &&
 			((PinCategory == TEXT("double")) || (PinCategory == TEXT("float")));
@@ -313,8 +309,6 @@ bool FEdGraphPinType::Serialize(FArchive& Ar)
 			PinCategory = TEXT("real");
 			PinSubCategory = TEXT("double");
 		}
-#endif
-
 #endif
 
 		bIsUObjectWrapper = bIsUObjectWrapperBool;
