Control Rig: Added BlueprintInternalUseOnly="true" to the basel rig unit struct so that all derived rig unit structs don't appear in the pin type drop down menu and create confusion

diff --git a/Engine/Plugins/Animation/ControlRig/Source/ControlRig/Public/Units/RigUnit.h b/Engine/Plugins/Animation/ControlRig/Source/ControlRig/Public/Units/RigUnit.h
--- a/Engine/Plugins/Animation/ControlRig/Source/ControlRig/Public/Units/RigUnit.h
+++ b/Engine/Plugins/Animation/ControlRig/Source/ControlRig/Public/Units/RigUnit.h
@@ -10,7 +10,7 @@
 struct FRigUnitContext;
 
 /** Base class for all rig units */
-USTRUCT(BlueprintType, meta=(Abstract, NodeColor = "0.1 0.1 0.1"))
+USTRUCT(BlueprintType, meta=(Abstract, NodeColor = "0.1 0.1 0.1", BlueprintInternalUseOnly="true"))
 struct CONTROLRIG_API FRigUnit : public FRigVMStruct
 {
 	GENERATED_BODY()
diff --git a/Engine/Source/Editor/BlueprintGraph/Private/EdGraphSchema_K2.cpp b/Engine/Source/Editor/BlueprintGraph/Private/EdGraphSchema_K2.cpp
--- a/Engine/Source/Editor/BlueprintGraph/Private/EdGraphSchema_K2.cpp
+++ b/Engine/Source/Editor/BlueprintGraph/Private/EdGraphSchema_K2.cpp
@@ -1267,7 +1267,7 @@ bool UEdGraphSchema_K2::IsAllowableBlueprintVariableType(const UScriptStruct* In
 		// User-defined structs are always allowed as BP variable types.
 		return true;
 	}
-	return (InStruct && InStruct->GetBoolMetaDataHierarchical(FBlueprintMetadata::MD_AllowableBlueprintVariableType) && (bForInternalUse || !InStruct->GetBoolMetaData(FBlueprintMetadata::MD_BlueprintInternalUseOnly)));
+	return (InStruct && InStruct->GetBoolMetaDataHierarchical(FBlueprintMetadata::MD_AllowableBlueprintVariableType) && (bForInternalUse || !InStruct->GetBoolMetaDataHierarchical(FBlueprintMetadata::MD_BlueprintInternalUseOnly)));
 }
 
 bool UEdGraphSchema_K2::DoesGraphSupportImpureFunctions(const UEdGraph* InGraph) const
