UE-141320 - Float to double conversion warning in existing blueprint function

Macro assignment nodes don't work like other BP nodes: only their nets refer to meaningful data. As a result, its RegisterNets function needs to insert a custom entry into the implicit cast table.

However, when the compiler calls RegisterImplicitCasts to detect implicit cast points in the graph, it can mistakenly identify that a cast is needed for its value pin. This is incorrect since a cast would only be needed between the *nets* of the value and variable pins. The fix here is to have the node explicitly remove the erroneous entry from the map, which we do in other nodes that have special semantics.

Additionally, changed the math expression handler to use RemoveRegisteredImplicitCast, which is just a simple wrapper around ImplicitCastMap.Remove that returns the number of entries removed. This can be useful for debugging when trying to determine how many entries were removed from the cast table.

diff --git a/Engine/Source/Editor/BlueprintGraph/Private/K2Node_AssignmentStatement.cpp b/Engine/Source/Editor/BlueprintGraph/Private/K2Node_AssignmentStatement.cpp
--- a/Engine/Source/Editor/BlueprintGraph/Private/K2Node_AssignmentStatement.cpp
+++ b/Engine/Source/Editor/BlueprintGraph/Private/K2Node_AssignmentStatement.cpp
@@ -26,6 +26,8 @@ public:
 
 	virtual void RegisterNets(FKismetFunctionContext& Context, UEdGraphNode* Node) override
 	{
+		using namespace UE::KismetCompiler;
+
 		UEdGraphPin* VariablePin = Node->FindPin(TEXT("Variable"));
 		UEdGraphPin* ValuePin = Node->FindPin(TEXT("Value"));
 
@@ -43,47 +45,44 @@ public:
 
 		ValidateAndRegisterNetIfLiteral(Context, ValuePin);
 
-		{
-			using namespace UE::KismetCompiler;
+		UEdGraphPin* VariablePinNet = FEdGraphUtilities::GetNetFromPin(VariablePin);
+		UEdGraphPin* ValuePinNet = FEdGraphUtilities::GetNetFromPin(ValuePin);
 
-			UEdGraphPin* VariablePinNet = FEdGraphUtilities::GetNetFromPin(VariablePin);
-			UEdGraphPin* ValuePinNet = FEdGraphUtilities::GetNetFromPin(ValuePin);
+		if (VariablePinNet && ValuePinNet)
+		{
+			TOptional<CastingUtils::StatementNamePair> ConversionType =
+				CastingUtils::GetFloatingPointConversionType(*ValuePinNet, *VariablePinNet);
 
-			if (VariablePinNet && ValuePinNet)
+			if (ConversionType)
 			{
-				TOptional<CastingUtils::StatementNamePair> ConversionType =
-					CastingUtils::GetFloatingPointConversionType(*ValuePinNet, *VariablePinNet);
+				FString DescriptiveName = Node->GetName();
 
-				if (ConversionType)
-				{
+				FString TerminalName = FString::Printf(TEXT("%s_%s_%s"),
+					*DescriptiveName,
+					*VariablePinNet->PinName.ToString(),
+					ConversionType->Get<1>());
 
-					FString DescriptiveName = Node->GetName();
+				FBPTerminal* NewTerm = Context.CreateLocalTerminal();
+				NewTerm->Name = TerminalName;
+				NewTerm->Type = VariablePinNet->PinType;
+				NewTerm->Source = Node;
 
-					FString TerminalName = FString::Printf(TEXT("%s_%s_%s"),
-						*DescriptiveName,
-						*VariablePinNet->PinName.ToString(),
-						ConversionType->Get<1>());
+				EKismetCompiledStatementType CastType = ConversionType->Get<0>();
 
-					FBPTerminal* NewTerm = Context.CreateLocalTerminal();
-					NewTerm->Name = TerminalName;
-					NewTerm->Type = VariablePinNet->PinType;
-					NewTerm->Source = Node;
-
-					EKismetCompiledStatementType CastType = ConversionType->Get<0>();
-
-					Context.ImplicitCastMap.Add(VariablePin, FImplicitCastParams{CastType, NewTerm, Node});
-				}
-			}
-			else
-			{
-				CompilerContext.MessageLog.Error(*LOCTEXT("NoVariableOrValueNets_Error", "Expected Variable and Value pins to have valid connections in @@").ToString(), Node);
+				Context.ImplicitCastMap.Add(VariablePin, FImplicitCastParams{CastType, NewTerm, Node});
 			}
 		}
+		else
+		{
+			CompilerContext.MessageLog.Error(*LOCTEXT("NoVariableOrValueNets_Error", "Expected Variable and Value pins to have valid connections in @@").ToString(), Node);
+		}
 	}
 
 
 	virtual void Compile(FKismetFunctionContext& Context, UEdGraphNode* Node) override
 	{
+		using namespace UE::KismetCompiler;
+
 		UEdGraphPin* VariablePin = Node->FindPin(TEXT("Variable"));
 		UEdGraphPin* ValuePin = Node->FindPin(TEXT("Value"));
 
@@ -95,6 +94,12 @@ public:
 
 		// Generate the output impulse from this node
 		GenerateSimpleThenGoto(Context, *Node);
+
+		// The assignment node is very much a non-standard node: its pins don't directly reference any data, but their nets do.
+		// The only cast that could happen is from the value pin's net to the variable's pin net, which we handle in RegisterNets.
+		// Due to how RegisterImplicitCasts works, it can actually register an erroneous entry to the value pin,
+		// which we need to remove ourselves.
+		CastingUtils::RemoveRegisteredImplicitCast(Context, ValuePin);
 	}
 
 protected:
diff --git a/Engine/Source/Editor/BlueprintGraph/Private/MathExpressionHandler.cpp b/Engine/Source/Editor/BlueprintGraph/Private/MathExpressionHandler.cpp
--- a/Engine/Source/Editor/BlueprintGraph/Private/MathExpressionHandler.cpp
+++ b/Engine/Source/Editor/BlueprintGraph/Private/MathExpressionHandler.cpp
@@ -131,44 +131,44 @@ FBlueprintCompiledStatement* FKCHandler_MathExpression::GenerateFunctionRPN(UEdG
 
 			if (RHSTerm)
 			{
+				const FImplicitCastParams* CastParams =
+					Context.ImplicitCastMap.Find(PinMatch);
+				if (CastParams)
 				{
-					const FImplicitCastParams* CastParams =
-						Context.ImplicitCastMap.Find(PinMatch);
-					if (CastParams)
-					{
-						check(CastParams->TargetTerminal);
-
-						UFunction* CastFunction = nullptr;
+					using namespace UE::KismetCompiler;
 
-						switch (CastParams->CastType)
-						{
-						case KCST_DoubleToFloatCast:
-							CastFunction = 
-								UKismetMathLibrary::StaticClass()->FindFunctionByName(GET_FUNCTION_NAME_CHECKED(UKismetMathLibrary, Conv_DoubleToFloat));
-							break;
+					check(CastParams->TargetTerminal);
 
-						case KCST_FloatToDoubleCast:
-							CastFunction = 
-								UKismetMathLibrary::StaticClass()->FindFunctionByName(GET_FUNCTION_NAME_CHECKED(UKismetMathLibrary, Conv_FloatToDouble));
-							break;
+					UFunction* CastFunction = nullptr;
 
-						default:
-							checkf(false, TEXT("Unsupported cast type used in math expression node: %d"), CastParams->CastType);
-						}
+					switch (CastParams->CastType)
+					{
+					case KCST_DoubleToFloatCast:
+						CastFunction = 
+							UKismetMathLibrary::StaticClass()->FindFunctionByName(GET_FUNCTION_NAME_CHECKED(UKismetMathLibrary, Conv_DoubleToFloat));
+						break;
+
+					case KCST_FloatToDoubleCast:
+						CastFunction = 
+							UKismetMathLibrary::StaticClass()->FindFunctionByName(GET_FUNCTION_NAME_CHECKED(UKismetMathLibrary, Conv_FloatToDouble));
+						break;
+
+					default:
+						checkf(false, TEXT("Unsupported cast type used in math expression node: %d"), CastParams->CastType);
+					}
 
-						check(CastFunction);
+					check(CastFunction);
 
-						FBlueprintCompiledStatement* CastStatement = new FBlueprintCompiledStatement();
-						CastStatement->FunctionToCall = CastFunction;
-						CastStatement->Type = KCST_CallFunction;
-						CastStatement->RHS.Add(RHSTerm);
+					FBlueprintCompiledStatement* CastStatement = new FBlueprintCompiledStatement();
+					CastStatement->FunctionToCall = CastFunction;
+					CastStatement->Type = KCST_CallFunction;
+					CastStatement->RHS.Add(RHSTerm);
 
-						RHSTerm = CastParams->TargetTerminal;
-						CastParams->TargetTerminal->InlineGeneratedParameter = CastStatement;
-						Context.AllGeneratedStatements.Add(CastStatement);
+					RHSTerm = CastParams->TargetTerminal;
+					CastParams->TargetTerminal->InlineGeneratedParameter = CastStatement;
+					Context.AllGeneratedStatements.Add(CastStatement);
 
-						Context.ImplicitCastMap.Remove(PinMatch);
-					}
+					CastingUtils::RemoveRegisteredImplicitCast(Context, PinMatch);
 				}
 
 				RHSTerms.Add(RHSTerm);
