Control Rig: Added BlueprintInternalUseOnlyHierarchical tag that can be inherited by derived structs. Addded this tag to base rig unit to hide all raw rig unit structs from BP

previous change at CL 18866227 was backed out because it broke the bot1 cooks: https://horde.devtools.epicgames.com/log/61ffcf94676584e71ea0ab7e?lineindex=104697

Basically GetBoolMetaDataHierarchical returns the wrong thing and the return pin can no longer be broken up and so the blueprint won't compile.

diff --git a/Engine/Plugins/Animation/ControlRig/Source/ControlRig/Public/Units/RigUnit.h b/Engine/Plugins/Animation/ControlRig/Source/ControlRig/Public/Units/RigUnit.h
--- a/Engine/Plugins/Animation/ControlRig/Source/ControlRig/Public/Units/RigUnit.h
+++ b/Engine/Plugins/Animation/ControlRig/Source/ControlRig/Public/Units/RigUnit.h
@@ -10,7 +10,7 @@
 struct FRigUnitContext;
 
 /** Base class for all rig units */
-USTRUCT(BlueprintType, meta=(Abstract, NodeColor = "0.1 0.1 0.1"))
+USTRUCT(BlueprintType, meta=(Abstract, NodeColor = "0.1 0.1 0.1"), BlueprintInternalUseOnlyHierarchical)
 struct CONTROLRIG_API FRigUnit : public FRigVMStruct
 {
 	GENERATED_BODY()
diff --git a/Engine/Source/Editor/BlueprintGraph/Classes/EdGraphSchema_K2.h b/Engine/Source/Editor/BlueprintGraph/Classes/EdGraphSchema_K2.h
--- a/Engine/Source/Editor/BlueprintGraph/Classes/EdGraphSchema_K2.h
+++ b/Engine/Source/Editor/BlueprintGraph/Classes/EdGraphSchema_K2.h
@@ -177,6 +177,7 @@ public:
 	static const FName MD_HidePin;
 
 	static const FName MD_BlueprintInternalUseOnly;
+	static const FName MD_BlueprintInternalUseOnlyHierarchical;
 	static const FName MD_NeedsLatentFixup;
 
 	static const FName MD_LatentInfo;
diff --git a/Engine/Source/Editor/BlueprintGraph/Private/EdGraphSchema_K2.cpp b/Engine/Source/Editor/BlueprintGraph/Private/EdGraphSchema_K2.cpp
--- a/Engine/Source/Editor/BlueprintGraph/Private/EdGraphSchema_K2.cpp
+++ b/Engine/Source/Editor/BlueprintGraph/Private/EdGraphSchema_K2.cpp
@@ -147,6 +147,7 @@ const FName FBlueprintMetadata::MD_ShowWorldContextPin(TEXT("ShowWorldContextPin
 const FName FBlueprintMetadata::MD_Private(TEXT("BlueprintPrivate"));
 
 const FName FBlueprintMetadata::MD_BlueprintInternalUseOnly(TEXT("BlueprintInternalUseOnly"));
+const FName FBlueprintMetadata::MD_BlueprintInternalUseOnlyHierarchical(TEXT("BlueprintInternalUseOnlyHierarchical"));
 const FName FBlueprintMetadata::MD_NeedsLatentFixup(TEXT("NeedsLatentFixup"));
 const FName FBlueprintMetadata::MD_LatentInfo(TEXT("LatentInfo"));
 const FName FBlueprintMetadata::MD_LatentCallbackTarget(TEXT("LatentCallbackTarget"));
@@ -1267,7 +1268,29 @@ bool UEdGraphSchema_K2::IsAllowableBlueprintVariableType(const UScriptStruct* In
 		// User-defined structs are always allowed as BP variable types.
 		return true;
 	}
-	return (InStruct && InStruct->GetBoolMetaDataHierarchical(FBlueprintMetadata::MD_AllowableBlueprintVariableType) && (bForInternalUse || !InStruct->GetBoolMetaData(FBlueprintMetadata::MD_BlueprintInternalUseOnly)));
+
+	// struct needs to be marked as BP type
+	if (InStruct && InStruct->GetBoolMetaDataHierarchical(FBlueprintMetadata::MD_AllowableBlueprintVariableType))
+	{
+		// for internal use, all BP types are allowed
+		if (bForInternalUse)
+		{
+			return true;
+		}
+
+		// for user-facing use case, only allow structs that don't have the internal-use-only tag
+		// struct itself should not be tagged
+		if (!InStruct->GetBoolMetaData(FBlueprintMetadata::MD_BlueprintInternalUseOnly))
+		{
+			// struct's base structs should not be tagged
+			if (!InStruct->GetBoolMetaDataHierarchical(FBlueprintMetadata::MD_BlueprintInternalUseOnlyHierarchical))
+			{
+				return true;
+			}
+		}
+	}
+
+	return false;
 }
 
 bool UEdGraphSchema_K2::DoesGraphSupportImpureFunctions(const UEdGraph* InGraph) const
diff --git a/Engine/Source/Editor/BlueprintGraph/Private/K2Node_StructOperation.cpp b/Engine/Source/Editor/BlueprintGraph/Private/K2Node_StructOperation.cpp
--- a/Engine/Source/Editor/BlueprintGraph/Private/K2Node_StructOperation.cpp
+++ b/Engine/Source/Editor/BlueprintGraph/Private/K2Node_StructOperation.cpp
@@ -155,6 +155,10 @@ bool UK2Node_StructOperation::IsActionFilteredOut(const FBlueprintActionFilter&
 		{
 			bIsFiltered = true;
 		}
+		else if (StructType->GetBoolMetaDataHierarchical(FBlueprintMetadata::MD_BlueprintInternalUseOnlyHierarchical))
+		{
+			bIsFiltered = true;
+		}
 		else if (!StructType->GetBoolMetaData(FBlueprintMetadata::MD_AllowableBlueprintVariableType))
 		{
 			bIsFiltered = true;
diff --git a/Engine/Source/Programs/UnrealHeaderTool/Private/BaseParser.cpp b/Engine/Source/Programs/UnrealHeaderTool/Private/BaseParser.cpp
--- a/Engine/Source/Programs/UnrealHeaderTool/Private/BaseParser.cpp
+++ b/Engine/Source/Programs/UnrealHeaderTool/Private/BaseParser.cpp
@@ -96,6 +96,11 @@ namespace
 			BlueprintInternalUseOnly.InsertAddAction(TEXT("BlueprintInternalUseOnly"), TEXT("true"));
 			BlueprintInternalUseOnly.InsertAddAction(TEXT("BlueprintType"), TEXT("true"));
 
+			FMetadataKeyword& BlueprintInternalUseOnlyHierarchical = MetadataDictionary.Add(TEXT("BlueprintInternalUseOnlyHierarchical"), EMetadataValueArgument::None);
+			BlueprintInternalUseOnlyHierarchical.InsertAddAction(TEXT("BlueprintInternalUseOnlyHierarchical"), TEXT("true"));
+			BlueprintInternalUseOnlyHierarchical.InsertAddAction(TEXT("BlueprintInternalUseOnly"), TEXT("true"));
+			BlueprintInternalUseOnlyHierarchical.InsertAddAction(TEXT("BlueprintType"), TEXT("true"));
+			
 			FMetadataKeyword& BlueprintType = MetadataDictionary.Add(TEXT("BlueprintType"), EMetadataValueArgument::None);
 			BlueprintType.InsertAddAction(TEXT("BlueprintType"), TEXT("true"));
 
diff --git a/Engine/Source/Runtime/CoreUObject/Public/UObject/ObjectMacros.h b/Engine/Source/Runtime/CoreUObject/Public/UObject/ObjectMacros.h
--- a/Engine/Source/Runtime/CoreUObject/Public/UObject/ObjectMacros.h
+++ b/Engine/Source/Runtime/CoreUObject/Public/UObject/ObjectMacros.h
@@ -1071,7 +1071,10 @@ namespace US
 		BlueprintType,
 
 		/// Indicates that a BlueprintType struct should not be exposed to the end user
-		BlueprintInternalUseOnly
+		BlueprintInternalUseOnly,
+
+		/// Indicates that a BlueprintType struct and its derived structs should not be exposed to the end user
+		BlueprintInternalUseOnlyHierarchical,
 	};
 }
 
