UE-141737 - Math Expression nodes don't handle BP Reals correctly

This issue was limited to any nodes in the math expression subgraph that were connected to the entry node. The problem is that nodes in the expression subgraph are completely decoupled from any inputs that are connected to the primary math expression node. As a result, the compiler isn't able to detect that a cast might be needed between subgraph nodes and any inputs connected to the expression node.

Since this is a one off issue, we need to modify FKCHandler_MathExpression::RegisterNets to perform the cast analysis.

diff --git a/Engine/Source/Editor/BlueprintGraph/Private/MathExpressionHandler.cpp b/Engine/Source/Editor/BlueprintGraph/Private/MathExpressionHandler.cpp
--- a/Engine/Source/Editor/BlueprintGraph/Private/MathExpressionHandler.cpp
+++ b/Engine/Source/Editor/BlueprintGraph/Private/MathExpressionHandler.cpp
@@ -185,6 +185,8 @@ FBlueprintCompiledStatement* FKCHandler_MathExpression::GenerateFunctionRPN(UEdG
 
 void FKCHandler_MathExpression::RegisterNets(FKismetFunctionContext& Context, UEdGraphNode* InNode)
 {
+	using namespace UE::KismetCompiler;
+
 	FNodeHandlingFunctor::RegisterNets(Context, InNode);
 
 	UK2Node_MathExpression* Node_MathExpression = CastChecked<UK2Node_MathExpression>(InNode);
@@ -232,6 +234,54 @@ void FKCHandler_MathExpression::RegisterNets(FKismetFunctionContext& Context, UE
 			}
 		}
 	}
+
+	// We have two nodes that handle input: the math expression node and the inner entry node.
+	// For any pins connected to the entry node, we need to compare their types to the nets
+	// of the math expression node. 
+	// 
+	// This is a fairly convoluted scenario that RegisterImplicitCasts fails to detect,
+	// and it only affects pins connected to the inner entry node.
+
+	int EntryNodeCursor = 0;
+	for (UEdGraphPin* Pin : Node_MathExpression->Pins)
+	{
+		if (Pin && (Pin->Direction == EEdGraphPinDirection::EGPD_Input))
+		{
+			UEdGraphPin* PinNet = FEdGraphUtilities::GetNetFromPin(Pin);
+			if (PinNet)
+			{
+				UEdGraphPin* InnerEntryNodePin = InnerEntryNode->Pins[EntryNodeCursor];
+				if (InnerEntryNodePin && (InnerEntryNodePin->GetName() == Pin->GetName()))
+				{
+					UEdGraphPin* Linked = InnerEntryNodePin->LinkedTo[0];
+					if (Linked)
+					{
+						TOptional<CastingUtils::StatementNamePair> ConversionType =
+							CastingUtils::GetFloatingPointConversionType(*PinNet, *Linked);
+
+						if (ConversionType)
+						{
+							EKismetCompiledStatementType CastType = ConversionType->Get<0>();
+							const TCHAR* TermName = ConversionType->Get<1>();
+
+							FBPTerminal* NewTerm = Context.CreateLocalTerminal();
+							UEdGraphNode* OwningNode = Linked->GetOwningNode();
+							NewTerm->CopyFromPin(Linked, Context.NetNameMap->MakeValidName(Linked, TermName));
+							NewTerm->Source = OwningNode;
+
+							Context.ImplicitCastMap.Add(Linked, FImplicitCastParams{ CastType, NewTerm, OwningNode });
+						}
+					}
+				}
+				else
+				{
+					Context.MessageLog.Error(*LOCTEXT("Compile_PinError", "ICE - mismatched pins found on @@ and @@!").ToString(), Node_MathExpression, InnerEntryNode);
+				}
+			}
+
+			++EntryNodeCursor;
+		}
+	}
 }
 
 void FKCHandler_MathExpression::RegisterNet(FKismetFunctionContext& Context, UEdGraphPin* Net) 
@@ -281,24 +331,22 @@ void FKCHandler_MathExpression::Compile(FKismetFunctionContext& Context, UEdGrap
 	FBlueprintCompiledStatement* DetachedStatement = GenerateFunctionRPN(LastInnerNode, Context, *Node_MathExpression, OutputTerm, InnerToOuterInput);
 	if (DetachedStatement)
 	{
+		using namespace UE::KismetCompiler;
+	
 		Context.AllGeneratedStatements.Add(DetachedStatement);
 		TArray<FBlueprintCompiledStatement*>& StatementList = Context.StatementsPerNode.FindOrAdd(Node);
 		StatementList.Add(DetachedStatement);
 
-		{
-			using namespace UE::KismetCompiler;
-
-			// The casts that we care about occur *within* the nodes generated by the math expression.
-			// The expression node merely serves as a proxy, so we can silently remove the casts
-			// on the input pins that were registered by the compiler.
+		// The casts that we care about occur *within* the nodes generated by the math expression.
+		// The expression node merely serves as a proxy, so we can silently remove the casts
+		// on the input pins that were registered by the compiler.
 
-			for (UEdGraphPin* Pin : Node_MathExpression->Pins)
+		for (UEdGraphPin* Pin : Node_MathExpression->Pins)
+		{
+			check(Pin);
+			if (!Context.Schema->IsMetaPin(*Pin) && (Pin->Direction == EGPD_Input))
 			{
-				check(Pin);
-				if (!Context.Schema->IsMetaPin(*Pin) && (Pin->Direction == EGPD_Input))
-				{
-					CastingUtils::RemoveRegisteredImplicitCast(Context, Pin);
-				}
+				CastingUtils::RemoveRegisteredImplicitCast(Context, Pin);
 			}
 		}
 	}
