Allow unloaded type assets to be included in the member type selector in the user-defined structure editor.

diff --git a/Engine/Source/Editor/BlueprintGraph/Private/EdGraphSchema_K2.cpp b/Engine/Source/Editor/BlueprintGraph/Private/EdGraphSchema_K2.cpp
--- a/Engine/Source/Editor/BlueprintGraph/Private/EdGraphSchema_K2.cpp
+++ b/Engine/Source/Editor/BlueprintGraph/Private/EdGraphSchema_K2.cpp
@@ -484,30 +484,20 @@ public:
 
 		OutChildren.Sort(FCompareChildren());
 	}
-
-	/** Loads an asset based on the AssetReference through the asset registry */
-	static UObject* LoadAsset(const FSoftObjectPath& AssetReference)
-	{
-		if (AssetReference.IsValid())
-		{
-			const FAssetRegistryModule& AssetRegistryModule = FModuleManager::LoadModuleChecked<FAssetRegistryModule>(TEXT("AssetRegistry"));
-			const FAssetData AssetData = AssetRegistryModule.Get().GetAssetByObjectPath(*AssetReference.ToString());
-			return AssetData.GetAsset(); //-V558
-		}
-		return nullptr;
-	}
 };
 
 const FEdGraphPinType& UEdGraphSchema_K2::FPinTypeTreeInfo::GetPinType(bool bForceLoadedSubCategoryObject)
 {
-	if (bForceLoadedSubCategoryObject)
+	// Only attempt to load the sub category object if we need to
+	if (SubCategoryObjectAssetReference.IsValid() && (!PinType.PinSubCategoryObject.IsValid() || FSoftObjectPath(PinType.PinSubCategoryObject.Get()) != SubCategoryObjectAssetReference))
 	{
-		// Only attempt to load the sub category object if we need to
-		if ( SubCategoryObjectAssetReference.IsValid() && (!PinType.PinSubCategoryObject.IsValid() || FSoftObjectPath(PinType.PinSubCategoryObject.Get()) != SubCategoryObjectAssetReference) )
+		const FAssetRegistryModule& AssetRegistryModule = FModuleManager::LoadModuleChecked<FAssetRegistryModule>(TEXT("AssetRegistry"));
+		const FAssetData AssetData = AssetRegistryModule.Get().GetAssetByObjectPath(*SubCategoryObjectAssetReference.ToString());
+		if (bForceLoadedSubCategoryObject || AssetData.IsAssetLoaded())
 		{
-			UObject* LoadedObject = FGatherTypesHelper::LoadAsset(SubCategoryObjectAssetReference);
+			UObject* LoadedObject = AssetData.GetAsset();
 
-			if(UBlueprint* BlueprintObject = Cast<UBlueprint>(LoadedObject))
+			if (UBlueprint* BlueprintObject = Cast<UBlueprint>(LoadedObject))
 			{
 				PinType.PinSubCategoryObject = *BlueprintObject->GeneratedClass;
 			}
@@ -517,33 +507,7 @@ const FEdGraphPinType& UEdGraphSchema_K2::FPinTypeTreeInfo::GetPinType(bool bFor
 			}
 		}
 	}
-	else
-	{
-		if (SubCategoryObjectAssetReference.IsValid())
-		{
-			const FAssetRegistryModule& AssetRegistryModule = FModuleManager::LoadModuleChecked<FAssetRegistryModule>(TEXT("AssetRegistry"));
-			const FAssetData AssetData = AssetRegistryModule.Get().GetAssetByObjectPath(*SubCategoryObjectAssetReference.ToString());
-
-			if(!AssetData.IsAssetLoaded())
-			{
-				UObject* LoadedObject = FindObject<UClass>(ANY_PACKAGE, *AssetData.AssetClass.ToString());
 
-				// If the unloaded asset is a Blueprint, we need to pull the generated class and assign that
-				if(UBlueprint* BlueprintObject = Cast<UBlueprint>(LoadedObject))
-				{
-					PinType.PinSubCategoryObject = *BlueprintObject->GeneratedClass;
-				}
-				else
-				{
-					PinType.PinSubCategoryObject = LoadedObject;
-				}
-			}
-			else
-			{
-				PinType.PinSubCategoryObject = AssetData.GetAsset();
-			}
-		}
-	}
 	return PinType;
 }
 
diff --git a/Engine/Source/Editor/Kismet/Private/UserDefinedStructureEditor.cpp b/Engine/Source/Editor/Kismet/Private/UserDefinedStructureEditor.cpp
--- a/Engine/Source/Editor/Kismet/Private/UserDefinedStructureEditor.cpp
+++ b/Engine/Source/Editor/Kismet/Private/UserDefinedStructureEditor.cpp
@@ -990,8 +990,11 @@ public:
 			const auto Child = PinTypeNode->Children[ChildIndex];
 			if(Child.IsValid())
 			{
-				const bool bCanCheckSubObjectWithoutLoading = Child->GetPinType(false).PinSubCategoryObject.IsValid();
-				if (bCanCheckSubObjectWithoutLoading && !FStructureEditorUtils::CanHaveAMemberVariableOfType(Parent, Child->GetPinType(false)))
+				const bool bForceLoadSubCategoryObject = false;
+				const FEdGraphPinType& PinType = Child->GetPinType(bForceLoadSubCategoryObject);
+
+				const bool bCanCheckSubObjectWithoutLoading = PinType.PinSubCategoryObject.IsValid();
+				if(bCanCheckSubObjectWithoutLoading && !FStructureEditorUtils::CanHaveAMemberVariableOfType(Parent, PinType))
 				{
 					PinTypeNode->Children.RemoveAt(ChildIndex);
 					continue;
diff --git a/Engine/Source/Editor/UnrealEd/Private/Kismet2/StructureEditorUtils.cpp b/Engine/Source/Editor/UnrealEd/Private/Kismet2/StructureEditorUtils.cpp
--- a/Engine/Source/Editor/UnrealEd/Private/Kismet2/StructureEditorUtils.cpp
+++ b/Engine/Source/Editor/UnrealEd/Private/Kismet2/StructureEditorUtils.cpp
@@ -188,21 +188,24 @@ bool FStructureEditorUtils::CanHaveAMemberVariableOfType(const UUserDefinedStruc
 {
 	if ((VarType.PinCategory == UEdGraphSchema_K2::PC_Struct) && Struct)
 	{
-		if (const UScriptStruct* SubCategoryStruct = Cast<const UScriptStruct>(VarType.PinSubCategoryObject.Get()))
+		if (const UObject* TypeObject = VarType.PinSubCategoryObject.Get())
 		{
-			const EStructureError Result = IsStructureValid(SubCategoryStruct, Struct, OutMsg);
-			if (EStructureError::Ok != Result)
+			if (const UScriptStruct* SubCategoryStruct = Cast<const UScriptStruct>(TypeObject))
 			{
-				return false;
+				const EStructureError Result = IsStructureValid(SubCategoryStruct, Struct, OutMsg);
+				if (EStructureError::Ok != Result)
+				{
+					return false;
+				}
 			}
-		}
-		else
-		{
-			if (OutMsg)
+			else
 			{
-				*OutMsg = LOCTEXT("StructureIncorrectStructType", "Incorrect struct type in a structure member variable.").ToString();
+				if (OutMsg)
+				{
+					*OutMsg = LOCTEXT("StructureIncorrectStructType", "Incorrect struct type in a structure member variable.").ToString();
+				}
+				return false;
 			}
-			return false;
 		}
 	}
 	else if ((VarType.PinCategory == UEdGraphSchema_K2::PC_Exec) 
@@ -216,18 +219,6 @@ bool FStructureEditorUtils::CanHaveAMemberVariableOfType(const UUserDefinedStruc
 		}
 		return false;
 	}
-	else
-	{
-		const UClass* PinSubCategoryClass = Cast<const UClass>(VarType.PinSubCategoryObject.Get());
-		if (PinSubCategoryClass && PinSubCategoryClass->IsChildOf(UBlueprint::StaticClass()))
-		{
-			if (OutMsg)
-			{
-				*OutMsg = LOCTEXT("StructureUseBlueprintReferences", "Struct cannot use any blueprint references").ToString();
-			}
-			return false;
-		}
-	}
 	return true;
 }
 
