Blueprints: Added a cvar to control the how Reals are presented to users in the Blueprint editor, changing the default to Float

diff --git a/Engine/Source/Editor/BlueprintGraph/Private/EdGraphSchema_K2.cpp b/Engine/Source/Editor/BlueprintGraph/Private/EdGraphSchema_K2.cpp
--- a/Engine/Source/Editor/BlueprintGraph/Private/EdGraphSchema_K2.cpp
+++ b/Engine/Source/Editor/BlueprintGraph/Private/EdGraphSchema_K2.cpp
@@ -101,6 +101,35 @@
 #include "Editor/EditorPerProjectUserSettings.h"
 #include "BlueprintPaletteFavorites.h"
 
+//////////////////////////////////////////////////////////////////////////
+
+// How to display PC_Real pin types to users
+enum class EBlueprintRealDisplayMode
+{
+	Real,
+	Float,
+	Number
+};
+
+namespace BlueprintDisplay
+{
+	int32 LastRealNamingMode = -1;
+	int32 RealNamingMode = (int32)EBlueprintRealDisplayMode::Float;
+	FAutoConsoleVariableRef CVarRealNamingMode(TEXT("Blueprint.PC_Real.DisplayMode"), RealNamingMode, TEXT("Real naming mode\n\t0: Real\n\t1: Float (default)\n\t2: Number\n\nNote the editor needs to be restarted for this to fully take effect"));
+
+	EBlueprintRealDisplayMode GetRealDisplayMode()
+	{
+		return static_cast<EBlueprintRealDisplayMode>(FMath::Clamp(RealNamingMode, 0, 2));
+	}
+
+	bool ShouldRefreshRealDisplay()
+	{
+		const bool bResult = LastRealNamingMode != RealNamingMode;
+		LastRealNamingMode = RealNamingMode;
+		return bResult;
+	}
+}
+
 //////////////////////////////////////////////////////////////////////////
 // FBlueprintMetadata
 
@@ -3640,7 +3669,22 @@ FText UEdGraphSchema_K2::GetCategoryText(const FName Category, const bool bForMe
 		CategoryDescriptions.Add(PC_SoftClass, LOCTEXT("SoftClassReferenceCategory", "Soft Class Reference"));
 		CategoryDescriptions.Add(PC_FieldPath, LOCTEXT("FieldPathReferenceCategory", "Property Reference"));
 		CategoryDescriptions.Add(AllObjectTypes, LOCTEXT("AllObjectTypes", "Object Types"));
+	}
 
+	if (BlueprintDisplay::ShouldRefreshRealDisplay())
+	{
+		switch (BlueprintDisplay::GetRealDisplayMode())
+		{
+		case EBlueprintRealDisplayMode::Real:
+			CategoryDescriptions[PC_Real] = LOCTEXT("RealCategory_DisplayAsReal", "Real");
+			break;
+		case EBlueprintRealDisplayMode::Float:
+			CategoryDescriptions[PC_Real] = LOCTEXT("RealCategory_DisplayAsFloat", "Float");
+			break;
+		case EBlueprintRealDisplayMode::Number:
+			CategoryDescriptions[PC_Real] = LOCTEXT("RealCategory_DisplayAsNumber", "Number");
+			break;
+		}
 	}
 
 	if (FText const* TypeDesc = CategoryDescriptions.Find(Category))
@@ -3713,10 +3757,28 @@ FText UEdGraphSchema_K2::TerminalTypeToText(const FName Category, const FName Su
 	}
 	else if (!SubCategory.IsNone())
 	{
-		FFormatNamedArguments Args;
-		Args.Add(TEXT("Category"), UEdGraphSchema_K2::GetCategoryText(Category));
-		Args.Add(TEXT("ObjectName"), FText::FromString(FName::NameToDisplayString(SubCategory.ToString(), false)));
-		PropertyText = FText::Format(LOCTEXT("ObjectAsText", "{ObjectName} {Category}"), Args);
+		if (Category == UEdGraphSchema_K2::PC_Real)
+		{
+			switch (BlueprintDisplay::GetRealDisplayMode())
+			{
+			case EBlueprintRealDisplayMode::Real:
+				PropertyText = (SubCategory == UEdGraphSchema_K2::PC_Float) ? LOCTEXT("SinglePrecisionReal", "Real (single-precision)") : LOCTEXT("DoublePrecisionReal", "Real (double-precision)");
+				break;
+			case EBlueprintRealDisplayMode::Float:
+				PropertyText = (SubCategory == UEdGraphSchema_K2::PC_Float) ? LOCTEXT("SinglePrecisionFloat", "Float (single-precision)") : LOCTEXT("DoublePrecisionFloat", "Float (double-precision)");
+				break;
+			case EBlueprintRealDisplayMode::Number:
+				PropertyText = (SubCategory == UEdGraphSchema_K2::PC_Float) ? LOCTEXT("SinglePrecisionNumber", "Number (single-precision)") : LOCTEXT("DoublePrecisionNumber", "Number (double-precision)");
+				break;
+			}
+		}
+		else
+		{
+			FFormatNamedArguments Args;
+			Args.Add(TEXT("Category"), UEdGraphSchema_K2::GetCategoryText(Category));
+			Args.Add(TEXT("ObjectName"), FText::FromString(FName::NameToDisplayString(SubCategory.ToString(), false)));
+			PropertyText = FText::Format(LOCTEXT("ObjectAsText", "{ObjectName} {Category}"), Args);
+		}
 	}
 	else
 	{
