UE-142401 - Dragging off a pin and adding a Real operator can break links to other nodes

We need to relax the rules for breaking pins when real numbers are used. We should maintain their connections since we implicitly perform casts between floats and doubles.

diff --git a/Engine/Source/Editor/BlueprintGraph/Private/K2Node_PromotableOperator.cpp b/Engine/Source/Editor/BlueprintGraph/Private/K2Node_PromotableOperator.cpp
--- a/Engine/Source/Editor/BlueprintGraph/Private/K2Node_PromotableOperator.cpp
+++ b/Engine/Source/Editor/BlueprintGraph/Private/K2Node_PromotableOperator.cpp
@@ -881,6 +881,19 @@ void UK2Node_PromotableOperator::EvaluatePinsFromChange(UEdGraphPin* ChangedPin,
 		for(int32 i = ChangedPin->LinkedTo.Num() - 1; i >= 0; --i)
 		{
 			UEdGraphPin* Link = ChangedPin->LinkedTo[i];
+
+			const bool bIsRealNumberConnection =
+				(Link->PinType.PinCategory == UEdGraphSchema_K2::PC_Real) &&
+				(MostRecentConnection.PinCategory == UEdGraphSchema_K2::PC_Real);
+
+			// Real numbers are an exception to the normal link breaking rules.
+			// Even if the subcategories don't match, they're still valid connections
+			// since we implicitly cast between float and double types.
+			if (bIsRealNumberConnection)
+			{
+				continue;
+			}
+
 			if( Link->PinType.PinCategory != MostRecentConnection.PinCategory ||
 				Link->PinType.PinSubCategory != MostRecentConnection.PinSubCategory ||
 				Link->PinType.PinSubCategoryObject != MostRecentConnection.PinSubCategoryObject
