UE-143681 - Constructed object in BP is ignoring override value for exposed float property

Two issues were discovered:

1. UKismetSystemLibrary::SetFloatPropertyByName is no longer used since it's been replaced by SetDoublePropertyByName. As a result, SetDoublePropertyByName would fail if the property was actually a float, which can happen if it's a native class. In this case, we should check for a FFloatProperty as a fallback.
2. Deprecation of float<->string converion functions in CL 19083449 introduced a regression: we can no longer use autocasts when the type is a single-precision float. In this case, we should search for autocast functions that take a double instead.

diff --git a/Engine/Source/Editor/BlueprintGraph/Private/EdGraphSchema_K2.cpp b/Engine/Source/Editor/BlueprintGraph/Private/EdGraphSchema_K2.cpp
--- a/Engine/Source/Editor/BlueprintGraph/Private/EdGraphSchema_K2.cpp
+++ b/Engine/Source/Editor/BlueprintGraph/Private/EdGraphSchema_K2.cpp
@@ -2549,11 +2549,35 @@ bool UEdGraphSchema_K2::SearchForAutocastFunction(const FEdGraphPinType& OutputP
 	// Try looking for a marked up autocast if we've not found a built-in one that works
 	if (TargetFunction == NAME_None)
 	{
+		auto FindAndSetCastFunction = [&TargetFunction, &FunctionOwner](const FEdGraphPinType& OutputPinType, const FEdGraphPinType& InputPinType)
+		{
+			const FAutocastFunctionMap& AutocastFunctionMap = FAutocastFunctionMap::Get();
+			if (const UFunction* Func = AutocastFunctionMap.Find(OutputPinType, InputPinType))
+			{
+				TargetFunction = Func->GetFName();
+				FunctionOwner = Func->GetOwnerClass();
+				return true;
+			}
+			return false;
+		};
+
 		const FAutocastFunctionMap& AutocastFunctionMap = FAutocastFunctionMap::Get();
-		if (UFunction* Func = AutocastFunctionMap.Find(OutputPinType, InputPinType))
+		if (!FindAndSetCastFunction(OutputPinType, InputPinType))
 		{
-			TargetFunction = Func->GetFName();
-			FunctionOwner = Func->GetOwnerClass();
+			// Since single-precision float interfaces have been deprecated,
+			// we should try to find a double-precision equivalent.
+			if (OutputPinType.PinSubCategory == PC_Float)
+			{
+				FEdGraphPinType PinTypeCopy(OutputPinType);
+				PinTypeCopy.PinSubCategory = PC_Double;
+				FindAndSetCastFunction(PinTypeCopy, InputPinType);
+			}
+			else if (InputPinType.PinSubCategory == PC_Float)
+			{
+				FEdGraphPinType PinTypeCopy(InputPinType);
+				PinTypeCopy.PinSubCategory = PC_Double;
+				FindAndSetCastFunction(OutputPinType, PinTypeCopy);
+			}
 		}
 	}
 
@@ -5187,7 +5211,6 @@ namespace FSetVariableByNameFunctionNames
 	static const FName SetIntName(GET_FUNCTION_NAME_CHECKED(UKismetSystemLibrary, SetIntPropertyByName));
 	static const FName SetInt64Name(GET_FUNCTION_NAME_CHECKED(UKismetSystemLibrary, SetInt64PropertyByName));
 	static const FName SetByteName(GET_FUNCTION_NAME_CHECKED(UKismetSystemLibrary, SetBytePropertyByName));
-	static const FName SetFloatName(GET_FUNCTION_NAME_CHECKED(UKismetSystemLibrary, SetFloatPropertyByName));
 	static const FName SetDoubleName(GET_FUNCTION_NAME_CHECKED(UKismetSystemLibrary, SetDoublePropertyByName));
 	static const FName SetBoolName(GET_FUNCTION_NAME_CHECKED(UKismetSystemLibrary, SetBoolPropertyByName));
 	static const FName SetObjectName(GET_FUNCTION_NAME_CHECKED(UKismetSystemLibrary, SetObjectPropertyByName));
diff --git a/Engine/Source/Runtime/Engine/Classes/Kismet/KismetSystemLibrary.h b/Engine/Source/Runtime/Engine/Classes/Kismet/KismetSystemLibrary.h
--- a/Engine/Source/Runtime/Engine/Classes/Kismet/KismetSystemLibrary.h
+++ b/Engine/Source/Runtime/Engine/Classes/Kismet/KismetSystemLibrary.h
@@ -829,7 +829,7 @@ class ENGINE_API UKismetSystemLibrary : public UBlueprintFunctionLibrary
 	static void SetBytePropertyByName(UObject* Object, FName PropertyName, uint8 Value);
 
 	/** Set a float property by name */
-	UFUNCTION(BlueprintCallable, meta=(BlueprintInternalUseOnly = "true"))
+	UE_DEPRECATED(5.0, "This method has been deprecated and will be removed. Use the double version instead.")
 	static void SetFloatPropertyByName(UObject* Object, FName PropertyName, float Value);
 
 	/** Set a double property by name */
diff --git a/Engine/Source/Runtime/Engine/Private/KismetSystemLibrary.cpp b/Engine/Source/Runtime/Engine/Private/KismetSystemLibrary.cpp
--- a/Engine/Source/Runtime/Engine/Private/KismetSystemLibrary.cpp
+++ b/Engine/Source/Runtime/Engine/Private/KismetSystemLibrary.cpp
@@ -994,11 +994,17 @@ void UKismetSystemLibrary::SetDoublePropertyByName(UObject* Object, FName Proper
 {
 	if (Object != nullptr)
 	{
-		FDoubleProperty* DoubleProp = FindFProperty<FDoubleProperty>(Object->GetClass(), PropertyName);
-		if (DoubleProp != nullptr)
+		if (FDoubleProperty* DoubleProp = FindFProperty<FDoubleProperty>(Object->GetClass(), PropertyName))
 		{
 			DoubleProp->SetPropertyValue_InContainer(Object, Value);
 		}
+		// It's entirely possible that the property refers to a native float property,
+		// so we need to make that check here.
+		else if (FFloatProperty* FloatProp = FindFProperty<FFloatProperty>(Object->GetClass(), PropertyName))
+		{
+			float floatValue = static_cast<float>(Value);
+			FloatProp->SetPropertyValue_InContainer(Object, floatValue);
+		}
 	}
 }
 
