Improved tooltip to display the class name over the casting options in blueprint editor.

#rnx

diff --git a/Engine/Source/Editor/BlueprintGraph/Private/K2Node_ClassDynamicCast.cpp b/Engine/Source/Editor/BlueprintGraph/Private/K2Node_ClassDynamicCast.cpp
--- a/Engine/Source/Editor/BlueprintGraph/Private/K2Node_ClassDynamicCast.cpp
+++ b/Engine/Source/Editor/BlueprintGraph/Private/K2Node_ClassDynamicCast.cpp
@@ -82,14 +82,16 @@ FText UK2Node_ClassDynamicCast::GetTooltipText() const
 {
 	if (TargetType && TargetType->IsChildOf(UInterface::StaticClass()))
 	{
-		return LOCTEXT("CastToInterfaceTooltip", "Tries to access class as an interface it may implement.");
+		return FText::Format(LOCTEXT("CastToInterfaceTooltip", "Tries to access class as an interface '{0}' it may implement."), FText::FromString(TargetType->GetName()));
 	}
 	UBlueprint* CastToBP = UBlueprint::GetBlueprintFromClass(TargetType);
 	if (CastToBP)
 	{
-		return LOCTEXT("CastToBPTooltip", "Tries to access class as a blueprint class it may inherit from.\n\nNOTE: This will cause the blueprint to always be loaded, which can be expensive.");
+		return FText::Format(LOCTEXT("CastToBPTooltip", "Tries to access class as a blueprint class '{0}' it may inherit from.\n\nNOTE: This will cause the blueprint to always be loaded, which can be expensive."), FText::FromString(CastToBP->GetName()));
 	}
-	return LOCTEXT("CastToNativeTooltip", "Tries to access class as one it may inherit from.");
+	
+	const FString ClassName = TargetType ? TargetType->GetName() : TEXT("");
+	return FText::Format(LOCTEXT("CastToNativeTooltip", "Tries to access class '{0}' as one it may inherit from."), FText::FromString(ClassName));
 }
 
 UEdGraphPin* UK2Node_ClassDynamicCast::GetCastSourcePin() const
diff --git a/Engine/Source/Editor/BlueprintGraph/Private/K2Node_DynamicCast.cpp b/Engine/Source/Editor/BlueprintGraph/Private/K2Node_DynamicCast.cpp
--- a/Engine/Source/Editor/BlueprintGraph/Private/K2Node_DynamicCast.cpp
+++ b/Engine/Source/Editor/BlueprintGraph/Private/K2Node_DynamicCast.cpp
@@ -121,14 +121,18 @@ FText UK2Node_DynamicCast::GetTooltipText() const
 {
 	if (TargetType && TargetType->IsChildOf(UInterface::StaticClass()))
 	{
-		return LOCTEXT("CastToInterfaceTooltip", "Tries to access object as an interface it may implement.");
+		return FText::Format(LOCTEXT("CastToInterfaceTooltip", "Tries to access object as an interface '{0}' it may implement."), FText::FromString(TargetType->GetName()));
 	}
+	
 	UBlueprint* CastToBP = UBlueprint::GetBlueprintFromClass(TargetType);
 	if (CastToBP)
 	{
-		return LOCTEXT("CastToBPTooltip", "Tries to access object as a blueprint class it may be an instance of.\n\nNOTE: This will cause the blueprint to always be loaded, which can be expensive.");
+		return FText::Format(LOCTEXT("CastToBPTooltip", "Tries to access object as a blueprint class '{0}' it may be an instance of.\n\nNOTE: This will cause the blueprint to always be loaded, which can be expensive."), FText::FromString(CastToBP->GetName()));
 	}
-	return LOCTEXT("CastToNativeTooltip", "Tries to access object as a class it may be an instance of.");
+
+	const FString ClassName = TargetType ? TargetType->GetName() : TEXT("");
+
+	return FText::Format(LOCTEXT("CastToNativeTooltip", "Tries to access object as a class '{0}' it may be an instance of."), FText::FromString(ClassName));
 }
 
 void UK2Node_DynamicCast::GetNodeContextMenuActions(UToolMenu* Menu, UGraphNodeContextMenuContext* Context) const
