UE-144351 - Blueprint real numbers causing Dev-EngineMerge blocker

The KismetMathLibrary changes made in CL 19159126 revealed a problem with type promotions: it's too strict when it comes to finding matching functions. This is currently preventing a  BP class from compiling, which is blocking Dev-EngineMerge.

We need to modify ArePinTypesEquivalent to be less strict about real number comparisons, which permits using single precision float pins with double precision math functions.

diff --git a/Engine/Source/Editor/BlueprintGraph/Private/EdGraphSchema_K2.cpp b/Engine/Source/Editor/BlueprintGraph/Private/EdGraphSchema_K2.cpp
--- a/Engine/Source/Editor/BlueprintGraph/Private/EdGraphSchema_K2.cpp
+++ b/Engine/Source/Editor/BlueprintGraph/Private/EdGraphSchema_K2.cpp
@@ -4525,6 +4525,12 @@ bool UEdGraphSchema_K2::ArePinTypesCompatible(const FEdGraphPinType& Output, con
 
 bool UEdGraphSchema_K2::ArePinTypesEquivalent(const FEdGraphPinType& PinA, const FEdGraphPinType& PinB) const
 {
+	// Real pins are effectively equivalent since we implicitly cast where necessary.
+	if ((PinA.PinCategory == PC_Real) && (PinB.PinCategory == PC_Real))
+	{
+		return true;
+	}
+
 	return 
 		PinA.PinCategory == PinB.PinCategory &&
 		PinA.PinSubCategory == PinB.PinSubCategory &&
