Prevent abstract classes from being used by create object from class nodes

diff --git a/Engine/Source/Editor/BlueprintGraph/Classes/K2Node_ConstructObjectFromClass.h b/Engine/Source/Editor/BlueprintGraph/Classes/K2Node_ConstructObjectFromClass.h
--- a/Engine/Source/Editor/BlueprintGraph/Classes/K2Node_ConstructObjectFromClass.h
+++ b/Engine/Source/Editor/BlueprintGraph/Classes/K2Node_ConstructObjectFromClass.h
@@ -28,6 +28,7 @@ class BLUEPRINTGRAPH_API UK2Node_ConstructObjectFromClass : public UK2Node
 	virtual void GetPinHoverText(const UEdGraphPin& Pin, FString& HoverTextOut) const override;
 	virtual void PostPlacedNewNode() override;
 	virtual void AddSearchMetaDataInfo(TArray<struct FSearchTagDataPair>& OutTaggedMetaData) const override;
+	virtual FString GetPinMetaData(FName InPinName, FName InKey) override;
 	//~ End UEdGraphNode Interface.
 
 	//~ Begin UK2Node Interface
diff --git a/Engine/Source/Editor/BlueprintGraph/Private/K2Node_ConstructObjectFromClass.cpp b/Engine/Source/Editor/BlueprintGraph/Private/K2Node_ConstructObjectFromClass.cpp
--- a/Engine/Source/Editor/BlueprintGraph/Private/K2Node_ConstructObjectFromClass.cpp
+++ b/Engine/Source/Editor/BlueprintGraph/Private/K2Node_ConstructObjectFromClass.cpp
@@ -202,6 +202,21 @@ void UK2Node_ConstructObjectFromClass::PostPlacedNewNode()
 	}
 }
 
+FString UK2Node_ConstructObjectFromClass::GetPinMetaData(FName InPinName, FName InKey)
+{
+	FString MetaData = Super::GetPinMetaData(InPinName, InKey);
+
+	if (MetaData.IsEmpty())
+	{
+		if (InPinName == FK2Node_ConstructObjectFromClassHelper::GetClassPinName() && InKey == FBlueprintMetadata::MD_AllowAbstractClasses)
+		{
+			MetaData = TEXT("false");
+		}
+	}
+
+	return MetaData;
+}
+
 void UK2Node_ConstructObjectFromClass::AddSearchMetaDataInfo(TArray<struct FSearchTagDataPair>& OutTaggedMetaData) const
 {
 	Super::AddSearchMetaDataInfo(OutTaggedMetaData);
