UE-145269 - Math Expression nodes failing within functions

We were only considering the first linked pin when iterating over the connections for the input pins. An input pin can be connected to multiple internal nodes, so we need to process those. For example, this problem would manifest if the expression were (x+x), and we connected the x input to a single precision float.

diff --git a/Engine/Source/Editor/BlueprintGraph/Private/MathExpressionHandler.cpp b/Engine/Source/Editor/BlueprintGraph/Private/MathExpressionHandler.cpp
--- a/Engine/Source/Editor/BlueprintGraph/Private/MathExpressionHandler.cpp
+++ b/Engine/Source/Editor/BlueprintGraph/Private/MathExpressionHandler.cpp
@@ -253,23 +253,25 @@ void FKCHandler_MathExpression::RegisterNets(FKismetFunctionContext& Context, UE
 				UEdGraphPin* InnerEntryNodePin = InnerEntryNode->Pins[EntryNodeCursor];
 				if (InnerEntryNodePin && (InnerEntryNodePin->GetName() == Pin->GetName()))
 				{
-					UEdGraphPin* Linked = InnerEntryNodePin->LinkedTo[0];
-					if (Linked)
+					for (UEdGraphPin* DestinationPin : InnerEntryNodePin->LinkedTo)
 					{
-						TOptional<CastingUtils::StatementNamePair> ConversionType =
-							CastingUtils::GetFloatingPointConversionType(*PinNet, *Linked);
-
-						if (ConversionType)
+						if (DestinationPin)
 						{
-							EKismetCompiledStatementType CastType = ConversionType->Get<0>();
-							const TCHAR* TermName = ConversionType->Get<1>();
+							TOptional<CastingUtils::StatementNamePair> ConversionType =
+								CastingUtils::GetFloatingPointConversionType(*PinNet, *DestinationPin);
+
+							if (ConversionType)
+							{
+								EKismetCompiledStatementType CastType = ConversionType->Get<0>();
+								const TCHAR* TermName = ConversionType->Get<1>();
 
-							FBPTerminal* NewTerm = Context.CreateLocalTerminal();
-							UEdGraphNode* OwningNode = Linked->GetOwningNode();
-							NewTerm->CopyFromPin(Linked, Context.NetNameMap->MakeValidName(Linked, TermName));
-							NewTerm->Source = OwningNode;
+								FBPTerminal* NewTerm = Context.CreateLocalTerminal();
+								UEdGraphNode* OwningNode = DestinationPin->GetOwningNode();
+								NewTerm->CopyFromPin(DestinationPin, Context.NetNameMap->MakeValidName(DestinationPin, TermName));
+								NewTerm->Source = OwningNode;
 
-							Context.ImplicitCastMap.Add(Linked, FImplicitCastParams{ CastType, NewTerm, OwningNode });
+								Context.ImplicitCastMap.Add(DestinationPin, FImplicitCastParams{ CastType, NewTerm, OwningNode });
+							}
 						}
 					}
 				}
