Blueprint: Add generic metadata on function that can be edit from a plugin.

diff --git a/Engine/Source/Editor/BlueprintGraph/Classes/K2Node_EditablePinBase.h b/Engine/Source/Editor/BlueprintGraph/Classes/K2Node_EditablePinBase.h
--- a/Engine/Source/Editor/BlueprintGraph/Classes/K2Node_EditablePinBase.h
+++ b/Engine/Source/Editor/BlueprintGraph/Classes/K2Node_EditablePinBase.h
@@ -45,7 +45,7 @@ struct FUserPinInfo
 USTRUCT()
 struct FKismetUserDeclaredFunctionMetadata
 {
-	GENERATED_USTRUCT_BODY()
+	GENERATED_BODY()
 
 	UPROPERTY()
 	FText ToolTip;
@@ -87,6 +87,23 @@ public:
 		, HasLatentFunctions(INDEX_NONE)
 	{
 	}
+
+	/** Set a metadata value on the function */
+	BLUEPRINTGRAPH_API void SetMetaData(FName Key, FString&& Value);
+	/** Set a metadata value on the function */
+	BLUEPRINTGRAPH_API void SetMetaData(FName Key, const FStringView Value);
+	/** Gets a metadata value on the function; asserts if the value isn't present. */
+	BLUEPRINTGRAPH_API const FString& GetMetaData(FName Key) const;
+	/** Clear metadata value on the function */
+	BLUEPRINTGRAPH_API void RemoveMetaData(FName Key);
+	/** Checks if there is metadata for a key */
+	BLUEPRINTGRAPH_API bool HasMetaData(FName Key) const;
+	/** Gets all metadata associated with this function */
+	BLUEPRINTGRAPH_API const TMap<FName, FString>& GetMetaDataMap() const;
+
+private:
+	UPROPERTY()
+	TMap<FName, FString> MetaDataMap;
 };
 
 UCLASS(abstract, MinimalAPI)
diff --git a/Engine/Source/Editor/BlueprintGraph/Private/K2Node_EditablePinBase.cpp b/Engine/Source/Editor/BlueprintGraph/Private/K2Node_EditablePinBase.cpp
--- a/Engine/Source/Editor/BlueprintGraph/Private/K2Node_EditablePinBase.cpp
+++ b/Engine/Source/Editor/BlueprintGraph/Private/K2Node_EditablePinBase.cpp
@@ -16,6 +16,65 @@
 //        with a version check and not have to do this anymore for updated assets
 #define ALWAYS_VALIDATE_DESIRED_PIN_DIRECTION_ON_LOAD 1
 
+
+//////////////////////////////////////////////////////////////////////////
+// FKismetUserDeclaredFunctionMetadata
+
+bool FKismetUserDeclaredFunctionMetadata::HasMetaData(FName Key) const
+{
+	const FString* ValuePtr = nullptr;
+	if (!Key.IsNone())
+	{
+		ValuePtr = MetaDataMap.Find(Key);
+	}
+	return ValuePtr != nullptr;
+}
+
+const FString& FKismetUserDeclaredFunctionMetadata::GetMetaData(FName Key) const
+{
+	// if not found, return a static empty string
+	static FString EmptyString;
+
+	if (Key.IsNone())
+	{
+		return EmptyString;
+	}
+	const FString* ValuePtr = MetaDataMap.Find(Key);
+	return ValuePtr ? *ValuePtr : EmptyString;
+}
+
+void FKismetUserDeclaredFunctionMetadata::SetMetaData(FName Key, FString&& Value)
+{
+	if (!Key.IsNone())
+	{
+		MetaDataMap.Add(Key, MoveTempIfPossible(Value));
+	}
+}
+
+void FKismetUserDeclaredFunctionMetadata::SetMetaData(FName Key, const FStringView Value)
+{
+	if (!Key.IsNone())
+	{
+		MetaDataMap.Add(Key, FString(Value));
+	}
+}
+
+void FKismetUserDeclaredFunctionMetadata::RemoveMetaData(FName Key)
+{
+	if (!Key.IsNone())
+	{
+		MetaDataMap.Remove(Key);
+	}
+}
+
+const TMap<FName, FString>& FKismetUserDeclaredFunctionMetadata::GetMetaDataMap() const
+{
+	return MetaDataMap;
+}
+
+//////////////////////////////////////////////////////////////////////////
+// FUserPinInfo
+
 FArchive& operator<<(FArchive& Ar, FUserPinInfo& Info)
 {
 	Ar.UsingCustomVersion(FFrameworkObjectVersion::GUID);
@@ -75,6 +134,8 @@ FArchive& operator<<(FArchive& Ar, FUserPinInfo& Info)
 	return Ar;
 }
 
+//////////////////////////////////////////////////////////////////////////
+// UK2Node_EditablePinBase
 
 UK2Node_EditablePinBase::UK2Node_EditablePinBase(const FObjectInitializer& ObjectInitializer)
 	: Super(ObjectInitializer)
diff --git a/Engine/Source/Editor/KismetCompiler/Private/KismetCompiler.cpp b/Engine/Source/Editor/KismetCompiler/Private/KismetCompiler.cpp
--- a/Engine/Source/Editor/KismetCompiler/Private/KismetCompiler.cpp
+++ b/Engine/Source/Editor/KismetCompiler/Private/KismetCompiler.cpp
@@ -2770,6 +2770,11 @@ void FKismetCompilerContext::SetCalculatedMetaDataAndFlags(UFunction* Function,
 			Function->SetMetaData(FBlueprintMetadata::MD_ThreadSafe, TEXT("true"));
 		}
 	}
+
+	for (const TPair<FName, FString>& Entry : FunctionMetaData.GetMetaDataMap())
+	{
+		Function->SetMetaData(Entry.Key, *Entry.Value);
+	}
 }
 
 void FKismetCompilerContext::SetDefaultInputValueMetaData(UFunction* Function, const TArray< TSharedPtr<FUserPinInfo> >& InputData)
diff --git a/Engine/Source/Runtime/Engine/Private/Blueprint.cpp b/Engine/Source/Runtime/Engine/Private/Blueprint.cpp
--- a/Engine/Source/Runtime/Engine/Private/Blueprint.cpp
+++ b/Engine/Source/Runtime/Engine/Private/Blueprint.cpp
@@ -174,7 +174,7 @@ void FBPVariableDescription::RemoveMetaData(const FName Key)
 	int32 EntryIndex = FindMetaDataEntryIndexForKey(Key);
 	if(EntryIndex != INDEX_NONE)
 	{
-		MetaDataArray.RemoveAt(EntryIndex);
+		MetaDataArray.RemoveAtSwap(EntryIndex);
 	}
 }
 
