Fix a crash when you connect a wildcard pin that has no owning node on it to a promotable operator Equals node. The function would hit a "Check" inside of the "GetOwningNode" function when dragging off of  an "Enum Select" node.

Also added an ensure to make sure that the function doesn't crash when evaluating pin changes. This wasn't happening,  but it is preventative just in case.

#rnx

diff --git a/Engine/Source/Editor/BlueprintGraph/Private/K2Node_PromotableOperator.cpp b/Engine/Source/Editor/BlueprintGraph/Private/K2Node_PromotableOperator.cpp
--- a/Engine/Source/Editor/BlueprintGraph/Private/K2Node_PromotableOperator.cpp
+++ b/Engine/Source/Editor/BlueprintGraph/Private/K2Node_PromotableOperator.cpp
@@ -844,6 +844,11 @@ void UK2Node_PromotableOperator::EvaluatePinsFromChange(UEdGraphPin* ChangedPin,
 {
 	UpdateOpName();
 
+	if (!ensureMsgf(ChangedPin, TEXT("UK2Node_PromotableOperator::EvaluatePinsFromChange failed to evaluate on a null pin!")))
+	{
+		return;
+	}
+
 	// True if the pin that has changed now has zero connections
 	const bool bWasAFullDisconnect = (ChangedPin->LinkedTo.Num() == 0) && !HasAnyConnectionsOrDefaults();
 	const bool bIsComparisonOp = FTypePromotion::IsComparisonOpName(OperationName);
@@ -857,7 +862,7 @@ void UK2Node_PromotableOperator::EvaluatePinsFromChange(UEdGraphPin* ChangedPin,
 	}
 	// If the pin that was connected is linked to a wildcard pin, then we should make it a wildcard
 	// and do nothing else.
-	else if (ChangedPin->GetOwningNode() == this && FWildcardNodeUtils::IsLinkedToWildcard(ChangedPin))
+	else if (ChangedPin->GetOwningNodeUnchecked() == this && FWildcardNodeUtils::IsLinkedToWildcard(ChangedPin))
 	{
 		return;
 	}
