Add "Degree ATan" and "Degree ATan2" to the k2node_mathexpression, as well as "DegreeToRadian" and "RadianToDegree" functoins.

diff --git a/Engine/Source/Editor/BlueprintGraph/Private/K2Node_MathExpression.cpp b/Engine/Source/Editor/BlueprintGraph/Private/K2Node_MathExpression.cpp
--- a/Engine/Source/Editor/BlueprintGraph/Private/K2Node_MathExpression.cpp
+++ b/Engine/Source/Editor/BlueprintGraph/Private/K2Node_MathExpression.cpp
@@ -977,6 +977,26 @@ private:
 			ADD_ALIAS("ARCTAN")
 		FUNC_ALIASES_END
 
+		FUNC_ALIASES_BEGIN("DegAtan")
+			ADD_ALIAS("DEGATAN")
+			ADD_ALIAS("DEGARCTAN")
+		FUNC_ALIASES_END
+
+		FUNC_ALIASES_BEGIN("DegAtan2")
+			ADD_ALIAS("DEGATAN2")
+			ADD_ALIAS("DEGARCTAN2")
+		FUNC_ALIASES_END
+
+		FUNC_ALIASES_BEGIN("DegreesToRadians")
+			ADD_ALIAS("DEGTORAD")
+			ADD_ALIAS("D2R")
+		FUNC_ALIASES_END
+
+		FUNC_ALIASES_BEGIN("RadiansToDegrees")
+			ADD_ALIAS("RADTODEG")
+			ADD_ALIAS("R2D")
+		FUNC_ALIASES_END
+
 		FUNC_ALIASES_BEGIN("ATan2")
 			ADD_ALIAS("ATAN2")
 			ADD_ALIAS("ARCTAN2")
