Provide custom context menu search criteria for material editor

diff --git a/Engine/Source/Editor/BlueprintGraph/Classes/EdGraphSchema_K2.h b/Engine/Source/Editor/BlueprintGraph/Classes/EdGraphSchema_K2.h
--- a/Engine/Source/Editor/BlueprintGraph/Classes/EdGraphSchema_K2.h
+++ b/Engine/Source/Editor/BlueprintGraph/Classes/EdGraphSchema_K2.h
@@ -558,6 +558,7 @@ public:
 
 #if WITH_EDITORONLY_DATA
 	virtual float GetActionFilteredWeight(const FGraphActionListBuilderBase::ActionGroup& InCurrentAction, const TArray<FString>& InFilterTerms, const TArray<FString>& InSanitizedFilterTerms, const TArray<UEdGraphPin*>& DraggedFromPins) const override;
+	virtual FGraphSchemaSearchWeightModifiers GetSearchWeightModifiers() const override;
 #endif // WITH_EDITORONLY_DATA	
 
 	//~ End EdGraphSchema Interface
diff --git a/Engine/Source/Editor/BlueprintGraph/Private/EdGraphSchema_K2.cpp b/Engine/Source/Editor/BlueprintGraph/Private/EdGraphSchema_K2.cpp
--- a/Engine/Source/Editor/BlueprintGraph/Private/EdGraphSchema_K2.cpp
+++ b/Engine/Source/Editor/BlueprintGraph/Private/EdGraphSchema_K2.cpp
@@ -7535,129 +7535,130 @@ bool UEdGraphSchema_K2::SafeDeleteNodeFromGraph(UEdGraph* Graph, UEdGraphNode* N
 
 //////////////////////////////////////////////////////////////////////////
 /** CVars for tweaking how the blueprint context menu search picks the best match */
-namespace ContextMenuConsoleVariables
+namespace BPContextMenuConsoleVariables
 {
 	/** Increasing this weight will give a bonus to shorter matching words */
 	static float ShorterWeight = 10.0f;
 	static FAutoConsoleVariableRef CVarShorterWeight(
-		TEXT("ContextMenu.ShorterWeight"), ShorterWeight,
+		TEXT("BP.ContextMenu.ShorterWeight"), ShorterWeight,
 		TEXT("Increasing this weight will make shorter words preferred"),
 		ECVF_Default);
 
 	/** When calculating shorter weight, this is the maximum length to make it relative to */
 	static int32 MaxWordLength = 30;
 	static FAutoConsoleVariableRef CVarMaxWordLength(
-		TEXT("ContextMenu.MaxWordLength"), MaxWordLength,
+		TEXT("BP.ContextMenu.MaxWordLength"), MaxWordLength,
 		TEXT("Maximum length to count while awarding short word weight"),
 		ECVF_Default);
 
 	/** Increasing this will prefer whole percentage matches when comparing the keyword to what the user has typed in */
 	static float PercentageMatchWeightMultiplier = 1.0f;
 	static FAutoConsoleVariableRef CVarPercentageMatchWeightMultiplier(
-		TEXT("ContextMenu.PercentageMatchWeightMultiplier"), PercentageMatchWeightMultiplier,
+		TEXT("BP.ContextMenu.PercentageMatchWeightMultiplier"), PercentageMatchWeightMultiplier,
 		TEXT("A multiplier for how much weight to give something based on the percentage match it is"),
 		ECVF_Default);
 
 	/** How much weight the description of actions have */
 	static float DescriptionWeight = 10.0f;
 	static FAutoConsoleVariableRef CVarDescriptionWeight(
-		TEXT("ContextMenu.DescriptionWeight"), DescriptionWeight,
+		TEXT("BP.ContextMenu.DescriptionWeight"), DescriptionWeight,
 		TEXT("The amount of weight placed on search items description"),
 		ECVF_Default);
 
 	/** Weight used to prefer categories that are the same as the node that was dragged off of */
 	static float MatchingFromPinCategory = 500.0f;
 	static FAutoConsoleVariableRef CVarMatchingFromPinCategory(
-		TEXT("ContextMenu.MatchingFromPinCategory"), MatchingFromPinCategory,
+		TEXT("BP.ContextMenu.MatchingFromPinCategory"), MatchingFromPinCategory,
 		TEXT("The amount of weight placed on actions with the same category as the node being dragged off of"),
 		ECVF_Default);
 
 	/** Weight that a match to a category search has */
 	static float CategoryWeight = 4.0f;
 	static FAutoConsoleVariableRef CVarCategoryWeight(
-		TEXT("ContextMenu.CategoryWeight"), CategoryWeight,
+		TEXT("BP.ContextMenu.CategoryWeight"), CategoryWeight,
 		TEXT("The amount of weight placed on categories that match what the user has typed in"),
 		ECVF_Default);
 
 	/** How much weight the node's title has */
 	static float NodeTitleWeight = 10.0f;
 	static FAutoConsoleVariableRef CVarNodeTitleWeight(
-		TEXT("ContextMenu.NodeTitleWeight"), NodeTitleWeight,
+		TEXT("BP.ContextMenu.NodeTitleWeight"), NodeTitleWeight,
 		TEXT("The amount of weight placed on the search items title"),
 		ECVF_Default);
 
 	/** Weight used to prefer keywords of actions  */
 	static float KeywordWeight = 30.0f;
 	static FAutoConsoleVariableRef CVarKeywordWeight(
-		TEXT("ContextMenu.KeywordWeight"), KeywordWeight,
+		TEXT("BP.ContextMenu.KeywordWeight"), KeywordWeight,
 		TEXT("The amount of weight placed on search items keyword"),
 		ECVF_Default);
 
 	/** The multiplier given if the keyword starts with a term the user typed in */
 	static float StartsWithBonusWeightMultiplier = 4.0f;
 	static FAutoConsoleVariableRef CVarStartsWithBonusWeightMultiplier(
-		TEXT("ContextMenu.StartsWithBonusWeightMultiplier"), StartsWithBonusWeightMultiplier,
+		TEXT("BP.ContextMenu.StartsWithBonusWeightMultiplier"), StartsWithBonusWeightMultiplier,
 		TEXT("The multiplier given if the keyword starts with a term the user typed in"),
 		ECVF_Default);
 
 	/** The multiplier given if the keyword contains a term the user typed in */
 	static float WordContainsLetterWeightMultiplier = 0.5f;
 	static FAutoConsoleVariableRef CVarWordContainsLetterWeightMultiplier(
-		TEXT("ContextMenu.WordContainsLetterWeightMultiplier"), WordContainsLetterWeightMultiplier,
+		TEXT("BP.ContextMenu.WordContainsLetterWeightMultiplier"), WordContainsLetterWeightMultiplier,
 		TEXT("The multiplier given if the keyword only contains a term the user typed in"),
 		ECVF_Default);
 
 	/** The bonus given if node is a favorite */
 	static float FavoriteBonus = 1000.0f;
 	static FAutoConsoleVariableRef CVarWordContainsLetterFavoriteBonus(
-		TEXT("ContextMenu.FavoriteBonus"), FavoriteBonus,
+		TEXT("BP.ContextMenu.FavoriteBonus"), FavoriteBonus,
 		TEXT("The bonus given if node is a favorite"),
 		ECVF_Default);
 
 	/** The bonus given if an action has the same container type as the dragged from pin */
 	static float ContainerBonus = 1000.0f;
 	static FAutoConsoleVariableRef CVarContainerBonus(
-		TEXT("ContextMenu.ContainerBonus"), ContainerBonus,
+		TEXT("BP.ContextMenu.ContainerBonus"), ContainerBonus,
 		TEXT("The bonus given if the dragged from pin matches the same container type of the action"),
 		ECVF_Default);
+};	// namespace BPContextMenuConsoleVariables
 
-	/** Enabling the debug printing of context menu selections */
-	static bool bPrintDebugContextSelection = false;
-	static FAutoConsoleVariableRef CVarPrintDebugContextSelection(
-		TEXT("ContextMenu.bPrintDebugContextSelection"), bPrintDebugContextSelection,
-		TEXT("Flag for printing the debug info about the context menu selection"),
-		ECVF_Default);
-
-};	// namespace ContextMenuConsoleVariables
+FGraphSchemaSearchWeightModifiers UEdGraphSchema_K2::GetSearchWeightModifiers() const
+{
+	FGraphSchemaSearchWeightModifiers Modifiers;
+	Modifiers.NodeTitleWeight = BPContextMenuConsoleVariables::NodeTitleWeight;
+	Modifiers.KeywordWeight = BPContextMenuConsoleVariables::KeywordWeight;
+	Modifiers.DescriptionWeight = BPContextMenuConsoleVariables::DescriptionWeight;
+	Modifiers.CategoryWeight = BPContextMenuConsoleVariables::DescriptionWeight;
+	Modifiers.WholeMatchLocalizedWeightMultiplier = BPContextMenuConsoleVariables::WordContainsLetterWeightMultiplier;
+	Modifiers.WholeMatchWeightMultiplier = BPContextMenuConsoleVariables::WordContainsLetterWeightMultiplier;
+	Modifiers.StartsWithBonusWeightMultiplier = BPContextMenuConsoleVariables::StartsWithBonusWeightMultiplier;
+	Modifiers.PercentageMatchWeightMultiplier = BPContextMenuConsoleVariables::PercentageMatchWeightMultiplier;
+	Modifiers.ShorterMatchWeight = BPContextMenuConsoleVariables::ShorterWeight;
+	return Modifiers;
+}
 
 /**
 * Debug Info about how the preferred context menu action is chosen
 * @see SGraphActionMenu::GetActionFilteredWeight
 */
-struct FContextMenuWeightDebugInfo
+struct FBPContextMenuWeightDebugInfo : public FGraphSchemaSearchTextDebugInfo
 {
-	float TotalWeight = 0.0f;
-	// Info about which sets of words match up with what weight
-	float KeywordArrayWeight = 0.0f;
-	float DescriptionWeight = 0.0f;
-	float NodeTitleWeight = 0.0f;
-	float CategoryWeight = 0.0f;
 	float FavoriteBonusWeight = 0.0f;
-
-	float PercentageMatchWeight = 0.0f;
-	/** The calculated match percentage */
-	float PercMatch = 0.0f;
-	float ShorterWeight = 0.0f;
 	float CategoryBonusWeight = 0.0f;
 
 	/**
 	* Print out the debug info about this weight info to the console
 	*/
-	void Print(const FString& Desc)
+	virtual void Print(const TArray<FString>& SearchForKeywords, const FGraphActionListBuilderBase::ActionGroup& Action) const override
 	{
+		// Combine the actions string, separate with \n so terms don't run into each other, and remove the spaces (incase the user is searching for a variable)
+		// In the case of groups containing multiple actions, they will have been created and added at the same place in the code, using the same description
+		// and keywords, so we only need to use the first one for filtering.
+		const FString& SearchText = Action.GetSearchTextForFirstAction();
+
 		UE_LOG(LogTemp, Warning, TEXT("[Weight for %s] \
 TotalWeight: %-8.2f | PercentageMatchWeight: %-8.2f | PercMatch: %-8.2f | ShorterWeight: %-8.2f | CategoryBonusWeight: %-8.2f | KeywordArrayWeight: %-8.2f | DescriptionWeight: %-8.2f | NodeTitleWeight: %-8.2f | CategoryWeight: %-8.2f | Fav. Bonus:%-8.2f\n"),
-			*Desc, TotalWeight, PercentageMatchWeight, PercMatch, ShorterWeight, CategoryBonusWeight, KeywordArrayWeight, DescriptionWeight, NodeTitleWeight, CategoryWeight, FavoriteBonusWeight);
+			*SearchText, TotalWeight, PercentMatchWeight, PercentMatch, ShorterMatchWeight, CategoryBonusWeight, KeywordWeight, DescriptionWeight, NodeTitleWeight, CategoryWeight, FavoriteBonusWeight);
 	}
 };
 
@@ -7665,24 +7666,10 @@ float UEdGraphSchema_K2::GetActionFilteredWeight(const FGraphActionListBuilderBa
 {
 	// The overall 'weight' of this action 
 	float TotalWeight = 0.0f;
-	// Helper array for tracking actions and the weight that they have added
-	struct FArrayWithWeight
-	{
-		FArrayWithWeight(const TArray< FString >* InArray, float InWeightModifier, float* OutDebugWeight)
-			: Array(InArray)
-			, OutWeight(OutDebugWeight)
-			, WeightModifier(InWeightModifier)
-		{
-		}
-
-		const TArray< FString >* Array = nullptr;
-		float* OutWeight = nullptr;
-		float WeightModifier = 0.0f;
-	};
 
 	// Setup an array of arrays so we can do a weighted search			
-	TArray< FArrayWithWeight > WeightedArrayList;
-	FContextMenuWeightDebugInfo OutDebugInfo;
+	TArray< FGraphSchemaSearchTextWeightInfo > WeightedArrayList;
+	FBPContextMenuWeightDebugInfo OutDebugInfo;
 
 	const bool bIsFromDrag = (DraggedFromPins.Num() > 0);
 
@@ -7691,38 +7678,12 @@ float UEdGraphSchema_K2::GetActionFilteredWeight(const FGraphActionListBuilderBa
 	{
 		TSharedPtr<FEdGraphSchemaAction> CurrentAction = InCurrentAction.Actions[Action];
 
-		// Combine the actions string, separate with \n so terms don't run into each other, and remove the spaces (incase the user is searching for a variable)
-		// In the case of groups containing multiple actions, they will have been created and added at the same place in the code, using the same description
-		// and keywords, so we only need to use the first one for filtering.
-		const FString& SearchText = InCurrentAction.GetSearchTextForFirstAction();
-
+		FGraphSchemaSearchWeightModifiers WeightModifiers = GetSearchWeightModifiers();
 		// If there are no keywords, bump the weight on description to compensate
 		const TArray<FString>& LocKeywords = InCurrentAction.GetLocalizedSearchKeywordsArrayForFirstAction();
-		const float DescriptionWeight = LocKeywords.Num() > 0 ? ContextMenuConsoleVariables::DescriptionWeight : ContextMenuConsoleVariables::DescriptionWeight * 2.0f;
-
-		// First the localized keywords
-		WeightedArrayList.Add(FArrayWithWeight(&LocKeywords, ContextMenuConsoleVariables::KeywordWeight, &OutDebugInfo.KeywordArrayWeight));
-
-		// The localized description
-		WeightedArrayList.Add(FArrayWithWeight(&InCurrentAction.GetLocalizedMenuDescriptionArrayForFirstAction(), DescriptionWeight, &OutDebugInfo.DescriptionWeight));
+		WeightModifiers.DescriptionWeight = LocKeywords.Num() > 0 ? WeightModifiers.DescriptionWeight : WeightModifiers.DescriptionWeight * 2.0f;
 
-		// The node search localized title weight, in most case this will be the same as description
-		WeightedArrayList.Add(FArrayWithWeight(&InCurrentAction.GetLocalizedSearchTitleArrayForFirstAction(), ContextMenuConsoleVariables::NodeTitleWeight, &OutDebugInfo.NodeTitleWeight));
-
-		// The localized category
-		WeightedArrayList.Add(FArrayWithWeight(&InCurrentAction.GetLocalizedSearchCategoryArrayForFirstAction(), ContextMenuConsoleVariables::CategoryWeight, &OutDebugInfo.CategoryWeight));
-
-		// First the keywords
-		WeightedArrayList.Add(FArrayWithWeight(&InCurrentAction.GetSearchKeywordsArrayForFirstAction(), ContextMenuConsoleVariables::KeywordWeight, &OutDebugInfo.KeywordArrayWeight));
-
-		// The description
-		WeightedArrayList.Add(FArrayWithWeight(&InCurrentAction.GetMenuDescriptionArrayForFirstAction(), DescriptionWeight, &OutDebugInfo.DescriptionWeight));
-
-		// The node search title weight
-		WeightedArrayList.Add(FArrayWithWeight(&InCurrentAction.GetSearchTitleArrayForFirstAction(), ContextMenuConsoleVariables::NodeTitleWeight, &OutDebugInfo.NodeTitleWeight));
-
-		// The category
-		WeightedArrayList.Add(FArrayWithWeight(&InCurrentAction.GetSearchCategoryArrayForFirstAction(), ContextMenuConsoleVariables::CategoryWeight, &OutDebugInfo.CategoryWeight));
+		CollectSearchTextWeightInfo(InCurrentAction, WeightModifiers, WeightedArrayList, &OutDebugInfo);
 
 		// Give a weight bonus to actions whose category matches what was dragged off of
 		if (bIsFromDrag)
@@ -7762,7 +7723,7 @@ float UEdGraphSchema_K2::GetActionFilteredWeight(const FGraphActionListBuilderBa
 					// For containers, add a preference for functions that are marked in their category
 					if (!bAddedContainerPreferenceBonus && FromPin->PinType.IsContainer() && InActionCategory == GetContainerTypeString(FromPin->PinType.ContainerType))
 					{
-						TotalWeight += ContextMenuConsoleVariables::ContainerBonus;
+						TotalWeight += BPContextMenuConsoleVariables::ContainerBonus;
 						bAddedContainerPreferenceBonus = true;
 					}
 
@@ -7791,8 +7752,8 @@ float UEdGraphSchema_K2::GetActionFilteredWeight(const FGraphActionListBuilderBa
 					// If we found match in any cases above then add the weight bonus and stop looking
 					if (bAddMatchBonus)
 					{
-						TotalWeight += ContextMenuConsoleVariables::MatchingFromPinCategory;
-						OutDebugInfo.CategoryBonusWeight += ContextMenuConsoleVariables::MatchingFromPinCategory;
+						TotalWeight += BPContextMenuConsoleVariables::MatchingFromPinCategory;
+						OutDebugInfo.CategoryBonusWeight += BPContextMenuConsoleVariables::MatchingFromPinCategory;
 
 						// Break out of the loop so that we don't give any extra bonuses
 						break;
@@ -7807,8 +7768,8 @@ float UEdGraphSchema_K2::GetActionFilteredWeight(const FGraphActionListBuilderBa
 		{
 			if (BlueprintFavorites->IsFavorited(CurrentAction))
 			{
-				TotalWeight += ContextMenuConsoleVariables::FavoriteBonus;
-				OutDebugInfo.FavoriteBonusWeight += ContextMenuConsoleVariables::FavoriteBonus;
+				TotalWeight += BPContextMenuConsoleVariables::FavoriteBonus;
+				OutDebugInfo.FavoriteBonusWeight += BPContextMenuConsoleVariables::FavoriteBonus;
 			}
 		}
 
@@ -7845,12 +7806,12 @@ float UEdGraphSchema_K2::GetActionFilteredWeight(const FGraphActionListBuilderBa
 					if (KeywordArray[iEachWord].Contains(*EachTermSanitized, ESearchCase::CaseSensitive) || KeywordArray[iEachWord].Contains(*EachTerm, ESearchCase::CaseSensitive))
 					{
 						++WordMatchCount;
-						WeightPerWord += KeywordArrayWeight * ContextMenuConsoleVariables::WordContainsLetterWeightMultiplier;
+						WeightPerWord += KeywordArrayWeight * BPContextMenuConsoleVariables::WordContainsLetterWeightMultiplier;
 
 						// If the word starts with the letter, give it a little extra boost of weight
 						if (KeywordArray[iEachWord].StartsWith(*EachTermSanitized, ESearchCase::CaseSensitive) || KeywordArray[iEachWord].StartsWith(*EachTerm, ESearchCase::CaseSensitive))
 						{
-							WeightPerWord += KeywordArrayWeight * ContextMenuConsoleVariables::StartsWithBonusWeightMultiplier;
+							WeightPerWord += KeywordArrayWeight * BPContextMenuConsoleVariables::StartsWithBonusWeightMultiplier;
 						}
 
 						if (WeightPerWord > WeightPerList)
@@ -7870,33 +7831,30 @@ float UEdGraphSchema_K2::GetActionFilteredWeight(const FGraphActionListBuilderBa
 					// How many words that we are checking had partial matches compared to what the user typed in?
 					float PercMatch = (float)WordMatchCount / (float)KeywordArray.Num();
 
-					float PercentageBonus = (WeightPerList * PercMatch * ContextMenuConsoleVariables::PercentageMatchWeightMultiplier);
+					float PercentageBonus = (WeightPerList * PercMatch * BPContextMenuConsoleVariables::PercentageMatchWeightMultiplier);
 					WeightPerList += PercentageBonus;
 
 					// The shorter the matching word, the larger bonus it gets
-					float ShortFactor = ContextMenuConsoleVariables::MaxWordLength - FMath::Min(BestMatchCharLength, ContextMenuConsoleVariables::MaxWordLength);
-					float ShortWeight = ShortFactor * ContextMenuConsoleVariables::ShorterWeight * (bIsFromDrag ? 0.25f : 1.0f);
+					float ShortFactor = BPContextMenuConsoleVariables::MaxWordLength - FMath::Min(BestMatchCharLength, BPContextMenuConsoleVariables::MaxWordLength);
+					float ShortWeight = ShortFactor * BPContextMenuConsoleVariables::ShorterWeight * (bIsFromDrag ? 0.25f : 1.0f);
 					WeightPerList += ShortWeight;
 
-					OutDebugInfo.PercMatch += PercMatch;
-					OutDebugInfo.ShorterWeight += ShortWeight;
-					OutDebugInfo.PercentageMatchWeight += PercentageBonus;
+					OutDebugInfo.PercentMatch += PercMatch;
+					OutDebugInfo.ShorterMatchWeight += ShortWeight;
+					OutDebugInfo.PercentMatchWeight += PercentageBonus;
 				}
 
 				TotalWeight += WeightPerList;
-				if (WeightedArrayList[iFindCount].OutWeight)
+				if (WeightedArrayList[iFindCount].DebugWeight)
 				{
 					// Each weight is used twice so add them
-					*WeightedArrayList[iFindCount].OutWeight += WeightPerList;
+					*WeightedArrayList[iFindCount].DebugWeight += WeightPerList;
 				}
 			}
 		}
 		OutDebugInfo.TotalWeight = TotalWeight;
 
-		if (ContextMenuConsoleVariables::bPrintDebugContextSelection)
-		{
-			OutDebugInfo.Print(SearchText);
-		}
+		PrintSearchTextDebugInfo(InFilterTerms, CurrentAction, &OutDebugInfo);
 	}
 
 	return TotalWeight;
diff --git a/Engine/Source/Editor/UnrealEd/Classes/MaterialGraph/MaterialGraphSchema.h b/Engine/Source/Editor/UnrealEd/Classes/MaterialGraph/MaterialGraphSchema.h
--- a/Engine/Source/Editor/UnrealEd/Classes/MaterialGraph/MaterialGraphSchema.h
+++ b/Engine/Source/Editor/UnrealEd/Classes/MaterialGraph/MaterialGraphSchema.h
@@ -265,6 +265,10 @@ class UMaterialGraphSchema : public UEdGraphSchema
 	virtual void OnPinConnectionDoubleCicked(UEdGraphPin* PinA, UEdGraphPin* PinB, const FVector2D& GraphPosition) const override;
 	virtual bool SafeDeleteNodeFromGraph(UEdGraph* Graph, UEdGraphNode* NodeToDelete) const override;
 	virtual void GetAssetsGraphHoverMessage(const TArray<FAssetData>& Assets, const UEdGraph* HoverGraph, FString& OutTooltipText, bool& OutOkIcon) const;
+#if WITH_EDITORONLY_DATA
+	virtual float GetActionFilteredWeight(const FGraphActionListBuilderBase::ActionGroup& InCurrentAction, const TArray<FString>& InFilterTerms, const TArray<FString>& InSanitizedFilterTerms, const TArray<UEdGraphPin*>& DraggedFromPins) const override;
+	virtual FGraphSchemaSearchWeightModifiers GetSearchWeightModifiers() const override;
+#endif // WITH_EDITORONLY_DATA	
 	//~ End UEdGraphSchema Interface
 
 private:
diff --git a/Engine/Source/Editor/UnrealEd/Private/MaterialGraphSchema.cpp b/Engine/Source/Editor/UnrealEd/Private/MaterialGraphSchema.cpp
--- a/Engine/Source/Editor/UnrealEd/Private/MaterialGraphSchema.cpp
+++ b/Engine/Source/Editor/UnrealEd/Private/MaterialGraphSchema.cpp
@@ -1144,4 +1144,180 @@ void UMaterialGraphSchema::GetAssetsGraphHoverMessage(const TArray<FAssetData>&
 	}
 }
 
+#if WITH_EDITORONLY_DATA
+//////////////////////////////////////////////////////////////////////////
+/** CVars for tweaking how the material editor context menu search picks the best match */
+namespace MaterialEditorContextMenuConsoleVariables
+{
+	/** How much weight the node's title has */
+	static float NodeTitleWeight = 20.0f;
+	static FAutoConsoleVariableRef CVarNodeTitleWeight(
+		TEXT("r.MaterialEditor.ContextMenu.NodeTitleWeight"), NodeTitleWeight,
+		TEXT("The amount of weight placed on the search items title"),
+		ECVF_Default);
+
+	/** Weight used to prefer keywords of actions  */
+	static float KeywordWeight = 30.0f;
+	static FAutoConsoleVariableRef CVarKeywordWeight(
+		TEXT("r.MaterialEditor.ContextMenu.KeywordWeight"), KeywordWeight,
+		TEXT("The amount of weight placed on search items keyword"),
+		ECVF_Default);
+
+	/** Weight that a match to a description search has */
+	static float DescriptionWeight = 4.0f;
+	static FAutoConsoleVariableRef CVarDescriptionWeight(
+		TEXT("r.MaterialEditor.ContextMenu.DescriptionWeight"), DescriptionWeight,
+		TEXT("The amount of weight placed on description that match what the user has typed in"),
+		ECVF_Default);
+
+	/** Weight that a match to a category search has */
+	static float CategoryWeight = 4.0f;
+	static FAutoConsoleVariableRef CVarCategoryWeight(
+		TEXT("r.MaterialEditor.ContextMenu.CategoryWeight"), CategoryWeight,
+		TEXT("The amount of weight placed on categories that match what the user has typed in"),
+		ECVF_Default);
+
+	/** The multiplier given if there is an exact localized match to the search term */
+	static float WholeMatchLocalizedWeightMultiplier = 0.5f;
+	static FAutoConsoleVariableRef CVarWholeMatchLocalizedWeightMultiplier(
+		TEXT("r.MaterialEditor.ContextMenu.WholeMatchLocalizedWeightMultiplier"), WholeMatchLocalizedWeightMultiplier,
+		TEXT("The multiplier given if there is an exact localized match to the search term"),
+		ECVF_Default);
+
+	/** The multiplier given if there is an exact match to the search term */
+	static float WholeMatchWeightMultiplier = 0.5f;
+	static FAutoConsoleVariableRef CVarWholeMatchWeightMultiplier(
+		TEXT("r.MaterialEditor.ContextMenu.WholeMatchWeightMultiplier"), WholeMatchWeightMultiplier,
+		TEXT("The multiplier given if there is an exact match to the search term"),
+		ECVF_Default);
+
+	/** The multiplier given if the keyword starts with a term the user typed in */
+	static float StartsWithBonusWeightMultiplier = 4.0f;
+	static FAutoConsoleVariableRef CVarStartsWithBonusWeightMultiplier(
+		TEXT("r.MaterialEditor.ContextMenu.StartsWithBonusWeightMultiplier"), StartsWithBonusWeightMultiplier,
+		TEXT("The multiplier given if the keyword starts with a term the user typed in"),
+		ECVF_Default);
+
+	/** Increasing this will prefer whole percentage matches when comparing the keyword to what the user has typed in */
+	static float PercentageMatchWeightMultiplier = 1.0f;
+	static FAutoConsoleVariableRef CVarPercentageMatchWeightMultiplier(
+		TEXT("r.MaterialEditor.ContextMenu.PercentageMatchWeightMultiplier"), PercentageMatchWeightMultiplier,
+		TEXT("A multiplier for how much weight to give something based on the percentage match it is"),
+		ECVF_Default);
+
+	/** Increasing this weight will give a bonus to shorter matching words */
+	static float ShorterMatchWeight = 10.0f;
+	static FAutoConsoleVariableRef CVarShorterWeight(
+		TEXT("r.MaterialEditor.ContextMenu.ShorterMatchWeight"), ShorterMatchWeight,
+		TEXT("Increasing this weight will make shorter words preferred"),
+		ECVF_Default);
+}
+
+FGraphSchemaSearchWeightModifiers UMaterialGraphSchema::GetSearchWeightModifiers() const
+{
+	FGraphSchemaSearchWeightModifiers Modifiers;
+	Modifiers.NodeTitleWeight = MaterialEditorContextMenuConsoleVariables::NodeTitleWeight;
+	Modifiers.KeywordWeight = MaterialEditorContextMenuConsoleVariables::KeywordWeight;
+	Modifiers.DescriptionWeight = MaterialEditorContextMenuConsoleVariables::DescriptionWeight;
+	Modifiers.CategoryWeight = MaterialEditorContextMenuConsoleVariables::DescriptionWeight;
+	Modifiers.WholeMatchLocalizedWeightMultiplier = MaterialEditorContextMenuConsoleVariables::WholeMatchLocalizedWeightMultiplier;
+	Modifiers.WholeMatchWeightMultiplier = MaterialEditorContextMenuConsoleVariables::WholeMatchWeightMultiplier;
+	Modifiers.StartsWithBonusWeightMultiplier = MaterialEditorContextMenuConsoleVariables::StartsWithBonusWeightMultiplier;
+	Modifiers.PercentageMatchWeightMultiplier = MaterialEditorContextMenuConsoleVariables::PercentageMatchWeightMultiplier;
+	Modifiers.ShorterMatchWeight = MaterialEditorContextMenuConsoleVariables::ShorterMatchWeight;
+	return Modifiers;
+}
+
+float UMaterialGraphSchema::GetActionFilteredWeight(const FGraphActionListBuilderBase::ActionGroup& InCurrentAction, const TArray<FString>& InFilterTerms, const TArray<FString>& InSanitizedFilterTerms, const TArray<UEdGraphPin*>& DraggedFromPins) const
+{
+	// The overall 'weight'
+	float TotalWeight = 0.0f;
+
+	// Setup an array of arrays so we can do a weighted search			
+	TArray<FGraphSchemaSearchTextWeightInfo> WeightedArrayList;
+	FGraphSchemaSearchTextDebugInfo DebugInfo;
+
+	int32 Action = 0;
+	if (InCurrentAction.Actions[Action].IsValid() == true)
+	{
+		FGraphSchemaSearchWeightModifiers WeightModifiers = GetSearchWeightModifiers();
+		int32 NonLocalizedFirstIndex = CollectSearchTextWeightInfo(InCurrentAction, WeightModifiers, WeightedArrayList, &DebugInfo);
+
+		// Now iterate through all the filter terms and calculate a 'weight' using the values and multipliers
+		for (int32 FilterIndex = 0; FilterIndex < InFilterTerms.Num(); ++FilterIndex)
+		{
+			const FString& EachTerm = InFilterTerms[FilterIndex];
+			const FString& EachTermSanitized = InSanitizedFilterTerms[FilterIndex];
+			// Now check the weighted lists
+			for (int32 iFindCount = 0; iFindCount < WeightedArrayList.Num(); iFindCount++)
+			{
+				float WeightPerList = 0.0f;
+				const TArray<FString>& WordArray = *WeightedArrayList[iFindCount].Array;
+				float ArrayWeight = WeightedArrayList[iFindCount].WeightModifier;
+				int32 WholeMatchCount = 0;
+				float WholeMatchMultiplier = (iFindCount < NonLocalizedFirstIndex) ? MaterialEditorContextMenuConsoleVariables::WholeMatchLocalizedWeightMultiplier : MaterialEditorContextMenuConsoleVariables::WholeMatchWeightMultiplier;
+
+				// Count of how many words in this array contain a search term that the user has typed in
+				int32 WordMatchCount = 0;
+				// The number of characters in the best matching word
+				int32 BestMatchCharLength = 0;
+
+				for (int32 iEachWord = 0; iEachWord < WordArray.Num(); iEachWord++)
+				{
+					float WeightPerWord = 0.0f;
+
+					// If a word contains the search phrase that the user has typed in, then give it weight					
+					if (WordArray[iEachWord].Contains(EachTermSanitized, ESearchCase::CaseSensitive) || WordArray[iEachWord].Contains(EachTerm, ESearchCase::CaseSensitive))
+					{
+						++WordMatchCount;
+						WeightPerWord += ArrayWeight * WholeMatchMultiplier;
+
+						// If the word starts with the search term, give it extra boost of weight
+						if (WordArray[iEachWord].StartsWith(EachTermSanitized, ESearchCase::CaseSensitive) || WordArray[iEachWord].StartsWith(EachTerm, ESearchCase::CaseSensitive))
+						{
+							WeightPerWord += ArrayWeight * MaterialEditorContextMenuConsoleVariables::StartsWithBonusWeightMultiplier;
+						}
+					}
+
+					if (WeightPerWord > WeightPerList)
+					{
+						// Use the best word match weight, we don't want to count similar words more than one
+						WeightPerList = WeightPerWord;
+						BestMatchCharLength = WordArray[iEachWord].Len();
+					}
+				}
+
+				if (BestMatchCharLength > 0 && WeightPerList > 0)
+				{
+					// Higher number of matching words contributes to higher weight
+					float PercentMatch = (float)WordMatchCount / (float)WordArray.Num();
+					float PercentMatchWeight = (WeightPerList * PercentMatch * MaterialEditorContextMenuConsoleVariables::PercentageMatchWeightMultiplier);
+					WeightPerList += PercentMatchWeight;
+					DebugInfo.PercentMatchWeight += PercentMatchWeight;
+					DebugInfo.PercentMatch += PercentMatch;
+
+					// The shorter the best matched word, the larger bonus it gets
+					float ShorterMatchFactor = (float)EachTerm.Len() / (float)BestMatchCharLength;
+					float ShorterMatchWeight = ShorterMatchFactor * MaterialEditorContextMenuConsoleVariables::ShorterMatchWeight;
+					WeightPerList += ShorterMatchWeight;
+					DebugInfo.ShorterMatchWeight += ShorterMatchWeight;
+				}
+
+				if (WeightedArrayList[iFindCount].DebugWeight)
+				{
+					*WeightedArrayList[iFindCount].DebugWeight += WeightPerList;
+				}
+
+				TotalWeight += WeightPerList;
+			}
+		}
+
+		DebugInfo.TotalWeight = TotalWeight;
+		PrintSearchTextDebugInfo(InFilterTerms, InCurrentAction, &DebugInfo);
+	}
+
+	return TotalWeight;
+}
+#endif // WITH_EDITORONLY_DATA
+
 #undef LOCTEXT_NAMESPACE
\ No newline at end of file
diff --git a/Engine/Source/Runtime/Engine/Classes/EdGraph/EdGraphSchema.h b/Engine/Source/Runtime/Engine/Classes/EdGraph/EdGraphSchema.h
--- a/Engine/Source/Runtime/Engine/Classes/EdGraph/EdGraphSchema.h
+++ b/Engine/Source/Runtime/Engine/Classes/EdGraph/EdGraphSchema.h
@@ -622,6 +622,50 @@ public:
 	ENGINE_API FString GetNotesAsString() const;
 };
 
+#if WITH_EDITORONLY_DATA
+struct FGraphSchemaSearchWeightModifiers
+{
+	float NodeTitleWeight = 0.0f;
+	float KeywordWeight = 0.0f;
+	float DescriptionWeight = 0.0f;
+	float CategoryWeight = 0.0f;
+	float WholeMatchLocalizedWeightMultiplier = 0.0f;
+	float WholeMatchWeightMultiplier = 0.0f;
+	float StartsWithBonusWeightMultiplier = 0.0f;
+	float PercentageMatchWeightMultiplier = 0.0f;
+	float ShorterMatchWeight = 0.0f;
+};
+
+// Helper struct storing the search text array with its weight info
+struct FGraphSchemaSearchTextWeightInfo
+{
+	FGraphSchemaSearchTextWeightInfo(const TArray< FString >* InArray, float InWeightModifier, float* OutDebugWeight)
+		: Array(InArray), WeightModifier(InWeightModifier), DebugWeight(OutDebugWeight)
+	{}
+
+	const TArray< FString >* Array = nullptr;
+	float WeightModifier = 0.0f;
+	float* DebugWeight = nullptr;
+};
+
+// Helper struct storing the breakdown of the weights assigned to the search text
+struct FGraphSchemaSearchTextDebugInfo
+{
+	float TotalWeight = 0.0f;			// Overall weight
+
+	float NodeTitleWeight = 0.0f;		// Weight for the node's title
+	float KeywordWeight = 0.0f;			// Weight for the node's keywords
+	float DescriptionWeight = 0.0f;		// Weight for the node's description
+	float CategoryWeight = 0.0f;		// Weight for the category
+
+	float PercentMatch = 0.0f;			// The calculated whole match percentage
+	float PercentMatchWeight = 0.0f;	// Weight for the whole match percentage
+	float ShorterMatchWeight = 0.0f;	// Weight for the shorter matched words
+
+	/** Print out the debug info about this weight info to the console */
+	ENGINE_API virtual void Print(const TArray<FString>& SearchForKeywords, const FGraphActionListBuilderBase::ActionGroup& Action) const;
+};
+#endif // WITH_EDITORONLY_DATA
 
 UCLASS(abstract)
 class ENGINE_API UEdGraphSchema : public UObject
@@ -844,6 +888,9 @@ class ENGINE_API UEdGraphSchema : public UObject
 	 * @param DraggedFromPins				Any pins that this action was dragged off of
 	 */
 	virtual float GetActionFilteredWeight(const FGraphActionListBuilderBase::ActionGroup& InCurrentAction, const TArray<FString>& InFilterTerms, const TArray<FString>& InSanitizedFilterTerms, const TArray<UEdGraphPin*>& DraggedFromPins) const;
+
+	/** Get the weight modifiers from the console variable settings */
+	virtual FGraphSchemaSearchWeightModifiers GetSearchWeightModifiers() const;
 #endif // WITH_EDITORONLY_DATA
 
 	/**
@@ -1255,4 +1302,12 @@ class ENGINE_API UEdGraphSchema : public UObject
 	}
 #endif
 	
+#if WITH_EDITORONLY_DATA
+protected:
+	/** Build an array containing all search types, return the index of the first non-localized entry. */
+	int32 CollectSearchTextWeightInfo(const FGraphActionListBuilderBase::ActionGroup& InCurrentAction, const FGraphSchemaSearchWeightModifiers& InWeightModifiers,
+		TArray<FGraphSchemaSearchTextWeightInfo>& OutWeightedArrayList, FGraphSchemaSearchTextDebugInfo* InDebugInfo) const;
+
+	void PrintSearchTextDebugInfo(const TArray<FString>& InFilterTerms, const FGraphActionListBuilderBase::ActionGroup& InCurrentAction, const FGraphSchemaSearchTextDebugInfo* InDebugInfo) const;
+#endif // WITH_EDITORONLY_DATA
 };
diff --git a/Engine/Source/Runtime/Engine/Private/EdGraph/EdGraphSchema.cpp b/Engine/Source/Runtime/Engine/Private/EdGraph/EdGraphSchema.cpp
--- a/Engine/Source/Runtime/Engine/Private/EdGraph/EdGraphSchema.cpp
+++ b/Engine/Source/Runtime/Engine/Private/EdGraph/EdGraphSchema.cpp
@@ -759,30 +759,6 @@ float UEdGraphSchema::GetActionFilteredWeight(const FGraphActionListBuilderBase:
 	// The overall 'weight'
 	int32 TotalWeight = 0;
 
-	// Some simple weight figures to help find the most appropriate match	
-	const int32 WholeMatchWeightMultiplier = 2;
-	const int32 WholeMatchLocalizedWeightMultiplier = 3;
-	const int32 DescriptionWeight = 10;
-	const int32 CategoryWeight = 1;
-	const int32 NodeTitleWeight = 1;
-	const int32 KeywordWeight = 4;
-
-	// Helper array
-	struct FArrayWithWeight
-	{
-		FArrayWithWeight(const TArray< FString >* InArray, int32 InWeight)
-			: Array(InArray)
-			, Weight(InWeight)
-		{
-		}
-
-		const TArray< FString >* Array;
-		int32 Weight;
-	};
-
-	// Setup an array of arrays so we can do a weighted search			
-	TArray<FArrayWithWeight> WeightedArrayList;
-
 	int32 Action = 0;
 	if (InCurrentAction.Actions[Action].IsValid() == true)
 	{
@@ -791,29 +767,11 @@ float UEdGraphSchema::GetActionFilteredWeight(const FGraphActionListBuilderBase:
 		// and keywords, so we only need to use the first one for filtering.
 		const FString& SearchText = InCurrentAction.GetSearchTextForFirstAction();
 
-		// First the localized keywords
-		WeightedArrayList.Add(FArrayWithWeight(&InCurrentAction.GetLocalizedSearchKeywordsArrayForFirstAction(), KeywordWeight));
-
-		// The localized description
-		WeightedArrayList.Add(FArrayWithWeight(&InCurrentAction.GetLocalizedMenuDescriptionArrayForFirstAction(), DescriptionWeight));
-
-		// The node search localized title weight
-		WeightedArrayList.Add(FArrayWithWeight(&InCurrentAction.GetLocalizedSearchTitleArrayForFirstAction(), NodeTitleWeight));
-
-		// The localized category
-		WeightedArrayList.Add(FArrayWithWeight(&InCurrentAction.GetLocalizedSearchCategoryArrayForFirstAction(), CategoryWeight));
-
-		// First the keywords
-		int32 NonLocalizedFirstIndex = WeightedArrayList.Add(FArrayWithWeight(&InCurrentAction.GetSearchKeywordsArrayForFirstAction(), KeywordWeight));
-
-		// The description
-		WeightedArrayList.Add(FArrayWithWeight(&InCurrentAction.GetMenuDescriptionArrayForFirstAction(), DescriptionWeight));
-
-		// The node search title weight
-		WeightedArrayList.Add(FArrayWithWeight(&InCurrentAction.GetSearchTitleArrayForFirstAction(), NodeTitleWeight));
-
-		// The category
-		WeightedArrayList.Add(FArrayWithWeight(&InCurrentAction.GetSearchCategoryArrayForFirstAction(), CategoryWeight));
+		// Setup an array of arrays so we can do a weighted search			
+		TArray<FGraphSchemaSearchTextWeightInfo> WeightedArrayList;
+		FGraphSchemaSearchWeightModifiers WeightModifiers = GetSearchWeightModifiers();
+		FGraphSchemaSearchTextDebugInfo DebugInfo;
+		int32 NonLocalizedFirstIndex = CollectSearchTextWeightInfo(InCurrentAction, WeightModifiers, WeightedArrayList, &DebugInfo);
 
 		// Now iterate through all the filter terms and calculate a 'weight' using the values and multipliers
 		const FString* EachTerm = nullptr;
@@ -835,9 +793,9 @@ float UEdGraphSchema::GetActionFilteredWeight(const FGraphActionListBuilderBase:
 			{
 				int32 WeightPerList = 0;
 				const TArray<FString>& KeywordArray = *WeightedArrayList[iFindCount].Array;
-				int32 EachWeight = WeightedArrayList[iFindCount].Weight;
+				int32 EachWeight = WeightedArrayList[iFindCount].WeightModifier;
 				int32 WholeMatchCount = 0;
-				int32 WholeMatchMultiplier = (iFindCount < NonLocalizedFirstIndex) ? WholeMatchLocalizedWeightMultiplier : WholeMatchWeightMultiplier;
+				int32 WholeMatchMultiplier = (iFindCount < NonLocalizedFirstIndex) ? WeightModifiers.WholeMatchLocalizedWeightMultiplier : WeightModifiers.WholeMatchWeightMultiplier;
 
 				for (int32 iEachWord = 0; iEachWord < KeywordArray.Num(); iEachWord++)
 				{
@@ -875,19 +833,171 @@ float UEdGraphSchema::GetActionFilteredWeight(const FGraphActionListBuilderBase:
 						WeightPerList += EachWeight / 2;
 					}
 				}
+
 				// Increase the weight if theres a larger % of matches in the keyword list
 				if (WholeMatchCount != 0)
 				{
 					int32 PercentAdjust = (100 / KeywordArray.Num()) * WholeMatchCount;
-					WeightPerList *= PercentAdjust;
+					int32 PercentAdjustedWeight = WeightPerList * PercentAdjust;
+					DebugInfo.PercentMatch += (float)KeywordArray.Num() / (float)WholeMatchCount;
+					DebugInfo.PercentMatchWeight += (PercentAdjustedWeight - WeightPerList);
+					WeightPerList = PercentAdjustedWeight;
+				}
+
+				if (WeightedArrayList[iFindCount].DebugWeight)
+				{
+					*WeightedArrayList[iFindCount].DebugWeight += WeightPerList;
 				}
 				TotalWeight += WeightPerList;
 			}
 		}
+
+		PrintSearchTextDebugInfo(InFilterTerms, InCurrentAction, &DebugInfo);
 	}
 	return TotalWeight;
 
 }
+
+int32 UEdGraphSchema::CollectSearchTextWeightInfo(const FGraphActionListBuilderBase::ActionGroup& InCurrentAction, const FGraphSchemaSearchWeightModifiers& InWeightModifiers, 
+													TArray<FGraphSchemaSearchTextWeightInfo>& OutWeightedArrayList, FGraphSchemaSearchTextDebugInfo* InDebugInfo) const
+{
+	// First the localized keywords
+	OutWeightedArrayList.Add(FGraphSchemaSearchTextWeightInfo(&InCurrentAction.GetLocalizedSearchKeywordsArrayForFirstAction(), InWeightModifiers.KeywordWeight, InDebugInfo ? &InDebugInfo->KeywordWeight : nullptr));
+
+	// The localized description
+	OutWeightedArrayList.Add(FGraphSchemaSearchTextWeightInfo(&InCurrentAction.GetLocalizedMenuDescriptionArrayForFirstAction(), InWeightModifiers.DescriptionWeight, InDebugInfo ? &InDebugInfo->DescriptionWeight : nullptr));
+
+	// The node search localized title weight
+	OutWeightedArrayList.Add(FGraphSchemaSearchTextWeightInfo(&InCurrentAction.GetLocalizedSearchTitleArrayForFirstAction(), InWeightModifiers.NodeTitleWeight, InDebugInfo ? &InDebugInfo->NodeTitleWeight : nullptr));
+
+	// The localized category
+	OutWeightedArrayList.Add(FGraphSchemaSearchTextWeightInfo(&InCurrentAction.GetLocalizedSearchCategoryArrayForFirstAction(), InWeightModifiers.CategoryWeight, InDebugInfo ? &InDebugInfo->CategoryWeight : nullptr));
+
+	// First the keywords
+	int32 NonLocalizedFirstIndex = OutWeightedArrayList.Add(FGraphSchemaSearchTextWeightInfo(&InCurrentAction.GetSearchKeywordsArrayForFirstAction(), InWeightModifiers.KeywordWeight, InDebugInfo ? &InDebugInfo->KeywordWeight : nullptr));
+
+	// The description
+	OutWeightedArrayList.Add(FGraphSchemaSearchTextWeightInfo(&InCurrentAction.GetMenuDescriptionArrayForFirstAction(), InWeightModifiers.DescriptionWeight, InDebugInfo ? &InDebugInfo->DescriptionWeight : nullptr));
+
+	// The node search title weight
+	OutWeightedArrayList.Add(FGraphSchemaSearchTextWeightInfo(&InCurrentAction.GetSearchTitleArrayForFirstAction(), InWeightModifiers.NodeTitleWeight, InDebugInfo ? &InDebugInfo->NodeTitleWeight : nullptr));
+
+	// The category
+	OutWeightedArrayList.Add(FGraphSchemaSearchTextWeightInfo(&InCurrentAction.GetSearchCategoryArrayForFirstAction(), InWeightModifiers.CategoryWeight, InDebugInfo ? &InDebugInfo->CategoryWeight : nullptr));
+
+	return NonLocalizedFirstIndex;
+}
+
+//////////////////////////////////////////////////////////////////////////
+/** CVars for tweaking how context menu search picks the best match */
+namespace ContextMenuConsoleVariables
+{
+	/** How much weight the node's title has */
+	static float NodeTitleWeight = 1.0f;
+	static FAutoConsoleVariableRef CVarNodeTitleWeight(
+		TEXT("ContextMenu.NodeTitleWeight"), NodeTitleWeight,
+		TEXT("The amount of weight placed on the search items title"),
+		ECVF_Default);
+
+	/** Weight used to prefer keywords of actions  */
+	static float KeywordWeight = 4.0f;
+	static FAutoConsoleVariableRef CVarKeywordWeight(
+		TEXT("ContextMenu.KeywordWeight"), KeywordWeight,
+		TEXT("The amount of weight placed on search items keyword"),
+		ECVF_Default);
+
+	/** Weight used to prefer description of actions  */
+	static float DescriptionWeight = 10.0f;
+	static FAutoConsoleVariableRef CVarDescriptionWeight(
+		TEXT("ContextMenu.DescriptionWeight"), DescriptionWeight,
+		TEXT("The amount of weight placed on search items description"),
+		ECVF_Default);
+
+	/** Weight that a match to a category search has */
+	static float CategoryWeight = 1.0f;
+	static FAutoConsoleVariableRef CVarCategoryWeight(
+		TEXT("ContextMenu.CategoryWeight"), CategoryWeight,
+		TEXT("The amount of weight placed on categories that match what the user has typed in"),
+		ECVF_Default);
+
+	/** The multiplier given if there is an exact localized match to the search term */
+	static float WholeMatchLocalizedWeightMultiplier = 3.0f;
+	static FAutoConsoleVariableRef CVarWholeMatchLocalizedWeightMultiplier(
+		TEXT("ContextMenu.WholeMatchLocalizedWeightMultiplier"), WholeMatchLocalizedWeightMultiplier,
+		TEXT("The multiplier given if there is an exact localized match to the search term"),
+		ECVF_Default);
+
+	/** The multiplier given if there is an exact match to the search term */
+	static float WholeMatchWeightMultiplier = 2.0f;
+	static FAutoConsoleVariableRef CVarWholeMatchWeightMultiplier(
+		TEXT("ContextMenu.WholeMatchWeightMultiplier"), WholeMatchWeightMultiplier,
+		TEXT("The multiplier given if there is an exact match to the search term"),
+		ECVF_Default);
+
+	/** Increasing this will prefer whole percentage matches when comparing the keyword to what the user has typed in */
+	static float PercentageMatchWeightMultiplier = 1.0f;
+	static FAutoConsoleVariableRef CVarPercentageMatchWeightMultiplier(
+		TEXT("rContextMenu.PercentageMatchWeightMultiplier"), PercentageMatchWeightMultiplier,
+		TEXT("A multiplier for how much weight to give something based on the percentage match it is"),
+		ECVF_Default);
+
+	/** Enabling the debug printing of context menu selections */
+	static bool bPrintDebugInfo = false;
+	static FAutoConsoleVariableRef CVarPrintDebugInfo(
+		TEXT("ContextMenu.PrintDebugInfo"), bPrintDebugInfo,
+		TEXT("Print the debug info about the context menu selection"),
+		ECVF_Default);
+}
+
+FGraphSchemaSearchWeightModifiers UEdGraphSchema::GetSearchWeightModifiers() const
+{
+	FGraphSchemaSearchWeightModifiers Modifiers;
+	Modifiers.NodeTitleWeight = ContextMenuConsoleVariables::NodeTitleWeight;
+	Modifiers.KeywordWeight = ContextMenuConsoleVariables::KeywordWeight;
+	Modifiers.DescriptionWeight = ContextMenuConsoleVariables::DescriptionWeight;
+	Modifiers.CategoryWeight = ContextMenuConsoleVariables::DescriptionWeight;
+	Modifiers.WholeMatchLocalizedWeightMultiplier = ContextMenuConsoleVariables::WholeMatchLocalizedWeightMultiplier;
+	Modifiers.WholeMatchWeightMultiplier = ContextMenuConsoleVariables::WholeMatchWeightMultiplier;
+	Modifiers.PercentageMatchWeightMultiplier = ContextMenuConsoleVariables::PercentageMatchWeightMultiplier;
+	return Modifiers;
+}
+
+void UEdGraphSchema::PrintSearchTextDebugInfo(const TArray<FString>& InFilterTerms, const FGraphActionListBuilderBase::ActionGroup& InCurrentAction, const FGraphSchemaSearchTextDebugInfo* InDebugInfo) const
+{
+	if (ContextMenuConsoleVariables::bPrintDebugInfo && InDebugInfo)
+	{
+		InDebugInfo->Print(InFilterTerms, InCurrentAction);
+	}
+}
+
+void FGraphSchemaSearchTextDebugInfo::Print(const TArray<FString>& SearchForKeywords, const FGraphActionListBuilderBase::ActionGroup& Action) const
+{
+	auto CombineStrings = [](const TArray<FString>& Strings) -> FString
+	{
+		FString Output;
+		for (int32 i = 0; i < Strings.Num(); ++i)
+		{
+			Output += Strings[i];
+			Output += i > 0 ? TEXT(" ") : TEXT("");
+		}
+		return Output;
+	};
+
+	FString SearchText = CombineStrings(SearchForKeywords);
+	FString LocalizedTitle = CombineStrings(Action.GetLocalizedSearchTitleArrayForFirstAction());
+	FString LocalizedKeywords = CombineStrings(Action.GetLocalizedSearchKeywordsArrayForFirstAction());
+	FString LocalizedDescription = CombineStrings(Action.GetLocalizedMenuDescriptionArrayForFirstAction());
+	FString LocalizedCategory = CombineStrings(Action.GetLocalizedSearchCategoryArrayForFirstAction());
+	FString Title = CombineStrings(Action.GetSearchTitleArrayForFirstAction());
+	FString Keywords = CombineStrings(Action.GetSearchKeywordsArrayForFirstAction());
+	FString Description = CombineStrings(Action.GetMenuDescriptionArrayForFirstAction());
+	FString Category = CombineStrings(Action.GetSearchCategoryArrayForFirstAction());
+
+	UE_LOG(LogTemp, Log, TEXT("Searching for \"%s\" in [LocTitle:\"%s\" LocKeywords:\"%s\" LocDescription:\"%s\" LocCategory:\"%s\" Title:\"%s\" Keywords:\"%s\" Description:\"%s\" Category:\"%s\"] \
+TotalWeight: %.2f | NodeTitleWeight: %.2f | KeywordWeight: %.2f | CategoryWeight: %.2f | PercentMatchWeight: %.2f | ShorterMatchWeight: %.2f"),
+		*SearchText, *LocalizedTitle, *LocalizedKeywords, *LocalizedDescription, *LocalizedCategory, *Title, *Keywords, *Description, *Category, 
+		TotalWeight, NodeTitleWeight, KeywordWeight, CategoryWeight, PercentMatchWeight, ShorterMatchWeight);
+}
 #endif // WITH_EDITORONLY_DATA
 
 void UEdGraphSchema::ConstructBasicPinTooltip(UEdGraphPin const& Pin, FText const& PinDescription, FString& TooltipOut) const
