Fixed defaults not being able to be set on interface pins

This allows assets that implement an interface to be set as pin defaults

diff --git a/Engine/Source/Editor/BlueprintGraph/Private/EdGraphSchema_K2.cpp b/Engine/Source/Editor/BlueprintGraph/Private/EdGraphSchema_K2.cpp
--- a/Engine/Source/Editor/BlueprintGraph/Private/EdGraphSchema_K2.cpp
+++ b/Engine/Source/Editor/BlueprintGraph/Private/EdGraphSchema_K2.cpp
@@ -4232,7 +4232,11 @@ bool UEdGraphSchema_K2::DefaultValueSimpleValidation(const FEdGraphPinType& PinT
 		}
 		if (NewDefaultObject != nullptr && ObjectClass != nullptr && !NewDefaultObject->GetClass()->GetAuthoritativeClass()->IsChildOf(ObjectClass))
 		{
-			DVSV_RETURN_MSG(FString::Printf(TEXT("%s isn't a %s (specified on pin %s)"), *NewDefaultObject->GetPathName(), *ObjectClass->GetName(), *PinName.ToString()));
+			// Not a type of object, but is it an object implementing an interface?
+			if(PinCategory != PC_Interface || !NewDefaultObject->GetClass()->ImplementsInterface(ObjectClass))
+			{
+				DVSV_RETURN_MSG(FString::Printf(TEXT("%s isn't a %s (specified on pin %s)"), *NewDefaultObject->GetPathName(), *ObjectClass->GetName(), *PinName.ToString()));
+			}
 		}
 	}
 	else if ((PinCategory == PC_SoftObject) || (PinCategory == PC_SoftClass))
