Make Enum to String Node not display as a compact node, making it look different from Byte to String for clarity

diff --git a/Engine/Source/Editor/BlueprintGraph/Classes/K2Node_GetEnumeratorNameAsString.h b/Engine/Source/Editor/BlueprintGraph/Classes/K2Node_GetEnumeratorNameAsString.h
--- a/Engine/Source/Editor/BlueprintGraph/Classes/K2Node_GetEnumeratorNameAsString.h
+++ b/Engine/Source/Editor/BlueprintGraph/Classes/K2Node_GetEnumeratorNameAsString.h
@@ -25,6 +25,7 @@ class UK2Node_GetEnumeratorNameAsString : public UK2Node_GetEnumeratorName
 	//~ Begin UK2Node Interface
 	virtual void GetMenuActions(FBlueprintActionDatabaseRegistrar& ActionRegistrar) const override;
 	virtual FText GetMenuCategory() const override;
+	virtual bool ShouldDrawCompact() const override { return false; }
 	//~ End UK2Node Interface
 
 	virtual FName GetFunctionName() const override;
