UE-147869 - Tooltips for floats are not the same in the My Blueprint window

For floats, the display should now include either a "(single-precision)" or "(double-precision)" next to the float type. Unfortunately, this required a bit of finagling.

When we query for the tooltip, we also need to specify that we want to include subcategory information, which requires changing the signature of GetTypeDescription (et al) and GetCategoryText to include that information. Additionally, pluralizing the type gets a little weird. For containers, you might get something such as "Current type: Array of Float (single-precision)s". Instead of trying to fix the out of place "s", the wording has been changed to "Current: Array of Float (single-precision) types".

diff --git a/Engine/Source/Editor/BlueprintGraph/Classes/EdGraphSchema_K2.h b/Engine/Source/Editor/BlueprintGraph/Classes/EdGraphSchema_K2.h
--- a/Engine/Source/Editor/BlueprintGraph/Classes/EdGraphSchema_K2.h
+++ b/Engine/Source/Editor/BlueprintGraph/Classes/EdGraphSchema_K2.h
@@ -987,6 +987,17 @@ public:
 	 */
 	static FText GetCategoryText(const FName Category, const bool bForMenu = false);
 
+	/**
+	 * Returns the FText to use for a given schema category and subcategory
+	 *
+	 * @param	Category	The category to convert into a FText.
+	 * @param	SubCategory	The subcategory to convert into a FText.
+	 * @param	bForMenu	Indicates if this is for display in tooltips or menu
+	 *
+	 * @return	The text to display for the category.
+	 */
+	static FText GetCategoryText(FName Category, FName SubCategory, bool bForMenu = false);
+
 	/**
 	 * Get the type tree for all of the property types valid for this schema
 	 *
diff --git a/Engine/Source/Editor/BlueprintGraph/Private/EdGraphSchema_K2.cpp b/Engine/Source/Editor/BlueprintGraph/Private/EdGraphSchema_K2.cpp
--- a/Engine/Source/Editor/BlueprintGraph/Private/EdGraphSchema_K2.cpp
+++ b/Engine/Source/Editor/BlueprintGraph/Private/EdGraphSchema_K2.cpp
@@ -104,17 +104,17 @@
 //////////////////////////////////////////////////////////////////////////
 
 // How to display PC_Real pin types to users
-enum class EBlueprintRealDisplayMode
+enum class EBlueprintRealDisplayMode : uint8
 {
 	Real,
 	Float,
 	Number
 };
 
-namespace BlueprintDisplay
+namespace UE::BlueprintDisplay::Private
 {
 	int32 LastRealNamingMode = -1;
-	int32 RealNamingMode = (int32)EBlueprintRealDisplayMode::Float;
+	int32 RealNamingMode = static_cast<int32>(EBlueprintRealDisplayMode::Float);
 	FAutoConsoleVariableRef CVarRealNamingMode(TEXT("Blueprint.PC_Real.DisplayMode"), RealNamingMode, TEXT("Real naming mode\n\t0: Real\n\t1: Float (default)\n\t2: Number\n\nNote the editor needs to be restarted for this to fully take effect"));
 
 	EBlueprintRealDisplayMode GetRealDisplayMode()
@@ -3664,6 +3664,13 @@ FText UEdGraphSchema_K2::TypeToText(const FProperty* const Property)
 
 FText UEdGraphSchema_K2::GetCategoryText(const FName Category, const bool bForMenu)
 {
+	return GetCategoryText(Category, NAME_None, bForMenu);
+}
+
+FText UEdGraphSchema_K2::GetCategoryText(FName Category, FName SubCategory, bool bForMenu)
+{
+	using namespace UE::BlueprintDisplay::Private;
+
 	if (Category.IsNone())
 	{
 		return FText::GetEmpty();
@@ -3673,46 +3680,68 @@ FText UEdGraphSchema_K2::GetCategoryText(const FName Category, const bool bForMe
 	if (CategoryDescriptions.Num() == 0)
 	{
 		CategoryDescriptions.Add(PC_Exec, LOCTEXT("Exec", "Exec"));
-		CategoryDescriptions.Add(PC_Boolean, LOCTEXT("BoolCategory","Boolean"));
-		CategoryDescriptions.Add(PC_Byte, LOCTEXT("ByteCategory","Byte"));
-		CategoryDescriptions.Add(PC_Class, LOCTEXT("ClassCategory","Class Reference"));
+		CategoryDescriptions.Add(PC_Boolean, LOCTEXT("BoolCategory", "Boolean"));
+		CategoryDescriptions.Add(PC_Byte, LOCTEXT("ByteCategory", "Byte"));
+		CategoryDescriptions.Add(PC_Class, LOCTEXT("ClassCategory", "Class Reference"));
 		CategoryDescriptions.Add(PC_Int, LOCTEXT("IntCategory", "Integer"));
 		CategoryDescriptions.Add(PC_Int64, LOCTEXT("Int64Category", "Integer64"));
 		CategoryDescriptions.Add(PC_Real, LOCTEXT("RealCategory", "Real"));
-		CategoryDescriptions.Add(PC_Name, LOCTEXT("NameCategory","Name"));
-		CategoryDescriptions.Add(PC_Delegate, LOCTEXT("DelegateCategory","Delegate"));
-		CategoryDescriptions.Add(PC_MCDelegate, LOCTEXT("MulticastDelegateCategory","Multicast Delegate"));
-		CategoryDescriptions.Add(PC_Object, LOCTEXT("ObjectCategory","Object Reference"));
-		CategoryDescriptions.Add(PC_Interface, LOCTEXT("InterfaceCategory","Interface"));
-		CategoryDescriptions.Add(PC_String, LOCTEXT("StringCategory","String"));
-		CategoryDescriptions.Add(PC_Text, LOCTEXT("TextCategory","Text"));
-		CategoryDescriptions.Add(PC_Struct, LOCTEXT("StructCategory","Structure"));
-		CategoryDescriptions.Add(PC_Wildcard, LOCTEXT("WildcardCategory","Wildcard"));
-		CategoryDescriptions.Add(PC_Enum, LOCTEXT("EnumCategory","Enum"));
+		CategoryDescriptions.Add(PC_Float, LOCTEXT("FloatCategory", "Real (single-precision)"));
+		CategoryDescriptions.Add(PC_Double, LOCTEXT("DoubleCategory", "Real (double-precision)"));
+		CategoryDescriptions.Add(PC_Name, LOCTEXT("NameCategory", "Name"));
+		CategoryDescriptions.Add(PC_Delegate, LOCTEXT("DelegateCategory", "Delegate"));
+		CategoryDescriptions.Add(PC_MCDelegate, LOCTEXT("MulticastDelegateCategory", "Multicast Delegate"));
+		CategoryDescriptions.Add(PC_Object, LOCTEXT("ObjectCategory", "Object Reference"));
+		CategoryDescriptions.Add(PC_Interface, LOCTEXT("InterfaceCategory", "Interface"));
+		CategoryDescriptions.Add(PC_String, LOCTEXT("StringCategory", "String"));
+		CategoryDescriptions.Add(PC_Text, LOCTEXT("TextCategory", "Text"));
+		CategoryDescriptions.Add(PC_Struct, LOCTEXT("StructCategory", "Structure"));
+		CategoryDescriptions.Add(PC_Wildcard, LOCTEXT("WildcardCategory", "Wildcard"));
+		CategoryDescriptions.Add(PC_Enum, LOCTEXT("EnumCategory", "Enum"));
 		CategoryDescriptions.Add(PC_SoftObject, LOCTEXT("SoftObjectReferenceCategory", "Soft Object Reference"));
 		CategoryDescriptions.Add(PC_SoftClass, LOCTEXT("SoftClassReferenceCategory", "Soft Class Reference"));
 		CategoryDescriptions.Add(PC_FieldPath, LOCTEXT("FieldPathReferenceCategory", "Property Reference"));
 		CategoryDescriptions.Add(AllObjectTypes, LOCTEXT("AllObjectTypes", "Object Types"));
 	}
 
-	if (BlueprintDisplay::ShouldRefreshRealDisplay())
+	if (ShouldRefreshRealDisplay())
 	{
-		switch (BlueprintDisplay::GetRealDisplayMode())
+		switch (GetRealDisplayMode())
 		{
 		case EBlueprintRealDisplayMode::Real:
 			CategoryDescriptions[PC_Real] = LOCTEXT("RealCategory_DisplayAsReal", "Real");
+			CategoryDescriptions[PC_Float] = LOCTEXT("RealCategory_DisplayAsReal", "Real (single-precision)");
+			CategoryDescriptions[PC_Double] = LOCTEXT("RealCategory_DisplayAsReal", "Real (double-precision)");
 			break;
 		case EBlueprintRealDisplayMode::Float:
 			CategoryDescriptions[PC_Real] = LOCTEXT("RealCategory_DisplayAsFloat", "Float");
+			CategoryDescriptions[PC_Float] = LOCTEXT("RealCategory_DisplayAsFloat", "Float (single-precision)");
+			CategoryDescriptions[PC_Double] = LOCTEXT("RealCategory_DisplayAsFloat", "Float (double-precision)");
 			break;
 		case EBlueprintRealDisplayMode::Number:
 			CategoryDescriptions[PC_Real] = LOCTEXT("RealCategory_DisplayAsNumber", "Number");
+			CategoryDescriptions[PC_Float] = LOCTEXT("RealCategory_DisplayAsNumber", "Number (single-precision)");
+			CategoryDescriptions[PC_Double] = LOCTEXT("RealCategory_DisplayAsNumber", "Number (double-precision)");
+			break;
+		default:
+			check(false);
 			break;
 		}
 	}
 
-	if (FText const* TypeDesc = CategoryDescriptions.Find(Category))
+	if (const FText* TypeDesc = CategoryDescriptions.Find(Category))
 	{
+		const bool bUseDetailedRealCategory =
+			bForMenu &&
+			(Category == PC_Real) &&
+			((SubCategory == PC_Float) || (SubCategory == PC_Double));
+
+		if (bUseDetailedRealCategory)
+		{
+			TypeDesc = CategoryDescriptions.Find(SubCategory);
+			check(TypeDesc);
+		}
+
 		return *TypeDesc;
 	}
 	else
@@ -3783,7 +3812,9 @@ FText UEdGraphSchema_K2::TerminalTypeToText(const FName Category, const FName Su
 	{
 		if (Category == UEdGraphSchema_K2::PC_Real)
 		{
-			switch (BlueprintDisplay::GetRealDisplayMode())
+			using namespace UE::BlueprintDisplay::Private;
+
+			switch (GetRealDisplayMode())
 			{
 			case EBlueprintRealDisplayMode::Real:
 				PropertyText = (SubCategory == UEdGraphSchema_K2::PC_Float) ? LOCTEXT("SinglePrecisionReal", "Real (single-precision)") : LOCTEXT("DoublePrecisionReal", "Real (double-precision)");
diff --git a/Engine/Source/Editor/KismetWidgets/Private/SPinTypeSelector.cpp b/Engine/Source/Editor/KismetWidgets/Private/SPinTypeSelector.cpp
--- a/Engine/Source/Editor/KismetWidgets/Private/SPinTypeSelector.cpp
+++ b/Engine/Source/Editor/KismetWidgets/Private/SPinTypeSelector.cpp
@@ -430,7 +430,7 @@ void SPinTypeSelector::Construct(const FArguments& InArgs, FGetPinTypeTree GetPi
 					.AutoWidth()
 					[
 						SNew(STextBlock)
-						.Text(this, &SPinTypeSelector::GetTypeDescription )
+						.Text(this, &SPinTypeSelector::GetTypeDescription, false)
 						.Font(InArgs._Font)
 						.ColorAndOpacity(FSlateColor::UseForeground())
 					]
@@ -479,7 +479,7 @@ void SPinTypeSelector::Construct(const FArguments& InArgs, FGetPinTypeTree GetPi
 						.Padding(2.0f, 0.0f, 0.0f, 0.0f)
 						[
 							SNew(STextBlock)
-							.Text( this, &SPinTypeSelector::GetSecondaryTypeDescription )
+							.Text(this, &SPinTypeSelector::GetSecondaryTypeDescription, false)
 							.Font(InArgs._Font)
 						]
 					]
@@ -520,7 +520,7 @@ void SPinTypeSelector::Construct(const FArguments& InArgs, FGetPinTypeTree GetPi
 			.AutoWidth()
 			[
 				SNew(STextBlock)
-				.Text(this, &SPinTypeSelector::GetTypeDescription)
+				.Text(this, &SPinTypeSelector::GetTypeDescription, false)
 				.Font(InArgs._Font)
 				.ColorAndOpacity(FSlateColor::UseForeground())
 			]
@@ -531,7 +531,7 @@ void SPinTypeSelector::Construct(const FArguments& InArgs, FGetPinTypeTree GetPi
 //=======================================================================
 // Attribute Helpers
 
-FText SPinTypeSelector::GetTypeDescription() const
+FText SPinTypeSelector::GetTypeDescription(const bool bIncludeSubcategory) const
 {
 	const FName PinSubCategory = TargetPinType.Get().PinSubCategory;
 	const UObject* PinSubCategoryObject = TargetPinType.Get().PinSubCategoryObject.Get();
@@ -545,11 +545,11 @@ FText SPinTypeSelector::GetTypeDescription() const
 	}
 	else
 	{
-		return UEdGraphSchema_K2::GetCategoryText(TargetPinType.Get().PinCategory, true);
+		return UEdGraphSchema_K2::GetCategoryText(TargetPinType.Get().PinCategory, bIncludeSubcategory ? PinSubCategory : NAME_None, true);
 	}
 }
 
-FText SPinTypeSelector::GetSecondaryTypeDescription() const
+FText SPinTypeSelector::GetSecondaryTypeDescription(const bool bIncludeSubcategory) const
 {
 	const FName PinSubCategory = TargetPinType.Get().PinValueType.TerminalSubCategory;
 	const UObject* PinSubCategoryObject = TargetPinType.Get().PinValueType.TerminalSubCategoryObject.Get();
@@ -563,11 +563,11 @@ FText SPinTypeSelector::GetSecondaryTypeDescription() const
 	}
 	else
 	{
-		return UEdGraphSchema_K2::GetCategoryText(TargetPinType.Get().PinValueType.TerminalCategory, true);
+		return UEdGraphSchema_K2::GetCategoryText(TargetPinType.Get().PinValueType.TerminalCategory, bIncludeSubcategory ? PinSubCategory : NAME_None, true);
 	}
 }
 
-FText SPinTypeSelector::GetCombinedTypeDescription() const
+FText SPinTypeSelector::GetCombinedTypeDescription(const bool bIncludeSubcategory) const
 {
 	FFormatNamedArguments Args;
 
@@ -575,23 +575,24 @@ FText SPinTypeSelector::GetCombinedTypeDescription() const
 	switch (TargetPinType.Get().ContainerType)
 	{
 	case EPinContainerType::Map:
-		Args.Add(TEXT("KeyTitle"), GetTypeDescription());
-		Args.Add(TEXT("ValueTitle"), GetSecondaryTypeDescription());
-		return FText::Format(NSLOCTEXT("KismetSchema", "MapAsText", "Map of {KeyTitle}s to {ValueTitle}s"), Args);
+		Args.Add(TEXT("KeyTitle"), GetTypeDescription(bIncludeSubcategory));
+		Args.Add(TEXT("ValueTitle"), GetSecondaryTypeDescription(bIncludeSubcategory));
+		return FText::Format(NSLOCTEXT("SPinTypeSelector", "MapAsText", "Map of {KeyTitle} types to {ValueTitle} types"), Args);
 	case EPinContainerType::Set:
-		Args.Add(TEXT("PropertyTitle"), GetTypeDescription());
-		return FText::Format(NSLOCTEXT("SetAsText", "MapAsText", "Set of {PropertyTitle}s"), Args);
+		Args.Add(TEXT("PropertyTitle"), GetTypeDescription(bIncludeSubcategory));
+		return FText::Format(NSLOCTEXT("SPinTypeSelector", "SetAsText", "Set of {PropertyTitle} types"), Args);
 	case EPinContainerType::Array:
-		Args.Add(TEXT("PropertyTitle"), GetTypeDescription());
-		return FText::Format(NSLOCTEXT("ArrayAsText", "MapAsText", "Array of {PropertyTitle}s"), Args);
+		Args.Add(TEXT("PropertyTitle"), GetTypeDescription(bIncludeSubcategory));
+		return FText::Format(NSLOCTEXT("SPinTypeSelector", "ArrayAsText", "Array of {PropertyTitle} types"), Args);
 	default:
-		return GetTypeDescription();
+		Args.Add(TEXT("PropertyTitle"), GetTypeDescription(bIncludeSubcategory));
+		return FText::Format(NSLOCTEXT("SPinTypeSelector", "PrimitiveAsText", "{PropertyTitle} type"), Args);
 	}
 }
 
 const FSlateBrush* SPinTypeSelector::GetTypeIconImage() const
 {
-	return FBlueprintEditorUtils::GetIconFromPin( TargetPinType.Get() );
+	return FBlueprintEditorUtils::GetIconFromPin(TargetPinType.Get());
 }
 
 const FSlateBrush* SPinTypeSelector::GetSecondaryTypeIconImage() const
@@ -1379,7 +1380,15 @@ FText SPinTypeSelector::GetToolTipForComboBoxType() const
 
 	// With the full selector type, just display the primary type in the tooltip because the secondary type has its own combo box.
 	// With the partial selector type, we need to jam everything into the tooltip for the single combo box.
-	return FText::Format(LOCTEXT("PrimaryTypeTwoLines", "{0}Current Type: {1}"), EditText, SelectorType == ESelectorType::Full ? GetTypeDescription() : GetCombinedTypeDescription());
+	// Additionally, with the partial selector type, GetCombinedTypeDescription automatically appends "type" to its description.
+	if (SelectorType == ESelectorType::Full)
+	{
+		return FText::Format(LOCTEXT("PrimaryTypeTwoLines", "{0}Current: {1} type"), EditText, GetTypeDescription(true));
+	}
+	else
+	{
+		return FText::Format(LOCTEXT("PrimaryCombinedTypeTwoLines", "{0}Current: {1}"), EditText, GetCombinedTypeDescription(true));
+	}
 }
 
 FText SPinTypeSelector::GetToolTipForComboBoxSecondaryType() const
diff --git a/Engine/Source/Editor/KismetWidgets/Public/SPinTypeSelector.h b/Engine/Source/Editor/KismetWidgets/Public/SPinTypeSelector.h
--- a/Engine/Source/Editor/KismetWidgets/Public/SPinTypeSelector.h
+++ b/Engine/Source/Editor/KismetWidgets/Public/SPinTypeSelector.h
@@ -107,13 +107,13 @@ protected:
 	FSlateColor GetSecondaryTypeIconColor() const;
 
 	/** Gets a succinct type description for the type being manipulated */
-	FText GetTypeDescription() const;
+	FText GetTypeDescription(bool bIncludeSubcategory = false) const;
 
 	/** Gets the secondary type description. E.g. the value type for TMaps */
-	FText GetSecondaryTypeDescription() const;
+	FText GetSecondaryTypeDescription(bool bIncludeSubcategory = false) const;
 
 	/** Gets a combined description of the primary, container, and secondary types. E.g. "Map of Strings to Floats" */
-	FText GetCombinedTypeDescription() const;
+	FText GetCombinedTypeDescription(bool bIncludeSubcategory = false) const;
 
 	TSharedPtr<SComboButton>		TypeComboButton;
 	TSharedPtr<SComboButton>		SecondaryTypeComboButton;
