Fixes a BP compiler crash occurring with SetFieldsOnStruct nodes with missing struct type dependencies, in favor of a compiler error. This was causing an editor-load-time crash in the Safari sample, due to removal of some Network Prediction plugin classes.

diff --git a/Engine/Source/Editor/BlueprintGraph/Private/K2Node_SetFieldsInStruct.cpp b/Engine/Source/Editor/BlueprintGraph/Private/K2Node_SetFieldsInStruct.cpp
--- a/Engine/Source/Editor/BlueprintGraph/Private/K2Node_SetFieldsInStruct.cpp
+++ b/Engine/Source/Editor/BlueprintGraph/Private/K2Node_SetFieldsInStruct.cpp
@@ -66,25 +66,31 @@ public:
 	{
 		FKCHandler_MakeStruct::RegisterNets(Context, Node);
 
-		UEdGraphPin* ReturnPin = Node->FindPin(SetFieldsInStructHelper::StructOutPinName());
-		UEdGraphPin* ReturnStructNet = FEdGraphUtilities::GetNetFromPin(ReturnPin);
+		if (UEdGraphPin* ReturnPin = Node->FindPin(SetFieldsInStructHelper::StructOutPinName()))
+		{
+			UEdGraphPin* ReturnStructNet = FEdGraphUtilities::GetNetFromPin(ReturnPin);
 
-		UEdGraphPin* InputPin = Node->FindPinChecked(SetFieldsInStructHelper::StructRefPinName());
-		UEdGraphPin* InputPinNet = FEdGraphUtilities::GetNetFromPin(InputPin);
-		FBPTerminal** InputTermRef = Context.NetMap.Find(InputPinNet);
+			UEdGraphPin* InputPin = Node->FindPinChecked(SetFieldsInStructHelper::StructRefPinName());
+			UEdGraphPin* InputPinNet = FEdGraphUtilities::GetNetFromPin(InputPin);
+			FBPTerminal** InputTermRef = Context.NetMap.Find(InputPinNet);
 		
-		if (InputTermRef == nullptr)
-		{
-			CompilerContext.MessageLog.Error(*LOCTEXT("MakeStruct_NoTerm_Error", "Failed to generate a term for the @@ pin; was it a struct reference that was left unset?").ToString(), InputPin);
+			if (InputTermRef == nullptr)
+			{
+				CompilerContext.MessageLog.Error(*LOCTEXT("MakeStruct_NoTerm_Error", "Failed to generate a term for the @@ pin; was it a struct reference that was left unset?").ToString(), InputPin);
+			}
+			else
+			{
+				FBPTerminal* InputTerm = *InputTermRef;
+				if (InputTerm->bPassedByReference) //InputPinNet->PinType.bIsReference)
+				{
+					// Forward the net to the output pin because it's being passed by-ref and this pin is a by-ref pin
+					Context.NetMap.Add(ReturnStructNet, InputTerm);
+				}
+			}
 		}
 		else
 		{
-			FBPTerminal* InputTerm = *InputTermRef;
-			if (InputTerm->bPassedByReference) //InputPinNet->PinType.bIsReference)
-			{
-				// Forward the net to the output pin because it's being passed by-ref and this pin is a by-ref pin
-				Context.NetMap.Add(ReturnStructNet, InputTerm);
-			}
+			CompilerContext.MessageLog.Error(*LOCTEXT("SetFieldsInStruct_NoReturnPin_Error", "Failed to find a return pin for node @@. This is likely due to an unresolved dependency.").ToString(), Node);
 		}
 	}
 
