UE-133114 - FBasicTokenParser::GetToken doesn't support double literals

The token parser assumed that numerical literals could only be floats. We simply need to change the parser to use doubles instead.

diff --git a/Engine/Source/Editor/BlueprintGraph/Private/BasicTokenParser.cpp b/Engine/Source/Editor/BlueprintGraph/Private/BasicTokenParser.cpp
--- a/Engine/Source/Editor/BlueprintGraph/Private/BasicTokenParser.cpp
+++ b/Engine/Source/Editor/BlueprintGraph/Private/BasicTokenParser.cpp
@@ -81,10 +81,10 @@ void FBasicToken::SetConstBool(bool InBool)
 }
 
 //------------------------------------------------------------------------------
-void FBasicToken::SetConstFloat(float InFloat)
+void FBasicToken::SetConstDouble(double InDouble)
 {
-	ConstantType = CPT_Float;
-	Float        = InFloat;
+	ConstantType = CPT_Double;
+	Double       = InDouble;
 	TokenType    = TOKEN_Const;
 }
 
@@ -128,8 +128,8 @@ FString FBasicToken::GetConstantValue() const
 		case CPT_Bool:
 			// Don't use FCoreTexts::True/FCoreTexts::False here because they can be localized
 			return FString::Printf(TEXT("%s"), NativeBool ? *(FName::GetEntry(NAME_TRUE)->GetPlainNameString()) : *(FName::GetEntry(NAME_FALSE)->GetPlainNameString()));
-		case CPT_Float:
-			return FString::Printf(TEXT("%f"), Float);
+		case CPT_Double:
+			return FString::Printf(TEXT("%f"), Double);
 		case CPT_Name:
 			return FString::Printf(TEXT("%s"), *ScriptNameToName(*(FScriptName*)NameBytes).ToString());
 		case CPT_String:
@@ -153,19 +153,22 @@ bool FBasicToken::GetConstInt(int32& ValOut) const
 	if(TokenType==TOKEN_Const && ConstantType==CPT_Int)
 	{
 		ValOut = Int;
-		return 1;
+		return true;
 	}
 	else if(TokenType==TOKEN_Const && ConstantType==CPT_Byte)
 	{
 		ValOut = Byte;
-		return 1;
+		return true;
 	}
-	else if(TokenType==TOKEN_Const && ConstantType==CPT_Float && Float==FMath::TruncToInt(Float))
+	else if(TokenType==TOKEN_Const && ConstantType==CPT_Double && Double==FMath::TruncToInt(Double))
 	{
-		ValOut = (int32)Float;
-		return 1;
+		ValOut = static_cast<int32>(Double);
+		return true;
+	}
+	else
+	{
+		return false;
 	}
-	else return 0;
 }
 
 /*******************************************************************************
@@ -361,15 +364,15 @@ bool FBasicTokenParser::GetToken(FBasicToken& Token, bool bNoConsts/* = false*/)
 	// if const values are allowed, determine whether the non-identifier token represents a const
 	else if ( !bNoConsts && (IsNumericChar(c) || ((c=='+' || c=='-') && IsNumericChar(p))) )
 	{
-		// Integer or floating point constant.
-		bool  bIsFloat = 0;
+		// Integer or double-precision floating point constant.
+		bool  bIsDouble = 0;
 		int32 Length   = 0;
 		bool  bIsHex   = 0;
 		do
 		{
 			if( c==TEXT('.') )
 			{
-				bIsFloat = true;
+				bIsDouble = true;
 			}
 			if( c==TEXT('X') || c == TEXT('x') )
 			{
@@ -388,17 +391,17 @@ bool FBasicTokenParser::GetToken(FBasicToken& Token, bool bNoConsts/* = false*/)
 				break;
 			}
 			c = FChar::ToUpper(GetChar());
-		} while (IsNumericChar(c) || (!bIsFloat && c == TEXT('.')) || (!bIsHex && c == TEXT('X')) || (bIsHex && c >= TEXT('A') && c <= TEXT('F')));
+		} while (IsNumericChar(c) || (!bIsDouble && c == TEXT('.')) || (!bIsHex && c == TEXT('X')) || (bIsHex && c >= TEXT('A') && c <= TEXT('F')));
 
 		Token.Identifier[Length]=0;
-		if (!bIsFloat || c != 'F')
+		if (!bIsDouble || c != 'F')
 		{
 			UngetChar();
 		}
 
-		if (bIsFloat)
+		if (bIsDouble)
 		{
-			Token.SetConstFloat( FCString::Atof(Token.Identifier) );
+			Token.SetConstDouble( FCString::Atod(Token.Identifier) );
 		}
 		else if (bIsHex)
 		{
@@ -1044,7 +1047,7 @@ void FBasicTokenParser::SetError(TEnumAsByte<FErrorState::EErrorType> ErrorCode,
 }
 
 //------------------------------------------------------------------------------
-FBasicTokenParser::FErrorState const& FBasicTokenParser::GetErrorState() const
+const FBasicTokenParser::FErrorState& FBasicTokenParser::GetErrorState() const
 {
 	return CurrentError;
 }
diff --git a/Engine/Source/Editor/BlueprintGraph/Private/BasicTokenParser.h b/Engine/Source/Editor/BlueprintGraph/Private/BasicTokenParser.h
--- a/Engine/Source/Editor/BlueprintGraph/Private/BasicTokenParser.h
+++ b/Engine/Source/Editor/BlueprintGraph/Private/BasicTokenParser.h
@@ -31,7 +31,7 @@ public:
 	*
 	* @param  Other		The token to copy properties from.
 	*/
-	void Clone(FBasicToken const& Other);
+	void Clone(const FBasicToken& Other);
 
 	//--------------------------------------
 	// Queries
@@ -55,7 +55,7 @@ public:
 	
 	void SetConstInt(int32 InInt);
 	void SetConstBool(bool InBool);
-	void SetConstFloat(float InFloat);
+	void SetConstDouble(double InDouble);
 	void SetConstName(FName InName);
 	void SetConstString(TCHAR* InString, int32 MaxLength=MAX_STRING_CONST_SIZE);
 	void SetGuid(TCHAR* InString, int32 MaxLength=MAX_STRING_CONST_SIZE);
@@ -66,7 +66,7 @@ public:
 
 	/**
 	 * Retrieves an int value from this token (if it is a constant 
-	 * int/byte/float/etc. type).
+	 * int/byte/double/etc. type).
 	 * 
 	 * @param  ValOut	Will hold the integer value if this is successfully.
 	 * @return True if this represents a constant that can be converted to an integer, otherwise false (ValOut is invalid).
@@ -113,12 +113,12 @@ public:
 	/* TOKEN_Const values */
 	union
 	{
-		uint8 Byte;								// If CPT_Byte.
-		int32 Int;								// If CPT_Int.
-		bool  NativeBool;						// If CPT_Bool
-		float Float;							// If CPT_Float.
-		uint8 NameBytes[sizeof(FScriptName)];	// If CPT_Name.
-		TCHAR String[MAX_STRING_CONST_SIZE];	// If CPT_String
+		uint8  Byte;							// If CPT_Byte.
+		int32  Int;								// If CPT_Int.
+		bool   NativeBool;						// If CPT_Bool
+		double Double;							// If CPT_Double.
+		uint8  NameBytes[sizeof(FScriptName)];	// If CPT_Name.
+		TCHAR  String[MAX_STRING_CONST_SIZE];	// If CPT_String
 	};
 };
 
@@ -357,7 +357,7 @@ public:
 	 * Retrieves the parser's internal error state (so that users can interpret 
 	 * what might have gone wrong while tokenizing).
 	 */
-	FErrorState const& GetErrorState() const;
+	const FErrorState& GetErrorState() const;
 
 protected:
 	/**
diff --git a/Engine/Source/Editor/BlueprintGraph/Private/K2Node_MathExpression.cpp b/Engine/Source/Editor/BlueprintGraph/Private/K2Node_MathExpression.cpp
--- a/Engine/Source/Editor/BlueprintGraph/Private/K2Node_MathExpression.cpp
+++ b/Engine/Source/Editor/BlueprintGraph/Private/K2Node_MathExpression.cpp
@@ -70,12 +70,12 @@ static bool PromoteByteToInt(FEdGraphPinType& InOutType)
 
 /**
 * If the specified type is a "int" type, then this will modify it to
-* a "float". Helps when trying to match function signatures.
+* a "double". Helps when trying to match function signatures.
 *
 * @param  InOutType	The type you want to attempt to promote.
 * @return True if the type was modified, false if not.
 */
-static bool PromoteIntToFloat(FEdGraphPinType& InOutType)
+static bool PromoteIntToDouble(FEdGraphPinType& InOutType)
 {
 	if (InOutType.PinCategory == UEdGraphSchema_K2::PC_Int)
 	{
@@ -717,7 +717,7 @@ public:
 	 * Attempts to lookup a function matching the supplied signature (where 
 	 * 'Operator' identifies the function's name and 'InputTypeList' defines
 	 * the desired parameters). If one can't be found, it attempts to find a
-	 * match by promoting the input types (like from int to float, etc.)
+	 * match by promoting the input types (like from int to double, etc.)
 	 * 
 	 * @param  Operator			The operator you want to find a function for.
 	 * @param  InputTypeList	A list of parameter types you want to feed the function.
@@ -734,7 +734,7 @@ public:
 
 		// if we didn't find a function that matches the supplied function 
 		// signature, then try to promote the parameters (like from int to 
-		// float), and see if we can lookup a function with those types
+		// double), and see if we can lookup a function with those types
 		for (int32 promoterIndex = 0; (promoterIndex < OrderedTypePromoters.Num()) && (MatchingFunc == NULL); ++promoterIndex)
 		{
 			const FTypePromoter& PromotionOperator = OrderedTypePromoters[promoterIndex];
@@ -833,9 +833,9 @@ public:
 		ByteToIntPromoter.BindStatic(&PromoteByteToInt);
 		OrderedTypePromoters.Add(ByteToIntPromoter);
 
-		FTypePromoter IntToFloatPromoter;
-		IntToFloatPromoter.BindStatic(&PromoteIntToFloat);
-		OrderedTypePromoters.Add(IntToFloatPromoter);
+		FTypePromoter IntToDoublePromoter;
+		IntToDoublePromoter.BindStatic(&PromoteIntToDouble);
+		OrderedTypePromoters.Add(IntToDoublePromoter);
 	}
 
 private:
@@ -1068,7 +1068,7 @@ private:
 
 	/**
 	 * When looking to match parameters, there are some implicit conversions we
-	 * can make to try and find a match (like converting from int to float). 
+	 * can make to try and find a match (like converting from int to double). 
 	 * This holds an ordered list of delegates that will try and promote the 
 	 * supplied types.
 	 */
@@ -1668,7 +1668,7 @@ private:
 		{
 			// Create an input pin (using the default guessed type)
 			FEdGraphPinType DefaultType;
-			// currently, generated expressions ALWAYS take a float (it is the most versatile type)
+			// currently, generated expressions ALWAYS take a double (it is the most versatile type)
 			DefaultType.PinCategory = UEdGraphSchema_K2::PC_Real;
 			DefaultType.PinSubCategory = UEdGraphSchema_K2::PC_Double;
 			
@@ -1752,7 +1752,7 @@ private:
 			case CPT_Bool:
 				LiteralType.PinCategory = UEdGraphSchema_K2::PC_Boolean;
 				break;
-			case CPT_Float:
+			case CPT_Double:
 				LiteralType.PinCategory = UEdGraphSchema_K2::PC_Real;
 				LiteralType.PinSubCategory = UEdGraphSchema_K2::PC_Double;
 				break;
@@ -2252,7 +2252,7 @@ private:
 	}
 
 	/**
-	 * Intended to handle type-casts (like from float to int, etc.).
+	 * Intended to handle type-casts (like from double to int, etc.).
 	 * 
 	 * @TODO   Implement!
 	 * @return Root node of an expression tree that was generated from where we 
