Make Soft object references not supported by [REDACTED_PERSON] because it doesnt work correctly. This will make it so that the old options show up in the context menu by default, which is the correct behavior.

diff --git a/Engine/Source/Editor/BlueprintGraph/Private/BlueprintEditorSettings.cpp b/Engine/Source/Editor/BlueprintGraph/Private/BlueprintEditorSettings.cpp
--- a/Engine/Source/Editor/BlueprintGraph/Private/BlueprintEditorSettings.cpp
+++ b/Engine/Source/Editor/BlueprintGraph/Private/BlueprintEditorSettings.cpp
@@ -35,7 +35,7 @@ UBlueprintEditorSettings::UBlueprintEditorSettings(const FObjectInitializer& Obj
 	, bNavigateToNativeFunctionsFromCallNodes(true)
 	, bDoubleClickNavigatesToParent(true)
 	, bEnableTypePromotion(true)
-	, TypePromotionPinDenyList { UEdGraphSchema_K2::PC_String, UEdGraphSchema_K2::PC_Text }
+	, TypePromotionPinDenyList { UEdGraphSchema_K2::PC_String, UEdGraphSchema_K2::PC_Text, UEdGraphSchema_K2::PC_SoftObject }
 	, BreakpointReloadMethod(EBlueprintBreakpointReloadMethod::RestoreAll)
 	, bEnablePinValueInspectionTooltips(true)
 	// Experimental
