Fix a regression that leads to an assert in the editor when trying to spawn a new AddComponent node for a non-native component class type that hasn't yet been loaded.

#rnx

diff --git a/Engine/Source/Editor/BlueprintGraph/Private/BlueprintComponentNodeSpawner.cpp b/Engine/Source/Editor/BlueprintGraph/Private/BlueprintComponentNodeSpawner.cpp
--- a/Engine/Source/Editor/BlueprintGraph/Private/BlueprintComponentNodeSpawner.cpp
+++ b/Engine/Source/Editor/BlueprintGraph/Private/BlueprintComponentNodeSpawner.cpp
@@ -180,17 +180,30 @@ UEdGraphNode* UBlueprintComponentNodeSpawner::Invoke(UEdGraph* ParentGraph, FBin
 		if (Class == nullptr)
 		{
 			const ELoadFlags LoadFlags = LOAD_None;
-			UBlueprint* LoadedObject = LoadObject<UBlueprint>(NULL, *ComponentAssetName, NULL, LoadFlags, NULL);
+			UObject* LoadedObject = LoadObject<UObject>(/*Outer =*/ nullptr, *ComponentAssetName, /*Filename =*/ nullptr, LoadFlags);
 			if (LoadedObject == nullptr)
 			{
 				return nullptr;
 			}
 
-			Class = TSubclassOf<UActorComponent>(Cast<UBlueprintGeneratedClass>(LoadedObject->GeneratedClass));
+			// Note: The asset name may refer to either the generated class or the outer BP asset.
+			if (const UBlueprint* LoadedObjectAsBlueprint = Cast<UBlueprint>(LoadedObject))
+			{
+				Class = TSubclassOf<UActorComponent>(LoadedObjectAsBlueprint->GeneratedClass);
+			}
+			else
+			{
+				Class = TSubclassOf<UActorComponent>(Cast<UBlueprintGeneratedClass>(LoadedObject));
+			}
+			
 			if (Class == nullptr)
 			{
 				return nullptr;
 			}
+
+			// Since the node has already been spawned, we need to update its template type.
+			// Note that the return pin's type will be updated when we reconstruct the node below.
+			NewNode->TemplateType = Class;
 		}
 
 		UBlueprint* Blueprint = NewNode->GetBlueprint();
