Exclude node spawner actions for non-imported component types from the Blueprint graph context menu when Blueprint namespace editor features are enabled.

Change summary:
- Extended FBlueprintActionContext to include a weak ptr to the current IBlueprintEditor context.
- Added ENUM_CLASS_FLAGS() for the FBlueprintActionFilter::EFlags enumeration (for use as bitflags).
- Deprecated the current FBlueprintActionFilter constructor, replaced with a more explicitly-typed EFlags param.
- FBlueprintActionFilter now internally stores its filter flags. Also added FBlueprintActionFilter::HasAnyFlags/HasAllFlags().
- Minor refactor to a few action filter delegates to query the filter flags locally rather than add bool flags to delegate payloads.
- The FBlueprintActionFilter constructor now unsets the non-imported filter flag if Blueprint namespace editor settings are disabled.
- BlueprintActionFilterImpl::IsNonImportedObject() now utilizes the filter context rather than iterate the Blueprints array to find open asset editors.
- Added an override for UBlueprintComponentNodeSpawner::IsTemplateNodeFilteredOut(). This is now used to filter out any loaded or unloaded BP component types from the context menu.
- Cleaned up display inconsistencies in the menu UI spec between loaded vs. unloaded "Add Component" node spawner actions in the context menu.
- Added an FSoftObjectPath variant of IBlueprintEditor::IsNonImportedObject() (to support unloaded non-native component type entries which store the class object path). Implemented as FBlueprintEditor::IsNonImportedObject(FSoftObjectPath).
- Added an override for UK2Node_DynamicCast::IsActionFilteredOut() to exclude cast operator actions for non-native component types if they are not imported into the current editor context.
- Modified SBlueprintActionMenu, SBlueprintFavoritesPalette and SBlueprintLibraryPalette to include the current editor context in applicable action filter context(s).

diff --git a/Engine/Source/Editor/BlueprintGraph/Classes/BlueprintComponentNodeSpawner.h b/Engine/Source/Editor/BlueprintGraph/Classes/BlueprintComponentNodeSpawner.h
--- a/Engine/Source/Editor/BlueprintGraph/Classes/BlueprintComponentNodeSpawner.h
+++ b/Engine/Source/Editor/BlueprintGraph/Classes/BlueprintComponentNodeSpawner.h
@@ -39,6 +39,7 @@ public:
 	virtual FBlueprintNodeSignature GetSpawnerSignature() const override;
 	virtual FBlueprintActionUiSpec GetUiSpec(FBlueprintActionContext const& Context, FBindingSet const& Bindings) const override;
 	virtual UEdGraphNode* Invoke(UEdGraph* ParentGraph, FBindingSet const& Bindings, FVector2D const Location) const override;
+	virtual bool IsTemplateNodeFilteredOut(const FBlueprintActionFilter& Filter) const override;
 	// End UBlueprintNodeSpawner interface
 	
 	// IBlueprintNodeBinder interface
diff --git a/Engine/Source/Editor/BlueprintGraph/Classes/K2Node_DynamicCast.h b/Engine/Source/Editor/BlueprintGraph/Classes/K2Node_DynamicCast.h
--- a/Engine/Source/Editor/BlueprintGraph/Classes/K2Node_DynamicCast.h
+++ b/Engine/Source/Editor/BlueprintGraph/Classes/K2Node_DynamicCast.h
@@ -44,6 +44,7 @@ class UK2Node_DynamicCast : public UK2Node
 	virtual void NotifyPinConnectionListChanged(UEdGraphPin* Pin) override;
 	virtual void ReallocatePinsDuringReconstruction(TArray<UEdGraphPin*>& OldPins) override;
 	virtual void ValidateNodeDuringCompilation(class FCompilerResultsLog& MessageLog) const override;
+	virtual bool IsActionFilteredOut(const class FBlueprintActionFilter& Filter) override;
 	//~ End UK2Node Interface
 
 	/** Get the 'valid cast' exec pin */
diff --git a/Engine/Source/Editor/BlueprintGraph/Private/BlueprintActionFilter.cpp b/Engine/Source/Editor/BlueprintGraph/Private/BlueprintActionFilter.cpp
--- a/Engine/Source/Editor/BlueprintGraph/Private/BlueprintActionFilter.cpp
+++ b/Engine/Source/Editor/BlueprintGraph/Private/BlueprintActionFilter.cpp
@@ -38,6 +38,7 @@
 #include "EditorCategoryUtils.h"
 #include "BlueprintEditorModule.h"
 #include "Kismet2/KismetEditorUtilities.h"
+#include "BlueprintEditorSettings.h"
 
 /*******************************************************************************
  * Static BlueprintActionFilter Helpers
@@ -231,10 +232,9 @@ namespace BlueprintActionFilterImpl
 	 * 
 	 * @param  Filter			Holds the TagetClass context for this test.
 	 * @param  BlueprintAction	The action you wish to query.
-	 * @param  Flags			Filter flags that may modify this test (@see FBlueprintActionFilter::EFlags).
 	 * @return True if the action should be rejected.
 	 */
-	static bool IsRejectedGlobalField(FBlueprintActionFilter const& Filter, FBlueprintActionInfo& BlueprintAction, uint32 Flags);
+	static bool IsRejectedGlobalField(FBlueprintActionFilter const& Filter, FBlueprintActionInfo& BlueprintAction);
 
 	/**
 	 * Rejection test that checks to see if the node-spawner represents an object
@@ -253,10 +253,9 @@ namespace BlueprintActionFilterImpl
 	 * 
 	 * @param  Filter					Holds the class context for this test.
 	 * @param  BlueprintAction			The action you wish to query.
-	 * @param  bPermitNonTargetGlobals	Determines if this test should pass for external global/static fields.
 	 * @return True if the action is associated with a non-allowed class member.
 	 */
-	static bool IsNonTargetMember(FBlueprintActionFilter const& Filter, FBlueprintActionInfo& BlueprintAction, bool const bPermitNonTargetGlobals);
+	static bool IsNonTargetMember(FBlueprintActionFilter const& Filter, FBlueprintActionInfo& BlueprintAction);
 
 	/**
 	 * Rejection test that checks to see if the node-spawner is associated with 
@@ -283,11 +282,9 @@ namespace BlueprintActionFilterImpl
 	 * 
 	 * @param  Filter				Holds the class context for this test.
 	 * @param  BlueprintAction		The action you wish to query.
-	 * @param  bPermitChildClasses	When true, the node's class doesn't have to perfectly match a class to be passed (it can be a sub-class).
-	 * @param  bRejectChildClasses	When true, the node's class doesn't have to perfectly match a class to be rejected (it can be a sub-class).
 	 * @return True if the action would produce a denied node.
 	 */
-	static bool IsFilteredNodeType(FBlueprintActionFilter const& Filter, FBlueprintActionInfo& BlueprintAction, bool const bPermitChildClasses, bool const bRejectChildClasses);
+	static bool IsFilteredNodeType(FBlueprintActionFilter const& Filter, FBlueprintActionInfo& BlueprintAction);
 
 	/**
 	 * Rejection test that checks to see if the node-spawner is tied to a 
@@ -421,10 +418,9 @@ namespace BlueprintActionFilterImpl
 	 * 
 	 * @param  Filter			Holds the graph context for this test.
 	 * @param  BlueprintAction	The action you wish to query.
-	 * @param  Flags			Filter flags that may modify this test (@see FBlueprintActionFilter::EFlags).
 	 * @return True if the macro instance is incompatible with the current graph context
 	 */
-	static bool IsIncompatibleMacroInstance(FBlueprintActionFilter const& Filter, FBlueprintActionInfo& BlueprintAction, uint32 Flags);
+	static bool IsIncompatibleMacroInstance(FBlueprintActionFilter const& Filter, FBlueprintActionInfo& BlueprintAction);
 
 	/**
 	 * Rejection test to help unblock common crashes, where programmers forget 
@@ -889,25 +885,10 @@ static bool BlueprintActionFilterImpl::IsNonImportedObject(FBlueprintActionFilte
 
 	if (InObject)
 	{
-		// @todo_namespaces - Maybe cache this as part of the filter context, or decouple queries from the editor?
-		TArray<TSharedPtr<IBlueprintEditor>> BlueprintEditors;
-		for (const UBlueprint* Blueprint : Filter.Context.Blueprints)
+		TSharedPtr<IBlueprintEditor> BlueprintEditor = Filter.Context.EditorPtr.Pin();
+		if (BlueprintEditor.IsValid())
 		{
-			const bool bOpenEditor = false;
-			TSharedPtr<IBlueprintEditor> BlueprintEditor = FKismetEditorUtilities::GetIBlueprintEditorForObject(Blueprint, bOpenEditor);
-			if (BlueprintEditor.IsValid())
-			{
-				BlueprintEditors.AddUnique(BlueprintEditor);
-			}
-		}
-
-		for (const TSharedPtr<IBlueprintEditor>& BlueprintEditor : BlueprintEditors)
-		{
-			if (BlueprintEditor->IsNonImportedObject(InObject))
-			{
-				bIsFilteredOut = true;
-				break;
-			}
+			bIsFilteredOut = BlueprintEditor->IsNonImportedObject(InObject);
 		}
 	}
 
@@ -915,10 +896,10 @@ static bool BlueprintActionFilterImpl::IsNonImportedObject(FBlueprintActionFilte
 }
 
 //------------------------------------------------------------------------------
-static bool BlueprintActionFilterImpl::IsRejectedGlobalField(FBlueprintActionFilter const& Filter, FBlueprintActionInfo& BlueprintAction, uint32 Flags)
+static bool BlueprintActionFilterImpl::IsRejectedGlobalField(FBlueprintActionFilter const& Filter, FBlueprintActionInfo& BlueprintAction)
 {
-	const bool bRejectAllGlobalFields = ((Flags & FBlueprintActionFilter::EFlags::BPFILTER_RejectGlobalFields) != 0);
-	const bool bRejectNonImportedFields = ((Flags & FBlueprintActionFilter::EFlags::BPFILTER_RejectNonImportedFields) != 0);
+	const bool bRejectAllGlobalFields = Filter.HasAnyFlags(FBlueprintActionFilter::BPFILTER_RejectGlobalFields);
+	const bool bRejectNonImportedFields = Filter.HasAnyFlags(FBlueprintActionFilter::BPFILTER_RejectNonImportedFields);
 
 	bool bIsFilteredOut = false;
 	FFieldVariant Field = BlueprintAction.GetAssociatedMemberField();
@@ -955,10 +936,13 @@ static bool BlueprintActionFilterImpl::IsRejectedGlobalField(FBlueprintActionFil
 }
 
 //------------------------------------------------------------------------------
-static bool BlueprintActionFilterImpl::IsNonTargetMember(FBlueprintActionFilter const& Filter, FBlueprintActionInfo& BlueprintAction, bool bPermitNonTargetGlobals)
+static bool BlueprintActionFilterImpl::IsNonTargetMember(FBlueprintActionFilter const& Filter, FBlueprintActionInfo& BlueprintAction)
 {
 	bool bIsFilteredOut = false;
 
+	// Determines if this test should pass for external global/static fields.
+	const bool bPermitNonTargetGlobals = !Filter.HasAnyFlags(FBlueprintActionFilter::BPFILTER_RejectGlobalFields);
+
 	FFieldVariant ClassField = BlueprintAction.GetAssociatedMemberField();
 	bool const bIsMemberAction = ClassField.IsValid() && ClassField.GetOwnerClass() != nullptr;
 
@@ -1149,10 +1133,16 @@ static bool BlueprintActionFilterImpl::IsIncompatibleWithGraphType(FBlueprintAct
 }
 
 //------------------------------------------------------------------------------
-static bool BlueprintActionFilterImpl::IsFilteredNodeType(FBlueprintActionFilter const& Filter, FBlueprintActionInfo& BlueprintAction, bool const bPermitChildClasses, bool const bRejectChildClasses)
+static bool BlueprintActionFilterImpl::IsFilteredNodeType(FBlueprintActionFilter const& Filter, FBlueprintActionInfo& BlueprintAction)
 {
 	bool bIsFilteredOut = (Filter.PermittedNodeTypes.Num() > 0);
 
+	// When true, the node's class doesn't have to perfectly match a class to be passed (it can be a sub-class).
+	const bool bPermitChildClasses = !Filter.HasAnyFlags(FBlueprintActionFilter::BPFILTER_RejectPermittedSubClasses);
+
+	// When true, the node's class doesn't have to perfectly match a class to be rejected (it can be a sub-class).
+	const bool bRejectChildClasses = !Filter.HasAnyFlags(FBlueprintActionFilter::BPFILTER_PermitRejectionSubClasses);
+
 	UClass const* NodeClass = BlueprintAction.GetNodeClass();
 	if (NodeClass != nullptr)
 	{
@@ -1796,11 +1786,11 @@ static bool BlueprintActionFilterImpl::IsExtraneousInterfaceCall(FBlueprintActio
 }
 
 //------------------------------------------------------------------------------
-static bool BlueprintActionFilterImpl::IsIncompatibleMacroInstance(FBlueprintActionFilter const& Filter, FBlueprintActionInfo& BlueprintAction, uint32 Flags)
+static bool BlueprintActionFilterImpl::IsIncompatibleMacroInstance(FBlueprintActionFilter const& Filter, FBlueprintActionInfo& BlueprintAction)
 {
 	bool bIsFilteredOut = false;
 
-	const bool bRejectNonImportedMacros = ((Flags & FBlueprintActionFilter::EFlags::BPFILTER_RejectNonImportedFields) != 0);
+	const bool bRejectNonImportedMacros = Filter.HasAllFlags(FBlueprintActionFilter::BPFILTER_RejectNonImportedFields);
 
 	if(BlueprintAction.GetNodeClass()->IsChildOf<UK2Node_MacroInstance>())
 	{
@@ -2038,10 +2028,25 @@ UFunction const* FBlueprintActionInfo::GetAssociatedFunction()
  ******************************************************************************/
 
 //------------------------------------------------------------------------------
-FBlueprintActionFilter::FBlueprintActionFilter(uint32 Flags/*= 0x00*/)
+// DEPRECATED
+FBlueprintActionFilter::FBlueprintActionFilter(uint32 Flags)
+	:FBlueprintActionFilter(EFlags(Flags))
+{
+}
+
+//------------------------------------------------------------------------------
+FBlueprintActionFilter::FBlueprintActionFilter(const EFlags InFlags /*= BPFILTER_NoFlags*/)
 {
 	using namespace BlueprintActionFilterImpl;
 
+	FilterFlags = InFlags;
+
+	// Disable non-imported type filtering if namespace filtering features are disabled.
+	if (!GetDefault<UBlueprintEditorSettings>()->bEnableNamespaceFilteringFeatures)
+	{
+		FilterFlags &= ~BPFILTER_RejectNonImportedFields;
+	}
+
 	BluprintGraphModule = &FModuleManager::LoadModuleChecked<FBlueprintGraphModule>("BlueprintGraph");
 
 	//
@@ -2064,16 +2069,16 @@ FBlueprintActionFilter::FBlueprintActionFilter(uint32 Flags/*= 0x00*/)
 	AddRejectionTest(FRejectionTestDelegate::CreateStatic(IsIncompatibleLatentNode));
 	AddRejectionTest(FRejectionTestDelegate::CreateStatic(IsIncompatibleImpureNode));
 	AddRejectionTest(FRejectionTestDelegate::CreateStatic(IsPropertyAccessorNode));
-	if (Flags & BPFILTER_RejectIncompatibleThreadSafety)
+	if (HasAnyFlags(BPFILTER_RejectIncompatibleThreadSafety))
 	{
 		AddRejectionTest(FRejectionTestDelegate::CreateStatic(IsThreadSafetyIncompatible));
 	}
 	
 	AddRejectionTest(FRejectionTestDelegate::CreateStatic(IsActionHiddenByConfig));
 	AddRejectionTest(FRejectionTestDelegate::CreateStatic(IsFieldCategoryHidden));
-	if (Flags & (BPFILTER_RejectGlobalFields | BPFILTER_RejectNonImportedFields))
+	if (HasAnyFlags(BPFILTER_RejectGlobalFields | BPFILTER_RejectNonImportedFields))
 	{
-		AddRejectionTest(FRejectionTestDelegate::CreateStatic(IsRejectedGlobalField, Flags));
+		AddRejectionTest(FRejectionTestDelegate::CreateStatic(IsRejectedGlobalField));
 	}
 	
 	AddRejectionTest(FRejectionTestDelegate::CreateStatic(IsFieldInaccessible));
@@ -2084,15 +2089,15 @@ FBlueprintActionFilter::FBlueprintActionFilter(uint32 Flags/*= 0x00*/)
 	AddRejectionTest(FRejectionTestDelegate::CreateStatic(IsIncompatibleWithGraphType));
 	AddRejectionTest(FRejectionTestDelegate::CreateStatic(IsSchemaIncompatible));
 	AddRejectionTest(FRejectionTestDelegate::CreateStatic(IsExtraneousInterfaceCall));
-	AddRejectionTest(FRejectionTestDelegate::CreateStatic(IsIncompatibleMacroInstance, Flags));
+	AddRejectionTest(FRejectionTestDelegate::CreateStatic(IsIncompatibleMacroInstance));
 
-	if (!(Flags & BPFILTER_PermitDeprecated))
+	if (!HasAnyFlags(BPFILTER_PermitDeprecated))
 	{
 		AddRejectionTest(FRejectionTestDelegate::CreateStatic(IsDeprecated));
 	}
 
-	AddRejectionTest(FRejectionTestDelegate::CreateStatic(IsFilteredNodeType, !(Flags & BPFILTER_RejectPermittedSubClasses), !(Flags & BPFILTER_PermitRejectionSubClasses)));
-	AddRejectionTest(FRejectionTestDelegate::CreateStatic(IsNonTargetMember, !(Flags & BPFILTER_RejectGlobalFields)));
+	AddRejectionTest(FRejectionTestDelegate::CreateStatic(IsFilteredNodeType));
+	AddRejectionTest(FRejectionTestDelegate::CreateStatic(IsNonTargetMember));
 	AddRejectionTest(FRejectionTestDelegate::CreateStatic(IsUnBoundBindingSpawner));
 	AddRejectionTest(FRejectionTestDelegate::CreateStatic(IsOutOfScopeLocalVariable));
 	AddRejectionTest(FRejectionTestDelegate::CreateStatic(IsLevelScriptActionValid));
diff --git a/Engine/Source/Editor/BlueprintGraph/Private/BlueprintComponentNodeSpawner.cpp b/Engine/Source/Editor/BlueprintGraph/Private/BlueprintComponentNodeSpawner.cpp
--- a/Engine/Source/Editor/BlueprintGraph/Private/BlueprintComponentNodeSpawner.cpp
+++ b/Engine/Source/Editor/BlueprintGraph/Private/BlueprintComponentNodeSpawner.cpp
@@ -12,6 +12,7 @@
 #include "BlueprintNodeTemplateCache.h"
 #include "ComponentAssetBroker.h"
 #include "ComponentTypeRegistry.h"
+#include "BlueprintEditorModule.h"
 
 #define LOCTEXT_NAMESPACE "BlueprintComponenetNodeSpawner"
 
@@ -21,27 +22,27 @@
 
 namespace BlueprintComponentNodeSpawnerImpl
 {
-	static FText GetDefaultMenuCategory(TSubclassOf<UActorComponent> const ComponentClass);
-}
-
-//------------------------------------------------------------------------------
-static FText BlueprintComponentNodeSpawnerImpl::GetDefaultMenuCategory(TSubclassOf<UActorComponent> const ComponentClass)
-{
-	FText ClassGroup;
-	TArray<FString> ClassGroupNames;
-	ComponentClass->GetClassGroupNames(ClassGroupNames);
-
-	static FText const DefaultClassGroup(LOCTEXT("DefaultClassGroup", "Common"));
-	// 'Common' takes priority over other class groups
-	if (ClassGroupNames.Contains(DefaultClassGroup.ToString()) || (ClassGroupNames.Num() == 0))
+	//------------------------------------------------------------------------------
+	static FText GetMenuCategoryFormat()
 	{
-		ClassGroup = DefaultClassGroup;
+		return LOCTEXT("ComponentCategory", "Add Component|{0}");
 	}
-	else
+
+	//------------------------------------------------------------------------------
+	static FText GetDefaultMenuCategory(const TSubclassOf<UActorComponent> ComponentClass)
 	{
-		ClassGroup = FText::FromString(ClassGroupNames[0]);
+		TArray<FString> ClassGroupNames;
+		ComponentClass->GetClassGroupNames(ClassGroupNames);
+
+		static const FText DefaultClassGroup(LOCTEXT("DefaultClassGroup", "Common"));
+		// 'Common' takes priority over other class groups
+		if (ClassGroupNames.Contains(DefaultClassGroup.ToString()) || (ClassGroupNames.Num() == 0))
+		{
+			return DefaultClassGroup;
+		}
+		
+		return FText::FromString(ClassGroupNames[0]);
 	}
-	return FText::Format(LOCTEXT("ComponentCategory", "Add Component|{0}"), ClassGroup);
 }
 
 /*******************************************************************************
@@ -51,6 +52,8 @@ static FText BlueprintComponentNodeSpawnerImpl::GetDefaultMenuCategory(TSubclass
 //------------------------------------------------------------------------------
 UBlueprintComponentNodeSpawner* UBlueprintComponentNodeSpawner::Create(const FComponentTypeEntry& Entry)
 {
+	using namespace BlueprintComponentNodeSpawnerImpl;
+
 	UClass* ComponentClass = Entry.ComponentClass;
 	if (ComponentClass == nullptr)
 	{
@@ -65,7 +68,9 @@ UBlueprintComponentNodeSpawner* UBlueprintComponentNodeSpawner::Create(const FCo
 		FBlueprintActionUiSpec& MenuSignature = NodeSpawner->DefaultMenuSignature;
 		FText const ComponentTypeName = FText::FromString(Entry.ComponentName);
 		MenuSignature.MenuName = FText::Format(LOCTEXT("AddComponentMenuName", "Add {0}"), ComponentTypeName);
-		MenuSignature.Category = LOCTEXT("BlueprintComponentCategory", "Custom");
+		// Note: Non-native (i.e. BP) component types are automatically assigned to the "Custom" group name.
+		// => @see FBlueprintEditorUtils::RecreateClassMetaData()
+		MenuSignature.Category = FText::Format(GetMenuCategoryFormat(), LOCTEXT("BlueprintComponentCategory", "Custom"));
 		MenuSignature.Tooltip = FText::Format(LOCTEXT("AddComponentTooltip", "Spawn a {0}"), ComponentTypeName);
 		// add at least one character, so that PrimeDefaultUiSpec() doesn't 
 		// attempt to query the template node
@@ -73,7 +78,10 @@ UBlueprintComponentNodeSpawner* UBlueprintComponentNodeSpawner::Create(const FCo
 		{
 			MenuSignature.Keywords = FText::FromString(TEXT(" "));
 		}
-		MenuSignature.Icon = FSlateIconFinder::FindIconForClass(nullptr);
+		// Note: Currently using whatever styling is in place for UActorComponent. Note that the actual
+		// class when loaded could be a USceneComponent derivative, in which case it might end up having
+		// an alternate styling. For now just going with this as the placeholder to match the basic type.
+		MenuSignature.Icon = FSlateIconFinder::FindIconForClass(UActorComponent::StaticClass());
 		MenuSignature.DocLink  = TEXT("Shared/GraphNodes/Blueprint/UK2Node_AddComponent");
 		MenuSignature.DocExcerptTag = TEXT("AddComponent");
 
@@ -102,7 +110,7 @@ UBlueprintComponentNodeSpawner* UBlueprintComponentNodeSpawner::Create(const FCo
 	FBlueprintActionUiSpec& MenuSignature = NodeSpawner->DefaultMenuSignature;
 	FText const ComponentTypeName = AuthoritativeClass->GetDisplayNameText();
 	MenuSignature.MenuName = FText::Format(LOCTEXT("AddComponentMenuName", "Add {0}"), ComponentTypeName);
-	MenuSignature.Category = BlueprintComponentNodeSpawnerImpl::GetDefaultMenuCategory(AuthoritativeClass);
+	MenuSignature.Category = FText::Format(GetMenuCategoryFormat(), GetDefaultMenuCategory(AuthoritativeClass));
 	MenuSignature.Tooltip  = FText::Format(LOCTEXT("AddComponentTooltip", "Spawn a {0}"), ComponentTypeName);
 	MenuSignature.Keywords = AuthoritativeClass->GetMetaDataText(*FBlueprintMetadata::MD_FunctionKeywords.ToString(), TEXT("UObjectKeywords"), AuthoritativeClass->GetFullGroupName(false));
 	// add at least one character, so that PrimeDefaultUiSpec() doesn't 
@@ -287,4 +295,28 @@ TSubclassOf<UActorComponent> UBlueprintComponentNodeSpawner::GetComponentClass()
 	return ComponentClass;
 }
 
+//------------------------------------------------------------------------------
+bool UBlueprintComponentNodeSpawner::IsTemplateNodeFilteredOut(const FBlueprintActionFilter& Filter) const
+{
+	bool bIsFilteredOut = false;
+
+	if (Filter.HasAnyFlags(FBlueprintActionFilter::BPFILTER_RejectNonImportedFields))
+	{
+		TSharedPtr<IBlueprintEditor> BlueprintEditor = Filter.Context.EditorPtr.Pin();
+		if (BlueprintEditor.IsValid())
+		{
+			if (ComponentClass)
+			{
+				bIsFilteredOut = BlueprintEditor->IsNonImportedObject(ComponentClass);
+			}
+			else if(FPackageName::IsValidObjectPath(ComponentAssetName))
+			{
+				bIsFilteredOut = BlueprintEditor->IsNonImportedObject(ComponentAssetName);
+			}
+		}
+	}
+
+	return bIsFilteredOut || UBlueprintNodeSpawner::IsTemplateNodeFilteredOut(Filter);
+}
+
 #undef LOCTEXT_NAMESPACE
diff --git a/Engine/Source/Editor/BlueprintGraph/Private/K2Node_DynamicCast.cpp b/Engine/Source/Editor/BlueprintGraph/Private/K2Node_DynamicCast.cpp
--- a/Engine/Source/Editor/BlueprintGraph/Private/K2Node_DynamicCast.cpp
+++ b/Engine/Source/Editor/BlueprintGraph/Private/K2Node_DynamicCast.cpp
@@ -7,13 +7,14 @@
 #include "Framework/Commands/UIAction.h"
 #include "ToolMenus.h"
 #include "EdGraphSchema_K2.h"
-
+#include "BlueprintEditorModule.h"
 #include "BlueprintEditorSettings.h"
 #include "Kismet2/CompilerResultsLog.h"
 #include "DynamicCastHandler.h"
 #include "EditorCategoryUtils.h"
 #include "Kismet2/KismetEditorUtilities.h"
 #include "ScopedTransaction.h"
+#include "BlueprintActionFilter.h"
 
 #define LOCTEXT_NAMESPACE "K2Node_DynamicCast"
 
@@ -516,4 +517,20 @@ bool UK2Node_DynamicCast::ReconnectPureExecPins(TArray<UEdGraphPin*>& OldPins)
 	return false;
 }
 
+bool UK2Node_DynamicCast::IsActionFilteredOut(const FBlueprintActionFilter& Filter)
+{
+	bool bIsFilteredOut = false;
+
+	if (Filter.HasAnyFlags(FBlueprintActionFilter::BPFILTER_RejectNonImportedFields))
+	{
+		TSharedPtr<IBlueprintEditor> BlueprintEditor = Filter.Context.EditorPtr.Pin();
+		if (BlueprintEditor.IsValid() && TargetType)
+		{
+			bIsFilteredOut = BlueprintEditor->IsNonImportedObject(TargetType);
+		}
+	}
+
+	return bIsFilteredOut;
+}
+
 #undef LOCTEXT_NAMESPACE
diff --git a/Engine/Source/Editor/BlueprintGraph/Public/BlueprintActionFilter.h b/Engine/Source/Editor/BlueprintGraph/Public/BlueprintActionFilter.h
--- a/Engine/Source/Editor/BlueprintGraph/Public/BlueprintActionFilter.h
+++ b/Engine/Source/Editor/BlueprintGraph/Public/BlueprintActionFilter.h
@@ -15,6 +15,7 @@ class UBlueprint;
 class UBlueprintNodeSpawner;
 class UEdGraph;
 class UEdGraphPin;
+class IBlueprintEditor;
 
 /*******************************************************************************
  * FBlueprintActionContext
@@ -22,6 +23,13 @@ class UEdGraphPin;
 
 struct FBlueprintActionContext
 {
+	/**
+	 * A weak reference to the Blueprint editor context in which actions are
+	 * being filtered. Note that this may be NULL or invalid if we're filtering
+	 * actions outside of an editor context (e.g. commandlet).
+	 */
+	TWeakPtr<IBlueprintEditor> EditorPtr;
+
 	/** 
 	 * A list of all blueprints you want actions for. Generally, this will
 	 * only contain a single blueprint, but it can have many (where an action
@@ -163,6 +171,8 @@ public:
 public:
 	enum EFlags // Flags, which configure certain rejection tests.
 	{
+		BPFILTER_NoFlags					= (0),
+
 		/** Deprecated class actions will not be filtered out*/
 		BPFILTER_PermitDeprecated			= (1<<0),
 
@@ -201,7 +211,12 @@ public:
 		 */
 		BPFILTER_RejectIncompatibleThreadSafety = (1<<5)
 	};
-	FBlueprintActionFilter(uint32 const Flags = 0x00);
+
+	/** @param InFlags - Filter configuration flags. */
+	FBlueprintActionFilter(const EFlags InFlags = BPFILTER_NoFlags);
+
+	UE_DEPRECATED(5.1, "Please use the version that takes the EFlags type as input")
+	FBlueprintActionFilter(uint32 const Flags);
 	
 	/**
 	 * Contains the full blueprint/graph/pin context that this is filtering 
@@ -289,6 +304,22 @@ public:
 	 */
 	FBlueprintActionFilter const& operator&=(FBlueprintActionFilter const& Rhs);
 
+	/**
+	 * @return TRUE if any of the given configuration flags are set on this filter.
+	 */
+	FORCEINLINE bool HasAnyFlags(EFlags InFlags) const
+	{
+		return (FilterFlags & InFlags) != 0;
+	}
+
+	/**
+	 * @return TRUE if all of the given configuration flags are set on this filter.
+	 */
+	FORCEINLINE bool HasAllFlags(EFlags InFlags) const
+	{
+		return (FilterFlags & InFlags) == InFlags;
+	}
+
 private:
 	/**
 	 * Query to check and see if the specified action gets filtered out by this 
@@ -299,6 +330,9 @@ private:
 	 */
 	bool IsFilteredByThis(FBlueprintActionInfo& BlueprintAction) const;
 
+	/** Configuration flags for this filter. */
+	EFlags FilterFlags = BPFILTER_NoFlags;
+
 	/** Set of rejection tests for this specific filter. */
 	TArray<FRejectionTestDelegate> FilterTests;
 
@@ -308,3 +342,5 @@ private:
 	/** Alternative filters to be logically or'd in with the IsFilteredByThis() result. */
 	TArray<FBlueprintActionFilter> OrFilters;
 };
+
+ENUM_CLASS_FLAGS(FBlueprintActionFilter::EFlags);
\ No newline at end of file
diff --git a/Engine/Source/Editor/Kismet/Private/BlueprintActionMenuUtils.cpp b/Engine/Source/Editor/Kismet/Private/BlueprintActionMenuUtils.cpp
--- a/Engine/Source/Editor/Kismet/Private/BlueprintActionMenuUtils.cpp
+++ b/Engine/Source/Editor/Kismet/Private/BlueprintActionMenuUtils.cpp
@@ -416,7 +416,7 @@ void FBlueprintActionMenuUtils::MakePaletteMenu(FBlueprintActionContext const& C
 {
 	MenuOut.Empty();
 	
-	uint32 FilterFlags = 0x00;
+	FBlueprintActionFilter::EFlags FilterFlags = FBlueprintActionFilter::BPFILTER_NoFlags;
 	if (FilterClass != nullptr)
 	{
 		// make sure we exclude global and static library actions
@@ -450,7 +450,7 @@ void FBlueprintActionMenuUtils::MakeContextMenu(FBlueprintActionContext const& C
 	// Composing Filters
 	//--------------------------------------
 
-	uint32 FilterFlags = 0x00;
+	FBlueprintActionFilter::EFlags FilterFlags = FBlueprintActionFilter::BPFILTER_NoFlags;
 	if ( bIsContextSensitive && ((ClassTargetMask & EContextTargetFlags::TARGET_BlueprintLibraries) == 0) )
 	{
 		FilterFlags |= FBlueprintActionFilter::BPFILTER_RejectGlobalFields;
diff --git a/Engine/Source/Editor/Kismet/Private/BlueprintEditor.cpp b/Engine/Source/Editor/Kismet/Private/BlueprintEditor.cpp
--- a/Engine/Source/Editor/Kismet/Private/BlueprintEditor.cpp
+++ b/Engine/Source/Editor/Kismet/Private/BlueprintEditor.cpp
@@ -10402,6 +10402,16 @@ bool FBlueprintEditor::IsNonImportedObject(const UObject* InObject) const
 	return false;
 }
 
+bool FBlueprintEditor::IsNonImportedObject(const FSoftObjectPath& InObject) const
+{
+	if (ImportedNamespaceHelper.IsValid() && !ImportedNamespaceHelper->IsImportedObject(InObject))
+	{
+		return true;
+	}
+
+	return false;
+}
+
 void FBlueprintEditor::OnBlueprintProjectSettingsChanged(UObject*, struct FPropertyChangedEvent&)
 {
 	ModifyDuringPIEStatus = ESafeToModifyDuringPIEStatus::Unknown;
diff --git a/Engine/Source/Editor/Kismet/Private/SBlueprintActionMenu.cpp b/Engine/Source/Editor/Kismet/Private/SBlueprintActionMenu.cpp
--- a/Engine/Source/Editor/Kismet/Private/SBlueprintActionMenu.cpp
+++ b/Engine/Source/Editor/Kismet/Private/SBlueprintActionMenu.cpp
@@ -521,6 +521,7 @@ void SBlueprintActionMenu::ConstructActionContext(FBlueprintActionContext& Conte
 		return;
 	}
 
+	ContextDescOut.EditorPtr = EditorPtr;
 	ContextDescOut.Blueprints.Add(Blueprint);
 
 	if (bIsContextSensitive)
diff --git a/Engine/Source/Editor/Kismet/Private/SBlueprintFavoritesPalette.cpp b/Engine/Source/Editor/Kismet/Private/SBlueprintFavoritesPalette.cpp
--- a/Engine/Source/Editor/Kismet/Private/SBlueprintFavoritesPalette.cpp
+++ b/Engine/Source/Editor/Kismet/Private/SBlueprintFavoritesPalette.cpp
@@ -11,6 +11,7 @@
 #include "Widgets/SToolTip.h"
 #include "Styling/CoreStyle.h"
 #include "EditorStyleSet.h"
+#include "BlueprintEditor.h"
 #include "Editor/EditorPerProjectUserSettings.h"
 #include "BlueprintPaletteFavorites.h"
 #include "BlueprintActionMenuBuilder.h"
@@ -315,6 +316,7 @@ void SBlueprintFavoritesPalette::CollectAllActions(FGraphActionListBuilderBase&
 {
 	FBlueprintActionContext FilterContext;
 	FilterContext.Blueprints.Add(GetBlueprint());
+	FilterContext.EditorPtr = BlueprintEditorPtr;
 
 	FBlueprintActionMenuBuilder FavoritesBuilder(BlueprintEditorPtr);
 	FBlueprintActionMenuUtils::MakeFavoritesMenu(FilterContext, FavoritesBuilder);
diff --git a/Engine/Source/Editor/Kismet/Private/SBlueprintLibraryPalette.cpp b/Engine/Source/Editor/Kismet/Private/SBlueprintLibraryPalette.cpp
--- a/Engine/Source/Editor/Kismet/Private/SBlueprintLibraryPalette.cpp
+++ b/Engine/Source/Editor/Kismet/Private/SBlueprintLibraryPalette.cpp
@@ -10,6 +10,7 @@
 #include "Widgets/Input/SComboButton.h"
 #include "EditorStyleSet.h"
 #include "Engine/Blueprint.h"
+#include "BlueprintEditor.h"
 #include "Editor/EditorPerProjectUserSettings.h"
 #include "EdGraphSchema_K2.h"
 #include "ClassViewerModule.h"
@@ -220,6 +221,7 @@ void SBlueprintLibraryPalette::CollectAllActions(FGraphActionListBuilderBase& Ou
 	
 	FBlueprintActionContext FilterContext;
 	FilterContext.Blueprints.Add(GetBlueprint());
+	FilterContext.EditorPtr = BlueprintEditorPtr;
 	
 	UClass* ClassFilter = nullptr;
 	if (FilterClass.IsValid())
diff --git a/Engine/Source/Editor/Kismet/Public/BlueprintEditor.h b/Engine/Source/Editor/Kismet/Public/BlueprintEditor.h
--- a/Engine/Source/Editor/Kismet/Public/BlueprintEditor.h
+++ b/Engine/Source/Editor/Kismet/Public/BlueprintEditor.h
@@ -246,6 +246,7 @@ public:
 	virtual TSharedPtr<class IClassViewerFilter> GetImportedClassViewerFilter() const override { return ImportedClassViewerFilter; }
 	virtual TSharedPtr<class IPinTypeSelectorFilter> GetImportedPinTypeSelectorFilter() const override { return ImportedPinTypeSelectorFilter; }
 	virtual bool IsNonImportedObject(const UObject* InObject) const;
+	virtual bool IsNonImportedObject(const FSoftObjectPath& InObject) const;
 	//~ End IBlueprintEditor Interface
 
 	//~ Begin FTickableEditorObject Interface
diff --git a/Engine/Source/Editor/Kismet/Public/BlueprintEditorModule.h b/Engine/Source/Editor/Kismet/Public/BlueprintEditorModule.h
--- a/Engine/Source/Editor/Kismet/Public/BlueprintEditorModule.h
+++ b/Engine/Source/Editor/Kismet/Public/BlueprintEditorModule.h
@@ -114,6 +114,9 @@ public:
 	/** Return whether the given object falls outside the scope of the current set of imported namespaces within this editor context. Default is FALSE (imported). */
 	virtual bool IsNonImportedObject(const UObject* InObject) const { return false; }
 
+	/** Return whether the given object (referenced by path) falls outside the scope of the current set of imported namespaces within this editor context. Default is FALSE (imported). */
+	virtual bool IsNonImportedObject(const FSoftObjectPath& InObject) const { return false; }
+
 	UE_DEPRECATED(5.0, "GetSelectedSCSEditorTreeNodes has been deprecated. Use GetSelectedSubobjectEditorTreeNodes instead.")
 	virtual TArray<TSharedPtr<class FSCSEditorTreeNode> >  GetSelectedSCSEditorTreeNodes() const = 0;
 	UE_DEPRECATED(5.0, "FindAndSelectSCSEditorTreeNode has been deprecated. Use FindAndSelectSubobjectEditorTreeNode instead.")
