Adding AppStyle.h includes

diff --git a/Engine/Source/Editor/BlueprintGraph/Private/BlueprintEventNodeSpawner.cpp b/Engine/Source/Editor/BlueprintGraph/Private/BlueprintEventNodeSpawner.cpp
--- a/Engine/Source/Editor/BlueprintGraph/Private/BlueprintEventNodeSpawner.cpp
+++ b/Engine/Source/Editor/BlueprintGraph/Private/BlueprintEventNodeSpawner.cpp
@@ -6,6 +6,7 @@
 #include "K2Node_FunctionEntry.h"
 #include "Kismet2/BlueprintEditorUtils.h"
 #include "BlueprintNodeTemplateCache.h"
+#include "Styling/AppStyle.h"
 
 #define LOCTEXT_NAMESPACE "BlueprintEventNodeSpawner"
 
diff --git a/Engine/Source/Editor/BlueprintGraph/Private/BlueprintNodeSpawner.cpp b/Engine/Source/Editor/BlueprintGraph/Private/BlueprintNodeSpawner.cpp
--- a/Engine/Source/Editor/BlueprintGraph/Private/BlueprintNodeSpawner.cpp
+++ b/Engine/Source/Editor/BlueprintGraph/Private/BlueprintNodeSpawner.cpp
@@ -7,6 +7,7 @@
 #include "K2Node_IfThenElse.h"
 #include "BlueprintNodeTemplateCache.h"
 #include "BlueprintNodeSpawnerUtils.h"
+#include "Styling/AppStyle.h"
 
 /*******************************************************************************
  * Static UBlueprintNodeSpawner Helpers
diff --git a/Engine/Source/Editor/BlueprintGraph/Private/K2Node_CastByteToEnum.cpp b/Engine/Source/Editor/BlueprintGraph/Private/K2Node_CastByteToEnum.cpp
--- a/Engine/Source/Editor/BlueprintGraph/Private/K2Node_CastByteToEnum.cpp
+++ b/Engine/Source/Editor/BlueprintGraph/Private/K2Node_CastByteToEnum.cpp
@@ -11,6 +11,7 @@
 #include "BlueprintFieldNodeSpawner.h"
 #include "EditorCategoryUtils.h"
 #include "BlueprintActionDatabaseRegistrar.h"
+#include "Styling/AppStyle.h"
 
 const FName UK2Node_CastByteToEnum::ByteInputPinName(TEXT("Byte"));
 
diff --git a/Engine/Source/Editor/BlueprintGraph/Private/K2Node_CustomEvent.cpp b/Engine/Source/Editor/BlueprintGraph/Private/K2Node_CustomEvent.cpp
--- a/Engine/Source/Editor/BlueprintGraph/Private/K2Node_CustomEvent.cpp
+++ b/Engine/Source/Editor/BlueprintGraph/Private/K2Node_CustomEvent.cpp
@@ -14,6 +14,7 @@
 #include "BlueprintEventNodeSpawner.h"
 #include "BlueprintActionDatabaseRegistrar.h"
 #include "FindInBlueprintManager.h"
+#include "Styling/AppStyle.h"
 
 #define LOCTEXT_NAMESPACE "K2Node_CustomEvent"
 
diff --git a/Engine/Source/Editor/BlueprintGraph/Private/K2Node_ForEachElementInEnum.cpp b/Engine/Source/Editor/BlueprintGraph/Private/K2Node_ForEachElementInEnum.cpp
--- a/Engine/Source/Editor/BlueprintGraph/Private/K2Node_ForEachElementInEnum.cpp
+++ b/Engine/Source/Editor/BlueprintGraph/Private/K2Node_ForEachElementInEnum.cpp
@@ -17,6 +17,7 @@
 #include "BlueprintFieldNodeSpawner.h"
 #include "EditorCategoryUtils.h"
 #include "BlueprintActionDatabaseRegistrar.h"
+#include "Styling/AppStyle.h"
 
 #define LOCTEXT_NAMESPACE "K2Node"
 
diff --git a/Engine/Source/Editor/BlueprintGraph/Private/K2Node_GetDataTableRow.cpp b/Engine/Source/Editor/BlueprintGraph/Private/K2Node_GetDataTableRow.cpp
--- a/Engine/Source/Editor/BlueprintGraph/Private/K2Node_GetDataTableRow.cpp
+++ b/Engine/Source/Editor/BlueprintGraph/Private/K2Node_GetDataTableRow.cpp
@@ -12,6 +12,7 @@
 #include "BlueprintNodeSpawner.h"
 #include "EditorCategoryUtils.h"
 #include "BlueprintActionDatabaseRegistrar.h"
+#include "Styling/AppStyle.h"
 
 #define LOCTEXT_NAMESPACE "K2Node_GetDataTableRow"
 
diff --git a/Engine/Source/Editor/BlueprintGraph/Private/K2Node_GetEnumeratorName.cpp b/Engine/Source/Editor/BlueprintGraph/Private/K2Node_GetEnumeratorName.cpp
--- a/Engine/Source/Editor/BlueprintGraph/Private/K2Node_GetEnumeratorName.cpp
+++ b/Engine/Source/Editor/BlueprintGraph/Private/K2Node_GetEnumeratorName.cpp
@@ -10,6 +10,7 @@
 #include "BlueprintNodeSpawner.h"
 #include "EditorCategoryUtils.h"
 #include "BlueprintActionDatabaseRegistrar.h"
+#include "Styling/AppStyle.h"
 
 FName UK2Node_GetEnumeratorName::EnumeratorPinName(TEXT("Enumerator"));
 
diff --git a/Engine/Source/Editor/BlueprintGraph/Private/K2Node_GetEnumeratorNameAsString.cpp b/Engine/Source/Editor/BlueprintGraph/Private/K2Node_GetEnumeratorNameAsString.cpp
--- a/Engine/Source/Editor/BlueprintGraph/Private/K2Node_GetEnumeratorNameAsString.cpp
+++ b/Engine/Source/Editor/BlueprintGraph/Private/K2Node_GetEnumeratorNameAsString.cpp
@@ -7,6 +7,7 @@
 #include "BlueprintNodeSpawner.h"
 #include "EditorCategoryUtils.h"
 #include "BlueprintActionDatabaseRegistrar.h"
+#include "Styling/AppStyle.h"
 
 UK2Node_GetEnumeratorNameAsString::UK2Node_GetEnumeratorNameAsString(const FObjectInitializer& ObjectInitializer)
 	: Super(ObjectInitializer)
diff --git a/Engine/Source/Editor/BlueprintGraph/Private/K2Node_IfThenElse.cpp b/Engine/Source/Editor/BlueprintGraph/Private/K2Node_IfThenElse.cpp
--- a/Engine/Source/Editor/BlueprintGraph/Private/K2Node_IfThenElse.cpp
+++ b/Engine/Source/Editor/BlueprintGraph/Private/K2Node_IfThenElse.cpp
@@ -11,6 +11,7 @@
 #include "BlueprintNodeSpawner.h"
 #include "EditorCategoryUtils.h"
 #include "BlueprintActionDatabaseRegistrar.h"
+#include "Styling/AppStyle.h"
 
 #define LOCTEXT_NAMESPACE "K2Node"
 
diff --git a/Engine/Source/Editor/BlueprintGraph/Private/K2Node_InputAxisKeyEvent.cpp b/Engine/Source/Editor/BlueprintGraph/Private/K2Node_InputAxisKeyEvent.cpp
--- a/Engine/Source/Editor/BlueprintGraph/Private/K2Node_InputAxisKeyEvent.cpp
+++ b/Engine/Source/Editor/BlueprintGraph/Private/K2Node_InputAxisKeyEvent.cpp
@@ -8,6 +8,7 @@
 #include "EditorCategoryUtils.h"
 #include "Engine/InputAxisKeyDelegateBinding.h"
 #include "BlueprintActionDatabaseRegistrar.h"
+#include "Styling/AppStyle.h"
 
 #define LOCTEXT_NAMESPACE "UK2Node_InputAxisKeyEvent"
 
diff --git a/Engine/Source/Editor/BlueprintGraph/Private/K2Node_Switch.cpp b/Engine/Source/Editor/BlueprintGraph/Private/K2Node_Switch.cpp
--- a/Engine/Source/Editor/BlueprintGraph/Private/K2Node_Switch.cpp
+++ b/Engine/Source/Editor/BlueprintGraph/Private/K2Node_Switch.cpp
@@ -10,6 +10,7 @@
 #include "KismetCompilerMisc.h"
 #include "KismetCompiler.h"
 #include "EditorCategoryUtils.h"
+#include "Styling/AppStyle.h"
 
 #define LOCTEXT_NAMESPACE "K2Node_Switch"
 
diff --git a/Engine/Source/Editor/UnrealEd/Private/EdGraphNode_Comment.cpp b/Engine/Source/Editor/UnrealEd/Private/EdGraphNode_Comment.cpp
--- a/Engine/Source/Editor/UnrealEd/Private/EdGraphNode_Comment.cpp
+++ b/Engine/Source/Editor/UnrealEd/Private/EdGraphNode_Comment.cpp
@@ -6,6 +6,8 @@
 #include "GraphEditorSettings.h"
 #include "Kismet2/Kismet2NameValidators.h"
 
+#include "Styling/AppStyle.h"
+
 #define LOCTEXT_NAMESPACE "EdGraph"
 
 namespace FEdGraphNode_Comment_Utils
diff --git a/Engine/Source/Runtime/Engine/Private/EdGraph/EdGraphNode.cpp b/Engine/Source/Runtime/Engine/Private/EdGraph/EdGraphNode.cpp
--- a/Engine/Source/Runtime/Engine/Private/EdGraph/EdGraphNode.cpp
+++ b/Engine/Source/Runtime/Engine/Private/EdGraph/EdGraphNode.cpp
@@ -21,6 +21,7 @@
 #include "FindInBlueprintManager.h"
 #include "DiffResults.h"
 #endif
+#include "Styling/AppStyle.h"
 
 #define LOCTEXT_NAMESPACE "EdGraph"
 
diff --git a/Engine/Source/Runtime/Engine/Private/EdGraph/EdGraphNode_Documentation.cpp b/Engine/Source/Runtime/Engine/Private/EdGraph/EdGraphNode_Documentation.cpp
--- a/Engine/Source/Runtime/Engine/Private/EdGraph/EdGraphNode_Documentation.cpp
+++ b/Engine/Source/Runtime/Engine/Private/EdGraph/EdGraphNode_Documentation.cpp
@@ -7,6 +7,8 @@
 #include "Kismet2/Kismet2NameValidators.h"
 #endif
 
+#include "Styling/AppStyle.h"
+
 #define LOCTEXT_NAMESPACE "EdGraph"
 
 /////////////////////////////////////////////////////
