Fix type promotion to correctly exclude UEdGraphSchema_K2::PC_SoftClass, not the PC_SoftObject.

diff --git a/Engine/Source/Editor/BlueprintGraph/Private/BlueprintEditorSettings.cpp b/Engine/Source/Editor/BlueprintGraph/Private/BlueprintEditorSettings.cpp
--- a/Engine/Source/Editor/BlueprintGraph/Private/BlueprintEditorSettings.cpp
+++ b/Engine/Source/Editor/BlueprintGraph/Private/BlueprintEditorSettings.cpp
@@ -35,7 +35,7 @@ UBlueprintEditorSettings::UBlueprintEditorSettings(const FObjectInitializer& Obj
 	, bNavigateToNativeFunctionsFromCallNodes(true)
 	, bDoubleClickNavigatesToParent(true)
 	, bEnableTypePromotion(true)
-	, TypePromotionPinDenyList { UEdGraphSchema_K2::PC_String, UEdGraphSchema_K2::PC_Text, UEdGraphSchema_K2::PC_SoftObject }
+	, TypePromotionPinDenyList { UEdGraphSchema_K2::PC_String, UEdGraphSchema_K2::PC_Text, UEdGraphSchema_K2::PC_SoftClass }
 	, BreakpointReloadMethod(EBlueprintBreakpointReloadMethod::RestoreAll)
 	, bEnablePinValueInspectionTooltips(true)
 	// Experimental
