K2Node_BaseAsyncTask uses param DisplayName if available
(consistent with K2Node_CallFunction)

diff --git a/Engine/Source/Editor/BlueprintGraph/Private/K2Node_BaseAsyncTask.cpp b/Engine/Source/Editor/BlueprintGraph/Private/K2Node_BaseAsyncTask.cpp
--- a/Engine/Source/Editor/BlueprintGraph/Private/K2Node_BaseAsyncTask.cpp
+++ b/Engine/Source/Editor/BlueprintGraph/Private/K2Node_BaseAsyncTask.cpp
@@ -169,6 +169,13 @@
 
 			if (bPinGood)
 			{
+				// Check for a display name override
+				const FString& PinDisplayName = Param->GetMetaData(FBlueprintMetadata::MD_DisplayName);
+				if (!PinDisplayName.IsEmpty())
+				{
+					Pin->PinFriendlyName = FText::FromString(PinDisplayName);
+				}
+				
 				//Flag pin as read only for const reference property
 				Pin->bDefaultValueIsIgnored = Param->HasAllPropertyFlags(CPF_ConstParm | CPF_ReferenceParm) && (!Function->HasMetaData(FBlueprintMetadata::MD_AutoCreateRefTerm) || Pin->PinType.IsContainer());
 
