Avoid crash because this == null no longer works on Linux
#rnx

diff --git a/Engine/Source/Editor/BlueprintGraph/Private/K2Node_CallFunction.cpp b/Engine/Source/Editor/BlueprintGraph/Private/K2Node_CallFunction.cpp
--- a/Engine/Source/Editor/BlueprintGraph/Private/K2Node_CallFunction.cpp
+++ b/Engine/Source/Editor/BlueprintGraph/Private/K2Node_CallFunction.cpp
@@ -1031,7 +1031,7 @@ bool UK2Node_CallFunction::CreatePinsForFunctionCall(const UFunction* Function)
 	ensure(BP);
 	if (BP != nullptr)
 	{
-		const bool bIsFunctionCompatibleWithSelf = BP->SkeletonGeneratedClass->IsChildOf(FunctionOwnerClass);
+		const bool bIsFunctionCompatibleWithSelf = BP->SkeletonGeneratedClass && BP->SkeletonGeneratedClass->IsChildOf(FunctionOwnerClass);
 
 		if (bIsStaticFunc)
 		{
