Fix another fallout from this == null no longer protecting IsChildOf
#rnx

diff --git a/Engine/Source/Editor/BlueprintGraph/Private/K2Node_VariableSet.cpp b/Engine/Source/Editor/BlueprintGraph/Private/K2Node_VariableSet.cpp
--- a/Engine/Source/Editor/BlueprintGraph/Private/K2Node_VariableSet.cpp
+++ b/Engine/Source/Editor/BlueprintGraph/Private/K2Node_VariableSet.cpp
@@ -296,14 +296,15 @@ bool UK2Node_VariableSet::HasLocalRepNotify() const
 
 bool UK2Node_VariableSet::ShouldFlushDormancyOnSet() const
 {
-	if (!GetVariableSourceClass()->IsChildOf(AActor::StaticClass()))
+	const UClass* VariableSourceClass = GetVariableSourceClass();
+	if (VariableSourceClass == nullptr || !VariableSourceClass->IsChildOf(AActor::StaticClass()))
 	{
 		return false;
 	}
 
 	// Flush net dormancy before setting a replicated property
-	FProperty *Property = FindFProperty<FProperty>(GetVariableSourceClass(), GetVarName());
-	return (Property != NULL && (Property->PropertyFlags & CPF_Net));
+	FProperty* Property = FindFProperty<FProperty>(VariableSourceClass, GetVarName());
+	return (Property && (Property->PropertyFlags & CPF_Net));
 }
 
 bool UK2Node_VariableSet::IsNetProperty() const
@@ -314,8 +315,7 @@ bool UK2Node_VariableSet::IsNetProperty() const
 
 FName UK2Node_VariableSet::GetRepNotifyName() const
 {
-	FProperty * Property = GetPropertyForVariable();
-	if (Property)
+	if (FProperty* Property = GetPropertyForVariable())
 	{
 		return Property->RepNotifyFunc;
 	}
