Clang now optimizes away this==nullptr. Calls to IsChildOf from a null UObject pointer will cause undefined behavior. This is a retroactive attempt to pad potentially dangerous calls to IsChildOf with a null check in the following directories:
- Plugins/BlueprintContext
- Editor/GlueprintGraph
- Editor/GraphEditor
- Editor/Kismet
- Editor/KismetCompiler
- Editor/UnrealEd/Private/Kismet2

note: if you're seeing this CL in the perforce history because you're trying to figure out why there's a null check that doesn't make sense, This is why. The goal of this CL is to preserve the behavior before IsChildOf changed rather than analyze whether that behavior makes sense. Use your best judgement

diff --git a/Engine/Source/Editor/BlueprintGraph/Private/BlueprintActionDatabaseRegistrar.cpp b/Engine/Source/Editor/BlueprintGraph/Private/BlueprintActionDatabaseRegistrar.cpp
--- a/Engine/Source/Editor/BlueprintGraph/Private/BlueprintActionDatabaseRegistrar.cpp
+++ b/Engine/Source/Editor/BlueprintGraph/Private/BlueprintActionDatabaseRegistrar.cpp
@@ -319,7 +319,8 @@ int32 FBlueprintActionDatabaseRegistrar::RegisterClassFactoryActions(const UClas
 			{
 				FObjectProperty* ReturnProperty = CastField<FObjectProperty>(Function->GetReturnProperty());
 				// see if the function is a static factory method
-				bool const bIsFactoryMethod = (ReturnProperty != nullptr) && ReturnProperty->PropertyClass->IsChildOf(InTargetType);
+				bool const bIsFactoryMethod = (ReturnProperty != nullptr) && (ReturnProperty->PropertyClass != nullptr) &&
+					ReturnProperty->PropertyClass->IsChildOf(InTargetType);
 
 				return bIsFactoryMethod;
 			}
diff --git a/Engine/Source/Editor/BlueprintGraph/Private/BlueprintActionFilter.cpp b/Engine/Source/Editor/BlueprintGraph/Private/BlueprintActionFilter.cpp
--- a/Engine/Source/Editor/BlueprintGraph/Private/BlueprintActionFilter.cpp
+++ b/Engine/Source/Editor/BlueprintGraph/Private/BlueprintActionFilter.cpp
@@ -1814,7 +1814,7 @@ static bool BlueprintActionFilterImpl::IsIncompatibleMacroInstance(FBlueprintAct
 					return true;
 				}
 
-				bIsFilteredOut = (Blueprint != MacroBP) && (MacroBP->BlueprintType != BPTYPE_MacroLibrary || !Blueprint->ParentClass->IsChildOf(MacroBP->ParentClass));
+				bIsFilteredOut = (Blueprint != MacroBP) && (MacroBP->BlueprintType != BPTYPE_MacroLibrary || !Blueprint->ParentClass || !Blueprint->ParentClass->IsChildOf(MacroBP->ParentClass));
 			}
 
 			// Note: The rest is handled by IsNodeTemplateSelfFiltered() - the check above is a "fast path" in that we don't have to instance the node template (see UK2Node_MacroInstance::IsActionFilteredOut())
diff --git a/Engine/Source/Editor/BlueprintGraph/Private/BlueprintBoundEventNodeSpawner.cpp b/Engine/Source/Editor/BlueprintGraph/Private/BlueprintBoundEventNodeSpawner.cpp
--- a/Engine/Source/Editor/BlueprintGraph/Private/BlueprintBoundEventNodeSpawner.cpp
+++ b/Engine/Source/Editor/BlueprintGraph/Private/BlueprintBoundEventNodeSpawner.cpp
@@ -147,7 +147,7 @@ bool UBlueprintBoundEventNodeSpawner::IsBindingCompatible(FBindingObject Binding
 	UClass* DelegateOwner = Delegate->GetOwnerClass()->GetAuthoritativeClass();
 	UClass* BindingClass  = FBlueprintNodeSpawnerUtils::GetBindingClass(BindingCandidate)->GetAuthoritativeClass();
 
-	return bMatchesNodeType && BindingClass->IsChildOf(DelegateOwner) && !FObjectEditorUtils::IsVariableCategoryHiddenFromClass(Delegate, BindingClass);
+	return bMatchesNodeType && BindingClass && BindingClass->IsChildOf(DelegateOwner) && !FObjectEditorUtils::IsVariableCategoryHiddenFromClass(Delegate, BindingClass);
 }
 
 //------------------------------------------------------------------------------
diff --git a/Engine/Source/Editor/BlueprintGraph/Private/BlueprintFunctionNodeSpawner.cpp b/Engine/Source/Editor/BlueprintGraph/Private/BlueprintFunctionNodeSpawner.cpp
--- a/Engine/Source/Editor/BlueprintGraph/Private/BlueprintFunctionNodeSpawner.cpp
+++ b/Engine/Source/Editor/BlueprintGraph/Private/BlueprintFunctionNodeSpawner.cpp
@@ -376,7 +376,7 @@ FBlueprintActionUiSpec UBlueprintFunctionNodeSpawner::GetUiSpec(FBlueprintAction
 		checkSlow(WrappedFunction != nullptr);
 		UClass* FunctionClass = WrappedFunction->GetOwnerClass()->GetAuthoritativeClass();
 
-		if (!TargetClass->IsChildOf(FunctionClass))
+		if (!TargetClass || !TargetClass->IsChildOf(FunctionClass))
 		{
 			// When there are no bindings set, functions need to be categorized into a category "Class" to help reduce clutter in the tree root
 			if(Bindings.Num() == 0)
@@ -487,7 +487,7 @@ bool UBlueprintFunctionNodeSpawner::IsBindingCompatible(FBindingObject BindingCa
 	UClass* BindingClass = FBlueprintNodeSpawnerUtils::GetBindingClass(BindingCandidate)->GetAuthoritativeClass();
 	if (UClass const* FuncOwner = Function->GetOwnerClass()->GetAuthoritativeClass())
 	{
-		bClassOwnerMatches = BindingClass->IsChildOf(FuncOwner);
+		bClassOwnerMatches = BindingClass && BindingClass->IsChildOf(FuncOwner);
 	}
 
 	return bNodeTypeMatches && bClassOwnerMatches && !FObjectEditorUtils::IsFunctionHiddenFromClass(Function, BindingClass);
diff --git a/Engine/Source/Editor/BlueprintGraph/Private/BlueprintVariableNodeSpawner.cpp b/Engine/Source/Editor/BlueprintGraph/Private/BlueprintVariableNodeSpawner.cpp
--- a/Engine/Source/Editor/BlueprintGraph/Private/BlueprintVariableNodeSpawner.cpp
+++ b/Engine/Source/Editor/BlueprintGraph/Private/BlueprintVariableNodeSpawner.cpp
@@ -42,15 +42,18 @@ UBlueprintVariableNodeSpawner* UBlueprintVariableNodeSpawner::CreateFromMemberOr
 
 	FText const VarName = NodeSpawner->GetVariableName();
 	// @TODO: NodeClass could be modified post Create()
-	if (NodeClass->IsChildOf(UK2Node_VariableGet::StaticClass()))
+	if (NodeClass != nullptr)
 	{
-		MenuSignature.MenuName = FText::Format(LOCTEXT("GetterMenuName", "Get {0}"), VarName);
-		MenuSignature.Tooltip  = UK2Node_VariableGet::GetPropertyTooltip(VarProperty);
-	}
-	else if (NodeClass->IsChildOf(UK2Node_VariableSet::StaticClass()))
-	{
-		MenuSignature.MenuName = FText::Format(LOCTEXT("SetterMenuName", "Set {0}"), VarName);
-		MenuSignature.Tooltip  = UK2Node_VariableSet::GetPropertyTooltip(VarProperty);
+		if (NodeClass->IsChildOf(UK2Node_VariableGet::StaticClass()))
+		{
+			MenuSignature.MenuName = FText::Format(LOCTEXT("GetterMenuName", "Get {0}"), VarName);
+			MenuSignature.Tooltip  = UK2Node_VariableGet::GetPropertyTooltip(VarProperty);
+		}
+		else if (NodeClass->IsChildOf(UK2Node_VariableSet::StaticClass()))
+		{
+			MenuSignature.MenuName = FText::Format(LOCTEXT("SetterMenuName", "Set {0}"), VarName);
+			MenuSignature.Tooltip  = UK2Node_VariableSet::GetPropertyTooltip(VarProperty);
+		}
 	}
 	// add at least one character, so that PrimeDefaultUiSpec() doesn't 
 	// attempt to query the template node
@@ -117,15 +120,18 @@ UBlueprintVariableNodeSpawner* UBlueprintVariableNodeSpawner::CreateFromLocal(TS
 
 	FText const VarName = NodeSpawner->GetVariableName();
 	// @TODO: NodeClass could be modified post Create()
-	if (NodeClass->IsChildOf(UK2Node_VariableGet::StaticClass()))
+	if (NodeClass != nullptr)
 	{
-		MenuSignature.MenuName = FText::Format(LOCTEXT("LocalGetterMenuName", "Get {0}"), VarName);
-		MenuSignature.Tooltip  = UK2Node_VariableGet::GetBlueprintVarTooltip(VarDesc);
-	}
-	else if (NodeClass->IsChildOf(UK2Node_VariableSet::StaticClass()))
-	{
-		MenuSignature.MenuName = FText::Format(LOCTEXT("LocalSetterMenuName", "Set {0}"), VarName);
-		MenuSignature.Tooltip  = UK2Node_VariableSet::GetBlueprintVarTooltip(VarDesc);
+		if (NodeClass->IsChildOf(UK2Node_VariableGet::StaticClass()))
+		{
+			MenuSignature.MenuName = FText::Format(LOCTEXT("LocalGetterMenuName", "Get {0}"), VarName);
+			MenuSignature.Tooltip  = UK2Node_VariableGet::GetBlueprintVarTooltip(VarDesc);
+		}
+		else if (NodeClass->IsChildOf(UK2Node_VariableSet::StaticClass()))
+		{
+			MenuSignature.MenuName = FText::Format(LOCTEXT("LocalSetterMenuName", "Set {0}"), VarName);
+			MenuSignature.Tooltip  = UK2Node_VariableSet::GetBlueprintVarTooltip(VarDesc);
+		}
 	}
 	// add at least one character, so that PrimeDefaultUiSpec() doesn't 
 	// attempt to query the template node
@@ -192,7 +198,7 @@ FBlueprintActionUiSpec UBlueprintVariableNodeSpawner::GetUiSpec(FBlueprintAction
 		EffectiveOwnerClass = EffectiveOwnerClass ? EffectiveOwnerClass : ToRawPtr(OwnerClass);
 		const bool bIsOwningClassValid = EffectiveOwnerClass && (!Cast<const UBlueprintGeneratedClass>(EffectiveOwnerClass) || EffectiveOwnerClass->ClassGeneratedBy); //todo: more general validation
 		UClass const* VariableClass = bIsOwningClassValid ? EffectiveOwnerClass->GetAuthoritativeClass() : nullptr;
-		if (VariableClass && !TargetClass->IsChildOf(VariableClass))
+		if (VariableClass && (!TargetClass || !TargetClass->IsChildOf(VariableClass)))
 		{
 			MenuSignature.Category = FEditorCategoryUtils::BuildCategoryString(FCommonEditorCategory::Class,
 				FText::FromString(VariableClass->GetDisplayNameText().ToString()));
diff --git a/Engine/Source/Editor/BlueprintGraph/Private/EdGraphSchema_K2.cpp b/Engine/Source/Editor/BlueprintGraph/Private/EdGraphSchema_K2.cpp
--- a/Engine/Source/Editor/BlueprintGraph/Private/EdGraphSchema_K2.cpp
+++ b/Engine/Source/Editor/BlueprintGraph/Private/EdGraphSchema_K2.cpp
@@ -795,7 +795,7 @@ bool UEdGraphSchema_K2::CanFunctionBeUsedInGraph(const UClass* InClass, const UF
 			// Blueprint functions visibility flags can be enforced in blueprints - native functions
 			// are often using these flags to only hide functionality from other native functions:
 			const bool bIsProtected = (InFunction->FunctionFlags & FUNC_Protected) != 0;
-			const bool bFuncBelongsToSubClass = InClass->IsChildOf(InFunction->GetOuterUClass()->GetSuperStruct());
+			const bool bFuncBelongsToSubClass = InClass && InClass->IsChildOf(InFunction->GetOuterUClass()->GetSuperStruct());
 			if (bIsProtected)
 			{
 				const bool bAllowProtectedFuncs = (InAllowedFunctionTypes & FT_Protected) != 0;
@@ -2613,7 +2613,7 @@ bool UEdGraphSchema_K2::FindSpecializedConversionNode(const UEdGraphPin* OutputP
 			{
 				FObjectProperty* ObjProp = *PropIt;
 				// .. if we have a blueprint visible var, and is of the type which contains this function..
-				if(ObjProp->HasAllPropertyFlags(CPF_BlueprintVisible) && ObjProp->PropertyClass->IsChildOf(FunctionClass))
+				if(ObjProp->HasAllPropertyFlags(CPF_BlueprintVisible) && ObjProp->PropertyClass && ObjProp->PropertyClass->IsChildOf(FunctionClass))
 				{
 					// say we can convert
 					bCanConvert = true;
@@ -4066,7 +4066,7 @@ namespace
 		// We have a native (C++) class, do a normal IsChildOf check
 		if (!bIsNonNativeClass)
 		{
-			bResult = InSourceStruct->IsChildOf(InTargetStruct);
+			bResult = InSourceStruct && InSourceStruct->IsChildOf(InTargetStruct);
 		}
 
 		return bResult;
@@ -4074,7 +4074,7 @@ namespace
 
 	static bool ExtendedIsChildOf(const UClass* Child, const UClass* Parent)
 	{
-		if (Child->IsChildOf(Parent))
+		if (Child && Child->IsChildOf(Parent))
 		{
 			return true;
 		}
@@ -4086,7 +4086,7 @@ namespace
 		}
 
 		const UClass* OriginalParent = GetOriginalClassToFixCompatibilit(Parent);
-		if (OriginalParent && Child->IsChildOf(OriginalParent))
+		if (OriginalParent && Child && Child->IsChildOf(OriginalParent))
 		{
 			return true;
 		}
@@ -4261,12 +4261,16 @@ bool UEdGraphSchema_K2::DefaultValueSimpleValidation(const FEdGraphPinType& PinT
 		{
 			ObjectClass = ObjectClass->GetAuthoritativeClass();
 		}
-		if (NewDefaultObject != nullptr && ObjectClass != nullptr && !NewDefaultObject->GetClass()->GetAuthoritativeClass()->IsChildOf(ObjectClass))
+		if (NewDefaultObject != nullptr && ObjectClass != nullptr)
 		{
-			// Not a type of object, but is it an object implementing an interface?
-			if(PinCategory != PC_Interface || !NewDefaultObject->GetClass()->ImplementsInterface(ObjectClass))
+			const UClass* AuthoritativeClass = NewDefaultObject->GetClass()->GetAuthoritativeClass();
+			if (!AuthoritativeClass || !AuthoritativeClass->IsChildOf(ObjectClass))
 			{
-				DVSV_RETURN_MSG(FString::Printf(TEXT("%s isn't a %s (specified on pin %s)"), *NewDefaultObject->GetPathName(), *ObjectClass->GetName(), *PinName.ToString()));
+				// Not a type of object, but is it an object implementing an interface?
+                if(PinCategory != PC_Interface || !NewDefaultObject->GetClass()->ImplementsInterface(ObjectClass))
+                {
+                	DVSV_RETURN_MSG(FString::Printf(TEXT("%s isn't a %s (specified on pin %s)"), *NewDefaultObject->GetPathName(), *ObjectClass->GetName(), *PinName.ToString()));
+                }
 			}
 		}
 	}
@@ -5738,7 +5742,7 @@ void UEdGraphSchema_K2::DroppedAssetsOnGraph(const TArray<FAssetData>& Assets, c
 			}
 
 			TSubclassOf<UActorComponent> DestinationComponentType;
-			if (AssetClass->IsChildOf(UActorComponent::StaticClass()) && IsAllowableBlueprintVariableType(AssetClass))
+			if (AssetClass && AssetClass->IsChildOf(UActorComponent::StaticClass()) && IsAllowableBlueprintVariableType(AssetClass))
 			{
 				// If it's an actor component subclass that is a BlueprintableComponent, we're good to go
 				DestinationComponentType = AssetClass;
@@ -5747,7 +5751,7 @@ void UEdGraphSchema_K2::DroppedAssetsOnGraph(const TArray<FAssetData>& Assets, c
 			{
 				// Otherwise see if we can factory a component from the asset
 				DestinationComponentType = FComponentAssetBrokerage::GetPrimaryComponentForAsset(AssetClass);
-				if ((DestinationComponentType == nullptr) && AssetClass->IsChildOf(AActor::StaticClass()))
+				if ((DestinationComponentType == nullptr) && AssetClass && AssetClass->IsChildOf(AActor::StaticClass()))
 				{
 					DestinationComponentType = UChildActorComponent::StaticClass();
 				}
@@ -5845,7 +5849,7 @@ void UEdGraphSchema_K2::GetAssetsGraphHoverMessage(const TArray<FAssetData>& Ass
 				}
 
 				TSubclassOf<UActorComponent> DestinationComponentType;
-				if (AssetClass->IsChildOf(UActorComponent::StaticClass()) && IsAllowableBlueprintVariableType(AssetClass))
+				if (AssetClass && AssetClass->IsChildOf(UActorComponent::StaticClass()) && IsAllowableBlueprintVariableType(AssetClass))
 				{
 					// If it's an actor component subclass that is a BlueprintableComponent, we're good to go
 					DestinationComponentType = AssetClass;
@@ -5854,7 +5858,7 @@ void UEdGraphSchema_K2::GetAssetsGraphHoverMessage(const TArray<FAssetData>& Ass
 				{
 					// Otherwise, see if we have a way to make a component out of the specified asset
 					DestinationComponentType = FComponentAssetBrokerage::GetPrimaryComponentForAsset(AssetClass);
-					if ((DestinationComponentType == nullptr) && AssetClass->IsChildOf(AActor::StaticClass()))
+					if ((DestinationComponentType == nullptr) && AssetClass && AssetClass->IsChildOf(AActor::StaticClass()))
 					{
 						DestinationComponentType = UChildActorComponent::StaticClass();
 					}
@@ -7344,7 +7348,7 @@ void UEdGraphSchema_K2::OnPinConnectionDoubleCicked(UEdGraphPin* PinA, UEdGraphP
 void UEdGraphSchema_K2::ConfigureVarNode(UK2Node_Variable* InVarNode, FName InVariableName, UStruct* InVariableSource, UBlueprint* InTargetBlueprint)
 {
 	// See if this is a 'self context' (ie. blueprint class is owner (or child of owner) of dropped var class)
-	if ((InVariableSource == NULL) || InTargetBlueprint->SkeletonGeneratedClass->IsChildOf(InVariableSource))
+	if ((InVariableSource == NULL) || (InTargetBlueprint->SkeletonGeneratedClass && InTargetBlueprint->SkeletonGeneratedClass->IsChildOf(InVariableSource)))
 	{
 		FGuid Guid = FBlueprintEditorUtils::FindMemberVariableGuidByName(InTargetBlueprint, InVariableName);
 		InVarNode->VariableReference.SetSelfMember(InVariableName, Guid);
diff --git a/Engine/Source/Editor/BlueprintGraph/Private/K2Node_CallFunction.cpp b/Engine/Source/Editor/BlueprintGraph/Private/K2Node_CallFunction.cpp
--- a/Engine/Source/Editor/BlueprintGraph/Private/K2Node_CallFunction.cpp
+++ b/Engine/Source/Editor/BlueprintGraph/Private/K2Node_CallFunction.cpp
@@ -1175,7 +1175,7 @@ void UK2Node_CallFunction::PostReconstructNode()
 		if (UBlueprint* BP = GetBlueprint())
 		{
 			UClass* FunctionOwnerClass = Function->GetOuterUClass();
-			if (!BP->SkeletonGeneratedClass->IsChildOf(FunctionOwnerClass))
+			if (!BP->SkeletonGeneratedClass || !BP->SkeletonGeneratedClass->IsChildOf(FunctionOwnerClass))
 			{
 				SelfPin->DefaultObject = FunctionOwnerClass->GetAuthoritativeClass()->GetDefaultObject();
 			}
@@ -2132,7 +2132,7 @@ void UK2Node_CallFunction::ValidateNodeDuringCompilation(class FCompilerResultsL
 			const bool bCanTreatAsError = Blueprint->GetLinkerCustomVersion(FFrameworkObjectVersion::GUID) >= FFrameworkObjectVersion::EnforceBlueprintFunctionVisibility;
 
 			const bool bIsProtected = (Function->FunctionFlags & FUNC_Protected) != 0;
-			const bool bFuncBelongsToSubClass = Blueprint->SkeletonGeneratedClass->IsChildOf(Function->GetOuterUClass());
+			const bool bFuncBelongsToSubClass = Blueprint->SkeletonGeneratedClass && Blueprint->SkeletonGeneratedClass->IsChildOf(Function->GetOuterUClass());
 			if (bIsProtected && !bFuncBelongsToSubClass)
 			{
 				if(bCanTreatAsError)
diff --git a/Engine/Source/Editor/BlueprintGraph/Private/K2Node_ComponentBoundEvent.cpp b/Engine/Source/Editor/BlueprintGraph/Private/K2Node_ComponentBoundEvent.cpp
--- a/Engine/Source/Editor/BlueprintGraph/Private/K2Node_ComponentBoundEvent.cpp
+++ b/Engine/Source/Editor/BlueprintGraph/Private/K2Node_ComponentBoundEvent.cpp
@@ -102,7 +102,7 @@ void UK2Node_ComponentBoundEvent::RegisterDynamicBinding(UDynamicBlueprintBindin
 
 void UK2Node_ComponentBoundEvent::HandleVariableRenamed(UBlueprint* InBlueprint, UClass* InVariableClass, UEdGraph* InGraph, const FName& InOldVarName, const FName& InNewVarName)
 {	
-	if (InVariableClass->IsChildOf(InBlueprint->GeneratedClass))
+	if (InVariableClass && InVariableClass->IsChildOf(InBlueprint->GeneratedClass))
 	{
 		// This could be the case if the component that this was originally bound to was removed, and a new one was 
 		// added in it's place. @see UE-88511
diff --git a/Engine/Source/Editor/BlueprintGraph/Private/K2Node_DynamicCast.cpp b/Engine/Source/Editor/BlueprintGraph/Private/K2Node_DynamicCast.cpp
--- a/Engine/Source/Editor/BlueprintGraph/Private/K2Node_DynamicCast.cpp
+++ b/Engine/Source/Editor/BlueprintGraph/Private/K2Node_DynamicCast.cpp
@@ -461,7 +461,7 @@ void UK2Node_DynamicCast::ValidateNodeDuringCompilation(FCompilerResultsLog& Mes
 				FText const WarningFormat = LOCTEXT("UnneededObjectCastFmt", "'{0}' is already a '{1}' (which inherits from '{2}'), so you don't need @@.");
 				MessageLog.Note( *FText::Format(WarningFormat, FText::FromString(SourcePinName), SourceClass->GetDisplayNameText(), TargetType->GetDisplayNameText()).ToString(), this );
 			}
-			else if (!SourceType->IsChildOf(SourceClass) && !FKismetEditorUtilities::IsClassABlueprintInterface(SourceType))
+			else if ((!SourceType || !SourceType->IsChildOf(SourceClass)) && !FKismetEditorUtilities::IsClassABlueprintInterface(SourceType))
 			{
 				FText const WarningFormat = LOCTEXT("DisallowedObjectCast", "'{0}' does not inherit from '{1}' (@@ would always fail).");
 				MessageLog.Warning( *FText::Format(WarningFormat, TargetType->GetDisplayNameText(), SourceClass->GetDisplayNameText()).ToString(), this );
diff --git a/Engine/Source/Editor/BlueprintGraph/Private/K2Node_MCDelegate.cpp b/Engine/Source/Editor/BlueprintGraph/Private/K2Node_MCDelegate.cpp
--- a/Engine/Source/Editor/BlueprintGraph/Private/K2Node_MCDelegate.cpp
+++ b/Engine/Source/Editor/BlueprintGraph/Private/K2Node_MCDelegate.cpp
@@ -239,7 +239,8 @@ void UK2Node_BaseMCDelegate::AutowireNewNode(UEdGraphPin* FromPin)
 				if (DelegateProperty)
 				{
 					UClass* DelegateOwner = DelegateProperty->GetOwnerClass();
-					if (FromPin->PinType.PinSubCategoryObject == DelegateOwner || dynamic_cast<UClass*>(FromPin->PinType.PinSubCategoryObject.Get())->IsChildOf(DelegateOwner))
+					const UClass* PinSubCategoryObject = dynamic_cast<UClass*>(FromPin->PinType.PinSubCategoryObject.Get());
+					if (FromPin->PinType.PinSubCategoryObject == DelegateOwner || (PinSubCategoryObject && PinSubCategoryObject->IsChildOf(DelegateOwner)))
 					{
 						// If we get here, then the property delegate is also available on the FromPin's class.
 						// Fix up the type by propagating it from the FromPin to our target pin
diff --git a/Engine/Source/Editor/BlueprintGraph/Private/K2Node_MacroInstance.cpp b/Engine/Source/Editor/BlueprintGraph/Private/K2Node_MacroInstance.cpp
--- a/Engine/Source/Editor/BlueprintGraph/Private/K2Node_MacroInstance.cpp
+++ b/Engine/Source/Editor/BlueprintGraph/Private/K2Node_MacroInstance.cpp
@@ -430,7 +430,8 @@ bool UK2Node_MacroInstance::CanPasteHere(const UEdGraph* TargetGraph) const
 	{
 		// Only allow "local" macro instances or instances from a macro library blueprint with the same parent class
 		check(MacroBlueprint->ParentClass != nullptr && TargetBlueprint->ParentClass != nullptr);
-		bCanPaste = (MacroBlueprint == TargetBlueprint) || (MacroBlueprint->BlueprintType == BPTYPE_MacroLibrary && TargetBlueprint->ParentClass->IsChildOf(MacroBlueprint->ParentClass));
+		bCanPaste = (MacroBlueprint == TargetBlueprint) || (MacroBlueprint->BlueprintType == BPTYPE_MacroLibrary
+			&& TargetBlueprint->ParentClass && TargetBlueprint->ParentClass->IsChildOf(MacroBlueprint->ParentClass));
 	}
 
 	// Macro Instances are not allowed in it's own graph
diff --git a/Engine/Source/Editor/BlueprintGraph/Private/K2Node_Switch.cpp b/Engine/Source/Editor/BlueprintGraph/Private/K2Node_Switch.cpp
--- a/Engine/Source/Editor/BlueprintGraph/Private/K2Node_Switch.cpp
+++ b/Engine/Source/Editor/BlueprintGraph/Private/K2Node_Switch.cpp
@@ -308,7 +308,7 @@ void UK2Node_Switch::CreateFunctionPin()
 		if (UBlueprint* BP = GetBlueprint())
 		{
 			UClass* FunctionOwnerClass = Function->GetOuterUClass();
-			if (!BP->SkeletonGeneratedClass->IsChildOf(FunctionOwnerClass))
+			if (!BP->SkeletonGeneratedClass || !BP->SkeletonGeneratedClass->IsChildOf(FunctionOwnerClass))
 			{
 				FunctionPin->DefaultObject = FunctionOwnerClass->GetDefaultObject();
 			}
diff --git a/Engine/Source/Editor/BlueprintGraph/Private/K2Node_Variable.cpp b/Engine/Source/Editor/BlueprintGraph/Private/K2Node_Variable.cpp
--- a/Engine/Source/Editor/BlueprintGraph/Private/K2Node_Variable.cpp
+++ b/Engine/Source/Editor/BlueprintGraph/Private/K2Node_Variable.cpp
@@ -982,7 +982,7 @@ bool UK2Node_Variable::CanPasteHere(const UEdGraph* TargetGraph) const
 		{
 			const UClass* CurrentClass = GetBlueprint()->SkeletonGeneratedClass->GetAuthoritativeClass();
 			const UClass* PropertyClass = Property->GetOwnerClass()->GetAuthoritativeClass();
-			const bool bIsChildOf = CurrentClass->IsChildOf(PropertyClass);
+			const bool bIsChildOf = CurrentClass && CurrentClass->IsChildOf(PropertyClass);
 			return bIsChildOf;
 		}
 		return false;
diff --git a/Engine/Source/Editor/BlueprintGraph/Private/StructMemberNodeHandlers.cpp b/Engine/Source/Editor/BlueprintGraph/Private/StructMemberNodeHandlers.cpp
--- a/Engine/Source/Editor/BlueprintGraph/Private/StructMemberNodeHandlers.cpp
+++ b/Engine/Source/Editor/BlueprintGraph/Private/StructMemberNodeHandlers.cpp
@@ -71,7 +71,8 @@ static void ResolveAndRegisterScopedStructTerm(FCompilerResultsLog& MessageLog,
 		Term->Context = ContextTerm;
 
 		// Read-only variables and variables in const classes are both const
-		if (BoundProperty->HasAnyPropertyFlags(CPF_BlueprintReadOnly) || (Context.IsConstFunction() && Context.NewClass->IsChildOf(StructType)))
+		if (BoundProperty->HasAnyPropertyFlags(CPF_BlueprintReadOnly) || (Context.IsConstFunction() &&
+			Context.NewClass && Context.NewClass->IsChildOf(StructType)))
 		{
 			Term->bIsConst = true;
 		}
diff --git a/Engine/Source/Editor/Kismet/Private/BPDelegateDragDropAction.cpp b/Engine/Source/Editor/Kismet/Private/BPDelegateDragDropAction.cpp
--- a/Engine/Source/Editor/Kismet/Private/BPDelegateDragDropAction.cpp
+++ b/Engine/Source/Editor/Kismet/Private/BPDelegateDragDropAction.cpp
@@ -65,7 +65,8 @@ FReply FKismetDelegateDragDropAction::DroppedOnPanel(const TSharedRef< SWidget >
 		const UBlueprint* DropOnBlueprint = FBlueprintEditorUtils::FindBlueprintForGraph(&Graph);
 		NewNodeParams.Graph = &Graph;
 		NewNodeParams.GraphPosition = GraphPosition;
-		NewNodeParams.bSelfContext = VariableSourceClass == NULL || DropOnBlueprint->SkeletonGeneratedClass->IsChildOf(VariableSourceClass);;
+		NewNodeParams.bSelfContext = VariableSourceClass == NULL ||
+			(DropOnBlueprint->SkeletonGeneratedClass && DropOnBlueprint->SkeletonGeneratedClass->IsChildOf(VariableSourceClass));
 		NewNodeParams.AnalyticCallback = AnalyticCallback;
 
 		FMenuBuilder MenuBuilder(true, NULL);
diff --git a/Engine/Source/Editor/Kismet/Private/BlueprintActionMenuBuilder.cpp b/Engine/Source/Editor/Kismet/Private/BlueprintActionMenuBuilder.cpp
--- a/Engine/Source/Editor/Kismet/Private/BlueprintActionMenuBuilder.cpp
+++ b/Engine/Source/Editor/Kismet/Private/BlueprintActionMenuBuilder.cpp
@@ -152,7 +152,7 @@ TSharedPtr<FBlueprintDragDropMenuItem> FBlueprintActionMenuItemFactory::MakeDrag
 
 			UClass const* PropertyClass = SampleProperty->GetOwnerClass();
 			checkSlow(PropertyClass != nullptr);
-			bool const bIsMemberProperty = BlueprintClass->IsChildOf(PropertyClass);
+			bool const bIsMemberProperty = BlueprintClass && BlueprintClass->IsChildOf(PropertyClass);
 
 			FText TextCategory;
 			if (Category.IsEmpty())
diff --git a/Engine/Source/Editor/Kismet/Private/BlueprintDetailsCustomization.cpp b/Engine/Source/Editor/Kismet/Private/BlueprintDetailsCustomization.cpp
--- a/Engine/Source/Editor/Kismet/Private/BlueprintDetailsCustomization.cpp
+++ b/Engine/Source/Editor/Kismet/Private/BlueprintDetailsCustomization.cpp
@@ -1527,7 +1527,8 @@ void FBlueprintVarActionDetails::PopulateCategories(SMyBlueprint* MyBlueprint, T
 	CategorySource.Add(MakeShared<FText>(UEdGraphSchema_K2::VR_DefaultCategory));
 	for (const FAdditionalBlueprintCategory& AdditionalBlueprintCategory : GetDefault<UBlueprintEditorSettings>()->AdditionalBlueprintCategories)
 	{
-		if (!AdditionalBlueprintCategory.Name.IsEmpty() && (AdditionalBlueprintCategory.ClassFilter.IsNull() || Blueprint->ParentClass->IsChildOf(AdditionalBlueprintCategory.ClassFilter.TryLoadClass<UObject>())))
+		if (!AdditionalBlueprintCategory.Name.IsEmpty() && (AdditionalBlueprintCategory.ClassFilter.IsNull() ||
+			(Blueprint->ParentClass && Blueprint->ParentClass->IsChildOf(AdditionalBlueprintCategory.ClassFilter.TryLoadClass<UObject>()))))
 		{
 			CategorySource.Add(MakeShared<FText>(AdditionalBlueprintCategory.Name));
 		}
@@ -1965,11 +1966,15 @@ EVisibility FBlueprintVarActionDetails::ExposeToCinematicsVisibility() const
 		const bool bIsFloat = VariableProperty->IsA(FFloatProperty::StaticClass());
 		const bool bIsBool = VariableProperty->IsA(FBoolProperty::StaticClass());
 		const bool bIsStr = VariableProperty->IsA(FStrProperty::StaticClass());
-		const bool bIsVectorStruct = VariableProperty->IsA(FStructProperty::StaticClass()) && CastField<FStructProperty>(VariableProperty)->Struct->GetFName() == NAME_Vector;
-		const bool bIsTransformStruct = VariableProperty->IsA(FStructProperty::StaticClass()) && CastField<FStructProperty>(VariableProperty)->Struct->GetFName() == NAME_Transform;
-		const bool bIsColorStruct = VariableProperty->IsA(FStructProperty::StaticClass()) && CastField<FStructProperty>(VariableProperty)->Struct->GetFName() == NAME_Color;
-		const bool bIsLinearColorStruct = VariableProperty->IsA(FStructProperty::StaticClass()) && CastField<FStructProperty>(VariableProperty)->Struct->GetFName() == NAME_LinearColor;
-		const bool bIsActorProperty = VariableProperty->IsA(FObjectProperty::StaticClass()) && CastField<FObjectProperty>(VariableProperty)->PropertyClass->IsChildOf(AActor::StaticClass());
+		
+		const FStructProperty* AsStructProperty = CastField<FStructProperty>(VariableProperty);
+		const bool bIsVectorStruct = AsStructProperty != nullptr && AsStructProperty->Struct->GetFName() == NAME_Vector;
+		const bool bIsTransformStruct = AsStructProperty != nullptr && AsStructProperty->Struct->GetFName() == NAME_Transform;
+		const bool bIsColorStruct = AsStructProperty != nullptr && AsStructProperty->Struct->GetFName() == NAME_Color;
+		const bool bIsLinearColorStruct = AsStructProperty != nullptr && AsStructProperty->Struct->GetFName() == NAME_LinearColor;
+		
+		const FObjectProperty* AsObjectProperty = CastField<FObjectProperty>(VariableProperty);
+		const bool bIsActorProperty = AsObjectProperty != nullptr && AsObjectProperty->PropertyClass && AsObjectProperty->PropertyClass->IsChildOf(AActor::StaticClass());
 
 		if (bIsInteger || bIsByte || bIsEnum || bIsFloat || bIsBool || bIsStr || bIsVectorStruct || bIsTransformStruct || bIsColorStruct || bIsLinearColorStruct || bIsActorProperty)
 		{
@@ -2798,7 +2803,8 @@ bool FBlueprintVarActionDetails::IsVariableInheritedByBlueprint() const
 	{
 		PropertyOwnerClass = CachedVariableProperty->GetOwnerClass();
 	}
-	return GetBlueprintObj()->SkeletonGeneratedClass->IsChildOf(PropertyOwnerClass);
+	const UClass* SkeletonGeneratedClass = GetBlueprintObj()->SkeletonGeneratedClass;
+	return SkeletonGeneratedClass && SkeletonGeneratedClass->IsChildOf(PropertyOwnerClass);
 }
 
 bool FBlueprintVarActionDetails::IsVariableDeprecated() const
diff --git a/Engine/Source/Editor/KismetCompiler/Private/KismetCompiler.cpp b/Engine/Source/Editor/KismetCompiler/Private/KismetCompiler.cpp
--- a/Engine/Source/Editor/KismetCompiler/Private/KismetCompiler.cpp
+++ b/Engine/Source/Editor/KismetCompiler/Private/KismetCompiler.cpp
@@ -1658,7 +1658,7 @@ void FKismetCompilerContext::CopyTermDefaultsToDefaultObject(UObject* DefaultObj
 						UClass* RealClass = Result->GetClass()->GetAuthoritativeClass();
 
 						// If object is compatible, write it into cdo:
-						if( RealClass->IsChildOf(AsObjectProperty->PropertyClass) )
+						if( RealClass && RealClass->IsChildOf(AsObjectProperty->PropertyClass) )
 						{
 							AsObjectProperty->SetObjectPropertyValue( AsObjectProperty->ContainerPtrToValuePtr<uint8>(DefaultObject), Result );
 							continue;
diff --git a/Engine/Source/Editor/KismetCompiler/Private/KismetCompilerMisc.cpp b/Engine/Source/Editor/KismetCompiler/Private/KismetCompilerMisc.cpp
--- a/Engine/Source/Editor/KismetCompiler/Private/KismetCompilerMisc.cpp
+++ b/Engine/Source/Editor/KismetCompiler/Private/KismetCompilerMisc.cpp
@@ -108,7 +108,7 @@ static bool DoesTypeNotMatchProperty(UEdGraphPin* SourcePin, const FEdGraphPinTy
 				InputClass = InputClass->GetAuthoritativeClass();
 
 				// It matches if it's an exact match or if the output class is more derived than the input class
-				bTypeMismatch = bSubtypeMismatch = !((OutputClass == InputClass) || (OutputClass->IsChildOf(InputClass)));
+				bTypeMismatch = bSubtypeMismatch = !((OutputClass == InputClass) || (OutputClass && OutputClass->IsChildOf(InputClass)));
 
 				if ((PinCategory == UEdGraphSchema_K2::PC_SoftClass) && (!TestProperty->IsA<FSoftClassProperty>()))
 				{
@@ -222,7 +222,7 @@ static bool DoesTypeNotMatchProperty(UEdGraphPin* SourcePin, const FEdGraphPinTy
 				OutputClass = OutputClass->GetAuthoritativeClass();
 
 				// It matches if it's an exact match or if the output class is more derived than the input class
-				bTypeMismatch = bSubtypeMismatch = !((OutputClass == InputClass) || (OutputClass->IsChildOf(InputClass)));
+				bTypeMismatch = bSubtypeMismatch = !((OutputClass == InputClass) || (OutputClass && OutputClass->IsChildOf(InputClass)));
 
 				if ((PinCategory == UEdGraphSchema_K2::PC_SoftObject) && (!TestProperty->IsA<FSoftObjectProperty>()))
 				{
@@ -2152,7 +2152,7 @@ void FNodeHandlingFunctor::ResolveAndRegisterScopedTerm(FKismetFunctionContext&
 		{
 			Term->SetVarTypeLocal(true);
 		}
-		else if (BoundProperty->HasAnyPropertyFlags(CPF_BlueprintReadOnly) || (Context.IsConstFunction() && Context.NewClass->IsChildOf(SearchScope)))
+		else if (BoundProperty->HasAnyPropertyFlags(CPF_BlueprintReadOnly) || (Context.IsConstFunction() && Context.NewClass && Context.NewClass->IsChildOf(SearchScope)))
 		{
 			// Read-only variables and variables in const classes are both const
 			Term->bIsConst = true;
diff --git a/Engine/Source/Editor/UnrealEd/Private/Kismet2/BlueprintEditorUtils.cpp b/Engine/Source/Editor/UnrealEd/Private/Kismet2/BlueprintEditorUtils.cpp
--- a/Engine/Source/Editor/UnrealEd/Private/Kismet2/BlueprintEditorUtils.cpp
+++ b/Engine/Source/Editor/UnrealEd/Private/Kismet2/BlueprintEditorUtils.cpp
@@ -411,7 +411,7 @@ void FBasePinChangeHelper::Broadcast(UBlueprint* InBlueprint, UK2Node_EditablePi
 						const bool bClassMatchesEasy = (MemberParentClass != nullptr)
 							&& ((SignatureClass != nullptr && MemberParentClass->IsChildOf(SignatureClass)) || MemberParentClass->IsChildOf(InBlueprint->GeneratedClass));
 						const bool bClassMatchesHard = !bClassMatchesEasy && CallSite->FunctionReference.IsSelfContext() && (SignatureClass == nullptr)
-							&& (CallSiteBlueprint == InBlueprint || CallSiteBlueprint->SkeletonGeneratedClass->IsChildOf(InBlueprint->SkeletonGeneratedClass));
+							&& (CallSiteBlueprint == InBlueprint || (CallSiteBlueprint->SkeletonGeneratedClass && CallSiteBlueprint->SkeletonGeneratedClass->IsChildOf(InBlueprint->SkeletonGeneratedClass)));
 
 						if (bClassMatchesEasy || bClassMatchesHard)
 						{
@@ -3001,10 +3001,13 @@ UK2Node_Event* FBlueprintEditorUtils::FindOverrideForFunction(const UBlueprint*
 		UK2Node_Event* EventNode = AllEvents[i];
 		check(EventNode);
 		if(	EventNode->bOverrideFunction == true &&
-			EventNode->EventReference.GetMemberParentClass(EventNode->GetBlueprintClassFromNode())->IsChildOf(SignatureClass) &&
 			EventNode->EventReference.GetMemberName() == SignatureName )
 		{
-			return EventNode;
+			const UClass* MemberParentClass = EventNode->EventReference.GetMemberParentClass(EventNode->GetBlueprintClassFromNode());
+			if(MemberParentClass && MemberParentClass->IsChildOf(SignatureClass))
+			{
+				return EventNode;
+			}
 		}
 	}
 
@@ -8635,10 +8638,10 @@ TSharedRef<SWidget> FBlueprintEditorUtils::ConstructBlueprintParentClassPicker(
 	for( auto BlueprintIter = Blueprints.CreateConstIterator(); (!bIsActor && !bIsAnimBlueprint) && BlueprintIter; ++BlueprintIter )
 	{
 		const UBlueprint* Blueprint = *BlueprintIter;
-		bIsActor |= Blueprint->ParentClass->IsChildOf( AActor::StaticClass() );
+		bIsActor |= Blueprint->ParentClass && Blueprint->ParentClass->IsChildOf( AActor::StaticClass() );
 		bIsAnimBlueprint |= Blueprint->IsA(UAnimBlueprint::StaticClass());
-		bIsLevelScriptActor |= Blueprint->ParentClass->IsChildOf( ALevelScriptActor::StaticClass() );
-		bIsComponentBlueprint |= Blueprint->ParentClass->IsChildOf( UActorComponent::StaticClass() );
+		bIsLevelScriptActor |= Blueprint->ParentClass && Blueprint->ParentClass->IsChildOf( ALevelScriptActor::StaticClass() );
+		bIsComponentBlueprint |= Blueprint->ParentClass && Blueprint->ParentClass->IsChildOf( UActorComponent::StaticClass() );
 		bIsEditorOnlyBlueprint |= IsEditorUtilityBlueprint(Blueprint);
 		bIsWidgetBlueprint = Blueprint->IsA(UBaseWidgetBlueprint::StaticClass());
 		if(Blueprint->GeneratedClass)
diff --git a/Engine/Source/Editor/UnrealEd/Private/Kismet2/ComponentEditorUtils.cpp b/Engine/Source/Editor/UnrealEd/Private/Kismet2/ComponentEditorUtils.cpp
--- a/Engine/Source/Editor/UnrealEd/Private/Kismet2/ComponentEditorUtils.cpp
+++ b/Engine/Source/Editor/UnrealEd/Private/Kismet2/ComponentEditorUtils.cpp
@@ -102,12 +102,12 @@ protected:
 	virtual bool CanCreateClass(UClass* ObjectClass, bool& bOmitSubObjs) const override
 	{
 		// Allow actor component types to be created
-		bool bCanCreate = ObjectClass->IsChildOf(UActorComponent::StaticClass());
+		bool bCanCreate = ObjectClass && ObjectClass->IsChildOf(UActorComponent::StaticClass());
 
 		if (!bCanCreate)
 		{
 			// Also allow Blueprint-able actor types to pass, in order to enable proper creation of actor component types as subobjects. The actor instance will be discarded after processing.
-			bCanCreate = ObjectClass->IsChildOf(AActor::StaticClass()) && FKismetEditorUtilities::CanCreateBlueprintOfClass(ObjectClass);
+			bCanCreate = ObjectClass && ObjectClass->IsChildOf(AActor::StaticClass()) && FKismetEditorUtilities::CanCreateBlueprintOfClass(ObjectClass);
 		}
 		else
 		{
diff --git a/Engine/Source/Editor/UnrealEd/Private/Kismet2/Kismet2.cpp b/Engine/Source/Editor/UnrealEd/Private/Kismet2/Kismet2.cpp
--- a/Engine/Source/Editor/UnrealEd/Private/Kismet2/Kismet2.cpp
+++ b/Engine/Source/Editor/UnrealEd/Private/Kismet2/Kismet2.cpp
@@ -550,7 +550,7 @@ UBlueprint* FKismetEditorUtilities::CreateBlueprint(UClass* ParentClass, UObject
 		for (TMultiMap<void*, FKismetEditorUtilities::FOnBlueprintCreatedData>::TIterator DataIt(OnBlueprintCreatedCallbacks); DataIt; ++DataIt)
 		{
 			FOnBlueprintCreatedData Data = DataIt.Value();
-			if (NewBP->GeneratedClass->IsChildOf(Data.TargetClass))
+			if (NewBP->GeneratedClass && NewBP->GeneratedClass->IsChildOf(Data.TargetClass))
 			{
 				FKismetEditorUtilities::FOnBlueprintCreated BlueprintCreatedDelegate = Data.OnBlueprintCreated;
 				BlueprintCreatedDelegate.Execute(NewBP);
