Fixing bug where deprecated Event Dispatchers on components do not warn on their attached actors

diff --git a/Engine/Source/Editor/BlueprintGraph/Classes/K2Node_ComponentBoundEvent.h b/Engine/Source/Editor/BlueprintGraph/Classes/K2Node_ComponentBoundEvent.h
--- a/Engine/Source/Editor/BlueprintGraph/Classes/K2Node_ComponentBoundEvent.h
+++ b/Engine/Source/Editor/BlueprintGraph/Classes/K2Node_ComponentBoundEvent.h
@@ -42,6 +42,8 @@ class UK2Node_ComponentBoundEvent : public UK2Node_Event
 	virtual FText GetTooltipText() const override;
 	virtual FString GetDocumentationLink() const override;
 	virtual FString GetDocumentationExcerptName() const override;
+	virtual bool HasDeprecatedReference() const override;
+	virtual FEdGraphNodeDeprecationResponse GetDeprecationResponse(EEdGraphNodeDeprecationType DeprecationType) const override;
 	//~ End UEdGraphNode Interface
 
 	//~ Begin K2Node Interface
diff --git a/Engine/Source/Editor/BlueprintGraph/Private/K2Node_ComponentBoundEvent.cpp b/Engine/Source/Editor/BlueprintGraph/Private/K2Node_ComponentBoundEvent.cpp
--- a/Engine/Source/Editor/BlueprintGraph/Private/K2Node_ComponentBoundEvent.cpp
+++ b/Engine/Source/Editor/BlueprintGraph/Private/K2Node_ComponentBoundEvent.cpp
@@ -9,6 +9,12 @@
 
 #define LOCTEXT_NAMESPACE "K2Node"
 
+static TAutoConsoleVariable<bool> CVarBPEnableDeprecatedWarningForComponentDelegateNodes(
+	TEXT("bp.EnableDeprecatedWarningForComponentDelegateNodes"),
+	true,
+	TEXT("Show Deprecated warning for component delegate event nodes"),
+	ECVF_Cheat);
+
 // @TODO_BH: Remove the CVar for validity checking when we can get all the errors sorted out
 namespace PinValidityCheck
 {
@@ -145,6 +151,35 @@ bool UK2Node_ComponentBoundEvent::IsDelegateValid() const
 		&& (GetTargetDelegateProperty() || FMemberReference::FindRemappedField<FMulticastDelegateProperty>(DelegateOwnerClass, DelegatePropertyName));
 }
 
+bool UK2Node_ComponentBoundEvent::HasDeprecatedReference() const
+{
+	if (CVarBPEnableDeprecatedWarningForComponentDelegateNodes.GetValueOnGameThread())
+	{
+		if (const FMulticastDelegateProperty* DelegateProperty = GetTargetDelegateProperty())
+		{
+			return DelegateProperty->HasAnyPropertyFlags(EPropertyFlags::CPF_Deprecated);	
+		}
+	}
+	return false;
+}
+
+FEdGraphNodeDeprecationResponse UK2Node_ComponentBoundEvent::GetDeprecationResponse(EEdGraphNodeDeprecationType DeprecationType) const
+{
+	FEdGraphNodeDeprecationResponse Response = Super::GetDeprecationResponse(DeprecationType);
+	if (DeprecationType == EEdGraphNodeDeprecationType::NodeHasDeprecatedReference)
+	{
+		const UFunction* Function = EventReference.ResolveMember<UFunction>(GetBlueprintClassFromNode());
+		if (ensureMsgf(Function != nullptr, TEXT("This node should not be able to report having a deprecated reference if the event override cannot be resolved.")))
+		{
+			Response.MessageType = EEdGraphNodeDeprecationMessageType::Warning;
+			const FText DetailedMessage = FText::FromString(Function->GetMetaData(FBlueprintMetadata::MD_DeprecationMessage));
+			Response.MessageText = FBlueprintEditorUtils::GetDeprecatedMemberUsageNodeWarning(DelegatePropertyDisplayName, DetailedMessage);
+		}
+	}
+
+	return Response;
+}
+
 bool UK2Node_ComponentBoundEvent::IsUsedByAuthorityOnlyDelegate() const
 {
 	FMulticastDelegateProperty* TargetDelegateProp = GetTargetDelegateProperty();
