Fix to allow users to coerce a BlueprintImplementableEvent with no return value to treat the override as a function rather than an event

#3078

diff --git a/Engine/Source/Editor/BlueprintGraph/Classes/EdGraphSchema_K2.h b/Engine/Source/Editor/BlueprintGraph/Classes/EdGraphSchema_K2.h
--- a/Engine/Source/Editor/BlueprintGraph/Classes/EdGraphSchema_K2.h
+++ b/Engine/Source/Editor/BlueprintGraph/Classes/EdGraphSchema_K2.h
@@ -132,6 +132,9 @@ public:
 	// [FunctionMetadata] Indicates that a particular function parameter is for internal use only, which means it will be both hidden and not connectible.
 	static const FName MD_InternalUseParam;
 
+	// [FunctionMetadata] Indicates that the function should appear as blueprint function even if it doesn't return a value.
+	static const FName MD_ForceAsFunction;
+
 	// [FunctionMetadata] Indicates that the function should be ignored when considered for blueprint type promotion
 	static const FName MD_IgnoreTypePromotion;
 
diff --git a/Engine/Source/Editor/BlueprintGraph/Private/EdGraphSchema_K2.cpp b/Engine/Source/Editor/BlueprintGraph/Private/EdGraphSchema_K2.cpp
--- a/Engine/Source/Editor/BlueprintGraph/Private/EdGraphSchema_K2.cpp
+++ b/Engine/Source/Editor/BlueprintGraph/Private/EdGraphSchema_K2.cpp
@@ -157,6 +157,7 @@ const FName FBlueprintMetadata::MD_CompactNodeTitle(TEXT("CompactNodeTitle"));
 const FName FBlueprintMetadata::MD_DisplayName(TEXT("DisplayName"));
 const FName FBlueprintMetadata::MD_ReturnDisplayName(TEXT("ReturnDisplayName"));
 const FName FBlueprintMetadata::MD_InternalUseParam(TEXT("InternalUseParam"));
+const FName FBlueprintMetadata::MD_ForceAsFunction(TEXT("ForceAsFunction"));
 const FName FBlueprintMetadata::MD_IgnoreTypePromotion(TEXT("IgnoreTypePromotion"));
 
 const FName FBlueprintMetadata::MD_PropertyGetFunction(TEXT("BlueprintGetter"));
@@ -979,6 +980,12 @@ bool UEdGraphSchema_K2::FunctionCanBePlacedAsEvent(const UFunction* InFunction)
 		return false;
 	}
 
+	// Check if meta data has been set to force this to appear as blueprint function even if it doesn't return a value.
+	if (InFunction->HasAllFunctionFlags(FUNC_BlueprintEvent) && InFunction->HasMetaData(FBlueprintMetadata::MD_ForceAsFunction))
+	{
+		return false;
+	}
+
 	// Then look to see if we have any output, return, or reference params
 	return !HasFunctionAnyOutputParameter(InFunction);
 }
diff --git a/Engine/Source/Runtime/CoreUObject/Public/UObject/ObjectMacros.h b/Engine/Source/Runtime/CoreUObject/Public/UObject/ObjectMacros.h
--- a/Engine/Source/Runtime/CoreUObject/Public/UObject/ObjectMacros.h
+++ b/Engine/Source/Runtime/CoreUObject/Public/UObject/ObjectMacros.h
@@ -953,6 +953,9 @@ namespace UF
 		/// [FunctionMetadata] Indicates that a particular function parameter is for internal use only, which means it will be both hidden and not connectible.
 		InternalUseParam, 
 
+		// [FunctionMetadata] Indicates that the function should appear as blueprint function even if it doesn't return a value.
+		ForceAsFunction, 
+
 		/// [FunctionMetadata] Indicates that the function should be ignored when considered for blueprint type promotion
 		IgnoreTypePromotion,
 	};
