Added additional check to ensure that we're using the correct interface function.

In order to a stack overrun due to erroneous bytecode, we should issue an error if we're using the incorrect interface function. This is a temporary fix until we find a more permanent solution.

diff --git a/Engine/Source/Editor/BlueprintGraph/Private/CallFunctionHandler.cpp b/Engine/Source/Editor/BlueprintGraph/Private/CallFunctionHandler.cpp
--- a/Engine/Source/Editor/BlueprintGraph/Private/CallFunctionHandler.cpp
+++ b/Engine/Source/Editor/BlueprintGraph/Private/CallFunctionHandler.cpp
@@ -969,16 +969,44 @@ void FKCHandler_CallFunction::Compile(FKismetFunctionContext& Context, UEdGraphN
 
 	// Validate the self pin again if it is disconnected, because pruning isolated nodes could have caused an invalid target
 	UEdGraphPin* SelfPin = CompilerContext.GetSchema()->FindSelfPin(*Node, EGPD_Input);
-	if (SelfPin && (SelfPin->LinkedTo.Num() == 0))
+	if (SelfPin)
 	{
-		FEdGraphPinType SelfType;
-		SelfType.PinCategory = UEdGraphSchema_K2::PC_Object;
-		SelfType.PinSubCategory = UEdGraphSchema_K2::PSC_Self;
+		if ((SelfPin->LinkedTo.Num() == 0))
+		{
+			FEdGraphPinType SelfType;
+			SelfType.PinCategory = UEdGraphSchema_K2::PC_Object;
+			SelfType.PinSubCategory = UEdGraphSchema_K2::PSC_Self;
 
-		if (!CompilerContext.GetSchema()->ArePinTypesCompatible(SelfType, SelfPin->PinType, Context.NewClass) && (SelfPin->DefaultObject == NULL))
+			if (!CompilerContext.GetSchema()->ArePinTypesCompatible(SelfType, SelfPin->PinType, Context.NewClass) && (SelfPin->DefaultObject == NULL))
+			{
+				CompilerContext.MessageLog.Error(*NSLOCTEXT("KismetCompiler", "PinMustHaveConnectionPruned_Error", "Pin @@ must have a connection.  Self pins cannot be connected to nodes that are culled.").ToString(), SelfPin);
+			}
+		}
+		else
 		{
-			CompilerContext.MessageLog.Error(*NSLOCTEXT("KismetCompiler", "PinMustHaveConnectionPruned_Error", "Pin @@ must have a connection.  Self pins cannot be connected to nodes that are culled.").ToString(), SelfPin);
+			if (UK2Node_CallFunction* CallFunctionNode = Cast<UK2Node_CallFunction>(Node))
+			{
+				for (int32 i = 0; i < SelfPin->LinkedTo.Num(); ++i)
+				{
+					UEdGraphPin* LinkedPin = SelfPin->LinkedTo[i];
+					
+					bool bHasMismatchedInterfaceConnection =
+						LinkedPin &&
+						(SelfPin->PinType.PinCategory == UEdGraphSchema_K2::PC_Interface) && 
+						(LinkedPin->PinType.PinCategory == UEdGraphSchema_K2::PC_Interface) &&
+						(CallFunctionNode->FunctionReference.GetMemberParentClass() != LinkedPin->PinType.PinSubCategoryObject);
+
+					if (bHasMismatchedInterfaceConnection)
+					{
+						CompilerContext.MessageLog.Error(*NSLOCTEXT("KismetCompiler",
+																	"MismatchedInterfaces_Error", 
+																	"Node @@ can't be used with the linked interface pin. Ensure that the interface's function is used here.").ToString(),
+																	Node);
+					}
+				}
+			}
 		}
+		
 	}
 
 	// Make sure the function node is valid to call
