Added BP action DB filtering

Applies optional filtering to the BP context menu and related UIs.

Primarily this is implemented at the BlueprintActionDatabase/Registrar level as permissions are mostly static. One additional filter has been added to BlueprintActionFilter to handle filtering out items from assets that cannot be referenced from *this* blueprint.

Added type filtering to the variable type tree to prevent creating disallowed variable types.

Also added some auditing console commands: bp.AuditThreadSafeFunctions and bp.AuditFunctionCallsForBlueprint
Corrected FBlueprintEditorUtils::HasFunctionBlueprintThreadSafeMetaData to not track FUNC_BlueprintEvent flags for native functions.

diff --git a/Engine/Source/Editor/AnimGraph/Private/AnimGraphNode_AssetPlayerBase.cpp b/Engine/Source/Editor/AnimGraph/Private/AnimGraphNode_AssetPlayerBase.cpp
--- a/Engine/Source/Editor/AnimGraph/Private/AnimGraphNode_AssetPlayerBase.cpp
+++ b/Engine/Source/Editor/AnimGraph/Private/AnimGraphNode_AssetPlayerBase.cpp
@@ -5,6 +5,7 @@
 #include "Animation/AnimComposite.h"
 #include "BlueprintActionFilter.h"
 #include "BlueprintNodeSpawner.h"
+#include "BlueprintAssetNodeSpawner.h"
 #include "UObject/UObjectIterator.h"
 #include "Animation/AnimLayerInterface.h"
 #include "IAnimBlueprintGeneratedClassCompiledData.h"
@@ -256,7 +257,7 @@ void UAnimGraphNode_AssetPlayerBase::GetMenuActionsHelper(
 			InSetupNewNodeFromAssetFunction(InNewNode, bInIsTemplateNode, InAssetData);
 		};
 
-		UBlueprintNodeSpawner* NodeSpawner = UBlueprintNodeSpawner::Create(InNodeClass.Get());
+		UBlueprintAssetNodeSpawner* NodeSpawner = UBlueprintAssetNodeSpawner::Create(InNodeClass.Get(), InAssetData);
 		NodeSpawner->CustomizeNodeDelegate = UBlueprintNodeSpawner::FCustomizeNodeDelegate::CreateLambda(AssetSetup);
 		NodeSpawner->DefaultMenuSignature.MenuName = InMenuNameFunction(InAssetData, InAssetData.GetClass());
 		NodeSpawner->DefaultMenuSignature.Category = InMenuCategoryFunction != nullptr ? InMenuCategoryFunction(InAssetData) : FText::GetEmpty();
diff --git a/Engine/Source/Editor/AnimGraph/Private/AnimGraphNode_LinkedAnimGraph.cpp b/Engine/Source/Editor/AnimGraph/Private/AnimGraphNode_LinkedAnimGraph.cpp
--- a/Engine/Source/Editor/AnimGraph/Private/AnimGraphNode_LinkedAnimGraph.cpp
+++ b/Engine/Source/Editor/AnimGraph/Private/AnimGraphNode_LinkedAnimGraph.cpp
@@ -160,6 +160,12 @@ bool UAnimGraphNode_LinkedAnimGraph::IsActionFilteredOut(class FBlueprintActionF
 		{
 			if (UAnimBlueprint* AnimBlueprint = Cast<UAnimBlueprint>(Blueprint))
 			{
+				if (AnimBlueprint->BlueprintType == BPTYPE_Interface)
+				{
+					bIsFilteredOut = true;
+					break;
+				}
+
 				if(!AnimBlueprint->IsCompatibleByAssetString(SkeletonName, bIsTemplateAnimBlueprint, bIsInterfaceBlueprint))
 				{
 					bIsFilteredOut = true;
diff --git a/Engine/Source/Editor/BlueprintGraph/BlueprintGraph.Build.cs b/Engine/Source/Editor/BlueprintGraph/BlueprintGraph.Build.cs
--- a/Engine/Source/Editor/BlueprintGraph/BlueprintGraph.Build.cs
+++ b/Engine/Source/Editor/BlueprintGraph/BlueprintGraph.Build.cs
@@ -41,6 +41,7 @@ public class BlueprintGraph : ModuleRules
                 "KismetWidgets",
                 "PropertyEditor",
 				"ToolMenus",
+				"AssetTools",
 			}
 		);
 
diff --git a/Engine/Source/Editor/BlueprintGraph/Classes/BlueprintAssetNodeSpawner.h b/Engine/Source/Editor/BlueprintGraph/Classes/BlueprintAssetNodeSpawner.h
new file mode 100644
--- /dev/null
+++ b/Engine/Source/Editor/BlueprintGraph/Classes/BlueprintAssetNodeSpawner.h
@@ -0,0 +1,50 @@
+// Copyright Epic Games, Inc. All Rights Reserved.
+
+#pragma once
+
+#include "CoreMinimal.h"
+#include "UObject/ObjectMacros.h"
+#include "UObject/Class.h"
+#include "Templates/SubclassOf.h"
+#include "BlueprintNodeSignature.h"
+#include "BlueprintNodeSpawner.h"
+#include "AssetRegistry/AssetData.h"
+#include "BlueprintAssetNodeSpawner.generated.h"
+
+class UEdGraph;
+
+/**
+ * Takes care of spawning various asset related nodes (nodes associated with 
+ * an asset).
+ */
+UCLASS(Transient)
+class BLUEPRINTGRAPH_API UBlueprintAssetNodeSpawner : public UBlueprintNodeSpawner
+{
+	GENERATED_BODY()
+
+public:
+	/**
+	 * Creates a new UBlueprintAssetNodeSpawner for the supplied asset.
+	 * Does not do any compatibility checking to ensure that the asset is
+	 * viable for blueprint use.
+	 * 
+	 * @param  InNodeClass			The type of node you want the spawner to create.
+	 * @param  InAssetData			The asset you want assigned to new nodes.
+	 * @param  InOuter				Optional outer for the new spawner (if left null, the transient package will be used).
+	 * @param  InPostSpawnDelegate	A delegate to perform specialized node setup post-spawn.
+	 * @return A newly allocated instance of this class.
+	 */
+	static UBlueprintAssetNodeSpawner* Create(TSubclassOf<UEdGraphNode> const InNodeClass, const FAssetData& InAssetData, UObject* InOuter = nullptr, FCustomizeNodeDelegate InPostSpawnDelegate = FCustomizeNodeDelegate());
+
+	/**
+	 * Retrieves the asset that this assigns to spawned nodes.
+	 *
+	 * @return The asset that this class was initialized with.
+	 */
+	const FAssetData& GetAssetData() const;
+
+private:
+	/** The asset to configure new nodes with. */
+	UPROPERTY()
+	FAssetData AssetData;
+};
diff --git a/Engine/Source/Editor/BlueprintGraph/Private/BlueprintActionDatabase.cpp b/Engine/Source/Editor/BlueprintGraph/Private/BlueprintActionDatabase.cpp
--- a/Engine/Source/Editor/BlueprintGraph/Private/BlueprintActionDatabase.cpp
+++ b/Engine/Source/Editor/BlueprintGraph/Private/BlueprintActionDatabase.cpp
@@ -67,6 +67,7 @@
 #include "EdGraph/EdGraphNode_Documentation.h"
 
 #include "BlueprintTypePromotion.h"
+#include "PropertyEditorPermissionList.h"
 
 #define LOCTEXT_NAMESPACE "BlueprintActionDatabase"
 
@@ -555,7 +556,13 @@ static bool BlueprintActionDatabaseImpl::IsPropertyBlueprintVisible(FProperty co
 	bool const bIsDelegate = Property->IsA(FMulticastDelegateProperty::StaticClass());
 	bool const bIsAssignableOrCallable = Property->HasAnyPropertyFlags(CPF_BlueprintAssignable | CPF_BlueprintCallable);
 
-	return !Property->HasAnyPropertyFlags(CPF_Parm) && (bIsAccessible || (bIsDelegate && bIsAssignableOrCallable));
+	bool bVisible = !Property->HasAnyPropertyFlags(CPF_Parm) && (bIsAccessible || (bIsDelegate && bIsAssignableOrCallable));
+	if (bVisible)
+	{
+		bVisible = FPropertyEditorPermissionList::Get().DoesPropertyPassFilter(Property->GetOwnerStruct(), Property->GetFName());
+	}
+
+	return bVisible;
 }
 
 //------------------------------------------------------------------------------
@@ -637,6 +644,12 @@ static void BlueprintActionDatabaseImpl::AddClassFunctionActions(UClass const* c
 			continue;
 		}
 
+		// Apply general filtering for functions
+		if(!FBlueprintActionDatabase::IsFunctionAllowed(Function, FBlueprintActionDatabase::EPermissionsContext::Node))
+		{
+			continue;
+		}
+
 		bool const bIsBpInterfaceFunc = BlueprintActionDatabaseImpl::IsBlueprintInterfaceFunction(Function);
 		if (UEdGraphSchema_K2::FunctionCanBePlacedAsEvent(Function) && !bIsBpInterfaceFunc)
 		{
@@ -774,7 +787,7 @@ static void BlueprintActionDatabaseImpl::AddClassCastActions(UClass* Class, FAct
 	check(Class);
 
 	UEdGraphSchema_K2 const* K2Schema = GetDefault<UEdGraphSchema_K2>();
-	bool bIsCastPermitted  = UEdGraphSchema_K2::IsAllowableBlueprintVariableType(Class);
+	bool bIsCastPermitted  = UEdGraphSchema_K2::IsAllowableBlueprintVariableType(Class) && FBlueprintActionDatabase::IsClassAllowed(Class, FBlueprintActionDatabase::EPermissionsContext::Node);
 
 	if (bIsCastPermitted)
 	{
@@ -1355,11 +1368,15 @@ void FBlueprintActionDatabase::RefreshAll()
 		UClass* const Class = (*ClassIt);
 		RefreshClassActions(Class);
 	}
-	// this handles creating entries for skeletons that were loaded before the database was alive:
-	for( TObjectIterator<USkeleton> SkeletonIt; SkeletonIt; ++SkeletonIt )
+
+	if(IsClassAllowed(USkeleton::StaticClass(), EPermissionsContext::Asset))
 	{
-		FActionList& ClassActionList = ActionRegistry.FindOrAdd(*SkeletonIt);
-		BlueprintActionDatabaseImpl::AddSkeletonActions(**SkeletonIt, ClassActionList);
+		// this handles creating entries for skeletons that were loaded before the database was alive:
+		for( TObjectIterator<USkeleton> SkeletonIt; SkeletonIt; ++SkeletonIt )
+		{
+			FActionList& ClassActionList = ActionRegistry.FindOrAdd(*SkeletonIt);
+			BlueprintActionDatabaseImpl::AddSkeletonActions(**SkeletonIt, ClassActionList);
+		}
 	}
 
 	FComponentTypeRegistry::Get().SubscribeToComponentList(ComponentTypes).RemoveAll(this);
@@ -1393,6 +1410,12 @@ void FBlueprintActionDatabase::RefreshWorlds()
 //------------------------------------------------------------------------------
 void FBlueprintActionDatabase::RefreshClassActions(UClass* const Class)
 {
+	// Early out if the class is filtered
+	if (!IsClassAllowed(Class, EPermissionsContext::Asset))
+	{
+		return;
+	}
+
 	using namespace BlueprintActionDatabaseImpl;
 	check(Class != nullptr);
 
@@ -1427,10 +1450,12 @@ void FBlueprintActionDatabase::RefreshClassActions(UClass* const Class)
 	// own actions (presumably ones that would spawn that node)...
 	else if (Class->IsChildOf<UEdGraphNode>())
 	{
-		FActionList& ClassActionList = ActionRegistry.FindOrAdd(Class);
-		if (!bIsInitializing)
 		{
-			ClassActionList.Empty();
+			FActionList& ClassActionList = ActionRegistry.FindOrAdd(Class);
+			if (!bIsInitializing)
+			{
+				ClassActionList.Empty();
+			}
 		}
 
 		FBlueprintActionDatabaseRegistrar Registrar(ActionRegistry, UnloadedActionRegistry, ActionPrimingQueue, Class);
@@ -1447,6 +1472,16 @@ void FBlueprintActionDatabase::RefreshClassActions(UClass* const Class)
 		// normally when sifting through fields on all known classes)		
 		GetNodeSpecificActions(Class, Registrar);
 		// don't worry, the registrar marks new actions for priming
+
+		// Filter out actions by node class
+		if(HasClassFiltering())
+		{
+			FActionList& ClassActionList = ActionRegistry.FindOrAdd(Class);
+			ClassActionList.RemoveAllSwap([this](UBlueprintNodeSpawner* InAction)
+			{
+				return !IsClassAllowed(InAction->NodeClass.Get(), EPermissionsContext::Node);
+			});
+		}
 	}
 	else if (Class->IsChildOf<UBlueprint>())
 	{
@@ -1473,6 +1508,15 @@ void FBlueprintActionDatabase::RefreshClassActions(UClass* const Class)
 		if (ClassActionList.Num() > 0)
 		{
 			ActionPrimingQueue.Add(Class, 0);
+
+			if (HasClassFiltering())
+			{
+				// Filter out actions by node class
+				ClassActionList.RemoveAllSwap([this](UBlueprintNodeSpawner* InAction)
+				{
+					return !IsClassAllowed(InAction->NodeClass.Get(), EPermissionsContext::Node);
+				});
+			}
 		}
 		else
 		{
@@ -1514,6 +1558,11 @@ void FBlueprintActionDatabase::RefreshAssetActions(UObject* const AssetObject)
 		return;
 	}
 
+	if (!IsClassAllowed(AssetObject->GetClass(), EPermissionsContext::Asset))
+	{
+		return;
+	}
+
 	if(const USkeleton* Skeleton = Cast<USkeleton>(AssetObject))
 	{
 		AddSkeletonActions(*Skeleton, AssetActionList);
@@ -1613,6 +1662,15 @@ void FBlueprintActionDatabase::RefreshAssetActions(UObject* const AssetObject)
 	{
 		EntryRefreshDelegate.Broadcast(AssetObject);
 	}
+
+	if (HasClassFiltering())
+	{
+		// Filter out actions by node class
+		AssetActionList.RemoveAllSwap([this](UBlueprintNodeSpawner* InAction)
+		{
+			return !IsClassAllowed(InAction->NodeClass.Get(), EPermissionsContext::Node);
+		});
+	}
 }
 
 //------------------------------------------------------------------------------
@@ -1623,11 +1681,25 @@ void FBlueprintActionDatabase::RefreshComponentActions()
 	ClassActionList.Empty(ComponentTypes->Num());
 	for (const FComponentTypeEntry& ComponentType : *ComponentTypes)
 	{
+		if (!IsClassAllowed(ComponentType.ComponentClass, EPermissionsContext::Node))
+		{
+			continue;
+		}
+
 		if (UBlueprintComponentNodeSpawner* NodeSpawner = UBlueprintComponentNodeSpawner::Create(ComponentType))
 		{
 			ClassActionList.Add(NodeSpawner);
 		}
 	}
+
+	if (HasClassFiltering())
+	{
+		// Filter out actions by node class
+		ClassActionList.RemoveAllSwap([this](UBlueprintNodeSpawner* InAction)
+		{
+			return !IsClassAllowed(InAction->NodeClass.Get(), EPermissionsContext::Node);
+		});
+	}
 }
 
 //------------------------------------------------------------------------------
@@ -1660,16 +1732,16 @@ bool FBlueprintActionDatabase::ClearAssetActions(const FObjectKey& AssetObjectKe
 
 	if (UObject* AssetObject = AssetObjectKey.ResolveObjectPtr())
 	{
-	if (UBlueprint* BlueprintAsset = Cast<UBlueprint>(AssetObject))
-	{
-		BlueprintAsset->OnChanged().RemoveAll(this);
-		BlueprintAsset->OnCompiled().RemoveAll(this);
-	}
+		if (UBlueprint* BlueprintAsset = Cast<UBlueprint>(AssetObject))
+		{
+			BlueprintAsset->OnChanged().RemoveAll(this);
+			BlueprintAsset->OnCompiled().RemoveAll(this);
+		}
 
-	if (bHasEntry && (ActionList->Num() > 0) && !BlueprintActionDatabaseImpl::bIsInitializing)
-	{
-		EntryRemovedDelegate.Broadcast(AssetObject);
-	}
+		if (bHasEntry && (ActionList->Num() > 0) && !BlueprintActionDatabaseImpl::bIsInitializing)
+		{
+			EntryRemovedDelegate.Broadcast(AssetObject);
+		}
 	}
 	
 	return bHasEntry;
@@ -1749,4 +1821,271 @@ void FBlueprintActionDatabase::OnBlueprintChanged(UBlueprint* InBlueprint)
 	}
 }
 
+bool FBlueprintActionDatabase::IsClassAllowed(UClass const* InClass, EPermissionsContext InContext)
+{
+	if (InClass == nullptr)
+	{
+		return false;
+	}
+
+	UBlueprintEditorSettings* BlueprintEditorSettings = GetMutableDefault<UBlueprintEditorSettings>();
+	
+	switch(InContext)
+	{
+	case EPermissionsContext::Property:
+	case EPermissionsContext::Node:
+	case EPermissionsContext::Asset:
+		if(BlueprintEditorSettings->HasClassFiltering())
+		{
+			return BlueprintEditorSettings->IsClassAllowed(InClass);
+		}
+		break;
+	case EPermissionsContext::Pin:
+		if(BlueprintEditorSettings->HasClassOnPinFiltering())
+		{
+			return BlueprintEditorSettings->IsClassAllowedOnPin(InClass);
+		}
+		break;
+	default:
+		break;
+	}
+
+	return true;
+}
+
+bool FBlueprintActionDatabase::IsClassAllowed(const FTopLevelAssetPath& InClassPath, EPermissionsContext InContext)
+{
+	if(!InClassPath.IsValid())
+	{
+		return false;
+	}
+
+	UBlueprintEditorSettings* BlueprintEditorSettings = GetMutableDefault<UBlueprintEditorSettings>();
+	
+	switch(InContext)
+	{
+	case EPermissionsContext::Property:
+	case EPermissionsContext::Node:
+	case EPermissionsContext::Asset:
+		if (BlueprintEditorSettings->HasClassPathFiltering())
+		{
+			return BlueprintEditorSettings->IsClassPathAllowed(InClassPath);
+		}
+		break;
+	case EPermissionsContext::Pin:
+		if(BlueprintEditorSettings->HasClassPathOnPinFiltering())
+		{
+			return BlueprintEditorSettings->IsClassPathAllowedOnPin(InClassPath);
+		}
+		break;
+	default:
+		break;
+	}
+
+	return true;
+}
+
+bool FBlueprintActionDatabase::HasClassFiltering()
+{
+	UBlueprintEditorSettings* BlueprintEditorSettings = GetMutableDefault<UBlueprintEditorSettings>();
+	return	BlueprintEditorSettings->HasClassFiltering() || 
+			BlueprintEditorSettings->HasClassOnPinFiltering() || 
+			BlueprintEditorSettings->HasClassPathFiltering() ||
+			BlueprintEditorSettings->HasClassPathOnPinFiltering();
+}
+
+bool FBlueprintActionDatabase::IsFieldAllowed(UField const* InField, EPermissionsContext InContext)
+{
+	if (UFunction const* Function = Cast<UFunction>(InField))
+	{
+		return IsFunctionAllowed(Function, InContext);
+	}
+	else if (UEnum const* Enum = Cast<UEnum>(InField))
+	{
+		return IsEnumAllowed(Enum, InContext);
+	}
+	else if (UScriptStruct const* ScriptStruct = Cast<UScriptStruct>(InField))
+	{
+		return IsStructAllowed(ScriptStruct, InContext);
+	}
+	else if (UClass const* Class = Cast<UClass>(InField))
+	{
+		return IsClassAllowed(Class, InContext);
+	}
+
+	return true;
+}
+
+bool FBlueprintActionDatabase::IsFunctionAllowed(UFunction const* InFunction, EPermissionsContext InContext)
+{
+	UBlueprintEditorSettings* BlueprintEditorSettings = GetMutableDefault<UBlueprintEditorSettings>();
+	const FPathPermissionList& FunctionPermissions = BlueprintEditorSettings->GetFunctionPermissions();
+	const FPathPermissionList& EnumPermissions = BlueprintEditorSettings->GetEnumPermissions();
+	const FPathPermissionList& StructPermissions = BlueprintEditorSettings->GetStructPermissions();
+	const FNamePermissionList& PinCategoryPermissions = BlueprintEditorSettings->GetPinCategoryPermissions();
+
+	// Apply general filtering for functions
+	if (FunctionPermissions.HasFiltering())
+	{
+		TStringBuilder<256> ResultBuilder;
+		InFunction->GetPathName(nullptr, ResultBuilder);
+		if (!FunctionPermissions.PassesFilter(ResultBuilder.ToView()))
+		{
+			return false;
+		}
+	}
+
+	if (UClass* Class = InFunction->GetOuterUClass())
+	{
+		if (!IsClassAllowed(Class, EPermissionsContext::Asset))
+		{
+			UE_LOG(LogBlueprint, Warning, TEXT("Function %s was filtered because its class (%s) was filtered"), *InFunction->GetPathName(), *Class->GetName());
+			return false;
+		}
+	}
+
+	// If we have filtering for other fields we need to check function parameters that might reference them
+	if (EnumPermissions.HasFiltering() || StructPermissions.HasFiltering() || PinCategoryPermissions.HasFiltering())
+	{
+		for (TFieldIterator<FProperty> PropertyIt(InFunction); PropertyIt; ++PropertyIt)
+		{
+			FProperty* Property = *PropertyIt;
+			if (FStructProperty* StructProperty = CastField<FStructProperty>(Property))
+			{
+				if (StructProperty->Struct)
+				{
+					if (!IsStructAllowed(StructProperty->Struct, EPermissionsContext::Pin))
+					{
+						UE_LOG(LogBlueprint, Warning, TEXT("Function %s was filtered because one of its parameters (struct %s) was filtered"), *InFunction->GetPathName(), *Property->GetName(), *StructProperty->Struct->GetPathName());
+						return false;
+					}
+				}
+			}
+			else if (FEnumProperty* EnumProperty = CastField<FEnumProperty>(Property))
+			{
+				if (UEnum* Enum = EnumProperty->GetEnum())
+				{
+					if (!IsEnumAllowed(Enum, EPermissionsContext::Pin))
+					{
+						UE_LOG(LogBlueprint, Warning, TEXT("Function %s was filtered because one of its parameters %s (enum %s) was filtered"), *InFunction->GetPathName(), *Property->GetName(), *Enum->GetPathName());
+						return false;
+					}
+				}
+			}
+			else
+			{
+				FEdGraphPinType PinType;
+				if(GetDefault<UEdGraphSchema_K2>()->ConvertPropertyToPinType(Property, PinType))
+				{
+					if(!IsPinTypeAllowed(PinType))
+					{
+						UE_LOG(LogBlueprint, Warning, TEXT("Function %s was filtered because one of its parameters %s (%s) was filtered"), *InFunction->GetPathName(), *Property->GetName(), *Property->GetCPPType());
+						return false;
+					}
+				}
+				else
+				{
+					UE_LOG(LogBlueprint, Warning, TEXT("Function %s was filtered because one of its parameters %s (%s) was not able to be converted to a pin"), *InFunction->GetPathName(), *Property->GetName(), *Property->GetCPPType());
+					return false;
+				}
+			}
+		}
+	}
+
+	return true;
+}
+
+bool FBlueprintActionDatabase::IsEnumAllowed(UEnum const* InEnum, EPermissionsContext InContext)
+{
+	const FPathPermissionList& EnumPermissions = GetMutableDefault<UBlueprintEditorSettings>()->GetEnumPermissions();
+	if(EnumPermissions.HasFiltering())
+	{
+		TStringBuilder<256> ResultBuilder;
+		InEnum->GetPathName(nullptr, ResultBuilder);
+		return EnumPermissions.PassesFilter(ResultBuilder.ToView());
+	}
+	return true;
+}
+
+bool FBlueprintActionDatabase::IsEnumAllowed(const FTopLevelAssetPath& InEnumPath, EPermissionsContext InContext)
+{
+	const FPathPermissionList& EnumPermissions = GetMutableDefault<UBlueprintEditorSettings>()->GetEnumPermissions();
+	if(EnumPermissions.HasFiltering())
+	{
+		return EnumPermissions.PassesFilter(InEnumPath.ToString());
+	}
+	return true;
+}
+
+bool FBlueprintActionDatabase::IsStructAllowed(UScriptStruct const* InStruct, EPermissionsContext InContext)
+{
+	const FPathPermissionList& StructPermissions = GetMutableDefault<UBlueprintEditorSettings>()->GetStructPermissions();
+	if(StructPermissions.HasFiltering())
+	{
+		TStringBuilder<256> ResultBuilder;
+		InStruct->GetPathName(nullptr, ResultBuilder);
+		return StructPermissions.PassesFilter(ResultBuilder.ToView());
+	}
+	return true;
+}
+
+bool FBlueprintActionDatabase::IsStructAllowed(const FTopLevelAssetPath& InStructPath, EPermissionsContext InContext)
+{
+	const FPathPermissionList& StructPermissions = GetMutableDefault<UBlueprintEditorSettings>()->GetStructPermissions();
+	if(StructPermissions.HasFiltering())
+	{
+		return StructPermissions.PassesFilter(InStructPath.ToString());
+	}
+	return true;
+}
+
+bool FBlueprintActionDatabase::IsPinTypeAllowed(const FEdGraphPinType& InPinType)
+{
+	// First check if the pin's category is allowed
+	const FNamePermissionList& PinCategoryPermissions = GetMutableDefault<UBlueprintEditorSettings>()->GetPinCategoryPermissions();
+	if(PinCategoryPermissions.HasFiltering())
+	{
+		if(!PinCategoryPermissions.PassesFilter(InPinType.PinCategory))
+		{
+			return false;
+		}
+	}
+
+	if(InPinType.PinCategory == UEdGraphSchema_K2::PC_Struct)
+	{
+		if(const UScriptStruct* ScriptStruct = CastChecked<UScriptStruct>(InPinType.PinSubCategoryObject, ECastCheckedType::NullAllowed))
+		{
+			return IsStructAllowed(ScriptStruct, EPermissionsContext::Pin);
+		}
+	}
+	else if(InPinType.PinCategory == UEdGraphSchema_K2::PC_Enum)
+	{
+		if(const UEnum* Enum = CastChecked<UEnum>(InPinType.PinSubCategoryObject, ECastCheckedType::NullAllowed))
+		{
+			return IsEnumAllowed(Enum, EPermissionsContext::Pin);
+		}
+	}
+	else if(InPinType.PinCategory == UEdGraphSchema_K2::PC_Class ||
+			InPinType.PinCategory == UEdGraphSchema_K2::PC_Object ||
+			InPinType.PinCategory == UEdGraphSchema_K2::PC_Interface ||
+			InPinType.PinCategory == UEdGraphSchema_K2::PC_SoftClass ||
+			InPinType.PinCategory == UEdGraphSchema_K2::PC_SoftObject)
+	{
+		if(const UClass* Class = CastChecked<UClass>(InPinType.PinSubCategoryObject, ECastCheckedType::NullAllowed))
+		{
+			return IsClassAllowed(Class, EPermissionsContext::Pin);
+		}
+	}
+	else if(InPinType.PinCategory == UEdGraphSchema_K2::PC_Delegate ||
+			InPinType.PinCategory == UEdGraphSchema_K2::PC_MCDelegate)
+	{
+		if(const UFunction* PinSignature = FMemberReference::ResolveSimpleMemberReference<UFunction>(InPinType.PinSubCategoryMemberReference))
+		{
+			return IsFunctionAllowed(PinSignature, EPermissionsContext::Pin);
+		}
+	}
+
+	return true;
+}
+
 #undef LOCTEXT_NAMESPACE
diff --git a/Engine/Source/Editor/BlueprintGraph/Private/BlueprintActionDatabaseRegistrar.cpp b/Engine/Source/Editor/BlueprintGraph/Private/BlueprintActionDatabaseRegistrar.cpp
--- a/Engine/Source/Editor/BlueprintGraph/Private/BlueprintActionDatabaseRegistrar.cpp
+++ b/Engine/Source/Editor/BlueprintGraph/Private/BlueprintActionDatabaseRegistrar.cpp
@@ -9,6 +9,8 @@
 #include "EdGraphSchema_K2.h"
 #include "BlueprintNodeSpawnerUtils.h"
 #include "BlueprintEditorSettings.h"
+#include "AssetToolsModule.h"
+#include "IAssetTools.h"
 
 /*******************************************************************************
  * BlueprintActionDatabaseRegistrarImpl
@@ -26,7 +28,7 @@ static UObject const* BlueprintActionDatabaseRegistrarImpl::ResolveClassKey(UCla
 	UObject const* ResolvedKey = ClassKey;
 	if (UBlueprintGeneratedClass const* BlueprintClass = Cast<UBlueprintGeneratedClass>(ClassKey))
 	{
-		ResolvedKey = CastChecked<UBlueprint>(BlueprintClass->ClassGeneratedBy);
+		ResolvedKey = CastChecked<UBlueprint>(BlueprintClass->ClassGeneratedBy, ECastCheckedType::NullAllowed);
 	}
 	return ResolvedKey;
 }
@@ -72,6 +74,7 @@ FBlueprintActionDatabaseRegistrar::FBlueprintActionDatabaseRegistrar(FActionRegi
 	, ActionKeyFilter(nullptr)
 	, ActionPrimingQueue(PrimingQueue)
 {
+	AssetTools = &FModuleManager::LoadModuleChecked<FAssetToolsModule>("AssetTools").Get();
 }
 
 //------------------------------------------------------------------------------
@@ -94,12 +97,23 @@ bool FBlueprintActionDatabaseRegistrar::AddBlueprintAction(UBlueprintNodeSpawner
 	{
 		ActionKey = MemberField;
 	}
+
+	if (!FBlueprintActionDatabase::IsFieldAllowed(ActionKey, FBlueprintActionDatabase::EPermissionsContext::Asset))
+	{
+		return false;
+	}
+
 	return AddActionToDatabase(ActionKey, NodeSpawner);
 }
 
 //------------------------------------------------------------------------------
 bool FBlueprintActionDatabaseRegistrar::AddBlueprintAction(UClass const* ClassOwner, UBlueprintNodeSpawner* NodeSpawner)
 {
+	if (!FBlueprintActionDatabase::IsClassAllowed(ClassOwner, FBlueprintActionDatabase::EPermissionsContext::Asset))
+	{
+		return false;
+	}
+
 	// @TODO: assert that AddBlueprintAction(UBlueprintNodeSpawner* NodeSpawner) 
 	//        wouldn't come up with a different key (besides GeneratingClass)
 
@@ -111,6 +125,11 @@ bool FBlueprintActionDatabaseRegistrar::AddBlueprintAction(UClass const* ClassOw
 //------------------------------------------------------------------------------
 bool FBlueprintActionDatabaseRegistrar::AddBlueprintAction(UEnum const* EnumOwner, UBlueprintNodeSpawner* NodeSpawner)
 {
+	if (!FBlueprintActionDatabase::IsEnumAllowed(EnumOwner, FBlueprintActionDatabase::EPermissionsContext::Asset))
+	{
+		return false;
+	}
+
 	// @TODO: assert that AddBlueprintAction(UBlueprintNodeSpawner* NodeSpawner) 
 	//        wouldn't come up with a different key (besides GeneratingClass)
 
@@ -122,6 +141,11 @@ bool FBlueprintActionDatabaseRegistrar::AddBlueprintAction(UEnum const* EnumOwne
 //------------------------------------------------------------------------------
 bool FBlueprintActionDatabaseRegistrar::AddBlueprintAction(UScriptStruct const* StructOwner, UBlueprintNodeSpawner* NodeSpawner)
 {
+	if (!FBlueprintActionDatabase::IsStructAllowed(StructOwner, FBlueprintActionDatabase::EPermissionsContext::Asset))
+	{
+		return false;
+	}
+
 	// @TODO: assert that AddBlueprintAction(UBlueprintNodeSpawner* NodeSpawner) 
 	//        wouldn't come up with a different key (besides GeneratingClass)
 
@@ -133,6 +157,11 @@ bool FBlueprintActionDatabaseRegistrar::AddBlueprintAction(UScriptStruct const*
 //------------------------------------------------------------------------------
 bool FBlueprintActionDatabaseRegistrar::AddBlueprintAction(UField const* FieldOwner, UBlueprintNodeSpawner* NodeSpawner)
 {
+	if (!FBlueprintActionDatabase::IsFieldAllowed(FieldOwner, FBlueprintActionDatabase::EPermissionsContext::Asset))
+	{
+		return false;
+	}
+
 	// @TODO: assert that AddBlueprintAction(UBlueprintNodeSpawner* NodeSpawner) 
 	//        wouldn't come up with a different key (besides GeneratingClass)
 
@@ -144,6 +173,11 @@ bool FBlueprintActionDatabaseRegistrar::AddBlueprintAction(UField const* FieldOw
 //------------------------------------------------------------------------------
 bool FBlueprintActionDatabaseRegistrar::AddBlueprintAction(UObject const* AssetOwner, UBlueprintNodeSpawner* NodeSpawner)
 {
+	if (AssetTools && !AssetTools->IsAssetClassSupported(AssetOwner->GetClass()))
+	{
+		return false;
+	}
+
 	// @TODO: assert that AddBlueprintAction(UBlueprintNodeSpawner* NodeSpawner) 
 	//        wouldn't come up with a different key (besides GeneratingClass)
 
@@ -204,6 +238,29 @@ int32 FBlueprintActionDatabaseRegistrar::RegisterStructActions(const FMakeStruct
 {
 	int32 RegisteredCount = 0;
 
+	const FPathPermissionList& Permissions = GetMutableDefault<UBlueprintEditorSettings>()->GetStructPermissions();
+
+	auto RegisterStruct = [this, &MakeActionCallback, &RegisteredCount, &Permissions](const UScriptStruct* InStruct)
+	{
+		if ((InStruct->StructFlags & STRUCT_NewerVersionExists) == 0)
+		{
+			if (Permissions.HasFiltering())
+			{
+				TStringBuilder<256> ResultBuilder;
+				InStruct->GetPathName(nullptr, ResultBuilder);
+				if (!Permissions.PassesFilter(ResultBuilder.ToView()))
+				{
+					return;
+				}
+			}
+
+			if (UBlueprintNodeSpawner* NewAction = MakeActionCallback.Execute(InStruct))
+			{
+				RegisteredCount += (int32)AddActionToDatabase(InStruct, NewAction);
+			}
+		}
+	};
+
 	// to keep from needlessly looping through UScriptStructs, first check to 
 	// see if the registrar is looking for only certain actions of this type
 	// (could be regenerating actions for a specific asset, and therefore the 
@@ -213,13 +270,7 @@ int32 FBlueprintActionDatabaseRegistrar::RegisterStructActions(const FMakeStruct
 		if (const UScriptStruct* StructAsset = Cast<UScriptStruct>(RegistrarTarget))
 		{
 			check(IsOpenForRegistration(StructAsset));
-			if ((StructAsset->StructFlags & STRUCT_NewerVersionExists) == 0)
-			{
-				if (UBlueprintNodeSpawner* NewAction = MakeActionCallback.Execute(StructAsset))
-				{
-					RegisteredCount += (int32)AddBlueprintAction(StructAsset, NewAction);
-				}
-			}
+			RegisterStruct(StructAsset);
 		}
 		// else, the target is a class or a different asset type... not something pertaining to a struct
 	}
@@ -228,13 +279,7 @@ int32 FBlueprintActionDatabaseRegistrar::RegisterStructActions(const FMakeStruct
 		for (TObjectIterator<UScriptStruct> StructIt; StructIt; ++StructIt)
 		{
 			UScriptStruct const* Struct = (*StructIt);
-			if ((Struct->StructFlags & STRUCT_NewerVersionExists) == 0)
-			{
-				if (UBlueprintNodeSpawner* NewAction = MakeActionCallback.Execute(Struct))
-				{
-					RegisteredCount += (int32)AddBlueprintAction(Struct, NewAction);
-				}
-			}
+			RegisterStruct(Struct);
 		}
 	}
 	return RegisteredCount;
@@ -244,58 +289,49 @@ int32 FBlueprintActionDatabaseRegistrar::RegisterStructActions(const FMakeStruct
 int32 FBlueprintActionDatabaseRegistrar::RegisterEnumActions(const FMakeEnumSpawnerDelegate& MakeActionCallback)
 {
 	int32 RegisteredCount = 0;
-	if (const UObject* RegistrarTarget = GetActionKeyFilter())
+
+	const FPathPermissionList& Permissions = GetMutableDefault<UBlueprintEditorSettings>()->GetEnumPermissions();
+
+	auto RegisterEnum = [this, &MakeActionCallback, &RegisteredCount, &Permissions](const UEnum* InEnum)
 	{
-		if (const UClass* TargetClass = Cast<UClass>(RegistrarTarget))
+		if (!InEnum->HasAnyEnumFlags(EEnumFlags::NewerVersionExists) && UEdGraphSchema_K2::IsAllowableBlueprintVariableType(InEnum))
 		{
-			for (TFieldIterator<UEnum> EnumIt(TargetClass, EFieldIteratorFlags::ExcludeSuper); EnumIt; ++EnumIt)
+			if (Permissions.HasFiltering())
 			{
-				UEnum const* EnumToConsider = (*EnumIt);
-				if (EnumToConsider->HasAnyEnumFlags(EEnumFlags::NewerVersionExists))
+				TStringBuilder<256> ResultBuilder;
+				InEnum->GetPathName(nullptr, ResultBuilder);
+				if (!Permissions.PassesFilter(ResultBuilder.ToView()))
 				{
-					continue;
-				}
-				if (!UEdGraphSchema_K2::IsAllowableBlueprintVariableType(EnumToConsider))
-				{
-					continue;
+					return;
 				}
+			}
 
-				if (UBlueprintNodeSpawner* NewAction = MakeActionCallback.Execute(EnumToConsider))
-				{
-					RegisteredCount += (int32)AddBlueprintAction(EnumToConsider, NewAction);
-				}
+			if (UBlueprintNodeSpawner* NewAction = MakeActionCallback.Execute(InEnum))
+			{
+				RegisteredCount += (int32)AddActionToDatabase(InEnum, NewAction);
 			}
 		}
-		else if (const UEnum* TargetEnum = Cast<UEnum>(RegistrarTarget))
+	};
+
+	if (const UObject* RegistrarTarget = GetActionKeyFilter())
+	{
+		if (const UClass* TargetClass = Cast<UClass>(RegistrarTarget))
 		{
-			if (!TargetEnum->HasAnyEnumFlags(EEnumFlags::NewerVersionExists) && UEdGraphSchema_K2::IsAllowableBlueprintVariableType(TargetEnum))
+			for (TFieldIterator<UEnum> EnumIt(TargetClass, EFieldIteratorFlags::ExcludeSuper); EnumIt; ++EnumIt)
 			{
-				if (UBlueprintNodeSpawner* NewAction = MakeActionCallback.Execute(TargetEnum))
-				{
-					RegisteredCount += (int32)AddBlueprintAction(TargetEnum, NewAction);
-				}
+				RegisterEnum(*EnumIt);
 			}
 		}
+		else if (const UEnum* TargetEnum = Cast<UEnum>(RegistrarTarget))
+		{
+			RegisterEnum(TargetEnum);
+		}
 	}
 	else
 	{
 		for (TObjectIterator<UEnum> EnumIt; EnumIt; ++EnumIt)
 		{
-			UEnum const* EnumToConsider = (*EnumIt);
-			if (EnumToConsider->HasAnyEnumFlags(EEnumFlags::NewerVersionExists))
-			{
-				continue;
-			}
-
-			if (!UEdGraphSchema_K2::IsAllowableBlueprintVariableType(EnumToConsider))
-			{
-				continue;
-			}
-
-			if (UBlueprintNodeSpawner* NewAction = MakeActionCallback.Execute(EnumToConsider))
-			{
-				RegisteredCount += (int32)AddBlueprintAction(EnumToConsider, NewAction);
-			}
+			RegisterEnum(*EnumIt);
 		}
 	}
 
diff --git a/Engine/Source/Editor/BlueprintGraph/Private/BlueprintActionFilter.cpp b/Engine/Source/Editor/BlueprintGraph/Private/BlueprintActionFilter.cpp
--- a/Engine/Source/Editor/BlueprintGraph/Private/BlueprintActionFilter.cpp
+++ b/Engine/Source/Editor/BlueprintGraph/Private/BlueprintActionFilter.cpp
@@ -30,6 +30,7 @@
 #include "BlueprintEventNodeSpawner.h"
 #include "BlueprintBoundEventNodeSpawner.h"
 #include "BlueprintBoundNodeSpawner.h"
+#include "BlueprintAssetNodeSpawner.h"
 #include "Algo/Transform.h"
 // "impure" node types (utilized in BlueprintActionFilterImpl::IsImpure)
 #include "K2Node_MultiGate.h"
@@ -39,6 +40,7 @@
 #include "BlueprintEditorModule.h"
 #include "Kismet2/KismetEditorUtilities.h"
 #include "BlueprintEditorSettings.h"
+#include "Editor/EditorEngine.h"
 
 /*******************************************************************************
  * Static BlueprintActionFilter Helpers
@@ -155,6 +157,17 @@ namespace BlueprintActionFilterImpl
 	 */
 	static bool IsRestrictedClassMember(FBlueprintActionFilter const& Filter, FBlueprintActionInfo& BlueprintAction);
 	
+	/**
+	 * Rejection test that checks to see if the supplied node-spawner would 
+	 * produce an asset-related action that cannot be accessed according to 
+	 * asset permissions
+	 * 
+	 * @param  Filter			Holds the blueprint context for this test.
+	 * @param  BlueprintAction	The action you wish to query.
+	 * @return True if the action would spawn a variable-set node for a read-only property.
+	 */
+	static bool IsAssetPermissionNotGranted(FBlueprintActionFilter const& Filter, FBlueprintActionInfo& BlueprintAction);
+
 	/**
 	 * Rejection test that checks to see if the supplied node-spawner would 
 	 * produce a variable-set node when the property is read-only (in this
@@ -836,6 +849,42 @@ static bool BlueprintActionFilterImpl::IsRestrictedClassMember(FBlueprintActionF
 	return bIsFilteredOut;
 }
 
+//------------------------------------------------------------------------------
+static bool BlueprintActionFilterImpl::IsAssetPermissionNotGranted(FBlueprintActionFilter const& Filter, FBlueprintActionInfo& BlueprintAction)
+{
+	bool bIsFilteredOut = false;
+	FBlueprintActionContext const& FilterContext = Filter.Context;
+
+	// This handles actions with loaded assets
+	UObject const* const Owner = BlueprintAction.GetActionOwner();
+	if (Owner && Owner->IsAsset())
+	{
+		if (Filter.AssetReferenceFilter.IsValid())
+		{
+			bool const bIsAccessible = Filter.AssetReferenceFilter->PassesFilter(FAssetData(Owner));
+			if (!bIsAccessible)
+			{
+				bIsFilteredOut = true;
+			}
+		}
+	}
+	// Handle unloaded asset spawners
+	else if(UBlueprintAssetNodeSpawner const* AssetNodeSpawner = Cast<UBlueprintAssetNodeSpawner>(BlueprintAction.NodeSpawner))
+	{
+		const FAssetData& AssetData = AssetNodeSpawner->GetAssetData();
+		if (Filter.AssetReferenceFilter.IsValid() && AssetData.IsValid())
+		{
+			bool const bIsAccessible = Filter.AssetReferenceFilter->PassesFilter(AssetData);
+			if (!bIsAccessible)
+			{
+				bIsFilteredOut = true;
+			}
+		}
+	}
+
+	return bIsFilteredOut;
+}
+
 //------------------------------------------------------------------------------
 static bool BlueprintActionFilterImpl::IsPermissionNotGranted(FBlueprintActionFilter const& Filter, FBlueprintActionInfo& BlueprintAction)
 {
@@ -2061,8 +2110,9 @@ FBlueprintActionFilter::FBlueprintActionFilter(const EFlags InFlags /*= BPFILTER
 	//
 	// this test in-particular spawns a template-node and then calls 
 	// AllocateDefaultPins() which is costly, so it should be very last!
-	AddRejectionTest(FRejectionTestDelegate::CreateStatic(IsIncompatibleAnimNotification));
 	AddRejectionTest(FRejectionTestDelegate::CreateStatic(IsNodeTemplateSelfFiltered));
+	AddRejectionTest(FRejectionTestDelegate::CreateStatic(IsIncompatibleAnimNotification));
+	AddRejectionTest(FRejectionTestDelegate::CreateStatic(IsAssetPermissionNotGranted));
 	AddRejectionTest(FRejectionTestDelegate::CreateStatic(IsMissingMatchingPinParam));
 	AddRejectionTest(FRejectionTestDelegate::CreateStatic(IsMissmatchedPropertyType));
 	AddRejectionTest(FRejectionTestDelegate::CreateStatic(IsFunctionMissingPinParam));
diff --git a/Engine/Source/Editor/BlueprintGraph/Private/BlueprintAssetNodeSpawner.cpp b/Engine/Source/Editor/BlueprintGraph/Private/BlueprintAssetNodeSpawner.cpp
new file mode 100644
--- /dev/null
+++ b/Engine/Source/Editor/BlueprintGraph/Private/BlueprintAssetNodeSpawner.cpp
@@ -0,0 +1,31 @@
+// Copyright Epic Games, Inc. All Rights Reserved.
+
+#include "BlueprintAssetNodeSpawner.h"
+#include "UObject/Package.h"
+
+#define LOCTEXT_NAMESPACE "BlueprintAssetNodeSpawner"
+
+UBlueprintAssetNodeSpawner* UBlueprintAssetNodeSpawner::Create(TSubclassOf<UEdGraphNode> const InNodeClass, const FAssetData& InAssetData, UObject* InOuter, FCustomizeNodeDelegate InPostSpawnDelegate)
+{
+	check(InNodeClass != nullptr);
+	check(InNodeClass->IsChildOf<UEdGraphNode>());
+
+	if (InOuter == nullptr)
+	{
+		InOuter = GetTransientPackage();
+	}
+
+	UBlueprintAssetNodeSpawner* NodeSpawner = NewObject<UBlueprintAssetNodeSpawner>(InOuter);
+	NodeSpawner->NodeClass = InNodeClass;
+	NodeSpawner->CustomizeNodeDelegate = InPostSpawnDelegate;
+	NodeSpawner->AssetData = InAssetData;
+
+	return NodeSpawner;
+}
+
+const FAssetData& UBlueprintAssetNodeSpawner::GetAssetData() const
+{
+	return AssetData;
+}
+
+#undef LOCTEXT_NAMESPACE
diff --git a/Engine/Source/Editor/BlueprintGraph/Private/BlueprintEditorSettings.cpp b/Engine/Source/Editor/BlueprintGraph/Private/BlueprintEditorSettings.cpp
--- a/Engine/Source/Editor/BlueprintGraph/Private/BlueprintEditorSettings.cpp
+++ b/Engine/Source/Editor/BlueprintGraph/Private/BlueprintEditorSettings.cpp
@@ -166,4 +166,92 @@ void UBlueprintEditorSettings::PostEditChangeProperty(FPropertyChangedEvent& Pro
 	}
 
 	Super::PostEditChangeProperty(PropertyChangedEvent);
+}
+
+void UBlueprintEditorSettings::RegisterIsClassAllowedDelegate(const FName OwnerName, FOnIsClassAllowed Delegate)
+{
+	IsClassAllowedDelegates.Add(OwnerName, Delegate);
+}
+
+void UBlueprintEditorSettings::UnregisterIsClassAllowedDelegate(const FName OwnerName)
+{
+	IsClassAllowedDelegates.Remove(OwnerName);
+}
+
+bool UBlueprintEditorSettings::IsClassAllowed(const UClass* InClass) const
+{
+	for (const TPair<FName, FOnIsClassAllowed>& DelegatePair : IsClassAllowedDelegates)
+	{
+		if (!DelegatePair.Value.Execute(InClass))
+		{
+			return false;
+		}
+	}
+	return true;
+}
+
+void UBlueprintEditorSettings::RegisterIsClassPathAllowedDelegate(const FName OwnerName, FOnIsClassPathAllowed Delegate)
+{
+	IsClassPathAllowedDelegates.Add(OwnerName, Delegate);
+}
+
+void UBlueprintEditorSettings::UnregisterIsClassPathAllowedDelegate(const FName OwnerName)
+{
+	IsClassPathAllowedDelegates.Remove(OwnerName);
+}
+
+bool UBlueprintEditorSettings::IsClassPathAllowed(const FTopLevelAssetPath& InClassPath) const
+{
+	for (const TPair<FName, FOnIsClassPathAllowed>& DelegatePair : IsClassPathAllowedDelegates)
+	{
+		if (!DelegatePair.Value.Execute(InClassPath))
+		{
+			return false;
+		}
+	}
+	return true;
+}
+
+void UBlueprintEditorSettings::RegisterIsClassAllowedOnPinDelegate(const FName OwnerName, FOnIsClassAllowed Delegate)
+{
+	IsClassAllowedOnPinDelegates.Add(OwnerName, Delegate);
+}
+
+void UBlueprintEditorSettings::UnregisterIsClassAllowedOnPinDelegate(const FName OwnerName)
+{
+	IsClassAllowedOnPinDelegates.Remove(OwnerName);
+}
+
+bool UBlueprintEditorSettings::IsClassAllowedOnPin(const UClass* InClass) const
+{
+	for (const TPair<FName, FOnIsClassAllowed>& DelegatePair : IsClassAllowedOnPinDelegates)
+	{
+		if (!DelegatePair.Value.Execute(InClass))
+		{
+			return false;
+		}
+	}
+	return true;
+}
+
+void UBlueprintEditorSettings::RegisterIsClassPathAllowedOnPinDelegate(const FName OwnerName, FOnIsClassPathAllowed Delegate)
+{
+	IsClassPathAllowedOnPinDelegates.Add(OwnerName, Delegate);
+}
+
+void UBlueprintEditorSettings::UnregisterIsClassPathAllowedOnPinDelegate(const FName OwnerName)
+{
+	IsClassPathAllowedOnPinDelegates.Remove(OwnerName);
+}
+
+bool UBlueprintEditorSettings::IsClassPathAllowedOnPin(const FTopLevelAssetPath& InClassPath) const
+{
+	for (const TPair<FName, FOnIsClassPathAllowed>& DelegatePair : IsClassPathAllowedOnPinDelegates)
+	{
+		if (!DelegatePair.Value.Execute(InClassPath))
+		{
+			return false;
+		}
+	}
+	return true;
 }
\ No newline at end of file
diff --git a/Engine/Source/Editor/BlueprintGraph/Private/EdGraphSchema_K2.cpp b/Engine/Source/Editor/BlueprintGraph/Private/EdGraphSchema_K2.cpp
--- a/Engine/Source/Editor/BlueprintGraph/Private/EdGraphSchema_K2.cpp
+++ b/Engine/Source/Editor/BlueprintGraph/Private/EdGraphSchema_K2.cpp
@@ -325,7 +325,7 @@ public:
 			for (TObjectIterator<UEnum> EnumIt; EnumIt; ++EnumIt)
 			{
 				UEnum* CurrentEnum = *EnumIt;
-				if (UEdGraphSchema_K2::IsAllowableBlueprintVariableType(CurrentEnum))
+				if (UEdGraphSchema_K2::IsAllowableBlueprintVariableType(CurrentEnum) && FBlueprintActionDatabase::IsEnumAllowed(CurrentEnum, FBlueprintActionDatabase::EPermissionsContext::Property))
 				{
 					LoadedTypesList->Add(FLoadedAssetData(CurrentEnum));
 				}
@@ -342,7 +342,7 @@ public:
 				for (TObjectIterator<UScriptStruct> StructIt; StructIt; ++StructIt)
 				{
 					UScriptStruct* ScriptStruct = *StructIt;
-					if (UEdGraphSchema_K2::IsAllowableBlueprintVariableType(ScriptStruct))
+					if (UEdGraphSchema_K2::IsAllowableBlueprintVariableType(ScriptStruct) && FBlueprintActionDatabase::IsStructAllowed(ScriptStruct, FBlueprintActionDatabase::EPermissionsContext::Property))
 					{
 						LoadedTypesList->Add(FLoadedAssetData(ScriptStruct));
 					}
@@ -363,7 +363,7 @@ public:
 					const bool bIsInterface = CurrentClass->IsChildOf(UInterface::StaticClass());
 					const bool bIsBlueprintType = UEdGraphSchema_K2::IsAllowableBlueprintVariableType(CurrentClass);
 					const bool bIsDeprecated = CurrentClass->HasAnyClassFlags(CLASS_Deprecated);
-					if (bIsBlueprintType && !bIsDeprecated)
+					if (bIsBlueprintType && !bIsDeprecated && FBlueprintActionDatabase::IsClassAllowed(CurrentClass, FBlueprintActionDatabase::EPermissionsContext::Property))
 					{
 						if (bIsInterface)
 						{
@@ -389,6 +389,7 @@ public:
 
 		const FAssetRegistryModule& AssetRegistryModule = FModuleManager::LoadModuleChecked<FAssetRegistryModule>(TEXT("AssetRegistry"));
 
+		if(FBlueprintActionDatabase::IsClassAllowed(UUserDefinedEnum::StaticClass(), FBlueprintActionDatabase::EPermissionsContext::Property))
 		{
 			TArray<FAssetData> AssetData;
 			AssetRegistryModule.Get().GetAssetsByClass(UUserDefinedEnum::StaticClass()->GetClassPathName(), AssetData);
@@ -397,7 +398,7 @@ public:
 			for (int32 AssetIndex = 0; AssetIndex < AssetData.Num(); ++AssetIndex)
 			{
 				const FAssetData& Asset = AssetData[AssetIndex];
-				if (Asset.IsValid() && !Asset.IsAssetLoaded())
+				if (Asset.IsValid() && !Asset.IsAssetLoaded() && FBlueprintActionDatabase::IsEnumAllowed(FTopLevelAssetPath(Asset.PackageName, Asset.AssetName), FBlueprintActionDatabase::EPermissionsContext::Property))
 				{
 					UnLoadedTypesList->Add(FUnloadedAssetData(Asset));
 				}
@@ -408,6 +409,7 @@ public:
 
 		if (!bIndexTypesOnly)
 		{
+			if (FBlueprintActionDatabase::IsClassAllowed(UUserDefinedStruct::StaticClass(), FBlueprintActionDatabase::EPermissionsContext::Property))
 			{
 				TArray<FAssetData> AssetData;
 				AssetRegistryModule.Get().GetAssetsByClass(UUserDefinedStruct::StaticClass()->GetClassPathName(), AssetData);
@@ -416,7 +418,7 @@ public:
 				for (int32 AssetIndex = 0; AssetIndex < AssetData.Num(); ++AssetIndex)
 				{
 					const FAssetData& Asset = AssetData[AssetIndex];
-					if (Asset.IsValid() && !Asset.IsAssetLoaded())
+					if (Asset.IsValid() && !Asset.IsAssetLoaded() && FBlueprintActionDatabase::IsStructAllowed(FTopLevelAssetPath(Asset.PackageName, Asset.AssetName), FBlueprintActionDatabase::EPermissionsContext::Property))
 					{
 						UnLoadedTypesList->Add(FUnloadedAssetData(Asset));
 					}
@@ -425,6 +427,7 @@ public:
 				TypesDatabase.UnLoadedTypesMap.Add(UEdGraphSchema_K2::PC_Struct, UnLoadedTypesList);
 			}
 
+			if (FBlueprintActionDatabase::IsClassAllowed(UBlueprint::StaticClass(), FBlueprintActionDatabase::EPermissionsContext::Property))
 			{
 				TArray<FAssetData> AssetData;
 				AssetRegistryModule.Get().GetAssetsByClass(UBlueprint::StaticClass()->GetClassPathName(), AssetData);
@@ -439,7 +442,7 @@ public:
 				{
 					const FAssetData& Asset = AssetData[AssetIndex];
 
-					if (Asset.IsValid() && !Asset.IsAssetLoaded())
+					if (Asset.IsValid() && !Asset.IsAssetLoaded() && FBlueprintActionDatabase::IsClassAllowed(FTopLevelAssetPath(Asset.PackageName, Asset.AssetName), FBlueprintActionDatabase::EPermissionsContext::Property))
 					{
 						const FString BlueprintTypeStr = Asset.GetTagValueRef<FString>(FBlueprintTags::BlueprintType);
 						const bool bNormalBP = BlueprintTypeStr == BPNormalTypeAllowed;
@@ -3923,57 +3926,66 @@ void UEdGraphSchema_K2::GetVariableTypeTree(TArray< TSharedPtr<FPinTypeTreeInfo>
 
 	// Clear the list
 	TypeTree.Empty();
-
+	
+	auto AddToTypeTree = [&TypeTree](const TSharedRef<FPinTypeTreeInfo>& InPinTypeTreeInfo)
+	{
+		constexpr bool bForceLoadedSubCategoryObject = false;
+		if(FBlueprintActionDatabase::IsPinTypeAllowed(InPinTypeTreeInfo->GetPinType(bForceLoadedSubCategoryObject)))
+		{
+			TypeTree.Add(InPinTypeTreeInfo);
+		}
+	};
+	
 	if( bAllowExec )
 	{
-		TypeTree.Add( MakeShareable( new FPinTypeTreeInfo(GetCategoryText(PC_Exec, true), PC_Exec, this, LOCTEXT("ExecType", "Execution pin")) ) );
+		AddToTypeTree( MakeShareable( new FPinTypeTreeInfo(GetCategoryText(PC_Exec, true), PC_Exec, this, LOCTEXT("ExecType", "Execution pin")) ) );
 	}
 
-	TypeTree.Add( MakeShareable( new FPinTypeTreeInfo(GetCategoryText(PC_Boolean, true), PC_Boolean, this, LOCTEXT("BooleanType", "True or false value")) ) );
-	TypeTree.Add( MakeShareable( new FPinTypeTreeInfo(GetCategoryText(PC_Byte, true), PC_Byte, this, LOCTEXT("ByteType", "8 bit number")) ) );
-	TypeTree.Add( MakeShareable( new FPinTypeTreeInfo(GetCategoryText(PC_Int, true), PC_Int, this, LOCTEXT("IntegerType", "Integer number")) ) );
-	TypeTree.Add( MakeShareable( new FPinTypeTreeInfo(GetCategoryText(PC_Int64, true), PC_Int64, this, LOCTEXT("Integer64Type", "64 bit Integer number")) ) );
+	AddToTypeTree( MakeShareable( new FPinTypeTreeInfo(GetCategoryText(PC_Boolean, true), PC_Boolean, this, LOCTEXT("BooleanType", "True or false value")) ) );
+	AddToTypeTree( MakeShareable( new FPinTypeTreeInfo(GetCategoryText(PC_Byte, true), PC_Byte, this, LOCTEXT("ByteType", "8 bit number")) ) );
+	AddToTypeTree( MakeShareable( new FPinTypeTreeInfo(GetCategoryText(PC_Int, true), PC_Int, this, LOCTEXT("IntegerType", "Integer number")) ) );
+	AddToTypeTree( MakeShareable( new FPinTypeTreeInfo(GetCategoryText(PC_Int64, true), PC_Int64, this, LOCTEXT("Integer64Type", "64 bit Integer number")) ) );
 
 	if (!bIndexTypesOnly)
 	{
-		TypeTree.Add(MakeShareable(new FPinTypeTreeInfo(GetCategoryText(PC_Real, true), PC_Real, this, LOCTEXT("RealType", "Floating point number"))));
-		TypeTree.Add(MakeShareable(new FPinTypeTreeInfo(GetCategoryText(PC_Name, true), PC_Name, this, LOCTEXT("NameType", "A text name"))));
-		TypeTree.Add(MakeShareable(new FPinTypeTreeInfo(GetCategoryText(PC_String, true), PC_String, this, LOCTEXT("StringType", "A text string"))));
-		TypeTree.Add(MakeShareable(new FPinTypeTreeInfo(GetCategoryText(PC_Text, true), PC_Text, this, LOCTEXT("TextType", "A localizable text string"))));
+		AddToTypeTree(MakeShareable(new FPinTypeTreeInfo(GetCategoryText(PC_Real, true), PC_Real, this, LOCTEXT("RealType", "Floating point number"))));
+		AddToTypeTree(MakeShareable(new FPinTypeTreeInfo(GetCategoryText(PC_Name, true), PC_Name, this, LOCTEXT("NameType", "A text name"))));
+		AddToTypeTree(MakeShareable(new FPinTypeTreeInfo(GetCategoryText(PC_String, true), PC_String, this, LOCTEXT("StringType", "A text string"))));
+		AddToTypeTree(MakeShareable(new FPinTypeTreeInfo(GetCategoryText(PC_Text, true), PC_Text, this, LOCTEXT("TextType", "A localizable text string"))));
 
 		// Add in special first-class struct types
 		if (!bRootTypesOnly)
 		{
-			TypeTree.Add(MakeShareable(new FPinTypeTreeInfo(PC_Struct, TBaseStructure<FVector>::Get(), LOCTEXT("VectorType", "A 3D vector"))));
-			TypeTree.Add(MakeShareable(new FPinTypeTreeInfo(PC_Struct, TBaseStructure<FRotator>::Get(), LOCTEXT("RotatorType", "A 3D rotation"))));
-			TypeTree.Add(MakeShareable(new FPinTypeTreeInfo(PC_Struct, TBaseStructure<FTransform>::Get(), LOCTEXT("TransformType", "A 3D transformation, including translation, rotation and 3D scale."))));
+			AddToTypeTree(MakeShareable(new FPinTypeTreeInfo(PC_Struct, TBaseStructure<FVector>::Get(), LOCTEXT("VectorType", "A 3D vector"))));
+			AddToTypeTree(MakeShareable(new FPinTypeTreeInfo(PC_Struct, TBaseStructure<FRotator>::Get(), LOCTEXT("RotatorType", "A 3D rotation"))));
+			AddToTypeTree(MakeShareable(new FPinTypeTreeInfo(PC_Struct, TBaseStructure<FTransform>::Get(), LOCTEXT("TransformType", "A 3D transformation, including translation, rotation and 3D scale."))));
 		}
 	}
 	// Add wildcard type
 	if (bAllowWildCard)
 	{
-		TypeTree.Add( MakeShareable( new FPinTypeTreeInfo(GetCategoryText(PC_Wildcard, true), PC_Wildcard, this, LOCTEXT("WildcardType", "Wildcard type (unspecified)")) ) );
+		AddToTypeTree( MakeShareable( new FPinTypeTreeInfo(GetCategoryText(PC_Wildcard, true), PC_Wildcard, this, LOCTEXT("WildcardType", "Wildcard type (unspecified)")) ) );
 	}
 
 	// Add the types that have subtrees
 	if (!bIndexTypesOnly)
 	{
-		TypeTree.Add(MakeShareable(new FPinTypeTreeInfo(GetCategoryText(PC_Struct, true), PC_Struct, this, LOCTEXT("StructType", "Struct (value) types"), true, TypesDatabasePtr)));
-		TypeTree.Add(MakeShareable(new FPinTypeTreeInfo(GetCategoryText(PC_Interface, true), PC_Interface, this, LOCTEXT("InterfaceType", "Interface types"), true, TypesDatabasePtr)));
+		AddToTypeTree(MakeShareable(new FPinTypeTreeInfo(GetCategoryText(PC_Struct, true), PC_Struct, this, LOCTEXT("StructType", "Struct (value) types"), true, TypesDatabasePtr)));
+		AddToTypeTree(MakeShareable(new FPinTypeTreeInfo(GetCategoryText(PC_Interface, true), PC_Interface, this, LOCTEXT("InterfaceType", "Interface types"), true, TypesDatabasePtr)));
 
 		if (!bRootTypesOnly)
 		{
-			TypeTree.Add(MakeShareable(new FPinTypeTreeInfo(GetCategoryText(AllObjectTypes, true), AllObjectTypes, this, LOCTEXT("ObjectType", "Object types"), true, TypesDatabasePtr)));
+			AddToTypeTree(MakeShareable(new FPinTypeTreeInfo(GetCategoryText(AllObjectTypes, true), AllObjectTypes, this, LOCTEXT("ObjectType", "Object types"), true, TypesDatabasePtr)));
 		}
 		else
 		{
-			TypeTree.Add(MakeShareable(new FPinTypeTreeInfo(GetCategoryText(PC_Object, true), PC_Object, this, LOCTEXT("ObjectTypeHardReference", "Hard reference to an Object"), true, TypesDatabasePtr)));
-			TypeTree.Add(MakeShareable(new FPinTypeTreeInfo(GetCategoryText(PC_Class, true), PC_Class, this, LOCTEXT("ClassType", "Hard reference to a Class"), true, TypesDatabasePtr)));
-			TypeTree.Add(MakeShareable(new FPinTypeTreeInfo(GetCategoryText(PC_SoftObject, true), PC_SoftObject, this, LOCTEXT("SoftObjectType", "Soft reference to an Object"), true, TypesDatabasePtr)));
-			TypeTree.Add(MakeShareable(new FPinTypeTreeInfo(GetCategoryText(PC_SoftClass, true), PC_SoftClass, this, LOCTEXT("SoftClassType", "Soft reference to a Class"), true, TypesDatabasePtr)));
+			AddToTypeTree(MakeShareable(new FPinTypeTreeInfo(GetCategoryText(PC_Object, true), PC_Object, this, LOCTEXT("ObjectTypeHardReference", "Hard reference to an Object"), true, TypesDatabasePtr)));
+			AddToTypeTree(MakeShareable(new FPinTypeTreeInfo(GetCategoryText(PC_Class, true), PC_Class, this, LOCTEXT("ClassType", "Hard reference to a Class"), true, TypesDatabasePtr)));
+			AddToTypeTree(MakeShareable(new FPinTypeTreeInfo(GetCategoryText(PC_SoftObject, true), PC_SoftObject, this, LOCTEXT("SoftObjectType", "Soft reference to an Object"), true, TypesDatabasePtr)));
+			AddToTypeTree(MakeShareable(new FPinTypeTreeInfo(GetCategoryText(PC_SoftClass, true), PC_SoftClass, this, LOCTEXT("SoftClassType", "Soft reference to a Class"), true, TypesDatabasePtr)));
 		}
 	}
-	TypeTree.Add( MakeShareable( new FPinTypeTreeInfo(GetCategoryText(PC_Enum, true), PC_Enum, this, LOCTEXT("EnumType", "Enumeration types."), true, TypesDatabasePtr) ) );
+	AddToTypeTree( MakeShareable( new FPinTypeTreeInfo(GetCategoryText(PC_Enum, true), PC_Enum, this, LOCTEXT("EnumType", "Enumeration types."), true, TypesDatabasePtr) ) );
 
 #if SCHEMA_K2_GETVARIABLETYPETREE_LOG_TIME
 	const double EndTime = FPlatformTime::Seconds();
diff --git a/Engine/Source/Editor/BlueprintGraph/Public/BlueprintActionDatabase.h b/Engine/Source/Editor/BlueprintGraph/Public/BlueprintActionDatabase.h
--- a/Engine/Source/Editor/BlueprintGraph/Public/BlueprintActionDatabase.h
+++ b/Engine/Source/Editor/BlueprintGraph/Public/BlueprintActionDatabase.h
@@ -8,10 +8,12 @@
 #include "UObject/ObjectKey.h"
 #include "UObject/GCObject.h"
 #include "TickableEditorObject.h"
+#include "Misc/NamePermissionList.h"
 
 class FBlueprintActionDatabaseRegistrar;
 class UBlueprint;
 class UBlueprintNodeSpawner;
+struct FEdGraphPinType;
 
 /**
  * Serves as a container for all available blueprint actions (no matter the 
@@ -148,6 +150,51 @@ public:
 	/** */
 	FOnDatabaseEntryUpdated& OnEntryRemoved() { return EntryRemovedDelegate; }
 
+	/**
+	 * Field visibility is different depending on context, this enum differentiates between the different
+	 * contexts that fields can be present in
+	 */
+	enum class EPermissionsContext
+	{
+		/** A property on a class - from a user perspective a variable in the My Blueprint tab */
+		Property,
+		
+		/** An asset - e.g. an anim blueprint asset player or skeleton notify */
+		Asset,
+		
+		/** A K2 node in a graph - e.g. a function call */
+		Node,
+		
+		/** An exposed pin on a K2 node */
+		Pin
+	};
+	
+	/** Check whether the global filter applies to this class */
+	static bool IsClassAllowed(UClass const* InClass, EPermissionsContext InContext);
+	static bool IsClassAllowed(const FTopLevelAssetPath& InClassPath, EPermissionsContext InContext);
+
+	/** Check whether we have any class filtering applied */
+	static bool HasClassFiltering();
+
+	/** Check whether permissions allow this field */
+	static bool IsFieldAllowed(UField const* InField, EPermissionsContext InContext);
+
+	/** Check whether permissions allow this function */
+	static bool IsFunctionAllowed(UFunction const* InFunction, EPermissionsContext InContext);
+
+	/** Check whether permissions allow this enum */
+	static bool IsEnumAllowed(UEnum const* InEnum, EPermissionsContext InContext);
+	static bool IsEnumAllowed(const FTopLevelAssetPath& InEnumPath, EPermissionsContext InContext);
+
+	/** Check whether permissions allow this struct */
+	static bool IsStructAllowed(UScriptStruct const* InStruct, EPermissionsContext InContext);
+	static bool IsStructAllowed(const FTopLevelAssetPath& InStructPath, EPermissionsContext InContext);
+	
+	/** Check whether permissions allow this ping type */
+	static bool IsPinTypeAllowed(const FEdGraphPinType& InPinType);
+	
+
+
 private:
 	/** Private constructor for singleton purposes. */
 	FBlueprintActionDatabase();
diff --git a/Engine/Source/Editor/BlueprintGraph/Public/BlueprintActionDatabaseRegistrar.h b/Engine/Source/Editor/BlueprintGraph/Public/BlueprintActionDatabaseRegistrar.h
--- a/Engine/Source/Editor/BlueprintGraph/Public/BlueprintActionDatabaseRegistrar.h
+++ b/Engine/Source/Editor/BlueprintGraph/Public/BlueprintActionDatabaseRegistrar.h
@@ -11,6 +11,7 @@
 #include "BlueprintActionDatabase.h"
 
 class UBlueprintNodeSpawner;
+class IAssetTools;
 
 /**
  * From the BlueprintActionDatabase, passed around to all UK2Nodes, giving each 
@@ -151,4 +152,7 @@ private:
 
 	/** A reference to FBlueprintActionDatabase's action priming queue */
 	FPrimingQueue& ActionPrimingQueue;
+
+	/** Cached ptr to asset tools to avoid LoadModule calls per-AddBlueprintAction call */
+	IAssetTools* AssetTools;
 };
diff --git a/Engine/Source/Editor/BlueprintGraph/Public/BlueprintActionFilter.h b/Engine/Source/Editor/BlueprintGraph/Public/BlueprintActionFilter.h
--- a/Engine/Source/Editor/BlueprintGraph/Public/BlueprintActionFilter.h
+++ b/Engine/Source/Editor/BlueprintGraph/Public/BlueprintActionFilter.h
@@ -16,6 +16,7 @@ class UBlueprintNodeSpawner;
 class UEdGraph;
 class UEdGraphPin;
 class IBlueprintEditor;
+class IAssetReferenceFilter;
 
 /*******************************************************************************
  * FBlueprintActionContext
@@ -261,6 +262,9 @@ public:
 	/** Cached reference to the BluprintGraphModule, which has extra rejection tests: */
 	class FBlueprintGraphModule* BluprintGraphModule;
 
+	/** Filter for asset references */
+	TSharedPtr<IAssetReferenceFilter> AssetReferenceFilter;
+
 	/**
 	 * Users can extend the filter and add their own rejection tests with this
 	 * method. We use rejection "IsFiltered" tests rather than inclusive tests 
diff --git a/Engine/Source/Editor/BlueprintGraph/Public/BlueprintEditorSettings.h b/Engine/Source/Editor/BlueprintGraph/Public/BlueprintEditorSettings.h
--- a/Engine/Source/Editor/BlueprintGraph/Public/BlueprintEditorSettings.h
+++ b/Engine/Source/Editor/BlueprintGraph/Public/BlueprintEditorSettings.h
@@ -10,6 +10,8 @@
 #include "Kismet2/WatchedPin.h"
 #include "Engine/DeveloperSettings.h"
 #include "Kismet2/KismetDebugUtilities.h"
+#include "Misc/NamePermissionList.h"
+
 #include "BlueprintEditorSettings.generated.h"
 
 UENUM()
@@ -278,6 +280,95 @@ public:
 	UPROPERTY(EditAnywhere, config, Category=Play, meta=(AllowAbstract))
 	TArray<TSoftClassPtr<UObject>> BaseClassesToAllowRecompilingDuringPlayInEditor;
 
+	/** Get allowed structs permissions list */
+	FPathPermissionList& GetStructPermissions() { return StructPermissions; }
+
+	/** Get allowed enums permissions list */
+	FPathPermissionList& GetEnumPermissions() { return EnumPermissions; }
+
+	/** Get allowed pin categories permissions list */
+	FNamePermissionList& GetPinCategoryPermissions() { return PinCategoryPermissions; }
+
+	DECLARE_DELEGATE_RetVal_OneParam(bool, FOnIsClassAllowed, const UClass* /*InClass*/)
+	DECLARE_DELEGATE_RetVal_OneParam(bool, FOnIsClassPathAllowed, const FTopLevelAssetPath& /*InClassPath*/)
+	
+	/** Delegates called to determine whether a class type is allowed */
+	void RegisterIsClassAllowedDelegate(const FName OwnerName, FOnIsClassAllowed Delegate);
+	void UnregisterIsClassAllowedDelegate(const FName OwnerName);
+	bool IsClassAllowed(const UClass* InClass) const;
+	bool HasClassFiltering() const { return IsClassAllowedDelegates.Num() > 0; }
+
+	void RegisterIsClassPathAllowedDelegate(const FName OwnerName, FOnIsClassPathAllowed Delegate);
+	void UnregisterIsClassPathAllowedDelegate(const FName OwnerName);
+	bool IsClassPathAllowed(const FTopLevelAssetPath& InClassPath) const;
+	bool HasClassPathFiltering() const { return IsClassPathAllowedDelegates.Num() > 0; }
+
+	void RegisterIsClassAllowedOnPinDelegate(const FName OwnerName, FOnIsClassAllowed Delegate);
+	void UnregisterIsClassAllowedOnPinDelegate(const FName OwnerName);
+	bool IsClassAllowedOnPin(const UClass* InClass) const;
+	bool HasClassOnPinFiltering() const { return IsClassAllowedOnPinDelegates.Num() > 0; }
+
+	void RegisterIsClassPathAllowedOnPinDelegate(const FName OwnerName, FOnIsClassPathAllowed Delegate);
+	void UnregisterIsClassPathAllowedOnPinDelegate(const FName OwnerName);
+	bool IsClassPathAllowedOnPin(const FTopLevelAssetPath& InClassPath) const;
+	bool HasClassPathOnPinFiltering() const { return IsClassPathAllowedOnPinDelegates.Num() > 0; }
+	
+private:
+
+	/** Get allowed structs permissions list */
+	FPathPermissionList& GetStructPermissions() { return StructPermissions; }
+
+	/** Get allowed enums permissions list */
+	FPathPermissionList& GetEnumPermissions() { return EnumPermissions; }
+
+	/** Get allowed pin categories permissions list */
+	FNamePermissionList& GetPinCategoryPermissions() { return PinCategoryPermissions; }
+
+	DECLARE_DELEGATE_RetVal_OneParam(bool, FOnIsClassAllowed, const UClass* /*InClass*/)
+	DECLARE_DELEGATE_RetVal_OneParam(bool, FOnIsClassPathAllowed, const FTopLevelAssetPath& /*InClassPath*/)
+	
+	/** Delegates called to determine whether a class type is allowed */
+	void RegisterIsClassAllowedDelegate(const FName OwnerName, FOnIsClassAllowed Delegate);
+	void UnregisterIsClassAllowedDelegate(const FName OwnerName);
+	bool IsClassAllowed(const UClass* InClass) const;
+	bool HasClassFiltering() const { return IsClassAllowedDelegates.Num() > 0; }
+
+	void RegisterIsClassPathAllowedDelegate(const FName OwnerName, FOnIsClassPathAllowed Delegate);
+	void UnregisterIsClassPathAllowedDelegate(const FName OwnerName);
+	bool IsClassPathAllowed(const FTopLevelAssetPath& InClassPath) const;
+	bool HasClassPathFiltering() const { return IsClassPathAllowedDelegates.Num() > 0; }
+
+	void RegisterIsClassAllowedOnPinDelegate(const FName OwnerName, FOnIsClassAllowed Delegate);
+	void UnregisterIsClassAllowedOnPinDelegate(const FName OwnerName);
+	bool IsClassAllowedOnPin(const UClass* InClass) const;
+	bool HasClassOnPinFiltering() const { return IsClassAllowedOnPinDelegates.Num() > 0; }
+
+	void RegisterIsClassPathAllowedOnPinDelegate(const FName OwnerName, FOnIsClassPathAllowed Delegate);
+	void UnregisterIsClassPathAllowedOnPinDelegate(const FName OwnerName);
+	bool IsClassPathAllowedOnPin(const FTopLevelAssetPath& InClassPath) const;
+	bool HasClassPathOnPinFiltering() const { return IsClassPathAllowedOnPinDelegates.Num() > 0; }
+	
+private:
+	/** All function permissions */
+	FPathPermissionList FunctionPermissions;
+
+	/** All struct permissions */
+	FPathPermissionList StructPermissions;
+
+	/** All enum permissions */
+	FPathPermissionList EnumPermissions;
+
+	/** All pin category permissions */
+	FNamePermissionList PinCategoryPermissions;
+
+	/** Delegates called to determine whether a class type is allowed to be displayed */
+	TMap<FName, FOnIsClassAllowed> IsClassAllowedDelegates;
+	TMap<FName, FOnIsClassPathAllowed> IsClassPathAllowedDelegates;
+
+	/** Delegates called to determine whether a class type is allowed to be displayed on a pin */
+	TMap<FName, FOnIsClassAllowed> IsClassAllowedOnPinDelegates;
+	TMap<FName, FOnIsClassPathAllowed> IsClassPathAllowedOnPinDelegates;
+
 protected:
 	//~ Begin UObject Interface
 	virtual void PostInitProperties();
diff --git a/Engine/Source/Editor/Kismet/Private/BlueprintActionMenuUtils.cpp b/Engine/Source/Editor/Kismet/Private/BlueprintActionMenuUtils.cpp
--- a/Engine/Source/Editor/Kismet/Private/BlueprintActionMenuUtils.cpp
+++ b/Engine/Source/Editor/Kismet/Private/BlueprintActionMenuUtils.cpp
@@ -31,6 +31,7 @@
 #include "ContentBrowserModule.h"
 #include "ComponentAssetBroker.h"
 #include "BlueprintNamespaceHelper.h"
+#include "Editor/EditorEngine.h"
 
 #define LOCTEXT_NAMESPACE "BlueprintActionMenuUtils"
 
@@ -486,6 +487,15 @@ void FBlueprintActionMenuUtils::MakeContextMenu(FBlueprintActionContext const& C
 	// only want bound actions for this menu section
 	LevelActorsFilter.AddRejectionTest(FBlueprintActionFilter::FRejectionTestDelegate::CreateStatic(IsUnBoundSpawner));
 
+	// Build asset reference filter
+	FAssetReferenceFilterContext AssetReferenceFilterContext;
+	for (UBlueprint* Blueprint : Context.Blueprints)
+	{
+		AssetReferenceFilterContext.ReferencingAssets.Add(FAssetData(Blueprint));
+	}
+
+	MainMenuFilter.AssetReferenceFilter = GEditor->MakeAssetReferenceFilter(AssetReferenceFilterContext);
+
 	const UBlueprintEditorSettings* BlueprintSettings = GetDefault<UBlueprintEditorSettings>();
 	bool bCanOperateOnLevelActors = bIsContextSensitive && (Context.Pins.Num() == 0);
 	bool bCanHaveActorComponents  = bIsContextSensitive;
diff --git a/Engine/Source/Editor/UnrealEd/Private/Kismet2/BlueprintEditorUtils.cpp b/Engine/Source/Editor/UnrealEd/Private/Kismet2/BlueprintEditorUtils.cpp
--- a/Engine/Source/Editor/UnrealEd/Private/Kismet2/BlueprintEditorUtils.cpp
+++ b/Engine/Source/Editor/UnrealEd/Private/Kismet2/BlueprintEditorUtils.cpp
@@ -133,6 +133,7 @@
 #include "Elements/Interfaces/TypedElementHierarchyInterface.h"
 #include "Elements/Interfaces/TypedElementObjectInterface.h"
 #include "Misc/MessageDialog.h"
+#include "UObject/FastReferenceCollector.h"
 
 #define LOCTEXT_NAMESPACE "Blueprint"
 
@@ -10446,6 +10447,98 @@ void FBlueprintEditorUtils::RecombineNestedSubPins(UK2Node* Node)
 	}
 }
 
+static FAutoConsoleCommand AuditThreadSafeFunctions(
+	TEXT("bp.AuditThreadSafeFunctions"),
+	TEXT("Audit currently loaded thread safe functions. Writes results to the log."),
+	FConsoleCommandDelegate::CreateLambda([]()
+		{
+			UE_LOG(LogBlueprint, Display, TEXT("--- BEGIN audit all BlueprintThreadSafe functions ---"));
+			UE_LOG(LogBlueprint, Display, TEXT("Name, Path, Type, BPCallType"));
+
+			for (TObjectIterator<UFunction> It; It; ++It)
+			{
+				UFunction* Function = *It;
+				if (FBlueprintEditorUtils::HasFunctionBlueprintThreadSafeMetaData(Function))
+				{
+					const TCHAR* Native = Function->HasAnyFunctionFlags(FUNC_Native) ? TEXT("Native") : TEXT("Script");
+					const TCHAR* Purity = [Function]()
+					{
+						if (Function->HasAnyFunctionFlags(FUNC_BlueprintCallable))
+						{
+							return Function->HasAllFunctionFlags(FUNC_BlueprintPure | FUNC_BlueprintCallable) ? TEXT("Pure") : TEXT("Callable");
+						}
+
+						return TEXT("NotCallable");
+					}();
+					UE_LOG(LogBlueprint, Display, TEXT("%s, %s, %s, %s"), *Function->GetName(), *Function->GetPathName(), Native, Purity);
+				}
+			}
+
+			UE_LOG(LogBlueprint, Display, TEXT("--- END audit all BlueprintThreadSafe functions ---"));
+		}));
+
+static FAutoConsoleCommand AuditFunctionCallsForBlueprint(
+	TEXT("bp.AuditFunctionCallsForBlueprint"),
+	TEXT("Audit all functions called by a specified blueprint. Single argument supplies the asset to audit. Writes results to the log."),
+	FConsoleCommandWithArgsDelegate::CreateLambda([](const TArray<FString>& InArgs)
+		{
+			if (InArgs.Num() != 1)
+			{
+				return;
+			}
+
+			// Find our Blueprint & load it
+			UBlueprint* Blueprint = LoadObject<UBlueprint>(nullptr, *InArgs[0]);
+			if (Blueprint == nullptr)
+			{
+				UE_LOG(LogBlueprint, Warning, TEXT("--- Could not load Blueprint %s ---"), *InArgs[0]);
+				return;
+			}
+
+			if (Blueprint->GeneratedClass == nullptr)
+			{
+				UE_LOG(LogBlueprint, Warning, TEXT("--- Blueprint %s as a null GeneratedClass ---"), *InArgs[0]);
+				return;
+			}
+
+			UE_LOG(LogBlueprint, Display, TEXT("--- BEGIN audit function calls for Blueprint %s ---"), *InArgs[0]);
+			UE_LOG(LogBlueprint, Display, TEXT("Name, Path, Type, BPCallType"));
+
+			struct FFunctionReferenceProcessor : public FSimpleReferenceProcessorBase
+			{
+				FORCEINLINE void HandleTokenStreamObjectReference(FGCArrayStruct& ObjectsToSerializeStruct, UObject* ReferencingObject, UObject*& Object, const int32 TokenIndex, const EGCTokenType TokenType, bool bAllowReferenceElimination)
+				{
+					if (UFunction* Function = Cast<UFunction>(Object))
+					{
+						const TCHAR* Native = Function->HasAnyFunctionFlags(FUNC_Native) ? TEXT("Native") : TEXT("Script");
+						const TCHAR* Purity = [Function]()
+						{
+							if (Function->HasAnyFunctionFlags(FUNC_BlueprintCallable))
+							{
+								return Function->HasAllFunctionFlags(FUNC_BlueprintPure | FUNC_BlueprintCallable) ? TEXT("Pure") : TEXT("Callable");
+							}
+
+							return TEXT("NotCallable");
+						}();
+						UE_LOG(LogBlueprint, Display, TEXT("%s, %s, %s, %s"), *Function->GetName(), *Function->GetOuterUClass()->GetPathName(), Native, Purity);
+					}
+				}
+			} Processor;
+
+			TFastReferenceCollector<
+				FFunctionReferenceProcessor,
+				TDefaultReferenceCollector<FFunctionReferenceProcessor>,
+				FGCArrayPool,
+				EFastReferenceCollectorOptions::AutogenerateTokenStream | EFastReferenceCollectorOptions::ProcessNoOpTokens
+			> ReferenceCollector(Processor, FGCArrayPool::Get());
+			FGCArrayStruct ArrayStruct;
+			ArrayStruct.ObjectsToSerialize.Add(Blueprint->GeneratedClass);
+			ReferenceCollector.CollectReferences(ArrayStruct);
+
+			UE_LOG(LogBlueprint, Display, TEXT("--- END audit all BlueprintThreadSafe functions ---"));
+		}));
+
+
 bool FBlueprintEditorUtils::HasFunctionBlueprintThreadSafeMetaData(const UFunction* InFunction)
 {
 	if(InFunction)
@@ -10454,8 +10547,8 @@ bool FBlueprintEditorUtils::HasFunctionBlueprintThreadSafeMetaData(const UFuncti
 		const bool bHasNotThreadSafeMetaData = InFunction->HasMetaData(FBlueprintMetadata::MD_NotThreadSafe);
 		const bool bClassHasThreadSafeMetaData = InFunction->GetOwnerClass() && InFunction->GetOwnerClass()->HasMetaData(FBlueprintMetadata::MD_ThreadSafe);
 
-		// Native (or BP event) functions need to just have the correct class/function metadata
-		const bool bThreadSafeNative = InFunction->HasAnyFunctionFlags(FUNC_Native | FUNC_BlueprintEvent) && (bHasThreadSafeMetaData || (bClassHasThreadSafeMetaData && !bHasNotThreadSafeMetaData));
+		// Native functions need to just have the correct class/function metadata
+		const bool bThreadSafeNative = InFunction->HasAnyFunctionFlags(FUNC_Native) && (bHasThreadSafeMetaData || (bClassHasThreadSafeMetaData && !bHasNotThreadSafeMetaData));
 
 		// Script functions get their flag propagated from their entry point, and dont pay heed to class metadata
 		const bool bThreadSafeScript = !InFunction->HasAnyFunctionFlags(FUNC_Native) && bHasThreadSafeMetaData;
