CIS fix (resolved merge issue).

#rnx

diff --git a/Engine/Source/Editor/BlueprintGraph/Public/BlueprintEditorSettings.h b/Engine/Source/Editor/BlueprintGraph/Public/BlueprintEditorSettings.h
--- a/Engine/Source/Editor/BlueprintGraph/Public/BlueprintEditorSettings.h
+++ b/Engine/Source/Editor/BlueprintGraph/Public/BlueprintEditorSettings.h
@@ -279,41 +279,9 @@ public:
 	 */
 	UPROPERTY(EditAnywhere, config, Category=Play, meta=(AllowAbstract))
 	TArray<TSoftClassPtr<UObject>> BaseClassesToAllowRecompilingDuringPlayInEditor;
-
-	/** Get allowed structs permissions list */
-	FPathPermissionList& GetStructPermissions() { return StructPermissions; }
-
-	/** Get allowed enums permissions list */
-	FPathPermissionList& GetEnumPermissions() { return EnumPermissions; }
-
-	/** Get allowed pin categories permissions list */
-	FNamePermissionList& GetPinCategoryPermissions() { return PinCategoryPermissions; }
-
-	DECLARE_DELEGATE_RetVal_OneParam(bool, FOnIsClassAllowed, const UClass* /*InClass*/)
-	DECLARE_DELEGATE_RetVal_OneParam(bool, FOnIsClassPathAllowed, const FTopLevelAssetPath& /*InClassPath*/)
-	
-	/** Delegates called to determine whether a class type is allowed */
-	void RegisterIsClassAllowedDelegate(const FName OwnerName, FOnIsClassAllowed Delegate);
-	void UnregisterIsClassAllowedDelegate(const FName OwnerName);
-	bool IsClassAllowed(const UClass* InClass) const;
-	bool HasClassFiltering() const { return IsClassAllowedDelegates.Num() > 0; }
-
-	void RegisterIsClassPathAllowedDelegate(const FName OwnerName, FOnIsClassPathAllowed Delegate);
-	void UnregisterIsClassPathAllowedDelegate(const FName OwnerName);
-	bool IsClassPathAllowed(const FTopLevelAssetPath& InClassPath) const;
-	bool HasClassPathFiltering() const { return IsClassPathAllowedDelegates.Num() > 0; }
-
-	void RegisterIsClassAllowedOnPinDelegate(const FName OwnerName, FOnIsClassAllowed Delegate);
-	void UnregisterIsClassAllowedOnPinDelegate(const FName OwnerName);
-	bool IsClassAllowedOnPin(const UClass* InClass) const;
-	bool HasClassOnPinFiltering() const { return IsClassAllowedOnPinDelegates.Num() > 0; }
-
-	void RegisterIsClassPathAllowedOnPinDelegate(const FName OwnerName, FOnIsClassPathAllowed Delegate);
-	void UnregisterIsClassPathAllowedOnPinDelegate(const FName OwnerName);
-	bool IsClassPathAllowedOnPin(const FTopLevelAssetPath& InClassPath) const;
-	bool HasClassPathOnPinFiltering() const { return IsClassPathAllowedOnPinDelegates.Num() > 0; }
 	
-private:
+	/** Get allowed functions permissions list */
+	FPathPermissionList& GetFunctionPermissions() { return FunctionPermissions; }
 
 	/** Get allowed structs permissions list */
 	FPathPermissionList& GetStructPermissions() { return StructPermissions; }
@@ -347,7 +315,7 @@ private:
 	void UnregisterIsClassPathAllowedOnPinDelegate(const FName OwnerName);
 	bool IsClassPathAllowedOnPin(const FTopLevelAssetPath& InClassPath) const;
 	bool HasClassPathOnPinFiltering() const { return IsClassPathAllowedOnPinDelegates.Num() > 0; }
-	
+
 private:
 	/** All function permissions */
 	FPathPermissionList FunctionPermissions;
