Fix bad merge in CL 20892543

Removed settings delegate that came from another change

diff --git a/Engine/Source/Editor/BlueprintGraph/Public/BlueprintEditorSettings.h b/Engine/Source/Editor/BlueprintGraph/Public/BlueprintEditorSettings.h
--- a/Engine/Source/Editor/BlueprintGraph/Public/BlueprintEditorSettings.h
+++ b/Engine/Source/Editor/BlueprintGraph/Public/BlueprintEditorSettings.h
@@ -280,17 +280,6 @@ public:
 	UPROPERTY(EditAnywhere, config, Category=Play, meta=(AllowAbstract))
 	TArray<TSoftClassPtr<UObject>> BaseClassesToAllowRecompilingDuringPlayInEditor;
 
-	void UnregisterOnUpdateSettings(FDelegateHandle Object)
-	{
-		OnSettingsChange.Remove(Object);
-	}
-	
-
-	void UnregisterOnUpdateSettings(FDelegateHandle Object)
-	{
-		OnSettingsChange.Remove(Object);
-	}
-	
 	/** Get allowed functions permissions list */
 	FPathPermissionList& GetFunctionPermissions() { return FunctionPermissions; }
 
