Unshelved from pending changelist '20809968':

Additional right-click Blueprint graph context menu profiling.

diff --git a/Engine/Source/Editor/BlueprintGraph/Private/BlueprintActionDatabase.cpp b/Engine/Source/Editor/BlueprintGraph/Private/BlueprintActionDatabase.cpp
--- a/Engine/Source/Editor/BlueprintGraph/Private/BlueprintActionDatabase.cpp
+++ b/Engine/Source/Editor/BlueprintGraph/Private/BlueprintActionDatabase.cpp
@@ -1338,6 +1338,8 @@ void FBlueprintActionDatabase::DeferredRemoveEntry(FObjectKey const& InKey)
 //------------------------------------------------------------------------------
 void FBlueprintActionDatabase::RefreshAll()
 {
+	TRACE_CPUPROFILER_EVENT_SCOPE(FBlueprintActionDatabase::RefreshAll);
+
 	TGuardValue<bool> ScopedInitialization(BlueprintActionDatabaseImpl::bIsInitializing, true);
 	BlueprintActionDatabaseImpl::bRefreshAllRequested = false;
 
@@ -1410,6 +1412,8 @@ void FBlueprintActionDatabase::RefreshWorlds()
 //------------------------------------------------------------------------------
 void FBlueprintActionDatabase::RefreshClassActions(UClass* const Class)
 {
+	TRACE_CPUPROFILER_EVENT_SCOPE(FBlueprintActionDatabase::RefreshClassActions);
+
 	// Early out if the class is filtered
 	if (!IsClassAllowed(Class, EPermissionsContext::Asset))
 	{
@@ -1534,6 +1538,8 @@ void FBlueprintActionDatabase::RefreshClassActions(UClass* const Class)
 //------------------------------------------------------------------------------
 void FBlueprintActionDatabase::RefreshAssetActions(UObject* const AssetObject)
 {
+	TRACE_CPUPROFILER_EVENT_SCOPE(FBlueprintActionDatabase::RefreshAssetActions);
+
 	using namespace BlueprintActionDatabaseImpl;
 
 	// this method is very expensive and is only for blueprint editor functionality
@@ -1676,6 +1682,8 @@ void FBlueprintActionDatabase::RefreshAssetActions(UObject* const AssetObject)
 //------------------------------------------------------------------------------
 void FBlueprintActionDatabase::RefreshComponentActions()
 {
+	TRACE_CPUPROFILER_EVENT_SCOPE(FBlueprintActionDatabase::RefreshComponentActions);
+
 	check(ComponentTypes);
 	FActionList& ClassActionList = ActionRegistry.FindOrAdd(UBlueprintComponentNodeSpawner::StaticClass());
 	ClassActionList.Empty(ComponentTypes->Num());
@@ -1797,6 +1805,8 @@ FBlueprintActionDatabase::FActionRegistry const& FBlueprintActionDatabase::GetAl
 //------------------------------------------------------------------------------
 void FBlueprintActionDatabase::RegisterAllNodeActions(FBlueprintActionDatabaseRegistrar& Registrar)
 {
+	TRACE_CPUPROFILER_EVENT_SCOPE(FBlueprintActionDatabase::RegisterAllNodeActions);
+
 	// nodes may have actions they wish to add for this asset
 	TArray<UClass*> NodeClassList;
 	GetDerivedClasses(UK2Node::StaticClass(), NodeClassList);
diff --git a/Engine/Source/Editor/BlueprintGraph/Private/BlueprintActionFilter.cpp b/Engine/Source/Editor/BlueprintGraph/Private/BlueprintActionFilter.cpp
--- a/Engine/Source/Editor/BlueprintGraph/Private/BlueprintActionFilter.cpp
+++ b/Engine/Source/Editor/BlueprintGraph/Private/BlueprintActionFilter.cpp
@@ -2105,59 +2105,62 @@ FBlueprintActionFilter::FBlueprintActionFilter(const EFlags InFlags /*= BPFILTER
 	//       here are ran last)
 	//
 
+#define ADD_REJECTION_TEST(TestFunc, ...) \
+	AddRejectionTest(FRejectionTestDelegate::CreateStatic(TestFunc, ##__VA_ARGS__), TEXT(#TestFunc))
+
 	// add first the most expensive tests (they will be ran last, and therefore
 	// should be operating on a smaller subset of node-spawners)
 	//
 	// this test in-particular spawns a template-node and then calls 
 	// AllocateDefaultPins() which is costly, so it should be very last!
-	AddRejectionTest(FRejectionTestDelegate::CreateStatic(IsNodeTemplateSelfFiltered));
-	AddRejectionTest(FRejectionTestDelegate::CreateStatic(IsIncompatibleAnimNotification));
-	AddRejectionTest(FRejectionTestDelegate::CreateStatic(IsAssetPermissionNotGranted));
-	AddRejectionTest(FRejectionTestDelegate::CreateStatic(IsMissingMatchingPinParam));
-	AddRejectionTest(FRejectionTestDelegate::CreateStatic(IsMissmatchedPropertyType));
-	AddRejectionTest(FRejectionTestDelegate::CreateStatic(IsFunctionMissingPinParam));
-	AddRejectionTest(FRejectionTestDelegate::CreateStatic(IsIncompatibleLatentNode));
-	AddRejectionTest(FRejectionTestDelegate::CreateStatic(IsIncompatibleImpureNode));
-	AddRejectionTest(FRejectionTestDelegate::CreateStatic(IsPropertyAccessorNode));
+	ADD_REJECTION_TEST(IsNodeTemplateSelfFiltered);
+	ADD_REJECTION_TEST(IsIncompatibleAnimNotification);
+	ADD_REJECTION_TEST(IsAssetPermissionNotGranted);
+	ADD_REJECTION_TEST(IsMissingMatchingPinParam);
+	ADD_REJECTION_TEST(IsMissmatchedPropertyType);
+	ADD_REJECTION_TEST(IsFunctionMissingPinParam);
+	ADD_REJECTION_TEST(IsIncompatibleLatentNode);
+	ADD_REJECTION_TEST(IsIncompatibleImpureNode);
+	ADD_REJECTION_TEST(IsPropertyAccessorNode);
 	if (HasAnyFlags(BPFILTER_RejectIncompatibleThreadSafety))
 	{
-		AddRejectionTest(FRejectionTestDelegate::CreateStatic(IsThreadSafetyIncompatible));
+		ADD_REJECTION_TEST(IsThreadSafetyIncompatible);
 	}
 	
-	AddRejectionTest(FRejectionTestDelegate::CreateStatic(IsActionHiddenByConfig));
-	AddRejectionTest(FRejectionTestDelegate::CreateStatic(IsFieldCategoryHidden));
+	ADD_REJECTION_TEST(IsActionHiddenByConfig);
+	ADD_REJECTION_TEST(IsFieldCategoryHidden);
 	if (HasAnyFlags(BPFILTER_RejectGlobalFields | BPFILTER_RejectNonImportedFields))
 	{
-		AddRejectionTest(FRejectionTestDelegate::CreateStatic(IsRejectedGlobalField));
+		ADD_REJECTION_TEST(IsRejectedGlobalField);
 	}
 	
-	AddRejectionTest(FRejectionTestDelegate::CreateStatic(IsFieldInaccessible));
-	AddRejectionTest(FRejectionTestDelegate::CreateStatic(IsNotSubClassCast));
-	AddRejectionTest(FRejectionTestDelegate::CreateStatic(IsEventUnimplementable));
-	AddRejectionTest(FRejectionTestDelegate::CreateStatic(IsPermissionNotGranted));
-	AddRejectionTest(FRejectionTestDelegate::CreateStatic(IsRestrictedClassMember));
-	AddRejectionTest(FRejectionTestDelegate::CreateStatic(IsIncompatibleWithGraphType));
-	AddRejectionTest(FRejectionTestDelegate::CreateStatic(IsSchemaIncompatible));
-	AddRejectionTest(FRejectionTestDelegate::CreateStatic(IsExtraneousInterfaceCall));
-	AddRejectionTest(FRejectionTestDelegate::CreateStatic(IsIncompatibleMacroInstance));
+	ADD_REJECTION_TEST(IsFieldInaccessible);
+	ADD_REJECTION_TEST(IsNotSubClassCast);
+	ADD_REJECTION_TEST(IsEventUnimplementable);
+	ADD_REJECTION_TEST(IsPermissionNotGranted);
+	ADD_REJECTION_TEST(IsRestrictedClassMember);
+	ADD_REJECTION_TEST(IsIncompatibleWithGraphType);
+	ADD_REJECTION_TEST(IsSchemaIncompatible);
+	ADD_REJECTION_TEST(IsExtraneousInterfaceCall);
+	ADD_REJECTION_TEST(IsIncompatibleMacroInstance);
 
 	if (!HasAnyFlags(BPFILTER_PermitDeprecated))
 	{
-		AddRejectionTest(FRejectionTestDelegate::CreateStatic(IsDeprecated));
+		ADD_REJECTION_TEST(IsDeprecated);
 	}
 
-	AddRejectionTest(FRejectionTestDelegate::CreateStatic(IsFilteredNodeType));
-	AddRejectionTest(FRejectionTestDelegate::CreateStatic(IsNonTargetMember));
-	AddRejectionTest(FRejectionTestDelegate::CreateStatic(IsUnBoundBindingSpawner));
-	AddRejectionTest(FRejectionTestDelegate::CreateStatic(IsOutOfScopeLocalVariable));
-	AddRejectionTest(FRejectionTestDelegate::CreateStatic(IsLevelScriptActionValid));
+	ADD_REJECTION_TEST(IsFilteredNodeType);
+	ADD_REJECTION_TEST(IsNonTargetMember);
+	ADD_REJECTION_TEST(IsUnBoundBindingSpawner);
+	ADD_REJECTION_TEST(IsOutOfScopeLocalVariable);
+	ADD_REJECTION_TEST(IsLevelScriptActionValid);
 
-	AddRejectionTest(FRejectionTestDelegate::CreateStatic(IsHiddenInNonEditorBlueprint));	
+	ADD_REJECTION_TEST(IsHiddenInNonEditorBlueprint);
 
 
 	// added as the first rejection test, so that we don't operate on stale 
 	// (TRASH/REINST) class fields
-	AddRejectionTest(FRejectionTestDelegate::CreateStatic(IsStaleFieldAction));
+	ADD_REJECTION_TEST(IsStaleFieldAction);
 }
 
 //------------------------------------------------------------------------------
@@ -2183,17 +2186,31 @@ void FBlueprintActionFilter::Add(TArray<FTargetClassFilterData>& ToArray, UClass
 }
 
 //------------------------------------------------------------------------------
-void FBlueprintActionFilter::AddRejectionTest(FRejectionTestDelegate IsFilteredDelegate)
+void FBlueprintActionFilter::AddRejectionTest(FRejectionTestDelegate RejectionTestDelegate, const TCHAR* RejectionTestName)
 {
-	if (IsFilteredDelegate.IsBound())
+	if (RejectionTestDelegate.IsBound())
 	{
-		FilterTests.Add(IsFilteredDelegate);
+		FilterTests.Add(RejectionTestDelegate);
+
+#if ENABLE_BLUEPRINT_ACTION_FILTER_PROFILING
+		FFilterTestProfileRecord FilterTestProfileRecord;
+		
+		FilterTestProfileRecord.TestName = RejectionTestName ? FString(RejectionTestName) : RejectionTestDelegate.TryGetBoundFunctionName().ToString();
+		if (FilterTestProfileRecord.TestName.IsEmpty())
+		{
+			FilterTestProfileRecord.TestName = FString::Printf(TEXT("[FuncPtr:0x%016x]"), RejectionTestDelegate.GetBoundProgramCounterForTimerManager());
+		}
+
+		FilterTestProfiles.Add(FilterTests.Num() - 1, MoveTemp(FilterTestProfileRecord));
+#endif	// ENABLE_BLUEPRINT_ACTION_FILTER_PROFILING
 	}
 }
 
 //------------------------------------------------------------------------------
 bool FBlueprintActionFilter::IsFiltered(FBlueprintActionInfo& BlueprintAction)
 {
+	TRACE_CPUPROFILER_EVENT_SCOPE(FBlueprintActionFilter::IsFiltered);
+
 	bool bIsFiltered = IsFilteredByThis(BlueprintAction);
 	if (!bIsFiltered)
 	{
@@ -2245,6 +2262,10 @@ bool FBlueprintActionFilter::IsFilteredByThis(FBlueprintActionInfo& BlueprintAct
 // 		bDebugBreak = true;
 // 	}
 
+#if ENABLE_BLUEPRINT_ACTION_FILTER_PROFILING
+	static TMap<FString, uint32> FilterTestToCpuSpecIds;
+#endif	// ENABLE_BLUEPRINT_ACTION_FILTER_PROFILING
+
 	bool bIsFiltered = false;
 	// iterate backwards so that custom user test are ran first (and the slow
 	// internal tests are ran last).
@@ -2253,20 +2274,68 @@ bool FBlueprintActionFilter::IsFilteredByThis(FBlueprintActionInfo& BlueprintAct
 		FRejectionTestDelegate const& RejectionTestDelegate = FilterTests[TestIndex];
 		check(RejectionTestDelegate.IsBound());
 
-		if (RejectionTestDelegate.Execute(FilterRef, BlueprintAction))
+#if ENABLE_BLUEPRINT_ACTION_FILTER_PROFILING
+		FFilterTestProfileRecord& FilterTestProfileRecord = FilterTestProfiles.FindOrAdd(TestIndex);
+		uint32& CpuSpecId = FilterTestToCpuSpecIds.FindOrAdd(FilterTestProfileRecord.TestName);
+		FCpuProfilerTrace::FEventScope ProfilerEvent(CpuSpecId, *FString::Printf(TEXT("IsFiltered_%s"), *FilterTestProfileRecord.TestName), true, __FILE__, __LINE__);
+
+		const double StartTime = FPlatformTime::Seconds();
+#endif	// ENABLE_BLUEPRINT_ACTION_FILTER_PROFILING
+
+		bIsFiltered = RejectionTestDelegate.Execute(FilterRef, BlueprintAction);
+
+#if ENABLE_BLUEPRINT_ACTION_FILTER_PROFILING
+		FilterTestProfileRecord.NumIterations += 1;
+		FilterTestProfileRecord.TotalTimeMs += (FPlatformTime::Seconds() - StartTime) * 1000.0f;
+		if (bIsFiltered)
+		{
+			FilterTestProfileRecord.NumFilteredOut += 1;
+		}
+#endif	// ENABLE_BLUEPRINT_ACTION_FILTER_PROFILING
+
+		if (bIsFiltered)
 		{
-			bIsFiltered = true;
 			break;
 		}
 	}
 
 	if (!bIsFiltered)
 	{
-		for (auto& ExtraRejectionTest : BluprintGraphModule->GetExtendedActionMenuFilters())
+#if ENABLE_BLUEPRINT_ACTION_FILTER_PROFILING
+		int32 ExtraTestProfileRecordIndex = FilterTests.Num();
+#endif	// ENABLE_BLUEPRINT_ACTION_FILTER_PROFILING
+
+		for (const FBlueprintGraphModule::FActionMenuRejectionTest& ExtraRejectionTest : BluprintGraphModule->GetExtendedActionMenuFilters())
 		{
-			if (ExtraRejectionTest.Execute(FilterRef, BlueprintAction))
+#if ENABLE_BLUEPRINT_ACTION_FILTER_PROFILING
+			FFilterTestProfileRecord& ExtraTestProfileRecord = FilterTestProfiles.FindOrAdd(ExtraTestProfileRecordIndex++);
+			if (ExtraTestProfileRecord.TestName.IsEmpty())
+			{
+				ExtraTestProfileRecord.TestName = ExtraRejectionTest.TryGetBoundFunctionName().ToString();
+				if (ExtraTestProfileRecord.TestName.IsEmpty())
+				{
+					ExtraTestProfileRecord.TestName = FString::Printf(TEXT("[FuncPtr:0x%016x]"), ExtraRejectionTest.GetBoundProgramCounterForTimerManager());
+				}
+			}
+			uint32& CpuSpecId = FilterTestToCpuSpecIds.FindOrAdd(ExtraTestProfileRecord.TestName);
+			FCpuProfilerTrace::FEventScope ProfilerEvent(CpuSpecId, *FString::Printf(TEXT("Extended_IsFiltered_%s"), *ExtraTestProfileRecord.TestName), true, __FILE__, __LINE__);
+
+			const double StartTime = FPlatformTime::Seconds();
+#endif	// ENABLE_BLUEPRINT_ACTION_FILTER_PROFILING
+
+			bIsFiltered = ExtraRejectionTest.Execute(FilterRef, BlueprintAction);
+
+#if ENABLE_BLUEPRINT_ACTION_FILTER_PROFILING
+			ExtraTestProfileRecord.NumIterations += 1;
+			ExtraTestProfileRecord.TotalTimeMs += (FPlatformTime::Seconds() - StartTime) * 1000.0f;
+			if (bIsFiltered)
+			{
+				ExtraTestProfileRecord.NumFilteredOut += 1;
+			}
+#endif	// ENABLE_BLUEPRINT_ACTION_FILTER_PROFILING
+
+			if (bIsFiltered)
 			{
-				bIsFiltered = true;
 				break;
 			}
 		}
@@ -2274,3 +2343,21 @@ bool FBlueprintActionFilter::IsFilteredByThis(FBlueprintActionInfo& BlueprintAct
 
 	return bIsFiltered;
 }
+
+//------------------------------------------------------------------------------
+#if ENABLE_BLUEPRINT_ACTION_FILTER_PROFILING
+TArray<FString> FBlueprintActionFilter::GetFilterTestProfile()
+{
+	TArray<FString> Result;
+
+	int32 TestIdx = 0;
+	for (int32 TestIndex = FilterTests.Num() - 1; TestIndex >= 0; --TestIndex)
+	{
+		Result.Add(FString::Printf(TEXT("[%02d]: %s | %d tested | %d failed | %.02f ms | %.02f us (per test)"), FilterTests.Num() - TestIndex, *FilterTestProfiles[TestIndex].TestName, FilterTestProfiles[TestIndex].NumIterations, FilterTestProfiles[TestIndex].NumFilteredOut, FilterTestProfiles[TestIndex].TotalTimeMs, FilterTestProfiles[TestIndex].TotalTimeMs / FilterTestProfiles[TestIndex].NumIterations * 1000.0f));
+	}
+
+	FilterTestProfiles.Reset();
+
+	return Result;
+}
+#endif	// ENABLE_BLUEPRINT_ACTION_FILTER_PROFILING
\ No newline at end of file
diff --git a/Engine/Source/Editor/BlueprintGraph/Private/BlueprintNodeSpawner.cpp b/Engine/Source/Editor/BlueprintGraph/Private/BlueprintNodeSpawner.cpp
--- a/Engine/Source/Editor/BlueprintGraph/Private/BlueprintNodeSpawner.cpp
+++ b/Engine/Source/Editor/BlueprintGraph/Private/BlueprintNodeSpawner.cpp
@@ -100,6 +100,8 @@ void UBlueprintNodeSpawner::Prime()
 //------------------------------------------------------------------------------
 FBlueprintActionUiSpec const& UBlueprintNodeSpawner::PrimeDefaultUiSpec(UEdGraph* TargetGraph) const
 {
+	TRACE_CPUPROFILER_EVENT_SCOPE(UBlueprintNodeSpawner::PrimeDefaultUiSpec);
+
 	bool bTemplateNodeFetched = false;
 	UEdGraphNode* NodeTemplate = nullptr;
 	// @TODO: boo! const cast... all to make this callable from GetUiSpec()
@@ -254,6 +256,8 @@ FBlueprintNodeSignature UBlueprintNodeSpawner::GetSpawnerSignature() const
 //------------------------------------------------------------------------------
 FBlueprintActionUiSpec UBlueprintNodeSpawner::GetUiSpec(FBlueprintActionContext const& Context, FBindingSet const& Bindings) const
 {
+	TRACE_CPUPROFILER_EVENT_SCOPE(UBlueprintNodeSpawner::GetUiSpec);
+
 	FBlueprintActionUiSpec MenuSignature = PrimeDefaultUiSpec();
 	DynamicUiSignatureGetter.ExecuteIfBound(Context, Bindings, &MenuSignature);
 	return MenuSignature;
@@ -274,6 +278,8 @@ UEdGraphNode* UBlueprintNodeSpawner::GetCachedTemplateNode() const
 //------------------------------------------------------------------------------
 UEdGraphNode* UBlueprintNodeSpawner::GetTemplateNode(UEdGraph* TargetGraph, FBindingSet const& Bindings) const
 {       
+	TRACE_CPUPROFILER_EVENT_SCOPE(UBlueprintNodeSpawner::GetTemplateNode);
+	
 	UEdGraphNode* TemplateNode = BlueprintNodeSpawnerImpl::GetSharedTemplateCache()->GetNodeTemplate(this, TargetGraph);
 
 	if (TemplateNode && Bindings.Num() > 0) 
@@ -297,6 +303,8 @@ void UBlueprintNodeSpawner::ClearCachedTemplateNode() const
 
 bool UBlueprintNodeSpawner::IsTemplateNodeFilteredOut(FBlueprintActionFilter const& Filter) const
 {
+	TRACE_CPUPROFILER_EVENT_SCOPE(UBlueprintNodeSpawner::IsTemplateNodeFilteredOut);
+
 	bool bIsFilteredOut = false;
 	if(UK2Node* NodeTemplate = Cast<UK2Node>(GetTemplateNode()))
 	{
diff --git a/Engine/Source/Editor/BlueprintGraph/Private/BlueprintNodeTemplateCache.cpp b/Engine/Source/Editor/BlueprintGraph/Private/BlueprintNodeTemplateCache.cpp
--- a/Engine/Source/Editor/BlueprintGraph/Private/BlueprintNodeTemplateCache.cpp
+++ b/Engine/Source/Editor/BlueprintGraph/Private/BlueprintNodeTemplateCache.cpp
@@ -255,6 +255,8 @@ FBlueprintNodeTemplateCache::FBlueprintNodeTemplateCache()
 //------------------------------------------------------------------------------
 UEdGraphNode* FBlueprintNodeTemplateCache::GetNodeTemplate(UBlueprintNodeSpawner const* NodeSpawner, UEdGraph* TargetGraph)
 {
+	TRACE_CPUPROFILER_EVENT_SCOPE(FBlueprintNodeTemplateCache::GetNodeTemplate);
+
 	using namespace BlueprintNodeTemplateCacheImpl;
 
 	bool bIsOverMemCap = false;
diff --git a/Engine/Source/Editor/BlueprintGraph/Public/BlueprintActionFilter.h b/Engine/Source/Editor/BlueprintGraph/Public/BlueprintActionFilter.h
--- a/Engine/Source/Editor/BlueprintGraph/Public/BlueprintActionFilter.h
+++ b/Engine/Source/Editor/BlueprintGraph/Public/BlueprintActionFilter.h
@@ -11,6 +11,8 @@
 #include "BlueprintNodeBinder.h"
 #include "BlueprintGraphModule.h"
 
+#define ENABLE_BLUEPRINT_ACTION_FILTER_PROFILING (!(UE_BUILD_SHIPPING || UE_BUILD_TEST) && 1)
+
 class UBlueprint;
 class UBlueprintNodeSpawner;
 class UEdGraph;
@@ -271,8 +273,9 @@ public:
 	 * because it is more optimal to whittle down the list of actions early.
 	 * 
 	 * @param  RejectionTestDelegate	The rejection test you wish to add to this filter.
+	 * @param  RejectionTestName		(Optional) Test name (used for profiling/debugging).
 	 */
-	void AddRejectionTest(FRejectionTestDelegate RejectionTestDelegate);
+	void AddRejectionTest(FRejectionTestDelegate RejectionTestDelegate, const TCHAR* RejectionTestName = nullptr);
 
 	/**
 	 * Query to check and see if the specified action gets filtered out by this 
@@ -307,6 +310,10 @@ public:
 	 * @return This.
 	 */
 	FBlueprintActionFilter const& operator&=(FBlueprintActionFilter const& Rhs);
+	
+#if ENABLE_BLUEPRINT_ACTION_FILTER_PROFILING
+	TArray<FString> GetFilterTestProfile();
+#endif	// ENABLE_BLUEPRINT_ACTION_FILTER_PROFILING
 
 	/**
 	 * @return TRUE if any of the given configuration flags are set on this filter.
@@ -345,6 +352,18 @@ private:
 
 	/** Alternative filters to be logically or'd in with the IsFilteredByThis() result. */
 	TArray<FBlueprintActionFilter> OrFilters;
+
+#if ENABLE_BLUEPRINT_ACTION_FILTER_PROFILING
+	struct FFilterTestProfileRecord
+	{
+		FString TestName;
+		float TotalTimeMs = 0.0f;
+		int32 NumIterations = 0;
+		int32 NumFilteredOut = 0;
+	};
+
+	mutable TMap<int32, FFilterTestProfileRecord> FilterTestProfiles;
+#endif	// ENABLE_BLUEPRINT_ACTION_FILTER_PROFILING
 };
 
 ENUM_CLASS_FLAGS(FBlueprintActionFilter::EFlags);
\ No newline at end of file
diff --git a/Engine/Source/Editor/GraphEditor/Private/GraphActionNode.cpp b/Engine/Source/Editor/GraphEditor/Private/GraphActionNode.cpp
--- a/Engine/Source/Editor/GraphEditor/Private/GraphActionNode.cpp
+++ b/Engine/Source/Editor/GraphEditor/Private/GraphActionNode.cpp
@@ -211,6 +211,8 @@ TSharedPtr<FGraphActionNode> FGraphActionNode::AddSection(int32 InGrouping, int3
 //------------------------------------------------------------------------------
 void FGraphActionNode::SortChildren(bool bAlphabetically/* = true*/, bool bRecursive/* = true*/)
 {
+	TRACE_CPUPROFILER_EVENT_SCOPE(SGraphActionMenu::GenerateFilteredItems_SortNodes);
+
 	if (bRecursive)
 	{
 		for (TSharedPtr<FGraphActionNode>& ChildNode : Children)
diff --git a/Engine/Source/Editor/GraphEditor/Private/SGraphActionMenu.cpp b/Engine/Source/Editor/GraphEditor/Private/SGraphActionMenu.cpp
--- a/Engine/Source/Editor/GraphEditor/Private/SGraphActionMenu.cpp
+++ b/Engine/Source/Editor/GraphEditor/Private/SGraphActionMenu.cpp
@@ -365,12 +365,19 @@ void SGraphActionMenu::Construct( const FArguments& InArgs, bool bIsReadOnly/* =
 
 void SGraphActionMenu::RefreshAllActions(bool bPreserveExpansion, bool bHandleOnSelectionEvent/*=true*/)
 {
+	TRACE_CPUPROFILER_EVENT_SCOPE(SGraphActionMenu::RefreshAllActions);
+
 	// Save Selection (of only the first selected thing)
 	TArray< TSharedPtr<FGraphActionNode> > SelectedNodes = TreeView->GetSelectedItems();
 	TSharedPtr<FGraphActionNode> SelectedAction = SelectedNodes.Num() > 0 ? SelectedNodes[0] : nullptr;
 
-	AllActions.Empty();
-	OnCollectAllActions.ExecuteIfBound(AllActions);
+	{
+		TRACE_CPUPROFILER_EVENT_SCOPE(SGraphActionMenu::CollectAllActions);
+
+		AllActions.Empty();
+		OnCollectAllActions.ExecuteIfBound(AllActions);
+	}
+
 	GenerateFilteredItems(bPreserveExpansion);
 
 	// Re-apply selection #0 if possible
@@ -660,6 +667,8 @@ void RestoreExpansionState(TSharedPtr< STreeView<ItemType> > InTree, const TArra
 
 void SGraphActionMenu::GenerateFilteredItems(bool bPreserveExpansion)
 {
+	TRACE_CPUPROFILER_EVENT_SCOPE(SGraphActionMenu::GenerateFilteredItems);
+
 	// First, save off current expansion state
 	TSet< TSharedPtr<FGraphActionNode> > OldExpansionState;
 	if(bPreserveExpansion)
diff --git a/Engine/Source/Editor/GraphEditor/Private/SGraphPanel.cpp b/Engine/Source/Editor/GraphEditor/Private/SGraphPanel.cpp
--- a/Engine/Source/Editor/GraphEditor/Private/SGraphPanel.cpp
+++ b/Engine/Source/Editor/GraphEditor/Private/SGraphPanel.cpp
@@ -1252,6 +1252,8 @@ void SGraphPanel::RemoveAllNodes()
 
 TSharedPtr<SWidget> SGraphPanel::SummonContextMenu(const FVector2D& WhereToSummon, const FVector2D& WhereToAddNode, UEdGraphNode* ForNode, UEdGraphPin* ForPin, const TArray<UEdGraphPin*>& DragFromPins)
 {
+	TRACE_CPUPROFILER_EVENT_SCOPE(SGraphPanel::SummonContextMenu);
+
 	if (OnGetContextMenuFor.IsBound())
 	{
 		FGraphContextMenuArguments SpawnInfo;
diff --git a/Engine/Source/Editor/Kismet/Private/BlueprintActionMenuBuilder.cpp b/Engine/Source/Editor/Kismet/Private/BlueprintActionMenuBuilder.cpp
--- a/Engine/Source/Editor/Kismet/Private/BlueprintActionMenuBuilder.cpp
+++ b/Engine/Source/Editor/Kismet/Private/BlueprintActionMenuBuilder.cpp
@@ -19,6 +19,14 @@
 #define LOCTEXT_NAMESPACE "BlueprintActionMenuBuilder"
 DEFINE_LOG_CATEGORY_STATIC(LogBlueprintActionMenuItemFactory, Log, All);
 
+#if ENABLE_BLUEPRINT_ACTION_FILTER_PROFILING
+static TAutoConsoleVariable<bool> CVarBPEnableActionMenuFilterTestStatsLogging(
+	TEXT("BP.EnableActionMenuFilterTestStatsLogging"),
+	false,
+	TEXT("If enabled, filter test stats will be logged for each menu section each time a Blueprint action menu is built for display.")
+);
+#endif	// ENABLE_BLUEPRINT_ACTION_FILTER_PROFILING
+
 /*******************************************************************************
  * FBlueprintActionMenuItemFactory
  ******************************************************************************/
@@ -106,6 +114,8 @@ FBlueprintActionMenuItemFactory::FBlueprintActionMenuItemFactory(FBlueprintActio
 //------------------------------------------------------------------------------
 TSharedPtr<FBlueprintActionMenuItem> FBlueprintActionMenuItemFactory::MakeActionMenuItem(TWeakPtr<FBlueprintEditor> EditorContext, FBlueprintActionInfo const& ActionInfo)
 {
+	TRACE_CPUPROFILER_EVENT_SCOPE(FBlueprintActionMenuItemFactory::MakeActionMenuItem);
+
 	FBlueprintActionUiSpec UiSignature = GetActionUiSignature(EditorContext, ActionInfo);
 
 	UBlueprintNodeSpawner const* Action = ActionInfo.NodeSpawner;
@@ -232,6 +242,8 @@ UEdGraph* FBlueprintActionMenuItemFactory::GetTargetGraph(TWeakPtr<FBlueprintEdi
 //------------------------------------------------------------------------------
 FBlueprintActionUiSpec FBlueprintActionMenuItemFactory::GetActionUiSignature(TWeakPtr<FBlueprintEditor> EditorContext, FBlueprintActionInfo const& ActionInfo)
 {
+	TRACE_CPUPROFILER_EVENT_SCOPE(FBlueprintActionMenuItemFactory::GetActionUiSignature);
+
 	UBlueprintNodeSpawner const* Action = ActionInfo.NodeSpawner;
 
 	UEdGraph* TargetGraph = GetTargetGraph(EditorContext);
@@ -406,7 +418,9 @@ void FBlueprintActionMenuBuilderImpl::FMenuSectionDefinition::AddBoundMenuItems(
 
 //------------------------------------------------------------------------------
 FBlueprintActionMenuBuilderImpl::MenuItemList FBlueprintActionMenuBuilderImpl::FMenuSectionDefinition::MakeMenuItems(TWeakPtr<FBlueprintEditor> EditorContext, FBlueprintActionInfo& DatabaseAction)
-{	
+{
+	TRACE_CPUPROFILER_EVENT_SCOPE(FMenuSectionDefinition::MakeMenuItems);
+
 	TSharedPtr<FEdGraphSchemaAction> UnBoundMenuEntry;
 	bool bPassedFilter = !Filter.IsFiltered(DatabaseAction);
 
@@ -551,6 +565,22 @@ void FBlueprintActionMenuBuilder::RebuildActionList()
 			ActionDatabase.DeferredRemoveEntry(ObjKey);
 		}
 	}
+
+#if ENABLE_BLUEPRINT_ACTION_FILTER_PROFILING
+	if (CVarBPEnableActionMenuFilterTestStatsLogging.GetValueOnGameThread())
+	{
+		// Dump detailed stats information about each filter test that was involved with building each menu section.
+		for (const TSharedRef<FMenuSectionDefinition>& MenuSection : MenuSections)
+		{
+			UE_LOG(LogBlueprintActionMenuItemFactory, Log, TEXT("=== %s: FILTER TEST PROFILE ==="), *MenuSection->GetSectionHeading().ToString());
+
+			for (const FString& ProfileEntry : MenuSection->Filter.GetFilterTestProfile())
+			{
+				UE_LOG(LogBlueprintActionMenuItemFactory, Log, TEXT("%s"), *ProfileEntry);
+			}
+		}
+	}
+#endif	// ENABLE_BLUEPRINT_ACTION_FILTER_PROFILING
 }
 
 #undef LOCTEXT_NAMESPACE
diff --git a/Engine/Source/Editor/Kismet/Private/BlueprintActionMenuUtils.cpp b/Engine/Source/Editor/Kismet/Private/BlueprintActionMenuUtils.cpp
--- a/Engine/Source/Editor/Kismet/Private/BlueprintActionMenuUtils.cpp
+++ b/Engine/Source/Editor/Kismet/Private/BlueprintActionMenuUtils.cpp
@@ -415,6 +415,8 @@ static void BlueprintActionMenuUtilsImpl::AddFavoritesSection(FBlueprintActionFi
 //------------------------------------------------------------------------------
 void FBlueprintActionMenuUtils::MakePaletteMenu(FBlueprintActionContext const& Context, UClass* FilterClass, FBlueprintActionMenuBuilder& MenuOut)
 {
+	TRACE_CPUPROFILER_EVENT_SCOPE(FBlueprintActionMenuUtils::MakePaletteMenu);
+
 	MenuOut.Empty();
 	
 	FBlueprintActionFilter::EFlags FilterFlags = FBlueprintActionFilter::BPFILTER_NoFlags;
@@ -445,6 +447,8 @@ void FBlueprintActionMenuUtils::MakePaletteMenu(FBlueprintActionContext const& C
 //------------------------------------------------------------------------------
 void FBlueprintActionMenuUtils::MakeContextMenu(FBlueprintActionContext const& Context, bool bIsContextSensitive, uint32 ClassTargetMask, FBlueprintActionMenuBuilder& MenuOut)
 {
+	TRACE_CPUPROFILER_EVENT_SCOPE(FBlueprintActionMenuUtils::MakeContextMenu);
+
 	using namespace BlueprintActionMenuUtilsImpl;
 
 	//--------------------------------------
@@ -734,6 +738,8 @@ void FBlueprintActionMenuUtils::MakeContextMenu(FBlueprintActionContext const& C
 //------------------------------------------------------------------------------
 void FBlueprintActionMenuUtils::MakeFavoritesMenu(FBlueprintActionContext const& Context, FBlueprintActionMenuBuilder& MenuOut)
 {
+	TRACE_CPUPROFILER_EVENT_SCOPE(FBlueprintActionMenuUtils::MakeFavoritesMenu);
+
 	MenuOut.Empty();
 
 	FBlueprintActionFilter MenuFilter;
diff --git a/Engine/Source/Editor/Kismet/Private/BlueprintEditor.cpp b/Engine/Source/Editor/Kismet/Private/BlueprintEditor.cpp
--- a/Engine/Source/Editor/Kismet/Private/BlueprintEditor.cpp
+++ b/Engine/Source/Editor/Kismet/Private/BlueprintEditor.cpp
@@ -1927,6 +1927,8 @@ void FBlueprintEditor::CommonInitialization(const TArray<UBlueprint*>& InitBluep
 
 void FBlueprintEditor::LoadLibrariesFromAssetRegistry()
 {
+	TRACE_CPUPROFILER_EVENT_SCOPE(FBlueprintEditor::LoadLibrariesFromAssetRegistry);
+
 	if (EnableAutomaticLibraryAssetLoading == 0)
 	{
 		return;
