Fixed incorrect comment for CanEverInsertExecutionPin().

#rnx

diff --git a/Engine/Source/Editor/BlueprintGraph/Classes/K2Node.h b/Engine/Source/Editor/BlueprintGraph/Classes/K2Node.h
--- a/Engine/Source/Editor/BlueprintGraph/Classes/K2Node.h
+++ b/Engine/Source/Editor/BlueprintGraph/Classes/K2Node.h
@@ -248,7 +248,7 @@ public:
 	/** Return whether the node's properties display in the blueprint details panel */
 	virtual bool ShouldShowNodeProperties() const { return false; }
 
-	/** Return whether the node's execution pins should support the remove execution pin action */
+	/** Return whether the node's execution pins should support the insert execution pin action */
 	virtual bool CanEverInsertExecutionPin() const { return false; }
 
 	/** Return whether the node's execution pins should support the remove execution pin action */
