Removing some unused local variables

#9346

diff --git a/Engine/Source/Editor/BlueprintGraph/Private/K2Node_ExecutionSequence.cpp b/Engine/Source/Editor/BlueprintGraph/Private/K2Node_ExecutionSequence.cpp
--- a/Engine/Source/Editor/BlueprintGraph/Private/K2Node_ExecutionSequence.cpp
+++ b/Engine/Source/Editor/BlueprintGraph/Private/K2Node_ExecutionSequence.cpp
@@ -118,7 +118,6 @@ public:
 				}
 
 				// Immediately jump to the first pin
-				UEdGraphNode* NextNode = OutputPins[0]->LinkedTo[0]->GetOwningNode();
 				FBlueprintCompiledStatement& NextExecutionState = Context.AppendStatementForNode(Node);
 				NextExecutionState.Type = KCST_UnconditionalGoto;
 				Context.GotoFixupRequestMap.Add(&NextExecutionState, OutputPins[0]);
diff --git a/Engine/Source/Editor/BlueprintGraph/Private/K2Node_Switch.cpp b/Engine/Source/Editor/BlueprintGraph/Private/K2Node_Switch.cpp
--- a/Engine/Source/Editor/BlueprintGraph/Private/K2Node_Switch.cpp
+++ b/Engine/Source/Editor/BlueprintGraph/Private/K2Node_Switch.cpp
@@ -81,7 +81,6 @@ public:
 
 		if ((BoolTerm != NULL) && (SwitchSelectionTerm != NULL))
 		{
-			UEdGraphNode* TargetNode = NULL;
 			UEdGraphPin* FuncPin = SwitchNode->GetFunctionPin();
 			FBPTerminal* FuncContext = Context.NetMap.FindRef(FuncPin);
 			UEdGraphPin* DefaultPin = SwitchNode->GetDefaultPin();
