PR #9162: Fix: Typo and format (Contributed by [REDACTED_PERSON])

diff --git a/Engine/Source/Editor/BlueprintGraph/Private/K2Node_IfThenElse.cpp b/Engine/Source/Editor/BlueprintGraph/Private/K2Node_IfThenElse.cpp
--- a/Engine/Source/Editor/BlueprintGraph/Private/K2Node_IfThenElse.cpp
+++ b/Engine/Source/Editor/BlueprintGraph/Private/K2Node_IfThenElse.cpp
@@ -59,7 +59,7 @@ public:
 			UEdGraphPin* PinToTry = FEdGraphUtilities::GetNetFromPin(CondPin);
 			FBPTerminal** CondTerm = Context.NetMap.Find(PinToTry);
 
-			if (CondTerm != NULL) //
+			if (CondTerm != NULL)
 			{
 				// First skip the if, if the term is false
 				{
@@ -99,7 +99,7 @@ void UK2Node_IfThenElse::AllocateDefaultPins()
 	K2Schema->SetPinAutogeneratedDefaultValue(ConditionPin, TEXT("true"));
 
 	UEdGraphPin* TruePin = CreatePin(EGPD_Output, UEdGraphSchema_K2::PC_Exec, UEdGraphSchema_K2::PN_Then);
-	TruePin->PinFriendlyName =LOCTEXT("true", "true");
+	TruePin->PinFriendlyName = LOCTEXT("true", "true");
 
 	UEdGraphPin* FalsePin = CreatePin(EGPD_Output, UEdGraphSchema_K2::PC_Exec, UEdGraphSchema_K2::PN_Else);
 	FalsePin->PinFriendlyName = LOCTEXT("false", "false");
@@ -125,7 +125,7 @@ FSlateIcon UK2Node_IfThenElse::GetIconAndTint(FLinearColor& OutColor) const
 
 FText UK2Node_IfThenElse::GetTooltipText() const
 {
-	return LOCTEXT("BrancStatement_Tooltip", "Branch Statement\nIf Condition is true, execution goes to True, otherwise it goes to False");
+	return LOCTEXT("BranchStatement_Tooltip", "Branch Statement\nIf Condition is true, execution goes to True, otherwise it goes to False");
 }
 
 UEdGraphPin* UK2Node_IfThenElse::GetThenPin() const
