Revise the legacy pin default value mismatch warning check on node reconstruction to handle string diffs that convert to identical values.

diff --git a/Engine/Source/Editor/BlueprintGraph/Private/K2Node.cpp b/Engine/Source/Editor/BlueprintGraph/Private/K2Node.cpp
--- a/Engine/Source/Editor/BlueprintGraph/Private/K2Node.cpp
+++ b/Engine/Source/Editor/BlueprintGraph/Private/K2Node.cpp
@@ -753,13 +753,27 @@ void UK2Node::ReconstructNode()
 			if (OldPinEntry && *OldPinEntry)
 			{
 				UEdGraphPin* OldPin = *OldPinEntry;
-				if (NewPin->LinkedTo.Num() == 0 &&
-					!OldPin->AutogeneratedDefaultValue.IsEmpty() &&
-					NewPin->AutogeneratedDefaultValue != OldPin->AutogeneratedDefaultValue &&
-					OldPin->DefaultValue == OldPin->AutogeneratedDefaultValue &&
-					NewPin->DefaultValue != OldPin->DefaultValue)
+				if (NewPin->LinkedTo.Num() == 0
+					&& !OldPin->AutogeneratedDefaultValue.IsEmpty()
+					&& Schema->DoesDefaultValueMatchAutogenerated(*OldPin))
 				{
-					Blueprint->CurrentMessageLog->Warning(*LOCTEXT("VerifyDefaultValues", "Default value for @@ on @@ has changed and this asset is from a version that may have had incorrect default value information - verify and resave").ToString(), NewPin, this);
+					// Rather than compare default value strings directly, we utilize the schema to convert and compare against the actual default value type.
+					auto HasPinDefaultValueChanged = [Schema](const UEdGraphPin& PinRef, FString& PinDefaultValueRef, const FString& OldPinDefaultValue) -> bool
+					{
+						TGuardValue<FString> TempSetPinDefaultValue(PinDefaultValueRef, OldPinDefaultValue);
+						return !Schema->DoesDefaultValueMatchAutogenerated(PinRef);
+					};
+
+					// Compare default value strings by temporarily setting the default value and auto-generated default value strings to the equivalent default
+					// value string from the old pin that we wish to compare it against. In some cases, the value string may have changed, but wouldn't otherwise
+					// result in a different value when converted.
+					// 
+					// Example: For vector types, {1, 1, 1} is equivalent to {1.000, 1.000, 1.000}, which can be considered a match that would not require a resave.
+					if (HasPinDefaultValueChanged(*NewPin, NewPin->DefaultValue, OldPin->AutogeneratedDefaultValue)
+						&& HasPinDefaultValueChanged(*NewPin, NewPin->AutogeneratedDefaultValue, OldPin->DefaultValue))
+					{
+						Blueprint->CurrentMessageLog->Warning(*LOCTEXT("VerifyDefaultValues", "Default value for @@ on @@ has changed and this asset is from a version that may have had incorrect default value information - verify and resave").ToString(), NewPin, this);
+					}
 				}
 			}
 		}
