Fix the "Create Event" delegate node to support reroute nodes.

The UK2Node_CreateDelegate node doesn't maintain accurate type information for its delegate pin. As a consequence, adjacent knot nodes may not have accurate information either. When obtaining the delegate signature via GetDelegateSignature, we need to traverse the knot chain and obtain the pin type info from the endpoint node.

Additionally, this is more of a temporary fix. UK2Node_CreateDelegate really needs to refresh its own delegate pin type.

diff --git a/Engine/Source/Editor/BlueprintGraph/Classes/K2Node_CreateDelegate.h b/Engine/Source/Editor/BlueprintGraph/Classes/K2Node_CreateDelegate.h
--- a/Engine/Source/Editor/BlueprintGraph/Classes/K2Node_CreateDelegate.h
+++ b/Engine/Source/Editor/BlueprintGraph/Classes/K2Node_CreateDelegate.h
@@ -44,7 +44,7 @@ public:
 	virtual ERedirectType DoPinsMatchForReconstruction(const UEdGraphPin* NewPin, int32 NewPinIndex, const UEdGraphPin* OldPin, int32 OldPinIndex) const override;
 	// End of UK2Node interface
 
-	bool IsValid(FString* OutMsg = NULL, bool bDontUseSkeletalClassForSelf = false) const;
+	bool IsValid(FString* OutMsg = nullptr, bool bDontUseSkeletalClassForSelf = false) const;
 
 	/** Set new Function name (Without notifying about the change) */
 	BLUEPRINTGRAPH_API void SetFunction(FName Name);
@@ -62,5 +62,5 @@ public:
 	BLUEPRINTGRAPH_API void ValidationAfterFunctionsAreCreated(class FCompilerResultsLog& MessageLog, bool bFullCompile) const;
 
 	// return Graph and Blueprint, when they should be notified about change. It allows to call BroadcastChanged only once per blueprint.
-	BLUEPRINTGRAPH_API void HandleAnyChange(UEdGraph* & OutGraph, UBlueprint* & OutBlueprint);
+	BLUEPRINTGRAPH_API void HandleAnyChange(UEdGraph*& OutGraph, UBlueprint*& OutBlueprint);
 };
diff --git a/Engine/Source/Editor/BlueprintGraph/Private/K2Node_CreateDelegate.cpp b/Engine/Source/Editor/BlueprintGraph/Private/K2Node_CreateDelegate.cpp
--- a/Engine/Source/Editor/BlueprintGraph/Private/K2Node_CreateDelegate.cpp
+++ b/Engine/Source/Editor/BlueprintGraph/Private/K2Node_CreateDelegate.cpp
@@ -6,6 +6,7 @@
 #include "EdGraphSchema_K2.h"
 #include "K2Node_Event.h"
 #include "K2Node_BaseMCDelegate.h"
+#include "K2Node_Knot.h"
 #include "Kismet2/BlueprintEditorUtils.h"
 
 #include "Kismet2/CompilerResultsLog.h"
@@ -15,13 +16,11 @@
 #include "BlueprintActionDatabaseRegistrar.h"
 #include "FindInBlueprintManager.h"
 
-struct FK2Node_CreateDelegate_Helper
+namespace UE::BlueprintGraph::Private
 {
-	static FName DelegateOutputName;
-	static FName InputObjectName; // Deprecated, for fixup
-};
-FName FK2Node_CreateDelegate_Helper::DelegateOutputName(TEXT("OutputDelegate"));
-FName FK2Node_CreateDelegate_Helper::InputObjectName(TEXT("InputObject"));
+	static const FName DelegateOutputName = TEXT("OutputDelegate");
+	static const FName InputObjectName = TEXT("InputObject");
+}
 
 UK2Node_CreateDelegate::UK2Node_CreateDelegate(const FObjectInitializer& ObjectInitializer)
 	: Super(ObjectInitializer)
@@ -35,7 +34,7 @@ void UK2Node_CreateDelegate::AllocateDefaultPins()
 		ObjPin->PinFriendlyName = NSLOCTEXT("K2Node", "CreateDelegate_ObjectInputName", "Object");
 	}
 
-	if(UEdGraphPin* DelegatePin = CreatePin(EGPD_Output, UEdGraphSchema_K2::PC_Delegate, FK2Node_CreateDelegate_Helper::DelegateOutputName))
+	if(UEdGraphPin* DelegatePin = CreatePin(EGPD_Output, UEdGraphSchema_K2::PC_Delegate, UE::BlueprintGraph::Private::DelegateOutputName))
 	{
 		DelegatePin->PinFriendlyName = NSLOCTEXT("K2Node", "CreateDelegate_DelegateOutName", "Event");
 	}
@@ -46,7 +45,7 @@ void UK2Node_CreateDelegate::AllocateDefaultPins()
 UK2Node::ERedirectType UK2Node_CreateDelegate::DoPinsMatchForReconstruction(const UEdGraphPin* NewPin, int32 NewPinIndex, const UEdGraphPin* OldPin, int32 OldPinIndex) const
 {
 	// Handles remap of InputObject to Self, from 4.10 time frame
-	if (OldPin->PinName == FK2Node_CreateDelegate_Helper::InputObjectName && NewPin->PinName == UEdGraphSchema_K2::PN_Self)
+	if (OldPin->PinName == UE::BlueprintGraph::Private::InputObjectName && NewPin->PinName == UEdGraphSchema_K2::PN_Self)
 	{
 		return ERedirectType_Name;
 	}
@@ -230,7 +229,7 @@ void UK2Node_CreateDelegate::HandleAnyChangeWithoutNotifying()
 	}
 }
 
-void UK2Node_CreateDelegate::HandleAnyChange(UEdGraph* & OutGraph, UBlueprint* & OutBlueprint)
+void UK2Node_CreateDelegate::HandleAnyChange(UEdGraph*& OutGraph, UBlueprint*& OutBlueprint)
 {
 	const FName OldSelectedFunctionName = GetFunctionName();
 	HandleAnyChangeWithoutNotifying();
@@ -318,17 +317,26 @@ UFunction* UK2Node_CreateDelegate::GetDelegateSignature() const
 {
 	UEdGraphPin* Pin = GetDelegateOutPin();
 	check(Pin);
-	if(Pin->LinkedTo.Num())
+	UFunction* Result = nullptr;
+
+	if (Pin->LinkedTo.Num())
 	{
-		if(UEdGraphPin* ResultPin = Pin->LinkedTo[0])
+		if (UEdGraphPin* ResultPin = Pin->LinkedTo[0])
 		{
-			if (UEdGraphSchema_K2::PC_Delegate == ResultPin->PinType.PinCategory)
+			// The knot nodes may not necessarily have the correct info for PinSubCategoryMemberReference.
+			// The safer approach is to traverse the knots to find the endpoint node, and then use its 
+			// PinSubCategoryMemberReference during resolution of the member reference.
+
+			UEdGraphPin* LinkedPin = FBlueprintEditorUtils::FindFirstCompilerRelevantLinkedPin(ResultPin);
+
+			if (LinkedPin && (UEdGraphSchema_K2::PC_Delegate == LinkedPin->PinType.PinCategory))
 			{
-				return FMemberReference::ResolveSimpleMemberReference<UFunction>(ResultPin->PinType.PinSubCategoryMemberReference);
+				Result = FMemberReference::ResolveSimpleMemberReference<UFunction>(LinkedPin->PinType.PinSubCategoryMemberReference);
 			}
 		}
 	}
-	return nullptr;
+
+	return Result;
 }
 
 UClass* UK2Node_CreateDelegate::GetScopeClass(bool bDontUseSkeletalClassForSelf/* = false*/) const
@@ -388,7 +396,7 @@ FName UK2Node_CreateDelegate::GetFunctionName() const
 
 UEdGraphPin* UK2Node_CreateDelegate::GetDelegateOutPin() const
 {
-	return FindPin(FK2Node_CreateDelegate_Helper::DelegateOutputName);
+	return FindPin(UE::BlueprintGraph::Private::DelegateOutputName);
 }
 
 UEdGraphPin* UK2Node_CreateDelegate::GetObjectInPin() const
