[Bug Fix] Crash when recombining Value Pin after swapping Target Wild Card pin from a Vector2D to a Vector on a Set By-Ref Var node

diff --git a/Engine/Source/Editor/BlueprintGraph/Private/K2Node_VariableSetRef.cpp b/Engine/Source/Editor/BlueprintGraph/Private/K2Node_VariableSetRef.cpp
--- a/Engine/Source/Editor/BlueprintGraph/Private/K2Node_VariableSetRef.cpp
+++ b/Engine/Source/Editor/BlueprintGraph/Private/K2Node_VariableSetRef.cpp
@@ -238,16 +238,28 @@ void UK2Node_VariableSetRef::NotifyPinConnectionListChanged(UEdGraphPin* Pin)
 		// If both target and value pins are unlinked, then reset types to wildcard
 		if(TargetPin->LinkedTo.Num() == 0 && ValuePin->LinkedTo.Num() == 0)
 		{
+			// collapse SubPins back into their parent if there are any
+			auto TryRecombineSubPins = [](UEdGraphPin* ParentPin)
+			{
+				if (!ParentPin->SubPins.IsEmpty())
+				{
+					const UEdGraphSchema_K2* K2Schema = GetDefault<UEdGraphSchema_K2>();
+					K2Schema->RecombinePin(ParentPin->SubPins[0]);
+				}
+			};
+			
 			// Pin disconnected...revert to wildcard
 			TargetPin->PinType.PinCategory = UEdGraphSchema_K2::PC_Wildcard;
 			TargetPin->PinType.PinSubCategory = NAME_None;
 			TargetPin->PinType.PinSubCategoryObject = nullptr;
 			TargetPin->BreakAllPinLinks();
+			TryRecombineSubPins(TargetPin);
 
 			ValuePin->PinType.PinCategory = UEdGraphSchema_K2::PC_Wildcard;
 			ValuePin->PinType.PinSubCategory = NAME_None;
 			ValuePin->PinType.PinSubCategoryObject = nullptr;
-			ValuePin->BreakAllPinLinks();			
+			ValuePin->BreakAllPinLinks();
+			TryRecombineSubPins(ValuePin);
 		}
 
 		CachedNodeTitle.MarkDirty();
