turned off blueprint action filter profiling

diff --git a/Engine/Source/Editor/BlueprintGraph/Private/BlueprintActionFilter.cpp b/Engine/Source/Editor/BlueprintGraph/Private/BlueprintActionFilter.cpp
--- a/Engine/Source/Editor/BlueprintGraph/Private/BlueprintActionFilter.cpp
+++ b/Engine/Source/Editor/BlueprintGraph/Private/BlueprintActionFilter.cpp
@@ -2209,7 +2209,10 @@ void FBlueprintActionFilter::AddRejectionTest(FRejectionTestDelegate RejectionTe
 //------------------------------------------------------------------------------
 bool FBlueprintActionFilter::IsFiltered(FBlueprintActionInfo& BlueprintAction)
 {
+	
+#if ENABLE_BLUEPRINT_ACTION_FILTER_PROFILING
 	TRACE_CPUPROFILER_EVENT_SCOPE(FBlueprintActionFilter::IsFiltered);
+#endif
 
 	bool bIsFiltered = IsFilteredByThis(BlueprintAction);
 	if (!bIsFiltered)
diff --git a/Engine/Source/Editor/BlueprintGraph/Public/BlueprintActionFilter.h b/Engine/Source/Editor/BlueprintGraph/Public/BlueprintActionFilter.h
--- a/Engine/Source/Editor/BlueprintGraph/Public/BlueprintActionFilter.h
+++ b/Engine/Source/Editor/BlueprintGraph/Public/BlueprintActionFilter.h
@@ -11,7 +11,7 @@
 #include "BlueprintNodeBinder.h"
 #include "BlueprintGraphModule.h"
 
-#define ENABLE_BLUEPRINT_ACTION_FILTER_PROFILING (!(UE_BUILD_SHIPPING || UE_BUILD_TEST) && 1)
+#define ENABLE_BLUEPRINT_ACTION_FILTER_PROFILING (!(UE_BUILD_SHIPPING || UE_BUILD_TEST) && 0)
 
 class UBlueprint;
 class UBlueprintNodeSpawner;
