Add guard code to try and catch invalid Blueprint node spawner refs during GC and report some additional context to the log without necessarily crashing the editor.

diff --git a/Engine/Source/Editor/BlueprintGraph/Private/BlueprintActionDatabase.cpp b/Engine/Source/Editor/BlueprintGraph/Private/BlueprintActionDatabase.cpp
--- a/Engine/Source/Editor/BlueprintGraph/Private/BlueprintActionDatabase.cpp
+++ b/Engine/Source/Editor/BlueprintGraph/Private/BlueprintActionDatabase.cpp
@@ -1219,10 +1219,23 @@ FBlueprintActionDatabase::~FBlueprintActionDatabase()
 void FBlueprintActionDatabase::AddReferencedObjects(FReferenceCollector& Collector)
 {
 	TSet<UBlueprintNodeSpawner*> AllActions;
-	for (auto& ActionListIt : ActionRegistry)
+	for (const TPair<FObjectKey, FActionList>& ActionListIt : ActionRegistry)
 	{
-		AllActions.Append(ActionListIt.Value);
-		Collector.AddReferencedObjects(ActionListIt.Value);
+		const FActionList& ActionList = ActionListIt.Value;
+		AllActions.Reserve(AllActions.Num() + ActionList.Num());
+		for (UBlueprintNodeSpawner* Action : ActionList)
+		{
+			// We have some reports of invalid action ptrs during GC - try to catch that case here without crashing the editor while reference gathering.
+			if (!Action || (GIsGarbageCollecting && !Action->IsValidLowLevel()))
+			{
+				const UObject* Key = ActionListIt.Key.ResolveObjectPtr();
+				ensureMsgf(false, TEXT("Invalid action (0x%016llx) registered for object: %s"), (int64)(PTRINT)Action, Key ? *Key->GetName() : TEXT("NULL"));
+				continue;
+			}
+
+			AllActions.Add(Action);
+			Collector.AddReferencedObject(Action);
+		}
 	}
 
 	// shouldn't have to do this, as the elements listed here should also be 
@@ -1232,10 +1245,22 @@ void FBlueprintActionDatabase::AddReferencedObjects(FReferenceCollector& Collect
 	if (UnloadedActionRegistry.Num() > 0)
 	{
 		TSet<UBlueprintNodeSpawner*> UnloadedActions;
-		for (auto& UnloadedActionListIt : UnloadedActionRegistry)
+		for (const TPair<FName, FActionList>& UnloadedActionListIt : UnloadedActionRegistry)
 		{
-			UnloadedActions.Append(UnloadedActionListIt.Value);
-			Collector.AddReferencedObjects(UnloadedActionListIt.Value);
+			const FActionList& ActionList = UnloadedActionListIt.Value;
+			UnloadedActions.Reserve(UnloadedActions.Num() + ActionList.Num());
+			for (UBlueprintNodeSpawner* Action : ActionList)
+			{
+				// Similar to above; however, we don't have any reports of failure here during GC. Nonetheless, we'll try and catch an invalid ptr value just in case.
+				if (!Action || (GIsGarbageCollecting && !Action->IsValidLowLevel()))
+				{
+					ensureMsgf(false, TEXT("Invalid action (0x%016llx) registered for unloaded object path: %s"), (int64)(PTRINT)Action, *UnloadedActionListIt.Key.ToString());
+					continue;
+				}
+
+				UnloadedActions.Add(Action);
+				Collector.AddReferencedObject(Action);
+			}
 		}
 
 		auto OrphanedUnloadedActions = UnloadedActions.Difference(AllActions.Intersect(UnloadedActions));
