Allow users to supress BP compiler deprecation warnings for specific properties or functions via project settings

diff --git a/Engine/Source/Editor/BlueprintGraph/Classes/K2Node_CallFunction.h b/Engine/Source/Editor/BlueprintGraph/Classes/K2Node_CallFunction.h
--- a/Engine/Source/Editor/BlueprintGraph/Classes/K2Node_CallFunction.h
+++ b/Engine/Source/Editor/BlueprintGraph/Classes/K2Node_CallFunction.h
@@ -95,6 +95,7 @@ public:
 	virtual UObject* GetJumpTargetForDoubleClick() const override;
 	virtual bool CanJumpToDefinition() const override;
 	virtual void JumpToDefinition() const override;
+	virtual void GetNodeContextMenuActions(class UToolMenu* Menu, class UGraphNodeContextMenuContext* Context) const override;
 	virtual FString GetPinMetaData(FName InPinName, FName InKey) override;
 	virtual bool HasExternalDependencies(TArray<class UStruct*>* OptionalOutput) const override;
 	// End of UEdGraphNode interface
@@ -238,6 +239,9 @@ protected:
 	/** Helper function to ensure function is called in our context */
 	virtual void FixupSelfMemberContext();
 
+	/** Adds this function to the suppressed deprecation warnings list for this project */
+	void SuppressDeprecationWarning() const;
+
 	/** Helper function to find UFunction entries from the skeleton class, use with caution.. */
 	UFunction* GetTargetFunctionFromSkeletonClass() const;
 };
diff --git a/Engine/Source/Editor/BlueprintGraph/Classes/K2Node_Variable.h b/Engine/Source/Editor/BlueprintGraph/Classes/K2Node_Variable.h
--- a/Engine/Source/Editor/BlueprintGraph/Classes/K2Node_Variable.h
+++ b/Engine/Source/Editor/BlueprintGraph/Classes/K2Node_Variable.h
@@ -76,6 +76,7 @@ public:
 	virtual UObject* GetJumpTargetForDoubleClick() const override;
 	virtual bool CanJumpToDefinition() const override;
 	virtual void JumpToDefinition() const override;
+	virtual void GetNodeContextMenuActions(class UToolMenu* Menu, class UGraphNodeContextMenuContext* Context) const override;
 	virtual FString GetPinMetaData(FName InPinName, FName InKey) override;
 	virtual bool HasExternalDependencies(TArray<class UStruct*>* OptionalOutput) const override;
 	//~ End UEdGraphNode Interface
@@ -196,6 +197,9 @@ protected:
 	 */
 	const UActorComponent* GetActorComponent(const FProperty* VariableProperty) const;
 
+	/** Adds the variable reference to the suppressed deprecation warnings list */
+	void SuppressDeprecationWarning() const;
+
 	/** Returns whether a Function Graph contains a parameter with the given name */
 	static bool FunctionParameterExists(const UEdGraph* InFunctionGraph, const FName InParameterName);
 };
diff --git a/Engine/Source/Editor/BlueprintGraph/Private/K2Node_CallFunction.cpp b/Engine/Source/Editor/BlueprintGraph/Private/K2Node_CallFunction.cpp
--- a/Engine/Source/Editor/BlueprintGraph/Private/K2Node_CallFunction.cpp
+++ b/Engine/Source/Editor/BlueprintGraph/Private/K2Node_CallFunction.cpp
@@ -45,6 +45,8 @@
 #include "HAL/FileManager.h"
 #include "Widgets/Notifications/SNotificationList.h"
 #include "BlueprintNodeStatics.h"
+#include "Settings/BlueprintEditorProjectSettings.h"
+#include "ToolMenu.h"
 
 #define LOCTEXT_NAMESPACE "K2Node"
 
@@ -491,8 +493,25 @@ UK2Node_CallFunction::UK2Node_CallFunction(const FObjectInitializer& ObjectIniti
 
 bool UK2Node_CallFunction::HasDeprecatedReference() const
 {
-	UFunction* Function = GetTargetFunction();
-	return (Function && Function->HasMetaData(FBlueprintMetadata::MD_DeprecatedFunction));
+	if (UFunction* Function = GetTargetFunction())
+	{
+		bool bDeprecated = Function->HasMetaData(FBlueprintMetadata::MD_DeprecatedFunction);
+
+		if (bDeprecated)
+		{
+			const UBlueprintEditorProjectSettings* BlueprintEditorProjectSettings = GetDefault<UBlueprintEditorProjectSettings>();
+			const FString PathName = Function->GetPathName();
+
+			if (BlueprintEditorProjectSettings->SuppressedDeprecationMessages.Contains(PathName))
+			{
+				bDeprecated = false;
+			}
+		}
+
+		return bDeprecated;
+	}
+
+	return false;
 }
 
 FEdGraphNodeDeprecationResponse UK2Node_CallFunction::GetDeprecationResponse(EEdGraphNodeDeprecationType DeprecationType) const
@@ -1967,6 +1986,18 @@ void UK2Node_CallFunction::FixupSelfMemberContext()
 	}
 }
 
+void UK2Node_CallFunction::SuppressDeprecationWarning() const
+{
+	if (UFunction* Function = GetTargetFunction())
+	{
+		FString PathName = Function->GetPathName();
+		UBlueprintEditorProjectSettings* BlueprintEditorProjectSettings = GetMutableDefault<UBlueprintEditorProjectSettings>();
+		BlueprintEditorProjectSettings->SuppressedDeprecationMessages.Add(MoveTemp(PathName));
+		BlueprintEditorProjectSettings->SaveConfig();
+		BlueprintEditorProjectSettings->TryUpdateDefaultConfigFile("", false);
+	}
+}
+
 void UK2Node_CallFunction::PostPasteNode()
 {
 	Super::PostPasteNode();
@@ -3355,6 +3386,30 @@ void UK2Node_CallFunction::JumpToDefinition() const
 	Super::JumpToDefinition();
 }
 
+void UK2Node_CallFunction::GetNodeContextMenuActions(class UToolMenu* Menu, class UGraphNodeContextMenuContext* Context) const
+{
+	Super::GetNodeContextMenuActions(Menu, Context);
+
+	if (HasDeprecatedReference())
+	{
+		FText MenuEntryTitle = LOCTEXT("SuppressDeprecationWarningTitle", "Suppress Deprecation Warning");
+		FText MenuEntryTooltip = LOCTEXT("SuppressDeprecationWarningTooltip", "Adds this variable to the suppressed deprecation warnings list for the project.");
+
+		FToolMenuSection& Section = Menu->AddSection("K2NodeCallFunction", LOCTEXT("FunctionHeader", "Function"));
+		Section.AddMenuEntry(
+			"SuppressDeprecationWarning",
+			MenuEntryTitle,
+			MenuEntryTooltip,
+			FSlateIcon(),
+			FUIAction(
+				FExecuteAction::CreateUObject(this, &UK2Node_CallFunction::SuppressDeprecationWarning),
+				FCanExecuteAction::CreateUObject(this, &UK2Node_CallFunction::HasDeprecatedReference),
+				FIsActionChecked()
+			)
+		);
+	}
+}
+
 FString UK2Node_CallFunction::GetPinMetaData(FName InPinName, FName InKey)
 {
 	FString MetaData = Super::GetPinMetaData(InPinName, InKey);
diff --git a/Engine/Source/Editor/BlueprintGraph/Private/K2Node_Variable.cpp b/Engine/Source/Editor/BlueprintGraph/Private/K2Node_Variable.cpp
--- a/Engine/Source/Editor/BlueprintGraph/Private/K2Node_Variable.cpp
+++ b/Engine/Source/Editor/BlueprintGraph/Private/K2Node_Variable.cpp
@@ -20,6 +20,8 @@
 #include "Editor/UnrealEdEngine.h"
 #include "Preferences/UnrealEdOptions.h"
 #include "UnrealEdGlobals.h"
+#include "Settings/BlueprintEditorProjectSettings.h"
+#include "ToolMenu.h"
 
 #define LOCTEXT_NAMESPACE "K2Node"
 
@@ -1036,22 +1038,48 @@ void UK2Node_Variable::PostPasteNode()
 
 bool UK2Node_Variable::HasDeprecatedReference() const
 {
+	bool bDeprecated = false;
+	FProperty* VariableProperty = nullptr; // Declare up here so we can reuse if we would have resolved twice
+
 	// Check if the referenced variable is deprecated.
 	if (VariableReference.IsDeprecated())
 	{
-		return true;
+		bDeprecated = true;
 	}
-	else if (FProperty* VariableProperty = VariableReference.ResolveMember<FProperty>(GetBlueprintClassFromNode()))
+	else
 	{
-		// Backcompat: Allow variables tagged only with 'DeprecationMessage' meta to be seen as deprecated if inherited from a native parent class.
-		const bool bHasDeprecationMessage = VariableProperty->HasMetaData(FBlueprintMetadata::MD_DeprecationMessage);
-		if (bHasDeprecationMessage && VariableProperty->GetOwnerUObject()->IsNative())
+		VariableProperty = VariableReference.ResolveMember<FProperty>(GetBlueprintClassFromNode());
+		if (VariableProperty)
 		{
-			return true;
+			// Backcompat: Allow variables tagged only with 'DeprecationMessage' meta to be seen as deprecated if inherited from a native parent class.
+			const bool bHasDeprecationMessage = VariableProperty->HasMetaData(FBlueprintMetadata::MD_DeprecationMessage);
+			if (bHasDeprecationMessage && VariableProperty->GetOwnerUObject()->IsNative())
+			{
+				bDeprecated = true;
+			}
 		}
 	}
 
-	return false;
+	if (bDeprecated)
+	{
+		const UBlueprintEditorProjectSettings* BlueprintEditorProjectSettings = GetDefault<UBlueprintEditorProjectSettings>();
+		if (!VariableProperty)
+		{
+			VariableProperty = VariableReference.ResolveMember<FProperty>(GetBlueprintClassFromNode());
+		}
+
+
+		if (VariableProperty)
+		{
+			const FString PathName = VariableProperty->GetPathName();
+			if (BlueprintEditorProjectSettings->SuppressedDeprecationMessages.Contains(PathName))
+			{
+				bDeprecated = false;
+			}
+		}
+	}
+
+	return bDeprecated;
 }
 
 FEdGraphNodeDeprecationResponse UK2Node_Variable::GetDeprecationResponse(EEdGraphNodeDeprecationType DeprecationType) const
@@ -1117,6 +1145,30 @@ void UK2Node_Variable::JumpToDefinition() const
 	Super::JumpToDefinition();
 }
 
+void UK2Node_Variable::GetNodeContextMenuActions(class UToolMenu* Menu, class UGraphNodeContextMenuContext* Context) const
+{
+	Super::GetNodeContextMenuActions(Menu, Context);
+
+	if (HasDeprecatedReference())
+	{
+		FText MenuEntryTitle = LOCTEXT("SuppressDeprecationWarningTitle", "Suppress Deprecation Warning");
+		FText MenuEntryTooltip = LOCTEXT("SuppressDeprecationWarningTooltip", "Adds this variable to the suppressed deprecation warnings list for the project.");
+
+		FToolMenuSection& Section = Menu->AddSection("K2NodeVariable", LOCTEXT("VariableHeader", "Variable"));
+		Section.AddMenuEntry(
+			"SuppressDeprecationWarning",
+			MenuEntryTitle,
+			MenuEntryTooltip,
+			FSlateIcon(),
+			FUIAction(
+				FExecuteAction::CreateUObject(this, &UK2Node_Variable::SuppressDeprecationWarning),
+				FCanExecuteAction::CreateUObject(this, &UK2Node_Variable::HasDeprecatedReference),
+				FIsActionChecked()
+			)
+		);
+	}
+}
+
 bool UK2Node_Variable::FunctionParameterExists(const UEdGraph* InFunctionGraph, const FName InParameterName)
 {
 	TArray<UK2Node_FunctionEntry*> Entry;
@@ -1158,4 +1210,16 @@ const UActorComponent* UK2Node_Variable::GetActorComponent(const FProperty* Vari
 	return nullptr;
 }
 
+void UK2Node_Variable::SuppressDeprecationWarning() const
+{
+	if (const FProperty* Property = VariableReference.ResolveMember<FProperty>(GetBlueprintClassFromNode()))
+	{
+		FString PathName = Property->GetPathName();
+		UBlueprintEditorProjectSettings* BlueprintEditorProjectSettings = GetMutableDefault<UBlueprintEditorProjectSettings>();
+		BlueprintEditorProjectSettings->SuppressedDeprecationMessages.Add(MoveTemp(PathName));
+		BlueprintEditorProjectSettings->SaveConfig();
+		BlueprintEditorProjectSettings->TryUpdateDefaultConfigFile("", false);
+	}
+}
+
 #undef LOCTEXT_NAMESPACE
diff --git a/Engine/Source/Editor/UnrealEd/Public/Settings/BlueprintEditorProjectSettings.h b/Engine/Source/Editor/UnrealEd/Public/Settings/BlueprintEditorProjectSettings.h
--- a/Engine/Source/Editor/UnrealEd/Public/Settings/BlueprintEditorProjectSettings.h
+++ b/Engine/Source/Editor/UnrealEd/Public/Settings/BlueprintEditorProjectSettings.h
@@ -63,6 +63,13 @@ public:
 	UPROPERTY(EditAnywhere, config, Category= Blueprints, DisplayName = "Compiler Messages Disabled Entirely")
 	TArray<FName> DisabledCompilerMessages;
 
+	/**
+	 * List of deprecated UProperties/UFunctions to supress warning messages for - useful for source changes
+	 * that would otherwise cause content warnings
+	 */
+	UPROPERTY(EditAnywhere, config, Category = Blueprints, DisplayName = "Deprecated Symbols to Supress")
+	TArray<FString> SuppressedDeprecationMessages;
+
 	/**
 	 * Any blueprint deriving from one of these base classes will be allowed to recompile during Play-in-Editor
 	 * (This setting exists both as an editor preference and project setting, and will be allowed if listed in either place) 
diff --git a/Engine/Source/Runtime/Engine/Classes/Components/SkinnedMeshComponent.h b/Engine/Source/Runtime/Engine/Classes/Components/SkinnedMeshComponent.h
--- a/Engine/Source/Runtime/Engine/Classes/Components/SkinnedMeshComponent.h
+++ b/Engine/Source/Runtime/Engine/Classes/Components/SkinnedMeshComponent.h
@@ -267,7 +267,7 @@ class ENGINE_API USkinnedMeshComponent : public UMeshComponent, public ILODSyncI
 
 	/** The skeletal mesh used by this component. */
 	UE_DEPRECATED(5.1, "Replaced by SkinnedAsset. Use GetSkinnedAsset()/SetSkinnedAsset() instead.")
-	UPROPERTY(EditAnywhere, BlueprintGetter = GetSkeletalMesh_DEPRECATED, Category = "Mesh|SkeletalAsset", meta = (DisallowedClasses = "/Script/ApexDestruction.DestructibleMesh"))
+	UPROPERTY(EditAnywhere, BlueprintGetter = GetSkeletalMesh_DEPRECATED, Category = "Mesh|SkeletalAsset", meta = (DisallowedClasses = "/Script/ApexDestruction.DestructibleMesh", DeprecatedProperty, DeprecationMessage = "Use USkeletalMeshComponent::GetSkeletalMesh() or GetSkinnedAsset() instead."))
 	TObjectPtr<class USkeletalMesh> SkeletalMesh;
 
 private:
@@ -927,7 +927,7 @@ public:
 	 * @return the SkeletalMesh set to this mesh.
 	 */
 	UE_DEPRECATED(5.1, "Use USkeletalMeshComponent::GetSkeletalMeshAsset() or GetSkinnedAsset() instead.")
-	UFUNCTION(BlueprintPure, Category = "Components|SkinnedMesh", DisplayName = "Get Skeletal Mesh")
+	UFUNCTION(BlueprintPure, Category = "Components|SkinnedMesh", DisplayName = "Get Skeletal Mesh", meta = (DeprecatedFunction, DeprecationMessage = "Use USkeletalMeshComponent::GetSkeletalMesh() or GetSkinnedAsset() instead."))
 	class USkeletalMesh* GetSkeletalMesh_DEPRECATED() const;
 
 	/**
