[Feature] Blueprint Action Filter Caching

diff --git a/Engine/Source/Editor/BlueprintGraph/Private/BlueprintActionFilter.cpp b/Engine/Source/Editor/BlueprintGraph/Private/BlueprintActionFilter.cpp
--- a/Engine/Source/Editor/BlueprintGraph/Private/BlueprintActionFilter.cpp
+++ b/Engine/Source/Editor/BlueprintGraph/Private/BlueprintActionFilter.cpp
@@ -1961,7 +1961,7 @@ FBlueprintActionInfo::FBlueprintActionInfo(FBlueprintActionInfo const& Rhs, IBlu
 }
 
 //------------------------------------------------------------------------------
-UObject const* FBlueprintActionInfo::GetActionOwner()
+UObject const* FBlueprintActionInfo::GetActionOwner() const
 {
 	return ActionOwner.Get();
 }
@@ -2072,6 +2072,318 @@ UFunction const* FBlueprintActionInfo::GetAssociatedFunction()
 	return CachedActionFunction;
 }
 
+bool FBlueprintActionInfo::operator==(const FBlueprintActionInfo& other) const
+{
+	if (NodeSpawner != other.NodeSpawner || ActionOwner != other.ActionOwner)
+	{
+		return false;
+	}
+	
+	for (const FBindingObject &Binding : Bindings)
+	{
+		if (!other.Bindings.Contains(Binding))
+		{
+			return false;
+		}
+	}
+	return true;
+}
+
+/*******************************************************************************
+ * FActionFilterCacheNode
+ ******************************************************************************/
+
+uint32 GetTypeHash(const FBlueprintActionInfo& BlueprintActionInfo)
+{
+	uint32 Hash = GetTypeHash(BlueprintActionInfo.NodeSpawner);
+	Hash = HashCombine(Hash,GetTypeHash(BlueprintActionInfo.GetActionOwner()));
+	Hash = HashCombine(Hash,BlueprintActionInfo.GetBindings().Num());
+	
+	return Hash;
+}
+
+/*******************************************************************************
+ * TActionFilterCacheKeyNode
+ ******************************************************************************/
+
+template<typename TKeyType>
+class TActionFilterCacheKeyNode : public FActionFilterCacheNode
+{
+public:
+	explicit TActionFilterCacheKeyNode(const TKeyType& Key, TDelegate<bool(const TKeyType&, const TKeyType&)> Comparison = nullptr) :
+		Key(MakeShared<TKeyType>(Key)),
+		Comparison(Comparison)
+	{}
+	
+	virtual ~TActionFilterCacheKeyNode() override = default;
+
+	virtual void UnCacheSelf() override
+	{
+		if (Parent.IsValid())
+		{
+			const FSharedNode SharedParent = Parent.Pin();
+
+			// remove this node from parent
+			SharedParent->Children.RemoveSingleSwap(this->AsShared(), false);
+
+			if (SharedParent->Children.IsEmpty())
+			{
+				SharedParent->UnCacheSelf();
+			}
+		}
+	}
+	
+	TSharedPtr<TKeyType> Key = nullptr;
+protected:
+	virtual bool KeysMatch(const void* RawKey) const override
+	{
+		if (Comparison.IsBound())
+		{
+			return Comparison.Execute(*Key, *StaticCast<const TKeyType*>(RawKey));
+		}
+		return *Key == *StaticCast<const TKeyType*>(RawKey);
+	}
+
+	TDelegate<bool(const TKeyType&, const TKeyType&)> Comparison = nullptr;
+};
+
+/*******************************************************************************
+ * TActionFilterCacheLeaf
+ ******************************************************************************/
+
+template<typename TKeyType>
+class TActionFilterCacheLeaf : public TActionFilterCacheKeyNode<TKeyType>
+{
+public:
+	explicit TActionFilterCacheLeaf(const TKeyType& Key, TDelegate<bool(const TKeyType&, const TKeyType&)> Comparison = nullptr) :
+		TActionFilterCacheKeyNode<TKeyType>(Key, Comparison),
+		CachedFilterResults(MakeShared<TMap<FBlueprintActionInfo, bool>>())
+	{}
+	
+	virtual ~TActionFilterCacheLeaf() override = default;
+
+	virtual uint32 GetLeafAccessCount() const override
+	{
+		return AccessCount;
+	}
+
+
+	virtual void DecrementLeafAccessCount() override
+	{
+		if (AccessCount > 0)
+		{
+			--AccessCount;
+		}
+	}
+
+	uint32 AccessCount = 1;
+	TSharedPtr<TMap<FBlueprintActionInfo, bool>> CachedFilterResults;
+};
+
+/*******************************************************************************
+ * FActionFilterCache
+ ******************************************************************************/
+
+FActionFilterCache::FActionFilterCache()
+{
+	if (!ReloadCompleteDelegate.IsValid())
+	{
+		ReloadCompleteDelegate = FCoreUObjectDelegates::ReloadCompleteDelegate.AddStatic(&FActionFilterCache::OnReloadComplete);
+	}
+	if (!ModuleChangedDelegate.IsValid())
+	{
+		ModuleChangedDelegate = FModuleManager::Get().OnModulesChanged().AddStatic(&FActionFilterCache::OnModuleChanged);
+	}
+}
+
+void FActionFilterCache::OnReloadComplete(EReloadCompleteReason Reason)
+{
+	ClearAllCache();
+}
+
+void FActionFilterCache::OnModuleChanged(FName ModuleName, EModuleChangeReason ChangeReason)
+{
+	ClearAllCache();
+}
+
+void FActionFilterCache::ClearAllCache()
+{
+	SharedCache = MakeShared<FActionFilterCacheNode>();
+	CacheLeafs.Empty();
+}
+
+TSharedPtr<TMap<FBlueprintActionInfo, bool>> FActionFilterCache::GetCachedFilterResults(const FBlueprintActionFilter &Filter)
+{
+	// if the leaf for this context has already been looked up, just grab it rather than researching the tree
+	if (CachedFilterResults.IsValid())
+	{
+		return CachedFilterResults;
+	}
+	
+	TSharedPtr<FActionFilterCacheNode> Current = SharedCache;
+
+/** update Current to it's child that matches Key. Create one if neccesary */
+#define LOOKUP_CACHE_KEY(Key) \
+	FindOrCacheChildNode(Current, Key, #Key)
+	
+/** update Current to it's child that matches Key as a leaf. Create one if neccesary */
+#define LOOKUP_CACHE_KEY_LEAF(Key) \
+	FindOrCacheLeaf(Current, Key, #Key)
+
+/** update Current to it's child that matches Key as a leaf. Create one if neccesary
+ *  uses Comparison lambda to check for key "equality"
+ */
+#define LOOKUP_CACHE_KEY_PREDICATE(Key, Comparison) \
+	FindOrCacheChildNode<decltype(Key)>(Current, Key, #Key, TDelegate<bool(const decltype(Key)&, const decltype(Key)&)>::CreateLambda(Comparison))
+
+/** update Current to it's child that matches Key as a leaf. Create one if neccesary
+ *  uses Comparison lambda to check for key "equality"
+ */
+#define LOOKUP_CACHE_KEY_LEAF_PREDICATE(Key, Comparison) \
+	FindOrCacheLeaf<decltype(Key)>(Current, Key, #Key, TDelegate<bool(const decltype(Key)&, const decltype(Key)&)>::CreateLambda(Comparison))
+
+	// traverse/build the cache tree looking for the cached results
+
+	// cache based on what filter tests are being applied
+	LOOKUP_CACHE_KEY(Filter.FilterTests);
+
+	// cache based on which blueprints/graphs/editor we're generating actions for
+	LOOKUP_CACHE_KEY(Filter.Context.Blueprints);
+	LOOKUP_CACHE_KEY(Filter.Context.Graphs);
+	LOOKUP_CACHE_KEY(Filter.Context.EditorPtr);
+
+	// cache based on the filter's target classes
+	LOOKUP_CACHE_KEY(Filter.TargetClasses);
+
+	// cache based on what node types are permitted or rejected
+	LOOKUP_CACHE_KEY(Filter.RejectedNodeTypes);
+	LOOKUP_CACHE_KEY(Filter.PermittedNodeTypes);
+
+	// cache based on what flags are active
+	LOOKUP_CACHE_KEY(Filter.FilterFlags);
+	
+	// cache based on the relevant pin types and directions
+	LOOKUP_CACHE_KEY_LEAF_PREDICATE(Filter.Context.Pins, [](const TArray<UEdGraphPin*>& A, const TArray<UEdGraphPin*>& B)
+	{
+		if (A.Num() != B.Num())
+		{
+			return false;
+		}
+		for (int32 Index = 0; Index < A.Num(); ++Index)
+		{
+			if (A[Index]->Direction != B[Index]->Direction)
+			{
+				return false;
+			}
+			if (A[Index]->PinType != B[Index]->PinType)
+			{
+				return false;
+			}
+		}
+		return true;
+	});
+
+	// at this point, IsFilteredByThisCache has been updated to a shared memory space that can be retrieved again
+	// by following the above procedure with the same keys values.
+	return CachedFilterResults;
+}
+
+void FActionFilterCache::PopCache()
+{
+	int32 MinIndex = 0;
+	uint32 MinAccessCount = TNumericLimits<uint32>::Max();
+	for (int32 LeafIndex = 0; LeafIndex < CacheLeafs.Num(); ++LeafIndex)
+	{
+		TWeakPtr<FActionFilterCacheNode>& Leaf = CacheLeafs[LeafIndex];
+		// if there's invalid cache, remove it first
+		if (!Leaf.IsValid())
+		{
+			MinIndex = LeafIndex;
+			break;
+		}
+
+		const uint32 AccessCount = Leaf.Pin()->GetLeafAccessCount();
+		if (AccessCount < MinAccessCount)
+		{
+			MinAccessCount = AccessCount;
+			MinIndex = LeafIndex;
+		}
+
+		// decrement access count of all leafs so that recently used leafs get prioritized more
+		Leaf.Pin()->DecrementLeafAccessCount();
+	}
+
+	TWeakPtr<FActionFilterCacheNode> RemovedElement = CacheLeafs[MinIndex];
+	if (RemovedElement.IsValid())
+	{
+		RemovedElement.Pin()->UnCacheSelf();
+	}
+
+	CacheLeafs.RemoveAt(MinIndex, 1, false);
+}
+
+TSharedPtr<FActionFilterCacheNode> FActionFilterCache::SharedCache = MakeShared<FActionFilterCacheNode>();
+
+TAutoConsoleVariable<int32> FActionFilterCache::CVarCacheLeafCapacity{
+	TEXT("BP.ActionMenuFilterCacheLeafCapacity"),
+	32,
+	TEXT("The number of action menu contexts to cache simultaniously. raising this number will increase the memory footprint but decrease how often the cache is blown")
+};
+
+TArray<TWeakPtr<FActionFilterCacheNode>> FActionFilterCache::CacheLeafs;
+
+FDelegateHandle FActionFilterCache::ReloadCompleteDelegate;
+FDelegateHandle FActionFilterCache::ModuleChangedDelegate;
+
+/*******************************************************************************
+ * FActionFilterTest
+ ******************************************************************************/
+
+TAutoConsoleVariable<bool> FActionFilterTest::CVarEnableCaching{
+		TEXT("BP.EnableActionMenuFilterCaching"),
+		false,
+		TEXT("If enabled, action filter tests with the CacheResults flag set will have their results cached")
+	};
+
+FActionFilterTest::FActionFilterTest(const FRejectionDelegate Delegate, FString Name, EActionFilterTestFlags Flags) :
+	RejectionDelegate(Delegate),
+	Name(Name),
+	Flags(Flags)
+{
+}
+
+bool FActionFilterTest::Call(FBlueprintActionFilter const& Filter, FBlueprintActionInfo& ActionInfo,
+                             TMap<FString, uint32>* FilterTestToCpuSpecIds)
+{
+	check(RejectionDelegate.IsBound());
+
+#if ENABLE_BLUEPRINT_ACTION_FILTER_PROFILING
+	checkf(FilterTestToCpuSpecIds, TEXT("If Action Filter Profiling is enabled, FilterTestToCpuSpecIds must be provided"));
+	uint32& CpuSpecId = FilterTestToCpuSpecIds->FindOrAdd(Name);
+	FCpuProfilerTrace::FEventScope ProfilerEvent(CpuSpecId, *FString::Printf(TEXT("IsFiltered_%s"), *Name), true, __FILE__, __LINE__);
+
+	const double StartTime = FPlatformTime::Seconds();
+#endif	// ENABLE_BLUEPRINT_ACTION_FILTER_PROFILING
+
+	const bool bIsFiltered = RejectionDelegate.Execute(Filter, ActionInfo);
+
+#if ENABLE_BLUEPRINT_ACTION_FILTER_PROFILING
+	ProfileRecord.NumIterations += 1;
+	ProfileRecord.TotalTimeMs += (FPlatformTime::Seconds() - StartTime) * 1000.0f;
+	if (bIsFiltered)
+	{
+		ProfileRecord.NumFilteredOut += 1;
+	}
+#endif	// ENABLE_BLUEPRINT_ACTION_FILTER_PROFILING
+	
+	return bIsFiltered;
+}
+
+bool FActionFilterTest::operator==(const FActionFilterTest& other) const
+{
+	return Name == other.Name && Flags == other.Flags;
+}
+
 /*******************************************************************************
  * FBlueprintActionFilter
  ******************************************************************************/
@@ -2105,62 +2417,64 @@ FBlueprintActionFilter::FBlueprintActionFilter(const EFlags InFlags /*= BPFILTER
 	//       here are ran last)
 	//
 
-#define ADD_REJECTION_TEST(TestFunc, ...) \
-	AddRejectionTest(FRejectionTestDelegate::CreateStatic(TestFunc, ##__VA_ARGS__), TEXT(#TestFunc))
-
 	// add first the most expensive tests (they will be ran last, and therefore
 	// should be operating on a smaller subset of node-spawners)
 	//
 	// this test in-particular spawns a template-node and then calls 
 	// AllocateDefaultPins() which is costly, so it should be very last!
-	ADD_REJECTION_TEST(IsNodeTemplateSelfFiltered);
-	ADD_REJECTION_TEST(IsIncompatibleAnimNotification);
-	ADD_REJECTION_TEST(IsAssetPermissionNotGranted);
-	ADD_REJECTION_TEST(IsMissingMatchingPinParam);
-	ADD_REJECTION_TEST(IsMissmatchedPropertyType);
-	ADD_REJECTION_TEST(IsFunctionMissingPinParam);
-	ADD_REJECTION_TEST(IsIncompatibleLatentNode);
-	ADD_REJECTION_TEST(IsIncompatibleImpureNode);
-	ADD_REJECTION_TEST(IsPropertyAccessorNode);
+	AddRejectionTest(MAKE_ACTION_FILTER_REJECTION_TEST(IsNodeTemplateSelfFiltered)); // since this is virtual, we'll avoid caching it.
+	AddRejectionTest(MAKE_ACTION_FILTER_REJECTION_TEST(IsIncompatibleAnimNotification)->WithFlags(EActionFilterTestFlags::CacheResults));
+	AddRejectionTest(MAKE_ACTION_FILTER_REJECTION_TEST(IsAssetPermissionNotGranted)); // since this depends on other filters, we'll avoid caching it
+
+	// TODO: blow cache when the results of the following might change:
+	//   - K2Schema->FunctionHasParamOfType
+	//   - K2Schema->ArePinTypesCompatible
+	//   - K2Schema->ConvertPropertyToPinType
+	AddRejectionTest(MAKE_ACTION_FILTER_REJECTION_TEST(IsMissingMatchingPinParam)->WithFlags(EActionFilterTestFlags::CacheResults));
+	AddRejectionTest(MAKE_ACTION_FILTER_REJECTION_TEST(IsMissmatchedPropertyType)->WithFlags(EActionFilterTestFlags::CacheResults));
+	AddRejectionTest(MAKE_ACTION_FILTER_REJECTION_TEST(IsFunctionMissingPinParam)->WithFlags(EActionFilterTestFlags::CacheResults));
+	AddRejectionTest(MAKE_ACTION_FILTER_REJECTION_TEST(IsIncompatibleLatentNode)->WithFlags(EActionFilterTestFlags::CacheResults));
+	AddRejectionTest(MAKE_ACTION_FILTER_REJECTION_TEST(IsIncompatibleImpureNode)->WithFlags(EActionFilterTestFlags::CacheResults));
+	AddRejectionTest(MAKE_ACTION_FILTER_REJECTION_TEST(IsPropertyAccessorNode)->WithFlags(EActionFilterTestFlags::CacheResults));
 	if (HasAnyFlags(BPFILTER_RejectIncompatibleThreadSafety))
 	{
-		ADD_REJECTION_TEST(IsThreadSafetyIncompatible);
+		AddRejectionTest(MAKE_ACTION_FILTER_REJECTION_TEST(IsThreadSafetyIncompatible)->WithFlags(EActionFilterTestFlags::CacheResults));
 	}
 	
-	ADD_REJECTION_TEST(IsActionHiddenByConfig);
-	ADD_REJECTION_TEST(IsFieldCategoryHidden);
+	AddRejectionTest(MAKE_ACTION_FILTER_REJECTION_TEST(IsActionHiddenByConfig)->WithFlags(EActionFilterTestFlags::CacheResults));
+	AddRejectionTest(MAKE_ACTION_FILTER_REJECTION_TEST(IsFieldCategoryHidden)->WithFlags(EActionFilterTestFlags::CacheResults));
 	if (HasAnyFlags(BPFILTER_RejectGlobalFields | BPFILTER_RejectNonImportedFields))
 	{
-		ADD_REJECTION_TEST(IsRejectedGlobalField);
+		AddRejectionTest(MAKE_ACTION_FILTER_REJECTION_TEST(IsRejectedGlobalField)->WithFlags(EActionFilterTestFlags::CacheResults));
 	}
 	
-	ADD_REJECTION_TEST(IsFieldInaccessible);
-	ADD_REJECTION_TEST(IsNotSubClassCast);
-	ADD_REJECTION_TEST(IsEventUnimplementable);
-	ADD_REJECTION_TEST(IsPermissionNotGranted);
-	ADD_REJECTION_TEST(IsRestrictedClassMember);
-	ADD_REJECTION_TEST(IsIncompatibleWithGraphType);
-	ADD_REJECTION_TEST(IsSchemaIncompatible);
-	ADD_REJECTION_TEST(IsExtraneousInterfaceCall);
-	ADD_REJECTION_TEST(IsIncompatibleMacroInstance);
+	AddRejectionTest(MAKE_ACTION_FILTER_REJECTION_TEST(IsFieldInaccessible)->WithFlags(EActionFilterTestFlags::CacheResults));
+	AddRejectionTest(MAKE_ACTION_FILTER_REJECTION_TEST(IsNotSubClassCast)->WithFlags(EActionFilterTestFlags::CacheResults));
+	AddRejectionTest(MAKE_ACTION_FILTER_REJECTION_TEST(IsEventUnimplementable)->WithFlags(EActionFilterTestFlags::CacheResults));
+	AddRejectionTest(MAKE_ACTION_FILTER_REJECTION_TEST(IsPermissionNotGranted)->WithFlags(EActionFilterTestFlags::CacheResults));
+	AddRejectionTest(MAKE_ACTION_FILTER_REJECTION_TEST(IsRestrictedClassMember)->WithFlags(EActionFilterTestFlags::CacheResults));
+	AddRejectionTest(MAKE_ACTION_FILTER_REJECTION_TEST(IsIncompatibleWithGraphType)->WithFlags(EActionFilterTestFlags::CacheResults));
+	AddRejectionTest(MAKE_ACTION_FILTER_REJECTION_TEST(IsSchemaIncompatible)->WithFlags(EActionFilterTestFlags::CacheResults));
+	AddRejectionTest(MAKE_ACTION_FILTER_REJECTION_TEST(IsExtraneousInterfaceCall)->WithFlags(EActionFilterTestFlags::CacheResults));
+	AddRejectionTest(MAKE_ACTION_FILTER_REJECTION_TEST(IsIncompatibleMacroInstance)->WithFlags(EActionFilterTestFlags::CacheResults));
 
 	if (!HasAnyFlags(BPFILTER_PermitDeprecated))
 	{
-		ADD_REJECTION_TEST(IsDeprecated);
+		AddRejectionTest(MAKE_ACTION_FILTER_REJECTION_TEST(IsDeprecated)->WithFlags(EActionFilterTestFlags::CacheResults));
 	}
 
-	ADD_REJECTION_TEST(IsFilteredNodeType);
-	ADD_REJECTION_TEST(IsNonTargetMember);
-	ADD_REJECTION_TEST(IsUnBoundBindingSpawner);
-	ADD_REJECTION_TEST(IsOutOfScopeLocalVariable);
-	ADD_REJECTION_TEST(IsLevelScriptActionValid);
+	AddRejectionTest(MAKE_ACTION_FILTER_REJECTION_TEST(IsFilteredNodeType)->WithFlags(EActionFilterTestFlags::CacheResults));
+	AddRejectionTest(MAKE_ACTION_FILTER_REJECTION_TEST(IsNonTargetMember)->WithFlags(EActionFilterTestFlags::CacheResults));
+	AddRejectionTest(MAKE_ACTION_FILTER_REJECTION_TEST(IsUnBoundBindingSpawner)->WithFlags(EActionFilterTestFlags::CacheResults));
+	AddRejectionTest(MAKE_ACTION_FILTER_REJECTION_TEST(IsOutOfScopeLocalVariable)->WithFlags(EActionFilterTestFlags::CacheResults));
+	AddRejectionTest(MAKE_ACTION_FILTER_REJECTION_TEST(IsLevelScriptActionValid)->WithFlags(EActionFilterTestFlags::CacheResults));
 
-	ADD_REJECTION_TEST(IsHiddenInNonEditorBlueprint);
+	AddRejectionTest(MAKE_ACTION_FILTER_REJECTION_TEST(IsHiddenInNonEditorBlueprint)->WithFlags(EActionFilterTestFlags::CacheResults));
 
 
 	// added as the first rejection test, so that we don't operate on stale 
 	// (TRASH/REINST) class fields
-	ADD_REJECTION_TEST(IsStaleFieldAction);
+	AddRejectionTest(MAKE_ACTION_FILTER_REJECTION_TEST(IsStaleFieldAction));
 }
 
 //------------------------------------------------------------------------------
@@ -2186,24 +2500,10 @@ void FBlueprintActionFilter::Add(TArray<FTargetClassFilterData>& ToArray, UClass
 }
 
 //------------------------------------------------------------------------------
-void FBlueprintActionFilter::AddRejectionTest(FRejectionTestDelegate RejectionTestDelegate, const TCHAR* RejectionTestName)
+void FBlueprintActionFilter::AddRejectionTest(TSharedRef<FActionFilterTest> RejectionTest)
 {
-	if (RejectionTestDelegate.IsBound())
-	{
-		FilterTests.Add(RejectionTestDelegate);
-
-#if ENABLE_BLUEPRINT_ACTION_FILTER_PROFILING
-		FFilterTestProfileRecord FilterTestProfileRecord;
-		
-		FilterTestProfileRecord.TestName = RejectionTestName ? FString(RejectionTestName) : RejectionTestDelegate.TryGetBoundFunctionName().ToString();
-		if (FilterTestProfileRecord.TestName.IsEmpty())
-		{
-			FilterTestProfileRecord.TestName = FString::Printf(TEXT("[FuncPtr:0x%016x]"), RejectionTestDelegate.GetBoundProgramCounterForTimerManager());
-		}
-
-		FilterTestProfiles.Add(FilterTests.Num() - 1, MoveTemp(FilterTestProfileRecord));
-#endif	// ENABLE_BLUEPRINT_ACTION_FILTER_PROFILING
-	}
+	checkf(RejectionTest->RejectionDelegate.IsBound(), TEXT("Cannot add a rejection test without a bound delegate"));
+	FilterTests.Add(RejectionTest);
 }
 
 //------------------------------------------------------------------------------
@@ -2253,8 +2553,24 @@ FBlueprintActionFilter const& FBlueprintActionFilter::operator&=(FBlueprintActio
 }
 
 //------------------------------------------------------------------------------
-bool FBlueprintActionFilter::IsFilteredByThis(FBlueprintActionInfo& BlueprintAction) const
+TSharedPtr<TMap<FBlueprintActionInfo, bool>> FBlueprintActionFilter::GetCachedFilterResults()
+{
+	return ActionFilterCache.GetCachedFilterResults(*this);
+}
+
+const TArray<TSharedRef<FActionFilterTest>> &FBlueprintActionFilter::GetFilterTests() const
+{
+	return FilterTests;
+}
+
+//------------------------------------------------------------------------------
+bool FBlueprintActionFilter::IsFilteredByThis(FBlueprintActionInfo& BlueprintAction)
 {
+	if (IsFilteredByCachedFilters(BlueprintAction))
+	{
+		return true;
+	}
+	
 	FBlueprintActionFilter const& FilterRef = *this;
 
 	// for debugging purposes:
@@ -2274,28 +2590,20 @@ bool FBlueprintActionFilter::IsFilteredByThis(FBlueprintActionInfo& BlueprintAct
 	// internal tests are ran last).
 	for (int32 TestIndex = FilterTests.Num()-1; TestIndex >= 0; --TestIndex)
 	{
-		FRejectionTestDelegate const& RejectionTestDelegate = FilterTests[TestIndex];
-		check(RejectionTestDelegate.IsBound());
-
-#if ENABLE_BLUEPRINT_ACTION_FILTER_PROFILING
-		FFilterTestProfileRecord& FilterTestProfileRecord = FilterTestProfiles.FindOrAdd(TestIndex);
-		uint32& CpuSpecId = FilterTestToCpuSpecIds.FindOrAdd(FilterTestProfileRecord.TestName);
-		FCpuProfilerTrace::FEventScope ProfilerEvent(CpuSpecId, *FString::Printf(TEXT("IsFiltered_%s"), *FilterTestProfileRecord.TestName), true, __FILE__, __LINE__);
-
-		const double StartTime = FPlatformTime::Seconds();
-#endif	// ENABLE_BLUEPRINT_ACTION_FILTER_PROFILING
-
-		bIsFiltered = RejectionTestDelegate.Execute(FilterRef, BlueprintAction);
+		TSharedRef<FActionFilterTest> RejectionTest = FilterTests[TestIndex];
 
-#if ENABLE_BLUEPRINT_ACTION_FILTER_PROFILING
-		FilterTestProfileRecord.NumIterations += 1;
-		FilterTestProfileRecord.TotalTimeMs += (FPlatformTime::Seconds() - StartTime) * 1000.0f;
-		if (bIsFiltered)
+		// cached tests have already been ran. we can skip them.
+		if (EnumHasAllFlags(RejectionTest->Flags, EActionFilterTestFlags::CacheResults))
 		{
-			FilterTestProfileRecord.NumFilteredOut += 1;
+			continue;
 		}
-#endif	// ENABLE_BLUEPRINT_ACTION_FILTER_PROFILING
 
+#if ENABLE_BLUEPRINT_ACTION_FILTER_PROFILING
+		bIsFiltered = RejectionTest.Call(FilterRef, BlueprintAction, &FilterTestToCpuSpecIds);
+#else
+		bIsFiltered = RejectionTest->Call(FilterRef, BlueprintAction);
+#endif
+		
 		if (bIsFiltered)
 		{
 			break;
@@ -2304,46 +2612,81 @@ bool FBlueprintActionFilter::IsFilteredByThis(FBlueprintActionInfo& BlueprintAct
 
 	if (!bIsFiltered)
 	{
-#if ENABLE_BLUEPRINT_ACTION_FILTER_PROFILING
-		int32 ExtraTestProfileRecordIndex = FilterTests.Num();
-#endif	// ENABLE_BLUEPRINT_ACTION_FILTER_PROFILING
-
 		for (const FBlueprintGraphModule::FActionMenuRejectionTest& ExtraRejectionTest : BluprintGraphModule->GetExtendedActionMenuFilters())
 		{
-#if ENABLE_BLUEPRINT_ACTION_FILTER_PROFILING
-			FFilterTestProfileRecord& ExtraTestProfileRecord = FilterTestProfiles.FindOrAdd(ExtraTestProfileRecordIndex++);
-			if (ExtraTestProfileRecord.TestName.IsEmpty())
+			bIsFiltered = ExtraRejectionTest.Execute(FilterRef, BlueprintAction);
+			if (bIsFiltered)
 			{
-				ExtraTestProfileRecord.TestName = ExtraRejectionTest.TryGetBoundFunctionName().ToString();
-				if (ExtraTestProfileRecord.TestName.IsEmpty())
-				{
-					ExtraTestProfileRecord.TestName = FString::Printf(TEXT("[FuncPtr:0x%016x]"), ExtraRejectionTest.GetBoundProgramCounterForTimerManager());
-				}
+				break;
 			}
-			uint32& CpuSpecId = FilterTestToCpuSpecIds.FindOrAdd(ExtraTestProfileRecord.TestName);
-			FCpuProfilerTrace::FEventScope ProfilerEvent(CpuSpecId, *FString::Printf(TEXT("Extended_IsFiltered_%s"), *ExtraTestProfileRecord.TestName), true, __FILE__, __LINE__);
-
-			const double StartTime = FPlatformTime::Seconds();
-#endif	// ENABLE_BLUEPRINT_ACTION_FILTER_PROFILING
+		}
+	}
 
-			bIsFiltered = ExtraRejectionTest.Execute(FilterRef, BlueprintAction);
+	return bIsFiltered;
+}
 
+bool FBlueprintActionFilter::IsFilteredByCachedFilters(FBlueprintActionInfo& BlueprintAction)
+{
+	
 #if ENABLE_BLUEPRINT_ACTION_FILTER_PROFILING
-			ExtraTestProfileRecord.NumIterations += 1;
-			ExtraTestProfileRecord.TotalTimeMs += (FPlatformTime::Seconds() - StartTime) * 1000.0f;
-			if (bIsFiltered)
-			{
-				ExtraTestProfileRecord.NumFilteredOut += 1;
-			}
-#endif	// ENABLE_BLUEPRINT_ACTION_FILTER_PROFILING
+	TRACE_CPUPROFILER_EVENT_SCOPE(FBlueprintActionFilter::IsFilteredByCachedFilters);
+#endif
+	
+	const TSharedPtr<TMap<FBlueprintActionInfo, bool>> Cache = GetCachedFilterResults();
+	FBlueprintActionFilter const& FilterRef = *this;
+		
+	if (BlueprintActionFilterImpl::IsStaleFieldAction(FilterRef, BlueprintAction))
+	{
+		// if BlueprintAction is stale, immediately filter the action
+		if (FActionFilterTest::CVarEnableCaching.GetValueOnGameThread())
+		{
+			Cache->FindOrAdd(BlueprintAction) = true;
+		}
+		return true;
+	}
 
-			if (bIsFiltered)
-			{
-				break;
-			}
+	if (FActionFilterTest::CVarEnableCaching.GetValueOnGameThread())
+	{
+		if (const bool* Result = Cache->Find(BlueprintAction))
+		{
+			return *Result;
 		}
 	}
 
+#if ENABLE_BLUEPRINT_ACTION_FILTER_PROFILING
+	static TMap<FString, uint32> FilterTestToCpuSpecIds;
+#endif	// ENABLE_BLUEPRINT_ACTION_FILTER_PROFILING
+	
+	bool bIsFiltered = false;
+	// iterate backwards so that custom user test are ran first (and the slow
+	// internal tests are ran last).
+	for (int32 TestIndex = FilterTests.Num()-1; TestIndex >= 0; --TestIndex)
+	{
+		TSharedRef<FActionFilterTest> RejectionTest = FilterTests[TestIndex];
+		
+		// this method only checks cached rejection tests
+		if (!EnumHasAllFlags(RejectionTest->Flags, EActionFilterTestFlags::CacheResults))
+		{
+			continue;
+		}
+		
+#if ENABLE_BLUEPRINT_ACTION_FILTER_PROFILING
+		bIsFiltered = RejectionTest.Call(FilterRef, BlueprintAction, &FilterTestToCpuSpecIds);
+#else
+		bIsFiltered = RejectionTest->Call(FilterRef, BlueprintAction);
+#endif
+		
+		if (bIsFiltered)
+		{
+			break;
+		}
+	}
+	
+	if (FActionFilterTest::CVarEnableCaching.GetValueOnGameThread())
+	{
+		Cache->Add(BlueprintAction, bIsFiltered);
+	}
+	
 	return bIsFiltered;
 }
 
@@ -2353,14 +2696,24 @@ TArray<FString> FBlueprintActionFilter::GetFilterTestProfile()
 {
 	TArray<FString> Result;
 
-	int32 TestIdx = 0;
 	for (int32 TestIndex = FilterTests.Num() - 1; TestIndex >= 0; --TestIndex)
 	{
-		Result.Add(FString::Printf(TEXT("[%02d]: %s | %d tested | %d failed | %.02f ms | %.02f us (per test)"), FilterTests.Num() - TestIndex, *FilterTestProfiles[TestIndex].TestName, FilterTestProfiles[TestIndex].NumIterations, FilterTestProfiles[TestIndex].NumFilteredOut, FilterTestProfiles[TestIndex].TotalTimeMs, FilterTestProfiles[TestIndex].TotalTimeMs / FilterTestProfiles[TestIndex].NumIterations * 1000.0f));
+		FActionFilterTest& RejectionTest = FilterTests[TestIndex];
+		FActionFilterTest::FFilterTestProfileRecord& ProfileRecord = RejectionTest.ProfileRecord;
+		
+		Result.Add(FString::Printf(TEXT("[%02d]: %s | %d tested | %d failed | %.02f ms | %.02f us (per test)"),
+			FilterTests.Num() - TestIndex,
+			*RejectionTest.Name,
+			ProfileRecord.NumIterations,
+			ProfileRecord.NumFilteredOut,
+			ProfileRecord.TotalTimeMs,
+			ProfileRecord.TotalTimeMs / ProfileRecord.NumIterations * 1000.0f
+		));
+
+		// reset record
+		ProfileRecord = FActionFilterTest::FFilterTestProfileRecord();
 	}
 
-	FilterTestProfiles.Reset();
-
 	return Result;
 }
 #endif	// ENABLE_BLUEPRINT_ACTION_FILTER_PROFILING
\ No newline at end of file
diff --git a/Engine/Source/Editor/BlueprintGraph/Public/BlueprintActionFilter.h b/Engine/Source/Editor/BlueprintGraph/Public/BlueprintActionFilter.h
--- a/Engine/Source/Editor/BlueprintGraph/Public/BlueprintActionFilter.h
+++ b/Engine/Source/Editor/BlueprintGraph/Public/BlueprintActionFilter.h
@@ -13,6 +13,8 @@
 
 #define ENABLE_BLUEPRINT_ACTION_FILTER_PROFILING (!(UE_BUILD_SHIPPING || UE_BUILD_TEST) && 0)
 
+
+
 class UBlueprint;
 class UBlueprintNodeSpawner;
 class UEdGraph;
@@ -20,6 +22,11 @@ class UEdGraphPin;
 class IAssetReferenceFilter;
 class IBlueprintEditor;
 
+class FActionFilterCacheNode;
+template<typename TKeyType>
+class TActionFilterCacheKeyNode;
+template<typename TKeyType>
+class TActionFilterCacheLeaf;
 /*******************************************************************************
  * FBlueprintActionContext
  ******************************************************************************/
@@ -85,7 +92,7 @@ struct BLUEPRINTGRAPH_API FBlueprintActionInfo
 	 * 
 	 * @return The class/asset that this action conceptually belongs to.
 	 */
-	UObject const* GetActionOwner();
+	UObject const* GetActionOwner() const;
 
 	/**
 	 * Retrieves any bindings that the action will apply to spawned nodes.
@@ -140,6 +147,8 @@ struct BLUEPRINTGRAPH_API FBlueprintActionInfo
 	 */
 	UFunction const* GetAssociatedFunction();
 
+	bool operator==(const FBlueprintActionInfo& other) const;
+
 	/** The raw action that this struct represent (const so we don't mutate the database) */
 	UBlueprintNodeSpawner const* const NodeSpawner;
 
@@ -160,6 +169,271 @@ private:
 	IBlueprintNodeBinder::FBindingSet Bindings;
 };
 
+uint32 GetTypeHash(const FBlueprintActionInfo& BlueprintActionInfo);
+
+/*******************************************************************************
+ * FActionFilterCacheNode
+ ******************************************************************************/
+
+/**
+ * This class is used by FActionFilterCache to create a tree structure that is used to store cached filter results in
+ * such a way that they can be looked up using all the relevant context data that determines how items are filtered.
+ * each node has a type agnostic key so that we can easily change how the cache is keyed in the future.
+ * (See FActionFilterCache::GetIsFilteredByThisCache)
+ */
+class FActionFilterCacheNode : public TSharedFromThis<FActionFilterCacheNode>
+{
+public:
+	template<typename TKeyType>
+	friend class TActionFilterCacheKeyNode;
+
+	// Branching nodes
+	using FWeakNode = TWeakPtr<FActionFilterCacheNode>;
+	template<typename TKeyType>
+	using TWeakKeyNode = TWeakPtr<TActionFilterCacheKeyNode<TKeyType>>;
+	template<typename TKeyType>
+	using TWeakLeaf = TWeakPtr<TActionFilterCacheLeaf<TKeyType>>;
+
+	// Leaf nodes
+	using FSharedNode = TSharedPtr<FActionFilterCacheNode>;
+	template<typename TKeyType>
+	using TSharedKeyNode = TSharedPtr<TActionFilterCacheKeyNode<TKeyType>>;
+	template<typename TKeyType>
+	using TSharedLeaf = TSharedPtr<TActionFilterCacheLeaf<TKeyType>>;
+	
+	FActionFilterCacheNode() = default;
+	virtual ~FActionFilterCacheNode() = default;
+
+	/** Looks up a child of this node by key. Note: the key is not hashed. This is O(n) */
+	template<typename TKeyType>
+	TSharedKeyNode<TKeyType> Find(const TKeyType& Key) const
+	{
+		for (const TSharedPtr<FActionFilterCacheNode>& Child : Children)
+		{
+			if (Child->KeysMatch(&Key))
+			{
+				return StaticCastSharedPtr<TActionFilterCacheKeyNode<TKeyType>>(Child);
+			}
+		}
+		return nullptr;
+	}
+
+	/** Looks up a child of this node by key and creates one if one doesnt exist.
+	 *	The node that is created will be a TActionFilterCacheKeyNode. Use FindOrCacheChildLeaf if you need it to be a
+	 *	TActionFilterCacheLeaf
+	 *  Note: the key is not hashed. This is O(n) */
+	template<typename TKeyType>
+	TSharedKeyNode<TKeyType> FindOrCacheChildNode(const TKeyType& Key, const FString& DebugName = FString(), TDelegate<bool(const TKeyType&, const TKeyType&)> Comparison = nullptr, bool *OutFoundWithoutCaching = nullptr)
+	{
+		TSharedKeyNode<TKeyType> Result = Find(Key);
+		if (OutFoundWithoutCaching)
+		{
+			*OutFoundWithoutCaching = Result.IsValid();
+		}
+		
+		if (Result.IsValid())
+		{
+			return Result;
+		}
+		
+#if UE_BUILD_DEBUG
+		this->DebugName = DebugName;
+#endif
+		return Cache(new TActionFilterCacheKeyNode<TKeyType>(Key, Comparison));
+	}
+
+	/** Looks up a child of this node by key and creates one if one doesnt exist.
+	 *	The node that is created will be a TActionFilterCacheLeaf. Use FindOrCacheChildNode if you need it to be a
+	 *	TActionFilterCacheKeyNode
+	 *  Note: the key is not hashed. This is O(n) */
+	template<typename TKeyType>
+	TSharedLeaf<TKeyType> FindOrCacheChildLeaf(const TKeyType& Key, const FString& DebugName = FString(), TDelegate<bool(const TKeyType&, const TKeyType&)> Comparison = nullptr, bool *OutFoundWithoutCaching = nullptr)
+	{
+		TSharedKeyNode<TKeyType> Result = Find(Key);
+		if (OutFoundWithoutCaching)
+		{
+			*OutFoundWithoutCaching = Result.IsValid();
+		}
+		
+		if (Result.IsValid())
+		{
+			TSharedLeaf<TKeyType> Leaf = StaticCastSharedPtr<TActionFilterCacheLeaf<TKeyType>>(Result);
+			++Leaf->AccessCount;
+			return Leaf;
+		}
+		return Cache(new TActionFilterCacheLeaf<TKeyType>(Key, Comparison));
+	}
+
+	/** get the number of times this leaf has been accessed
+	 *  Warning: throws and error if this isn't a TActionFilterCacheLeaf */
+	virtual uint32 GetLeafAccessCount() const
+	{
+		checkf(false, TEXT("GetLeafAccessCount was called on a non-leaf node"))
+		return 0;
+	}
+
+	/** decrement the number of times this leaf has been accessed
+	 *  Warning: throws and error if this isn't a TActionFilterCacheLeaf */
+	virtual void DecrementLeafAccessCount()
+	{
+		checkf(false, TEXT("DecrementLeafAccessCount was called on a non-leaf node"))
+	}
+
+	/** remove this node from the the tree */
+	virtual void UnCacheSelf(){}
+
+#if UE_BUILD_DEBUG
+	FString DebugName;
+#endif
+
+protected:
+	virtual bool KeysMatch(const void* RawKey) const
+	{
+		checkf(false, TEXT("Tried to call KeysMatch on a Node with no keys"))
+		return false;
+	}
+
+	template<typename TNodeType>
+	TSharedPtr<TNodeType> Cache(TNodeType* Child)
+	{
+		Child->Parent = AsShared();
+		Children.Emplace(MakeShareable<FActionFilterCacheNode>(Child));
+		return StaticCastSharedPtr<TNodeType>(Children.Last());
+	}
+	
+	TArray<FSharedNode> Children;
+	FWeakNode Parent;
+};
+
+/*******************************************************************************
+ * FActionFilterCache
+ ******************************************************************************/
+
+/**
+ * This class creates a tree structure that is used to store cached filter results in such a way that they can be
+ * looked up using all the relevant context data that determines how items are filtered.
+ * each node has a type agnostic key so that we can easily change how the cache is keyed.
+ * (See FActionFilterCache::GetIsFilteredByThisCache)
+ */
+class FActionFilterCache
+{
+public:
+	FActionFilterCache();
+
+	static void OnReloadComplete(EReloadCompleteReason Reason);
+	static void OnModuleChanged(FName ModuleName, EModuleChangeReason ChangeReason);
+	static void ClearAllCache();
+
+	/** get the cached map used by FBlueprintActionFilter::IsFilteredByThis to determine whether BP actions are filtered out */
+	TSharedPtr<TMap<FBlueprintActionInfo, bool>> GetCachedFilterResults(const FBlueprintActionFilter &Filter);
+
+private:
+	template<typename TKeyType>
+	void FindOrCacheChildNode(TSharedPtr<FActionFilterCacheNode> &Current, const TKeyType& Key, FString DebugName, TDelegate<bool(const TKeyType&, const TKeyType&)> Comparison = nullptr)
+	{
+		TSharedPtr<TActionFilterCacheKeyNode<TKeyType>> Found = Current->FindOrCacheChildNode(Key, DebugName, Comparison);
+		Current = StaticCastSharedPtr<FActionFilterCacheNode>(Found);
+		
+	}
+
+	template<typename TKeyType>
+	void FindOrCacheLeaf(TSharedPtr<FActionFilterCacheNode> &Current, const TKeyType& Key, FString DebugName, TDelegate<bool(const TKeyType&, const TKeyType&)> Comparison = nullptr)
+	{
+		bool bFoundWithoutCaching = false;
+		TSharedPtr<TActionFilterCacheLeaf<TKeyType>> Found = Current->FindOrCacheChildLeaf(Key, DebugName, Comparison, &bFoundWithoutCaching);
+		CachedFilterResults = Found->CachedFilterResults;
+		Current = StaticCastSharedPtr<FActionFilterCacheNode>(Found);
+
+		// if we just expanded the cache, log it and check if the cache is getting too big
+		if (!bFoundWithoutCaching)
+		{
+			if (CacheLeafs.Num() >= CVarCacheLeafCapacity.GetValueOnGameThread())
+			{
+				PopCache();
+			}
+			CacheLeafs.Push(Found);
+		}
+	}
+
+	/** Find the least likely item in the cache to be used again, and remove it to make room for more data */
+	static void PopCache();
+
+private:
+	/** the shared tree that stores previously calculated filter results for several different contects */
+	static TSharedPtr<FActionFilterCacheNode> SharedCache;
+	
+	static TAutoConsoleVariable<int32> CVarCacheLeafCapacity;
+
+	/** array of all the leaf nodes in SharedCache so we can pop a rarely used node if we exceed CacheLeafCapacity */
+	static TArray<TWeakPtr<FActionFilterCacheNode>> CacheLeafs;
+
+	static FDelegateHandle ReloadCompleteDelegate;
+	static FDelegateHandle ModuleChangedDelegate;
+
+	/** cached results from GetIsFilteredByThisCache so we don't have to re-search the tree */
+	TSharedPtr<TMap<FBlueprintActionInfo, bool>> CachedFilterResults;
+};
+
+/*******************************************************************************
+ * EBlueprintActionFilterRejectionTestFlags
+ ******************************************************************************/
+
+enum class EActionFilterTestFlags
+{
+	None = 0,
+	
+	/**
+	 * these tests will be have their combined result cached for faster lookup.
+	 * note: rejection tests can only be cached if they either
+	 *		A: Don't access any mutable data
+	 *		or B: Only access mutable data that is listed as a filter key in FActionFilterCache::GetCachedFilterResults
+	 */
+	CacheResults = 1 << 0,
+	
+	Default = None, // this would allow us to change the default behavior
+};
+
+ENUM_CLASS_FLAGS(EActionFilterTestFlags);
+
+/*******************************************************************************
+ * FActionFilterTest
+******************************************************************************/
+
+struct BLUEPRINTGRAPH_API FActionFilterTest : TSharedFromThis<FActionFilterTest>
+{
+	static TAutoConsoleVariable<bool> CVarEnableCaching;
+	
+	/** The filter uses a series of rejection tests matching */
+	typedef FBlueprintGraphModule::FActionMenuRejectionTest FRejectionDelegate;
+
+	FActionFilterTest(const FRejectionDelegate Delegate, FString Name, EActionFilterTestFlags Flags = EActionFilterTestFlags::Default);
+
+	TSharedRef<FActionFilterTest> WithFlags(EActionFilterTestFlags InFlags)
+	{
+		EnumAddFlags(Flags, InFlags);
+		return AsShared();
+	}
+	
+	bool Call(FBlueprintActionFilter const& Filter, FBlueprintActionInfo& ActionInfo, TMap<FString, uint32> *FilterTestToCpuSpecIds = nullptr);
+	bool operator==(const FActionFilterTest& other) const;
+	
+	FRejectionDelegate RejectionDelegate;
+	FString Name; // this name must be unique
+	EActionFilterTestFlags Flags = EActionFilterTestFlags::Default;
+	
+#if ENABLE_BLUEPRINT_ACTION_FILTER_PROFILING
+	struct FFilterTestProfileRecord
+	{
+		float TotalTimeMs = 0.0f;
+		int32 NumIterations = 0;
+		int32 NumFilteredOut = 0;
+	} ProfileRecord;
+#endif	// ENABLE_BLUEPRINT_ACTION_FILTER_PROFILING
+};
+
+// constructor for TSharedRef<FBlueprintActionFilterRejectionTest> that auto deduces the name
+#define MAKE_ACTION_FILTER_REJECTION_TEST(TestFunc, ...) \
+	MakeShared<FActionFilterTest>(FActionFilterTest::FRejectionDelegate::CreateStatic(TestFunc, ##__VA_ARGS__), TEXT(#TestFunc))
 
 /*******************************************************************************
  * FBlueprintActionFilter
@@ -168,8 +442,7 @@ private:
 class BLUEPRINTGRAPH_API FBlueprintActionFilter
 {
 public:
-	/** The filter uses a series of rejection tests matching */
-	typedef FBlueprintGraphModule::FActionMenuRejectionTest FRejectionTestDelegate;
+	friend FActionFilterCache;
 
 public:
 	enum EFlags // Flags, which configure certain rejection tests.
@@ -252,6 +525,10 @@ public:
 	{
 		UClass* TargetClass;
 		TArray<FString> HiddenCategories;
+		bool operator==(const FTargetClassFilterData& Other) const
+		{
+			return TargetClass == Other.TargetClass && HiddenCategories == Other.HiddenCategories;
+		}
 	};
 	TArray<FTargetClassFilterData> TargetClasses;
 
@@ -272,10 +549,10 @@ public:
 	 * method. We use rejection "IsFiltered" tests rather than inclusive tests 
 	 * because it is more optimal to whittle down the list of actions early.
 	 * 
-	 * @param  RejectionTestDelegate	The rejection test you wish to add to this filter.
-	 * @param  RejectionTestName		(Optional) Test name (used for profiling/debugging).
+	 * @param  RejectionTest	a test this filter will run to cull blueprint actions
+	 * @param  Flags	List of flags that will be added in the provided rejection test
 	 */
-	void AddRejectionTest(FRejectionTestDelegate RejectionTestDelegate, const TCHAR* RejectionTestName = nullptr);
+	void AddRejectionTest(TSharedRef<FActionFilterTest> RejectionTest);
 
 	/**
 	 * Query to check and see if the specified action gets filtered out by this 
@@ -310,6 +587,10 @@ public:
 	 * @return This.
 	 */
 	FBlueprintActionFilter const& operator&=(FBlueprintActionFilter const& Rhs);
+
+	TSharedPtr<TMap<FBlueprintActionInfo, bool>> GetCachedFilterResults();
+
+	const TArray<TSharedRef<FActionFilterTest>> &GetFilterTests() const;
 	
 #if ENABLE_BLUEPRINT_ACTION_FILTER_PROFILING
 	TArray<FString> GetFilterTestProfile();
@@ -337,15 +618,24 @@ private:
 	 * (does not take into consideration any and'd/or'd filters).
 	 * 
 	 * @param  BlueprintAction	The node-spawner you wish to test.
-	 * @return False if the action passes the filter, otherwise false (the action got filtered out).
+	 * @return False if the action passes the filter, otherwise true (the action got filtered out).
 	 */
-	bool IsFilteredByThis(FBlueprintActionInfo& BlueprintAction) const;
+	bool IsFilteredByThis(FBlueprintActionInfo& BlueprintAction);
+
+	/**
+	 * Query to check and see if the specified action gets filtered out by the filter tests
+	 * that have been marked as cacheable
+	 * 
+	 * @param  BlueprintAction	The node-spawner you wish to test.
+	 * @return False if the action passes the filter tests, otherwise true (the action got filtered out).
+	 */
+	bool IsFilteredByCachedFilters(FBlueprintActionInfo& BlueprintAction);
 
 	/** Configuration flags for this filter. */
 	EFlags FilterFlags = BPFILTER_NoFlags;
 
 	/** Set of rejection tests for this specific filter. */
-	TArray<FRejectionTestDelegate> FilterTests;
+	TArray<TSharedRef<FActionFilterTest>> FilterTests;
 
 	/** Filters to be logically and'd in with the IsFilteredByThis() result. */
 	TArray<FBlueprintActionFilter> AndFilters;
@@ -353,17 +643,7 @@ private:
 	/** Alternative filters to be logically or'd in with the IsFilteredByThis() result. */
 	TArray<FBlueprintActionFilter> OrFilters;
 
-#if ENABLE_BLUEPRINT_ACTION_FILTER_PROFILING
-	struct FFilterTestProfileRecord
-	{
-		FString TestName;
-		float TotalTimeMs = 0.0f;
-		int32 NumIterations = 0;
-		int32 NumFilteredOut = 0;
-	};
-
-	mutable TMap<int32, FFilterTestProfileRecord> FilterTestProfiles;
-#endif	// ENABLE_BLUEPRINT_ACTION_FILTER_PROFILING
+	FActionFilterCache ActionFilterCache;
 };
 
 ENUM_CLASS_FLAGS(FBlueprintActionFilter::EFlags);
\ No newline at end of file
diff --git a/Engine/Source/Editor/Kismet/Private/BlueprintActionMenuUtils.cpp b/Engine/Source/Editor/Kismet/Private/BlueprintActionMenuUtils.cpp
--- a/Engine/Source/Editor/Kismet/Private/BlueprintActionMenuUtils.cpp
+++ b/Engine/Source/Editor/Kismet/Private/BlueprintActionMenuUtils.cpp
@@ -285,19 +285,19 @@ static FBlueprintActionFilter BlueprintActionMenuUtilsImpl::MakeCallOnMemberFilt
 	FBlueprintActionFilter CallOnMemberFilter;
 	CallOnMemberFilter.Context = MainMenuFilter.Context;
 	CallOnMemberFilter.PermittedNodeTypes.Add(UK2Node_CallFunction::StaticClass());
-	CallOnMemberFilter.AddRejectionTest(FBlueprintActionFilter::FRejectionTestDelegate::CreateStatic(IsUnBoundSpawner));
-
+	CallOnMemberFilter.AddRejectionTest(MAKE_ACTION_FILTER_REJECTION_TEST(IsUnBoundSpawner)->WithFlags(EActionFilterTestFlags::CacheResults));
+	
 	const UBlueprintEditorSettings* BlueprintSettings = GetDefault<UBlueprintEditorSettings>();
 	// instead of looking for "ExposeFunctionCategories" on component properties,
 	// we just expose functions for all components, but we still need to check
 	// for "ExposeFunctionCategories" on any non-component properties... 
 	if (BlueprintSettings->bExposeAllMemberComponentFunctions)
 	{
-		CallOnMemberFilter.AddRejectionTest(FBlueprintActionFilter::FRejectionTestDelegate::CreateStatic(IsUnexposedNonComponentAction));
+		CallOnMemberFilter.AddRejectionTest(MAKE_ACTION_FILTER_REJECTION_TEST(IsUnexposedNonComponentAction)->WithFlags(EActionFilterTestFlags::CacheResults));
 	}
 	else
 	{
-		CallOnMemberFilter.AddRejectionTest(FBlueprintActionFilter::FRejectionTestDelegate::CreateStatic(IsUnexposedMemberAction));
+		CallOnMemberFilter.AddRejectionTest(MAKE_ACTION_FILTER_REJECTION_TEST(IsUnexposedMemberAction)->WithFlags(EActionFilterTestFlags::CacheResults));
 	}
 
 	bool bForceAddComponents = ((ContextTargetMask & EContextTargetFlags::TARGET_SubComponents) != 0);
@@ -393,7 +393,9 @@ static void BlueprintActionMenuUtilsImpl::AddFavoritesSection(FBlueprintActionFi
 	if (BlueprintSettings->bShowContextualFavorites)
 	{
 		FBlueprintActionFilter FavoritesFilter = MainMenuFilter;
-		FavoritesFilter.AddRejectionTest(FBlueprintActionFilter::FRejectionTestDelegate::CreateStatic(IsNonFavoritedAction));
+
+		// make this action filter uncached because favorites can change on a moments notice
+		FavoritesFilter.AddRejectionTest(MAKE_ACTION_FILTER_REJECTION_TEST(IsNonFavoritedAction));
 		
 		uint32 SectionFlags = 0x00;
 		FText  SectionHeading = LOCTEXT("ContextMenuFavoritesTitle", "Favorites");
@@ -478,18 +480,18 @@ void FBlueprintActionMenuUtils::MakeContextMenu(FBlueprintActionContext const& C
 	FBlueprintActionFilter ComponentsFilter;
 	ComponentsFilter.Context = Context;
 	// only want bound actions for this menu section
-	ComponentsFilter.AddRejectionTest(FBlueprintActionFilter::FRejectionTestDelegate::CreateStatic(IsUnBoundSpawner));
+	ComponentsFilter.AddRejectionTest(MAKE_ACTION_FILTER_REJECTION_TEST(IsUnBoundSpawner)->WithFlags(EActionFilterTestFlags::CacheResults));
 	// @TODO: don't know exactly why we can only bind non-pure/const functions;
 	//        this is mirrored after FK2ActionMenuBuilder::GetFunctionCallsOnSelectedActors()
 	//        and FK2ActionMenuBuilder::GetFunctionCallsOnSelectedComponents(),
 	//        where we make the same stipulation
-	ComponentsFilter.AddRejectionTest(FBlueprintActionFilter::FRejectionTestDelegate::CreateStatic(IsPureNonConstAction));
+	ComponentsFilter.AddRejectionTest(MAKE_ACTION_FILTER_REJECTION_TEST(IsPureNonConstAction)->WithFlags(EActionFilterTestFlags::CacheResults));
 	
 
 	FBlueprintActionFilter LevelActorsFilter;
 	LevelActorsFilter.Context = Context;
 	// only want bound actions for this menu section
-	LevelActorsFilter.AddRejectionTest(FBlueprintActionFilter::FRejectionTestDelegate::CreateStatic(IsUnBoundSpawner));
+	LevelActorsFilter.AddRejectionTest(MAKE_ACTION_FILTER_REJECTION_TEST(IsUnBoundSpawner)->WithFlags(EActionFilterTestFlags::CacheResults));
 
 	// Build asset reference filter
 	FAssetReferenceFilterContext AssetReferenceFilterContext;
@@ -641,7 +643,7 @@ void FBlueprintActionMenuUtils::MakeContextMenu(FBlueprintActionContext const& C
 	FBlueprintActionFilter AddComponentFilter;
 	AddComponentFilter.Context = MainMenuFilter.Context;
 	AddComponentFilter.PermittedNodeTypes.Add(UK2Node_AddComponent::StaticClass());
-	AddComponentFilter.AddRejectionTest(FBlueprintActionFilter::FRejectionTestDelegate::CreateStatic(IsUnBoundSpawner));
+	LevelActorsFilter.AddRejectionTest(MAKE_ACTION_FILTER_REJECTION_TEST(IsUnBoundSpawner)->WithFlags(EActionFilterTestFlags::CacheResults));
 
 
 	FContentBrowserModule& ContentBrowserModule = FModuleManager::LoadModuleChecked<FContentBrowserModule>("ContentBrowser");
@@ -744,7 +746,7 @@ void FBlueprintActionMenuUtils::MakeFavoritesMenu(FBlueprintActionContext const&
 
 	FBlueprintActionFilter MenuFilter;
 	MenuFilter.Context = Context;
-	MenuFilter.AddRejectionTest(FBlueprintActionFilter::FRejectionTestDelegate::CreateStatic(BlueprintActionMenuUtilsImpl::IsNonFavoritedAction));
+	MenuFilter.AddRejectionTest(MAKE_ACTION_FILTER_REJECTION_TEST(BlueprintActionMenuUtilsImpl::IsNonFavoritedAction));
 
 	uint32 SectionFlags = 0x00;
 	const UBlueprintEditorSettings* BlueprintSettings = GetDefault<UBlueprintEditorSettings>();
diff --git a/Engine/Source/Editor/MovieSceneTools/Private/MovieSceneEventCustomization.cpp b/Engine/Source/Editor/MovieSceneTools/Private/MovieSceneEventCustomization.cpp
--- a/Engine/Source/Editor/MovieSceneTools/Private/MovieSceneEventCustomization.cpp
+++ b/Engine/Source/Editor/MovieSceneTools/Private/MovieSceneEventCustomization.cpp
@@ -91,7 +91,7 @@ namespace UE_MovieSceneEventCustomization
 			// but also prevents displaying functions on BP components. Comment out for now.
 			//CallOnMemberFilter.AddRejectionTest(FBlueprintActionFilter::FRejectionTestDelegate::CreateStatic(RejectAnyUnboundActions));
 			
-			CallOnMemberFilter.AddRejectionTest(FBlueprintActionFilter::FRejectionTestDelegate::CreateStatic(RejectAnyNonFunctions));
+			CallOnMemberFilter.AddRejectionTest(MAKE_ACTION_FILTER_REJECTION_TEST(RejectAnyNonFunctions)->WithFlags(EActionFilterTestFlags::CacheResults));
 
 			ContextMenuBuilder.AddMenuSection(CallOnMemberFilter, FText::FromName(BoundObjectPinClass->GetFName()), 0);
 		}
@@ -110,7 +110,7 @@ namespace UE_MovieSceneEventCustomization
 			{
 				FBlueprintActionFilter::AddUnique(MenuFilter.TargetClasses, Blueprint->SkeletonGeneratedClass);
 			}
-			MenuFilter.AddRejectionTest(FBlueprintActionFilter::FRejectionTestDelegate::CreateStatic(RejectAnyNonFunctions));
+			MenuFilter.AddRejectionTest(MAKE_ACTION_FILTER_REJECTION_TEST(RejectAnyNonFunctions)->WithFlags(EActionFilterTestFlags::CacheResults));
 
 
 			ContextMenuBuilder.AddMenuSection(MenuFilter, LOCTEXT("SequenceDirectorMenu", "This Sequence"), 0);
@@ -124,7 +124,7 @@ namespace UE_MovieSceneEventCustomization
 	void CollectAllRebindActions(FGraphActionListBuilderBase& OutAllActions, UBlueprint* Blueprint)
 	{
 		// Build up the context object
-		auto RejectAnyForeignFunctions = [Blueprint](const FBlueprintActionFilter& Filter, FBlueprintActionInfo& BlueprintAction)
+		auto RejectAnyForeignFunctions = [](const FBlueprintActionFilter& Filter, FBlueprintActionInfo& BlueprintAction, UBlueprint* Blueprint)
 		{
 			const UBlueprintFunctionNodeSpawner* FunctionNodeSpawner = Cast<UBlueprintFunctionNodeSpawner>(BlueprintAction.NodeSpawner);
 			const UFunction* FunctionToCall = FunctionNodeSpawner ? FunctionNodeSpawner->GetFunction() : nullptr;
@@ -143,7 +143,7 @@ namespace UE_MovieSceneEventCustomization
 
 		MenuFilter.Context.Blueprints.Add(Blueprint);
 
-		MenuFilter.AddRejectionTest(FBlueprintActionFilter::FRejectionTestDelegate::CreateLambda(RejectAnyForeignFunctions));
+		MenuFilter.AddRejectionTest(MAKE_ACTION_FILTER_REJECTION_TEST(RejectAnyForeignFunctions, Blueprint));
 
 		if (Blueprint->SkeletonGeneratedClass)
 		{
