fixed overshadowed variable compile error
#prevlight trivial

diff --git a/Engine/Source/Editor/BlueprintGraph/Public/BlueprintActionFilter.h b/Engine/Source/Editor/BlueprintGraph/Public/BlueprintActionFilter.h
--- a/Engine/Source/Editor/BlueprintGraph/Public/BlueprintActionFilter.h
+++ b/Engine/Source/Editor/BlueprintGraph/Public/BlueprintActionFilter.h
@@ -224,7 +224,7 @@ public:
 	 *	TActionFilterCacheLeaf
 	 *  Note: the key is not hashed. This is O(n) */
 	template<typename TKeyType>
-	TSharedKeyNode<TKeyType> FindOrCacheChildNode(const TKeyType& Key, const FString& DebugName = FString(), TDelegate<bool(const TKeyType&, const TKeyType&)> Comparison = nullptr, bool *OutFoundWithoutCaching = nullptr)
+	TSharedKeyNode<TKeyType> FindOrCacheChildNode(const TKeyType& Key, const FString& InDebugName = FString(), TDelegate<bool(const TKeyType&, const TKeyType&)> Comparison = nullptr, bool *OutFoundWithoutCaching = nullptr)
 	{
 		TSharedKeyNode<TKeyType> Result = Find(Key);
 		if (OutFoundWithoutCaching)
@@ -238,7 +238,7 @@ public:
 		}
 		
 #if UE_BUILD_DEBUG
-		this->DebugName = DebugName;
+		DebugName = InDebugName;
 #endif
 		return Cache(new TActionFilterCacheKeyNode<TKeyType>(Key, Comparison));
 	}
@@ -248,7 +248,7 @@ public:
 	 *	TActionFilterCacheKeyNode
 	 *  Note: the key is not hashed. This is O(n) */
 	template<typename TKeyType>
-	TSharedLeaf<TKeyType> FindOrCacheChildLeaf(const TKeyType& Key, const FString& DebugName = FString(), TDelegate<bool(const TKeyType&, const TKeyType&)> Comparison = nullptr, bool *OutFoundWithoutCaching = nullptr)
+	TSharedLeaf<TKeyType> FindOrCacheChildLeaf(const TKeyType& Key, const FString& InDebugName = FString(), TDelegate<bool(const TKeyType&, const TKeyType&)> Comparison = nullptr, bool *OutFoundWithoutCaching = nullptr)
 	{
 		TSharedKeyNode<TKeyType> Result = Find(Key);
 		if (OutFoundWithoutCaching)
@@ -262,6 +262,9 @@ public:
 			++Leaf->AccessCount;
 			return Leaf;
 		}
+#if UE_BUILD_DEBUG
+		DebugName = InDebugName;
+#endif
 		return Cache(new TActionFilterCacheLeaf<TKeyType>(Key, Comparison));
 	}
 
