Struct properties marked as EditAnywhere and BlueprintReadOnly were mistakenly editable in widget blueprints because we were checking if the Blueprint object was editor only, rather than its parent class...

diff --git a/Engine/Source/Editor/BlueprintGraph/Private/K2Node_MakeStruct.cpp b/Engine/Source/Editor/BlueprintGraph/Private/K2Node_MakeStruct.cpp
--- a/Engine/Source/Editor/BlueprintGraph/Private/K2Node_MakeStruct.cpp
+++ b/Engine/Source/Editor/BlueprintGraph/Private/K2Node_MakeStruct.cpp
@@ -92,7 +92,7 @@ static bool CanBeExposed(const FProperty* Property, UBlueprint* BP)
 		const UEdGraphSchema_K2* Schema = GetDefault<UEdGraphSchema_K2>();
 		check(Schema);
 
-		const bool bIsEditorBP = IsEditorOnlyObject(BP);
+		const bool bIsEditorBP = BP->ParentClass && IsEditorOnlyObject(BP->ParentClass);
 		const bool bIsEditAnywhereProperty = Property->HasAllPropertyFlags(CPF_Edit) &&
 			!Property->HasAnyPropertyFlags(CPF_EditConst);
 
