Back out 21452891 - lyra content is relying on this, possibly also some CitySample content

diff --git a/Engine/Source/Editor/BlueprintGraph/Private/K2Node_MakeStruct.cpp b/Engine/Source/Editor/BlueprintGraph/Private/K2Node_MakeStruct.cpp
--- a/Engine/Source/Editor/BlueprintGraph/Private/K2Node_MakeStruct.cpp
+++ b/Engine/Source/Editor/BlueprintGraph/Private/K2Node_MakeStruct.cpp
@@ -92,7 +92,7 @@ static bool CanBeExposed(const FProperty* Property, UBlueprint* BP)
 		const UEdGraphSchema_K2* Schema = GetDefault<UEdGraphSchema_K2>();
 		check(Schema);
 
-		const bool bIsEditorBP = BP->ParentClass && IsEditorOnlyObject(BP->ParentClass);
+		const bool bIsEditorBP = IsEditorOnlyObject(BP);
 		const bool bIsEditAnywhereProperty = Property->HasAllPropertyFlags(CPF_Edit) &&
 			!Property->HasAnyPropertyFlags(CPF_EditConst);
 
