Improve UX of Suppressed Deprecation Messages

diff --git a/Engine/Source/Editor/BlueprintGraph/Private/K2Node_CallFunction.cpp b/Engine/Source/Editor/BlueprintGraph/Private/K2Node_CallFunction.cpp
--- a/Engine/Source/Editor/BlueprintGraph/Private/K2Node_CallFunction.cpp
+++ b/Engine/Source/Editor/BlueprintGraph/Private/K2Node_CallFunction.cpp
@@ -3393,7 +3393,7 @@ void UK2Node_CallFunction::GetNodeContextMenuActions(class UToolMenu* Menu, clas
 	if (HasDeprecatedReference())
 	{
 		FText MenuEntryTitle = LOCTEXT("SuppressDeprecationWarningTitle", "Suppress Deprecation Warning");
-		FText MenuEntryTooltip = LOCTEXT("SuppressDeprecationWarningTooltip", "Adds this variable to the suppressed deprecation warnings list for the project.");
+		FText MenuEntryTooltip = LOCTEXT("SuppressDeprecationWarningTooltip", "Adds this function to the suppressed deprecation warnings list in the Bluperint Editor Project Settings for this project.");
 
 		FToolMenuSection& Section = Menu->AddSection("K2NodeCallFunction", LOCTEXT("FunctionHeader", "Function"));
 		Section.AddMenuEntry(
diff --git a/Engine/Source/Editor/BlueprintGraph/Private/K2Node_Variable.cpp b/Engine/Source/Editor/BlueprintGraph/Private/K2Node_Variable.cpp
--- a/Engine/Source/Editor/BlueprintGraph/Private/K2Node_Variable.cpp
+++ b/Engine/Source/Editor/BlueprintGraph/Private/K2Node_Variable.cpp
@@ -1152,7 +1152,7 @@ void UK2Node_Variable::GetNodeContextMenuActions(class UToolMenu* Menu, class UG
 	if (HasDeprecatedReference())
 	{
 		FText MenuEntryTitle = LOCTEXT("SuppressDeprecationWarningTitle", "Suppress Deprecation Warning");
-		FText MenuEntryTooltip = LOCTEXT("SuppressDeprecationWarningTooltip", "Adds this variable to the suppressed deprecation warnings list for the project.");
+		FText MenuEntryTooltip = LOCTEXT("SuppressDeprecationWarningTooltip", "Adds this variable to the suppressed deprecation warnings list in the Bluperint Editor Project Settings for this project.");
 
 		FToolMenuSection& Section = Menu->AddSection("K2NodeVariable", LOCTEXT("VariableHeader", "Variable"));
 		Section.AddMenuEntry(
diff --git a/Engine/Source/Editor/UnrealEd/Public/Settings/BlueprintEditorProjectSettings.h b/Engine/Source/Editor/UnrealEd/Public/Settings/BlueprintEditorProjectSettings.h
--- a/Engine/Source/Editor/UnrealEd/Public/Settings/BlueprintEditorProjectSettings.h
+++ b/Engine/Source/Editor/UnrealEd/Public/Settings/BlueprintEditorProjectSettings.h
@@ -66,6 +66,7 @@ public:
 	/**
 	 * List of deprecated UProperties/UFunctions to supress warning messages for - useful for source changes
 	 * that would otherwise cause content warnings
+	 * The easiest way to populate this list is using the context menu on nodes with deprecated references
 	 */
 	UPROPERTY(EditAnywhere, config, Category = Blueprints, DisplayName = "Deprecated Symbols to Supress")
 	TArray<FString> SuppressedDeprecationMessages;
