Fix conflicting LOCTEXT keys

diff --git a/Engine/Source/Editor/BlueprintGraph/Private/K2Node_CallFunction.cpp b/Engine/Source/Editor/BlueprintGraph/Private/K2Node_CallFunction.cpp
--- a/Engine/Source/Editor/BlueprintGraph/Private/K2Node_CallFunction.cpp
+++ b/Engine/Source/Editor/BlueprintGraph/Private/K2Node_CallFunction.cpp
@@ -3392,8 +3392,8 @@ void UK2Node_CallFunction::GetNodeContextMenuActions(class UToolMenu* Menu, clas
 
 	if (HasDeprecatedReference())
 	{
-		FText MenuEntryTitle = LOCTEXT("SuppressDeprecationWarningTitle", "Suppress Deprecation Warning");
-		FText MenuEntryTooltip = LOCTEXT("SuppressDeprecationWarningTooltip", "Adds this function to the suppressed deprecation warnings list in the Bluperint Editor Project Settings for this project.");
+		FText MenuEntryTitle = LOCTEXT("SuppressFunctionDeprecationWarningTitle", "Suppress Deprecation Warning");
+		FText MenuEntryTooltip = LOCTEXT("SuppressFunctionDeprecationWarningTooltip", "Adds this function to the suppressed deprecation warnings list in the Bluperint Editor Project Settings for this project.");
 
 		FToolMenuSection& Section = Menu->AddSection("K2NodeCallFunction", LOCTEXT("FunctionHeader", "Function"));
 		Section.AddMenuEntry(
diff --git a/Engine/Source/Editor/BlueprintGraph/Private/K2Node_Variable.cpp b/Engine/Source/Editor/BlueprintGraph/Private/K2Node_Variable.cpp
--- a/Engine/Source/Editor/BlueprintGraph/Private/K2Node_Variable.cpp
+++ b/Engine/Source/Editor/BlueprintGraph/Private/K2Node_Variable.cpp
@@ -1151,8 +1151,8 @@ void UK2Node_Variable::GetNodeContextMenuActions(class UToolMenu* Menu, class UG
 
 	if (HasDeprecatedReference())
 	{
-		FText MenuEntryTitle = LOCTEXT("SuppressDeprecationWarningTitle", "Suppress Deprecation Warning");
-		FText MenuEntryTooltip = LOCTEXT("SuppressDeprecationWarningTooltip", "Adds this variable to the suppressed deprecation warnings list in the Bluperint Editor Project Settings for this project.");
+		FText MenuEntryTitle = LOCTEXT("SuppressVariableDeprecationWarningTitle", "Suppress Deprecation Warning");
+		FText MenuEntryTooltip = LOCTEXT("SuppressVariableDeprecationWarningTooltip", "Adds this variable to the suppressed deprecation warnings list in the Bluperint Editor Project Settings for this project.");
 
 		FToolMenuSection& Section = Menu->AddSection("K2NodeVariable", LOCTEXT("VariableHeader", "Variable"));
 		Section.AddMenuEntry(
