Blueprint namespace editor features are now enabled by default.

Change summary:
- Moved Blueprint namespace editor settings out of the experimental category.
- Modified the editor settings details customization to allow the options to appear as adjacent rows, since the feature flag is also an edit condition.

diff --git a/Engine/Source/Editor/BlueprintGraph/Private/BlueprintEditorSettings.cpp b/Engine/Source/Editor/BlueprintGraph/Private/BlueprintEditorSettings.cpp
--- a/Engine/Source/Editor/BlueprintGraph/Private/BlueprintEditorSettings.cpp
+++ b/Engine/Source/Editor/BlueprintGraph/Private/BlueprintEditorSettings.cpp
@@ -38,11 +38,8 @@ UBlueprintEditorSettings::UBlueprintEditorSettings(const FObjectInitializer& Obj
 	, TypePromotionPinDenyList { UEdGraphSchema_K2::PC_String, UEdGraphSchema_K2::PC_Text, UEdGraphSchema_K2::PC_SoftClass }
 	, BreakpointReloadMethod(EBlueprintBreakpointReloadMethod::RestoreAll)
 	, bEnablePinValueInspectionTooltips(true)
+	, bEnableNamespaceEditorFeatures(true)
 	// Experimental
-	, bEnableNamespaceEditorFeatures(false)
-	, bEnableNamespaceFilteringFeatures(false)
-	, bEnableNamespaceImportingFeatures(false)
-	, bInheritImportedNamespacesFromParentBP(false)
 	, bFavorPureCastNodes(false)
 	// Compiler Settings
 	, SaveOnCompile(SoC_Never)
@@ -62,6 +59,9 @@ UBlueprintEditorSettings::UBlueprintEditorSettings(const FObjectInitializer& Obj
 	, bShowAccessSpecifier(false)
 	, bIncludeCommentNodesInBookmarksTab(true)
 	, bShowBookmarksForCurrentDocumentOnlyInTab(false)
+	, bEnableNamespaceFilteringFeatures(false)
+	, bEnableNamespaceImportingFeatures(false)
+	, bInheritImportedNamespacesFromParentBP(false)
 {
 	// settings that were moved out of experimental...
 	UEditorExperimentalSettings const* ExperimentalSettings = GetDefault<UEditorExperimentalSettings>();
diff --git a/Engine/Source/Editor/BlueprintGraph/Public/BlueprintEditorSettings.h b/Engine/Source/Editor/BlueprintGraph/Public/BlueprintEditorSettings.h
--- a/Engine/Source/Editor/BlueprintGraph/Public/BlueprintEditorSettings.h
+++ b/Engine/Source/Editor/BlueprintGraph/Public/BlueprintEditorSettings.h
@@ -161,31 +161,16 @@ public:
 	UPROPERTY(config, EditAnywhere, Category = Workflow)
 	bool bEnablePinValueInspectionTooltips;
 
-// Experimental
-public:
 	/** Whether to enable namespace importing and filtering features in the Blueprint editor */
-	UPROPERTY(config, EditAnywhere, Category = Experimental)
+	UPROPERTY(config, EditAnywhere, Category = Workflow)
 	bool bEnableNamespaceEditorFeatures;
 
-	/** Whether to enable namespace filtering features in the Blueprint editor */
-	// @todo_namespaces - Remove this if/when dependent code is changed to utilize the single setting above.
-	UPROPERTY(Transient)
-	bool bEnableNamespaceFilteringFeatures;
-
-	/** Whether to enable namespace importing features in the Blueprint editor */
-	// @todo_namespaces - Remove this if/when dependent code is changed to utilize the single setting above.
-	UPROPERTY(Transient)
-	bool bEnableNamespaceImportingFeatures;
-
-	/** Whether to inherit the set of imported namespaces from the parent class hierarchy */
-	// @todo_namespaces - Remove this if/when this becomes a permanent setting. For now this is experimental.
-	UPROPERTY(Transient)
-	bool bInheritImportedNamespacesFromParentBP;
-
 	// A list of namespace identifiers that the Blueprint editor should always import by default. Requires Blueprint namespace features to be enabled and only applies to the current local user. Editing this list will also cause any visible Blueprint editor windows to be closed.
-	UPROPERTY(EditAnywhere, config, Category = Experimental, meta = (EditCondition = "bEnableNamespaceEditorFeatures", DisplayName = "Global Namespace Imports (Local User Only)"))
+	UPROPERTY(EditAnywhere, config, Category = Workflow, meta = (EditCondition = "bEnableNamespaceEditorFeatures", DisplayName = "Global Namespace Imports (Local User Only)"))
 	TArray<FString> NamespacesToAlwaysInclude;
 
+// Experimental
+public:
 	/** If enabled, then placed cast nodes will default to their "pure" form (meaning: without execution pins). */
 	UPROPERTY(EditAnywhere, config, AdvancedDisplay, Category = Experimental, meta = (DisplayName = "Default to Using Pure Cast Nodes"))
 	bool bFavorPureCastNodes;
@@ -273,6 +258,21 @@ public:
 	UPROPERTY(config)
 	TMap<int32, FEditedDocumentInfo> GraphEditorQuickJumps;
 
+	/** Whether to enable namespace filtering features in the Blueprint editor */
+	// @todo_namespaces - Remove this if/when dependent code is changed to utilize the single setting above.
+	UPROPERTY(Transient)
+	bool bEnableNamespaceFilteringFeatures;
+
+	/** Whether to enable namespace importing features in the Blueprint editor */
+	// @todo_namespaces - Remove this if/when dependent code is changed to utilize the single setting above.
+	UPROPERTY(Transient)
+	bool bEnableNamespaceImportingFeatures;
+
+	/** Whether to inherit the set of imported namespaces from the parent class hierarchy */
+	// @todo_namespaces - Remove this if/when this becomes a permanent setting. For now this is experimental.
+	UPROPERTY(Transient)
+	bool bInheritImportedNamespacesFromParentBP;
+
 	/**
 	 * Any blueprint deriving from one of these base classes will be allowed to recompile during Play-in-Editor
 	 * (This setting exists both as an editor preference and project setting, and will be allowed if listed in either place) 
diff --git a/Engine/Source/Editor/UnrealEd/Private/Settings/BlueprintEditorSettingsCustomization.cpp b/Engine/Source/Editor/UnrealEd/Private/Settings/BlueprintEditorSettingsCustomization.cpp
--- a/Engine/Source/Editor/UnrealEd/Private/Settings/BlueprintEditorSettingsCustomization.cpp
+++ b/Engine/Source/Editor/UnrealEd/Private/Settings/BlueprintEditorSettingsCustomization.cpp
@@ -113,12 +113,19 @@ TSharedRef<IDetailCustomization> FBlueprintEditorSettingsCustomization::MakeInst
 
 void FBlueprintEditorSettingsCustomization::CustomizeDetails(IDetailLayoutBuilder& LayoutBuilder)
 {
+	// Get the property name and handle for the namespace feature toggle setting. Mark it hidden because we will add it via customization below.
+	static FName PropertyName_bEnableNamespaceEditorFeatures = GET_MEMBER_NAME_CHECKED(UBlueprintEditorSettings, bEnableNamespaceEditorFeatures);
+	TSharedRef<IPropertyHandle> PropertyHandle_bEnableNamespaceEditorFeatures = LayoutBuilder.GetProperty(PropertyName_bEnableNamespaceEditorFeatures);
+	PropertyHandle_bEnableNamespaceEditorFeatures->MarkHiddenByCustomization();
+
+	// Get the property name and handle for the global namespace list setting. Mark it hidden because we will add it via customization below.
 	static FName PropertyName_NamespacesToAlwaysInclude = GET_MEMBER_NAME_CHECKED(UBlueprintEditorSettings, NamespacesToAlwaysInclude);
 	TSharedRef<IPropertyHandle> PropertyHandle_NamespacesToAlwaysInclude = LayoutBuilder.GetProperty(PropertyName_NamespacesToAlwaysInclude);
-
 	PropertyHandle_NamespacesToAlwaysInclude->MarkHiddenByCustomization();
 
-	IDetailCategoryBuilder& CategoryBuilder = LayoutBuilder.EditCategory(PropertyHandle_NamespacesToAlwaysInclude->GetDefaultCategoryName());
+	// List the edit condition first, followed by the customization for the global namespace list (i.e. ensure they are adjacent within the same category).
+	IDetailCategoryBuilder& CategoryBuilder = LayoutBuilder.EditCategory(PropertyHandle_bEnableNamespaceEditorFeatures->GetDefaultCategoryName());
+	CategoryBuilder.AddProperty(PropertyHandle_bEnableNamespaceEditorFeatures);
 	CategoryBuilder.AddCustomBuilder(MakeShared<UE::Editor::BlueprintEditorSettingsCustomization::Private::FBlueprintGlobalEditorImportsLayout>(PropertyHandle_NamespacesToAlwaysInclude));
 }
 
diff --git a/Engine/Source/Editor/UnrealEd/Public/Settings/BlueprintEditorProjectSettings.h b/Engine/Source/Editor/UnrealEd/Public/Settings/BlueprintEditorProjectSettings.h
--- a/Engine/Source/Editor/UnrealEd/Public/Settings/BlueprintEditorProjectSettings.h
+++ b/Engine/Source/Editor/UnrealEd/Public/Settings/BlueprintEditorProjectSettings.h
@@ -44,7 +44,7 @@ public:
 	EChildActorComponentTreeViewVisualizationMode DefaultChildActorTreeViewMode;
 
 	// A list of namespace identifiers that all Blueprint assets in the project should import by default. Requires Blueprint namespace features to be enabled in editor preferences. Editing this list will also cause any visible Blueprint editor windows to be closed.
-	UPROPERTY(EditAnywhere, config, Category = Experimental, DisplayName = "Global Namespace Imports (Shared)")
+	UPROPERTY(EditAnywhere, config, Category = Blueprints, DisplayName = "Global Namespace Imports (Shared)")
 	TArray<FString> NamespacesToAlwaysInclude;
 
 	/** 
