Adding support for dragging variables onto pins to automatically convert/promote

diff --git a/Engine/Source/Editor/BlueprintGraph/Classes/EdGraphSchema_K2.h b/Engine/Source/Editor/BlueprintGraph/Classes/EdGraphSchema_K2.h
--- a/Engine/Source/Editor/BlueprintGraph/Classes/EdGraphSchema_K2.h
+++ b/Engine/Source/Editor/BlueprintGraph/Classes/EdGraphSchema_K2.h
@@ -1090,7 +1090,10 @@ public:
 
 	/** Find an appropriate node that can convert from one pin type to another (not a cast; e.g. "MakeLiteralArray" node) */
 	virtual bool FindSpecializedConversionNode(const UEdGraphPin* OutputPin, const UEdGraphPin* InputPin, bool bCreateNode, /*out*/ class UK2Node*& TargetNode) const;
-
+	
+	/** Find an appropriate node that can convert from one pin type to another (not a cast; e.g. "MakeLiteralArray" node) */
+	virtual bool FindSpecializedConversionNode(const FEdGraphPinType& OutputPinType, const UEdGraphPin* InputPinType, bool bCreateNode, /*out*/ class UK2Node*& TargetNode) const;
+	
 	/** Create menu for variable get/set nodes which refer to a variable which does not exist. */
 	void GetNonExistentVariableMenu(FToolMenuSection& Section, const UEdGraphNode* InGraphNode, UBlueprint* OwnerBlueprint) const;
 
diff --git a/Engine/Source/Editor/BlueprintGraph/Private/EdGraphSchema_K2.cpp b/Engine/Source/Editor/BlueprintGraph/Private/EdGraphSchema_K2.cpp
--- a/Engine/Source/Editor/BlueprintGraph/Private/EdGraphSchema_K2.cpp
+++ b/Engine/Source/Editor/BlueprintGraph/Private/EdGraphSchema_K2.cpp
@@ -2591,12 +2591,21 @@ bool UEdGraphSchema_K2::SearchForAutocastFunction(const FEdGraphPinType& OutputP
 }
 
 bool UEdGraphSchema_K2::FindSpecializedConversionNode(const UEdGraphPin* OutputPin, const UEdGraphPin* InputPin, bool bCreateNode, /*out*/ UK2Node*& TargetNode) const
+{
+	return FindSpecializedConversionNode(OutputPin->PinType, InputPin, bCreateNode, TargetNode);
+}
+
+bool UEdGraphSchema_K2::FindSpecializedConversionNode(const FEdGraphPinType& OutputPinType, const UEdGraphPin* InputPin, bool bCreateNode, UK2Node*& TargetNode) const
 {
 	bool bCanConvert = false;
 	TargetNode = nullptr;
+	if (!InputPin)
+	{
+		return bCanConvert;
+	}
 
 	// Conversion for scalar -> array
-	if( (!OutputPin->PinType.IsContainer() && InputPin->PinType.IsArray()) && ArePinTypesCompatible(OutputPin->PinType, InputPin->PinType, nullptr, true))
+	if( (!OutputPinType.IsContainer() && InputPin->PinType.IsArray()) && ArePinTypesCompatible(OutputPinType, InputPin->PinType, nullptr, true))
 	{
 		bCanConvert = true;
 		if(bCreateNode)
@@ -2606,10 +2615,10 @@ bool UEdGraphSchema_K2::FindSpecializedConversionNode(const UEdGraphPin* OutputP
 	}
 	// If connecting an object to a 'call function' self pin, and not currently compatible, see if there is a property we can call a function on
 	else if (InputPin->GetOwningNode()->IsA(UK2Node_CallFunction::StaticClass()) && IsSelfPin(*InputPin) && 
-		((OutputPin->PinType.PinCategory == PC_Object) || (OutputPin->PinType.PinCategory == PC_Interface)))
+		((OutputPinType.PinCategory == PC_Object) || (OutputPinType.PinCategory == PC_Interface)))
 	{
 		UK2Node_CallFunction* CallFunctionNode = (UK2Node_CallFunction*)(InputPin->GetOwningNode());
-		UClass* OutputPinClass = Cast<UClass>(OutputPin->PinType.PinSubCategoryObject.Get());
+		UClass* OutputPinClass = Cast<UClass>(OutputPinType.PinSubCategoryObject.Get());
 
 		UClass* FunctionClass = CallFunctionNode->FunctionReference.GetMemberParentClass(CallFunctionNode->GetBlueprintClassFromNode());
 		if(FunctionClass != NULL && OutputPinClass != NULL)
@@ -2640,7 +2649,7 @@ bool UEdGraphSchema_K2::FindSpecializedConversionNode(const UEdGraphPin* OutputP
 	{
 		// CHECK ENUM TO NAME CAST
 		const bool bInoputMatch = InputPin && !InputPin->PinType.IsContainer() && ((PC_Name == InputPin->PinType.PinCategory) || (PC_String == InputPin->PinType.PinCategory));
-		const bool bOutputMatch = OutputPin && !OutputPin->PinType.IsContainer() && (PC_Byte == OutputPin->PinType.PinCategory) && (NULL != Cast<UEnum>(OutputPin->PinType.PinSubCategoryObject.Get()));
+		const bool bOutputMatch = !OutputPinType.IsContainer() && (PC_Byte == OutputPinType.PinCategory) && (NULL != Cast<UEnum>(OutputPinType.PinSubCategoryObject.Get()));
 		if(bOutputMatch && bInoputMatch)
 		{
 			bCanConvert = true;
@@ -2659,10 +2668,10 @@ bool UEdGraphSchema_K2::FindSpecializedConversionNode(const UEdGraphPin* OutputP
 		}
 	}
 
-	if (!bCanConvert && InputPin && OutputPin)
+	if (!bCanConvert && InputPin)
 	{
 		FEdGraphPinType const& InputType  = InputPin->PinType;
-		FEdGraphPinType const& OutputType = OutputPin->PinType;
+		FEdGraphPinType const& OutputType = OutputPinType;
 
 		// CHECK BYTE TO ENUM CAST
 		UEnum* Enum = Cast<UEnum>(InputType.PinSubCategoryObject.Get());
@@ -2682,7 +2691,7 @@ bool UEdGraphSchema_K2::FindSpecializedConversionNode(const UEdGraphPin* OutputP
 		else
 		{
 			UClass* InputClass  = FBlueprintEditorUtils::GetTypeForPin(*InputPin);
-			UClass* OutputClass = FBlueprintEditorUtils::GetTypeForPin(*OutputPin);
+			UClass* OutputClass = Cast<UClass>(OutputPinType.PinSubCategoryObject.Get());
 
 			if ((OutputType.PinCategory == PC_Interface) && (InputType.PinCategory == PC_Object))
 			{
diff --git a/Engine/Source/Editor/Kismet/Private/BPVariableDragDropAction.cpp b/Engine/Source/Editor/Kismet/Private/BPVariableDragDropAction.cpp
--- a/Engine/Source/Editor/Kismet/Private/BPVariableDragDropAction.cpp
+++ b/Engine/Source/Editor/Kismet/Private/BPVariableDragDropAction.cpp
@@ -174,9 +174,23 @@ void FKismetVariableDragDropAction::HoverTargetChanged()
 				Schema->ConvertPropertyToPinType(VariableProperty, VariablePinType);
 				const bool bTypeMatch = Schema->ArePinTypesCompatible(VariablePinType, PinUnderCursor->PinType) || bIsExecPin;
 
+				FName DummyName;
+				UClass* DummyClass;
+				UK2Node* DummyNode;
+				const bool bCanAutoConvert = Schema->FindSpecializedConversionNode(VariablePinType, PinUnderCursor, false, /* out */ DummyNode);
+				bool bCanAutocast = false;
+				if (PinUnderCursor->Direction == EGPD_Output)
+				{
+					bCanAutocast = Schema->SearchForAutocastFunction(PinUnderCursor->PinType, VariablePinType, /*out*/ DummyName, DummyClass);
+				}
+				else
+				{
+					bCanAutocast = Schema->SearchForAutocastFunction(VariablePinType, PinUnderCursor->PinType, /*out*/ DummyName, DummyClass);
+				}
+				
 				Args.Add(TEXT("PinUnderCursor"), FText::FromName(PinUnderCursor->PinName));
 
-				if (bTypeMatch && bCanWriteIfNeeded)
+				if ((bTypeMatch  || bCanAutocast || bCanAutoConvert) && bCanWriteIfNeeded)
 				{
 					SetFeedbackMessageOK(bIsRead ?
 						FText::Format(LOCTEXT("MakeThisEqualThat_PinEqualVariableName", "Make {PinUnderCursor} = {VariableName}"), Args) :
@@ -284,8 +298,22 @@ FReply FKismetVariableDragDropAction::DroppedOnPin(FVector2D ScreenPosition, FVe
 					FEdGraphPinType VariablePinType;
 					Schema->ConvertPropertyToPinType(VariableProperty, VariablePinType);
 					const bool bTypeMatch = Schema->ArePinTypesCompatible(VariablePinType, TargetPin->PinType) || bIsExecPin;
-
-					if (bTypeMatch && bCanWriteIfNeeded)
+					
+					FName DummyName;
+					UClass* DummyClass;
+					UK2Node* DummyNode;
+					const bool bCanAutoConvert = Schema->FindSpecializedConversionNode(VariablePinType, TargetPin, false, /* out */ DummyNode);
+					bool bCanAutocast = false;
+					if (TargetPin->Direction == EGPD_Output)
+					{
+						bCanAutocast = Schema->SearchForAutocastFunction(TargetPin->PinType, VariablePinType, /*out*/ DummyName, DummyClass);
+					}
+					else
+					{
+						bCanAutocast = Schema->SearchForAutocastFunction(VariablePinType, TargetPin->PinType, /*out*/ DummyName, DummyClass);
+					}
+					
+					if ((bTypeMatch || bCanAutocast || bCanAutoConvert) && bCanWriteIfNeeded)
 					{
 						FEdGraphSchemaAction_K2NewNode Action;
 
