Added a matching GetThenPin to supplement GetExecPin.

Cleaned up some comments, too.

diff --git a/Engine/Source/Editor/BlueprintGraph/Classes/K2Node.h b/Engine/Source/Editor/BlueprintGraph/Classes/K2Node.h
--- a/Engine/Source/Editor/BlueprintGraph/Classes/K2Node.h
+++ b/Engine/Source/Editor/BlueprintGraph/Classes/K2Node.h
@@ -324,9 +324,12 @@ public:
 	/** Get the Blueprint object to which this node belongs */
 	BLUEPRINTGRAPH_API UBlueprint* GetBlueprint() const;
 
-	/** Get the input execution pin if this node is impure (will return NULL when IsNodePure() returns true) */
+	/** Get the input execution pin of this node (if one exists)*/
 	BLUEPRINTGRAPH_API UEdGraphPin* GetExecPin() const;
 
+	/** Get the output then pin of this node (if one exists)*/
+	BLUEPRINTGRAPH_API UEdGraphPin* GetThenPin() const;
+
 	/**
 	 * If this node references an actor in the level that should be selectable by "Find Actors In Level," this will return a reference to that actor
 	 *
diff --git a/Engine/Source/Editor/BlueprintGraph/Private/K2Node.cpp b/Engine/Source/Editor/BlueprintGraph/Private/K2Node.cpp
--- a/Engine/Source/Editor/BlueprintGraph/Private/K2Node.cpp
+++ b/Engine/Source/Editor/BlueprintGraph/Private/K2Node.cpp
@@ -1739,7 +1739,14 @@ void FOptionalPinManager::EvaluateOldShownPins(const TArray<FOptionalPinFromProp
 UEdGraphPin* UK2Node::GetExecPin() const
 {
 	UEdGraphPin* Pin = FindPin(UEdGraphSchema_K2::PN_Execute);
-	check(Pin == nullptr || Pin->Direction == EGPD_Input); // If pin exists, it must be input
+	check(Pin == nullptr || Pin->Direction == EGPD_Input);
+	return Pin;
+}
+
+UEdGraphPin* UK2Node::GetThenPin() const
+{
+	UEdGraphPin* Pin = FindPin(UEdGraphSchema_K2::PN_Then);
+	check(Pin == nullptr || Pin->Direction == EGPD_Output);
 	return Pin;
 }
 
