Misc cleanup:

* Constify CreateFunctionGraphTerminators
* Update the definition of RegisterCompilerForBP to use same the same typedef parameter as its declaration.

diff --git a/Engine/Source/Editor/BlueprintGraph/Classes/EdGraphSchema_K2.h b/Engine/Source/Editor/BlueprintGraph/Classes/EdGraphSchema_K2.h
--- a/Engine/Source/Editor/BlueprintGraph/Classes/EdGraphSchema_K2.h
+++ b/Engine/Source/Editor/BlueprintGraph/Classes/EdGraphSchema_K2.h
@@ -968,7 +968,7 @@ public:
 	 * @param	Graph			Graph to add the function terminators to
 	 * @param	FunctionSignature	The function signature to mimic when creating the inputs and outputs for the function.
 	 */
-	virtual void CreateFunctionGraphTerminators(UEdGraph& Graph, UFunction* FunctionSignature) const;
+	virtual void CreateFunctionGraphTerminators(UEdGraph& Graph, const UFunction* FunctionSignature) const;
 
 	/**
 	 * Converts the type of a property into a fully qualified string (e.g., object'ObjectName').
diff --git a/Engine/Source/Editor/BlueprintGraph/Private/EdGraphSchema_K2.cpp b/Engine/Source/Editor/BlueprintGraph/Private/EdGraphSchema_K2.cpp
--- a/Engine/Source/Editor/BlueprintGraph/Private/EdGraphSchema_K2.cpp
+++ b/Engine/Source/Editor/BlueprintGraph/Private/EdGraphSchema_K2.cpp
@@ -3297,7 +3297,7 @@ void UEdGraphSchema_K2::CreateFunctionGraphTerminators(UEdGraph& Graph, UClass*
 	}
 }
 
-void UEdGraphSchema_K2::CreateFunctionGraphTerminators(UEdGraph& Graph, UFunction* FunctionSignature) const
+void UEdGraphSchema_K2::CreateFunctionGraphTerminators(UEdGraph& Graph, const UFunction* FunctionSignature) const
 {
 	const FName GraphName = Graph.GetFName();
 
diff --git a/Engine/Source/Editor/KismetCompiler/Private/KismetCompiler.cpp b/Engine/Source/Editor/KismetCompiler/Private/KismetCompiler.cpp
--- a/Engine/Source/Editor/KismetCompiler/Private/KismetCompiler.cpp
+++ b/Engine/Source/Editor/KismetCompiler/Private/KismetCompiler.cpp
@@ -5393,7 +5393,7 @@ TSharedPtr<FKismetCompilerContext> FKismetCompilerContext::GetCompilerForBP(UBlu
 	}
 }
 
-void FKismetCompilerContext::RegisterCompilerForBP(UClass* BPClass, TFunction<TSharedPtr<FKismetCompilerContext>(UBlueprint*, FCompilerResultsLog&, const FKismetCompilerOptions&)> FactoryFunction)
+void FKismetCompilerContext::RegisterCompilerForBP(UClass* BPClass, CompilerContextFactoryFunction FactoryFunction)
 {
 	CustomCompilerMap.Add(BPClass, FactoryFunction);
 }
