Fix component bound event nodes (also used for UMG) to update the node title when a delegate is renamed
Change both component and actor bound events to compute the display name on demand like MCDelegate

diff --git a/Engine/Source/Editor/BlueprintGraph/Classes/K2Node_ActorBoundEvent.h b/Engine/Source/Editor/BlueprintGraph/Classes/K2Node_ActorBoundEvent.h
--- a/Engine/Source/Editor/BlueprintGraph/Classes/K2Node_ActorBoundEvent.h
+++ b/Engine/Source/Editor/BlueprintGraph/Classes/K2Node_ActorBoundEvent.h
@@ -77,6 +77,9 @@ class UK2Node_ActorBoundEvent : public UK2Node_Event
 	BLUEPRINTGRAPH_API FMulticastDelegateProperty* GetTargetDelegateProperty() const;
 	BLUEPRINTGRAPH_API FMulticastDelegateProperty* GetTargetDelegatePropertyFromSkel() const;
 
+	/** Gets the proper display name for the property */
+	BLUEPRINTGRAPH_API FText GetTargetDelegateDisplayName() const;
+
 private:
 	/** Constructing FText strings can be costly, so we cache the node's title */
 	FNodeTextCache CachedNodeTitle;
diff --git a/Engine/Source/Editor/BlueprintGraph/Classes/K2Node_ComponentBoundEvent.h b/Engine/Source/Editor/BlueprintGraph/Classes/K2Node_ComponentBoundEvent.h
--- a/Engine/Source/Editor/BlueprintGraph/Classes/K2Node_ComponentBoundEvent.h
+++ b/Engine/Source/Editor/BlueprintGraph/Classes/K2Node_ComponentBoundEvent.h
@@ -72,6 +72,9 @@ class UK2Node_ComponentBoundEvent : public UK2Node_Event
 	/** Return the delegate property that this event is bound to */
 	BLUEPRINTGRAPH_API FMulticastDelegateProperty* GetTargetDelegateProperty() const;
 
+	/** Gets the proper display name for the property */
+	BLUEPRINTGRAPH_API FText GetTargetDelegateDisplayName() const;
+
 	BLUEPRINTGRAPH_API void InitializeComponentBoundEventParams(FObjectProperty const* InComponentProperty, const FMulticastDelegateProperty* InDelegateProperty);
 
 private:
@@ -79,10 +82,6 @@ private:
 	/** Returns true if there is an FObjectProperty on this blueprint with a name that matches ComponentPropertyName */
 	bool IsDelegateValid() const;
 
-	/** Cached display name for the delegate property */
-	UPROPERTY()
-	FText DelegatePropertyDisplayName;
-
 	/** Constructing FText strings can be costly, so we cache the node's title */
 	FNodeTextCache CachedNodeTitle;
 };
diff --git a/Engine/Source/Editor/BlueprintGraph/Private/K2Node_ActorBoundEvent.cpp b/Engine/Source/Editor/BlueprintGraph/Private/K2Node_ActorBoundEvent.cpp
--- a/Engine/Source/Editor/BlueprintGraph/Private/K2Node_ActorBoundEvent.cpp
+++ b/Engine/Source/Editor/BlueprintGraph/Private/K2Node_ActorBoundEvent.cpp
@@ -148,13 +148,13 @@ FText UK2Node_ActorBoundEvent::GetNodeTitle(ENodeTitleType::Type TitleType) cons
 	if (EventOwner == nullptr)
 	{
 		FFormatNamedArguments Args;
-		Args.Add(TEXT("DelegatePropertyName"), FText::FromName(DelegatePropertyName));
+		Args.Add(TEXT("DelegatePropertyName"), GetTargetDelegateDisplayName());
 		return FText::Format(LOCTEXT("ActorBoundEventTitleNoOwner", "{DelegatePropertyName} (None)"), Args);
 	}
 	else if (CachedNodeTitle.IsOutOfDate(this))
 	{
 		FFormatNamedArguments Args;
-		Args.Add(TEXT("DelegatePropertyName"), FText::FromName(DelegatePropertyName));
+		Args.Add(TEXT("DelegatePropertyName"), GetTargetDelegateDisplayName());
 		Args.Add(TEXT("TargetName"), FText::FromString(EventOwner->GetActorLabel()));
 
 		// FText::Format() is slow, so we cache this to save on performance
@@ -254,6 +254,12 @@ FMulticastDelegateProperty* UK2Node_ActorBoundEvent::GetTargetDelegatePropertyFr
 	return FindFProperty<FMulticastDelegateProperty>(FBlueprintEditorUtils::GetMostUpToDateClass(DelegateOwnerClass), DelegatePropertyName);
 }
 
+FText UK2Node_ActorBoundEvent::GetTargetDelegateDisplayName() const
+{
+	FMulticastDelegateProperty* Prop = GetTargetDelegateProperty();
+	return Prop ? Prop->GetDisplayNameText() : FText::FromName(DelegatePropertyName);
+}
+
 bool UK2Node_ActorBoundEvent::IsUsedByAuthorityOnlyDelegate() const
 {
 	const FMulticastDelegateProperty* TargetDelegateProp = FindFProperty<FMulticastDelegateProperty>(DelegateOwnerClass, DelegatePropertyName);
diff --git a/Engine/Source/Editor/BlueprintGraph/Private/K2Node_ComponentBoundEvent.cpp b/Engine/Source/Editor/BlueprintGraph/Private/K2Node_ComponentBoundEvent.cpp
--- a/Engine/Source/Editor/BlueprintGraph/Private/K2Node_ComponentBoundEvent.cpp
+++ b/Engine/Source/Editor/BlueprintGraph/Private/K2Node_ComponentBoundEvent.cpp
@@ -84,7 +84,7 @@ FText UK2Node_ComponentBoundEvent::GetNodeTitle(ENodeTitleType::Type TitleType)
 	if (CachedNodeTitle.IsOutOfDate(this))
 	{
 		FFormatNamedArguments Args;
-		Args.Add(TEXT("DelegatePropertyName"), DelegatePropertyDisplayName.IsEmpty() ? FText::FromName(DelegatePropertyName) : DelegatePropertyDisplayName);
+		Args.Add(TEXT("DelegatePropertyName"), GetTargetDelegateDisplayName());
 		Args.Add(TEXT("ComponentPropertyName"), FText::FromName(ComponentPropertyName));
 
 		// FText::Format() is slow, so we cache this to save on performance
@@ -99,7 +99,6 @@ void UK2Node_ComponentBoundEvent::InitializeComponentBoundEventParams(FObjectPro
 	{
 		ComponentPropertyName = InComponentProperty->GetFName();
 		DelegatePropertyName = InDelegateProperty->GetFName();
-		DelegatePropertyDisplayName = InDelegateProperty->GetDisplayNameText();
 		DelegateOwnerClass = CastChecked<UClass>(InDelegateProperty->GetOwner<UObject>())->GetAuthoritativeClass();
 
 		EventReference.SetFromField<UFunction>(InDelegateProperty->SignatureFunction, /*bIsConsideredSelfContext =*/false);
@@ -196,7 +195,7 @@ FEdGraphNodeDeprecationResponse UK2Node_ComponentBoundEvent::GetDeprecationRespo
 		{
 			Response.MessageType = EEdGraphNodeDeprecationMessageType::Warning;
 			const FText DetailedMessage = FText::FromString(Function->GetMetaData(FBlueprintMetadata::MD_DeprecationMessage));
-			Response.MessageText = FBlueprintEditorUtils::GetDeprecatedMemberUsageNodeWarning(DelegatePropertyDisplayName, DetailedMessage);
+			Response.MessageText = FBlueprintEditorUtils::GetDeprecatedMemberUsageNodeWarning(GetTargetDelegateDisplayName(), DetailedMessage);
 		}
 	}
 
@@ -214,6 +213,11 @@ FMulticastDelegateProperty* UK2Node_ComponentBoundEvent::GetTargetDelegateProper
 	return FindFProperty<FMulticastDelegateProperty>(DelegateOwnerClass, DelegatePropertyName);
 }
 
+FText UK2Node_ComponentBoundEvent::GetTargetDelegateDisplayName() const
+{
+	FMulticastDelegateProperty* Prop = GetTargetDelegateProperty();
+	return Prop ? Prop->GetDisplayNameText() : FText::FromName(DelegatePropertyName);
+}
 
 FText UK2Node_ComponentBoundEvent::GetTooltipText() const
 {
