Check for if Pins are compatible, not if they are equal, when determining available promotable operator functions.

#rnx

diff --git a/Engine/Source/Editor/BlueprintGraph/Private/BlueprintEditorSettings.cpp b/Engine/Source/Editor/BlueprintGraph/Private/BlueprintEditorSettings.cpp
--- a/Engine/Source/Editor/BlueprintGraph/Private/BlueprintEditorSettings.cpp
+++ b/Engine/Source/Editor/BlueprintGraph/Private/BlueprintEditorSettings.cpp
@@ -48,7 +48,7 @@ UBlueprintEditorSettings::UBlueprintEditorSettings(const FObjectInitializer& Obj
 	, bNavigateToNativeFunctionsFromCallNodes(true)
 	, bDoubleClickNavigatesToParent(true)
 	, bEnableTypePromotion(true)
-	, TypePromotionPinDenyList { UEdGraphSchema_K2::PC_String, UEdGraphSchema_K2::PC_Text, UEdGraphSchema_K2::PC_SoftClass }
+	, TypePromotionPinDenyList { UEdGraphSchema_K2::PC_String, UEdGraphSchema_K2::PC_Text }
 	, BreakpointReloadMethod(EBlueprintBreakpointReloadMethod::RestoreAll)
 	, bEnablePinValueInspectionTooltips(true)
 	, bEnableNamespaceEditorFeatures(true)
diff --git a/Engine/Source/Editor/BlueprintGraph/Private/BlueprintTypePromotion.cpp b/Engine/Source/Editor/BlueprintGraph/Private/BlueprintTypePromotion.cpp
--- a/Engine/Source/Editor/BlueprintGraph/Private/BlueprintTypePromotion.cpp
+++ b/Engine/Source/Editor/BlueprintGraph/Private/BlueprintTypePromotion.cpp
@@ -322,7 +322,7 @@ UFunction* FTypePromotion::FindBestMatchingFunc_Internal(FName Operation, const
 					const bool bClassTypesMatch = ParamType.PinCategory == UEdGraphSchema_K2::PC_Class && Pin->PinType.PinCategory == UEdGraphSchema_K2::PC_Class;
 
 					// Give a point for each function parameter that matches up with a pin to consider
-					if ((!CheckedPins.Contains(Pin) || bIsSinglePin) && (bObjectTypesMatch || bClassTypesMatch || Schema->ArePinTypesEquivalent(ParamType, Pin->PinType)))
+					if ((!CheckedPins.Contains(Pin) || bIsSinglePin) && (bObjectTypesMatch || bClassTypesMatch || Schema->ArePinTypesCompatible(Pin->PinType, ParamType)))
 					{
 						// Are the directions compatible? 
 						// If we are a comparison or only a single pin then we don't care about the direction
