UE-160862 - Fix FW-Gameplay miscellaneous module truncation warnings

diff --git a/Engine/Source/Editor/BlueprintGraph/Private/BlueprintActionDatabase.cpp b/Engine/Source/Editor/BlueprintGraph/Private/BlueprintActionDatabase.cpp
--- a/Engine/Source/Editor/BlueprintGraph/Private/BlueprintActionDatabase.cpp
+++ b/Engine/Source/Editor/BlueprintGraph/Private/BlueprintActionDatabase.cpp
@@ -248,10 +248,10 @@ static UBlueprintNodeSpawner* FBlueprintNodeSpawnerFactory::MakeDocumentationNod
 		UBlueprint* Blueprint = FBlueprintEditorUtils::FindBlueprintForGraph(OuterGraph);
 		check(Blueprint != nullptr);
 
-		float const OldNodePosX   = NewNode->NodePosX;
-		float const OldNodePosY   = NewNode->NodePosY;
-		float const OldHalfHeight = NewNode->NodeHeight / 2.f;
-		float const OldHalfWidth  = NewNode->NodeWidth  / 2.f;
+		const float OldNodePosX = static_cast<float>(NewNode->NodePosX);
+		const float OldNodePosY = static_cast<float>(NewNode->NodePosY);
+		const float OldHalfHeight = NewNode->NodeHeight / 2.f;
+		const float OldHalfWidth  = NewNode->NodeWidth  / 2.f;
 		
 		static const float DocNodePadding = 50.0f;
 		FSlateRect Bounds(OldNodePosX - OldHalfWidth, OldNodePosY - OldHalfHeight, OldNodePosX + OldHalfWidth, OldNodePosY + OldHalfHeight);
diff --git a/Engine/Source/Editor/BlueprintGraph/Private/BlueprintNodeSpawner.cpp b/Engine/Source/Editor/BlueprintGraph/Private/BlueprintNodeSpawner.cpp
--- a/Engine/Source/Editor/BlueprintGraph/Private/BlueprintNodeSpawner.cpp
+++ b/Engine/Source/Editor/BlueprintGraph/Private/BlueprintNodeSpawner.cpp
@@ -323,7 +323,7 @@ bool UBlueprintNodeSpawner::IsTemplateNodeFilteredOut(FBlueprintActionFilter con
 }
 
 //------------------------------------------------------------------------------
-UEdGraphNode* UBlueprintNodeSpawner::SpawnEdGraphNode(TSubclassOf<UEdGraphNode> InNodeClass, UEdGraph* ParentGraph, FBindingSet const& Bindings, FVector2D const Location, FCustomizeNodeDelegate PostSpawnDelegate) const
+UEdGraphNode* UBlueprintNodeSpawner::SpawnEdGraphNode(TSubclassOf<UEdGraphNode> InNodeClass, UEdGraph* ParentGraph, const FBindingSet& Bindings, FVector2D Location, FCustomizeNodeDelegate PostSpawnDelegate) const
 {
 	UEdGraphNode* NewNode = nullptr;
 	if (InNodeClass != nullptr)
@@ -334,8 +334,8 @@ UEdGraphNode* UBlueprintNodeSpawner::SpawnEdGraphNode(TSubclassOf<UEdGraphNode>
 
 		// position the node before invoking PostSpawnDelegate (in case it 
 		// wishes to modify this positioning)
-		NewNode->NodePosX = Location.X;
-		NewNode->NodePosY = Location.Y;
+		NewNode->NodePosX = static_cast<int32>(Location.X);
+		NewNode->NodePosY = static_cast<int32>(Location.Y);
 
 		bool const bIsTemplateNode = FBlueprintNodeTemplateCache::IsTemplateOuter(ParentGraph);
 		PostSpawnDelegate.ExecuteIfBound(NewNode, bIsTemplateNode);
diff --git a/Engine/Source/Editor/BlueprintGraph/Private/BlueprintNodeTemplateCache.cpp b/Engine/Source/Editor/BlueprintGraph/Private/BlueprintNodeTemplateCache.cpp
--- a/Engine/Source/Editor/BlueprintGraph/Private/BlueprintNodeTemplateCache.cpp
+++ b/Engine/Source/Editor/BlueprintGraph/Private/BlueprintNodeTemplateCache.cpp
@@ -46,7 +46,7 @@ DEFINE_LOG_CATEGORY_STATIC(LogBlueprintNodeCache, Log, All);
 namespace BlueprintNodeTemplateCacheImpl
 {
 	/** */
-	static int32 ActiveMemFootprint   = 0;
+	static int64 ActiveMemFootprint   = 0;
 	/** Used to track the average blueprint size (so that we can try to predict when a blueprint would fail to be cached) */
 	static int32 MadeBlueprintCount   = 0;
 	static int32 AverageBlueprintSize = 0;
@@ -171,11 +171,6 @@ static UEdGraph* BlueprintNodeTemplateCacheImpl::FindCompatibleGraph(UBlueprint*
 static UBlueprint* BlueprintNodeTemplateCacheImpl::MakeCompatibleBlueprint(TSubclassOf<UBlueprint> BlueprintClass, UClass* ParentClass, TSubclassOf<UBlueprintGeneratedClass> GeneratedClassType)
 {
 	EBlueprintType BlueprintType = BPTYPE_Normal;
-	// @TODO: BPTYPE_LevelScript requires a level outer, which we don't want to have here... can we get away without it?
-// 	if (BlueprintClass->IsChildOf<ULevelScriptBlueprint>())
-// 	{
-// 		BlueprintType = BPTYPE_LevelScript;
-// 	}
 
 	if (GeneratedClassType == nullptr)
 	{
@@ -183,8 +178,8 @@ static UBlueprint* BlueprintNodeTemplateCacheImpl::MakeCompatibleBlueprint(TSubc
 	}
 
 	UPackage* BlueprintOuter = GetTransientPackage();
-	FString const DesiredName = FString::Printf(TEXT("PROTO_BP_%s"), *BlueprintClass->GetName());
-	FName   const BlueprintName = MakeUniqueObjectName(BlueprintOuter, BlueprintClass, FName(*DesiredName));
+	const FString DesiredName = FString::Printf(TEXT("PROTO_BP_%s"), *BlueprintClass->GetName());
+	const FName BlueprintName = MakeUniqueObjectName(BlueprintOuter, BlueprintClass, FName(*DesiredName));
 
 	BlueprintClass = FBlueprintEditorUtils::FindFirstNativeClass(BlueprintClass);
 	UBlueprint* NewBlueprint = FKismetEditorUtilities::CreateBlueprint(ParentClass, BlueprintOuter, BlueprintName, BlueprintType, BlueprintClass, GeneratedClassType);
@@ -192,11 +187,17 @@ static UBlueprint* BlueprintNodeTemplateCacheImpl::MakeCompatibleBlueprint(TSubc
 
 	++MadeBlueprintCount;
 
-	float const AproxBlueprintSize = ApproximateMemFootprint(NewBlueprint);
+	const float AproxBlueprintSize = static_cast<float>(ApproximateMemFootprint(NewBlueprint));
+	const float MadeBlueprintCountFloat = static_cast<float>(MadeBlueprintCount);
+
 	// track the average blueprint size, so that we can attempt to predict 
 	// whether a  blueprint will fail to be cached (when the cache is near full)
-	AverageBlueprintSize = AverageBlueprintSize * ((float)(MadeBlueprintCount-1) / MadeBlueprintCount) + 
-		(AproxBlueprintSize / MadeBlueprintCount) + 0.5f;
+	AverageBlueprintSize = 
+		static_cast<int32>(
+			(AverageBlueprintSize * ((MadeBlueprintCountFloat - 1.0f) / MadeBlueprintCountFloat)) +
+			(AproxBlueprintSize / MadeBlueprintCountFloat) +
+			0.5f
+		);
 
 	return NewBlueprint;
 }
@@ -237,9 +238,9 @@ bool BlueprintNodeTemplateCacheImpl::IsTemplateOuter(UEdGraph* ParentGraph)
 //------------------------------------------------------------------------------
 static int32 BlueprintNodeTemplateCacheImpl::GetCacheCapSize()
 {
-	UBlueprintEditorSettings const* BpSettings = GetDefault<UBlueprintEditorSettings>();
+	const UBlueprintEditorSettings* BpSettings = GetDefault<UBlueprintEditorSettings>();
 	// have to convert from MB to bytes
-	return (BpSettings->NodeTemplateCacheCapMB * 1024.f * 1024.f) + 0.5f;
+	return (static_cast<int32>(BpSettings->NodeTemplateCacheCapMB) << 20);
 }
 
 //------------------------------------------------------------------------------
@@ -462,9 +463,9 @@ void FBlueprintNodeTemplateCache::ClearCachedTemplate(UBlueprintNodeSpawner cons
 }
 
 //------------------------------------------------------------------------------
-int32 FBlueprintNodeTemplateCache::GetEstimateCacheSize() const
+int64 FBlueprintNodeTemplateCache::GetEstimateCacheSize() const
 {
-	int32 TotalEstimatedSize = ApproximateObjectMem;
+	int64 TotalEstimatedSize = ApproximateObjectMem;
 	TotalEstimatedSize += TemplateOuters.GetAllocatedSize();
 	TotalEstimatedSize += NodeTemplateCache.GetAllocatedSize();
 	TotalEstimatedSize += sizeof(*this);
@@ -473,7 +474,7 @@ int32 FBlueprintNodeTemplateCache::GetEstimateCacheSize() const
 }
 
 //------------------------------------------------------------------------------
-int32 FBlueprintNodeTemplateCache::RecalculateCacheSize()
+int64 FBlueprintNodeTemplateCache::RecalculateCacheSize()
 {
 	ApproximateObjectMem = 0;
 	for (UBlueprint* Blueprint : TemplateOuters)
diff --git a/Engine/Source/Editor/BlueprintGraph/Private/EdGraphSchema_K2.cpp b/Engine/Source/Editor/BlueprintGraph/Private/EdGraphSchema_K2.cpp
--- a/Engine/Source/Editor/BlueprintGraph/Private/EdGraphSchema_K2.cpp
+++ b/Engine/Source/Editor/BlueprintGraph/Private/EdGraphSchema_K2.cpp
@@ -7900,13 +7900,13 @@ float UEdGraphSchema_K2::GetActionFilteredWeight(const FGraphActionListBuilderBa
 				if (BestMatchCharLength > 0 && WeightPerList > 0)
 				{
 					// How many words that we are checking had partial matches compared to what the user typed in?
-					float PercMatch = (float)WordMatchCount / (float)KeywordArray.Num();
+					float PercMatch = static_cast<float>(WordMatchCount) / static_cast<float>(KeywordArray.Num());
 
 					float PercentageBonus = (WeightPerList * PercMatch * BPContextMenuConsoleVariables::PercentageMatchWeightMultiplier);
 					WeightPerList += PercentageBonus;
 
 					// The shorter the matching word, the larger bonus it gets
-					float ShortFactor = BPContextMenuConsoleVariables::MaxWordLength - FMath::Min(BestMatchCharLength, BPContextMenuConsoleVariables::MaxWordLength);
+					float ShortFactor = static_cast<float>(BPContextMenuConsoleVariables::MaxWordLength - FMath::Min(BestMatchCharLength, BPContextMenuConsoleVariables::MaxWordLength));
 					float ShortWeight = ShortFactor * BPContextMenuConsoleVariables::ShorterWeight * (bIsFromDrag ? 0.25f : 1.0f);
 					WeightPerList += ShortWeight;
 
diff --git a/Engine/Source/Editor/BlueprintGraph/Private/EdGraphSchema_K2_Actions.cpp b/Engine/Source/Editor/BlueprintGraph/Private/EdGraphSchema_K2_Actions.cpp
--- a/Engine/Source/Editor/BlueprintGraph/Private/EdGraphSchema_K2_Actions.cpp
+++ b/Engine/Source/Editor/BlueprintGraph/Private/EdGraphSchema_K2_Actions.cpp
@@ -272,11 +272,11 @@ UEdGraphNode* FEdGraphSchemaAction_K2NewNode::CreateNode(
 
 	// For input pins, new node will generally overlap node being dragged off
 	// Work out if we want to visually push away from connected node
-	int32 XLocation = Location.X;
+	int32 XLocation = static_cast<int32>(Location.X);
 	if (FromPinPtr.IsValid() && FromPinPtr->Direction == EGPD_Input)
 	{
 		UEdGraphNode* PinNode = FromPinPtr->GetOwningNode();
-		const float XDelta = FMath::Abs(PinNode->NodePosX - Location.X);
+		const double XDelta = FMath::Abs(PinNode->NodePosX - Location.X);
 
 		if (XDelta < NodeDistance)
 		{
@@ -286,7 +286,7 @@ UEdGraphNode* FEdGraphSchemaAction_K2NewNode::CreateNode(
 		}
 	}
 	ResultNode->NodePosX = XLocation;
-	ResultNode->NodePosY = Location.Y;
+	ResultNode->NodePosY = static_cast<int32>(Location.Y);
 	ResultNode->SnapToGrid(GetDefault<UEditorStyleSettings>()->GridSnapSize);
 
 	// make sure to auto-wire after we position the new node (in case the 
@@ -616,10 +616,10 @@ UEdGraphNode* FEdGraphSchemaAction_K2AddCustomEvent::PerformAction(class UEdGrap
 UEdGraphNode* FEdGraphSchemaAction_K2AddCallOnActor::PerformAction(class UEdGraph* ParentGraph, UEdGraphPin* FromPin, const FVector2D Location, bool bSelectNewNode/* = true*/)
 {
 	// Snap the node placement location to the grid, ensures calculations later match up better
-	const float  GridSnapSize = (float)GetDefault<UEditorStyleSettings>()->GridSnapSize;
-	FVector2f LocalLocation;	// LWC_TODO: Precision loss
-	LocalLocation.X = FMath::GridSnap( (float)Location.X, GridSnapSize);
-	LocalLocation.Y = FMath::GridSnap( (float)Location.Y, GridSnapSize);
+	const uint32 GridSnapSize = GetDefault<UEditorStyleSettings>()->GridSnapSize;
+	FVector2D LocalLocation;
+	LocalLocation.X = FMath::GridSnap(Location.X, GridSnapSize);
+	LocalLocation.Y = FMath::GridSnap(Location.Y, GridSnapSize);
 
 	// First use the base functionality to spawn the 'call function' node
 	UEdGraphNode* CallNode = FEdGraphSchemaAction_K2NewNode::PerformAction(ParentGraph, FromPin, FVector2D(LocalLocation));
@@ -628,9 +628,9 @@ UEdGraphNode* FEdGraphSchemaAction_K2AddCallOnActor::PerformAction(class UEdGrap
 	// this is the guesstimate of the function node's height, snapped to grid units
 	const float FunctionNodeHeight = FMath::GridSnap( FunctionNodeHeightUnsnapped, GridSnapSize );
 	// this is roughly the middle of the function node height
-	const float FunctionNodeMidY = LocalLocation.Y + FunctionNodeHeight * 0.5f;
+	const float FunctionNodeMidY = static_cast<float>(LocalLocation.Y + FunctionNodeHeight * 0.5f);
 	// this is the offset up from the mid point at which we start placing nodes 
-	const float StartYOffset = (float((LevelActors.Num() > 0) ? LevelActors.Num()-1 : 0) * -NodeLiteralHeight) * 0.5f;
+	const float StartYOffset = (static_cast<float>((LevelActors.Num() > 0) ? LevelActors.Num() - 1 : 0) * -NodeLiteralHeight) * 0.5f;
 	// The Y location we start placing nodes from
 	const float ReferencedNodesPlacementYLocation = FunctionNodeMidY + StartYOffset;
 
@@ -647,11 +647,11 @@ UEdGraphNode* FEdGraphSchemaAction_K2AddCallOnActor::PerformAction(class UEdGrap
 
 			LiteralNode->SetObjectRef(LevelActor);
 			LiteralNode->AllocateDefaultPins();
-			LiteralNode->NodePosX = LocalLocation.X - FunctionNodeLiteralReferencesXOffset;
+			LiteralNode->NodePosX = static_cast<int32>(LocalLocation.X - FunctionNodeLiteralReferencesXOffset);
 
 			// this is the current offset down from the Y start location to place the next node at
-			float CurrentNodeOffset = NodeLiteralHeight * float(ActorIndex);
-			LiteralNode->NodePosY = ReferencedNodesPlacementYLocation + CurrentNodeOffset;
+			float CurrentNodeOffset = NodeLiteralHeight * static_cast<float>(ActorIndex);
+			LiteralNode->NodePosY = static_cast<int32>(ReferencedNodesPlacementYLocation + CurrentNodeOffset);
 
 			LiteralNode->SnapToGrid(GridSnapSize);
 
diff --git a/Engine/Source/Editor/BlueprintGraph/Private/K2Node_AddPinInterface.cpp b/Engine/Source/Editor/BlueprintGraph/Private/K2Node_AddPinInterface.cpp
--- a/Engine/Source/Editor/BlueprintGraph/Private/K2Node_AddPinInterface.cpp
+++ b/Engine/Source/Editor/BlueprintGraph/Private/K2Node_AddPinInterface.cpp
@@ -8,6 +8,6 @@
 FName IK2Node_AddPinInterface::GetNameForAdditionalPin(int32 PinIndex)
 {
 	check(PinIndex < GetMaxInputPinsNum());
-	const FName Name(*FString::Chr(TCHAR('A') + PinIndex));
+	const FName Name(*FString::Chr(TCHAR('A') + static_cast<TCHAR>(PinIndex)));
 	return Name;
 }
\ No newline at end of file
diff --git a/Engine/Source/Editor/BlueprintGraph/Private/K2Node_CustomEvent.cpp b/Engine/Source/Editor/BlueprintGraph/Private/K2Node_CustomEvent.cpp
--- a/Engine/Source/Editor/BlueprintGraph/Private/K2Node_CustomEvent.cpp
+++ b/Engine/Source/Editor/BlueprintGraph/Private/K2Node_CustomEvent.cpp
@@ -516,8 +516,8 @@ UK2Node_CustomEvent* UK2Node_CustomEvent::CreateFromFunction(FVector2D GraphPosi
 			}
 		}
 
-		CustomEventNode->NodePosX = GraphPosition.X;
-		CustomEventNode->NodePosY = GraphPosition.Y;
+		CustomEventNode->NodePosX = static_cast<int32>(GraphPosition.X);
+		CustomEventNode->NodePosY = static_cast<int32>(GraphPosition.Y);
 		CustomEventNode->SnapToGrid(GetDefault<UEditorStyleSettings>()->GridSnapSize);
 	}
 
diff --git a/Engine/Source/Editor/BlueprintGraph/Private/K2Node_DoOnceMultiInput.cpp b/Engine/Source/Editor/BlueprintGraph/Private/K2Node_DoOnceMultiInput.cpp
--- a/Engine/Source/Editor/BlueprintGraph/Private/K2Node_DoOnceMultiInput.cpp
+++ b/Engine/Source/Editor/BlueprintGraph/Private/K2Node_DoOnceMultiInput.cpp
@@ -69,7 +69,7 @@ FText UK2Node_DoOnceMultiInput::GetNameForPin(int32 PinIndex, bool In)
 {
 	check(PinIndex < GetMaxInputPinsNum());
 	FString Name;
-	Name.AppendChar(TCHAR('A') + PinIndex);
+	Name.AppendChar(TCHAR('A') + static_cast<TCHAR>(PinIndex));
 	
 	FFormatNamedArguments Args;
 	Args.Add(TEXT("Identifier"), FText::FromString(Name));
diff --git a/Engine/Source/Editor/BlueprintGraph/Private/K2Node_EaseFunction.cpp b/Engine/Source/Editor/BlueprintGraph/Private/K2Node_EaseFunction.cpp
--- a/Engine/Source/Editor/BlueprintGraph/Private/K2Node_EaseFunction.cpp
+++ b/Engine/Source/Editor/BlueprintGraph/Private/K2Node_EaseFunction.cpp
@@ -481,7 +481,7 @@ void UK2Node_EaseFunction::RefreshPinVisibility()
 	const auto EaseFuncPin = GetEaseFuncPin();
 	UEnum * Enum = StaticEnum<EEasingFunc::Type>();
 	check(Enum != NULL);
-	const int32 NewEasingFunc = CanCustomizeCurve() ? Enum->GetValueByName(*EaseFuncPin->DefaultValue) : -1;
+	const int64 NewEasingFunc = CanCustomizeCurve() ? Enum->GetValueByName(*EaseFuncPin->DefaultValue) : -1;
 	
 	// Early exit in case no changes are required
 	const UEdGraphSchema_K2* K2Schema = Cast<const UEdGraphSchema_K2>(GetSchema());
diff --git a/Engine/Source/Editor/BlueprintGraph/Private/K2Node_EnumLiteral.cpp b/Engine/Source/Editor/BlueprintGraph/Private/K2Node_EnumLiteral.cpp
--- a/Engine/Source/Editor/BlueprintGraph/Private/K2Node_EnumLiteral.cpp
+++ b/Engine/Source/Editor/BlueprintGraph/Private/K2Node_EnumLiteral.cpp
@@ -137,7 +137,7 @@ void UK2Node_EnumLiteral::ExpandNode(class FKismetCompilerContext& CompilerConte
 		bSuccess = Index != INDEX_NONE;
 		if (bSuccess)
 		{
-			NewInputPin->DefaultValue = FString::FromInt(Enum->GetValueByIndex(Index));
+			NewInputPin->DefaultValue = FString::FromInt(IntCastChecked<int32, int64>(Enum->GetValueByIndex(Index)));
 		}
 
 		UK2Node_CastByteToEnum* CastByte = CompilerContext.SpawnIntermediateNode<UK2Node_CastByteToEnum>(this, SourceGraph);
diff --git a/Engine/Source/Editor/BlueprintGraph/Private/K2Node_GetEnumeratorName.cpp b/Engine/Source/Editor/BlueprintGraph/Private/K2Node_GetEnumeratorName.cpp
--- a/Engine/Source/Editor/BlueprintGraph/Private/K2Node_GetEnumeratorName.cpp
+++ b/Engine/Source/Editor/BlueprintGraph/Private/K2Node_GetEnumeratorName.cpp
@@ -187,7 +187,7 @@ void UK2Node_GetEnumeratorName::ExpandNode(class FKismetCompilerContext& Compile
 	{
 		//MAKE LITERAL BYTE FROM LITERAL ENUM
 		const FString EnumLiteral = IndexPin->GetDefaultAsString();
-		const int32 NumericValue = Enum->GetValueByName(*EnumLiteral);
+		const int32 NumericValue = IntCastChecked<int32, int64>(Enum->GetValueByName(*EnumLiteral));
 		if (NumericValue == INDEX_NONE) 
 		{
 			CompilerContext.MessageLog.Error(*FText::Format(NSLOCTEXT("K2Node", "GetEnumeratorNam_Error_InvalidNameFmt", "@@ has invalid enum value '{0}'"), FText::FromString(EnumLiteral)).ToString(), this);
diff --git a/Engine/Source/Editor/BlueprintGraph/Private/K2Node_MathExpression.cpp b/Engine/Source/Editor/BlueprintGraph/Private/K2Node_MathExpression.cpp
--- a/Engine/Source/Editor/BlueprintGraph/Private/K2Node_MathExpression.cpp
+++ b/Engine/Source/Editor/BlueprintGraph/Private/K2Node_MathExpression.cpp
@@ -1393,13 +1393,13 @@ public:
 
 		// position the entry and exit nodes somewhere sane
 		{
-			const FVector2D EntryPos = GetNodePosition(GraphXBounds.X - 1, 0);
-			EntryNode->NodePosX = EntryPos.X;
-			EntryNode->NodePosY = EntryPos.Y;
+			const FVector2D EntryPos = GetNodePosition(static_cast<int32>(GraphXBounds.X - 1), 0);
+			EntryNode->NodePosX = static_cast<int32>(EntryPos.X);
+			EntryNode->NodePosY = static_cast<int32>(EntryPos.Y);
 
-			const FVector2D ExitPos = GetNodePosition(GraphXBounds.Y + 1, 0);
-			ExitNode->NodePosX = ExitPos.X;
-			ExitNode->NodePosY = ExitPos.Y;
+			const FVector2D ExitPos = GetNodePosition(static_cast<int32>(GraphXBounds.Y + 1), 0);
+			ExitNode->NodePosX = static_cast<int32>(ExitPos.X);
+			ExitNode->NodePosY = static_cast<int32>(ExitPos.Y);
 		}
 				
 		bool bHasErrors = ((MessageLog.NumErrors - StartingErrorCount) > 0);
diff --git a/Engine/Source/Editor/BlueprintGraph/Private/K2Node_SwitchEnum.cpp b/Engine/Source/Editor/BlueprintGraph/Private/K2Node_SwitchEnum.cpp
--- a/Engine/Source/Editor/BlueprintGraph/Private/K2Node_SwitchEnum.cpp
+++ b/Engine/Source/Editor/BlueprintGraph/Private/K2Node_SwitchEnum.cpp
@@ -186,8 +186,8 @@ UK2Node::ERedirectType UK2Node_SwitchEnum::DoPinsMatchForReconstruction(const UE
 	UK2Node::ERedirectType ReturnValue = Super::DoPinsMatchForReconstruction(NewPin, NewPinIndex, OldPin, OldPinIndex);
 	if (ReturnValue == UK2Node::ERedirectType_None && Enum && OldPinIndex > 2 && NewPinIndex > 2)
 	{
-		const int32 OldValue = Enum->GetValueByName(OldPin->PinName);
-		const int32 NewValue = Enum->GetValueByName(NewPin->PinName);
+		const int64 OldValue = Enum->GetValueByName(OldPin->PinName);
+		const int64 NewValue = Enum->GetValueByName(NewPin->PinName);
 		// This handles redirects properly
 		if (OldValue == NewValue && OldValue != INDEX_NONE)
 		{
diff --git a/Engine/Source/Editor/BlueprintGraph/Public/BlueprintNodeSpawner.h b/Engine/Source/Editor/BlueprintGraph/Public/BlueprintNodeSpawner.h
--- a/Engine/Source/Editor/BlueprintGraph/Public/BlueprintNodeSpawner.h
+++ b/Engine/Source/Editor/BlueprintGraph/Public/BlueprintNodeSpawner.h
@@ -292,7 +292,7 @@ private:
 	 * @param  PostSpawnDelegate	A delegate to run after spawning the node, but prior to allocating the node's pins.
 	 * @return Null if it failed to spawn a node (if NodeClass is null), otherwise a newly spawned node.
 	 */
-	UEdGraphNode* SpawnEdGraphNode(TSubclassOf<UEdGraphNode> NodeClass, UEdGraph* ParentGraph, FBindingSet const& Bindings, FVector2D const Location, FCustomizeNodeDelegate PostSpawnDelegate) const;
+	UEdGraphNode* SpawnEdGraphNode(TSubclassOf<UEdGraphNode> NodeClass, UEdGraph* ParentGraph, const FBindingSet& Bindings, FVector2D Location, FCustomizeNodeDelegate PostSpawnDelegate) const;
 };
 
 /*******************************************************************************
diff --git a/Engine/Source/Editor/BlueprintGraph/Public/BlueprintNodeTemplateCache.h b/Engine/Source/Editor/BlueprintGraph/Public/BlueprintNodeTemplateCache.h
--- a/Engine/Source/Editor/BlueprintGraph/Public/BlueprintNodeTemplateCache.h
+++ b/Engine/Source/Editor/BlueprintGraph/Public/BlueprintNodeTemplateCache.h
@@ -75,7 +75,7 @@ public:
 	 * 
 	 * @return The approximated total (in bytes) that this cache has allocated.
 	 */
-	int32 GetEstimateCacheSize() const;
+	int64 GetEstimateCacheSize() const;
 
 	/**
 	 * External systems can make changes that alter the memory footprint of the
@@ -85,7 +85,7 @@ public:
 	 * 
 	 * @return The new approximated total (in bytes) that this cache has allocated.
 	 */
-	int32 RecalculateCacheSize();
+	int64 RecalculateCacheSize();
 
 	// FGCObject interface
 	virtual void AddReferencedObjects(FReferenceCollector& Collector) override;
@@ -136,5 +136,5 @@ private:
 	 * entry is added, so we keep this approximate tally of memory allocated for
 	 * UObjects (owned by this system).
 	 */
-	int32 ApproximateObjectMem;
+	int64 ApproximateObjectMem;
 };
diff --git a/Engine/Source/Editor/GraphEditor/Private/SGraphNodeResizable.cpp b/Engine/Source/Editor/GraphEditor/Private/SGraphNodeResizable.cpp
--- a/Engine/Source/Editor/GraphEditor/Private/SGraphNodeResizable.cpp
+++ b/Engine/Source/Editor/GraphEditor/Private/SGraphNodeResizable.cpp
@@ -168,7 +168,7 @@ FReply SGraphNodeResizable::OnMouseMove(const FGeometry& MyGeometry, const FPoin
 		DragSize.X += DeltaNodeSize.X;
 		DragSize.Y += DeltaNodeSize.Y;
 		// apply snap
-		const float SnapSize = SNodePanel::GetSnapGridSize();
+		const uint32 SnapSize = SNodePanel::GetSnapGridSize();
 		FVector2D SnappedSize;
 		SnappedSize.X = SnapSize * FMath::RoundToFloat( DragSize.X/SnapSize );
 		SnappedSize.Y = SnapSize * FMath::RoundToFloat( DragSize.Y/SnapSize );
diff --git a/Engine/Source/Editor/GraphEditor/Private/SGraphPanel.cpp b/Engine/Source/Editor/GraphEditor/Private/SGraphPanel.cpp
--- a/Engine/Source/Editor/GraphEditor/Private/SGraphPanel.cpp
+++ b/Engine/Source/Editor/GraphEditor/Private/SGraphPanel.cpp
@@ -760,24 +760,24 @@ FReply SGraphPanel::OnKeyDown( const FGeometry& MyGeometry, const FKeyEvent& InK
 		const bool bIsModifierActive = InKeyEvent.IsCommandDown() || InKeyEvent.IsAltDown() || InKeyEvent.IsShiftDown() || InKeyEvent.IsControlDown();
 		if (!bIsModifierActive)
 		{
-			if( InKeyEvent.GetKey() == EKeys::Up  ||InKeyEvent.GetKey() ==  EKeys::NumPadEight )
+			if( InKeyEvent.GetKey() == EKeys::Up || InKeyEvent.GetKey() == EKeys::NumPadEight )
 			{
-				UpdateSelectedNodesPositions(FVector2D(0.0f,-GetSnapGridSize()));
+				UpdateSelectedNodesPositions(FVector2D(0.0f,-1.0f * GetSnapGridSize()));
 				return FReply::Handled();
 			}
-			if( InKeyEvent.GetKey() ==  EKeys::Down || InKeyEvent.GetKey() ==  EKeys::NumPadTwo )
+			if( InKeyEvent.GetKey() == EKeys::Down || InKeyEvent.GetKey() == EKeys::NumPadTwo )
 			{
 				UpdateSelectedNodesPositions(FVector2D(0.0f,GetSnapGridSize()));
 				return FReply::Handled();
 			}
-			if( InKeyEvent.GetKey() ==  EKeys::Right || InKeyEvent.GetKey() ==  EKeys::NumPadSix )
+			if( InKeyEvent.GetKey() == EKeys::Right || InKeyEvent.GetKey() == EKeys::NumPadSix )
 			{
 				UpdateSelectedNodesPositions(FVector2D(GetSnapGridSize(),0.0f));
 				return FReply::Handled();
 			}
-			if( InKeyEvent.GetKey() ==  EKeys::Left || InKeyEvent.GetKey() ==  EKeys::NumPadFour )
+			if( InKeyEvent.GetKey() == EKeys::Left || InKeyEvent.GetKey() == EKeys::NumPadFour )
 			{
-				UpdateSelectedNodesPositions(FVector2D(-GetSnapGridSize(),0.0f));
+				UpdateSelectedNodesPositions(FVector2D(-1.0f * GetSnapGridSize(),0.0f));
 				return FReply::Handled();
 			}
 		}
diff --git a/Engine/Source/Editor/GraphEditor/Private/SNodePanel.cpp b/Engine/Source/Editor/GraphEditor/Private/SNodePanel.cpp
--- a/Engine/Source/Editor/GraphEditor/Private/SNodePanel.cpp
+++ b/Engine/Source/Editor/GraphEditor/Private/SNodePanel.cpp
@@ -1332,7 +1332,7 @@ void SNodePanel::RestoreViewSettings(const FVector2D& InViewOffset, float InZoom
 	CurrentBookmarkGuid = InBookmarkGuid;
 }
 
-float SNodePanel::GetSnapGridSize()
+uint32 SNodePanel::GetSnapGridSize()
 {
 	return GetDefault<UEditorStyleSettings>()->GridSnapSize;
 }
diff --git a/Engine/Source/Editor/GraphEditor/Public/SNodePanel.h b/Engine/Source/Editor/GraphEditor/Public/SNodePanel.h
--- a/Engine/Source/Editor/GraphEditor/Public/SNodePanel.h
+++ b/Engine/Source/Editor/GraphEditor/Public/SNodePanel.h
@@ -783,7 +783,7 @@ public:
 	void RestoreViewSettings(const FVector2D& InViewOffset, float InZoomAmount, const FGuid& InBookmarkGuid = FGuid());
 
 	/** Get the grid snap size */
-	static float GetSnapGridSize();
+	static uint32 GetSnapGridSize();
 
 	/** 
 	 * Zooms out to fit either all nodes or only the selected ones.
diff --git a/Engine/Source/Runtime/Engine/Classes/EdGraph/EdGraphNode.h b/Engine/Source/Runtime/Engine/Classes/EdGraph/EdGraphNode.h
--- a/Engine/Source/Runtime/Engine/Classes/EdGraph/EdGraphNode.h
+++ b/Engine/Source/Runtime/Engine/Classes/EdGraph/EdGraphNode.h
@@ -660,7 +660,7 @@ public:
 	void BreakAllNodeLinks();
 
 	/** Snap this node to a specified grid size */
-	void SnapToGrid(float GridSnapSize);
+	void SnapToGrid(uint32 GridSnapSize);
 
 	/** Clear error flag */
 	void ClearCompilerMessage()
diff --git a/Engine/Source/Runtime/Engine/Private/EdGraph/EdGraphNode.cpp b/Engine/Source/Runtime/Engine/Private/EdGraph/EdGraphNode.cpp
--- a/Engine/Source/Runtime/Engine/Private/EdGraph/EdGraphNode.cpp
+++ b/Engine/Source/Runtime/Engine/Private/EdGraph/EdGraphNode.cpp
@@ -488,10 +488,10 @@ void UEdGraphNode::GetPinHoverText(const UEdGraphPin& Pin, FString& HoverTextOut
 	HoverTextOut = Pin.PinToolTip;
 }
 
-void UEdGraphNode::SnapToGrid(float GridSnapSize)
+void UEdGraphNode::SnapToGrid(uint32 GridSnapSize)
 {
-	NodePosX = GridSnapSize * FMath::RoundToInt(NodePosX/GridSnapSize);
-	NodePosY = GridSnapSize * FMath::RoundToInt(NodePosY/GridSnapSize);
+	NodePosX = GridSnapSize * (NodePosX / GridSnapSize);
+	NodePosY = GridSnapSize * (NodePosY / GridSnapSize);
 }
 
 class UEdGraph* UEdGraphNode::GetGraph() const
diff --git a/Engine/Source/Runtime/GameplayTags/Private/GameplayTagContainer.cpp b/Engine/Source/Runtime/GameplayTags/Private/GameplayTagContainer.cpp
--- a/Engine/Source/Runtime/GameplayTags/Private/GameplayTagContainer.cpp
+++ b/Engine/Source/Runtime/GameplayTags/Private/GameplayTagContainer.cpp
@@ -102,12 +102,12 @@ void SerializeTagNetIndexPacked(FArchive& Ar, FGameplayTagNetIndex& Value, const
 		{
 			uint32 SecondData = 0;
 			Ar.SerializeBits(&SecondData, SecondSegment);
-			Value = (SecondData << FirstSegment);
+			Value = IntCastChecked<uint16, uint32>(SecondData << FirstSegment);
 			Value |= (FirstData & BitMasks[FirstSegment]);
 		}
 		else
 		{
-			Value = FirstData;
+			Value = IntCastChecked<uint16, uint32>(FirstData);
 		}
 
 	}
@@ -966,7 +966,7 @@ bool FGameplayTagContainer::NetSerialize(FArchive& Ar, class UPackageMap* Map, b
 
 	if (Ar.IsSaving())
 	{
-		uint8 NumTags = GameplayTags.Num();
+		uint8 NumTags = IntCastChecked<uint8, int32>(GameplayTags.Num());
 		uint8 MaxSize = (1 << NumBitsForContainerSize) - 1;
 		if (!ensureMsgf(NumTags <= MaxSize, TEXT("TagContainer has %d elements when max is %d! Tags: %s"), NumTags, MaxSize, *ToStringSimple()))
 		{
@@ -1220,7 +1220,7 @@ bool FGameplayTag::NetSerialize_Packed(FArchive& Ar, class UPackageMap* Map, boo
 
 			uint32 NetIndex32 = NetIndex;
 			Ar.SerializeIntPacked(NetIndex32);
-			NetIndex = NetIndex32;
+			NetIndex = IntCastChecked<uint16, uint32>(NetIndex32);
 
 			if (Ar.IsLoading())
 			{
@@ -1850,7 +1850,7 @@ FGameplayTagQueryExpression& FGameplayTagQueryExpression::AddTag(FName TagName)
 void FGameplayTagQueryExpression::EmitTokens(TArray<uint8>& TokenStream, TArray<FGameplayTag>& TagDictionary) const
 {
 	// emit exprtype
-	TokenStream.Add(ExprType);
+	TokenStream.Add(static_cast<uint8>(ExprType));
 
 	// emit exprdata
 	switch (ExprType)
@@ -1923,7 +1923,7 @@ static void TagPackingTest()
 		{
 			for (int32 NetIndex=0; NetIndex < FMath::Pow(2.f, TotalNetIndexBits); ++NetIndex)
 			{
-				FGameplayTagNetIndex NI = NetIndex;
+				FGameplayTagNetIndex NI = IntCastChecked<uint16, int32>(NetIndex);
 
 				FNetBitWriter	BitWriter(nullptr, 1024 * 8);
 				SerializeTagNetIndexPacked(BitWriter, NI, NetIndexBitsPerComponent, TotalNetIndexBits);
diff --git a/Engine/Source/Runtime/GameplayTags/Private/GameplayTagsManager.cpp b/Engine/Source/Runtime/GameplayTags/Private/GameplayTagsManager.cpp
--- a/Engine/Source/Runtime/GameplayTags/Private/GameplayTagsManager.cpp
+++ b/Engine/Source/Runtime/GameplayTags/Private/GameplayTagsManager.cpp
@@ -612,11 +612,11 @@ void UGameplayTagsManager::ConstructNetIndex()
 		checkf( Found, TEXT("Tag %s not found in NetworkGameplayTagNodeIndex"), *Tag.ToString() );
 	}
 
-	InvalidTagNetIndex = NetworkGameplayTagNodeIndex.Num()+1;
+	InvalidTagNetIndex = IntCastChecked<uint16, int32>(NetworkGameplayTagNodeIndex.Num() + 1);
 	NetIndexTrueBitNum = FMath::CeilToInt(FMath::Log2(static_cast<float>(InvalidTagNetIndex)));
 	
 	// This should never be smaller than NetIndexTrueBitNum
-	NetIndexFirstBitSegment = FMath::Min<int64>(GetDefault<UGameplayTagsSettings>()->NetIndexFirstBitSegment, NetIndexTrueBitNum);
+	NetIndexFirstBitSegment = FMath::Min<int32>(GetDefault<UGameplayTagsSettings>()->NetIndexFirstBitSegment, NetIndexTrueBitNum);
 
 	// This is now sorted and it should be the same on both client and server
 	if (NetworkGameplayTagNodeIndex.Num() >= INVALID_TAGNETINDEX)
