[Engine] added IsDefaultAsString to optimize and compare to default value

#p4v-preflight-copy 21907138

diff --git a/Engine/Source/Editor/BlueprintGraph/Private/EdGraphSchema_K2.cpp b/Engine/Source/Editor/BlueprintGraph/Private/EdGraphSchema_K2.cpp
--- a/Engine/Source/Editor/BlueprintGraph/Private/EdGraphSchema_K2.cpp
+++ b/Engine/Source/Editor/BlueprintGraph/Private/EdGraphSchema_K2.cpp
@@ -4960,17 +4960,18 @@ bool UEdGraphSchema_K2::DoesDefaultValueMatchAutogenerated(const UEdGraphPin& In
 		}
 		else if (InPin.AutogeneratedDefaultValue.IsEmpty())
 		{
-			const FString PinDefaultValue = InPin.GetDefaultAsString();
-			if (PinDefaultValue.IsEmpty())
+			if (InPin.IsDefaultAsStringEmpty())
 			{
 				return true;
 			}
 			else if (InPin.PinType.PinCategory == PC_Boolean)
 			{
+				const FString PinDefaultValue = InPin.GetDefaultAsString();
 				return (PinDefaultValue == TEXT("false"));
 			}
 			else if (InPin.PinType.PinCategory == PC_Int)
 			{
+				const FString PinDefaultValue = InPin.GetDefaultAsString();
 				if (FCString::Atoi(*PinDefaultValue) == 0)
 				{
 					return true;
@@ -4978,6 +4979,7 @@ bool UEdGraphSchema_K2::DoesDefaultValueMatchAutogenerated(const UEdGraphPin& In
 			}
 			else if (InPin.PinType.PinCategory == PC_Int64)
 			{
+				const FString PinDefaultValue = InPin.GetDefaultAsString();
 				if (FCString::Atoi64(*PinDefaultValue) == 0)
 				{
 					return true;
@@ -4985,6 +4987,7 @@ bool UEdGraphSchema_K2::DoesDefaultValueMatchAutogenerated(const UEdGraphPin& In
 			}
 			else if (InPin.PinType.PinCategory == PC_Name)
 			{
+				const FString PinDefaultValue = InPin.GetDefaultAsString();
 				return (PinDefaultValue == TEXT("None"));
 			}
 		}
diff --git a/Engine/Source/Runtime/Engine/Classes/EdGraph/EdGraphPin.h b/Engine/Source/Runtime/Engine/Classes/EdGraph/EdGraphPin.h
--- a/Engine/Source/Runtime/Engine/Classes/EdGraph/EdGraphPin.h
+++ b/Engine/Source/Runtime/Engine/Classes/EdGraph/EdGraphPin.h
@@ -469,6 +469,11 @@ public:
 	 */
 	ENGINE_API FString GetDefaultAsString() const;
 
+	/**
+	 * Returns true if GetDefaultAsString will return an empty string.
+	 */
+	ENGINE_API bool IsDefaultAsStringEmpty() const;
+
 	/** Returns a human readable FText representation of the string/object/text default value */
 	ENGINE_API FText GetDefaultAsText() const;
 
diff --git a/Engine/Source/Runtime/Engine/Private/EdGraph/EdGraphPin.cpp b/Engine/Source/Runtime/Engine/Private/EdGraph/EdGraphPin.cpp
--- a/Engine/Source/Runtime/Engine/Private/EdGraph/EdGraphPin.cpp
+++ b/Engine/Source/Runtime/Engine/Private/EdGraph/EdGraphPin.cpp
@@ -771,6 +771,11 @@ FString UEdGraphPin::GetDefaultAsString() const
 	}
 }
 
+bool UEdGraphPin::IsDefaultAsStringEmpty() const
+{
+	return DefaultObject.IsNullNoResolve() && DefaultTextValue.IsEmpty() && DefaultValue.IsEmpty();
+}
+
 FText UEdGraphPin::GetDefaultAsText() const
 {
 	if (DefaultObject)
diff --git a/Engine/Source/Runtime/Engine/Private/EdGraph/EdGraphSchema.cpp b/Engine/Source/Runtime/Engine/Private/EdGraph/EdGraphSchema.cpp
--- a/Engine/Source/Runtime/Engine/Private/EdGraph/EdGraphSchema.cpp
+++ b/Engine/Source/Runtime/Engine/Private/EdGraph/EdGraphSchema.cpp
@@ -508,6 +508,11 @@ void UEdGraphSchema::TrySetDefaultText(UEdGraphPin& InPin, const FText& InNewDef
 
 bool UEdGraphSchema::DoesDefaultValueMatch(const UEdGraphPin& InPin, const FString& InValue) const
 {
+	if (InValue.IsEmpty())
+	{
+		return InPin.IsDefaultAsStringEmpty();
+	}
+
 	// Same logic as UEdGraphPin::DoesDefaultValueMatchAutogenerated
 	// Ignoring case on purpose to match default behavior
 	FText ValueAsText;
