Fix a bug with ConvertAsset nodes that were spawned during the expansion of a promotable operator. The pin name on the ConvertAsset node does not match the normal "Return" pin name, so it would always be null and cause a crash.

#rnx

diff --git a/Engine/Source/Editor/BlueprintGraph/Private/K2Node_PromotableOperator.cpp b/Engine/Source/Editor/BlueprintGraph/Private/K2Node_PromotableOperator.cpp
--- a/Engine/Source/Editor/BlueprintGraph/Private/K2Node_PromotableOperator.cpp
+++ b/Engine/Source/Editor/BlueprintGraph/Private/K2Node_PromotableOperator.cpp
@@ -40,6 +40,7 @@
 #include "UObject/UnrealType.h"
 #include "UObject/WeakObjectPtr.h"
 #include "UObject/WeakObjectPtrTemplates.h"
+#include "EdGraphSchema_K2_Actions.h"
 
 #define LOCTEXT_NAMESPACE "PromotableOperatorNode"
 
@@ -1017,9 +1018,11 @@ bool UK2Node_PromotableOperator::CreateIntermediateCast(UK2Node_CallFunction* So
 		TemplateNode->AllocateDefaultPins();
 		CompilerContext.MessageLog.NotifyIntermediateObjectCreation(TemplateNode, this);
 	}
-	else
+	else if (Schema->FindSpecializedConversionNode(InputPin, OutputPin, true, /*out*/ TemplateConversionNode))
 	{
-		Schema->FindSpecializedConversionNode(InputPin, OutputPin, true, /*out*/ TemplateConversionNode);
+		FVector2D AverageLocation = UEdGraphSchema_K2::CalculateAveragePositionBetweenNodes(InputPin, OutputPin);		
+		TemplateConversionNode = FEdGraphSchemaAction_K2NewNode::SpawnNodeFromTemplate<UK2Node>(SourceGraph, TemplateConversionNode, AverageLocation);		
+		CompilerContext.MessageLog.NotifyIntermediateObjectCreation(TemplateConversionNode, this);
 	}
 
 	bool bInputSuccessful = false;
@@ -1028,16 +1031,24 @@ bool UK2Node_PromotableOperator::CreateIntermediateCast(UK2Node_CallFunction* So
 	if (TemplateConversionNode)
 	{
 		UEdGraphPin* ConversionInput = nullptr;
+		UEdGraphPin* ConversionOutput = TemplateConversionNode->FindPin(UEdGraphSchema_K2::PN_ReturnValue, EGPD_Output);
 
 		for (UEdGraphPin* ConvPin : TemplateConversionNode->Pins)
 		{
-			if (ConvPin && ConvPin->Direction == EGPD_Input && ConvPin->PinName != UEdGraphSchema_K2::PSC_Self)
+			if (ConvPin)
 			{
-				ConversionInput = ConvPin;
-				break;
+				if (!ConversionInput && ConvPin->Direction == EGPD_Input && ConvPin->PinName != UEdGraphSchema_K2::PSC_Self)
+				{
+					ConversionInput = ConvPin;
+				}
+				else if (!ConversionOutput && ConvPin->Direction == EGPD_Output)
+				{
+					ConversionOutput = ConvPin;
+				}
 			}
 		}
-		UEdGraphPin* ConversionOutput = TemplateConversionNode->FindPin(UEdGraphSchema_K2::PN_ReturnValue, EGPD_Output);
+
+		ensure(ConversionInput && ConversionOutput);
 
 		// Connect my input to the conversion node directly if we have links, otherwise we need to move the intermediate version of it
 		if (InputPin->LinkedTo.Num() > 0)
