Don't try to create transactions while recompiling BP in async loading

diff --git a/Engine/Source/Editor/BlueprintGraph/Private/EdGraphSchema_K2_Actions.cpp b/Engine/Source/Editor/BlueprintGraph/Private/EdGraphSchema_K2_Actions.cpp
--- a/Engine/Source/Editor/BlueprintGraph/Private/EdGraphSchema_K2_Actions.cpp
+++ b/Engine/Source/Editor/BlueprintGraph/Private/EdGraphSchema_K2_Actions.cpp
@@ -243,7 +243,7 @@ UEdGraphNode* FEdGraphSchemaAction_K2NewNode::CreateNode(
 	const bool bSelectNewNode = (Options & EK2NewNodeFlags::SelectNewNode) != EK2NewNodeFlags::None;
 	const bool bGotoNode = (Options & EK2NewNodeFlags::GotoNewNode) != EK2NewNodeFlags::None;
 
-	const FScopedTransaction Transaction(NSLOCTEXT("UnrealEd", "K2_AddNode", "Add Node"));
+	const FScopedTransaction Transaction(NSLOCTEXT("UnrealEd", "K2_AddNode", "Add Node"), !IsInAsyncLoadingThread());
 	ParentGraph->Modify();
 	for(UEdGraphPin* FromPin : FromPins)
 	{
