Exclude override flag members from the optional input pin set on MakeStruct nodes.

Overview:
- An "override flag" is an inline toggle-style Boolean edit condition. These are implicitly set to true at runtime within the output struct value if any member that's bound to it is exposed as a visible input pin on a MakeStruct node.
- For context, the Property Editor implicitly sets an override flag to true at edit time when a member that's bound to it is enabled for editing. These members are not otherwise labeled/exposed for direct editing.
- An override flag is meant to signal to a system that the user wishes to use the bound member's value in place of the current value (whatever that may be) when the full struct value is applied. Examples: FPostProcessSettings, FMovieSceneSequencePlaybackSettings, etc.

Previous UX:
- All boolean edit conditions were being treated as override flags on a MakeStruct node.
- Any inline toggle edit condition that did not begin with "bOverride_" or whose suffix otherwise did not match another value member could be exposed as an input on a MakeStruct node.
- Override flags exposed as inputs would always be set to TRUE at runtime regardless of input if it was declared at the top of the struct and if a member value bound to it was also exposed as an input pin.

After this change:
- Only inline toggle edit conditions and legacy struct members that follow the "bOverride_" naming convention will be treated as an override flag on a MakeStruct node.
- Inline toggle edit conditions can no longer be exposed directly as an input on a MakeStruct node. The intent was to bring the MakeStruct node UX closer to parity with the Property Editor UX.

Additional notes:
- Members that follow the legacy "bOverride_" naming convention were already being excluded from the optional input pin set on MakeStruct nodes if another member property name also matched its suffix. These have historically been excluded from ALL optional pin sets that utilize any FOptionalPinManager subtype (regardless of node type), so there was no change here.
- Existing MakeStruct nodes that may have already exposed inline toggle edit condition members as input pins will now orphan those pins on load if connected or if set to a non-default value (true). The "correct" way to set an override flag is by choosing to expose a member that's bound to the override condition as an input.
- Existing BreakStruct nodes are unchanged currently. Meaning, inline edit conditions that don't follow the legacy "bOverride_" convention can still be optionally exposed as an output pin. This UX was preserved as existing Blueprint logic could conceivably rely on the value of an override flag.
- Only one implicit assignment is now emitted for each override flag binding. Previously, we were emitting one assignment statement per bound property, so it could result in redundant assignments to the same flag if more than one property was bound to it.

diff --git a/Engine/Source/Editor/BlueprintGraph/Classes/K2Node_StructMemberSet.h b/Engine/Source/Editor/BlueprintGraph/Classes/K2Node_StructMemberSet.h
--- a/Engine/Source/Editor/BlueprintGraph/Classes/K2Node_StructMemberSet.h
+++ b/Engine/Source/Editor/BlueprintGraph/Classes/K2Node_StructMemberSet.h
@@ -18,6 +18,7 @@
 
 class FName;
 class FProperty;
+class FBoolProperty;
 class UEdGraphPin;
 class UObject;
 
@@ -51,6 +52,9 @@ class UK2Node_StructMemberSet : public UK2Node_StructOperation
 	// Helper for AllocateDefaultPins
 	BLUEPRINTGRAPH_API void AllocateExecPins();
 
+	// Returns the override condition that's bound to the given property, if one exists. An override is a edit condition flag that is not exposed as an optional input pin.
+	BLUEPRINTGRAPH_API FBoolProperty* GetOverrideConditionForProperty(const FProperty* InProperty) const;
+
 private:
 	/** Constructing FText strings can be costly, so we cache the node's title/tooltip */
 	FNodeTextCache CachedTooltip;
diff --git a/Engine/Source/Editor/BlueprintGraph/Private/K2Node.cpp b/Engine/Source/Editor/BlueprintGraph/Private/K2Node.cpp
--- a/Engine/Source/Editor/BlueprintGraph/Private/K2Node.cpp
+++ b/Engine/Source/Editor/BlueprintGraph/Private/K2Node.cpp
@@ -1523,43 +1523,50 @@ void FOptionalPinManager::RebuildPropertyList(TArray<FOptionalPinFromProperty>&
 		OldPinSettings.Add(PropertyEntry.PropertyName, FOldOptionalPinSettings(PropertyEntry.bShowPin, PropertyEntry.bIsOverrideEnabled, PropertyEntry.bIsSetValuePinVisible, PropertyEntry.bIsOverridePinVisible));
 	}
 
+	// This exists to support legacy struct types that used an older convention to declare
+	// override flag properties with a "bOverride_" prefix. Any property that starts with
+	// "bOverride_" will be implicitly excluded from the optional pin set if the suffix matches
+	// another property name within the same scope. For example, a property named "bOverride_MyValue"
+	// will not be included in the output set if we find a matching property named "MyValue".
+	TMap<FString, FProperty*> OverridesMap;
+	const FString OverridePrefix(TEXT("bOverride_"));
+
 	// Rebuild the property list
 	Properties.Reset();
-
-	// find all "bOverride_" properties
-	TMap<FName, FProperty*> OverridesMap;
-	const FString OverridePrefix(TEXT("bOverride_"));
+	TSet<FProperty*> TestProperties;
 	for (TFieldIterator<FProperty> It(SourceStruct, EFieldIteratorFlags::IncludeSuper); It; ++It)
 	{
-		FProperty* TestProperty = *It;
-		if (CanTreatPropertyAsOptional(TestProperty) && TestProperty->GetName().StartsWith(OverridePrefix))
+		FProperty* Property = *It;
+		if (CanTreatPropertyAsOptional(Property))
 		{
-			FString OriginalName = TestProperty->GetName();
+			TestProperties.Add(Property);
+
+			// find all "bOverride_" properties
+			FString OriginalName = Property->GetName();
 			if (OriginalName.RemoveFromStart(OverridePrefix) && !OriginalName.IsEmpty())
 			{
-				OverridesMap.Add(FName(*OriginalName), TestProperty);
+				OverridesMap.Add(OriginalName, Property);
 			}
 		}
 	}
 
 	// handle regular properties
-	for (TFieldIterator<FProperty> It(SourceStruct, EFieldIteratorFlags::IncludeSuper); It; ++It)
+	for (FProperty* TestProperty : TestProperties)
 	{
-		FProperty* TestProperty = *It;
-		if (CanTreatPropertyAsOptional(TestProperty) && !TestProperty->GetName().StartsWith(OverridePrefix))
+		if (!TestProperty->GetName().StartsWith(OverridePrefix))
 		{
 			FName CategoryName = NAME_None;
 #if WITH_EDITOR
 			CategoryName = FObjectEditorUtils::GetCategoryFName(TestProperty);
 #endif //WITH_EDITOR
 
-			OverridesMap.Remove(TestProperty->GetFName());
+			OverridesMap.Remove(TestProperty->GetName());
 			RebuildProperty(TestProperty, CategoryName, Properties, SourceStruct, OldPinSettings);
 		}
 	}
 
 	// add remaining "bOverride_" properties
-	for (const TPair<FName, FProperty*>& Pair : OverridesMap)
+	for (const TPair<FString, FProperty*>& Pair : OverridesMap)
 	{
 		FProperty* TestProperty = Pair.Value;
 
@@ -1580,9 +1587,12 @@ void FOptionalPinManager::RebuildProperty(FProperty* TestProperty, FName Categor
 	Record->PropertyTooltip = TestProperty->GetToolTipText();
 	Record->CategoryName = CategoryName;
 
+	// Determine if the property is bound to a mutable edit condition flag, which historically indicates an additional optional
+	// "override" pin is needed. Note that some manager subtypes may choose not to include these fields in the optional pin set.
 	bool bNegate = false;
 	FProperty* OverrideProperty = PropertyCustomizationHelpers::GetEditConditionProperty(TestProperty, bNegate);
 	Record->bHasOverridePin = OverrideProperty != nullptr && OverrideProperty->HasAllPropertyFlags(CPF_BlueprintVisible) && !OverrideProperty->HasAllPropertyFlags(CPF_BlueprintReadOnly);
+
 	Record->bIsMarkedForAdvancedDisplay = TestProperty->HasAnyPropertyFlags(CPF_AdvancedDisplay);
 
 	// Get the defaults
diff --git a/Engine/Source/Editor/BlueprintGraph/Private/K2Node_MakeStruct.cpp b/Engine/Source/Editor/BlueprintGraph/Private/K2Node_MakeStruct.cpp
--- a/Engine/Source/Editor/BlueprintGraph/Private/K2Node_MakeStruct.cpp
+++ b/Engine/Source/Editor/BlueprintGraph/Private/K2Node_MakeStruct.cpp
@@ -108,6 +108,18 @@ static bool CanBeExposed(const FProperty* Property, UBlueprint* BP)
 		const UEdGraphSchema_K2* Schema = GetDefault<UEdGraphSchema_K2>();
 		check(Schema);
 
+		// Treat all inline edit condition properties as override flags; that is, don't allow
+		// these to be exposed as part of the optional input pin set. Their value will be set
+		// implicitly at runtime based on whether or not any bound members are exposed, rather
+		// than explicitly via an exposed input pin. This emulates how the Property Editor
+		// handles setting these values at edit time (they appear as an inline checkbox that
+		// the user ticks on to set the flag and enable/override a bound property's value).
+		static const FName MD_InlineEditConditionToggle(TEXT("InlineEditConditionToggle"));
+		if (Property->HasMetaData(MD_InlineEditConditionToggle))
+		{
+			return false;
+		}
+
 		const bool bIsEditorBP = IsEditorOnlyObject(BP);
 		const bool bIsEditAnywhereProperty = Property->HasAllPropertyFlags(CPF_Edit) &&
 			!Property->HasAnyPropertyFlags(CPF_EditConst);
diff --git a/Engine/Source/Editor/BlueprintGraph/Private/K2Node_StructMemberSet.cpp b/Engine/Source/Editor/BlueprintGraph/Private/K2Node_StructMemberSet.cpp
--- a/Engine/Source/Editor/BlueprintGraph/Private/K2Node_StructMemberSet.cpp
+++ b/Engine/Source/Editor/BlueprintGraph/Private/K2Node_StructMemberSet.cpp
@@ -13,6 +13,7 @@
 #include "UObject/ObjectPtr.h"
 #include "UObject/UnrealNames.h"
 #include "UObject/UnrealType.h"
+#include "PropertyCustomizationHelpers.h"
 
 class FKismetCompilerContext;
 class UEdGraphPin;
@@ -103,4 +104,24 @@ FNodeHandlingFunctor* UK2Node_StructMemberSet::CreateNodeHandler(FKismetCompiler
 	return new FKCHandler_StructMemberVariableSet(CompilerContext);
 }
 
+FBoolProperty* UK2Node_StructMemberSet::GetOverrideConditionForProperty(const FProperty* InProperty) const
+{
+	bool bNegate = false;
+	if (FBoolProperty* OverrideProperty = PropertyCustomizationHelpers::GetEditConditionProperty(InProperty, bNegate))
+	{
+		// Determine if the edit condition is included in the optional input pin set - if so, then it is not considered to be an override condition.
+		const bool bIsOverrideExposedForInput = ShowPinForProperties.ContainsByPredicate([OverrideProperty](const FOptionalPinFromProperty& PropertyEntry)
+		{
+			return PropertyEntry.PropertyName == OverrideProperty->GetFName();
+		});
+
+		if (!bIsOverrideExposedForInput)
+		{
+			return OverrideProperty;
+		}
+	}
+
+	return nullptr;
+}
+
 #undef LOCTEXT_NAMESPACE
diff --git a/Engine/Source/Editor/BlueprintGraph/Private/MakeStructHandler.cpp b/Engine/Source/Editor/BlueprintGraph/Private/MakeStructHandler.cpp
--- a/Engine/Source/Editor/BlueprintGraph/Private/MakeStructHandler.cpp
+++ b/Engine/Source/Editor/BlueprintGraph/Private/MakeStructHandler.cpp
@@ -23,7 +23,6 @@
 #include "KismetCompiledFunctionContext.h"
 #include "KismetCompiler.h"
 #include "Misc/AssertionMacros.h"
-#include "PropertyCustomizationHelpers.h"
 #include "Templates/Casts.h"
 #include "UObject/Class.h"
 #include "UObject/NameTypes.h"
@@ -133,9 +132,15 @@ void FKCHandler_MakeStruct::Compile(FKismetFunctionContext& Context, UEdGraphNod
 	FBPTerminal* OutputStructTerm = FoundTerm ? *FoundTerm : NULL;
 	check(OutputStructTerm);
 
+	// A set of edit condition properties that should be set to signal an override state
+	// when the condition is not exposed as a unique input alongside any bound properties.
+	TMap<FProperty*, bool> OverrideProperties;
+
+	const UEdGraphSchema_K2* Schema = CompilerContext.GetSchema();
+
 	for (UEdGraphPin* Pin : Node->Pins)
 	{
-		if (Pin && !Pin->bOrphanedPin && (Pin != StructPin) && (Pin->Direction == EGPD_Input) && !CompilerContext.GetSchema()->IsMetaPin(*Pin))
+		if (Pin && !Pin->bOrphanedPin && (Pin != StructPin) && (Pin->Direction == EGPD_Input) && !Schema->IsMetaPin(*Pin))
 		{
 			FProperty* BoundProperty = FindFProperty<FProperty>(Node->StructType, Pin->PinName);
 			check(BoundProperty);
@@ -157,61 +162,83 @@ void FKCHandler_MakeStruct::Compile(FKismetFunctionContext& Context, UEdGraphNod
 				}
 			}
 
-			// Handle injecting the override property values into the node if the property has any
-			bool bNegate = false;
-			FProperty* OverrideProperty = PropertyCustomizationHelpers::GetEditConditionProperty(BoundProperty, bNegate);
-
+			// Determine if the property bound to the input pin has an associated override flag that's linked via
+			// an edit condition. If we find an override property and it's not included in the optional pin set,
+			// it signals that we need to implicitly set it at runtime since the value bound to it is exposed as
+			// an input. This convention allows us to roughly emulate the Property Editor which injects the override
+			// flag's value when the user modifies/exposes the bound property's value as an override at edit time.
+			// 
+			// Note: This API returns a Boolean flag type only. Other condition types are not supported since the
+			// "override" concept involves a direct association with a Boolean edit condition. Override properties
+			// must also be visible to the BP runtime as well as writable, or the terms used in the implicit
+			// assignment statement below may otherwise fail to compile.
+			FBoolProperty* OverrideProperty = Node->GetOverrideConditionForProperty(BoundProperty);
 			if (OverrideProperty && OverrideProperty->HasAllPropertyFlags(CPF_BlueprintVisible) && !OverrideProperty->HasAllPropertyFlags(CPF_BlueprintReadOnly))
 			{
-				const UEdGraphSchema_K2* Schema = GetDefault<UEdGraphSchema_K2>();
-				FEdGraphPinType PinType;
-				Schema->ConvertPropertyToPinType(OverrideProperty, /*out*/ PinType);
-
-				// Create the term in the list
-				FBPTerminal* OverrideTerm = new FBPTerminal();
-				Context.VariableReferences.Add(OverrideTerm);
-				OverrideTerm->Type = PinType;
-				OverrideTerm->AssociatedVarProperty = OverrideProperty;
-				OverrideTerm->Context = OutputStructTerm;
-
-				FBlueprintCompiledStatement* AssignBoolStatement = new FBlueprintCompiledStatement;
-				AssignBoolStatement->Type = KCST_Assignment;
-
-				// Literal Bool Term to set the OverrideProperty to
-				FBPTerminal* BoolTerm = Context.CreateLocalTerminal(ETerminalSpecification::TS_Literal);
-				BoolTerm->Type.PinCategory = UEdGraphSchema_K2::PC_Boolean;
-				BoolTerm->bIsLiteral = true;
-				// If we are showing the pin, then we are overriding the property
-
-				// Assigning the OverrideProperty to the literal bool term
-				AssignBoolStatement->LHS = OverrideTerm;
-				AssignBoolStatement->RHS.Add(BoolTerm);
-
 				// Need to dig up what the state of the override property should be
-				for (FOptionalPinFromProperty& PropertyEntry : Node->ShowPinForProperties)
+				const FOptionalPinFromProperty* BoundPropertyEntryPtr = Node->ShowPinForProperties.FindByPredicate([BoundProperty](const FOptionalPinFromProperty& PropertyEntry)
 				{
-					if (PropertyEntry.bHasOverridePin && PropertyEntry.bShowPin && PropertyEntry.PropertyName == BoundProperty->GetFName())
+					// If we are showing the pin, then we are overriding the property
+					return PropertyEntry.bHasOverridePin && PropertyEntry.bShowPin && PropertyEntry.PropertyName == BoundProperty->GetFName();
+				});
+
+				if (BoundPropertyEntryPtr)
+				{
+					const FOptionalPinFromProperty& PropertyEntry = *BoundPropertyEntryPtr;
+
+					if (!PropertyEntry.bIsOverridePinVisible || (PropertyEntry.bIsOverridePinVisible && !PropertyEntry.bIsOverrideEnabled && PropertyEntry.bIsSetValuePinVisible))
 					{
-						if (!PropertyEntry.bIsOverridePinVisible || (PropertyEntry.bIsOverridePinVisible && !PropertyEntry.bIsOverrideEnabled && PropertyEntry.bIsSetValuePinVisible))
-						{
-							CompilerContext.MessageLog.Warning(*LOCTEXT("MakeStruct_InvalidOverrideSetting", "Selected override setting on @@ in @@ is no longer a supported workflow and it is advised that you refactor your Blueprint to not use it!").ToString(), Pin, Node->GetBlueprint());
-						}
-
-						if (PropertyEntry.bIsOverridePinVisible)
-						{
-							BoolTerm->Name = PropertyEntry.bIsOverrideEnabled ? TEXT("true") : TEXT("false");
-							Context.AllGeneratedStatements.Add(AssignBoolStatement);
-
-							TArray<FBlueprintCompiledStatement*>& StatementList = Context.StatementsPerNode.FindOrAdd(Node);
-							StatementList.Add(AssignBoolStatement);
-						}
-						break;
+						CompilerContext.MessageLog.Warning(*LOCTEXT("MakeStruct_InvalidOverrideSetting", "Selected override setting on @@ in @@ is no longer a supported workflow and it is advised that you refactor your Blueprint to not use it!").ToString(), Pin, Node->GetBlueprint());
+					}
+
+					// This flag is a bit of a misnomer - it's generally true, except for assets saved prior to the convention described above, in which case the pin
+					// may have been explicitly hidden by the user, and thus this flag would be false to indicate that. For backwards-compatibility, we don't want to
+					// implicitly set the value in that case, because it means the user had explicitly opted not to show it, which implies it was not previously set.
+					if (PropertyEntry.bIsOverridePinVisible)
+					{
+						// Note: bIsOverrideEnabled is generally true, but may be false for older assets that were saved with an exposed override pin with a default
+						// value prior to the change that excludes all override properties from the optional input pin set.
+						OverrideProperties.Add(OverrideProperty, PropertyEntry.bIsOverrideEnabled);
 					}
 				}
 			}
 		}
 	}
 
+	// Handle injecting the override property values into the node
+	TArray<FBlueprintCompiledStatement*>& StatementList = Context.StatementsPerNode.FindOrAdd(Node);
+	for (auto OverridePropIt = OverrideProperties.CreateConstIterator(); OverridePropIt; ++OverridePropIt)
+	{
+		FProperty* OverrideProperty = OverridePropIt.Key();
+		const bool bIsOverrideEnabled = OverridePropIt.Value();
+
+		FEdGraphPinType PinType;
+		Schema->ConvertPropertyToPinType(OverrideProperty, /*out*/ PinType);
+
+		// Create the term in the list
+		FBPTerminal* OverrideTerm = new FBPTerminal();
+		Context.VariableReferences.Add(OverrideTerm);
+		OverrideTerm->Type = PinType;
+		OverrideTerm->AssociatedVarProperty = OverrideProperty;
+		OverrideTerm->Context = OutputStructTerm;
+
+		FBlueprintCompiledStatement* AssignBoolStatement = new FBlueprintCompiledStatement();
+		AssignBoolStatement->Type = KCST_Assignment;
+
+		// Literal Bool Term to set the OverrideProperty to
+		FBPTerminal* BoolTerm = Context.CreateLocalTerminal(ETerminalSpecification::TS_Literal);
+		BoolTerm->Type.PinCategory = UEdGraphSchema_K2::PC_Boolean;
+		BoolTerm->bIsLiteral = true;
+		BoolTerm->Name = bIsOverrideEnabled ? TEXT("true") : TEXT("false");
+
+		// Assigning the OverrideProperty to the literal bool term
+		AssignBoolStatement->LHS = OverrideTerm;
+		AssignBoolStatement->RHS.Add(BoolTerm);
+
+		Context.AllGeneratedStatements.Add(AssignBoolStatement);
+		StatementList.Add(AssignBoolStatement);
+	}
+
 	if (bAutoGenerateGotoForPure && !Node->IsNodePure())
 	{
 		GenerateSimpleThenGoto(Context, *Node);
diff --git a/Engine/Source/Editor/GraphEditor/Private/KismetNodes/SGraphNodeMakeStruct.cpp b/Engine/Source/Editor/GraphEditor/Private/KismetNodes/SGraphNodeMakeStruct.cpp
--- a/Engine/Source/Editor/GraphEditor/Private/KismetNodes/SGraphNodeMakeStruct.cpp
+++ b/Engine/Source/Editor/GraphEditor/Private/KismetNodes/SGraphNodeMakeStruct.cpp
@@ -383,36 +383,41 @@ TSharedPtr<SGraphPin> SGraphNodeMakeStruct::CreatePinWidget(UEdGraphPin* Pin) co
 
 	if (!Pin->PinType.bIsReference && Pin->PinType.PinCategory != UEdGraphSchema_K2::PC_Exec && Pin->Direction != EGPD_Output)
 	{
-		FOptionalPinFromProperty ReturnedPropertyEntry;
-		for (FOptionalPinFromProperty& PropertyEntry : SetFieldsNode->ShowPinForProperties)
+		// Attempt to find a matching optional pin entry that is bound to an edit condition as a potential override.
+		FOptionalPinFromProperty* BoundPropertyEntryPtr = SetFieldsNode->ShowPinForProperties.FindByPredicate([Pin](const FOptionalPinFromProperty& PropertyEntry)
 		{
-			if (PropertyEntry.bHasOverridePin)
+			return PropertyEntry.bHasOverridePin && PropertyEntry.PropertyName == Pin->PinName;
+		});
+
+		if (BoundPropertyEntryPtr)
+		{
+			FOptionalPinFromProperty& PropertyEntry = *BoundPropertyEntryPtr;
+
+			// Find the member property that's bound to the input pin.
+			const FProperty* BoundProperty = FindFProperty<FProperty>(SetFieldsNode->StructType, Pin->PinName);
+			check(BoundProperty);
+
+			// If this member is also bound to an override condition, insert the widget that allows the user to toggle input on/off for the override value.
+			if (const FBoolProperty* OverrideProperty = SetFieldsNode->GetOverrideConditionForProperty(BoundProperty))
 			{
-				if (PropertyEntry.PropertyName == Pin->PinName)
+				TWeakPtr<SHorizontalBox> HorizontalPin = ResultPin->GetFullPinHorizontalRowWidget();
+				if (HorizontalPin.IsValid())
 				{
-					TWeakPtr<SHorizontalBox> HorizontalPin = ResultPin->GetFullPinHorizontalRowWidget();
-					if (HorizontalPin.IsValid())
-					{
-						// Setup the pin's editable state to be dependent on whether the override is enabled
-						TAttribute<bool> IsEditableAttribute;
-						TAttribute<bool>::FGetter IsEditableGetter;
-						IsEditableGetter.BindSP(this, &SGraphNodeMakeStruct::IsPinEnabled, &PropertyEntry);
-						IsEditableAttribute.Bind(IsEditableGetter);
-						ResultPin->SetIsEditable(IsEditableAttribute);
-
-						HorizontalPin.Pin()->InsertSlot(1)
-							.Padding(0.0f, 0.0f, 2.0f, 0.0f)
-							[
-								SNew(SOptionalPinStateView, Pin, PropertyEntry)
-							];
-					}
-
-
-					break;
+					// Setup the pin's editable state to be dependent on whether the override is enabled
+					TAttribute<bool> IsEditableAttribute;
+					TAttribute<bool>::FGetter IsEditableGetter;
+					IsEditableGetter.BindSP(this, &SGraphNodeMakeStruct::IsPinEnabled, &PropertyEntry);
+					IsEditableAttribute.Bind(IsEditableGetter);
+					ResultPin->SetIsEditable(IsEditableAttribute);
+
+					HorizontalPin.Pin()->InsertSlot(1)
+						.Padding(0.0f, 0.0f, 2.0f, 0.0f)
+						[
+							SNew(SOptionalPinStateView, Pin, PropertyEntry)
+						];
 				}
 			}
 		}
-
 	}
 	return ResultPin;
 }
