Don't allow composite nodes to be pasted into Blueprint graphs that aren't schema-compatible.

Change summary:
- Added UEdGraphSchema_K2::DoesSupportCollapsedNodes() as an overridable API for K2-based schema types. Also added overrides for derived schemas that do not currently support collapsed subgraphs as a feature (e.g. AnimBPs).
- Added a UK2Node_Composite::CanCreateUnderSpecifiedSchema() override. This will be called when determining node types that can be created from clipboard content (see FGraphObjectTextFactory in EdGraphUtilities.cpp).
- Modified UK2Node_Composite::PostPasteNode() to exclude nodes that are not schema-compatible from the pasted subgraph. Also updated to reassign the subgraph schema to match the target graph.
- Modified FEdGraphUtilities::MergeChildrenGraphsIn() to emit an error to the message log when a subgraph cannot be merged into a parent graph due to schema incompatibility. This also adds an appropriate error context to the source node(s).
- Modified FKismetCompilerContext::ProcessOneFunctionGraph() to early exit if we fail to merge subgraphs while processing the intermediate function graph. This prevents us from attempting to further compile an intermediate graph that's in an invalid state, which previously would lead to an ensure() and later fail due to a non-specific compile error.

diff --git a/Engine/Source/Editor/AnimGraph/Public/AnimationConduitGraphSchema.h b/Engine/Source/Editor/AnimGraph/Public/AnimationConduitGraphSchema.h
--- a/Engine/Source/Editor/AnimGraph/Public/AnimationConduitGraphSchema.h
+++ b/Engine/Source/Editor/AnimGraph/Public/AnimationConduitGraphSchema.h
@@ -17,8 +17,12 @@ class UAnimationConduitGraphSchema : public UEdGraphSchema_K2
 	virtual void CreateDefaultNodesForGraph(UEdGraph& Graph) const override;
 	virtual bool CanDuplicateGraph(UEdGraph* InSourceGraph) const override { return false; }
 	virtual void GetGraphDisplayInformation(const UEdGraph& Graph, /*out*/ FGraphDisplayInfo& DisplayInfo) const override;
-	virtual bool DoesSupportEventDispatcher() const	override { return false; }
 	virtual bool ShouldAlwaysPurgeOnModification() const override { return true; }
 	virtual void HandleGraphBeingDeleted(UEdGraph& GraphBeingRemoved) const override;
 	// End of UEdGraphSchema interface
+
+	// UEdGraphSchema_K2 interface
+	virtual bool DoesSupportCollapsedNodes() const override { return false; }
+	virtual bool DoesSupportEventDispatcher() const	override { return false; }
+	// End of UEdGraphSchema_K2 interface
 };
diff --git a/Engine/Source/Editor/AnimGraph/Public/AnimationGraphSchema.h b/Engine/Source/Editor/AnimGraph/Public/AnimationGraphSchema.h
--- a/Engine/Source/Editor/AnimGraph/Public/AnimationGraphSchema.h
+++ b/Engine/Source/Editor/AnimGraph/Public/AnimationGraphSchema.h
@@ -62,7 +62,6 @@ class UAnimationGraphSchema : public UEdGraphSchema_K2
 	virtual void GetContextMenuActions(class UToolMenu* Menu, class UGraphNodeContextMenuContext* Context) const override;
 	virtual FText GetPinDisplayName(const UEdGraphPin* Pin) const override;
 	virtual bool CanDuplicateGraph(UEdGraph* InSourceGraph) const override;
-	virtual bool DoesSupportEventDispatcher() const	override { return false; }
 	virtual bool ShouldAlwaysPurgeOnModification() const override { return true; }
 	virtual bool TryCreateConnection(UEdGraphPin* A, UEdGraphPin* B) const override;
 	virtual void GetGraphDisplayInformation(const UEdGraph& Graph, /*out*/ FGraphDisplayInfo& DisplayInfo) const override;
@@ -74,6 +73,8 @@ class UAnimationGraphSchema : public UEdGraphSchema_K2
 	virtual const FPinConnectionResponse DetermineConnectionResponseOfCompatibleTypedPins(const UEdGraphPin* PinA, const UEdGraphPin* PinB, const UEdGraphPin* InputPin, const UEdGraphPin* OutputPin) const override;
 	virtual bool SearchForAutocastFunction(const FEdGraphPinType& OutputPinType, const FEdGraphPinType& InputPinType, /*out*/ FName& TargetFunction, /*out*/ UClass*& FunctionOwner) const override;
 	virtual bool ArePinsCompatible(const UEdGraphPin* PinA, const UEdGraphPin* PinB, const UClass* CallingContext = NULL, bool bIgnoreArray = false) const override;
+	virtual bool DoesSupportCollapsedNodes() const override { return false; }
+	virtual bool DoesSupportEventDispatcher() const	override { return false; }
 	virtual bool DoesSupportAnimNotifyActions() const override;
 	virtual void CreateFunctionGraphTerminators(UEdGraph& Graph, UClass* Class) const override;
 	virtual bool CanShowDataTooltipForPin(const UEdGraphPin& Pin) const override;
diff --git a/Engine/Source/Editor/AnimGraph/Public/AnimationTransitionSchema.h b/Engine/Source/Editor/AnimGraph/Public/AnimationTransitionSchema.h
--- a/Engine/Source/Editor/AnimGraph/Public/AnimationTransitionSchema.h
+++ b/Engine/Source/Editor/AnimGraph/Public/AnimationTransitionSchema.h
@@ -21,11 +21,15 @@ class UAnimationTransitionSchema : public UEdGraphSchema_K2
 	virtual void CreateDefaultNodesForGraph(UEdGraph& Graph) const override;
 	virtual bool CanDuplicateGraph(UEdGraph* InSourceGraph) const override { return false; }
 	virtual void GetGraphDisplayInformation(const UEdGraph& Graph, /*out*/ FGraphDisplayInfo& DisplayInfo) const override;
-	virtual bool DoesSupportEventDispatcher() const	override { return false; }
 	virtual bool ShouldAlwaysPurgeOnModification() const override { return true; }
 	virtual void HandleGraphBeingDeleted(UEdGraph& GraphBeingRemoved) const override;
 	//~ End UEdGraphSchema Interface.
 
+	//~ Begin UEdGraphSchema_K2 Interface.
+	virtual bool DoesSupportCollapsedNodes() const override { return false; }
+	virtual bool DoesSupportEventDispatcher() const	override { return false; }
+	//~ End UEdGraphSchema_K2 Interface.
+
 private:
 	static UAnimStateTransitionNode* GetTransitionNodeFromGraph(const FAnimBlueprintDebugData& DebugData, const UEdGraph* Graph);
 
diff --git a/Engine/Source/Editor/BlueprintGraph/Classes/EdGraphSchema_K2.h b/Engine/Source/Editor/BlueprintGraph/Classes/EdGraphSchema_K2.h
--- a/Engine/Source/Editor/BlueprintGraph/Classes/EdGraphSchema_K2.h
+++ b/Engine/Source/Editor/BlueprintGraph/Classes/EdGraphSchema_K2.h
@@ -587,6 +587,13 @@ public:
 
 	//~ End EdGraphSchema Interface
 
+	/**
+	 * Determine if this graph supports collapsing nodes into subgraphs
+	 *
+	 * @return True if this schema supports collapsed node subgraphs
+	 */
+	virtual bool DoesSupportCollapsedNodes() const { return true; }
+
 	/**
 	 *
 	 *	Determine if this graph supports event dispatcher
diff --git a/Engine/Source/Editor/BlueprintGraph/Classes/K2Node_Composite.h b/Engine/Source/Editor/BlueprintGraph/Classes/K2Node_Composite.h
--- a/Engine/Source/Editor/BlueprintGraph/Classes/K2Node_Composite.h
+++ b/Engine/Source/Editor/BlueprintGraph/Classes/K2Node_Composite.h
@@ -46,6 +46,7 @@ class BLUEPRINTGRAPH_API UK2Node_Composite : public UK2Node_Tunnel
 	virtual void OnRenameNode(const FString& NewName) override;
 	virtual TSharedPtr<class INameValidatorInterface> MakeNameValidator() const override;
 	virtual TArray<UEdGraph*> GetSubGraphs() const override { return TArray<UEdGraph*>( { BoundGraph } ); }
+	virtual bool CanCreateUnderSpecifiedSchema(const UEdGraphSchema* DesiredSchema) const override;
 	//~ End UEdGraphNode Interface
 
 	//~ Begin UK2Node Interface
diff --git a/Engine/Source/Editor/BlueprintGraph/Private/K2Node_Composite.cpp b/Engine/Source/Editor/BlueprintGraph/Private/K2Node_Composite.cpp
--- a/Engine/Source/Editor/BlueprintGraph/Private/K2Node_Composite.cpp
+++ b/Engine/Source/Editor/BlueprintGraph/Private/K2Node_Composite.cpp
@@ -85,34 +85,50 @@ void UK2Node_Composite::PostPasteNode()
 {
 	Super::PostPasteNode();
 
-	//@TODO: Should verify that each node in the composite can be pasted into this new graph successfully (CanPasteHere)
-
-	if (BoundGraph != NULL)
+	if (BoundGraph != nullptr)
 	{
 		UEdGraph* ParentGraph = CastChecked<UEdGraph>(GetOuter());
 		ensure(BoundGraph != ParentGraph);
 
+		const UEdGraphSchema* ParentSchema = ParentGraph->GetSchema();
+		check(ParentSchema);
+
+		UBlueprint* BP = FBlueprintEditorUtils::FindBlueprintForGraphChecked(BoundGraph);
+
 		// Update the InputSinkNode / OutputSourceNode pointers to point to the new graph
 		TSet<UEdGraphNode*> BoundaryNodes;
 		for (int32 NodeIndex = 0; NodeIndex < BoundGraph->Nodes.Num(); ++NodeIndex)
 		{
 			UEdGraphNode* Node = BoundGraph->Nodes[NodeIndex];
 			
-			//Remove this node if it should not exist more then one in blueprint
-			if(UK2Node_Event* Event = Cast<UK2Node_Event>(Node))
+			// Remove this node if it should not exist more then one in blueprint
+			bool bRemoveNode = false;
+			if (UK2Node_Event* Event = Cast<UK2Node_Event>(Node))
 			{
-				UBlueprint* BP = FBlueprintEditorUtils::FindBlueprintForGraphChecked(BoundGraph);
-				if(FBlueprintEditorUtils::FindOverrideForFunction(BP, Event->EventReference.GetMemberParentClass(Event->GetBlueprintClassFromNode()), Event->EventReference.GetMemberName()))
+				if (FBlueprintEditorUtils::FindOverrideForFunction(BP, Event->EventReference.GetMemberParentClass(Event->GetBlueprintClassFromNode()), Event->EventReference.GetMemberName()))
 				{
-					FBlueprintEditorUtils::RemoveNode(BP, Node, true);
-					NodeIndex--;
-					continue;
+					bRemoveNode = true;
 				}
 			}
+
+			// Intentional that we check for exact class here!
+			bool bIsTunnelNode = false;
+			if (Node->GetClass() == UK2Node_Tunnel::StaticClass())
+			{
+				bIsTunnelNode = true;
+			}
+
+			// Also remove any nodes from the subgraph that are otherwise not schema-compatible
+			if (bRemoveNode || (!bIsTunnelNode && !ParentSchema->CanEncapuslateNode(*Node)))
+			{
+				FBlueprintEditorUtils::RemoveNode(BP, Node, true);
+				NodeIndex--;
+				continue;
+			}
 			
 			BoundaryNodes.Add(Node);
 
-			if (Node->GetClass() == UK2Node_Tunnel::StaticClass())
+			if (bIsTunnelNode)
 			{
 				// Exactly a tunnel node, should be the entrance or exit node
 				UK2Node_Tunnel* Tunnel = CastChecked<UK2Node_Tunnel>(Node);
@@ -140,6 +156,11 @@ void UK2Node_Composite::PostPasteNode()
 		//Nested composites will already be in the SubGraph array
 		if(ParentGraph->SubGraphs.Find(BoundGraph) == INDEX_NONE)
 		{
+			// Set the subgraph's schema class to match the parent (as we do when placing a new node). This is needed in
+			// order to ensure that the new subgraph will compile. We've already verified that the new subgraph is schema-
+			// compatible at this point (see above).
+			BoundGraph->Schema = ParentGraph->Schema;
+
 			ParentGraph->SubGraphs.Add(BoundGraph);
 		}
 
@@ -338,4 +359,16 @@ TSharedPtr<class INameValidatorInterface> UK2Node_Composite::MakeNameValidator()
 	return MakeShareable(new FKismetNameValidator(GetBlueprint(), BoundGraph ? BoundGraph->GetFName() : NAME_None));
 }
 
+bool UK2Node_Composite::CanCreateUnderSpecifiedSchema(const UEdGraphSchema* DesiredSchema) const
+{
+	if (Super::CanCreateUnderSpecifiedSchema(DesiredSchema))
+	{
+		// Check for derived schemas that don't support a collapsed subgraph.
+		const UEdGraphSchema_K2* K2Schema = CastChecked<UEdGraphSchema_K2>(DesiredSchema);
+		return K2Schema->DoesSupportCollapsedNodes();
+	}
+
+	return false;
+}
+
 #undef LOCTEXT_NAMESPACE
diff --git a/Engine/Source/Editor/KismetCompiler/Private/KismetCompiler.cpp b/Engine/Source/Editor/KismetCompiler/Private/KismetCompiler.cpp
--- a/Engine/Source/Editor/KismetCompiler/Private/KismetCompiler.cpp
+++ b/Engine/Source/Editor/KismetCompiler/Private/KismetCompiler.cpp
@@ -4369,7 +4369,13 @@ void FKismetCompilerContext::ProcessOneFunctionGraph(UEdGraph* SourceGraph, bool
 		FunctionGraph->Rename(*FunctionGraphName, nullptr, RenameFlags);
 	}
 
-	FEdGraphUtilities::MergeChildrenGraphsIn(FunctionGraph, FunctionGraph, /*bRequireSchemaMatch=*/ true);
+	FEdGraphUtilities::MergeChildrenGraphsIn(FunctionGraph, FunctionGraph, /* bRequireSchemaMatch = */ true, /* bInIsCompiling = */ true, &MessageLog);
+
+	// If we failed to merge with any child graphs due to an error, we shouldn't continue processing the intermediate graph.
+	if (MessageLog.NumErrors > 0)
+	{
+		return;
+	}
 
 	ExpansionStep(FunctionGraph, false);
 
diff --git a/Engine/Source/Editor/UnrealEd/Private/EdGraphUtilities.cpp b/Engine/Source/Editor/UnrealEd/Private/EdGraphUtilities.cpp
--- a/Engine/Source/Editor/UnrealEd/Private/EdGraphUtilities.cpp
+++ b/Engine/Source/Editor/UnrealEd/Private/EdGraphUtilities.cpp
@@ -413,6 +413,10 @@ void FEdGraphUtilities::MergeChildrenGraphsIn(UEdGraph* MergeTarget, UEdGraph* P
 
 			MergeChildrenGraphsIn(MergeTarget, ChildGraph, bRequireSchemaMatch, bInIsCompiling, MessageLog);
 		}
+		else if (bNonVirtualGraph && bRequireSchemaMatch && MessageLog)
+		{
+			MessageLog->Error(*NSLOCTEXT("EdGraphUtilities", "CannotMergeChildGraph_Error", "Unable to merge child graph for node @@. The child graph's schema is incompatible with the outer graph.").ToString(), MessageLog->FindSourceObject(NodeOwner));
+		}
 	}
 }
 
