Fix minor BP issues when async loading thread is activated

diff --git a/Engine/Source/Editor/BlueprintGraph/Private/K2Node_CallMaterialParameterCollectionFunction.cpp b/Engine/Source/Editor/BlueprintGraph/Private/K2Node_CallMaterialParameterCollectionFunction.cpp
--- a/Engine/Source/Editor/BlueprintGraph/Private/K2Node_CallMaterialParameterCollectionFunction.cpp
+++ b/Engine/Source/Editor/BlueprintGraph/Private/K2Node_CallMaterialParameterCollectionFunction.cpp
@@ -37,7 +37,13 @@ void UK2Node_CallMaterialParameterCollectionFunction::PreloadRequiredAssets()
 		if (UMaterialParameterCollection* Collection = Cast<UMaterialParameterCollection>(CollectionPin->DefaultObject))
 		{
 			PreloadObject(Collection);
-			Collection->ConditionalPostLoad();
+
+			// When on async loading thread, postload happens later on GT unless it's possible
+			// to do it right now safely.
+			if (IsInGameThread() || Collection->IsPostLoadThreadSafe())
+			{
+				Collection->ConditionalPostLoad();
+			}
 		}
 	}
 
diff --git a/Engine/Source/Editor/BlueprintGraph/Private/K2Node_ComponentBoundEvent.cpp b/Engine/Source/Editor/BlueprintGraph/Private/K2Node_ComponentBoundEvent.cpp
--- a/Engine/Source/Editor/BlueprintGraph/Private/K2Node_ComponentBoundEvent.cpp
+++ b/Engine/Source/Editor/BlueprintGraph/Private/K2Node_ComponentBoundEvent.cpp
@@ -175,7 +175,7 @@ bool UK2Node_ComponentBoundEvent::IsDelegateValid() const
 
 bool UK2Node_ComponentBoundEvent::HasDeprecatedReference() const
 {
-	if (CVarBPEnableDeprecatedWarningForComponentDelegateNodes.GetValueOnGameThread())
+	if (CVarBPEnableDeprecatedWarningForComponentDelegateNodes.GetValueOnAnyThread())
 	{
 		if (const FMulticastDelegateProperty* DelegateProperty = GetTargetDelegateProperty())
 		{
diff --git a/Engine/Source/Editor/BlueprintGraph/Private/K2Node_SwitchEnum.cpp b/Engine/Source/Editor/BlueprintGraph/Private/K2Node_SwitchEnum.cpp
--- a/Engine/Source/Editor/BlueprintGraph/Private/K2Node_SwitchEnum.cpp
+++ b/Engine/Source/Editor/BlueprintGraph/Private/K2Node_SwitchEnum.cpp
@@ -47,7 +47,13 @@ void UK2Node_SwitchEnum::SetEnum(UEnum* InEnum)
 	if (Enum)
 	{
 		PreloadObject(Enum);
-		Enum->ConditionalPostLoad();
+
+		// When on async loading thread, postload happens later on GT unless it's possible
+		// to do it right now safely.
+		if (IsInGameThread() || Enum->IsPostLoadThreadSafe())
+		{
+			Enum->ConditionalPostLoad();
+		}
 
 		for (int32 EnumIndex = 0; EnumIndex < Enum->NumEnums() - 1; ++EnumIndex)
 		{
