[Bug Repro] Error in 'Multiple Cook Client' LogWindows: Error: begin: stack for UAT
K2Node_BaseMCDelegate::HasExternalDependencies is accessing garbage data during cooks. To help find a repro, this change will log more info when we see the problem

diff --git a/Engine/Source/Editor/BlueprintGraph/Private/K2Node_MCDelegate.cpp b/Engine/Source/Editor/BlueprintGraph/Private/K2Node_MCDelegate.cpp
--- a/Engine/Source/Editor/BlueprintGraph/Private/K2Node_MCDelegate.cpp
+++ b/Engine/Source/Editor/BlueprintGraph/Private/K2Node_MCDelegate.cpp
@@ -224,6 +224,13 @@ bool UK2Node_BaseMCDelegate::HasExternalDependencies(TArray<class UStruct*>* Opt
 	}
 
 	UFunction* Signature = GetDelegateSignature(true);
+
+	checkf(!Signature || Signature->IsValidLowLevel(), TEXT("Invalid Signature: BP - %s Delegate Ref - %s %s" ),
+		*SourceBlueprint->GetName(),
+		*DelegateReference.GetMemberScopeName(),
+		*DelegateReference.GetMemberName().ToString()
+	);
+	
 	UClass* SignatureSourceClass = Signature ? Signature->GetOwnerClass() : nullptr;
 	const bool bSignatureResult = (SignatureSourceClass != NULL) && (SignatureSourceClass->ClassGeneratedBy != SourceBlueprint);
 	if (bSignatureResult && OptionalOutput)
