Loctext Namespace fix

diff --git a/Engine/Source/Editor/BlueprintGraph/Private/K2Node_FunctionEntry.cpp b/Engine/Source/Editor/BlueprintGraph/Private/K2Node_FunctionEntry.cpp
--- a/Engine/Source/Editor/BlueprintGraph/Private/K2Node_FunctionEntry.cpp
+++ b/Engine/Source/Editor/BlueprintGraph/Private/K2Node_FunctionEntry.cpp
@@ -768,7 +768,7 @@ void UK2Node_FunctionEntry::FindDiffs(UEdGraphNode* OtherNode, struct FDiffResul
 				{
 					// type changed
 					AddDiff(FText::Format(
-						LOCTEXT("DIF_FunctionLocalVariableNameChange", "Local variable {0}::{1} was a {2} but is now a {3}"),
+						LOCTEXT("DIF_FunctionLocalVariableTypeChange", "Local variable {0}::{1} was a {2} but is now a {3}"),
 						NodeName,
 						NewVarName,
 						UEdGraphSchema_K2::TypeToText(ThisVar.VarType),
@@ -779,7 +779,7 @@ void UK2Node_FunctionEntry::FindDiffs(UEdGraphNode* OtherNode, struct FDiffResul
 				{
 					// category changed
 					AddDiff(FText::Format(
-						LOCTEXT("DIF_FunctionLocalVariableNameChange", "Local variable {0}::{1} changed category from {2} to {3}"),
+						LOCTEXT("DIF_FunctionLocalVariableCategoryChange", "Local variable {0}::{1} changed category from {2} to {3}"),
 						NodeName,
 						NewVarName,
 						ThisVar.Category,
@@ -790,7 +790,7 @@ void UK2Node_FunctionEntry::FindDiffs(UEdGraphNode* OtherNode, struct FDiffResul
 				{
 					// property flags changed
 					AddDiff(FText::Format(
-						LOCTEXT("DIF_FunctionLocalVariableNameChange", "Local variable {0}::{1} changed property flags"),
+						LOCTEXT("DIF_FunctionLocalVariableFlagChange", "Local variable {0}::{1} changed property flags"),
 						NodeName,
 						NewVarName
 					));
@@ -799,7 +799,7 @@ void UK2Node_FunctionEntry::FindDiffs(UEdGraphNode* OtherNode, struct FDiffResul
 				{
 					// replication notify function changed
 					AddDiff(FText::Format(
-						LOCTEXT("DIF_FunctionLocalVariableNameChange", "Local variable {0}::{1} changed RepNotifyFunc"),
+						LOCTEXT("DIF_FunctionLocalVariableRepNotifyFuncChange", "Local variable {0}::{1} changed RepNotifyFunc"),
 						NodeName,
 						NewVarName
 					));
@@ -808,7 +808,7 @@ void UK2Node_FunctionEntry::FindDiffs(UEdGraphNode* OtherNode, struct FDiffResul
 				{
 					// replication condition changed
 					AddDiff(FText::Format(
-						LOCTEXT("DIF_FunctionLocalVariableNameChange", "Local variable {0}::{1} changed ReplicationCondition"),
+						LOCTEXT("DIF_FunctionLocalVariableRepConditionChange", "Local variable {0}::{1} changed ReplicationCondition"),
 						NodeName,
 						NewVarName
 					));
@@ -817,7 +817,7 @@ void UK2Node_FunctionEntry::FindDiffs(UEdGraphNode* OtherNode, struct FDiffResul
 				{
 					// replication condition changed
 					AddDiff(FText::Format(
-						LOCTEXT("DIF_FunctionLocalVariableNameChange", "Local variable {0}::{1} changed default value from '{2}' to '{3}'"),
+						LOCTEXT("DIF_FunctionLocalVariableDefaultValueChange", "Local variable {0}::{1} changed default value from '{2}' to '{3}'"),
 						NodeName,
 						NewVarName,
 						FText::FromString(ThisVar.DefaultValue),
@@ -830,7 +830,7 @@ void UK2Node_FunctionEntry::FindDiffs(UEdGraphNode* OtherNode, struct FDiffResul
 				// local variable is in this node but not in other node
 				Diff.Category = EDiffType::SUBTRACTION;
 				AddDiff(FText::Format(
-				LOCTEXT("DIF_FunctionLocalVariableNameChange", "Local variable {0}::{1} removed"),
+				LOCTEXT("DIF_FunctionLocalVariableRemoved", "Local variable {0}::{1} removed"),
 					NodeName,
 					OldVarName
 				));
@@ -847,7 +847,7 @@ void UK2Node_FunctionEntry::FindDiffs(UEdGraphNode* OtherNode, struct FDiffResul
 				// local variable is in other node but not in this node
 				Diff.Category = EDiffType::ADDITION;
 				AddDiff(FText::Format(
-				LOCTEXT("DIF_FunctionLocalVariableNameChange", "Local variable {0}::{1} added"),
+				LOCTEXT("DIF_FunctionLocalVariableAdded", "Local variable {0}::{1} added"),
 					NodeName,
 					NewVarName
 				));
