Blueprint - Keeping things TObjectPtr in the Select node to prevent loading on expansion.

diff --git a/Engine/Source/Editor/BlueprintGraph/Private/K2Node_Select.cpp b/Engine/Source/Editor/BlueprintGraph/Private/K2Node_Select.cpp
--- a/Engine/Source/Editor/BlueprintGraph/Private/K2Node_Select.cpp
+++ b/Engine/Source/Editor/BlueprintGraph/Private/K2Node_Select.cpp
@@ -969,7 +969,7 @@ void UK2Node_Select::ExpandNode(class FKismetCompilerContext& CompilerContext, U
 
 		// copy defaults as default values can be reset when the pin is connected
 		const FString DefaultValue = Pin->DefaultValue;
-		UObject* DefaultObject = Pin->DefaultObject;
+		TObjectPtr<UObject> DefaultObject = Pin->DefaultObject;
 		const FText DefaultTextValue = Pin->DefaultTextValue;
 		bool bMatchesDefaults = Pin->DoesDefaultValueMatchAutogenerated();
 
