Fix some edgraph node trying to create transactions on async loading thread

diff --git a/Engine/Source/Editor/BlueprintGraph/Private/EdGraphSchema_K2.cpp b/Engine/Source/Editor/BlueprintGraph/Private/EdGraphSchema_K2.cpp
--- a/Engine/Source/Editor/BlueprintGraph/Private/EdGraphSchema_K2.cpp
+++ b/Engine/Source/Editor/BlueprintGraph/Private/EdGraphSchema_K2.cpp
@@ -694,6 +694,14 @@ int32 UEdGraphSchema_K2::CurrentCacheRefreshID = 0;
 
 const FName UEdGraphSchema_K2::AllObjectTypes(TEXT("ObjectTypes"));
 
+namespace UEdGraphSchemaImpl
+{
+	bool ShouldActuallyTransact()
+	{
+		return !IsInAsyncLoadingThread();
+	}
+}
+
 UEdGraphSchema_K2::UEdGraphSchema_K2(const FObjectInitializer& ObjectInitializer)
 	: Super(ObjectInitializer)
 {
@@ -1120,7 +1128,7 @@ void UEdGraphSchema_K2::ReplaceSelectedNode(UEdGraphNode* SourceNode, AActor* Ta
 		UK2Node_Literal* LiteralNode = (UK2Node_Literal*)(SourceNode);
 		if (LiteralNode)
 		{
-			const FScopedTransaction Transaction( LOCTEXT("ReplaceSelectedNodeUndoTransaction", "Replace Selected Node") );
+			const FScopedTransaction Transaction( LOCTEXT("ReplaceSelectedNodeUndoTransaction", "Replace Selected Node"), UEdGraphSchemaImpl::ShouldActuallyTransact());
 
 			LiteralNode->Modify();
 			LiteralNode->SetObjectRef( TargetActor );
@@ -1819,7 +1827,7 @@ void UEdGraphSchema_K2::OnCreateNonExistentVariable( UK2Node_Variable* Variable,
 {
 	if (UEdGraphPin* Pin = Variable->FindPin(Variable->GetVarName()))
 	{
-		const FScopedTransaction Transaction( LOCTEXT("CreateMissingVariable", "Create Missing Variable") );
+		const FScopedTransaction Transaction( LOCTEXT("CreateMissingVariable", "Create Missing Variable"), UEdGraphSchemaImpl::ShouldActuallyTransact());
 
 		if (FBlueprintEditorUtils::AddMemberVariable(OwnerBlueprint,Variable->GetVarName(), Pin->PinType))
 		{
@@ -1833,7 +1841,7 @@ void UEdGraphSchema_K2::OnCreateNonExistentLocalVariable( UK2Node_Variable* Vari
 {
 	if (UEdGraphPin* Pin = Variable->FindPin(Variable->GetVarName()))
 	{
-		const FScopedTransaction Transaction( LOCTEXT("CreateMissingLocalVariable", "Create Missing Local Variable") );
+		const FScopedTransaction Transaction( LOCTEXT("CreateMissingLocalVariable", "Create Missing Local Variable"), UEdGraphSchemaImpl::ShouldActuallyTransact());
 
 		FName VarName = Variable->GetVarName();
 		if (FBlueprintEditorUtils::AddLocalVariable(OwnerBlueprint, Variable->GetGraph(), VarName, Pin->PinType))
@@ -1862,7 +1870,7 @@ void UEdGraphSchema_K2::OnReplaceVariableForVariableNode( UK2Node_Variable* Vari
 {
 	if(UEdGraphPin* Pin = Variable->FindPin(Variable->GetVarName()))
 	{
-		const FScopedTransaction Transaction( NSLOCTEXT("UnrealEd", "GraphEd_ReplaceVariable", "Replace Variable") );
+		const FScopedTransaction Transaction( NSLOCTEXT("UnrealEd", "GraphEd_ReplaceVariable", "Replace Variable"), UEdGraphSchemaImpl::ShouldActuallyTransact());
 		Variable->Modify();
 		Pin->Modify();
 
@@ -4624,13 +4632,11 @@ void UEdGraphSchema_K2::BreakNodeLinks(UEdGraphNode& TargetNode) const
 
 void UEdGraphSchema_K2::BreakPinLinks(UEdGraphPin& TargetPin, bool bSendsNodeNotifcation) const
 {
-	const FScopedTransaction Transaction( NSLOCTEXT("UnrealEd", "GraphEd_BreakPinLinks", "Break Pin Links") );
+	const FScopedTransaction Transaction( NSLOCTEXT("UnrealEd", "GraphEd_BreakPinLinks", "Break Pin Links"), UEdGraphSchemaImpl::ShouldActuallyTransact());
 
 	// cache this here, as BreakPinLinks can trigger a node reconstruction invalidating the TargetPin referenceS
 	UBlueprint* const Blueprint = FBlueprintEditorUtils::FindBlueprintForNodeChecked(TargetPin.GetOwningNode());
 
-
-
 	Super::BreakPinLinks(TargetPin, bSendsNodeNotifcation);
 
 	FBlueprintEditorUtils::MarkBlueprintAsModified(Blueprint);
@@ -4638,7 +4644,7 @@ void UEdGraphSchema_K2::BreakPinLinks(UEdGraphPin& TargetPin, bool bSendsNodeNot
 
 void UEdGraphSchema_K2::BreakSinglePinLink(UEdGraphPin* SourcePin, UEdGraphPin* TargetPin) const
 {
-	const FScopedTransaction Transaction( NSLOCTEXT("UnrealEd", "GraphEd_BreakSinglePinLink", "Break Pin Link") );
+	const FScopedTransaction Transaction( NSLOCTEXT("UnrealEd", "GraphEd_BreakSinglePinLink", "Break Pin Link"), UEdGraphSchemaImpl::ShouldActuallyTransact());
 
 	UBlueprint* Blueprint = FBlueprintEditorUtils::FindBlueprintForNodeChecked(TargetPin->GetOwningNode());
 
@@ -5794,7 +5800,7 @@ void UEdGraphSchema_K2::DroppedAssetsOnGraph(const TArray<FAssetData>& Assets, c
 			// Make sure we have an asset type that's registered with the component list
 			if (DestinationComponentType != nullptr)
 			{
-				const FScopedTransaction Transaction(LOCTEXT("CreateAddComponentFromAsset", "Add Component From Asset"));
+				const FScopedTransaction Transaction(LOCTEXT("CreateAddComponentFromAsset", "Add Component From Asset"), UEdGraphSchemaImpl::ShouldActuallyTransact());
 
 				FComponentTypeEntry ComponentType = { FString(), FString(), DestinationComponentType };
 
@@ -7355,7 +7361,7 @@ void UEdGraphSchema_K2::RecombinePin(UEdGraphPin* Pin) const
 
 void UEdGraphSchema_K2::OnPinConnectionDoubleCicked(UEdGraphPin* PinA, UEdGraphPin* PinB, const FVector2D& GraphPosition) const
 {
-	const FScopedTransaction Transaction(LOCTEXT("CreateRerouteNodeOnWire", "Create Reroute Node"));
+	const FScopedTransaction Transaction(LOCTEXT("CreateRerouteNodeOnWire", "Create Reroute Node"), UEdGraphSchemaImpl::ShouldActuallyTransact());
 
 	//@TODO: This constant is duplicated from inside of SGraphNodeKnot
 	const FVector2D NodeSpacerSize(42.0f, 24.0f);
