Blueprint: Update an ensure to include the name of the asset

diff --git a/Engine/Source/Editor/BlueprintGraph/Private/K2Node_CallArrayFunction.cpp b/Engine/Source/Editor/BlueprintGraph/Private/K2Node_CallArrayFunction.cpp
--- a/Engine/Source/Editor/BlueprintGraph/Private/K2Node_CallArrayFunction.cpp
+++ b/Engine/Source/Editor/BlueprintGraph/Private/K2Node_CallArrayFunction.cpp
@@ -39,7 +39,7 @@ void UK2Node_CallArrayFunction::AllocateDefaultPins()
 	Super::AllocateDefaultPins();
 
 	UEdGraphPin* TargetArrayPin = GetTargetArrayPin();
-	if (ensure(TargetArrayPin))
+	if (ensureMsgf(TargetArrayPin, TEXT("%s"), *GetFullName()))
 	{
 		TargetArrayPin->PinType.ContainerType = EPinContainerType::Array;
 		TargetArrayPin->PinType.bIsReference = true;
