UE-166921 - Make Lidar Point Cloud Point Node has had the Location input change from a Vector to a Vector 3f

Promotable operators weren't aware that certain struct types are implicitly convertible (eg: FVector/FVector3f). Similar to PC_Real, we need to treat these pairs of structs as equivalent types. This can be easily done by checking to see if there's an implicit conversion entry in FStructConversionTable.

diff --git a/Engine/Source/Editor/BlueprintGraph/Private/K2Node_PromotableOperator.cpp b/Engine/Source/Editor/BlueprintGraph/Private/K2Node_PromotableOperator.cpp
--- a/Engine/Source/Editor/BlueprintGraph/Private/K2Node_PromotableOperator.cpp
+++ b/Engine/Source/Editor/BlueprintGraph/Private/K2Node_PromotableOperator.cpp
@@ -21,6 +21,7 @@
 #include "HAL/PlatformCrt.h"
 #include "Internationalization/Internationalization.h"
 #include "K2Node.h"
+#include "Kismet/BlueprintTypeConversions.h"
 #include "Kismet2/BlueprintEditorUtils.h"
 #include "Kismet2/CompilerResultsLog.h"
 #include "Kismet2/WildcardNodeUtils.h"
@@ -1257,6 +1258,8 @@ void UK2Node_PromotableOperator::UpdatePinsFromFunction(const UFunction* Functio
 
 	auto ConformPinLambda = [&ChangedPin, bIsFromConversion](const FEdGraphPinType& FunctionPinType, UEdGraphPin* NodePin)
 	{
+		using namespace UE::Kismet::BlueprintTypeConversions;
+
 		// If the pin types are already equal, then we don't have to do any work
 		// If this is linked to wildcard pins, then we can just ignore it and handle it on expansion
 		if (!NodePin || FWildcardNodeUtils::IsLinkedToWildcard(NodePin))
@@ -1273,10 +1276,14 @@ void UK2Node_PromotableOperator::UpdatePinsFromFunction(const UFunction* Functio
 		// By default, conform to the type of the function param
 		FEdGraphPinType ConformingType = FunctionPinType;
 		const FEdGraphPinType HighestLinkedType = NodePin->LinkedTo.Num() > 0 ? FTypePromotion::GetPromotedType(NodePin->LinkedTo) : NodePin->PinType;
-		const bool bDifferingStructs = 
+		const UScriptStruct* NodePinStruct = Cast<UScriptStruct>(HighestLinkedType.PinSubCategoryObject);
+		const UScriptStruct* FunctionPinStruct = Cast<UScriptStruct>(FunctionPinType.PinSubCategoryObject);
+		const bool bHasImplicitConversionFunction = FStructConversionTable::Get().GetConversionFunction(NodePinStruct, FunctionPinStruct).IsSet();
+		const bool bDifferingStructs =
 			HighestLinkedType.PinCategory == UEdGraphSchema_K2::PC_Struct &&
-			FunctionPinType.PinCategory == UEdGraphSchema_K2::PC_Struct && 
-			HighestLinkedType.PinSubCategoryObject != FunctionPinType.PinSubCategoryObject;
+			FunctionPinType.PinCategory == UEdGraphSchema_K2::PC_Struct &&
+			HighestLinkedType.PinSubCategoryObject != FunctionPinType.PinSubCategoryObject &&
+			!bHasImplicitConversionFunction;
 
 		const bool bHasDeterminingType = NodePin->LinkedTo.Num() > 0 || !NodePin->DoesDefaultValueMatchAutogenerated();
 
